#DUPEMERGE of CL 4067369

Fix on hierarchy transform change

diff --git a/Engine/Plugins/Experimental/ControlRig/Source/ControlRig/Private/AnimNode_ControlRigBase.cpp b/Engine/Plugins/Experimental/ControlRig/Source/ControlRig/Private/AnimNode_ControlRigBase.cpp
--- a/Engine/Plugins/Experimental/ControlRig/Source/ControlRig/Private/AnimNode_ControlRigBase.cpp
+++ b/Engine/Plugins/Experimental/ControlRig/Source/ControlRig/Private/AnimNode_ControlRigBase.cpp
@@ -59,66 +59,78 @@ void FAnimNode_ControlRigBase::Update_AnyThread(const FAnimationUpdateContext& C
 void FAnimNode_ControlRigBase::UpdateInput(UControlRig* ControlRig, const FPoseContext& InOutput)
 {
 	// copy output of the rig
-	const FBoneContainer& RequiredBones = InOutput.Pose.GetBoneContainer();
+	// @todo: find a better option to do this
+	// but for now, we use this flag to skip the execution and updating input
+	// this is used to skip when refpose is chosen to view
+	if (ControlRig->bExecutionOn)
+	{
+		const FBoneContainer& RequiredBones = InOutput.Pose.GetBoneContainer();
 
-	// get component pose from control rig
-	FCSPose<FCompactPose> MeshPoses;
-	// first I need to convert to local pose
-	MeshPoses.InitPose(InOutput.Pose);
+		// get component pose from control rig
+		FCSPose<FCompactPose> MeshPoses;
+		// first I need to convert to local pose
+		MeshPoses.InitPose(InOutput.Pose);
 
-	// should we also update refpose? 
-	const int32 NumNodes = RigHierarchyItemNameMapping.Num();
-	for (int32 Index = 0; Index < NumNodes; ++Index)
-	{
-		if (RigHierarchyItemNameMapping[Index] != NAME_None)
+		// should we also update refpose? 
+		const int32 NumNodes = RigHierarchyItemNameMapping.Num();
+		for (int32 Index = 0; Index < NumNodes; ++Index)
 		{
-			FCompactPoseBoneIndex CompactPoseIndex(Index);
-			FTransform ComponentTransform = MeshPoses.GetComponentSpaceTransform(CompactPoseIndex);
-			if (NodeMappingContainer.IsValid())
+			if (RigHierarchyItemNameMapping[Index] != NAME_None)
 			{
-				// @todo : I should reverse retarget
-				//ComponentTransform = NodeMappingContainer->GetSourceToTargetTransform(RigHierarchyItemNameMapping[Index]) * ComponentTransform;
+				FCompactPoseBoneIndex CompactPoseIndex(Index);
+				FTransform ComponentTransform = MeshPoses.GetComponentSpaceTransform(CompactPoseIndex);
+				if (NodeMappingContainer.IsValid())
+				{
+					// @todo : I should reverse retarget
+					//ComponentTransform = NodeMappingContainer->GetSourceToTargetTransform(RigHierarchyItemNameMapping[Index]) * ComponentTransform;
+				}
+
+				ControlRig->SetGlobalTransform(RigHierarchyItemNameMapping[Index], ComponentTransform);
 			}
-
-			ControlRig->SetGlobalTransform(RigHierarchyItemNameMapping[Index], ComponentTransform);
 		}
 	}
 }
 
 void FAnimNode_ControlRigBase::UpdateOutput(const UControlRig* ControlRig, FPoseContext& InOutput)
 {
-	// copy output of the rig
-	const FBoneContainer& RequiredBones = InOutput.Pose.GetBoneContainer();
+	// @todo: find a better option to do this
+	// but for now, we use this flag to skip the execution and updating output
+	// this is used to skip when refpose is chosen to view
+	if (ControlRig->bExecutionOn)
+	{
+		// copy output of the rig
+		const FBoneContainer& RequiredBones = InOutput.Pose.GetBoneContainer();
 
-	// get component pose from control rig
-	FCSPose<FCompactPose> MeshPoses;
-	MeshPoses.InitPose(InOutput.Pose);
+		// get component pose from control rig
+		FCSPose<FCompactPose> MeshPoses;
+		MeshPoses.InitPose(InOutput.Pose);
 
-	const int32 NumNodes = RigHierarchyItemNameMapping.Num();
-	for (int32 Index = 0; Index < NumNodes; ++Index)
-	{
-		if (RigHierarchyItemNameMapping[Index] != NAME_None)
+		const int32 NumNodes = RigHierarchyItemNameMapping.Num();
+		for (int32 Index = 0; Index < NumNodes; ++Index)
 		{
-			FCompactPoseBoneIndex CompactPoseIndex(Index);
-			FTransform ComponentTransform = ControlRig->GetGlobalTransform(RigHierarchyItemNameMapping[Index]);
-			if (NodeMappingContainer.IsValid())
+			if (RigHierarchyItemNameMapping[Index] != NAME_None)
 			{
-				ComponentTransform = NodeMappingContainer->GetSourceToTargetTransform(RigHierarchyItemNameMapping[Index]) * ComponentTransform;
+				FCompactPoseBoneIndex CompactPoseIndex(Index);
+				FTransform ComponentTransform = ControlRig->GetGlobalTransform(RigHierarchyItemNameMapping[Index]);
+				if (NodeMappingContainer.IsValid())
+				{
+					ComponentTransform = NodeMappingContainer->GetSourceToTargetTransform(RigHierarchyItemNameMapping[Index]) * ComponentTransform;
+				}
+
+				MeshPoses.SetComponentSpaceTransform(CompactPoseIndex, ComponentTransform);
 			}
-
-			MeshPoses.SetComponentSpaceTransform(CompactPoseIndex, ComponentTransform);
 		}
-	}
 
-	// now convert to what you'd like
-	// initialize with refpose
- 	InOutput.ResetToRefPose();
-	for (int32 Index = 0; Index < NumNodes; ++Index)
-	{
-		if (RigHierarchyItemNameMapping[Index] != NAME_None)
+		// now convert to what you'd like
+		// initialize with refpose
+	 	InOutput.ResetToRefPose();
+		for (int32 Index = 0; Index < NumNodes; ++Index)
 		{
-			FCompactPoseBoneIndex CompactPoseIndex(Index);
-			InOutput.Pose[CompactPoseIndex] = MeshPoses.GetLocalSpaceTransform(CompactPoseIndex);
+			if (RigHierarchyItemNameMapping[Index] != NAME_None)
+			{
+				FCompactPoseBoneIndex CompactPoseIndex(Index);
+				InOutput.Pose[CompactPoseIndex] = MeshPoses.GetLocalSpaceTransform(CompactPoseIndex);
+			}
 		}
 	}
 }
