Merging using Dev-Anim_Release-4.20
Fix build error

diff --git a/Engine/Plugins/Experimental/ControlRig/Source/ControlRig/Private/AnimNode_ControlRigBase.cpp b/Engine/Plugins/Experimental/ControlRig/Source/ControlRig/Private/AnimNode_ControlRigBase.cpp
--- a/Engine/Plugins/Experimental/ControlRig/Source/ControlRig/Private/AnimNode_ControlRigBase.cpp
+++ b/Engine/Plugins/Experimental/ControlRig/Source/ControlRig/Private/AnimNode_ControlRigBase.cpp
@@ -62,7 +62,9 @@ void FAnimNode_ControlRigBase::UpdateInput(UControlRig* ControlRig, const FPoseC
 	// @todo: find a better option to do this
 	// but for now, we use this flag to skip the execution and updating input
 	// this is used to skip when refpose is chosen to view
+#if WITH_EDITORONLY_DATA
 	if (ControlRig->bExecutionOn)
+#endif // WITH_EDITORONLY_DATA
 	{
 		const FBoneContainer& RequiredBones = InOutput.Pose.GetBoneContainer();
 
@@ -96,7 +98,9 @@ void FAnimNode_ControlRigBase::UpdateOutput(const UControlRig* ControlRig, FPose
 	// @todo: find a better option to do this
 	// but for now, we use this flag to skip the execution and updating output
 	// this is used to skip when refpose is chosen to view
+#if WITH_EDITORONLY_DATA
 	if (ControlRig->bExecutionOn)
+#endif // WITH_EDITORONLY_DATA
 	{
 		// copy output of the rig
 		const FBoneContainer& RequiredBones = InOutput.Pose.GetBoneContainer();
