Fix issue with displaying class that is in transient due to mid compilation

diff --git a/Engine/Plugins/Experimental/ControlRig/Source/ControlRigEditor/Private/ControlRigBlueprintFactory.cpp b/Engine/Plugins/Experimental/ControlRig/Source/ControlRigEditor/Private/ControlRigBlueprintFactory.cpp
--- a/Engine/Plugins/Experimental/ControlRig/Source/ControlRigEditor/Private/ControlRigBlueprintFactory.cpp
+++ b/Engine/Plugins/Experimental/ControlRig/Source/ControlRigEditor/Private/ControlRigBlueprintFactory.cpp
@@ -125,7 +125,17 @@ private:
 		virtual bool IsClassAllowed(const FClassViewerInitializationOptions& InInitOptions, const UClass* InClass, TSharedRef< FClassViewerFilterFuncs > InFilterFuncs ) override
 		{
 			// If it appears on the allowed child-of classes list (or there is nothing on that list)
-			return InFilterFuncs->IfInChildOfClassesSet( AllowedChildrenOfClasses, InClass) != EFilterReturn::Failed;
+			if (InClass)
+			{
+				if (InFilterFuncs->IfInChildOfClassesSet(AllowedChildrenOfClasses, InClass) == EFilterReturn::Failed)
+				{
+					return false;
+				}
+
+				return (!InClass->HasAnyClassFlags(CLASS_Deprecated) && InClass->GetOutermost() != GetTransientPackage());
+			}
+
+			return false;
 		}
 
 		virtual bool IsUnloadedClassAllowed(const FClassViewerInitializationOptions& InInitOptions, const TSharedRef< const IUnloadedBlueprintData > InUnloadedClassData, TSharedRef< FClassViewerFilterFuncs > InFilterFuncs) override
