Prevented Blueprint notifies from destroying the animation preview actor

We now use a custom actor class in the preview scene.
Also prevented snim instnace class from being altered from Blueprint.

diff --git a/Engine/Config/BaseEngine.ini b/Engine/Config/BaseEngine.ini
--- a/Engine/Config/BaseEngine.ini
+++ b/Engine/Config/BaseEngine.ini
@@ -577,6 +577,8 @@ bUseStreamingPause=false
 +StructRedirects=(OldName="QualifiedFrameTime",NewName="/Script/CoreUObject.QualifiedFrameTime")
 +StructRedirects=(OldName="Timecode",NewName="/Script/CoreUObject.Timecode")
 
++FunctionRedirects=(OldName="SkeletalMeshComponent.SetAnimInstanceClass",NewName="SkeletalMeshComponent.K2_SetAnimInstanceClass")
+
 [CoreUObject.Metadata]
 ; Note: UnrealHeaderTool should be rerun after making changes to MetadataRedirects to catch any keys specified in class headers
 +MetadataRedirects=(OldKey="K2Protected", NewKey="BlueprintProtected")
diff --git a/Engine/Plugins/Experimental/ControlRig/Source/ControlRigEditor/Private/Editor/ControlRigEditor.cpp b/Engine/Plugins/Experimental/ControlRig/Source/ControlRigEditor/Private/Editor/ControlRigEditor.cpp
--- a/Engine/Plugins/Experimental/ControlRig/Source/ControlRigEditor/Private/Editor/ControlRigEditor.cpp
+++ b/Engine/Plugins/Experimental/ControlRig/Source/ControlRigEditor/Private/Editor/ControlRigEditor.cpp
@@ -30,6 +30,7 @@
 #include "ControlRigBlueprintUtils.h"
 #include "IPersonaViewport.h"
 #include "EditorViewportClient.h"
+#include "AnimationEditorPreviewActor.h"
 
 #define LOCTEXT_NAMESPACE "ControlRigEditor"
 
@@ -507,7 +508,7 @@ void FControlRigEditor::HandleViewportCreated(const TSharedRef<class IPersonaVie
 
 void FControlRigEditor::HandlePreviewSceneCreated(const TSharedRef<IPersonaPreviewScene>& InPersonaPreviewScene)
 {
-	AActor* Actor = InPersonaPreviewScene->GetWorld()->SpawnActor<AActor>(AActor::StaticClass(), FTransform::Identity);
+	AAnimationEditorPreviewActor* Actor = InPersonaPreviewScene->GetWorld()->SpawnActor<AAnimationEditorPreviewActor>(AAnimationEditorPreviewActor::StaticClass(), FTransform::Identity);
 	InPersonaPreviewScene->SetActor(Actor);
 
 	// Create the preview component
diff --git a/Engine/Source/Editor/Persona/Private/AnimationEditorPreviewActor.cpp b/Engine/Source/Editor/Persona/Private/AnimationEditorPreviewActor.cpp
new file mode 100644
--- /dev/null
+++ b/Engine/Source/Editor/Persona/Private/AnimationEditorPreviewActor.cpp
@@ -0,0 +1,10 @@
+// Copyright 1998-2018 Epic Games, Inc. All Rights Reserved.
+
+#include "AnimationEditorPreviewActor.h"
+#include "AnimationRuntime.h"
+
+void AAnimationEditorPreviewActor::K2_DestroyActor()
+{
+	// Override this to do nothing and warn the user
+	UE_LOG(LogAnimation, Warning, TEXT("Attempting to destroy an animation preview actor, skipping."));
+}
\ No newline at end of file
diff --git a/Engine/Source/Editor/Persona/Private/PersonaToolkit.cpp b/Engine/Source/Editor/Persona/Private/PersonaToolkit.cpp
--- a/Engine/Source/Editor/Persona/Private/PersonaToolkit.cpp
+++ b/Engine/Source/Editor/Persona/Private/PersonaToolkit.cpp
@@ -13,6 +13,7 @@
 #include "PersonaModule.h"
 #include "PersonaAssetFamily.h"
 #include "Interfaces/Interface_PreviewMeshProvider.h"
+#include "AnimationEditorPreviewActor.h"
 
 FPersonaToolkit::FPersonaToolkit()
 	: Skeleton(nullptr)
@@ -113,7 +114,7 @@ void FPersonaToolkit::CreatePreviewScene(const FPersonaToolkitArgs& PersonaToolk
 		else
 		{
 			// setup default scene
-			AActor* Actor = PreviewScene->GetWorld()->SpawnActor<AActor>(AActor::StaticClass(), FTransform::Identity);
+			AAnimationEditorPreviewActor* Actor = PreviewScene->GetWorld()->SpawnActor<AAnimationEditorPreviewActor>(AAnimationEditorPreviewActor::StaticClass(), FTransform::Identity);
 			PreviewScene->SetActor(Actor);
 
 			// Create the preview component
diff --git a/Engine/Source/Editor/Persona/Public/AnimationEditorPreviewActor.h b/Engine/Source/Editor/Persona/Public/AnimationEditorPreviewActor.h
new file mode 100644
--- /dev/null
+++ b/Engine/Source/Editor/Persona/Public/AnimationEditorPreviewActor.h
@@ -0,0 +1,16 @@
+// Copyright 1998-2018 Epic Games, Inc. All Rights Reserved.
+
+#pragma once
+
+#include "GameFramework/Actor.h"
+#include "AnimationEditorPreviewActor.generated.h"
+
+UCLASS()
+class PERSONA_API AAnimationEditorPreviewActor : public AActor
+{
+	GENERATED_BODY()
+
+public:
+	/** AActor interface */
+	virtual void K2_DestroyActor() override;
+};
diff --git a/Engine/Source/Editor/PhysicsAssetEditor/Private/PhysicsAssetEditor.cpp b/Engine/Source/Editor/PhysicsAssetEditor/Private/PhysicsAssetEditor.cpp
--- a/Engine/Source/Editor/PhysicsAssetEditor/Private/PhysicsAssetEditor.cpp
+++ b/Engine/Source/Editor/PhysicsAssetEditor/Private/PhysicsAssetEditor.cpp
@@ -78,6 +78,7 @@
 #include "IPinnedCommandList.h"
 #include "Widgets/Input/SSpinBox.h"
 #include "Widgets/Input/SNumericEntryBox.h"
+#include "AnimationEditorPreviewActor.h"
 
 const FName PhysicsAssetEditorModes::PhysicsAssetEditorMode("PhysicsAssetEditorMode");
 
@@ -2761,7 +2762,7 @@ void FPhysicsAssetEditor::HandlePreviewSceneCreated(const TSharedRef<IPersonaPre
 
 	SharedData->Initialize(InPersonaPreviewScene);
 
-	AActor* Actor = InPersonaPreviewScene->GetWorld()->SpawnActor<AActor>(AActor::StaticClass(), FTransform::Identity);
+	AAnimationEditorPreviewActor* Actor = InPersonaPreviewScene->GetWorld()->SpawnActor<AAnimationEditorPreviewActor>(AAnimationEditorPreviewActor::StaticClass(), FTransform::Identity);
 	InPersonaPreviewScene->SetActor(Actor);
 
 	// Create the preview component
diff --git a/Engine/Source/Editor/UnrealEd/Classes/Animation/DebugSkelMeshComponent.h b/Engine/Source/Editor/UnrealEd/Classes/Animation/DebugSkelMeshComponent.h
--- a/Engine/Source/Editor/UnrealEd/Classes/Animation/DebugSkelMeshComponent.h
+++ b/Engine/Source/Editor/UnrealEd/Classes/Animation/DebugSkelMeshComponent.h
@@ -302,6 +302,7 @@ class UNREALED_API UDebugSkelMeshComponent : public USkeletalMeshComponent
 	//~ Begin SkeletalMeshComponent Interface
 	virtual void InitAnim(bool bForceReinit) override;
 	virtual bool IsWindEnabled() const override { return true; }
+	virtual void K2_SetAnimInstanceClass(class UClass* NewClass) override;
 	//~ End SkeletalMeshComponent Interface
 	// Preview.
 	// @todo document
diff --git a/Engine/Source/Editor/UnrealEd/Private/Animation/DebugSkelMeshComponent.cpp b/Engine/Source/Editor/UnrealEd/Private/Animation/DebugSkelMeshComponent.cpp
--- a/Engine/Source/Editor/UnrealEd/Private/Animation/DebugSkelMeshComponent.cpp
+++ b/Engine/Source/Editor/UnrealEd/Private/Animation/DebugSkelMeshComponent.cpp
@@ -295,6 +295,12 @@ void UDebugSkelMeshComponent::InitAnim(bool bForceReinit)
 	}
 }
 
+void UDebugSkelMeshComponent::K2_SetAnimInstanceClass(class UClass* NewClass)
+{
+	// Override this to do nothing and warn the user
+	UE_LOG(LogAnimation, Warning, TEXT("Attempting to destroy an animation preview actor, skipping."));
+}
+
 void UDebugSkelMeshComponent::EnablePreview(bool bEnable, UAnimationAsset* PreviewAsset)
 {
 	if (PreviewInstance)
diff --git a/Engine/Source/Runtime/Engine/Classes/Components/SkeletalMeshComponent.h b/Engine/Source/Runtime/Engine/Classes/Components/SkeletalMeshComponent.h
--- a/Engine/Source/Runtime/Engine/Classes/Components/SkeletalMeshComponent.h
+++ b/Engine/Source/Runtime/Engine/Classes/Components/SkeletalMeshComponent.h
@@ -651,7 +651,11 @@ public:
 	/** Array of physical interactions for the frame. This is a temporary solution for a more permanent force system and should not be used directly*/
 	TArray<FPendingRadialForces> PendingRadialForces;
 
-	UFUNCTION(BlueprintCallable, Category="Components|SkeletalMesh", meta=(Keywords = "AnimBlueprint"))
+	/** Set the anim instance class. Clears and re-initializes the anim instance with the new class and sets animation mode to 'AnimationBlueprint' */
+	UFUNCTION(BlueprintCallable, Category="Components|SkeletalMesh", meta=(Keywords = "AnimBlueprint", DisplayName = "Set Anim Instance Class"))
+	virtual void K2_SetAnimInstanceClass(class UClass* NewClass);
+
+	/** Set the anim instance class. Clears and re-initializes the anim instance with the new class and sets animation mode to 'AnimationBlueprint' */
 	void SetAnimInstanceClass(class UClass* NewClass);
 
 	/** 
diff --git a/Engine/Source/Runtime/Engine/Private/Components/SkeletalMeshComponent.cpp b/Engine/Source/Runtime/Engine/Private/Components/SkeletalMeshComponent.cpp
--- a/Engine/Source/Runtime/Engine/Private/Components/SkeletalMeshComponent.cpp
+++ b/Engine/Source/Runtime/Engine/Private/Components/SkeletalMeshComponent.cpp
@@ -2550,6 +2550,11 @@ void USkeletalMeshComponent::SetDisablePostProcessBlueprint(bool bInDisablePostP
 	bDisablePostProcessBlueprint = bInDisablePostProcess;
 }
 
+void USkeletalMeshComponent::K2_SetAnimInstanceClass(class UClass* NewClass)
+{
+	SetAnimInstanceClass(NewClass);
+}
+
 void USkeletalMeshComponent::SetAnimInstanceClass(class UClass* NewClass)
 {
 	if (NewClass != nullptr)
