Control Rig: Set default display (30fps) and tick resolution (24,000)

diff --git a/Engine/Plugins/Experimental/ControlRig/Source/ControlRig/Private/Sequencer/ControlRigSequence.cpp b/Engine/Plugins/Experimental/ControlRig/Source/ControlRig/Private/Sequencer/ControlRigSequence.cpp
--- a/Engine/Plugins/Experimental/ControlRig/Source/ControlRig/Private/Sequencer/ControlRigSequence.cpp
+++ b/Engine/Plugins/Experimental/ControlRig/Source/ControlRig/Private/Sequencer/ControlRigSequence.cpp
@@ -8,6 +8,24 @@
 
 #define LOCTEXT_NAMESPACE "ControlRigSequence"
 
+static TAutoConsoleVariable<int32> CVarControlRigDefaultEvaluationType(
+	TEXT("ControlRigSequence.DefaultEvaluationType"),
+	0,
+	TEXT("0: Playback locked to playback frames\n1: Unlocked playback with sub frame interpolation"),
+	ECVF_Default);
+
+static TAutoConsoleVariable<FString> CVarControlRigDefaultTickResolution(
+	TEXT("ControlRigSequence.DefaultTickResolution"),
+	TEXT("24000fps"),
+	TEXT("Specifies default a tick resolution for newly created control rig sequences. Examples: 30 fps, 120/1 (120 fps), 30000/1001 (29.97), 0.01s (10ms)."),
+	ECVF_Default);
+
+static TAutoConsoleVariable<FString> CVarControlRigDefaultDisplayRate(
+	TEXT("ControlRigSequence.DefaultDisplayRate"),
+	TEXT("30fps"),
+	TEXT("Specifies default a display frame rate for newly created control rig sequences; also defines frame locked frame rate where sequences are set to be frame locked. Examples: 30 fps, 120/1 (120 fps), 30000/1001 (29.97), 0.01s (10ms)."),
+	ECVF_Default); 
+
 UControlRigSequence::UControlRigSequence(const FObjectInitializer& ObjectInitializer)
 	: Super(ObjectInitializer)
 	, LastExportedFrameRate(30.0f)
@@ -19,6 +37,18 @@ void UControlRigSequence::Initialize()
 {
 	MovieScene = NewObject<UMovieScene>(this, NAME_None, RF_Transactional);
 	MovieScene->SetFlags(RF_Transactional);
+
+	const bool bFrameLocked = CVarControlRigDefaultEvaluationType.GetValueOnGameThread() != 0;
+
+	MovieScene->SetEvaluationType(bFrameLocked ? EMovieSceneEvaluationType::FrameLocked : EMovieSceneEvaluationType::WithSubFrames);
+
+	FFrameRate TickResolution(60000, 1);
+	TryParseString(TickResolution, *CVarControlRigDefaultTickResolution.GetValueOnGameThread());
+	MovieScene->SetTickResolutionDirectly(TickResolution);
+
+	FFrameRate DisplayRate(30, 1);
+	TryParseString(DisplayRate, *CVarControlRigDefaultDisplayRate.GetValueOnGameThread());
+	MovieScene->SetDisplayRate(DisplayRate);
 }
 
 void UControlRigSequence::BindPossessableObject(const FGuid& ObjectId, UObject& PossessedObject, UObject* Context)
