Interpreter Optimization 1 of N: Removing seemingly dangerous overrides of GetPersistentUberGraphFrame

diff --git a/Engine/Plugins/Experimental/ControlRig/Source/ControlRig/Private/ControlRigBlueprintGeneratedClass.cpp b/Engine/Plugins/Experimental/ControlRig/Source/ControlRig/Private/ControlRigBlueprintGeneratedClass.cpp
--- a/Engine/Plugins/Experimental/ControlRig/Source/ControlRig/Private/ControlRigBlueprintGeneratedClass.cpp
+++ b/Engine/Plugins/Experimental/ControlRig/Source/ControlRig/Private/ControlRigBlueprintGeneratedClass.cpp
@@ -37,15 +37,3 @@ void UControlRigBlueprintGeneratedClass::PurgeClass(bool bRecompilingOnLoad)
 {
 	Super::PurgeClass(bRecompilingOnLoad);
 }
-
-uint8* UControlRigBlueprintGeneratedClass::GetPersistentUberGraphFrame(UObject* Obj, UFunction* FuncToCheck) const
-{
-	if(!IsInGameThread())
-	{
-		// we cant use the persistent frame if we are executing in parallel (as we could potentially thunk to BP)
-		return nullptr;
-	}
-
-	return Super::GetPersistentUberGraphFrame(Obj, FuncToCheck);
-}
-
diff --git a/Engine/Plugins/Experimental/ControlRig/Source/ControlRig/Public/ControlRigBlueprintGeneratedClass.h b/Engine/Plugins/Experimental/ControlRig/Source/ControlRig/Public/ControlRigBlueprintGeneratedClass.h
--- a/Engine/Plugins/Experimental/ControlRig/Source/ControlRig/Public/ControlRigBlueprintGeneratedClass.h
+++ b/Engine/Plugins/Experimental/ControlRig/Source/ControlRig/Public/ControlRigBlueprintGeneratedClass.h
@@ -20,7 +20,6 @@ public:
 
 	// UClass interface
 	virtual void PurgeClass(bool bRecompilingOnLoad) override;
-	virtual uint8* GetPersistentUberGraphFrame(UObject* Obj, UFunction* FuncToCheck) const override;
 
 public:
 #if WITH_EDITORONLY_DATA
diff --git a/Engine/Source/Runtime/Engine/Classes/Animation/AnimBlueprintGeneratedClass.h b/Engine/Source/Runtime/Engine/Classes/Animation/AnimBlueprintGeneratedClass.h
--- a/Engine/Source/Runtime/Engine/Classes/Animation/AnimBlueprintGeneratedClass.h
+++ b/Engine/Source/Runtime/Engine/Classes/Animation/AnimBlueprintGeneratedClass.h
@@ -340,7 +340,6 @@ public:
 
 	// UClass interface
 	virtual void PurgeClass(bool bRecompilingOnLoad) override;
-	virtual uint8* GetPersistentUberGraphFrame(UObject* Obj, UFunction* FuncToCheck) const override;
 	virtual void PostLoadDefaultObject(UObject* Object) override;
 	// End of UClass interface
 };
diff --git a/Engine/Source/Runtime/Engine/Private/Animation/AnimBlueprintGeneratedClass.cpp b/Engine/Source/Runtime/Engine/Private/Animation/AnimBlueprintGeneratedClass.cpp
--- a/Engine/Source/Runtime/Engine/Private/Animation/AnimBlueprintGeneratedClass.cpp
+++ b/Engine/Source/Runtime/Engine/Private/Animation/AnimBlueprintGeneratedClass.cpp
@@ -312,17 +312,6 @@ void UAnimBlueprintGeneratedClass::PurgeClass(bool bRecompilingOnLoad)
 	BakedStateMachines.Empty();
 }
 
-uint8* UAnimBlueprintGeneratedClass::GetPersistentUberGraphFrame(UObject* Obj, UFunction* FuncToCheck) const
-{
-	if(!IsInGameThread())
-	{
-		// we cant use the persistent frame if we are executing in parallel (as we could potentially thunk to BP)
-		return nullptr;
-	}
-
-	return Super::GetPersistentUberGraphFrame(Obj, FuncToCheck);
-}
-
 void UAnimBlueprintGeneratedClass::PostLoadDefaultObject(UObject* Object)
 {
 	Super::PostLoadDefaultObject(Object);
