Only display native control rig class for parent class picker

diff --git a/Engine/Plugins/Experimental/ControlRig/Source/ControlRigEditor/Private/ControlRigBlueprintFactory.cpp b/Engine/Plugins/Experimental/ControlRig/Source/ControlRigEditor/Private/ControlRigBlueprintFactory.cpp
--- a/Engine/Plugins/Experimental/ControlRig/Source/ControlRigEditor/Private/ControlRigBlueprintFactory.cpp
+++ b/Engine/Plugins/Experimental/ControlRig/Source/ControlRigEditor/Private/ControlRigBlueprintFactory.cpp
@@ -132,7 +132,9 @@ private:
 					return false;
 				}
 
-				return (!InClass->HasAnyClassFlags(CLASS_Deprecated) && InClass->GetOutermost() != GetTransientPackage());
+				// in the future we might allow it to parent to BP classes, but right now, it won't work well because of multi rig graph
+				// for now we disable it and only allow native class. 
+				return (!InClass->HasAnyClassFlags(CLASS_Deprecated) && InClass->HasAnyClassFlags(CLASS_Native) && InClass->GetOutermost() != GetTransientPackage());
 			}
 
 			return false;
@@ -141,7 +143,9 @@ private:
 		virtual bool IsUnloadedClassAllowed(const FClassViewerInitializationOptions& InInitOptions, const TSharedRef< const IUnloadedBlueprintData > InUnloadedClassData, TSharedRef< FClassViewerFilterFuncs > InFilterFuncs) override
 		{
 			// If it appears on the allowed child-of classes list (or there is nothing on that list)
-			return InFilterFuncs->IfInChildOfClassesSet( AllowedChildrenOfClasses, InUnloadedClassData) != EFilterReturn::Failed;
+			// in the future we might allow it to parent to BP classes, but right now, it won't work well because of multi rig graph
+			// for now we disable it and only allow native class. 
+			return false; // InFilterFuncs->IfInChildOfClassesSet(AllowedChildrenOfClasses, InUnloadedClassData) != EFilterReturn::Failed;
 		}
 	};
 
