- added an option to refresh transform as that can solve quick issues where ref pose goes wrong
- fixed search

diff --git a/Engine/Plugins/Experimental/ControlRig/Source/ControlRigEditor/Private/Editor/SRigHierarchy.cpp b/Engine/Plugins/Experimental/ControlRig/Source/ControlRigEditor/Private/Editor/SRigHierarchy.cpp
--- a/Engine/Plugins/Experimental/ControlRig/Source/ControlRigEditor/Private/Editor/SRigHierarchy.cpp
+++ b/Engine/Plugins/Experimental/ControlRig/Source/ControlRigEditor/Private/Editor/SRigHierarchy.cpp
@@ -171,7 +171,7 @@ void SRigHierarchy::Construct(const FArguments& InArgs, TSharedRef<FControlRigEd
 					.Padding(3.0f, 1.0f)
 					[
 						SAssignNew(FilterBox, SSearchBox)
-			//			.OnTextChanged(this, &SRigHierarchy::OnFilterTextChanged)
+						.OnTextChanged(this, &SRigHierarchy::OnFilterTextChanged)
 					]
 				]
 			]
@@ -218,6 +218,13 @@ void SRigHierarchy::BindCommands()
 		FCanExecuteAction::CreateSP(this, &SRigHierarchy::CanRenameItem));
 }
 
+void SRigHierarchy::OnFilterTextChanged(const FText& SearchText)
+{
+	FilterText = SearchText;
+
+	RefreshTreeView();
+}
+
 void SRigHierarchy::RefreshTreeView()
 {
 	RootJoints.Reset();
@@ -229,31 +236,46 @@ void SRigHierarchy::RefreshTreeView()
 
 		TMap<FName, TSharedPtr<FRigTreeJoint>> SearchTable;
 
+		FString FilteredString = FilterText.ToString();
+		const bool bSearchOff = FilteredString.IsEmpty();
 		for (int32 JointIndex = 0; JointIndex < Hierarchy.Joints.Num(); ++JointIndex)
 		{
 			FRigJoint& Joint = Hierarchy.Joints[JointIndex];
 
 			// create new item
-			TSharedPtr<FRigTreeJoint> NewItem = MakeShared<FRigTreeJoint>(Hierarchy.Joints[JointIndex].Name, SharedThis(this));
-			SearchTable.Add(Joint.Name, NewItem);
-
-			if (Joint.ParentName == NAME_None)
+			if (bSearchOff)
 			{
-				RootJoints.Add(NewItem);
+				TSharedPtr<FRigTreeJoint> NewItem = MakeShared<FRigTreeJoint>(Hierarchy.Joints[JointIndex].Name, SharedThis(this));
+				SearchTable.Add(Joint.Name, NewItem);
+
+				if (Joint.ParentName == NAME_None)
+				{
+					RootJoints.Add(NewItem);
+				}
+				else
+				{
+					// you have to find one
+					TSharedPtr<FRigTreeJoint>* FoundItem = SearchTable.Find(Joint.ParentName);
+					check(FoundItem);
+					// add to children list
+					FoundItem->Get()->Children.Add(NewItem);
+				}
 			}
-			else
+			else if (Joint.Name.ToString().Contains(FilteredString))
 			{
-				// you have to find one
-				TSharedPtr<FRigTreeJoint>* FoundItem = SearchTable.Find(Joint.ParentName);
-				check(FoundItem);
-				// add to children list
-				FoundItem->Get()->Children.Add(NewItem);
+				// if contains, just list out everything to root
+				TSharedPtr<FRigTreeJoint> NewItem = MakeShared<FRigTreeJoint>(Hierarchy.Joints[JointIndex].Name, SharedThis(this));
+				// during search, everything is on root
+				RootJoints.Add(NewItem);
 			}
 		}
 
-		for (int32 RootIndex = 0; RootIndex < RootJoints.Num(); ++RootIndex)
+		if (bSearchOff)
 		{
-			SetExpansionRecursive(RootJoints[RootIndex]);
+			for (int32 RootIndex = 0; RootIndex < RootJoints.Num(); ++RootIndex)
+			{
+				SetExpansionRecursive(RootJoints[RootIndex]);
+			}
 		}
 	}
 
@@ -358,9 +380,17 @@ TSharedPtr< SWidget > SRigHierarchy::CreateContextMenu()
 		MenuBuilder.AddMenuSeparator();
 		MenuBuilder.AddSubMenu(
 			LOCTEXT("ImportSubMenu", "Import"),
-			LOCTEXT("ImportSubMenu_ToolTip", "Insert current pose to selected PoseAsset"),
+			LOCTEXT("ImportSubMenu_ToolTip", "Import hierarchy to the current rig. This only imports non-existing node. For example, if there is hand_r, it won't import hand_r. \
+				If you want to reimport whole new hiearchy, delete all nodes, and use import hierarchy."),
 			FNewMenuDelegate::CreateSP(this, &SRigHierarchy::CreateImportMenu)
 		);
+
+		MenuBuilder.AddMenuSeparator();
+		MenuBuilder.AddSubMenu(
+			LOCTEXT("RefreshSubMenu", "Refresh"),
+			LOCTEXT("RefreshSubMenu_ToolTip", "Refresh the existing initial transform from the selected mesh. This only updates if the node is found."),
+			FNewMenuDelegate::CreateSP(this, &SRigHierarchy::CreateRefreshMenu)
+		);
 	
 		MenuBuilder.EndSection();
 	}
@@ -368,6 +398,62 @@ TSharedPtr< SWidget > SRigHierarchy::CreateContextMenu()
 	return MenuBuilder.MakeWidget();
 }
 
+void SRigHierarchy::CreateRefreshMenu(FMenuBuilder& MenuBuilder)
+{
+	MenuBuilder.AddWidget(
+		SNew(SVerticalBox)
+
+		+ SVerticalBox::Slot()
+		.AutoHeight()
+		.Padding(3)
+		[
+			SNew(STextBlock)
+			.Font(FEditorStyle::GetFontStyle("ControlRig.Hierarchy.Menu"))
+			.Text(LOCTEXT("RefreshMesh_Title", "Select Mesh"))
+			.ToolTipText(LOCTEXT("RefreshMesh_Tooltip", "Select Mesh to refresh transform from... It will refresh init transform from selected mesh. This doesn't change hierarchy. \
+				If you want to reimport hierarchy, please delete all nodes, and use import hierarchy."))
+		]
+		+ SVerticalBox::Slot()
+		.AutoHeight()
+		.Padding(3)
+		[
+			SNew(SObjectPropertyEntryBox)
+			.AllowedClass(USkeletalMesh::StaticClass())
+			.OnObjectChanged(this, &SRigHierarchy::RefreshHierarchy)
+		]
+		,
+		FText()
+	);
+}
+
+void SRigHierarchy::RefreshHierarchy(const FAssetData& InAssetData)
+{
+	FRigHierarchy* Hier = GetHierarchy();
+	USkeletalMesh* Mesh = Cast<USkeletalMesh>(InAssetData.GetAsset());
+	if (Mesh && Hier)
+	{
+		FScopedTransaction Transaction(LOCTEXT("HierarchyRefresh", "Refresh Transform"));
+		ControlRigBlueprint->Modify();
+
+		const FReferenceSkeleton& RefSkeleton = Mesh->RefSkeleton;
+		const TArray<FMeshBoneInfo>& BoneInfos = RefSkeleton.GetRawRefBoneInfo();
+		const TArray<FTransform>& BonePoses = RefSkeleton.GetRawRefBonePose();
+
+		for (int32 BoneIndex = 0; BoneIndex < RefSkeleton.GetNum(); ++BoneIndex)
+		{
+			// only add if you don't have it. This may change in the future
+			int32 RigIndex = Hier->GetIndex(BoneInfos[BoneIndex].Name);
+			if (RigIndex != INDEX_NONE)
+			{
+				// @todo: add optimized version without sorting, but if no sort, we should make sure not to use find index function
+				Hier->SetInitialTransform(RigIndex, FAnimationRuntime::GetComponentSpaceTransform(RefSkeleton, BonePoses, BoneIndex));
+			}
+		}
+
+		RefreshTreeView();
+		FSlateApplication::Get().DismissAllMenus();
+	}
+}
 void SRigHierarchy::CreateImportMenu(FMenuBuilder& MenuBuilder)
 {
 	MenuBuilder.AddWidget(
@@ -380,6 +466,7 @@ void SRigHierarchy::CreateImportMenu(FMenuBuilder& MenuBuilder)
 			SNew(STextBlock)
 			.Font(FEditorStyle::GetFontStyle("ControlRig.Hierarchy.Menu"))
 			.Text(LOCTEXT("ImportMesh_Title", "Select Mesh"))
+			.ToolTipText(LOCTEXT("ImportMesh_Tooltip", "Select Mesh to import hierarchy from... It will only import if the node doens't exists in the current hierarchy."))
 		]
 		+ SVerticalBox::Slot()
 		.AutoHeight()
@@ -400,6 +487,9 @@ void SRigHierarchy::ImportHierarchy(const FAssetData& InAssetData)
 	USkeletalMesh* Mesh = Cast<USkeletalMesh> (InAssetData.GetAsset());
 	if (Mesh && Hier)
 	{
+		FScopedTransaction Transaction(LOCTEXT("HierarchyImport", "Import Hierarchy"));
+		ControlRigBlueprint->Modify();
+
 		const FReferenceSkeleton& RefSkeleton = Mesh->RefSkeleton;
 		const TArray<FMeshBoneInfo>& BoneInfos = RefSkeleton.GetRawRefBoneInfo();
 		const TArray<FTransform>& BonePoses = RefSkeleton.GetRawRefBonePose();
diff --git a/Engine/Plugins/Experimental/ControlRig/Source/ControlRigEditor/Private/Editor/SRigHierarchy.h b/Engine/Plugins/Experimental/ControlRig/Source/ControlRigEditor/Private/Editor/SRigHierarchy.h
--- a/Engine/Plugins/Experimental/ControlRig/Source/ControlRigEditor/Private/Editor/SRigHierarchy.h
+++ b/Engine/Plugins/Experimental/ControlRig/Source/ControlRigEditor/Private/Editor/SRigHierarchy.h
@@ -124,6 +124,9 @@ private:
 
 	/** Search box widget */
 	TSharedPtr<SSearchBox> FilterBox;
+	FText FilterText;
+
+	void OnFilterTextChanged(const FText& SearchText);
 
 	/** Tree view widget */
 	TSharedPtr<STreeView<TSharedPtr<FRigTreeJoint>>> TreeView;
@@ -147,6 +150,8 @@ private:
 
 	void ImportHierarchy(const FAssetData& InAssetData);
 	void CreateImportMenu(FMenuBuilder& MenuBuilder);
+	void CreateRefreshMenu(FMenuBuilder& MenuBuilder);
+	void RefreshHierarchy(const FAssetData& InAssetData);
 
 	FName CreateUniqueName(const FName& InBaseName) const;
 
