Copying //UE4/Dev-Anim to //UE4/Dev-Main (Source: //UE4/Dev-Anim @ 4335583)
============================
  MAJOR FEATURES & CHANGES
============================

Change 4300079 by [REDACTED_PERSON]

	Fixed incorrect sub instance error reporting when multiple sub-instances of the same class are used

Change 4301787 by [REDACTED_PERSON]

	Fix crash when multi select with bone driven controller

	- make sure you can't continue with mismatching types of classes, code becomes very vulnerable

Change 4302310 by [REDACTED_PERSON]

	Fix for anim slot window of small icons

Change 4302385 by [REDACTED_PERSON]

	remove constructor female data from control rig example

Change 4303861 by [REDACTED_PERSON]

	PR #4298: Avoid division by 0 in AnimLinkableElement (Contributed by DSDambuster)

Change 4304915 by [REDACTED_PERSON]

	PR #4842: UE-52767: UAnimComposite doesn't inherit from UAnimSequence (Contributed by [REDACTED_PERSON])

	Fixes issue with stepping not working with AnimComposite
	- it also loops around

Change 4307241 by [REDACTED_PERSON]

	Control rig : invalid FOV issue

Change 4310924 by [REDACTED_PERSON]

	Only display native control rig class for parent class picker

Change 4314164 by Jurre.deBaare

	Fix issue with propagating vertex colors to static meshes in some cases not working correctly
	#fix Using .OverrideVertexColors instead of .OverrideMapBuildData to do pointer comparison

Change 4321197 by [REDACTED_PERSON]

	fix aim offset assets

	- removed aimtest2 as it doens't really do anything

Change 4321673 by [REDACTED_PERSON]

	Autoload Animation Blueprint compiler module in editor so that we can set up logging category (following how normal blueprint compiler works)

Change 4321898 by [REDACTED_PERSON]

	Change Live Link tab to nomad

Change 4322333 by [REDACTED_PERSON]

	Fix additional params being displayed on blutility asset/actor actions dialog

Change 4322338 by [REDACTED_PERSON]

	Fixed mis-reported HLOD build indicator

Change 4322368 by [REDACTED_PERSON]

	fixed build error

Change 4322378 by [REDACTED_PERSON]

	Fix issue where preview asset is not saved and not loaded unless it's loaded manually

Change 4322490 by [REDACTED_PERSON]

	- added an option to refresh transform as that can solve quick issues where ref pose goes wrong
	- fixed search

Change 4322824 by [REDACTED_PERSON]

	Fix possible crash due to retarget asset being cleaned up on ULiveLinkInstance

Change 4325357 by [REDACTED_PERSON]

	Perform anim update even when we have zero bones, but still skip evaluation

	This ensures that notifies still get triggered in all update-relevant tick modes, even on dedicated servers.

Change 4325358 by [REDACTED_PERSON]

	Fix incorrect interpolations when using URO and CopyPoseFromMesh

	This ensures that CopyPoseFromMesh uses the correct target pose to inteprolate to, rather than the already-interpolated output pose.

Change 4325552 by [REDACTED_PERSON]

	Fix doubling of playrate in marker sync code

Change 4331507 by [REDACTED_PERSON]

	Fixed crash when creating a new notify in certain circumstances

	When creating a notify causes a timing order change, then the timing panel wasnt getting refreshed and was accessing the wrong notify to do its drawing.

diff --git a/Engine/Plugins/Animation/LiveLink/Source/LiveLink/Private/LiveLinkInstance.cpp b/Engine/Plugins/Animation/LiveLink/Source/LiveLink/Private/LiveLinkInstance.cpp
--- a/Engine/Plugins/Animation/LiveLink/Source/LiveLink/Private/LiveLinkInstance.cpp
+++ b/Engine/Plugins/Animation/LiveLink/Source/LiveLink/Private/LiveLinkInstance.cpp
@@ -24,10 +24,16 @@ void FLiveLinkInstanceProxy::UpdateAnimationNode(float DeltaSeconds)
 
 	FAnimationUpdateContext UpdateContext(this, DeltaSeconds);
 	PoseNode.Update_AnyThread(UpdateContext);
+	
+	if(ULiveLinkInstance* Instance = Cast<ULiveLinkInstance>(GetAnimInstanceObject()))
+	{
+		Instance->CurrentRetargetAsset = PoseNode.CurrentRetargetAsset; //Cache for GC
+	}
 }
 
 ULiveLinkInstance::ULiveLinkInstance(const FObjectInitializer& Initializer)
 	: Super(Initializer)
+	, CurrentRetargetAsset(nullptr)
 {
 
 }
@@ -35,4 +41,10 @@ ULiveLinkInstance::ULiveLinkInstance(const FObjectInitializer& Initializer)
 FAnimInstanceProxy* ULiveLinkInstance::CreateAnimInstanceProxy()
 {
 	return new FLiveLinkInstanceProxy(this);
+}
+
+void ULiveLinkInstance::DestroyAnimInstanceProxy(FAnimInstanceProxy* InProxy)
+{
+	Super::DestroyAnimInstanceProxy(InProxy);
+	CurrentRetargetAsset = nullptr;
 }
\ No newline at end of file
diff --git a/Engine/Plugins/Animation/LiveLink/Source/LiveLink/Public/LiveLinkInstance.h b/Engine/Plugins/Animation/LiveLink/Source/LiveLink/Public/LiveLinkInstance.h
--- a/Engine/Plugins/Animation/LiveLink/Source/LiveLink/Public/LiveLinkInstance.h
+++ b/Engine/Plugins/Animation/LiveLink/Source/LiveLink/Public/LiveLinkInstance.h
@@ -52,4 +52,11 @@ class LIVELINK_API ULiveLinkInstance : public UAnimInstance
 
 protected:
 	virtual FAnimInstanceProxy* CreateAnimInstanceProxy() override;
+	virtual void DestroyAnimInstanceProxy(FAnimInstanceProxy* InProxy) override;
+
+	// Cache for GC
+	UPROPERTY(transient)
+	ULiveLinkRetargetAsset* CurrentRetargetAsset;
+
+	friend FLiveLinkInstanceProxy;
 };
diff --git a/Engine/Plugins/Animation/LiveLink/Source/LiveLinkEditor/Private/LiveLinkEditorModule.cpp b/Engine/Plugins/Animation/LiveLink/Source/LiveLinkEditor/Private/LiveLinkEditorModule.cpp
--- a/Engine/Plugins/Animation/LiveLink/Source/LiveLinkEditor/Private/LiveLinkEditorModule.cpp
+++ b/Engine/Plugins/Animation/LiveLink/Source/LiveLinkEditor/Private/LiveLinkEditorModule.cpp
@@ -59,9 +59,11 @@ public:
 		static FName LiveLinkStyle(TEXT("LiveLinkStyle"));
 		StyleSet = MakeShareable(new FSlateStyleSet(LiveLinkStyle));
 
+		bHasRegisteredTabSpawners = false;
+
 		if (FModuleManager::Get().IsModuleLoaded(LevelEditorModuleName))
 		{
-			RegisterTab();
+			RegisterTabSpawner();
 		}
 		else
 		{
@@ -97,34 +99,13 @@ public:
 	{
 		if (ReasonForChange == EModuleChangeReason::ModuleLoaded && ModuleName == LevelEditorModuleName)
 		{
-			RegisterTab();
+			RegisterTabSpawner();
 		}
 	}
 
-	void RegisterTab()
-	{
-		FLevelEditorModule& LevelEditorModule = FModuleManager::LoadModuleChecked<FLevelEditorModule>(LevelEditorModuleName);
-
-		TSharedPtr<FSlateStyleSet> StyleSetPtr = StyleSet;
-
-		//Register our UI
-		LevelEditorTabManagerChangedHandle = LevelEditorModule.OnTabManagerChanged().AddLambda([StyleSetPtr]()
-		{
-			FLevelEditorModule& LocalLevelEditorModule = FModuleManager::GetModuleChecked<FLevelEditorModule>(LevelEditorModuleName);
-			LocalLevelEditorModule.GetLevelEditorTabManager()->RegisterTabSpawner(LiveLinkClientTabName, FOnSpawnTab::CreateStatic(&FLiveLinkEditorModule::SpawnLiveLinkTab, StyleSetPtr))
-				.SetGroup(WorkspaceMenu::GetMenuStructure().GetLevelEditorCategory())
-				.SetDisplayName(LOCTEXT("LiveLinkTabTitle", "Live Link"))
-				.SetTooltipText(LOCTEXT("SequenceRecorderTooltipText", "Open the Live Link streaming manager tab."))
-				.SetIcon(FSlateIcon(StyleSetPtr->GetStyleSetName(), "LiveLinkClient.Common.Icon.Small"));
-		});
-	}
-
 	virtual void ShutdownModule() override
 	{
-		if (FSlateApplication::IsInitialized())
-		{
-			FGlobalTabmanager::Get()->UnregisterTabSpawner(LiveLinkClientTabName);
-		}
+		UnregisterTabSpawner();
 
 		FModuleManager::Get().OnModulesChanged().Remove(ModulesChangedHandle);
 
@@ -166,12 +147,38 @@ public:
 
 private:
 
+	void RegisterTabSpawner()
+	{
+		if (bHasRegisteredTabSpawners)
+		{
+			UnregisterTabSpawner();
+		}
+
+		TSharedPtr<FSlateStyleSet> StyleSetPtr = StyleSet;
+		FTabSpawnerEntry& SpawnerEntry = FGlobalTabmanager::Get()->RegisterNomadTabSpawner(LiveLinkClientTabName, FOnSpawnTab::CreateStatic(&FLiveLinkEditorModule::SpawnLiveLinkTab, StyleSetPtr))
+			.SetDisplayName(LOCTEXT("LiveLinkTabTitle", "Live Link"))
+			.SetTooltipText(LOCTEXT("SequenceRecorderTooltipText", "Open the Live Link streaming manager tab."))
+			.SetIcon(FSlateIcon(StyleSetPtr->GetStyleSetName(), "LiveLinkClient.Common.Icon.Small"));
+
+		const IWorkspaceMenuStructure& MenuStructure = WorkspaceMenu::GetMenuStructure();
+		SpawnerEntry.SetGroup(MenuStructure.GetLevelEditorCategory());
+
+		bHasRegisteredTabSpawners = true;
+	}
+
+	void UnregisterTabSpawner()
+	{
+		FGlobalTabmanager::Get()->UnregisterNomadTabSpawner(LiveLinkClientTabName);
+		bHasRegisteredTabSpawners = false;
+	}
+
 	FDelegateHandle LevelEditorTabManagerChangedHandle;
 	FDelegateHandle ModulesChangedHandle;
 	FDelegateHandle CreateLiveLinkPropertyTrackEditorHandle;
 	FMovieSceneLiveLinkSectionRecorderFactory MovieSceneLiveLinkRecorder;
 
-
+	// Track if we have registered
+	bool bHasRegisteredTabSpawners;
 };
 
 IMPLEMENT_MODULE(FLiveLinkEditorModule, LiveLinkEditor);
diff --git a/Engine/Plugins/Experimental/ControlRig/Source/ControlRigEditor/Private/ControlRigBlueprintFactory.cpp b/Engine/Plugins/Experimental/ControlRig/Source/ControlRigEditor/Private/ControlRigBlueprintFactory.cpp
--- a/Engine/Plugins/Experimental/ControlRig/Source/ControlRigEditor/Private/ControlRigBlueprintFactory.cpp
+++ b/Engine/Plugins/Experimental/ControlRig/Source/ControlRigEditor/Private/ControlRigBlueprintFactory.cpp
@@ -132,7 +132,9 @@ private:
 					return false;
 				}
 
-				return (!InClass->HasAnyClassFlags(CLASS_Deprecated) && InClass->GetOutermost() != GetTransientPackage());
+				// in the future we might allow it to parent to BP classes, but right now, it won't work well because of multi rig graph
+				// for now we disable it and only allow native class. 
+				return (!InClass->HasAnyClassFlags(CLASS_Deprecated) && InClass->HasAnyClassFlags(CLASS_Native) && InClass->GetOutermost() != GetTransientPackage());
 			}
 
 			return false;
@@ -141,7 +143,9 @@ private:
 		virtual bool IsUnloadedClassAllowed(const FClassViewerInitializationOptions& InInitOptions, const TSharedRef< const IUnloadedBlueprintData > InUnloadedClassData, TSharedRef< FClassViewerFilterFuncs > InFilterFuncs) override
 		{
 			// If it appears on the allowed child-of classes list (or there is nothing on that list)
-			return InFilterFuncs->IfInChildOfClassesSet( AllowedChildrenOfClasses, InUnloadedClassData) != EFilterReturn::Failed;
+			// in the future we might allow it to parent to BP classes, but right now, it won't work well because of multi rig graph
+			// for now we disable it and only allow native class. 
+			return false; // InFilterFuncs->IfInChildOfClassesSet(AllowedChildrenOfClasses, InUnloadedClassData) != EFilterReturn::Failed;
 		}
 	};
 
diff --git a/Engine/Plugins/Experimental/ControlRig/Source/ControlRigEditor/Private/Editor/SRigHierarchy.cpp b/Engine/Plugins/Experimental/ControlRig/Source/ControlRigEditor/Private/Editor/SRigHierarchy.cpp
--- a/Engine/Plugins/Experimental/ControlRig/Source/ControlRigEditor/Private/Editor/SRigHierarchy.cpp
+++ b/Engine/Plugins/Experimental/ControlRig/Source/ControlRigEditor/Private/Editor/SRigHierarchy.cpp
@@ -171,7 +171,7 @@ void SRigHierarchy::Construct(const FArguments& InArgs, TSharedRef<FControlRigEd
 					.Padding(3.0f, 1.0f)
 					[
 						SAssignNew(FilterBox, SSearchBox)
-			//			.OnTextChanged(this, &SRigHierarchy::OnFilterTextChanged)
+						.OnTextChanged(this, &SRigHierarchy::OnFilterTextChanged)
 					]
 				]
 			]
@@ -218,6 +218,13 @@ void SRigHierarchy::BindCommands()
 		FCanExecuteAction::CreateSP(this, &SRigHierarchy::CanRenameItem));
 }
 
+void SRigHierarchy::OnFilterTextChanged(const FText& SearchText)
+{
+	FilterText = SearchText;
+
+	RefreshTreeView();
+}
+
 void SRigHierarchy::RefreshTreeView()
 {
 	RootJoints.Reset();
@@ -229,31 +236,46 @@ void SRigHierarchy::RefreshTreeView()
 
 		TMap<FName, TSharedPtr<FRigTreeJoint>> SearchTable;
 
+		FString FilteredString = FilterText.ToString();
+		const bool bSearchOff = FilteredString.IsEmpty();
 		for (int32 JointIndex = 0; JointIndex < Hierarchy.Joints.Num(); ++JointIndex)
 		{
 			FRigJoint& Joint = Hierarchy.Joints[JointIndex];
 
 			// create new item
-			TSharedPtr<FRigTreeJoint> NewItem = MakeShared<FRigTreeJoint>(Hierarchy.Joints[JointIndex].Name, SharedThis(this));
-			SearchTable.Add(Joint.Name, NewItem);
-
-			if (Joint.ParentName == NAME_None)
+			if (bSearchOff)
 			{
-				RootJoints.Add(NewItem);
+				TSharedPtr<FRigTreeJoint> NewItem = MakeShared<FRigTreeJoint>(Hierarchy.Joints[JointIndex].Name, SharedThis(this));
+				SearchTable.Add(Joint.Name, NewItem);
+
+				if (Joint.ParentName == NAME_None)
+				{
+					RootJoints.Add(NewItem);
+				}
+				else
+				{
+					// you have to find one
+					TSharedPtr<FRigTreeJoint>* FoundItem = SearchTable.Find(Joint.ParentName);
+					check(FoundItem);
+					// add to children list
+					FoundItem->Get()->Children.Add(NewItem);
+				}
 			}
-			else
+			else if (Joint.Name.ToString().Contains(FilteredString))
 			{
-				// you have to find one
-				TSharedPtr<FRigTreeJoint>* FoundItem = SearchTable.Find(Joint.ParentName);
-				check(FoundItem);
-				// add to children list
-				FoundItem->Get()->Children.Add(NewItem);
+				// if contains, just list out everything to root
+				TSharedPtr<FRigTreeJoint> NewItem = MakeShared<FRigTreeJoint>(Hierarchy.Joints[JointIndex].Name, SharedThis(this));
+				// during search, everything is on root
+				RootJoints.Add(NewItem);
 			}
 		}
 
-		for (int32 RootIndex = 0; RootIndex < RootJoints.Num(); ++RootIndex)
+		if (bSearchOff)
 		{
-			SetExpansionRecursive(RootJoints[RootIndex]);
+			for (int32 RootIndex = 0; RootIndex < RootJoints.Num(); ++RootIndex)
+			{
+				SetExpansionRecursive(RootJoints[RootIndex]);
+			}
 		}
 	}
 
@@ -358,9 +380,17 @@ TSharedPtr< SWidget > SRigHierarchy::CreateContextMenu()
 		MenuBuilder.AddMenuSeparator();
 		MenuBuilder.AddSubMenu(
 			LOCTEXT("ImportSubMenu", "Import"),
-			LOCTEXT("ImportSubMenu_ToolTip", "Insert current pose to selected PoseAsset"),
+			LOCTEXT("ImportSubMenu_ToolTip", "Import hierarchy to the current rig. This only imports non-existing node. For example, if there is hand_r, it won't import hand_r. \
+				If you want to reimport whole new hiearchy, delete all nodes, and use import hierarchy."),
 			FNewMenuDelegate::CreateSP(this, &SRigHierarchy::CreateImportMenu)
 		);
+
+		MenuBuilder.AddMenuSeparator();
+		MenuBuilder.AddSubMenu(
+			LOCTEXT("RefreshSubMenu", "Refresh"),
+			LOCTEXT("RefreshSubMenu_ToolTip", "Refresh the existing initial transform from the selected mesh. This only updates if the node is found."),
+			FNewMenuDelegate::CreateSP(this, &SRigHierarchy::CreateRefreshMenu)
+		);
 	
 		MenuBuilder.EndSection();
 	}
@@ -368,6 +398,62 @@ TSharedPtr< SWidget > SRigHierarchy::CreateContextMenu()
 	return MenuBuilder.MakeWidget();
 }
 
+void SRigHierarchy::CreateRefreshMenu(FMenuBuilder& MenuBuilder)
+{
+	MenuBuilder.AddWidget(
+		SNew(SVerticalBox)
+
+		+ SVerticalBox::Slot()
+		.AutoHeight()
+		.Padding(3)
+		[
+			SNew(STextBlock)
+			.Font(FEditorStyle::GetFontStyle("ControlRig.Hierarchy.Menu"))
+			.Text(LOCTEXT("RefreshMesh_Title", "Select Mesh"))
+			.ToolTipText(LOCTEXT("RefreshMesh_Tooltip", "Select Mesh to refresh transform from... It will refresh init transform from selected mesh. This doesn't change hierarchy. \
+				If you want to reimport hierarchy, please delete all nodes, and use import hierarchy."))
+		]
+		+ SVerticalBox::Slot()
+		.AutoHeight()
+		.Padding(3)
+		[
+			SNew(SObjectPropertyEntryBox)
+			.AllowedClass(USkeletalMesh::StaticClass())
+			.OnObjectChanged(this, &SRigHierarchy::RefreshHierarchy)
+		]
+		,
+		FText()
+	);
+}
+
+void SRigHierarchy::RefreshHierarchy(const FAssetData& InAssetData)
+{
+	FRigHierarchy* Hier = GetHierarchy();
+	USkeletalMesh* Mesh = Cast<USkeletalMesh>(InAssetData.GetAsset());
+	if (Mesh && Hier)
+	{
+		FScopedTransaction Transaction(LOCTEXT("HierarchyRefresh", "Refresh Transform"));
+		ControlRigBlueprint->Modify();
+
+		const FReferenceSkeleton& RefSkeleton = Mesh->RefSkeleton;
+		const TArray<FMeshBoneInfo>& BoneInfos = RefSkeleton.GetRawRefBoneInfo();
+		const TArray<FTransform>& BonePoses = RefSkeleton.GetRawRefBonePose();
+
+		for (int32 BoneIndex = 0; BoneIndex < RefSkeleton.GetNum(); ++BoneIndex)
+		{
+			// only add if you don't have it. This may change in the future
+			int32 RigIndex = Hier->GetIndex(BoneInfos[BoneIndex].Name);
+			if (RigIndex != INDEX_NONE)
+			{
+				// @todo: add optimized version without sorting, but if no sort, we should make sure not to use find index function
+				Hier->SetInitialTransform(RigIndex, FAnimationRuntime::GetComponentSpaceTransform(RefSkeleton, BonePoses, BoneIndex));
+			}
+		}
+
+		RefreshTreeView();
+		FSlateApplication::Get().DismissAllMenus();
+	}
+}
 void SRigHierarchy::CreateImportMenu(FMenuBuilder& MenuBuilder)
 {
 	MenuBuilder.AddWidget(
@@ -380,6 +466,7 @@ void SRigHierarchy::CreateImportMenu(FMenuBuilder& MenuBuilder)
 			SNew(STextBlock)
 			.Font(FEditorStyle::GetFontStyle("ControlRig.Hierarchy.Menu"))
 			.Text(LOCTEXT("ImportMesh_Title", "Select Mesh"))
+			.ToolTipText(LOCTEXT("ImportMesh_Tooltip", "Select Mesh to import hierarchy from... It will only import if the node doens't exists in the current hierarchy."))
 		]
 		+ SVerticalBox::Slot()
 		.AutoHeight()
@@ -400,6 +487,9 @@ void SRigHierarchy::ImportHierarchy(const FAssetData& InAssetData)
 	USkeletalMesh* Mesh = Cast<USkeletalMesh> (InAssetData.GetAsset());
 	if (Mesh && Hier)
 	{
+		FScopedTransaction Transaction(LOCTEXT("HierarchyImport", "Import Hierarchy"));
+		ControlRigBlueprint->Modify();
+
 		const FReferenceSkeleton& RefSkeleton = Mesh->RefSkeleton;
 		const TArray<FMeshBoneInfo>& BoneInfos = RefSkeleton.GetRawRefBoneInfo();
 		const TArray<FTransform>& BonePoses = RefSkeleton.GetRawRefBonePose();
diff --git a/Engine/Plugins/Experimental/ControlRig/Source/ControlRigEditor/Private/Editor/SRigHierarchy.h b/Engine/Plugins/Experimental/ControlRig/Source/ControlRigEditor/Private/Editor/SRigHierarchy.h
--- a/Engine/Plugins/Experimental/ControlRig/Source/ControlRigEditor/Private/Editor/SRigHierarchy.h
+++ b/Engine/Plugins/Experimental/ControlRig/Source/ControlRigEditor/Private/Editor/SRigHierarchy.h
@@ -124,6 +124,9 @@ private:
 
 	/** Search box widget */
 	TSharedPtr<SSearchBox> FilterBox;
+	FText FilterText;
+
+	void OnFilterTextChanged(const FText& SearchText);
 
 	/** Tree view widget */
 	TSharedPtr<STreeView<TSharedPtr<FRigTreeJoint>>> TreeView;
@@ -147,6 +150,8 @@ private:
 
 	void ImportHierarchy(const FAssetData& InAssetData);
 	void CreateImportMenu(FMenuBuilder& MenuBuilder);
+	void CreateRefreshMenu(FMenuBuilder& MenuBuilder);
+	void RefreshHierarchy(const FAssetData& InAssetData);
 
 	FName CreateUniqueName(const FName& InBaseName) const;
 
diff --git a/Engine/Source/Editor/AnimGraph/Private/AnimGraphNode_BoneDrivenController.cpp b/Engine/Source/Editor/AnimGraph/Private/AnimGraphNode_BoneDrivenController.cpp
--- a/Engine/Source/Editor/AnimGraph/Private/AnimGraphNode_BoneDrivenController.cpp
+++ b/Engine/Source/Editor/AnimGraph/Private/AnimGraphNode_BoneDrivenController.cpp
@@ -357,6 +357,12 @@ void UAnimGraphNode_BoneDrivenController::CustomizeDetails(IDetailLayoutBuilder&
 
 	TSharedRef<IPropertyHandle> NodeHandle = DetailBuilder.GetProperty(GET_MEMBER_NAME_CHECKED(UAnimGraphNode_BoneDrivenController, Node), GetClass());
 
+	// if it doesn't have this node, that means this isn't really bone driven controller
+	if (!NodeHandle->IsValidHandle())
+	{
+		return;
+	}
+
 	TAttribute<EVisibility> NotUsingCurveVisibility = TAttribute<EVisibility>::Create(TAttribute<EVisibility>::FGetter::CreateStatic(&UAnimGraphNode_BoneDrivenController::AreNonCurveMappingValuesVisible, &DetailBuilder));
 	TAttribute<EVisibility> MapRangeVisiblity = TAttribute<EVisibility>::Create(TAttribute<EVisibility>::FGetter::CreateStatic(&UAnimGraphNode_BoneDrivenController::AreRemappingValuesVisible, &DetailBuilder));
 
diff --git a/Engine/Source/Editor/AnimGraph/Private/AnimGraphNode_SubInstance.cpp b/Engine/Source/Editor/AnimGraph/Private/AnimGraphNode_SubInstance.cpp
--- a/Engine/Source/Editor/AnimGraph/Private/AnimGraphNode_SubInstance.cpp
+++ b/Engine/Source/Editor/AnimGraph/Private/AnimGraphNode_SubInstance.cpp
@@ -82,36 +82,15 @@ void UAnimGraphNode_SubInstance::ValidateAnimNodeDuringCompilation(USkeleton* Fo
 		MessageLog.Error(TEXT("Sub instance node @@ has no valid instance class to spawn."), this);
 	}
 
-	// Check for cycles from other sub instance nodes
-	TArray<UEdGraph*> Graphs;
-	AnimBP->GetAllGraphs(Graphs);
-
-	for(UEdGraph* Graph : Graphs)
-	{
-		TArray<UAnimGraphNode_SubInstance*> SubInstanceNodes;
-		Graph->GetNodesOfClass(SubInstanceNodes);
-
-		for(UAnimGraphNode_SubInstance* SubInstanceNode : SubInstanceNodes)
-		{
-			if(SubInstanceNode == OriginalNode)
-			{
-				continue;
-			}
-
-			FAnimNode_SubInstance& InnerNode = SubInstanceNode->Node;
-
-			if(*InnerNode.InstanceClass && *InnerNode.InstanceClass == *Node.InstanceClass)
-			{
-				MessageLog.Error(TEXT("Node @@ and node @@ both target the same class @@, causing a sub instance loop."), this, SubInstanceNode, *Node.InstanceClass);
-			}
-		}
-	}
-
 	if(HasInstanceLoop())
 	{
 		MessageLog.Error(TEXT("Detected loop in sub instance chain starting at @@ inside class @@"), this, AnimBP->GetAnimBlueprintGeneratedClass());
 	}
 
+	// Check for cycles from other sub instance nodes
+	TArray<UEdGraph*> Graphs;
+	AnimBP->GetAllGraphs(Graphs);
+
 	// Check for duplicate tags in this anim blueprint
 	for(UEdGraph* Graph : Graphs)
 	{
diff --git a/Engine/Source/Editor/Blutility/Private/BlutilityMenuExtensions.cpp b/Engine/Source/Editor/Blutility/Private/BlutilityMenuExtensions.cpp
--- a/Engine/Source/Editor/Blutility/Private/BlutilityMenuExtensions.cpp
+++ b/Engine/Source/Editor/Blutility/Private/BlutilityMenuExtensions.cpp
@@ -66,6 +66,13 @@ class SFunctionParamDialog : public SCompoundWidget
 		FPropertyEditorModule& PropertyEditorModule = FModuleManager::Get().LoadModuleChecked<FPropertyEditorModule>("PropertyEditor");
 		TSharedRef<IStructureDetailsView> StructureDetailsView = PropertyEditorModule.CreateStructureDetailView(DetailsViewArgs, StructureViewArgs, InStructOnScope);
 
+		StructureDetailsView->GetDetailsView()->SetIsPropertyVisibleDelegate(FIsPropertyVisible::CreateLambda([](const FPropertyAndParent& InPropertyAndParent)
+		{
+			return InPropertyAndParent.Property.HasAnyPropertyFlags(CPF_Parm);
+		}));
+
+		StructureDetailsView->GetDetailsView()->ForceRefresh();
+
 		ChildSlot
 		[
 			SNew(SVerticalBox)
diff --git a/Engine/Source/Editor/MeshPaint/Private/MeshPaintHelpers.cpp b/Engine/Source/Editor/MeshPaint/Private/MeshPaintHelpers.cpp
--- a/Engine/Source/Editor/MeshPaint/Private/MeshPaintHelpers.cpp
+++ b/Engine/Source/Editor/MeshPaint/Private/MeshPaintHelpers.cpp
@@ -134,7 +134,7 @@ bool MeshPaintHelpers::PropagateColorsToRawMesh(UStaticMesh* StaticMesh, int32 L
 	else
 	{
 		// If there's no raw mesh data, don't try to do any fixup here
-		if (SrcModel.IsRawMeshEmpty() || ComponentLODInfo.OverrideMapBuildData == nullptr)
+		if (SrcModel.IsRawMeshEmpty() || ComponentLODInfo.OverrideVertexColors == nullptr)
 		{
 			return false;
 		}
diff --git a/Engine/Source/Editor/Persona/Private/AnimGraphNodeDetails.cpp b/Engine/Source/Editor/Persona/Private/AnimGraphNodeDetails.cpp
--- a/Engine/Source/Editor/Persona/Private/AnimGraphNodeDetails.cpp
+++ b/Engine/Source/Editor/Persona/Private/AnimGraphNodeDetails.cpp
@@ -63,17 +63,29 @@ void FAnimGraphNodeDetails::CustomizeDetails(class IDetailLayoutBuilder& DetailB
 
 	// get first animgraph nodes
 	UAnimGraphNode_Base* AnimGraphNode = Cast<UAnimGraphNode_Base>(SelectedObjectsList[0].Get());
-	if (AnimGraphNode == NULL)
+	if (AnimGraphNode == nullptr)
 	{
 		return;
 	}
 
+	// make sure type matches with all the nodes. 
+	const UAnimGraphNode_Base* FirstNodeType = AnimGraphNode;
+	for (int32 Index = 1; Index < SelectedObjectsList.Num(); ++Index)
+	{
+		UAnimGraphNode_Base* CurrentNode = Cast<UAnimGraphNode_Base>(SelectedObjectsList[Index].Get());
+		if (!CurrentNode || CurrentNode->GetClass() != FirstNodeType->GetClass())
+		{
+			// if type mismatches, multi selection doesn't work, just return
+			return;
+		}
+	}
+
 	TargetSkeleton = AnimGraphNode->GetAnimBlueprint()->TargetSkeleton;
 	TargetSkeletonName = FString::Printf(TEXT("%s'%s'"), *TargetSkeleton->GetClass()->GetName(), *TargetSkeleton->GetPathName());
 
 	// Get the node property
 	const UStructProperty* NodeProperty = AnimGraphNode->GetFNodeProperty();
-	if (NodeProperty == NULL)
+	if (NodeProperty == nullptr)
 	{
 		return;
 	}
diff --git a/Engine/Source/Editor/Persona/Private/SAnimNotifyPanel.cpp b/Engine/Source/Editor/Persona/Private/SAnimNotifyPanel.cpp
--- a/Engine/Source/Editor/Persona/Private/SAnimNotifyPanel.cpp
+++ b/Engine/Source/Editor/Persona/Private/SAnimNotifyPanel.cpp
@@ -3886,6 +3886,7 @@ void SAnimNotifyPanel::Construct(const FArguments& InArgs, const TSharedRef<clas
 	Sequence = InArgs._Sequence;
 	MarkerBars = InArgs._MarkerBars;
 	OnInvokeTab = InArgs._OnInvokeTab;
+	OnNotifiesChanged = InArgs._OnNotifiesChanged;
 
 	FAnimNotifyPanelCommands::Register();
 	BindCommands();
@@ -4105,6 +4106,10 @@ void SAnimNotifyPanel::Update()
 	{
 		Sequence->RefreshCacheData();
 	}
+
+	RefreshNotifyTracks();
+
+	OnNotifiesChanged.ExecuteIfBound();
 }
 
 void SAnimNotifyPanel::RefreshNotifyTracks()
diff --git a/Engine/Source/Editor/Persona/Private/SAnimNotifyPanel.h b/Engine/Source/Editor/Persona/Private/SAnimNotifyPanel.h
--- a/Engine/Source/Editor/Persona/Private/SAnimNotifyPanel.h
+++ b/Engine/Source/Editor/Persona/Private/SAnimNotifyPanel.h
@@ -184,6 +184,7 @@ public:
 	SLATE_EVENT( FRefreshOffsetsRequest, OnRequestRefreshOffsets )
 	SLATE_EVENT( FOnGetTimingNodeVisibility, OnGetTimingNodeVisibility )
 	SLATE_EVENT( FOnInvokeTab, OnInvokeTab )
+	SLATE_EVENT( FSimpleDelegate, OnNotifiesChanged )
 
 	SLATE_END_ARGS()
 
@@ -340,4 +341,7 @@ private:
 
 	/** Delegate used to invoke a tab */
 	FOnInvokeTab OnInvokeTab;
+
+	/** Delegate used to inform others that notifies have changed (for timing) */
+	FSimpleDelegate OnNotifiesChanged;
 };
diff --git a/Engine/Source/Editor/Persona/Private/SMontageEditor.cpp b/Engine/Source/Editor/Persona/Private/SMontageEditor.cpp
--- a/Engine/Source/Editor/Persona/Private/SMontageEditor.cpp
+++ b/Engine/Source/Editor/Persona/Private/SMontageEditor.cpp
@@ -199,6 +199,7 @@ void SMontageEditor::Construct(const FArguments& InArgs, const FMontageEditorReq
 		.OnRequestRefreshOffsets(this, &SMontageEditor::RefreshNotifyTriggerOffsets)
 		.OnGetTimingNodeVisibility(TimingNodeVisibilityDelegate)
 		.OnInvokeTab(InArgs._OnInvokeTab)
+		.OnNotifiesChanged(this, &SMontageEditor::HandleNotifiesChanged)
 	];
 
 	EditorPanels->AddSlot()
@@ -333,6 +334,11 @@ void SMontageEditor::RefreshNotifyTriggerOffsets()
 	}
 }
 
+void SMontageEditor::HandleNotifiesChanged()
+{
+	AnimTimingPanel->Update();
+}
+
 bool SMontageEditor::GetSectionTime( int32 SectionIndex, float &OutTime ) const
 {
 	if (MontageObj != NULL && MontageObj->CompositeSections.IsValidIndex(SectionIndex))
diff --git a/Engine/Source/Editor/Persona/Private/SMontageEditor.h b/Engine/Source/Editor/Persona/Private/SMontageEditor.h
--- a/Engine/Source/Editor/Persona/Private/SMontageEditor.h
+++ b/Engine/Source/Editor/Persona/Private/SMontageEditor.h
@@ -131,6 +131,9 @@ private:
 	/** Updates Notify trigger offsets to take into account current montage state */
 	void RefreshNotifyTriggerOffsets();
 
+	/** Handle notifies changing, so we rebuild timing */
+	void HandleNotifiesChanged();
+
 	/** One-off active timer to trigger a montage panel rebuild */
 	EActiveTimerReturnType TriggerRebuildMontagePanel(double InCurrentTime, float InDeltaTime);
 
diff --git a/Engine/Source/Editor/UnrealEd/Classes/Preferences/PersonaOptions.h b/Engine/Source/Editor/UnrealEd/Classes/Preferences/PersonaOptions.h
--- a/Engine/Source/Editor/UnrealEd/Classes/Preferences/PersonaOptions.h
+++ b/Engine/Source/Editor/UnrealEd/Classes/Preferences/PersonaOptions.h
@@ -47,6 +47,12 @@ struct FViewportConfigOptions
 
 	UPROPERTY(config)
 	FName CameraFollowBoneName;
+
+	FViewportConfigOptions()
+		: ViewModeIndex(EViewModeIndex::VMI_Lit)
+		, ViewFOV(53.43)
+		, CameraFollowMode(EAnimationViewportCameraFollowMode::None)
+	{}
 };
 
 /** Options that should be unique per asset editor (like skeletal mesh or anim sequence editors) */
diff --git a/Engine/Source/Runtime/AnimGraphRuntime/Private/AnimNodes/AnimNode_CopyPoseFromMesh.cpp b/Engine/Source/Runtime/AnimGraphRuntime/Private/AnimNodes/AnimNode_CopyPoseFromMesh.cpp
--- a/Engine/Source/Runtime/AnimGraphRuntime/Private/AnimNodes/AnimNode_CopyPoseFromMesh.cpp
+++ b/Engine/Source/Runtime/AnimGraphRuntime/Private/AnimNodes/AnimNode_CopyPoseFromMesh.cpp
@@ -85,7 +85,9 @@ void FAnimNode_CopyPoseFromMesh::Evaluate_AnyThread(FPoseContext& Output)
 	if (CurrentMeshComponent && CurrentMeshComponent->SkeletalMesh && CurrentMeshComponent->IsRegistered())
 	{
 		const FBoneContainer& RequiredBones = OutPose.GetBoneContainer();
-		const TArray<FTransform>& SourcMeshTransformArray = CurrentMeshComponent->GetComponentSpaceTransforms();
+		const bool bUROInSync = CurrentMeshComponent->AnimUpdateRateParams != nullptr && CurrentMeshComponent->AnimUpdateRateParams == Output.AnimInstanceProxy->GetSkelMeshComponent()->AnimUpdateRateParams;
+		const bool bArraySizesMatch = CurrentMeshComponent->CachedComponentSpaceTransforms.Num() == CurrentMeshComponent->GetComponentSpaceTransforms().Num();
+		const TArray<FTransform>& SourceMeshTransformArray =  bUROInSync && bArraySizesMatch ? CurrentMeshComponent->CachedComponentSpaceTransforms : CurrentMeshComponent->GetComponentSpaceTransforms();
 		for(FCompactPoseBoneIndex PoseBoneIndex : OutPose.ForEachBoneIndex())
 		{
 			const int32& SkeletonBoneIndex = RequiredBones.GetSkeletonIndex(PoseBoneIndex);
@@ -99,13 +101,13 @@ void FAnimNode_CopyPoseFromMesh::Evaluate_AnyThread(FPoseContext& Output)
 				// only apply if I also have parent, otherwise, it should apply the space bases
 				if (ParentIndex != INDEX_NONE && MyParentIndex != INDEX_NONE)
 				{
-					const FTransform& ParentTransform = SourcMeshTransformArray[ParentIndex];
-					const FTransform& ChildTransform = SourcMeshTransformArray[SourceBoneIndex];
+					const FTransform& ParentTransform = SourceMeshTransformArray[ParentIndex];
+					const FTransform& ChildTransform = SourceMeshTransformArray[SourceBoneIndex];
 					OutPose[PoseBoneIndex] = ChildTransform.GetRelativeTransform(ParentTransform);
 				}
 				else
 				{
-					OutPose[PoseBoneIndex] = SourcMeshTransformArray[SourceBoneIndex];
+					OutPose[PoseBoneIndex] = SourceMeshTransformArray[SourceBoneIndex];
 				}
 			}
 		}
diff --git a/Engine/Source/Runtime/Engine/Private/AnimBlueprint.cpp b/Engine/Source/Runtime/Engine/Private/AnimBlueprint.cpp
--- a/Engine/Source/Runtime/Engine/Private/AnimBlueprint.cpp
+++ b/Engine/Source/Runtime/Engine/Private/AnimBlueprint.cpp
@@ -202,6 +202,10 @@ USkeletalMesh* UAnimBlueprint::GetPreviewMesh(bool bFindIfNotSet/*=false*/)
 USkeletalMesh* UAnimBlueprint::GetPreviewMesh() const
 {
 #if WITH_EDITORONLY_DATA
+	if (!PreviewSkeletalMesh.IsValid())
+	{
+		PreviewSkeletalMesh.LoadSynchronous();
+	}
 	return PreviewSkeletalMesh.Get();
 #else
 	return nullptr;
diff --git a/Engine/Source/Runtime/Engine/Private/Animation/AnimLinkableElement.cpp b/Engine/Source/Runtime/Engine/Private/Animation/AnimLinkableElement.cpp
--- a/Engine/Source/Runtime/Engine/Private/Animation/AnimLinkableElement.cpp
+++ b/Engine/Source/Runtime/Engine/Private/Animation/AnimLinkableElement.cpp
@@ -289,12 +289,28 @@ void FAnimLinkableElement::SetTimeFromProportional(float NewTime, EAnimLinkMetho
 	{
 		case EAnimLinkMethod::Absolute:
 		{
-			LinkValue = (NewTime - SegmentBeginTime) / SegmentLength;
+			if(SegmentLength != 0.0f)
+			{
+				LinkValue = (NewTime - SegmentBeginTime) / SegmentLength;
+			}
+			else
+			{
+				// if segment length is 0, we can't set anywhere else but 0.f
+				LinkValue = 0.f;
+			}
 		}
 			break;
 		case EAnimLinkMethod::Relative:
 		{
-			LinkValue = NewTime / SegmentLength;
+			if(SegmentLength != 0.0f)
+			{
+				LinkValue = NewTime / SegmentLength;
+			}
+			else
+			{
+				// if segment length is 0, we can't set anywhere else but 0.f
+				LinkValue = 0.f;			
+			}
 		}
 			break;
 	}
diff --git a/Engine/Source/Runtime/Engine/Private/Animation/AnimSingleNodeInstance.cpp b/Engine/Source/Runtime/Engine/Private/Animation/AnimSingleNodeInstance.cpp
--- a/Engine/Source/Runtime/Engine/Private/Animation/AnimSingleNodeInstance.cpp
+++ b/Engine/Source/Runtime/Engine/Private/Animation/AnimSingleNodeInstance.cpp
@@ -393,28 +393,45 @@ float UAnimSingleNodeInstance::GetLength()
 
 void UAnimSingleNodeInstance::StepForward()
 {
-	if (UAnimSequence* Sequence = Cast<UAnimSequence>(CurrentAsset))
+	if (UAnimSequenceBase* Sequence = Cast<UAnimSequenceBase>(CurrentAsset))
 	{
 		FAnimSingleNodeInstanceProxy& Proxy = GetProxyOnGameThread<FAnimSingleNodeInstanceProxy>();
-		const FAnimKeyHelper Helper(Sequence->SequenceLength, Sequence->NumFrames);
+		const FAnimKeyHelper Helper(Sequence->SequenceLength, Sequence->GetNumberOfFrames());
 		float KeyLength = Helper.TimePerKey() + SMALL_NUMBER;
 		float Fraction = (Proxy.GetCurrentTime()+KeyLength)/Sequence->SequenceLength;
-		int32 Frames = FMath::Clamp<int32>((float)(Helper.LastKey()*Fraction), 0, Helper.LastKey());
+		int32 Frames = (float)Helper.LastKey() * Fraction;
+		if (IsLooping())
+		{
+			int32 PreviousFrame = (float)Helper.LastKey() * (Proxy.GetCurrentTime() / Sequence->SequenceLength);
+			Frames = (PreviousFrame == Helper.LastKey()) ? 0 : Frames;
+		}
+		else
+		{
+			Frames = FMath::Clamp<int32>(Frames, 0, Helper.LastKey());
+		}
 		SetPosition(Frames*KeyLength);
 	}	
 }
 
 void UAnimSingleNodeInstance::StepBackward()
 {
-	if (UAnimSequence* Sequence = Cast<UAnimSequence>(CurrentAsset))
+	if (UAnimSequenceBase* Sequence = Cast<UAnimSequenceBase>(CurrentAsset))
 	{
 		FAnimSingleNodeInstanceProxy& Proxy = GetProxyOnGameThread<FAnimSingleNodeInstanceProxy>();
 
-		const FAnimKeyHelper Helper(Sequence->SequenceLength, Sequence->NumFrames);
+		const FAnimKeyHelper Helper(Sequence->SequenceLength, Sequence->GetNumberOfFrames());
 		float KeyLength = Helper.TimePerKey() + SMALL_NUMBER;
 		float Fraction = (Proxy.GetCurrentTime()-KeyLength)/Sequence->SequenceLength;
-		int32 Frames = FMath::Clamp<int32>((float)(Helper.LastKey()*Fraction), 0, Helper.LastKey());
-		SetPosition(Frames*KeyLength);
+		int32 Frames = (float)Helper.LastKey() * Fraction;
+		if (IsLooping())
+		{
+			Frames = (Frames < 0) ? Helper.LastKey() : Frames;
+		}
+		else
+		{
+			Frames = FMath::Clamp<int32>(Frames, 0, Helper.LastKey());
+		}
+		SetPosition(Frames * KeyLength);
 	}
 }
 
diff --git a/Engine/Source/Runtime/Engine/Private/Animation/BlendSpaceBase.cpp b/Engine/Source/Runtime/Engine/Private/Animation/BlendSpaceBase.cpp
--- a/Engine/Source/Runtime/Engine/Private/Animation/BlendSpaceBase.cpp
+++ b/Engine/Source/Runtime/Engine/Private/Animation/BlendSpaceBase.cpp
@@ -296,7 +296,7 @@ void UBlendSpaceBase::TickAssetPlayer(FAnimTickRecord& Instance, struct FAnimNot
 						Sample.Animation->GetMarkerIndicesForPosition(Context.MarkerTickContext.GetMarkerSyncStartPosition(), true, SampleDataItem.MarkerTickRecord.PreviousMarker, SampleDataItem.MarkerTickRecord.NextMarker, SampleDataItem.Time);
 					}
 
-					const float NewDeltaTime = Context.GetDeltaTime() * Instance.PlayRateMultiplier * Sample.RateScale;
+					const float NewDeltaTime = Context.GetDeltaTime() * Instance.PlayRateMultiplier * Sample.RateScale * Sample.Animation->RateScale;
 					if (!FMath::IsNearlyZero(NewDeltaTime))
 					{
 						Context.SetLeaderDelta(NewDeltaTime);
diff --git a/Engine/Source/Runtime/Engine/Private/Animation/Skeleton.cpp b/Engine/Source/Runtime/Engine/Private/Animation/Skeleton.cpp
--- a/Engine/Source/Runtime/Engine/Private/Animation/Skeleton.cpp
+++ b/Engine/Source/Runtime/Engine/Private/Animation/Skeleton.cpp
@@ -789,6 +789,10 @@ USkeletalMesh* USkeleton::GetPreviewMesh(bool bFindIfNotSet/*=false*/)
 USkeletalMesh* USkeleton::GetPreviewMesh() const
 {
 #if WITH_EDITORONLY_DATA
+	if (!PreviewSkeletalMesh.IsValid())
+	{
+		PreviewSkeletalMesh.LoadSynchronous();
+	}
 	return PreviewSkeletalMesh.Get();
 #else
 	return nullptr;
diff --git a/Engine/Source/Runtime/Engine/Private/AnimationAsset.cpp b/Engine/Source/Runtime/Engine/Private/AnimationAsset.cpp
--- a/Engine/Source/Runtime/Engine/Private/AnimationAsset.cpp
+++ b/Engine/Source/Runtime/Engine/Private/AnimationAsset.cpp
@@ -278,6 +278,10 @@ USkeletalMesh* UAnimationAsset::GetPreviewMesh(bool bFindIfNotSet)
 USkeletalMesh* UAnimationAsset::GetPreviewMesh() const
 {
 #if WITH_EDITORONLY_DATA
+	if (!PreviewSkeletalMesh.IsValid())
+	{
+		PreviewSkeletalMesh.LoadSynchronous();
+	}
 	return PreviewSkeletalMesh.Get();
 #else
 	return nullptr;
diff --git a/Engine/Source/Runtime/Engine/Private/Components/SkeletalMeshComponent.cpp b/Engine/Source/Runtime/Engine/Private/Components/SkeletalMeshComponent.cpp
--- a/Engine/Source/Runtime/Engine/Private/Components/SkeletalMeshComponent.cpp
+++ b/Engine/Source/Runtime/Engine/Private/Components/SkeletalMeshComponent.cpp
@@ -1753,8 +1753,7 @@ void USkeletalMeshComponent::PerformAnimationProcessing(const USkeletalMesh* InS
 	ANIM_MT_SCOPE_CYCLE_COUNTER(PerformAnimEvaluation, IsRunningParallelEvaluation());
 
 	// Can't do anything without a SkeletalMesh
-	// Do nothing more if no bones in skeleton.
-	if (!InSkeletalMesh || OutSpaceBases.Num() == 0)
+	if (!InSkeletalMesh)
 	{
 		return;
 	}
@@ -1771,7 +1770,8 @@ void USkeletalMeshComponent::PerformAnimationProcessing(const USkeletalMesh* InS
 		PostProcessAnimInstance->ParallelUpdateAnimation();
 	}
 
-	if(bInDoEvaluation)
+	// Do nothing more if no bones in skeleton.
+	if(bInDoEvaluation && OutSpaceBases.Num() > 0)
 	{
 		FMemMark Mark(FMemStack::Get());
 		FCompactPose EvaluatedPose;
diff --git a/Engine/Source/Runtime/Engine/Private/HLODProxy.cpp b/Engine/Source/Runtime/Engine/Private/HLODProxy.cpp
--- a/Engine/Source/Runtime/Engine/Private/HLODProxy.cpp
+++ b/Engine/Source/Runtime/Engine/Private/HLODProxy.cpp
@@ -367,7 +367,17 @@ FName UHLODProxy::GenerateKeyForActor(const ALODActor* LODActor)
 bool UHLODProxy::ContainsDataForActor(const ALODActor* InLODActor) const
 {
 #if WITH_EDITOR
-	FName Key = GenerateKeyForActor(InLODActor);
+	FName Key;
+
+	// Only re-generate the key in non-PIE worlds
+	if(InLODActor->GetOutermost()->HasAnyPackageFlags(PKG_PlayInEditor))
+	{
+		Key = InLODActor->GetKey();
+	}
+	else
+	{
+		Key = GenerateKeyForActor(InLODActor);
+	}
 #else
 	FName Key = InLODActor->GetKey();
 #endif
diff --git a/Engine/Source/Runtime/Engine/Private/LODActor.cpp b/Engine/Source/Runtime/Engine/Private/LODActor.cpp
--- a/Engine/Source/Runtime/Engine/Private/LODActor.cpp
+++ b/Engine/Source/Runtime/Engine/Private/LODActor.cpp
@@ -18,6 +18,7 @@
 #include "UObject/FrameworkObjectVersion.h"
 #include "UObject/FortniteMainBranchObjectVersion.h"
 #include "Engine/HLODProxy.h"
+#include "UObject/PropertyPortFlags.h"
 
 #if WITH_EDITOR
 #include "Editor.h"
@@ -974,6 +975,13 @@ void ALODActor::OnCVarsChanged()
 
 void ALODActor::Serialize(FArchive& Ar)
 {
+#if WITH_EDITOR
+	if (Ar.GetPortFlags() & PPF_DuplicateForPIE && Ar.IsSaving())
+	{
+		Key = UHLODProxy::GenerateKeyForActor(this);
+	}
+#endif
+
 	Super::Serialize(Ar);
 #if WITH_EDITOR
 	Ar.UsingCustomVersion(FFrameworkObjectVersion::GUID);
