- Fixed memory corruption when adding
- Control Rig node watch support for quick debugging

diff --git a/Engine/Plugins/Experimental/ControlRig/Source/ControlRig/Private/Units/Math/RigUnit_Float.h b/Engine/Plugins/Experimental/ControlRig/Source/ControlRig/Private/Units/Math/RigUnit_Float.h
--- a/Engine/Plugins/Experimental/ControlRig/Source/ControlRig/Private/Units/Math/RigUnit_Float.h
+++ b/Engine/Plugins/Experimental/ControlRig/Source/ControlRig/Private/Units/Math/RigUnit_Float.h
@@ -88,4 +88,34 @@ struct FRigUnit_Clamp_Float: public FRigUnit
 	{
 		Result = FMath::Clamp(Value, Min, Max);
 	}
+};
+
+/** Two args and a result of float type */
+USTRUCT(meta = (DisplayName = "MapRange", Category = "Math|Float"))
+struct FRigUnit_MapRange_Float: public FRigUnit
+{
+	GENERATED_BODY()
+
+	UPROPERTY(meta = (Input))
+	float Value;
+
+	UPROPERTY(meta = (Input))
+	float MinIn;
+
+	UPROPERTY(meta = (Input))
+	float MaxIn;
+
+	UPROPERTY(meta = (Input))
+	float MinOut;
+
+	UPROPERTY(meta = (Input))
+	float MaxOut;
+
+	UPROPERTY(meta = (Output))
+	float Result;
+
+	virtual void Execute(const FRigUnitContext& InContext) override
+	{
+		Result = FMath::GetMappedRangeValueClamped(FVector2D(MinIn, MaxIn), FVector2D(MinOut, MaxOut), Value);
+	}
 };
\ No newline at end of file
diff --git a/Engine/Plugins/Experimental/ControlRig/Source/ControlRigDeveloper/Public/Graph/ControlRigGraphSchema.h b/Engine/Plugins/Experimental/ControlRig/Source/ControlRigDeveloper/Public/Graph/ControlRigGraphSchema.h
--- a/Engine/Plugins/Experimental/ControlRig/Source/ControlRigDeveloper/Public/Graph/ControlRigGraphSchema.h
+++ b/Engine/Plugins/Experimental/ControlRig/Source/ControlRigDeveloper/Public/Graph/ControlRigGraphSchema.h
@@ -66,6 +66,7 @@ public:
 	virtual void TrySetDefaultText(UEdGraphPin& InPin, const FText& InNewDefaultText) const override;
 	virtual bool ShouldAlwaysPurgeOnModification() const override { return false; }
 	virtual bool ArePinsCompatible(const UEdGraphPin* PinA, const UEdGraphPin* PinB, const UClass* CallingContext, bool bIgnoreArray /*= false*/) const override;
+	virtual bool DoesSupportPinWatching() const	override { return true; }
 
 	/** Create a graph node for a rig */
 	UControlRigGraphNode* CreateGraphNode(UControlRigGraph* InGraph, const FName& InPropertyName) const;
diff --git a/Engine/Plugins/Experimental/ControlRig/Source/ControlRigEditor/ControlRigEditor.Build.cs b/Engine/Plugins/Experimental/ControlRig/Source/ControlRigEditor/ControlRigEditor.Build.cs
--- a/Engine/Plugins/Experimental/ControlRig/Source/ControlRigEditor/ControlRigEditor.Build.cs
+++ b/Engine/Plugins/Experimental/ControlRig/Source/ControlRigEditor/ControlRigEditor.Build.cs
@@ -47,6 +47,7 @@ namespace UnrealBuildTool.Rules
                     "Persona",
                     "UMG",
 					"TimeManagement",
+                    "PropertyPath",
                 }
             );
         }
diff --git a/Engine/Plugins/Experimental/ControlRig/Source/ControlRigEditor/Private/ControlRigEditorModule.cpp b/Engine/Plugins/Experimental/ControlRig/Source/ControlRigEditor/Private/ControlRigEditorModule.cpp
--- a/Engine/Plugins/Experimental/ControlRig/Source/ControlRigEditor/Private/ControlRigEditorModule.cpp
+++ b/Engine/Plugins/Experimental/ControlRig/Source/ControlRigEditor/Private/ControlRigEditorModule.cpp
@@ -54,6 +54,7 @@
 #include "Graph/NodeSpawners/ControlRigUnitNodeSpawner.h"
 #include "Graph/NodeSpawners/ControlRigVariableNodeSpawner.h"
 #include "Kismet2/BlueprintEditorUtils.h"
+#include "Kismet2/KismetDebugUtilities.h"
 #include "Graph/ControlRigGraphNode.h"
 #include "EdGraphUtilities.h"
 #include "ControlRigGraphPanelNodeFactory.h"
@@ -889,6 +890,25 @@ void FControlRigEditorModule::GetContextMenuActions(const UControlRigGraphSchema
 				}
 			}
 			MenuBuilder->EndSection();
+
+			// Add the watch pin / unwatch pin menu items
+			MenuBuilder->BeginSection("EdGraphSchemaWatches", LOCTEXT("WatchesHeader", "Watches"));
+			{
+				UBlueprint* OwnerBlueprint = FBlueprintEditorUtils::FindBlueprintForGraphChecked(CurrentGraph);
+				{
+					const UEdGraphPin* WatchedPin = ((InGraphPin->Direction == EGPD_Input) && (InGraphPin->LinkedTo.Num() > 0)) ? InGraphPin->LinkedTo[0] : InGraphPin;
+					if (FKismetDebugUtilities::IsPinBeingWatched(OwnerBlueprint, WatchedPin))
+					{
+						MenuBuilder->AddMenuEntry(FGraphEditorCommands::Get().StopWatchingPin);
+					}
+					else
+					{
+						MenuBuilder->AddMenuEntry(FGraphEditorCommands::Get().StartWatchingPin);
+					}
+				}
+			}
+			MenuBuilder->EndSection();
+
 		}
 	}
 }
diff --git a/Engine/Plugins/Experimental/ControlRig/Source/ControlRigEditor/Private/Editor/ControlRigEditor.cpp b/Engine/Plugins/Experimental/ControlRig/Source/ControlRigEditor/Private/Editor/ControlRigEditor.cpp
--- a/Engine/Plugins/Experimental/ControlRig/Source/ControlRigEditor/Private/Editor/ControlRigEditor.cpp
+++ b/Engine/Plugins/Experimental/ControlRig/Source/ControlRigEditor/Private/Editor/ControlRigEditor.cpp
@@ -277,6 +277,7 @@ void FControlRigEditor::OnCreateGraphEditorCommands(TSharedPtr<FUICommandList> G
 
 void FControlRigEditor::Compile()
 {
+	GetBlueprintObj()->SetObjectBeingDebugged(nullptr);
 	ClearDetailObject();
 	FBlueprintEditor::Compile();
 }
@@ -729,6 +730,9 @@ void FControlRigEditor::UpdateControlRig()
 			{
 				ControlRig->SetObjectBinding(MakeShared<FControlRigSkeletalMeshBinding>());
 			}
+			
+			// Make sure the object being debugged is the preview instance
+			GetBlueprintObj()->SetObjectBeingDebugged(ControlRig);
 
 			// initialize is moved post reinstance
 			FInputBlendPose Filter;
diff --git a/Engine/Plugins/Experimental/ControlRig/Source/ControlRigEditor/Private/Graph/ControlRigConnectionDrawingPolicy.cpp b/Engine/Plugins/Experimental/ControlRig/Source/ControlRigEditor/Private/Graph/ControlRigConnectionDrawingPolicy.cpp
--- a/Engine/Plugins/Experimental/ControlRig/Source/ControlRigEditor/Private/Graph/ControlRigConnectionDrawingPolicy.cpp
+++ b/Engine/Plugins/Experimental/ControlRig/Source/ControlRigEditor/Private/Graph/ControlRigConnectionDrawingPolicy.cpp
@@ -11,24 +11,25 @@ void FControlRigConnectionDrawingPolicy::BuildPinToPinWidgetMap(TMap<TSharedRef<
 	{
 		struct Local
 		{
-			static void AddSubPins_Recursive(UEdGraphPin* PinObj, TMap<UEdGraphPin*, TSharedRef<SGraphPin>>& InPinToPinWidgetMap, TSharedRef<SGraphPin>& InGraphPinWidget)
+			static void AddSubPins_Recursive(UEdGraphPin* PinObj, TMap<UEdGraphPin*, TSharedPtr<SGraphPin>>& InPinToPinWidgetMap, TSharedPtr<SGraphPin>& InGraphPinWidget)
 			{
 				for(UEdGraphPin* SubPin : PinObj->SubPins)
 				{
 					// Only add to the pin-to-pin widget map if the sub-pin widget is not there already
-					TSharedRef<SGraphPin>* SubPinWidgetPtr = InPinToPinWidgetMap.Find(SubPin);
+					TSharedPtr<SGraphPin>* SubPinWidgetPtr = InPinToPinWidgetMap.Find(SubPin);
 					if(SubPinWidgetPtr == nullptr)
 					{
 						SubPinWidgetPtr = &InGraphPinWidget;
 					}
 
-					InPinToPinWidgetMap.Add(SubPin, *SubPinWidgetPtr);
-					AddSubPins_Recursive(SubPin, InPinToPinWidgetMap, *SubPinWidgetPtr);
+					TSharedPtr<SGraphPin> PinWidgetPtr = *SubPinWidgetPtr;
+					InPinToPinWidgetMap.Add(SubPin, PinWidgetPtr);
+					AddSubPins_Recursive(SubPin, InPinToPinWidgetMap, PinWidgetPtr);
 				}
 			}
 		};
 
-		TSharedRef<SGraphPin> GraphPinWidget = StaticCastSharedRef<SGraphPin>(ConnectorIt.Key());
+		TSharedPtr<SGraphPin> GraphPinWidget = StaticCastSharedRef<SGraphPin>(ConnectorIt.Key());
 		Local::AddSubPins_Recursive(GraphPinWidget->GetPinObj(), PinToPinWidgetMap, GraphPinWidget);
 	}
 }
@@ -89,13 +90,13 @@ void FControlRigConnectionDrawingPolicy::DetermineLinkGeometry(
 	/*out*/ FArrangedWidget*& EndWidgetGeometry
 	)
 {
-	if (TSharedRef<SGraphPin>* pOutputWidget = PinToPinWidgetMap.Find(OutputPin))
+	if (TSharedPtr<SGraphPin>* pOutputWidget = PinToPinWidgetMap.Find(OutputPin))
 	{
-		StartWidgetGeometry = PinGeometries->Find(*pOutputWidget);
+		StartWidgetGeometry = PinGeometries->Find((*pOutputWidget).ToSharedRef());
 	}
 	
-	if (TSharedRef<SGraphPin>* pInputWidget = PinToPinWidgetMap.Find(InputPin))
+	if (TSharedPtr<SGraphPin>* pInputWidget = PinToPinWidgetMap.Find(InputPin))
 	{
-		EndWidgetGeometry = PinGeometries->Find(*pInputWidget);
+		EndWidgetGeometry = PinGeometries->Find((*pInputWidget).ToSharedRef());
 	}
 }
\ No newline at end of file
diff --git a/Engine/Plugins/Experimental/ControlRig/Source/ControlRigEditor/Private/Graph/SControlRigGraphNode.cpp b/Engine/Plugins/Experimental/ControlRig/Source/ControlRigEditor/Private/Graph/SControlRigGraphNode.cpp
--- a/Engine/Plugins/Experimental/ControlRig/Source/ControlRigEditor/Private/Graph/SControlRigGraphNode.cpp
+++ b/Engine/Plugins/Experimental/ControlRig/Source/ControlRigEditor/Private/Graph/SControlRigGraphNode.cpp
@@ -13,6 +13,11 @@
 #include "GraphEditorSettings.h"
 #include "ControlRigEditorStyle.h"
 #include "Widgets/Layout/SWrapBox.h"
+#include "Engine/Engine.h"
+#include "KismetNodes/KismetNodeInfoContext.h"
+#include "Kismet2/KismetDebugUtilities.h"
+#include "PropertyPathHelpers.h"
+#include "UObject/PropertyPortFlags.h"
 
 #define LOCTEXT_NAMESPACE "SControlRigGraphNode"
 
@@ -638,4 +643,58 @@ FReply SControlRigGraphNode::HandleAddArrayElement(TWeakPtr<FControlRigField> In
 	return FReply::Handled();
 }
 
+void SControlRigGraphNode::GetNodeInfoPopups(FNodeInfoContext* Context, TArray<FGraphInformationPopupInfo>& Popups) const
+{
+	FKismetNodeInfoContext* K2Context = (FKismetNodeInfoContext*)Context;
+
+	const FLinearColor LatentBubbleColor(1.f, 0.5f, 0.25f);
+	const FLinearColor PinnedWatchColor(0.35f, 0.25f, 0.25f);
+
+	// Display any pending latent actions
+	if (UObject* ActiveObject = K2Context->ActiveObjectBeingDebugged)
+	{
+		// Display pinned watches
+		if (K2Context->WatchedNodeSet.Contains(GraphNode))
+		{
+			UBlueprint* Blueprint = K2Context->SourceBlueprint;
+			const UEdGraphSchema* Schema = GraphNode->GetSchema();
+
+			FString PinnedWatchText;
+			int32 ValidWatchCount = 0;
+			for (int32 PinIndex = 0; PinIndex < GraphNode->Pins.Num(); ++PinIndex)
+			{
+				UEdGraphPin* WatchPin = GraphNode->Pins[PinIndex];
+				if (K2Context->WatchedPinSet.Contains(WatchPin))
+				{
+					if (ValidWatchCount > 0)
+					{
+						PinnedWatchText += TEXT("\n");
+					}
+
+					FString PinName = UEdGraphSchema_K2::TypeToText(WatchPin->PinType).ToString();
+					PinName += TEXT(" ");
+					PinName += Schema->GetPinDisplayName(WatchPin).ToString();
+
+					FString WatchText;
+					if (PropertyPathHelpers::GetPropertyValueAsString(ActiveObject, WatchPin->PinName.ToString(), WatchText))
+					{
+						PinnedWatchText += FText::Format(LOCTEXT("WatchingAndValidFmt", "Watching {0}\n\t{1}"), FText::FromString(PinName), FText::FromString(WatchText)).ToString();//@TODO: Print out object being debugged name?
+					}
+					else
+					{
+						PinnedWatchText += FText::Format(LOCTEXT("WatchingAndValidFmt", "Invalid Property {0}"), FText::FromString(PinName)).ToString();//@TODO: Print out object being debugged name?
+					}
+
+					ValidWatchCount++;
+				}
+			}
+
+			if (ValidWatchCount)
+			{
+				new (Popups) FGraphInformationPopupInfo(NULL, PinnedWatchColor, PinnedWatchText);
+			}
+		}
+	}
+}
+
 #undef LOCTEXT_NAMESPACE
\ No newline at end of file
diff --git a/Engine/Plugins/Experimental/ControlRig/Source/ControlRigEditor/Private/Graph/SControlRigGraphNode.h b/Engine/Plugins/Experimental/ControlRig/Source/ControlRigEditor/Private/Graph/SControlRigGraphNode.h
--- a/Engine/Plugins/Experimental/ControlRig/Source/ControlRigEditor/Private/Graph/SControlRigGraphNode.h
+++ b/Engine/Plugins/Experimental/ControlRig/Source/ControlRigEditor/Private/Graph/SControlRigGraphNode.h
@@ -39,7 +39,7 @@ public:
 
 	virtual TSharedRef<SWidget> CreateNodeContentArea() override;
 	virtual TSharedPtr<SGraphPin> GetHoveredPin( const FGeometry& MyGeometry, const FPointerEvent& MouseEvent ) const override;
-
+	virtual void GetNodeInfoPopups(FNodeInfoContext* Context, TArray<FGraphInformationPopupInfo>& Popups) const override;
 private:
 	bool ParentUseLowDetailNodeTitles() const
 	{
diff --git a/Engine/Source/Editor/AnimationBlueprintEditor/Private/StateMachineConnectionDrawingPolicy.cpp b/Engine/Source/Editor/AnimationBlueprintEditor/Private/StateMachineConnectionDrawingPolicy.cpp
--- a/Engine/Source/Editor/AnimationBlueprintEditor/Private/StateMachineConnectionDrawingPolicy.cpp
+++ b/Engine/Source/Editor/AnimationBlueprintEditor/Private/StateMachineConnectionDrawingPolicy.cpp
@@ -75,9 +75,9 @@ void FStateMachineConnectionDrawingPolicy::DetermineLinkGeometry(
 	{
 		StartWidgetGeometry = PinGeometries->Find(OutputPinWidget);
 
-		if (TSharedRef<SGraphPin>* pTargetWidget = PinToPinWidgetMap.Find(InputPin))
+		if (TSharedPtr<SGraphPin>* pTargetWidget = PinToPinWidgetMap.Find(InputPin))
 		{
-			TSharedRef<SGraphPin> InputWidget = *pTargetWidget;
+			TSharedRef<SGraphPin> InputWidget = (*pTargetWidget).ToSharedRef();
 			EndWidgetGeometry = PinGeometries->Find(InputWidget);
 		}
 	}
diff --git a/Engine/Source/Editor/GraphEditor/Private/BlueprintConnectionDrawingPolicy.cpp b/Engine/Source/Editor/GraphEditor/Private/BlueprintConnectionDrawingPolicy.cpp
--- a/Engine/Source/Editor/GraphEditor/Private/BlueprintConnectionDrawingPolicy.cpp
+++ b/Engine/Source/Editor/GraphEditor/Private/BlueprintConnectionDrawingPolicy.cpp
@@ -363,9 +363,9 @@ FKismetConnectionDrawingPolicy::FTimePair const* FKismetConnectionDrawingPolicy:
 
 bool FKismetConnectionDrawingPolicy::FindPinCenter(UEdGraphPin* Pin, FVector2D& OutCenter) const
 {
-	if (const TSharedRef<SGraphPin>* pPinWidget = PinToPinWidgetMap.Find(Pin))
+	if (const TSharedPtr<SGraphPin>* pPinWidget = PinToPinWidgetMap.Find(Pin))
 	{
-		if (FArrangedWidget* pPinEntry = PinGeometries->Find(*pPinWidget))
+		if (FArrangedWidget* pPinEntry = PinGeometries->Find((*pPinWidget).ToSharedRef()))
 		{
 			OutCenter = FGeometryHelper::CenterOf(pPinEntry->Geometry);
 			return true;
diff --git a/Engine/Source/Editor/GraphEditor/Private/ConnectionDrawingPolicy.cpp b/Engine/Source/Editor/GraphEditor/Private/ConnectionDrawingPolicy.cpp
--- a/Engine/Source/Editor/GraphEditor/Private/ConnectionDrawingPolicy.cpp
+++ b/Engine/Source/Editor/GraphEditor/Private/ConnectionDrawingPolicy.cpp
@@ -424,9 +424,9 @@ void FConnectionDrawingPolicy::DetermineLinkGeometry(
 {
 	StartWidgetGeometry = PinGeometries->Find(OutputPinWidget);
 	
-	if (TSharedRef<SGraphPin>* pTargetWidget = PinToPinWidgetMap.Find(InputPin))
+	if (TSharedPtr<SGraphPin>* pTargetWidget = PinToPinWidgetMap.Find(InputPin))
 	{
-		TSharedRef<SGraphPin> InputWidget = *pTargetWidget;
+		TSharedRef<SGraphPin> InputWidget = (*pTargetWidget).ToSharedRef();
 		EndWidgetGeometry = PinGeometries->Find(InputWidget);
 	}
 }
diff --git a/Engine/Source/Editor/GraphEditor/Private/MaterialGraphConnectionDrawingPolicy.cpp b/Engine/Source/Editor/GraphEditor/Private/MaterialGraphConnectionDrawingPolicy.cpp
--- a/Engine/Source/Editor/GraphEditor/Private/MaterialGraphConnectionDrawingPolicy.cpp
+++ b/Engine/Source/Editor/GraphEditor/Private/MaterialGraphConnectionDrawingPolicy.cpp
@@ -26,9 +26,9 @@ FMaterialGraphConnectionDrawingPolicy::FMaterialGraphConnectionDrawingPolicy(int
 
 bool FMaterialGraphConnectionDrawingPolicy::FindPinCenter(UEdGraphPin* Pin, FVector2D& OutCenter) const
 {
-	if (const TSharedRef<SGraphPin>* pPinWidget = PinToPinWidgetMap.Find(Pin))
+	if (const TSharedPtr<SGraphPin>* pPinWidget = PinToPinWidgetMap.Find(Pin))
 	{
-		if (FArrangedWidget* pPinEntry = PinGeometries->Find(*pPinWidget))
+		if (FArrangedWidget* pPinEntry = PinGeometries->Find((*pPinWidget).ToSharedRef()))
 		{
 			OutCenter = FGeometryHelper::CenterOf(pPinEntry->Geometry);
 			return true;
diff --git a/Engine/Source/Editor/GraphEditor/Public/ConnectionDrawingPolicy.h b/Engine/Source/Editor/GraphEditor/Public/ConnectionDrawingPolicy.h
--- a/Engine/Source/Editor/GraphEditor/Public/ConnectionDrawingPolicy.h
+++ b/Engine/Source/Editor/GraphEditor/Public/ConnectionDrawingPolicy.h
@@ -91,7 +91,7 @@ protected:
 	float HoverDeemphasisDarkFraction;
 	const FSlateRect& ClippingRect;
 	FSlateWindowElementList& DrawElementsList;
-	TMap< UEdGraphPin*, TSharedRef<SGraphPin> > PinToPinWidgetMap;
+	TMap< UEdGraphPin*, TSharedPtr<SGraphPin> > PinToPinWidgetMap;
 	TSet< FEdGraphPinReference > HoveredPins;
 	TMap<TSharedRef<SWidget>, FArrangedWidget>* PinGeometries;
 	double LastHoverTimeEvent;
