Control Rig: Fixing Cook Warning, Variable initialization and missing image.

diff --git a/Engine/Plugins/Experimental/ControlRig/Source/ControlRig/Private/Units/Math/RigUnit_MathFloat.h b/Engine/Plugins/Experimental/ControlRig/Source/ControlRig/Private/Units/Math/RigUnit_MathFloat.h
--- a/Engine/Plugins/Experimental/ControlRig/Source/ControlRig/Private/Units/Math/RigUnit_MathFloat.h
+++ b/Engine/Plugins/Experimental/ControlRig/Source/ControlRig/Private/Units/Math/RigUnit_MathFloat.h
@@ -16,6 +16,11 @@ struct FRigUnit_MathFloatConstant : public FRigUnit_MathFloatBase
 {
 	GENERATED_BODY()
 
+	FRigUnit_MathFloatConstant()
+	{
+		Value = 0.f;
+	}
+
 	UPROPERTY(meta=(Output, Constant))
 	float Value;
 };
@@ -25,6 +30,11 @@ struct FRigUnit_MathFloatUnaryOp : public FRigUnit_MathFloatBase
 {
 	GENERATED_BODY()
 
+	FRigUnit_MathFloatUnaryOp()
+	{
+		Value = Result = 0.f;
+	}
+
 	UPROPERTY(meta=(Input))
 	float Value;
 
@@ -37,6 +47,11 @@ struct FRigUnit_MathFloatBinaryOp : public FRigUnit_MathFloatBase
 {
 	GENERATED_BODY()
 
+	FRigUnit_MathFloatBinaryOp()
+	{
+		A = B = Result = 0.f;
+	}
+
 	UPROPERTY(meta=(Input))
 	float A;
 
@@ -97,12 +112,6 @@ struct FRigUnit_MathFloatAdd : public FRigUnit_MathFloatBinaryOp
 {
 	GENERATED_BODY()
 
-	FRigUnit_MathFloatAdd()
-	{
-		A = 0.f;
-		B = 0.f;
-	}
-
 	virtual void Execute(const FRigUnitContext& Context) override;
 };
 
@@ -114,12 +123,6 @@ struct FRigUnit_MathFloatSub : public FRigUnit_MathFloatBinaryOp
 {
 	GENERATED_BODY()
 
-	FRigUnit_MathFloatSub()
-	{
-		A = 0.f;
-		B = 0.f;
-	}
-
 	virtual void Execute(const FRigUnitContext& Context) override;
 };
 
@@ -150,7 +153,6 @@ struct FRigUnit_MathFloatDiv : public FRigUnit_MathFloatBinaryOp
 
 	FRigUnit_MathFloatDiv()
 	{
-		A = 0.f;
 		B = 1.f;
 	}
 
@@ -165,12 +167,6 @@ struct FRigUnit_MathFloatMin : public FRigUnit_MathFloatBinaryOp
 {
 	GENERATED_BODY()
 
-	FRigUnit_MathFloatMin()
-	{
-		A = 0.f;
-		B = 0.f;
-	}
-
 	virtual void Execute(const FRigUnitContext& Context) override;
 };
 
@@ -182,12 +178,6 @@ struct FRigUnit_MathFloatMax : public FRigUnit_MathFloatBinaryOp
 {
 	GENERATED_BODY()
 
-	FRigUnit_MathFloatMax()
-	{
-		A = 0.f;
-		B = 0.f;
-	}
-
 	virtual void Execute(const FRigUnitContext& Context) override;
 };
 
@@ -216,11 +206,6 @@ struct FRigUnit_MathFloatSqrt : public FRigUnit_MathFloatUnaryOp
 {
 	GENERATED_BODY()
 
-	FRigUnit_MathFloatSqrt()
-	{
-		Value = 0.f;
-	}
-
 	virtual void Execute(const FRigUnitContext& Context) override;
 };
 
@@ -232,11 +217,6 @@ struct FRigUnit_MathFloatNegate : public FRigUnit_MathFloatUnaryOp
 {
 	GENERATED_BODY()
 
-	FRigUnit_MathFloatNegate()
-	{
-		Value = 0.f;
-	}
-
 	virtual void Execute(const FRigUnitContext& Context) override;
 };
 
@@ -248,11 +228,6 @@ struct FRigUnit_MathFloatAbs : public FRigUnit_MathFloatUnaryOp
 {
 	GENERATED_BODY()
 
-	FRigUnit_MathFloatAbs()
-	{
-		Value = 0.f;
-	}
-
 	virtual void Execute(const FRigUnitContext& Context) override;
 };
 
@@ -264,11 +239,6 @@ struct FRigUnit_MathFloatFloor : public FRigUnit_MathFloatUnaryOp
 {
 	GENERATED_BODY()
 
-	FRigUnit_MathFloatFloor()
-	{
-		Value = 0.f;
-	}
-
 	virtual void Execute(const FRigUnitContext& Context) override;
 };
 
@@ -280,11 +250,6 @@ struct FRigUnit_MathFloatCeil : public FRigUnit_MathFloatUnaryOp
 {
 	GENERATED_BODY()
 
-	FRigUnit_MathFloatCeil()
-	{
-		Value = 0.f;
-	}
-
 	virtual void Execute(const FRigUnitContext& Context) override;
 };
 
@@ -296,11 +261,6 @@ struct FRigUnit_MathFloatRound : public FRigUnit_MathFloatUnaryOp
 {
 	GENERATED_BODY()
 
-	FRigUnit_MathFloatRound()
-	{
-		Value = 0.f;
-	}
-
 	virtual void Execute(const FRigUnitContext& Context) override;
 };
 
@@ -312,11 +272,6 @@ struct FRigUnit_MathFloatSign : public FRigUnit_MathFloatUnaryOp
 {
 	GENERATED_BODY()
 
-	FRigUnit_MathFloatSign()
-	{
-		Value = 0.f;
-	}
-
 	virtual void Execute(const FRigUnitContext& Context) override;
 };
 
@@ -330,8 +285,7 @@ struct FRigUnit_MathFloatClamp : public FRigUnit_MathFloatBase
 
 	FRigUnit_MathFloatClamp()
 	{
-		Minimum = 0.f;
-		Maximum = 1.f;
+		Value = Minimum = Maximum = Result = 0.f;
 	}
 
 	virtual void Execute(const FRigUnitContext& Context) override;
@@ -359,7 +313,7 @@ struct FRigUnit_MathFloatLerp : public FRigUnit_MathFloatBase
 
 	FRigUnit_MathFloatLerp()
 	{
-		A = 0.f;
+		A = B = T = Result = 0.f;
 		B = 1.f;
 	}
 
@@ -388,10 +342,9 @@ struct FRigUnit_MathFloatRemap : public FRigUnit_MathFloatBase
 
 	FRigUnit_MathFloatRemap()
 	{
-		SourceMinimum = 0.f;
-		SourceMaximum = 1.f;
-		TargetMinimum = 0.f;
-		TargetMaximum = 1.f;
+		Value = SourceMinimum = TargetMinimum = Result = 0.f;
+		SourceMaximum = TargetMaximum = 1.f;
+		bClamp = false;
 	}
 
 	virtual void Execute(const FRigUnitContext& Context) override;
@@ -428,6 +381,12 @@ struct FRigUnit_MathFloatEquals : public FRigUnit_MathFloatBase
 	GENERATED_BODY()
 	virtual void Execute(const FRigUnitContext& Context) override;
 
+	FRigUnit_MathFloatEquals()
+	{
+		A = B = 0.f;
+		Result = true;
+	}
+
 	UPROPERTY(meta=(Input))
 	float A;
 
@@ -447,6 +406,12 @@ struct FRigUnit_MathFloatNotEquals : public FRigUnit_MathFloatBase
 	GENERATED_BODY()
 	virtual void Execute(const FRigUnitContext& Context) override;
 
+	FRigUnit_MathFloatNotEquals()
+	{
+		A = B = 0.f;
+		Result = false;
+	}
+
 	UPROPERTY(meta=(Input))
 	float A;
 
@@ -464,6 +429,13 @@ USTRUCT(meta=(DisplayName="Greater", PrototypeName="Greater", Keywords="Larger,B
 struct FRigUnit_MathFloatGreater : public FRigUnit_MathFloatBase
 {
 	GENERATED_BODY()
+
+	FRigUnit_MathFloatGreater()
+	{
+		A = B = 0.f;
+		Result = false;
+	}
+
 	virtual void Execute(const FRigUnitContext& Context) override;
 
 	UPROPERTY(meta=(Input))
@@ -483,6 +455,13 @@ USTRUCT(meta=(DisplayName="Less", PrototypeName="Less", Keywords="Smaller,<"))
 struct FRigUnit_MathFloatLess : public FRigUnit_MathFloatBase
 {
 	GENERATED_BODY()
+	
+	FRigUnit_MathFloatLess()
+	{
+		A = B = 0.f;
+		Result = false;
+	}
+
 	virtual void Execute(const FRigUnitContext& Context) override;
 
 	UPROPERTY(meta=(Input))
@@ -502,6 +481,13 @@ USTRUCT(meta=(DisplayName="Greater Equal", PrototypeName="GreaterEqual", Keyword
 struct FRigUnit_MathFloatGreaterEqual : public FRigUnit_MathFloatBase
 {
 	GENERATED_BODY()
+
+	FRigUnit_MathFloatGreaterEqual()
+	{
+		A = B = 0.f;
+		Result = true;
+	}
+
 	virtual void Execute(const FRigUnitContext& Context) override;
 
 	UPROPERTY(meta=(Input))
@@ -521,6 +507,13 @@ USTRUCT(meta=(DisplayName="Less Equal", PrototypeName="LessEqual", Keywords="Sma
 struct FRigUnit_MathFloatLessEqual : public FRigUnit_MathFloatBase
 {
 	GENERATED_BODY()
+
+	FRigUnit_MathFloatLessEqual()
+	{
+		A = B = 0.f;
+		Result = true;
+	}
+
 	virtual void Execute(const FRigUnitContext& Context) override;
 
 	UPROPERTY(meta=(Input))
@@ -543,7 +536,8 @@ struct FRigUnit_MathFloatIsNearlyZero : public FRigUnit_MathFloatBase
 	
 	FRigUnit_MathFloatIsNearlyZero()
 	{
-		Tolerance = 0.f;
+		Value = Tolerance = 0.f;
+		Result = true;
 	}
 	
 	virtual void Execute(const FRigUnitContext& Context) override;
@@ -568,7 +562,8 @@ struct FRigUnit_MathFloatIsNearlyEqual : public FRigUnit_MathFloatBase
 
 	FRigUnit_MathFloatIsNearlyEqual()
 	{
-		Tolerance = 0.f;
+		A = B = Tolerance = 0.f;
+		Result = true;
 	}
 
 	virtual void Execute(const FRigUnitContext& Context) override;
@@ -593,6 +588,13 @@ USTRUCT(meta=(DisplayName="Select", PrototypeName="Select", Keywords="Pick,If"))
 struct FRigUnit_MathFloatSelectBool : public FRigUnit_MathFloatBase
 {
 	GENERATED_BODY()
+
+	FRigUnit_MathFloatSelectBool()
+	{
+		Condition = false;
+		IfTrue = IfFalse = Result = 0.f;
+	}
+
 	virtual void Execute(const FRigUnitContext& Context) override;
 
 	UPROPERTY(meta=(Input))
@@ -616,12 +618,6 @@ struct FRigUnit_MathFloatDeg : public FRigUnit_MathFloatUnaryOp
 {
 	GENERATED_BODY()
 
-	FRigUnit_MathFloatDeg()
-	{
-		Value = 0.f;
-	}
-
-
 	virtual void Execute(const FRigUnitContext& Context) override;
 };
 
@@ -633,12 +629,6 @@ struct FRigUnit_MathFloatRad : public FRigUnit_MathFloatUnaryOp
 {
 	GENERATED_BODY()
 
-	FRigUnit_MathFloatRad()
-	{
-		Value = 0.f;
-	}
-
-
 	virtual void Execute(const FRigUnitContext& Context) override;
 };
 
@@ -650,12 +640,6 @@ struct FRigUnit_MathFloatSin : public FRigUnit_MathFloatUnaryOp
 {
 	GENERATED_BODY()
 
-	FRigUnit_MathFloatSin()
-	{
-		Value = 0.f;
-	}
-
-
 	virtual void Execute(const FRigUnitContext& Context) override;
 };
 
@@ -667,12 +651,6 @@ struct FRigUnit_MathFloatCos : public FRigUnit_MathFloatUnaryOp
 {
 	GENERATED_BODY()
 
-	FRigUnit_MathFloatCos()
-	{
-		Value = 0.f;
-	}
-
-
 	virtual void Execute(const FRigUnitContext& Context) override;
 };
 
@@ -684,12 +662,6 @@ struct FRigUnit_MathFloatTan : public FRigUnit_MathFloatUnaryOp
 {
 	GENERATED_BODY()
 
-	FRigUnit_MathFloatTan()
-	{
-		Value = 0.f;
-	}
-
-
 	virtual void Execute(const FRigUnitContext& Context) override;
 };
 
@@ -701,12 +673,6 @@ struct FRigUnit_MathFloatAsin : public FRigUnit_MathFloatUnaryOp
 {
 	GENERATED_BODY()
 
-	FRigUnit_MathFloatAsin()
-	{
-		Value = 0.f;
-	}
-
-
 	virtual void Execute(const FRigUnitContext& Context) override;
 };
 
@@ -718,12 +684,6 @@ struct FRigUnit_MathFloatAcos : public FRigUnit_MathFloatUnaryOp
 {
 	GENERATED_BODY()
 
-	FRigUnit_MathFloatAcos()
-	{
-		Value = 0.f;
-	}
-
-
 	virtual void Execute(const FRigUnitContext& Context) override;
 };
 
@@ -735,12 +695,6 @@ struct FRigUnit_MathFloatAtan : public FRigUnit_MathFloatUnaryOp
 {
 	GENERATED_BODY()
 
-	FRigUnit_MathFloatAtan()
-	{
-		Value = 0.f;
-	}
-
-
 	virtual void Execute(const FRigUnitContext& Context) override;
 };
 
diff --git a/Engine/Plugins/Experimental/ControlRig/Source/ControlRig/Private/Units/Math/RigUnit_MathQuaternion.h b/Engine/Plugins/Experimental/ControlRig/Source/ControlRig/Private/Units/Math/RigUnit_MathQuaternion.h
--- a/Engine/Plugins/Experimental/ControlRig/Source/ControlRig/Private/Units/Math/RigUnit_MathQuaternion.h
+++ b/Engine/Plugins/Experimental/ControlRig/Source/ControlRig/Private/Units/Math/RigUnit_MathQuaternion.h
@@ -18,7 +18,7 @@ struct FRigUnit_MathQuaternionUnaryOp : public FRigUnit_MathQuaternionBase
 
 	FRigUnit_MathQuaternionUnaryOp()
 	{
-		Value = FQuat::Identity;
+		Value = Result = FQuat::Identity;
 	}
 
 	UPROPERTY(meta=(Input))
@@ -35,7 +35,7 @@ struct FRigUnit_MathQuaternionBinaryOp : public FRigUnit_MathQuaternionBase
 
 	FRigUnit_MathQuaternionBinaryOp()
 	{
-		A = B = FQuat::Identity;
+		A = B = Result = FQuat::Identity;
 	}
 
 	UPROPERTY(meta=(Input))
@@ -59,6 +59,8 @@ struct FRigUnit_MathQuaternionFromAxisAndAngle : public FRigUnit_MathQuaternionB
 	FRigUnit_MathQuaternionFromAxisAndAngle()
 	{
 		Axis = FVector(1.f, 0.f, 0.f);
+		Angle = 0.f;
+		Result = FQuat::Identity;
 	}
 
 	virtual void Execute(const FRigUnitContext& Context) override;
@@ -80,6 +82,13 @@ USTRUCT(meta=(DisplayName="From Euler", PrototypeName="FromEuler", Keywords="Mak
 struct FRigUnit_MathQuaternionFromEuler : public FRigUnit_MathQuaternionBase
 {
 	GENERATED_BODY()
+
+	FRigUnit_MathQuaternionFromEuler()
+	{
+		Euler = FVector();
+		Result = FQuat::Identity;
+	}
+
 	virtual void Execute(const FRigUnitContext& Context) override;
 
 	UPROPERTY(meta=(Input))
@@ -96,6 +105,13 @@ USTRUCT(meta=(DisplayName="From Rotator", PrototypeName="FromRotator", Keywords=
 struct FRigUnit_MathQuaternionFromRotator : public FRigUnit_MathQuaternionBase
 {
 	GENERATED_BODY()
+	
+	FRigUnit_MathQuaternionFromRotator()
+	{
+		Rotator = FRotator();
+		Result = FQuat::Identity;
+	}
+
 	virtual void Execute(const FRigUnitContext& Context) override;
 
 	UPROPERTY(meta=(Input))
@@ -116,6 +132,7 @@ struct FRigUnit_MathQuaternionFromTwoVectors : public FRigUnit_MathQuaternionBas
 	FRigUnit_MathQuaternionFromTwoVectors()
 	{
 		A = B = FVector(1.f, 0.f, 0.f);
+		Result = FQuat::Identity;
 	}
 
 	virtual void Execute(const FRigUnitContext& Context) override;
@@ -142,6 +159,7 @@ struct FRigUnit_MathQuaternionToAxisAndAngle : public FRigUnit_MathQuaternionBas
 	{
 		Value = FQuat::Identity;
 		Axis = FVector(1.f, 0.f, 0.f);
+		Angle = 0.f;
 	}
 
 	virtual void Execute(const FRigUnitContext& Context) override;
@@ -167,6 +185,7 @@ struct FRigUnit_MathQuaternionToEuler : public FRigUnit_MathQuaternionBase
 	FRigUnit_MathQuaternionToEuler()
 	{
 		Value = FQuat::Identity;
+		Result = FVector();
 	}
 
 	virtual void Execute(const FRigUnitContext& Context) override;
@@ -189,6 +208,7 @@ struct FRigUnit_MathQuaternionToRotator : public FRigUnit_MathQuaternionBase
 	FRigUnit_MathQuaternionToRotator()
 	{
 		Value = FQuat::Identity;
+		Result = FRotator();
 	}
 
 	virtual void Execute(const FRigUnitContext& Context) override;
@@ -230,7 +250,8 @@ struct FRigUnit_MathQuaternionSlerp : public FRigUnit_MathQuaternionBase
 
 	FRigUnit_MathQuaternionSlerp()
 	{
-		A = B = FQuat::Identity;
+		A = B = Result = FQuat::Identity;
+		T = 0.f;
 	}
 
 	virtual void Execute(const FRigUnitContext& Context) override;
@@ -259,7 +280,8 @@ struct FRigUnit_MathQuaternionEquals : public FRigUnit_MathQuaternionBase
 	FRigUnit_MathQuaternionEquals()
 	{
 		A = B = FQuat::Identity;
-	}
+		Result = true;
+	}	
 
 	virtual void Execute(const FRigUnitContext& Context) override;
 
@@ -284,6 +306,7 @@ struct FRigUnit_MathQuaternionNotEquals : public FRigUnit_MathQuaternionBase
 	FRigUnit_MathQuaternionNotEquals()
 	{
 		A = B = FQuat::Identity;
+		Result = false;
 	}
 
 	virtual void Execute(const FRigUnitContext& Context) override;
@@ -308,7 +331,8 @@ struct FRigUnit_MathQuaternionSelectBool : public FRigUnit_MathQuaternionBase
 
 	FRigUnit_MathQuaternionSelectBool()
 	{
-		IfTrue = IfFalse = FQuat::Identity;
+		IfTrue = IfFalse = Result = FQuat::Identity;
+		Condition = true;
 	}
 
 	virtual void Execute(const FRigUnitContext& Context) override;
@@ -337,6 +361,7 @@ struct FRigUnit_MathQuaternionDot : public FRigUnit_MathQuaternionBase
 	FRigUnit_MathQuaternionDot()
 	{
 		A = B = FQuat::Identity;
+		Result = 0.f;
 	}
 
 	virtual void Execute(const FRigUnitContext& Context) override;
@@ -372,6 +397,7 @@ struct FRigUnit_MathQuaternionRotateVector : public FRigUnit_MathQuaternionBase
 	FRigUnit_MathQuaternionRotateVector()
 	{
 		Quaternion = FQuat::Identity;
+		Vector = Result = FVector();
 	}
 
 	virtual void Execute(const FRigUnitContext& Context) override;
@@ -398,6 +424,7 @@ struct FRigUnit_MathQuaternionGetAxis: public FRigUnit_MathQuaternionBase
 	{
 		Quaternion = FQuat::Identity;
 		Axis = EAxis::X;
+		Result = FVector(1.f, 0.f, 0.f);
 	}
 
 	virtual void Execute(const FRigUnitContext& Context) override;
diff --git a/Engine/Plugins/Experimental/ControlRig/Source/ControlRig/Private/Units/Math/RigUnit_MathTransform.h b/Engine/Plugins/Experimental/ControlRig/Source/ControlRig/Private/Units/Math/RigUnit_MathTransform.h
--- a/Engine/Plugins/Experimental/ControlRig/Source/ControlRig/Private/Units/Math/RigUnit_MathTransform.h
+++ b/Engine/Plugins/Experimental/ControlRig/Source/ControlRig/Private/Units/Math/RigUnit_MathTransform.h
@@ -17,6 +17,11 @@ struct FRigUnit_MathTransformUnaryOp : public FRigUnit_MathTransformBase
 {
 	GENERATED_BODY()
 
+	FRigUnit_MathTransformUnaryOp()
+	{
+		Value = Result = FTransform::Identity;
+	}
+
 	UPROPERTY(meta=(Input))
 	FTransform Value;
 
@@ -29,6 +34,11 @@ struct FRigUnit_MathTransformBinaryOp : public FRigUnit_MathTransformBase
 {
 	GENERATED_BODY()
 
+	FRigUnit_MathTransformBinaryOp()
+	{
+		A = B = Result = FTransform::Identity;
+	}
+
 	UPROPERTY(meta=(Input))
 	FTransform A;
 
@@ -46,6 +56,13 @@ USTRUCT(meta=(DisplayName="From Euler Transform", PrototypeName="FromEulerTransf
 struct FRigUnit_MathTransformFromEulerTransform : public FRigUnit_MathTransformBase
 {
 	GENERATED_BODY()
+
+	FRigUnit_MathTransformFromEulerTransform()
+	{
+		EulerTransform = FEulerTransform::Identity;
+		Result = FTransform::Identity;
+	}
+
 	virtual void Execute(const FRigUnitContext& Context) override;
 
 	UPROPERTY(meta=(Input))
@@ -62,6 +79,13 @@ USTRUCT(meta=(DisplayName="To Euler Transform", PrototypeName="ToEulerTransform"
 struct FRigUnit_MathTransformToEulerTransform : public FRigUnit_MathTransformBase
 {
 	GENERATED_BODY()
+
+	FRigUnit_MathTransformToEulerTransform()
+	{
+		Value = FTransform::Identity;
+		Result = FEulerTransform::Identity;
+	}
+
 	virtual void Execute(const FRigUnitContext& Context) override;
 
 	UPROPERTY(meta=(Input))
@@ -88,6 +112,12 @@ USTRUCT(meta=(DisplayName="Make Relative", PrototypeName="MakeRelative", Keyword
 struct FRigUnit_MathTransformMakeRelative : public FRigUnit_MathTransformBase
 {
 	GENERATED_BODY()
+
+	FRigUnit_MathTransformMakeRelative()
+	{
+		Global = Parent = Local = FTransform::Identity;
+	}
+
 	virtual void Execute(const FRigUnitContext& Context) override;
 
 	UPROPERTY(meta=(Input))
@@ -117,6 +147,13 @@ USTRUCT(meta=(DisplayName="Interpolate", PrototypeName="Interpolate", Keywords="
 struct FRigUnit_MathTransformLerp : public FRigUnit_MathTransformBase
 {
 	GENERATED_BODY()
+
+	FRigUnit_MathTransformLerp()
+	{
+		A = B = Result = FTransform::Identity;
+		T = 0.f;
+	}
+
 	virtual void Execute(const FRigUnitContext& Context) override;
 
 	UPROPERTY(meta=(Input))
@@ -139,6 +176,13 @@ USTRUCT(meta=(DisplayName="Select", PrototypeName="Select", Keywords="Pick,If"))
 struct FRigUnit_MathTransformSelectBool : public FRigUnit_MathTransformBase
 {
 	GENERATED_BODY()
+
+	FRigUnit_MathTransformSelectBool()
+	{
+		Condition = false;
+		IfTrue = IfFalse = Result = FTransform::Identity;
+	}
+
 	virtual void Execute(const FRigUnitContext& Context) override;
 
 	UPROPERTY(meta=(Input))
@@ -161,6 +205,13 @@ USTRUCT(meta=(DisplayName="Transform Direction", PrototypeName="Rotate", Keyword
 struct FRigUnit_MathTransformRotateVector : public FRigUnit_MathTransformBase
 {
 	GENERATED_BODY()
+
+	FRigUnit_MathTransformRotateVector()
+	{
+		Transform = FTransform::Identity;
+		Direction = Result = FVector();
+	}
+
 	virtual void Execute(const FRigUnitContext& Context) override;
 
 	UPROPERTY(meta=(Input))
@@ -180,6 +231,13 @@ USTRUCT(meta=(DisplayName="Transform Location", PrototypeName="Multiply"))
 struct FRigUnit_MathTransformTransformVector : public FRigUnit_MathTransformBase
 {
 	GENERATED_BODY()
+
+	FRigUnit_MathTransformTransformVector()
+	{
+		Transform = FTransform::Identity;
+		Location = Result = FVector();
+	}
+	
 	virtual void Execute(const FRigUnitContext& Context) override;
 
 	UPROPERTY(meta=(Input))
diff --git a/Engine/Plugins/Experimental/ControlRig/Source/ControlRig/Private/Units/Math/RigUnit_MathVector.h b/Engine/Plugins/Experimental/ControlRig/Source/ControlRig/Private/Units/Math/RigUnit_MathVector.h
--- a/Engine/Plugins/Experimental/ControlRig/Source/ControlRig/Private/Units/Math/RigUnit_MathVector.h
+++ b/Engine/Plugins/Experimental/ControlRig/Source/ControlRig/Private/Units/Math/RigUnit_MathVector.h
@@ -15,6 +15,11 @@ USTRUCT(meta=(Abstract))
 struct FRigUnit_MathVectorUnaryOp : public FRigUnit_MathVectorBase
 {
 	GENERATED_BODY()
+	
+	FRigUnit_MathVectorUnaryOp()
+	{
+		Value = Result = FVector();
+	}
 
 	UPROPERTY(meta=(Input))
 	FVector Value;
@@ -28,6 +33,11 @@ struct FRigUnit_MathVectorBinaryOp : public FRigUnit_MathVectorBase
 {
 	GENERATED_BODY()
 
+	FRigUnit_MathVectorBinaryOp()
+	{
+		A = B = Result = FVector();
+	}
+
 	UPROPERTY(meta=(Input))
 	FVector A;
 
@@ -47,6 +57,12 @@ struct FRigUnit_MathVectorFromFloat : public FRigUnit_MathVectorBase
 	GENERATED_BODY()
 	virtual void Execute(const FRigUnitContext& Context) override;
 
+	FRigUnit_MathVectorFromFloat()
+	{
+		Value = 0.f;
+		Result = FVector();
+	}
+
 	UPROPERTY(meta=(Input))
 	float Value;
 
@@ -62,12 +78,6 @@ struct FRigUnit_MathVectorAdd : public FRigUnit_MathVectorBinaryOp
 {
 	GENERATED_BODY()
 
-	FRigUnit_MathVectorAdd()
-	{
-		A = FVector(0.f, 0.f, 0.f);
-		B = FVector(0.f, 0.f, 0.f);
-	}
-
 	virtual void Execute(const FRigUnitContext& Context) override;
 };
 
@@ -79,12 +89,6 @@ struct FRigUnit_MathVectorSub : public FRigUnit_MathVectorBinaryOp
 {
 	GENERATED_BODY()
 
-	FRigUnit_MathVectorSub()
-	{
-		A = FVector(0.f, 0.f, 0.f);
-		B = FVector(0.f, 0.f, 0.f);
-	}
-
 	virtual void Execute(const FRigUnitContext& Context) override;
 };
 
@@ -98,8 +102,7 @@ struct FRigUnit_MathVectorMul : public FRigUnit_MathVectorBinaryOp
 
 	FRigUnit_MathVectorMul()
 	{
-		A = FVector(1.f, 1.f, 1.f);
-		B = FVector(1.f, 1.f, 1.f);
+		A = B = FVector(1.f, 1.f, 1.f);
 	}
 
 	virtual void Execute(const FRigUnitContext& Context) override;
@@ -115,7 +118,6 @@ struct FRigUnit_MathVectorDiv : public FRigUnit_MathVectorBinaryOp
 
 	FRigUnit_MathVectorDiv()
 	{
-		A = FVector(0.f, 0.f, 0.f);
 		B = FVector(1.f, 1.f, 1.f);
 	}
 
@@ -130,12 +132,6 @@ struct FRigUnit_MathVectorMin : public FRigUnit_MathVectorBinaryOp
 {
 	GENERATED_BODY()
 
-	FRigUnit_MathVectorMin()
-	{
-		A = FVector(0.f, 0.f, 0.f);
-		B = FVector(0.f, 0.f, 0.f);
-	}
-
 	virtual void Execute(const FRigUnitContext& Context) override;
 };
 
@@ -147,12 +143,6 @@ struct FRigUnit_MathVectorMax : public FRigUnit_MathVectorBinaryOp
 {
 	GENERATED_BODY()
 
-	FRigUnit_MathVectorMax()
-	{
-		A = FVector(0.f, 0.f, 0.f);
-		B = FVector(0.f, 0.f, 0.f);
-	}
-
 	virtual void Execute(const FRigUnitContext& Context) override;
 };
 
@@ -164,11 +154,6 @@ struct FRigUnit_MathVectorNegate : public FRigUnit_MathVectorUnaryOp
 {
 	GENERATED_BODY()
 
-	FRigUnit_MathVectorNegate()
-	{
-		Value = FVector(0.f, 0.f, 0.f);
-	}
-
 	virtual void Execute(const FRigUnitContext& Context) override;
 };
 
@@ -180,11 +165,6 @@ struct FRigUnit_MathVectorAbs : public FRigUnit_MathVectorUnaryOp
 {
 	GENERATED_BODY()
 
-	FRigUnit_MathVectorAbs()
-	{
-		Value = FVector(0.f, 0.f, 0.f);
-	}
-
 	virtual void Execute(const FRigUnitContext& Context) override;
 };
 
@@ -196,11 +176,6 @@ struct FRigUnit_MathVectorFloor : public FRigUnit_MathVectorUnaryOp
 {
 	GENERATED_BODY()
 
-	FRigUnit_MathVectorFloor()
-	{
-		Value = FVector(0.f, 0.f, 0.f);
-	}
-
 	virtual void Execute(const FRigUnitContext& Context) override;
 };
 
@@ -212,11 +187,6 @@ struct FRigUnit_MathVectorCeil : public FRigUnit_MathVectorUnaryOp
 {
 	GENERATED_BODY()
 
-	FRigUnit_MathVectorCeil()
-	{
-		Value = FVector(0.f, 0.f, 0.f);
-	}
-
 	virtual void Execute(const FRigUnitContext& Context) override;
 };
 
@@ -228,11 +198,6 @@ struct FRigUnit_MathVectorRound : public FRigUnit_MathVectorUnaryOp
 {
 	GENERATED_BODY()
 
-	FRigUnit_MathVectorRound()
-	{
-		Value = FVector(0.f, 0.f, 0.f);
-	}
-
 	virtual void Execute(const FRigUnitContext& Context) override;
 };
 
@@ -244,11 +209,6 @@ struct FRigUnit_MathVectorSign : public FRigUnit_MathVectorUnaryOp
 {
 	GENERATED_BODY()
 
-	FRigUnit_MathVectorSign()
-	{
-		Value = FVector(0.f, 0.f, 0.f);
-	}
-
 	virtual void Execute(const FRigUnitContext& Context) override;
 };
 
@@ -262,7 +222,7 @@ struct FRigUnit_MathVectorClamp : public FRigUnit_MathVectorBase
 
 	FRigUnit_MathVectorClamp()
 	{
-		Minimum = FVector(0.f, 0.f, 0.f);
+		Value = Minimum = Result = FVector(0.f, 0.f, 0.f);
 		Maximum = FVector(1.f, 1.f, 1.f);
 	}
 
@@ -291,8 +251,9 @@ struct FRigUnit_MathVectorLerp : public FRigUnit_MathVectorBase
 
 	FRigUnit_MathVectorLerp()
 	{
-		A = FVector(0.f, 0.f, 0.f);
+		A = Result = FVector(0.f, 0.f, 0.f);
 		B = FVector(1.f, 1.f, 1.f);
+		T = 0.f;
 	}
 
 	virtual void Execute(const FRigUnitContext& Context) override;
@@ -320,10 +281,9 @@ struct FRigUnit_MathVectorRemap : public FRigUnit_MathVectorBase
 
 	FRigUnit_MathVectorRemap()
 	{
-		SourceMinimum = FVector(0.f, 0.f, 0.f);
-		SourceMaximum = FVector(1.f, 1.f, 1.f);
-		TargetMinimum = FVector(0.f, 0.f, 0.f);
-		TargetMaximum = FVector(1.f, 1.f, 1.f);
+		Value = SourceMinimum = TargetMinimum = Result = FVector(0.f, 0.f, 0.f);
+		SourceMaximum = TargetMaximum = FVector(1.f, 1.f, 1.f);
+		bClamp = false;
 	}
 
 	virtual void Execute(const FRigUnitContext& Context) override;
@@ -358,6 +318,13 @@ USTRUCT(meta=(DisplayName="Equals", PrototypeName="Equals", Keywords="Same,=="))
 struct FRigUnit_MathVectorEquals : public FRigUnit_MathVectorBase
 {
 	GENERATED_BODY()
+
+	FRigUnit_MathVectorEquals()
+	{
+		A = B = FVector();
+		Result = true;
+	}
+
 	virtual void Execute(const FRigUnitContext& Context) override;
 
 	UPROPERTY(meta=(Input))
@@ -377,6 +344,13 @@ USTRUCT(meta=(DisplayName="Not Equals", PrototypeName="NotEquals", Keywords="Dif
 struct FRigUnit_MathVectorNotEquals : public FRigUnit_MathVectorBase
 {
 	GENERATED_BODY()
+
+	FRigUnit_MathVectorNotEquals()
+	{
+		A = B = FVector();
+		Result = false;
+	}
+
 	virtual void Execute(const FRigUnitContext& Context) override;
 
 	UPROPERTY(meta=(Input))
@@ -399,7 +373,9 @@ struct FRigUnit_MathVectorIsNearlyZero : public FRigUnit_MathVectorBase
 
 	FRigUnit_MathVectorIsNearlyZero()
 	{
+		Value = FVector();
 		Tolerance = 0.f;
+		Result = true;
 	}
 
 	virtual void Execute(const FRigUnitContext& Context) override;
@@ -424,7 +400,9 @@ struct FRigUnit_MathVectorIsNearlyEqual : public FRigUnit_MathVectorBase
 
 	FRigUnit_MathVectorIsNearlyEqual()
 	{
+		A = B = FVector();
 		Tolerance = 0.f;
+		Result = true;
 	}
 
 	virtual void Execute(const FRigUnitContext& Context) override;
@@ -449,6 +427,13 @@ USTRUCT(meta=(DisplayName="Select", PrototypeName="Select", Keywords="Pick,If"))
 struct FRigUnit_MathVectorSelectBool : public FRigUnit_MathVectorBase
 {
 	GENERATED_BODY()
+
+	FRigUnit_MathVectorSelectBool()
+	{
+		Condition = false;
+		IfTrue = IfFalse = Result = FVector();
+	}
+
 	virtual void Execute(const FRigUnitContext& Context) override;
 
 	UPROPERTY(meta=(Input))
@@ -472,12 +457,6 @@ struct FRigUnit_MathVectorDeg : public FRigUnit_MathVectorUnaryOp
 {
 	GENERATED_BODY()
 
-	FRigUnit_MathVectorDeg()
-	{
-		Value = FVector(0.f, 0.f, 0.f);
-	}
-
-
 	virtual void Execute(const FRigUnitContext& Context) override;
 };
 
@@ -489,12 +468,6 @@ struct FRigUnit_MathVectorRad : public FRigUnit_MathVectorUnaryOp
 {
 	GENERATED_BODY()
 
-	FRigUnit_MathVectorRad()
-	{
-		Value = FVector(0.f, 0.f, 0.f);
-	}
-
-
 	virtual void Execute(const FRigUnitContext& Context) override;
 };
 
@@ -505,6 +478,13 @@ USTRUCT(meta=(DisplayName="Length Squared", PrototypeName="LengthSquared", Keywo
 struct FRigUnit_MathVectorLengthSquared : public FRigUnit_MathVectorBase
 {
 	GENERATED_BODY()
+
+	FRigUnit_MathVectorLengthSquared()
+	{
+		Value = FVector();
+		Result = 0.f;
+	}
+
 	virtual void Execute(const FRigUnitContext& Context) override;
 
 	UPROPERTY(meta=(Input))
@@ -521,6 +501,13 @@ USTRUCT(meta=(DisplayName="Length", PrototypeName="Length", Keywords="Size,Magni
 struct FRigUnit_MathVectorLength : public FRigUnit_MathVectorBase
 {
 	GENERATED_BODY()
+
+	FRigUnit_MathVectorLength()
+	{
+		Value = FVector();
+		Result = 0.f;
+	}
+
 	virtual void Execute(const FRigUnitContext& Context) override;
 
 	UPROPERTY(meta=(Input))
@@ -537,6 +524,13 @@ USTRUCT(meta=(DisplayName="Distance Between", PrototypeName="Distance"))
 struct FRigUnit_MathVectorDistance : public FRigUnit_MathVectorBase
 {
 	GENERATED_BODY()
+
+	FRigUnit_MathVectorDistance()
+	{
+		A = B = FVector();
+		Result = 0.f;
+	}
+
 	virtual void Execute(const FRigUnitContext& Context) override;
 
 	UPROPERTY(meta=(Input))
@@ -566,6 +560,13 @@ USTRUCT(meta=(DisplayName="Dot", PrototypeName="Dot,|"))
 struct FRigUnit_MathVectorDot : public FRigUnit_MathVectorBase
 {
 	GENERATED_BODY()
+
+	FRigUnit_MathVectorDot()
+	{
+		A = B = FVector();
+		Result = 0.f;
+	}
+
 	virtual void Execute(const FRigUnitContext& Context) override;
 
 	UPROPERTY(meta=(Input))
@@ -586,11 +587,6 @@ struct FRigUnit_MathVectorUnit : public FRigUnit_MathVectorUnaryOp
 {
 	GENERATED_BODY()
 
-	FRigUnit_MathVectorUnit()
-	{
-		Value = FVector(0.f, 0.f, 0.f);
-	}
-
 	virtual void Execute(const FRigUnitContext& Context) override;
 };
 
@@ -601,6 +597,13 @@ USTRUCT(meta=(DisplayName="Mirror", PrototypeName="Mirror"))
 struct FRigUnit_MathVectorMirror : public FRigUnit_MathVectorBase
 {
 	GENERATED_BODY()
+
+	FRigUnit_MathVectorMirror()
+	{
+		Value = Result = FVector();
+		Normal = FVector(1.f, 0.f, 0.f);
+	}
+
 	virtual void Execute(const FRigUnitContext& Context) override;
 
 	UPROPERTY(meta=(Input))
@@ -620,6 +623,13 @@ USTRUCT(meta=(DisplayName="Angle Between", PrototypeName="AngleBetween"))
 struct FRigUnit_MathVectorAngle : public FRigUnit_MathVectorBase
 {
 	GENERATED_BODY()
+
+		FRigUnit_MathVectorAngle()
+	{
+		A = B = FVector();
+		Result = 0.f;
+	}
+
 	virtual void Execute(const FRigUnitContext& Context) override;
 
 	UPROPERTY(meta=(Input))
@@ -639,6 +649,13 @@ USTRUCT(meta=(DisplayName="Parallel", PrototypeName="Parallel"))
 struct FRigUnit_MathVectorParallel : public FRigUnit_MathVectorBase
 {
 	GENERATED_BODY()
+
+	FRigUnit_MathVectorParallel()
+	{
+		A = B = FVector(1.f, 0.f, 0.f);
+		Result = true;
+	}
+
 	virtual void Execute(const FRigUnitContext& Context) override;
 
 	UPROPERTY(meta=(Input))
@@ -658,6 +675,14 @@ USTRUCT(meta=(DisplayName="Orthogonal", PrototypeName="Orthogonal"))
 struct FRigUnit_MathVectorOrthogonal : public FRigUnit_MathVectorBase
 {
 	GENERATED_BODY()
+
+	FRigUnit_MathVectorOrthogonal()
+	{
+		A = FVector(1.f, 0.f, 0.f);
+		B = FVector(0.f, 1.f, 0.f);
+		Result = true;
+	}
+
 	virtual void Execute(const FRigUnitContext& Context) override;
 
 	UPROPERTY(meta=(Input))
@@ -677,6 +702,13 @@ USTRUCT(meta=(DisplayName="Bezier Four Point"))
 struct FRigUnit_MathVectorBezierFourPoint : public FRigUnit_MathVectorBase
 {
 	GENERATED_BODY()
+
+	FRigUnit_MathVectorBezierFourPoint()
+	{
+		A = B = C = D = Result = Tangent = FVector();
+		T = 0.f;
+	}
+
 	virtual void Execute(const FRigUnitContext& Context) override;
 
 	UPROPERTY(meta=(Input))
