Control Rig: Fixed linux build break due to case issue in include

diff --git a/Engine/Plugins/Experimental/ControlRig/Source/ControlRigEditor/Private/Editor/SControlRigStackView.h b/Engine/Plugins/Experimental/ControlRig/Source/ControlRigEditor/Private/Editor/SControlRigStackView.h
--- a/Engine/Plugins/Experimental/ControlRig/Source/ControlRigEditor/Private/Editor/SControlRigStackView.h
+++ b/Engine/Plugins/Experimental/ControlRig/Source/ControlRigEditor/Private/Editor/SControlRigStackView.h
@@ -5,7 +5,7 @@
 #include "CoreMinimal.h"
 #include "Graph/ControlRigGraph.h"
 #include "ControlRig.h"
-#include "ControlRigblueprint.h"
+#include "ControlRigBlueprint.h"
 #include "Widgets/Views/STreeView.h"
 
 class FControlRigEditor;
