Control Rig - Improving Settings on Spring IK Unit

diff --git a/Engine/Plugins/Experimental/ControlRig/Source/ControlRig/Private/Units/Highlevel/Hierarchy/RigUnit_SpringIK.cpp b/Engine/Plugins/Experimental/ControlRig/Source/ControlRig/Private/Units/Highlevel/Hierarchy/RigUnit_SpringIK.cpp
--- a/Engine/Plugins/Experimental/ControlRig/Source/ControlRig/Private/Units/Highlevel/Hierarchy/RigUnit_SpringIK.cpp
+++ b/Engine/Plugins/Experimental/ControlRig/Source/ControlRig/Private/Units/Highlevel/Hierarchy/RigUnit_SpringIK.cpp
@@ -65,7 +65,7 @@ void FRigUnit_SpringIK::Execute(const FRigUnitContext& Context)
 				B = Hierarchy->GetInitialTransform(BoneIndices[0]);
 				Spring.SubjectA = PointIndex;
 				Spring.SubjectB = 0;
-				Spring.Equilibrium = (A.GetLocation() - B.GetLocation()).Size();
+				Spring.Equilibrium = FMath::Lerp<float>(0.f, (A.GetLocation() - B.GetLocation()).Size(), FMath::Clamp<float>(RootRatio, 0.f, 1.f));
 				Spring.Coefficient = RootStrength;
 				Simulation.Springs.Add(Spring);
 			}
@@ -76,7 +76,7 @@ void FRigUnit_SpringIK::Execute(const FRigUnitContext& Context)
 				B = Hierarchy->GetInitialTransform(BoneIndices[BoneIndices.Num() - 1]);
 				Spring.SubjectA = PointIndex;
 				Spring.SubjectB = BoneIndices.Num() - 1;
-				Spring.Equilibrium = (A.GetLocation() - B.GetLocation()).Size();
+				Spring.Equilibrium = FMath::Lerp<float>(0.f, (A.GetLocation() - B.GetLocation()).Size(), FMath::Clamp<float>(EffectorRatio, 0.f, 1.f));
 				Spring.Coefficient = EffectorStrength;
 				Simulation.Springs.Add(Spring);
 			}
@@ -201,7 +201,7 @@ void FRigUnit_SpringIK::Execute(const FRigUnitContext& Context)
 		}
 	}
 
-	Simulation.StepSemiExplicitEuler(bLiveSimulation ? Context.DeltaTime : Simulation.TimeStep * (float)FMath::Clamp<int32>(Iterations, 1, 32));
+	Simulation.StepSemiExplicitEuler((bLiveSimulation ? Context.DeltaTime : Simulation.TimeStep) * (float)FMath::Clamp<int32>(Iterations, 1, 32));
 
 	FVector AccumulatedTarget = FVector::ZeroVector;
 	FVector LastPrimaryTarget = FVector::ZeroVector;
@@ -213,6 +213,7 @@ void FRigUnit_SpringIK::Execute(const FRigUnitContext& Context)
 		{
 			Transform.SetLocation(Simulation.GetPointInterpolated(PointIndex).Position);
 		}
+
 		if (PointIndex != BoneIndices.Num() - 1) // skip the effector
 		{
 			Transform.SetRotation(PreRotation * Transform.GetRotation());
@@ -248,7 +249,7 @@ void FRigUnit_SpringIK::Execute(const FRigUnitContext& Context)
 			}
 		}
 
-		if(bLimitLocalPosition && PointIndex < BoneIndices.Num() - 1)
+		if(bLimitLocalPosition)
 		{
 			int32 ParentIndex = Hierarchy->GetParentIndex(BoneIndices[PointIndex]);
 			if (ParentIndex != INDEX_NONE)
@@ -257,7 +258,29 @@ void FRigUnit_SpringIK::Execute(const FRigUnitContext& Context)
 				FTransform ParentInitialTransform = Hierarchy->GetInitialTransform(ParentIndex);
 				FTransform ParentTransform = Hierarchy->GetGlobalTransform(ParentIndex);
 				float ExpectedDistance = (InitialTransform.GetLocation() - ParentInitialTransform.GetLocation()).Size();
-				if (ExpectedDistance > SMALL_NUMBER)
+				if (ExpectedDistance > SMALL_NUMBER && PointIndex < BoneIndices.Num() - 1)
+				{
+					FVector Direction = Transform.GetLocation() - ParentTransform.GetLocation();
+					if (!Direction.IsNearlyZero())
+					{
+						Transform.SetLocation(ParentTransform.GetLocation() + Direction.GetSafeNormal() * ExpectedDistance);
+					}
+				}
+
+				// correct the rotation on the last bone
+				if (PointIndex == BoneIndices.Num() - 2)
+				{
+					FVector Axis = Transform.TransformVectorNoScale(PrimaryAxis);
+					FVector Target = Simulation.GetPointInterpolated(PointIndex + 1).Position - Transform.GetLocation();
+					if (!Target.IsNearlyZero() && !Axis.IsNearlyZero())
+					{
+						Target = Target.GetSafeNormal();
+						FQuat Rotation = FQuat::FindBetweenNormals(Axis, Target);
+						Transform.SetRotation((Rotation * Transform.GetRotation()).GetNormalized());
+					}
+				}
+
+				if (ExpectedDistance > SMALL_NUMBER && PointIndex == BoneIndices.Num() - 1)
 				{
 					FVector Direction = Transform.GetLocation() - ParentTransform.GetLocation();
 					if (!Direction.IsNearlyZero())
diff --git a/Engine/Plugins/Experimental/ControlRig/Source/ControlRig/Private/Units/Highlevel/Hierarchy/RigUnit_SpringIK.h b/Engine/Plugins/Experimental/ControlRig/Source/ControlRig/Private/Units/Highlevel/Hierarchy/RigUnit_SpringIK.h
--- a/Engine/Plugins/Experimental/ControlRig/Source/ControlRig/Private/Units/Highlevel/Hierarchy/RigUnit_SpringIK.h
+++ b/Engine/Plugins/Experimental/ControlRig/Source/ControlRig/Private/Units/Highlevel/Hierarchy/RigUnit_SpringIK.h
@@ -59,6 +59,7 @@ struct FRigUnit_SpringIK : public FRigUnit_HighlevelBaseMutable
 		StartBone = EndBone = PoleVectorSpace = NAME_None;
 		HierarchyStrength = 256.f;
 		EffectorStrength = RootStrength = 64.f;
+		EffectorRatio = RootRatio = 0.5f;
 		Damping = 0.4f;
 		PoleVector = FVector(0.f, 0.f, 1.f);
 		bFlipPolePlane = false;
@@ -99,12 +100,26 @@ struct FRigUnit_SpringIK : public FRigUnit_HighlevelBaseMutable
 	UPROPERTY(meta = (Input, Constant))
 	float EffectorStrength;
 
+	/**
+	 * Defines the equilibrium of the effector springs.
+	 * This value ranges from 0.0 (zero distance) to 1.0 (distance in initial pose)
+	 */
+	UPROPERTY(meta = (Input, Constant))
+	float EffectorRatio;
+
 	/**
 	 * Sets the coefficient of the springs towards the root. Values between 1 and 2048 are common.
 	 */
 	UPROPERTY(meta = (Input, Constant))
 	float RootStrength;
 
+	/**
+	 * Defines the equilibrium of the root springs.
+	 * This value ranges from 0.0 (zero distance) to 1.0 (distance in initial pose)
+	 */
+	UPROPERTY(meta = (Input, Constant))
+	float RootRatio;
+
 	/**
 	 * The higher the value to more quickly the simulation calms down. Values between 0.0001 and 0.75 are common.
 	 */
