Control Rig: Improve errors + warnings for Aim node

diff --git a/Engine/Plugins/Experimental/ControlRig/Source/ControlRig/Private/Units/Highlevel/Hierarchy/RigUnit_AimBone.cpp b/Engine/Plugins/Experimental/ControlRig/Source/ControlRig/Private/Units/Highlevel/Hierarchy/RigUnit_AimBone.cpp
--- a/Engine/Plugins/Experimental/ControlRig/Source/ControlRig/Private/Units/Highlevel/Hierarchy/RigUnit_AimBone.cpp
+++ b/Engine/Plugins/Experimental/ControlRig/Source/ControlRig/Private/Units/Highlevel/Hierarchy/RigUnit_AimBone.cpp
@@ -19,13 +19,14 @@ void FRigUnit_AimBone::Execute(const FRigUnitContext& Context)
 
 	if (BoneIndex == INDEX_NONE)
 	{
+		UE_CONTROLRIG_RIGUNIT_REPORT_WARNING(TEXT("Bone not found '%s'."), *Bone.ToString());
 		return;
 	}
 
 	if (Primary.Weight <= SMALL_NUMBER && Secondary.Weight <= SMALL_NUMBER)
 	{
+		UE_CONTROLRIG_RIGUNIT_REPORT_WARNING(TEXT("Weights are zero."));
 		return;
-
 	}
 	FTransform Transform = Hierarchy->GetGlobalTransform(BoneIndex);
 
@@ -33,7 +34,7 @@ void FRigUnit_AimBone::Execute(const FRigUnitContext& Context)
 	{
 		FVector Target = Primary.Target;
 
-		if (PrimaryCachedSpaceName != Primary.Space)
+		if (PrimaryCachedSpaceName != Primary.Space || PrimaryCachedSpaceIndex == INDEX_NONE)
 		{
 			if (Primary.Space == NAME_None)
 			{
@@ -89,13 +90,17 @@ void FRigUnit_AimBone::Execute(const FRigUnitContext& Context)
 			FQuat Rotation = FQuat::FindBetweenNormals(Axis, Target);
 			Transform.SetRotation((Rotation * Transform.GetRotation()).GetNormalized());
 		}
+		else
+		{
+			UE_CONTROLRIG_RIGUNIT_REPORT_WARNING(TEXT("Invalid primary target."));
+		}
 	}
 
 	if (Secondary.Weight > SMALL_NUMBER)
 	{
 		FVector Target = Secondary.Target;
 
-		if (SecondaryCachedSpaceName != Secondary.Space)
+		if (SecondaryCachedSpaceName != Secondary.Space || SecondaryCachedSpaceIndex == INDEX_NONE)
 		{
 			if (Secondary.Space == NAME_None)
 			{
@@ -158,6 +163,10 @@ void FRigUnit_AimBone::Execute(const FRigUnitContext& Context)
 			FQuat Rotation = FQuat::FindBetweenNormals(Axis, Target);
 			Transform.SetRotation((Rotation * Transform.GetRotation()).GetNormalized());
 		}
+		else
+		{
+			UE_CONTROLRIG_RIGUNIT_REPORT_WARNING(TEXT("Invalid secondary target."));
+		}
 	}
 
 	Hierarchy->SetGlobalTransform(BoneIndex, Transform, bPropagateToChildren);
