Fixed broken no-unity build after merge

diff --git a/Engine/Plugins/Experimental/ControlRig/Source/ControlRigDeveloper/Private/SVariableMappingWidget.cpp b/Engine/Plugins/Experimental/ControlRig/Source/ControlRigDeveloper/Private/SVariableMappingWidget.cpp
--- a/Engine/Plugins/Experimental/ControlRig/Source/ControlRigDeveloper/Private/SVariableMappingWidget.cpp
+++ b/Engine/Plugins/Experimental/ControlRig/Source/ControlRigDeveloper/Private/SVariableMappingWidget.cpp
@@ -3,7 +3,7 @@
 
 #include "SVariableMappingWidget.h"
 #include "Widgets/Text/SInlineEditableTextBlock.h"
-#include "Widgets/Input/SSearchBox.h"
+#include "Widgets/Input/SCheckBox.h"
 #include "Widgets/Input/SButton.h"
 
 #define LOCTEXT_NAMESPACE "SVariableMappingWidget"
diff --git a/Engine/Plugins/Experimental/ControlRig/Source/ControlRigDeveloper/Private/SVariableMappingWidget.h b/Engine/Plugins/Experimental/ControlRig/Source/ControlRigDeveloper/Private/SVariableMappingWidget.h
--- a/Engine/Plugins/Experimental/ControlRig/Source/ControlRigDeveloper/Private/SVariableMappingWidget.h
+++ b/Engine/Plugins/Experimental/ControlRig/Source/ControlRigDeveloper/Private/SVariableMappingWidget.h
@@ -5,6 +5,7 @@
 #include "CoreMinimal.h"
 #include "SSearchableComboBox.h"
 #include "Widgets/Views/STreeView.h"
+#include "Widgets/Input/SSearchBox.h"
 
 //=====================================================================================
 // UI for displaying property info
diff --git a/Engine/Plugins/Experimental/ControlRig/Source/ControlRigEditor/Private/Graph/SControlRigGraphNode.cpp b/Engine/Plugins/Experimental/ControlRig/Source/ControlRigEditor/Private/Graph/SControlRigGraphNode.cpp
--- a/Engine/Plugins/Experimental/ControlRig/Source/ControlRigEditor/Private/Graph/SControlRigGraphNode.cpp
+++ b/Engine/Plugins/Experimental/ControlRig/Source/ControlRigEditor/Private/Graph/SControlRigGraphNode.cpp
@@ -21,6 +21,10 @@
 #include "ControlRigBlueprint.h"
 #include "ControlRigController.h"
 
+#if WITH_EDITOR
+#include "Editor.h"
+#endif
+
 #define LOCTEXT_NAMESPACE "SControlRigGraphNode"
 
 void SControlRigGraphNode::Construct( const FArguments& InArgs )
diff --git a/Engine/Source/Editor/DetailCustomizations/Private/LandscapeUIDetails.cpp b/Engine/Source/Editor/DetailCustomizations/Private/LandscapeUIDetails.cpp
--- a/Engine/Source/Editor/DetailCustomizations/Private/LandscapeUIDetails.cpp
+++ b/Engine/Source/Editor/DetailCustomizations/Private/LandscapeUIDetails.cpp
@@ -9,6 +9,7 @@
 #include "DetailCategoryBuilder.h"
 #include "DetailWidgetRow.h"
 #include "Widgets/Input/SButton.h"
+#include "Widgets/Input/SCheckBox.h"
 #include "Widgets/Text/STextBlock.h"
 #include "Misc/MessageDialog.h"
 #include "Editor.h"
