The case where missing joints, it was still thinking it has data to map, so it was using it. This change does query control rig for node data, and then map it if it requires. This wil need more thinking around how this will work with curve container

diff --git a/Engine/Plugins/Experimental/ControlRig/Source/ControlRig/Private/AnimNode_ControlRigBase.cpp b/Engine/Plugins/Experimental/ControlRig/Source/ControlRig/Private/AnimNode_ControlRigBase.cpp
--- a/Engine/Plugins/Experimental/ControlRig/Source/ControlRig/Private/AnimNode_ControlRigBase.cpp
+++ b/Engine/Plugins/Experimental/ControlRig/Source/ControlRig/Private/AnimNode_ControlRigBase.cpp
@@ -71,19 +71,19 @@ void FAnimNode_ControlRigBase::UpdateInput(UControlRig* ControlRig, const FPoseC
 
 	// @re-think - now control rig contains init pose from their default hierarchy and current pose from this instance.
 	// we may need this init pose somewhere (instance refpose)
-	const int32 NumNodes = ContolRigNodeMapping.Num();
+	const int32 NumNodes = ControlRigNodeMapping.Num();
 	for (int32 Index = 0; Index < NumNodes; ++Index)
 	{
-		if (ContolRigNodeMapping[Index] != NAME_None)
+		if (ControlRigNodeMapping[Index] != NAME_None)
 		{
 			FCompactPoseBoneIndex CompactPoseIndex(Index);
 			FTransform ComponentTransform = MeshPoses.GetComponentSpaceTransform(CompactPoseIndex);
 			if (NodeMappingContainer.IsValid())
 			{
-				ComponentTransform = NodeMappingContainer->GetSourceToTargetTransform(ContolRigNodeMapping[Index]).GetRelativeTransformReverse(ComponentTransform);
+				ComponentTransform = NodeMappingContainer->GetSourceToTargetTransform(ControlRigNodeMapping[Index]).GetRelativeTransformReverse(ComponentTransform);
 			}
 
-			ControlRig->SetGlobalTransform(ContolRigNodeMapping[Index], ComponentTransform);
+			ControlRig->SetGlobalTransform(ControlRigNodeMapping[Index], ComponentTransform);
 		}
 	}
 }
@@ -97,16 +97,16 @@ void FAnimNode_ControlRigBase::UpdateOutput(UControlRig* ControlRig, FPoseContex
 	FCSPose<FCompactPose> MeshPoses;
 	MeshPoses.InitPose(InOutput.Pose);
 
-	const int32 NumNodes = ContolRigNodeMapping.Num();
+	const int32 NumNodes = ControlRigNodeMapping.Num();
 	for (int32 Index = 0; Index < NumNodes; ++Index)
 	{
-		if (ContolRigNodeMapping[Index] != NAME_None)
+		if (ControlRigNodeMapping[Index] != NAME_None)
 		{
 			FCompactPoseBoneIndex CompactPoseIndex(Index);
-			FTransform ComponentTransform = ControlRig->GetGlobalTransform(ContolRigNodeMapping[Index]);
+			FTransform ComponentTransform = ControlRig->GetGlobalTransform(ControlRigNodeMapping[Index]);
 			if (NodeMappingContainer.IsValid())
 			{
-				ComponentTransform = NodeMappingContainer->GetSourceToTargetTransform(ContolRigNodeMapping[Index]) * ComponentTransform;
+				ComponentTransform = NodeMappingContainer->GetSourceToTargetTransform(ControlRigNodeMapping[Index]) * ComponentTransform;
 			}
 
 			MeshPoses.SetComponentSpaceTransform(CompactPoseIndex, ComponentTransform);
@@ -142,8 +142,8 @@ void FAnimNode_ControlRigBase::CacheBones_AnyThread(const FAnimationCacheBonesCo
 		FBoneContainer& RequiredBones = Context.AnimInstanceProxy->GetRequiredBones();
 		const TArray<FBoneIndexType>& RequiredBonesArray = RequiredBones.GetBoneIndicesArray();
 		const int32 NumBones = RequiredBonesArray.Num();
-		ContolRigNodeMapping.Reset(NumBones);
-		ContolRigNodeMapping.AddDefaulted(NumBones);
+		ControlRigNodeMapping.Reset(NumBones);
+		ControlRigNodeMapping.AddDefaulted(NumBones);
 
 		const FReferenceSkeleton& RefSkeleton = RequiredBones.GetReferenceSkeleton();
 
@@ -160,15 +160,28 @@ void FAnimNode_ControlRigBase::CacheBones_AnyThread(const FAnimationCacheBonesCo
 				// get bone name, and find reverse mapping
 				FName TargetNodeName = RefSkeleton.GetBoneName(RequiredBonesArray[Index]);
 				FName* SourceName = TargetToSourceMappingTable.Find(TargetNodeName);
-				ContolRigNodeMapping[Index] = (SourceName)? *SourceName : NAME_None;
+				ControlRigNodeMapping[Index] = (SourceName)? *SourceName : NAME_None;
 			}
-			UE_LOG(LogAnimation, Log, TEXT("%s : %d"), *GetNameSafe(ControlRig), ContolRigNodeMapping.Num());
+			UE_LOG(LogAnimation, Log, TEXT("%s : %d"), *GetNameSafe(ControlRig), ControlRigNodeMapping.Num());
 		}
 		else
 		{
+			TArray<FName> NodeNames;
+			TArray<FNodeItem> NodeItems;
+			ControlRig->GetMappableNodeData(NodeNames, NodeItems);
+
+			// even if not mapped, we map only node that exists in the controlrig
 			for (int32 Index = 0; Index < NumBones; ++Index)
 			{
-				ContolRigNodeMapping[Index] = RefSkeleton.GetBoneName(RequiredBonesArray[Index]);
+				const FName& BoneName = RefSkeleton.GetBoneName(RequiredBonesArray[Index]);
+				if (NodeNames.Contains(BoneName))
+				{
+					ControlRigNodeMapping[Index] = BoneName;
+				}
+				else
+				{
+					ControlRigNodeMapping[Index] = NAME_None;
+				}
 			}
 		}
 	}
diff --git a/Engine/Plugins/Experimental/ControlRig/Source/ControlRig/Public/AnimNode_ControlRigBase.h b/Engine/Plugins/Experimental/ControlRig/Source/ControlRig/Public/AnimNode_ControlRigBase.h
--- a/Engine/Plugins/Experimental/ControlRig/Source/ControlRig/Public/AnimNode_ControlRigBase.h
+++ b/Engine/Plugins/Experimental/ControlRig/Source/ControlRig/Public/AnimNode_ControlRigBase.h
@@ -31,7 +31,7 @@ struct CONTROLRIG_API FAnimNode_ControlRigBase : public FAnimNode_CustomProperty
 protected:
 	/** Rig Hierarchy node name mapping for the required bones array */
 	UPROPERTY(transient)
-	TArray<FName> ContolRigNodeMapping;
+	TArray<FName> ControlRigNodeMapping;
 
 	UPROPERTY(transient)
 	TMap<FName, uint16> CurveMappingUIDs;
