Filter out ControlRig specific node spawners for non-controlrig blueprints

diff --git a/Engine/Plugins/Experimental/ControlRig/Source/ControlRigEditor/Private/Graph/NodeSpawners/ControlRigCommentNodeSpawner.cpp b/Engine/Plugins/Experimental/ControlRig/Source/ControlRigEditor/Private/Graph/NodeSpawners/ControlRigCommentNodeSpawner.cpp
--- a/Engine/Plugins/Experimental/ControlRig/Source/ControlRigEditor/Private/Graph/NodeSpawners/ControlRigCommentNodeSpawner.cpp
+++ b/Engine/Plugins/Experimental/ControlRig/Source/ControlRigEditor/Private/Graph/NodeSpawners/ControlRigCommentNodeSpawner.cpp
@@ -110,4 +110,18 @@ UEdGraphNode* UControlRigCommentNodeSpawner::Invoke(UEdGraph* ParentGraph, FBind
 	return nullptr;
 }
 
+bool UControlRigCommentNodeSpawner::IsTemplateNodeFilteredOut(FBlueprintActionFilter const& Filter) const
+{
+	for (UBlueprint* Blueprint : Filter.Context.Blueprints)
+	{
+		UControlRigBlueprint* RigBlueprint = Cast<UControlRigBlueprint>(Blueprint);
+		if (RigBlueprint == nullptr)
+		{
+			return true;
+		}
+	}
+
+	return Super::IsTemplateNodeFilteredOut(Filter);
+}
+
 #undef LOCTEXT_NAMESPACE
diff --git a/Engine/Plugins/Experimental/ControlRig/Source/ControlRigEditor/Private/Graph/NodeSpawners/ControlRigCommentNodeSpawner.h b/Engine/Plugins/Experimental/ControlRig/Source/ControlRigEditor/Private/Graph/NodeSpawners/ControlRigCommentNodeSpawner.h
--- a/Engine/Plugins/Experimental/ControlRig/Source/ControlRigEditor/Private/Graph/NodeSpawners/ControlRigCommentNodeSpawner.h
+++ b/Engine/Plugins/Experimental/ControlRig/Source/ControlRigEditor/Private/Graph/NodeSpawners/ControlRigCommentNodeSpawner.h
@@ -35,6 +35,7 @@ public:
 	virtual FBlueprintNodeSignature GetSpawnerSignature() const override;
 	virtual FBlueprintActionUiSpec GetUiSpec(FBlueprintActionContext const& Context, FBindingSet const& Bindings) const override;
 	virtual UEdGraphNode* Invoke(UEdGraph* ParentGraph, FBindingSet const& Bindings, FVector2D const Location) const override;
+	virtual bool IsTemplateNodeFilteredOut(FBlueprintActionFilter const& Filter) const override;
 	// End UBlueprintNodeSpawner interface
 
 private:
