Fixing up a corrupt asset causing control rig asserts

diff --git a/Engine/Plugins/Experimental/ControlRig/Source/ControlRigDeveloper/Private/ControlRigModel.cpp b/Engine/Plugins/Experimental/ControlRig/Source/ControlRigDeveloper/Private/ControlRigModel.cpp
--- a/Engine/Plugins/Experimental/ControlRig/Source/ControlRigDeveloper/Private/ControlRigModel.cpp
+++ b/Engine/Plugins/Experimental/ControlRig/Source/ControlRigDeveloper/Private/ControlRigModel.cpp
@@ -491,7 +491,7 @@ bool UControlRigModel::AddNode(const FControlRigModelNode& InNode, bool bUndo)
 	int32 PreviousExecutePin = INDEX_NONE;
 
 	// if this is a mutable node let's add the begin execution and / or wire it up
-	if (NodeToAdd.UnitStruct()->IsChildOf(FRigUnitMutable::StaticStruct()))
+	if (NodeToAdd.UnitStruct()->IsChildOf(FRigUnitMutable::StaticStruct()) && bUndo)
 	{
 		float ClosestDistance = FLT_MAX;
 
