Add bWasCreated flag to BindToSkeletalMesh so when we know when we create the sequencer anim instance we should save the pose. This removes the need for control rig to override NativeInitialization.  Also add some checks to make sure we save a valid pose.

diff --git a/Engine/Plugins/Experimental/ControlRig/Source/ControlRig/Private/Sequencer/ControlRigBindingTemplate.cpp b/Engine/Plugins/Experimental/ControlRig/Source/ControlRig/Private/Sequencer/ControlRigBindingTemplate.cpp
--- a/Engine/Plugins/Experimental/ControlRig/Source/ControlRig/Private/Sequencer/ControlRigBindingTemplate.cpp
+++ b/Engine/Plugins/Experimental/ControlRig/Source/ControlRig/Private/Sequencer/ControlRigBindingTemplate.cpp
@@ -83,7 +83,8 @@ struct FBindControlRigObjectToken : IMovieSceneExecutionToken
 		check(ControlRig);
 		if (USkeletalMeshComponent* SkeletalMeshComponent = Cast<USkeletalMeshComponent>(ControlRig->GetObjectBinding()->GetBoundObject()))
 		{
-			if (UControlRigSequencerAnimInstance* AnimInstance = UAnimCustomInstance::BindToSkeletalMeshComponent<UControlRigSequencerAnimInstance>(SkeletalMeshComponent))
+			bool bWasCreated = false;
+			if (UControlRigSequencerAnimInstance* AnimInstance = UAnimCustomInstance::BindToSkeletalMeshComponent<UControlRigSequencerAnimInstance>(SkeletalMeshComponent,bWasCreated))
 			{
 				AnimInstance->RecalcRequiredBones();
 			}
diff --git a/Engine/Plugins/Experimental/ControlRig/Source/ControlRig/Private/Sequencer/ControlRigSequencerAnimInstance.cpp b/Engine/Plugins/Experimental/ControlRig/Source/ControlRig/Private/Sequencer/ControlRigSequencerAnimInstance.cpp
--- a/Engine/Plugins/Experimental/ControlRig/Source/ControlRig/Private/Sequencer/ControlRigSequencerAnimInstance.cpp
+++ b/Engine/Plugins/Experimental/ControlRig/Source/ControlRig/Private/Sequencer/ControlRigSequencerAnimInstance.cpp
@@ -20,9 +20,3 @@ bool UControlRigSequencerAnimInstance::UpdateControlRig(UControlRig* InControlRi
 	CachedControlRig = InControlRig;
 	return GetProxyOnGameThread<FControlRigSequencerAnimInstanceProxy>().UpdateControlRig(InControlRig, SequenceId, bAdditive, bApplyBoneFilter, BoneFilter, Weight, bExternalSource);
 }
-
-void UControlRigSequencerAnimInstance::NativeInitializeAnimation()
-{
-	//Do nothing since UAnimSequencerInstance save's a pose snapshot which can cause ensure and crash issues 
-	//since it may not have a skeletal mesh or component transform set up yet.
-}
diff --git a/Engine/Plugins/Experimental/ControlRig/Source/ControlRig/Public/Sequencer/ControlRigSequencerAnimInstance.h b/Engine/Plugins/Experimental/ControlRig/Source/ControlRig/Public/Sequencer/ControlRigSequencerAnimInstance.h
--- a/Engine/Plugins/Experimental/ControlRig/Source/ControlRig/Public/Sequencer/ControlRigSequencerAnimInstance.h
+++ b/Engine/Plugins/Experimental/ControlRig/Source/ControlRig/Public/Sequencer/ControlRigSequencerAnimInstance.h
@@ -22,5 +22,4 @@ public:
 protected:
 	// UAnimInstance interface
 	virtual FAnimInstanceProxy* CreateAnimInstanceProxy() override;
-	virtual void NativeInitializeAnimation() override;
 };
diff --git a/Engine/Plugins/Experimental/ControlRig/Source/ControlRigEditor/Private/Editor/ControlRigEditor.cpp b/Engine/Plugins/Experimental/ControlRig/Source/ControlRigEditor/Private/Editor/ControlRigEditor.cpp
--- a/Engine/Plugins/Experimental/ControlRig/Source/ControlRigEditor/Private/Editor/ControlRigEditor.cpp
+++ b/Engine/Plugins/Experimental/ControlRig/Source/ControlRigEditor/Private/Editor/ControlRigEditor.cpp
@@ -1238,7 +1238,8 @@ void FControlRigEditor::HandlePreviewSceneCreated(const TSharedRef<IPersonaPrevi
 	UControlRigSkeletalMeshComponent* EditorSkelComp = NewObject<UControlRigSkeletalMeshComponent>(Actor);
 	EditorSkelComp->SetSkeletalMesh(InPersonaPreviewScene->GetPersonaToolkit()->GetPreviewMesh());
 	InPersonaPreviewScene->SetPreviewMeshComponent(EditorSkelComp);
-	UAnimCustomInstance::BindToSkeletalMeshComponent<UControlRigSequencerAnimInstance>(EditorSkelComp);
+	bool bWasCreated = false;
+	UAnimCustomInstance::BindToSkeletalMeshComponent<UControlRigSequencerAnimInstance>(EditorSkelComp, bWasCreated);
 	InPersonaPreviewScene->AddComponent(EditorSkelComp, FTransform::Identity);
 
 	// set root component, so we can attach to it. 
@@ -1328,7 +1329,8 @@ void FControlRigEditor::RebindToSkeletalMeshComponent()
 	UDebugSkelMeshComponent* MeshComponent = GetPersonaToolkit()->GetPreviewScene()->GetPreviewMeshComponent();
 	if (MeshComponent)
 	{
-		UAnimCustomInstance::BindToSkeletalMeshComponent<UControlRigSequencerAnimInstance>(MeshComponent);
+		bool bWasCreated = false;
+		UAnimCustomInstance::BindToSkeletalMeshComponent<UControlRigSequencerAnimInstance>(MeshComponent , bWasCreated);
 	}
 }
 
diff --git a/Engine/Source/Editor/Persona/Private/AnimationEditorPreviewScene.cpp b/Engine/Source/Editor/Persona/Private/AnimationEditorPreviewScene.cpp
--- a/Engine/Source/Editor/Persona/Private/AnimationEditorPreviewScene.cpp
+++ b/Engine/Source/Editor/Persona/Private/AnimationEditorPreviewScene.cpp
@@ -352,7 +352,8 @@ void FAnimationEditorPreviewScene::RefreshAdditionalMeshes(bool bAllowOverrideBa
 						}
 						else
 						{
-							UAnimCustomInstance::BindToSkeletalMeshComponent<UAnimPreviewAttacheInstance>(NewComp);
+							bool bWasCreated = false;
+							UAnimCustomInstance::BindToSkeletalMeshComponent<UAnimPreviewAttacheInstance>(NewComp,bWasCreated);
 						}
 
 						AdditionalMeshes.Add(NewComp);
diff --git a/Engine/Source/Runtime/AnimGraphRuntime/Private/AnimSequencerInstance.cpp b/Engine/Source/Runtime/AnimGraphRuntime/Private/AnimSequencerInstance.cpp
--- a/Engine/Source/Runtime/AnimGraphRuntime/Private/AnimSequencerInstance.cpp
+++ b/Engine/Source/Runtime/AnimGraphRuntime/Private/AnimSequencerInstance.cpp
@@ -40,12 +40,13 @@ void UAnimSequencerInstance::ResetPose()
 	GetProxyOnGameThread<FAnimSequencerInstanceProxy>().ResetPose();
 }
 
-void UAnimSequencerInstance::NativeInitializeAnimation()
-{
-	SavePose();
-}
-
 void UAnimSequencerInstance::SavePose()
 {
-	SavePoseSnapshot(UAnimSequencerInstance::SequencerPoseName);
+	if (USkeletalMeshComponent* SkeletalMeshComponent = GetSkelMeshComponent())
+	{
+		if (SkeletalMeshComponent->SkeletalMesh && SkeletalMeshComponent->GetComponentSpaceTransforms().Num() > 0)
+		{
+			SavePoseSnapshot(UAnimSequencerInstance::SequencerPoseName);
+		}
+	}
 }
diff --git a/Engine/Source/Runtime/AnimGraphRuntime/Public/AnimCustomInstance.h b/Engine/Source/Runtime/AnimGraphRuntime/Public/AnimCustomInstance.h
--- a/Engine/Source/Runtime/AnimGraphRuntime/Public/AnimCustomInstance.h
+++ b/Engine/Source/Runtime/AnimGraphRuntime/Public/AnimCustomInstance.h
@@ -21,8 +21,9 @@ class ANIMGRAPHRUNTIME_API UAnimCustomInstance : public UAnimInstance
 	 * @return the current (or newly created) UAnimCustomInstance
 	 */
 	template<typename InstanceClassType>
-	static InstanceClassType* BindToSkeletalMeshComponent(USkeletalMeshComponent* InSkeletalMeshComponent)
+	static InstanceClassType* BindToSkeletalMeshComponent(USkeletalMeshComponent* InSkeletalMeshComponent, bool& bOutWasCreated)
 	{
+		bOutWasCreated = false;
 		// make sure to tick and refresh all the time when ticks
 		// @TODO: this needs restoring post-binding
 		InSkeletalMeshComponent->VisibilityBasedAnimTickOption = EVisibilityBasedAnimTickOption::AlwaysTickPoseAndRefreshBones;
@@ -50,6 +51,7 @@ class ANIMGRAPHRUNTIME_API UAnimCustomInstance : public UAnimInstance
 				InstanceClassType* SequencerInstance = NewObject<InstanceClassType>(InSkeletalMeshComponent, InstanceClassType::StaticClass());
 				InSkeletalMeshComponent->AnimScriptInstance = SequencerInstance;
 				InSkeletalMeshComponent->AnimScriptInstance->InitializeAnimation();
+				bOutWasCreated = true;
 				return SequencerInstance;
 			}
 			else
diff --git a/Engine/Source/Runtime/AnimGraphRuntime/Public/AnimSequencerInstance.h b/Engine/Source/Runtime/AnimGraphRuntime/Public/AnimSequencerInstance.h
--- a/Engine/Source/Runtime/AnimGraphRuntime/Public/AnimSequencerInstance.h
+++ b/Engine/Source/Runtime/AnimGraphRuntime/Public/AnimSequencerInstance.h
@@ -27,13 +27,14 @@ public:
 	/** Reset the pose for this instance*/
 	void ResetPose();
 
+	/** Saved the named pose to restore after */
+	void SavePose();
+
 protected:
 	// UAnimInstance interface
 	virtual FAnimInstanceProxy* CreateAnimInstanceProxy() override;
-	virtual void NativeInitializeAnimation() override;
 
-	/** Saved the named pose to restore after */
-	void SavePose();
+
 
 public:
 	static const FName SequencerPoseName;
diff --git a/Engine/Source/Runtime/MovieSceneTracks/Private/Evaluation/MovieSceneSkeletalAnimationTemplate.cpp b/Engine/Source/Runtime/MovieSceneTracks/Private/Evaluation/MovieSceneSkeletalAnimationTemplate.cpp
--- a/Engine/Source/Runtime/MovieSceneTracks/Private/Evaluation/MovieSceneSkeletalAnimationTemplate.cpp
+++ b/Engine/Source/Runtime/MovieSceneTracks/Private/Evaluation/MovieSceneSkeletalAnimationTemplate.cpp
@@ -236,8 +236,12 @@ namespace MovieScene
 			OriginalStack.SavePreAnimatedState(Player, *SkeletalMeshComponent, GetAnimControlTypeID(), FPreAnimatedAnimationTokenProducer());
 
 			UAnimInstance* ExistingAnimInstance = SkeletalMeshComponent->GetAnimInstance();
-
-			UAnimSequencerInstance* SequencerInstance = UAnimCustomInstance::BindToSkeletalMeshComponent<UAnimSequencerInstance>(SkeletalMeshComponent);
+			bool bWasCreated = false;
+			UAnimSequencerInstance* SequencerInstance = UAnimCustomInstance::BindToSkeletalMeshComponent<UAnimSequencerInstance>(SkeletalMeshComponent,bWasCreated);
+			if (bWasCreated)
+			{
+				SequencerInstance->SavePose();
+			}
 
 			const bool bPreviewPlayback = ShouldUsePreviewPlayback(Player, *SkeletalMeshComponent);
 
