Final copy-up for Control Rig Task Stream.

diff --git a/Engine/Plugins/Experimental/ControlRig/Source/ControlRigDeveloper/Private/ControlRigBlueprintCompiler.cpp b/Engine/Plugins/Experimental/ControlRig/Source/ControlRigDeveloper/Private/ControlRigBlueprintCompiler.cpp
--- a/Engine/Plugins/Experimental/ControlRig/Source/ControlRigDeveloper/Private/ControlRigBlueprintCompiler.cpp
+++ b/Engine/Plugins/Experimental/ControlRig/Source/ControlRigDeveloper/Private/ControlRigBlueprintCompiler.cpp
@@ -21,6 +21,8 @@ DEFINE_LOG_CATEGORY_STATIC(LogControlRigCompiler, Log, All);
 
 bool FControlRigBlueprintCompiler::CanCompile(const UBlueprint* Blueprint)
 {
+	DECLARE_SCOPE_HIERARCHICAL_COUNTER_FUNC()
+
 	if (Blueprint && Blueprint->ParentClass && Blueprint->ParentClass->IsChildOf(UControlRig::StaticClass()))
 	{
 		return true;
@@ -31,12 +33,16 @@ bool FControlRigBlueprintCompiler::CanCompile(const UBlueprint* Blueprint)
 
 void FControlRigBlueprintCompiler::Compile(UBlueprint* Blueprint, const FKismetCompilerOptions& CompileOptions, FCompilerResultsLog& Results)
 {
+	DECLARE_SCOPE_HIERARCHICAL_COUNTER_FUNC()
+
 	FControlRigBlueprintCompilerContext Compiler(Blueprint, Results, CompileOptions);
 	Compiler.Compile();
 }
 
 void FControlRigBlueprintCompilerContext::MarkCompilationFailed(const FString& Message)
 {
+	DECLARE_SCOPE_HIERARCHICAL_COUNTER_FUNC()
+
 	UControlRigBlueprint* ControlRigBlueprint = Cast<UControlRigBlueprint>(Blueprint);
 	if (ControlRigBlueprint)
 	{
@@ -59,6 +65,8 @@ void FControlRigBlueprintCompilerContext::MarkCompilationFailed(const FString& M
 
 void FControlRigBlueprintCompilerContext::BuildPropertyLinks()
 {
+	DECLARE_SCOPE_HIERARCHICAL_COUNTER_FUNC()
+
 	UControlRigBlueprint* ControlRigBlueprint = Cast<UControlRigBlueprint>(Blueprint);
 	if (ControlRigBlueprint)
 	{
@@ -109,11 +117,15 @@ void FControlRigBlueprintCompilerContext::BuildPropertyLinks()
 
 void FControlRigBlueprintCompilerContext::MergeUbergraphPagesIn(UEdGraph* Ubergraph)
 {
+	DECLARE_SCOPE_HIERARCHICAL_COUNTER_FUNC()
+
 	BuildPropertyLinks();
 }
 
 void FControlRigBlueprintCompilerContext::PostCompile()
 {
+	DECLARE_SCOPE_HIERARCHICAL_COUNTER_FUNC()
+
 	UControlRigBlueprint* ControlRigBlueprint = Cast<UControlRigBlueprint>(Blueprint);
 	if (ControlRigBlueprint)
 	{
@@ -123,19 +135,22 @@ void FControlRigBlueprintCompilerContext::PostCompile()
 
 		TArray<FName> UnitNames;
 		TMap<FName, int32> UnitNameToIndex;
-		for (const FControlRigBlueprintPropertyLink& Link : PropertyLinks)
 		{
-			FName SourceUnitName(*Link.GetSourceUnitName());
-			if (!UnitNameToIndex.Contains(SourceUnitName))
-			{
-				UnitNames.Add(SourceUnitName);
-				UnitNameToIndex.Add(SourceUnitName, UnitNameToIndex.Num());
-			}
-			FName DestUnitName(*Link.GetDestUnitName());
-			if (!UnitNameToIndex.Contains(DestUnitName))
+			DECLARE_SCOPE_HIERARCHICAL_COUNTER(CreateUnitNameMap)
+			for (const FControlRigBlueprintPropertyLink& Link : PropertyLinks)
 			{
-				UnitNames.Add(DestUnitName);
-				UnitNameToIndex.Add(DestUnitName, UnitNameToIndex.Num());
+				FName SourceUnitName(*Link.GetSourceUnitName());
+				if (!UnitNameToIndex.Contains(SourceUnitName))
+				{
+					UnitNames.Add(SourceUnitName);
+					UnitNameToIndex.Add(SourceUnitName, UnitNameToIndex.Num());
+				}
+				FName DestUnitName(*Link.GetDestUnitName());
+				if (!UnitNameToIndex.Contains(DestUnitName))
+				{
+					UnitNames.Add(DestUnitName);
+					UnitNameToIndex.Add(DestUnitName, UnitNameToIndex.Num());
+				}
 			}
 		}
 
@@ -168,11 +183,14 @@ void FControlRigBlueprintCompilerContext::PostCompile()
 
 		if (UnitNames.Num() > 0)
 		{
+			DECLARE_SCOPE_HIERARCHICAL_COUNTER(CycleDetection)
+
 			// add all of the nodes
 			FControlRigDAG SortGraph;
 			for (int32 UnitIndex = 0; UnitIndex < UnitNames.Num(); UnitIndex++)
 			{
 				bool IsMutableUnit = false;
+				bool IsOutputParameter = false;
 				UStructProperty* StructProperty = Cast<UStructProperty>(ControlRigBlueprint->GeneratedClass->FindPropertyByName(UnitNames[UnitIndex]));
 				if (StructProperty)
 				{
@@ -185,7 +203,13 @@ void FControlRigBlueprintCompilerContext::PostCompile()
 						IsMutableUnit = true;
 					}
 				}
-				SortGraph.AddNode(IsMutableUnit, UnitNames[UnitIndex]);
+
+				const FControlRigModelNode* UnitNode = ControlRigBlueprint->Model->FindNode(UnitNames[UnitIndex]);
+				if (UnitNode)
+				{
+					IsOutputParameter = UnitNode->IsParameter() && UnitNode->ParameterType == EControlRigModelParameterType::Output;
+				}
+				SortGraph.AddNode(IsMutableUnit, IsOutputParameter, UnitNames[UnitIndex]);
 			}
 
 			// add all of the links
@@ -243,74 +267,86 @@ void FControlRigBlueprintCompilerContext::PostCompile()
 
 #if WITH_EDITORONLY_DATA
 			// clear the errors on the graph
-			for (UEdGraph* UbergraphPage : Blueprint->UbergraphPages)
 			{
-				if (UControlRigGraph* ControlRigGraph = Cast<UControlRigGraph>(UbergraphPage))
+				DECLARE_SCOPE_HIERARCHICAL_COUNTER(ClearNodeErrors)
+
+				for (UEdGraph* UbergraphPage : Blueprint->UbergraphPages)
 				{
-					for (UEdGraphNode* Node : ControlRigGraph->Nodes)
+					if (UControlRigGraph* ControlRigGraph = Cast<UControlRigGraph>(UbergraphPage))
 					{
-						if(Node->ErrorType < (int32)EMessageSeverity::Info+1)
+						for (UEdGraphNode* Node : ControlRigGraph->Nodes)
 						{
-							Node->ErrorMsg.Reset();
-							Node->ErrorType = (int32)EMessageSeverity::Info+1;
-							Node->bHasCompilerMessage = false;
+							if (Node->ErrorType < (int32)EMessageSeverity::Info + 1)
+							{
+								Node->ErrorMsg.Reset();
+								Node->ErrorType = (int32)EMessageSeverity::Info + 1;
+								Node->bHasCompilerMessage = false;
+							}
 						}
 					}
 				}
 			}
 
-			int32 SortedPropertyLinkIndex = 1;
-			for (const FControlRigDAG::FNode Node : UnitOrder)
 			{
-				UE_LOG(LogControlRigCompiler, Log, TEXT("%d. %s"), SortedPropertyLinkIndex++, *Node.Name.ToString());
-				for (const FControlRigDAG::FPin& Pin : Node.Outputs)
+				DECLARE_SCOPE_HIERARCHICAL_COUNTER(PrintGraphToConsole)
+
+				int32 SortedPropertyLinkIndex = 1;
+				for (const FControlRigDAG::FNode Node : UnitOrder)
 				{
-					const int32 Index = Pin.Link;
-					UE_LOG(LogControlRigCompiler, Log, TEXT("%d. %s -> %s"), SortedPropertyLinkIndex++, *PropertyLinks[Index].GetSourcePropertyPath(), *PropertyLinks[Index].GetDestPropertyPath());
+					UE_LOG(LogControlRigCompiler, Log, TEXT("%d. %s"), SortedPropertyLinkIndex++, *Node.Name.ToString());
+					for (const FControlRigDAG::FPin& Pin : Node.Outputs)
+					{
+						const int32 Index = Pin.Link;
+						UE_LOG(LogControlRigCompiler, Log, TEXT("%d. %s -> %s"), SortedPropertyLinkIndex++, *PropertyLinks[Index].GetSourcePropertyPath(), *PropertyLinks[Index].GetDestPropertyPath());
+					}
 				}
 			}
 #endif
 
-			for (const FControlRigDAG::FNode Node : UnitOrder)
 			{
-				// copy all properties not originating from a unit
-				for (const FControlRigDAG::FPin& Pin : Node.Outputs)
+				DECLARE_SCOPE_HIERARCHICAL_COUNTER(CreateOperatorsForVM)
+
+				for (const FControlRigDAG::FNode Node : UnitOrder)
 				{
-					const int32 PropertyLinkIndex = Pin.Link;
-					const FControlRigBlueprintPropertyLink& Link = PropertyLinks[PropertyLinkIndex];
-					FString SourceUnitName = Link.GetSourceUnitName();
+					// copy all properties not originating from a unit
+					for (const FControlRigDAG::FPin& Pin : Node.Outputs)
+					{
+						const int32 PropertyLinkIndex = Pin.Link;
+						const FControlRigBlueprintPropertyLink& Link = PropertyLinks[PropertyLinkIndex];
+						FString SourceUnitName = Link.GetSourceUnitName();
+						UStructProperty* StructProperty = Cast<UStructProperty>(ControlRigBlueprint->GeneratedClass->FindPropertyByName(Node.Name));
+						if (StructProperty)
+						{
+							if (StructProperty->Struct->IsChildOf(FRigUnit::StaticStruct()))
+							{
+								continue;
+							}
+						}
+
+						Operators.Add(FControlRigOperator(EControlRigOpCode::Copy, Link.GetSourcePropertyPath(), Link.GetDestPropertyPath()));
+					}
+
 					UStructProperty* StructProperty = Cast<UStructProperty>(ControlRigBlueprint->GeneratedClass->FindPropertyByName(Node.Name));
 					if (StructProperty)
 					{
 						if (StructProperty->Struct->IsChildOf(FRigUnit::StaticStruct()))
 						{
-							continue;
+							Operators.Add(FControlRigOperator(EControlRigOpCode::Exec, Node.Name.ToString(), FCachedPropertyPath()));
 						}
 					}
 
-					Operators.Add(FControlRigOperator(EControlRigOpCode::Copy, Link.GetSourcePropertyPath(), Link.GetDestPropertyPath()));
-				}
-
-				UStructProperty* StructProperty = Cast<UStructProperty>(ControlRigBlueprint->GeneratedClass->FindPropertyByName(Node.Name));
-				if (StructProperty)
-				{
-					if (StructProperty->Struct->IsChildOf(FRigUnit::StaticStruct()))
+					// copy all properties to outputs
+					for (const FControlRigDAG::FPin& Pin : Node.Outputs)
 					{
-						Operators.Add(FControlRigOperator(EControlRigOpCode::Exec, Node.Name.ToString(), FCachedPropertyPath()));
+						const int32 PropertyLinkIndex = Pin.Link;
+						const FControlRigBlueprintPropertyLink& Link = PropertyLinks[PropertyLinkIndex];
+						Operators.Add(FControlRigOperator(EControlRigOpCode::Copy, Link.GetSourcePropertyPath(), Link.GetDestPropertyPath()));
 					}
 				}
-
-				// copy all properties to outputs
-				for (const FControlRigDAG::FPin& Pin : Node.Outputs)
-				{
-					const int32 PropertyLinkIndex = Pin.Link;
-					const FControlRigBlueprintPropertyLink& Link = PropertyLinks[PropertyLinkIndex];
-					Operators.Add(FControlRigOperator(EControlRigOpCode::Copy, Link.GetSourcePropertyPath(), Link.GetDestPropertyPath()));
-				}
 			}
-		}
 
-		Operators.Add(FControlRigOperator(EControlRigOpCode::Done));
+			Operators.Add(FControlRigOperator(EControlRigOpCode::Done));
+		}
 
 		// guard against the control rig failing due to serialization changes
 		if (PreviousOperatorCount > 1 && Operators.Num() == 1 && bIsFromVersionBeforeBeginExecution)
@@ -322,6 +358,8 @@ void FControlRigBlueprintCompilerContext::PostCompile()
 
 		// update allow source access properties
 		{
+			DECLARE_SCOPE_HIERARCHICAL_COUNTER(AllowAccessToSourceProperties)
+
 			TArray<FName> SourcePropertyLinkArray;
 			TArray<FName> DestPropertyLinkArray;
 
@@ -387,14 +425,18 @@ void FControlRigBlueprintCompilerContext::PostCompile()
 		}
 	}
 
-	FKismetCompilerContext::PostCompile();
+	{
+		DECLARE_SCOPE_HIERARCHICAL_COUNTER(FKismetCompilerContext::PostCompile)
+		FKismetCompilerContext::PostCompile();
+	}
 
 	// ask the model to update all defaults
-	bool bSetDefaultsFromModel = false;
 	if (ControlRigBlueprint)
 	{
 		if (ControlRigBlueprint->ModelController)
 		{
+			DECLARE_SCOPE_HIERARCHICAL_COUNTER(UpdateModelAfterCompilation)
+
 			// ensure that blueprint storage arrays have the right size.
 			// they might get out of sync due to compilation order.
 			for (const FControlRigModelNode& Node : ControlRigBlueprint->Model->Nodes())
@@ -427,35 +469,16 @@ void FControlRigBlueprintCompilerContext::PostCompile()
 				}
 			}
 
-			bSetDefaultsFromModel = ControlRigBlueprint->ModelController->ResendAllPinDefaultNotifications();
-		}
-	}
-
-	if (!bSetDefaultsFromModel)
-	{
-		// We need to copy any pin defaults over to underlying properties once the class is built
-		// as the defaults may not have been propagated from new nodes yet
-		for (UEdGraph* UbergraphPage : Blueprint->UbergraphPages)
-		{
-			if (UControlRigGraph* ControlRigGraph = Cast<UControlRigGraph>(UbergraphPage))
-			{
-				for (UEdGraphNode* Node : ControlRigGraph->Nodes)
-				{
-					if (UControlRigGraphNode* ControlRigGraphNode = Cast<UControlRigGraphNode>(Node))
-					{
-						for (UEdGraphPin* Pin : ControlRigGraphNode->Pins)
-						{
-							ControlRigGraphNode->CopyPinDefaultsToModel(Pin);
-						}
-					}
-				}
-			}
+			TGuardValue<bool> GuardNotifs(ControlRigBlueprint->bSuspendModelNotificationsForOthers, true);
+			ControlRigBlueprint->ModelController->ResendAllPinDefaultNotifications();
 		}
 	}
 }
 
 void FControlRigBlueprintCompilerContext::CopyTermDefaultsToDefaultObject(UObject* DefaultObject)
 {
+	DECLARE_SCOPE_HIERARCHICAL_COUNTER_FUNC()
+
 	FKismetCompilerContext::CopyTermDefaultsToDefaultObject(DefaultObject);
 
 	UControlRigBlueprint* ControlRigBlueprint = Cast<UControlRigBlueprint>(Blueprint);
@@ -463,6 +486,7 @@ void FControlRigBlueprintCompilerContext::CopyTermDefaultsToDefaultObject(UObjec
 	{
 		UControlRig* ControlRig = CastChecked<UControlRig>(DefaultObject);
 		ControlRig->Hierarchy.BaseHierarchy = ControlRigBlueprint->Hierarchy;
+		ControlRig->CurveContainer = ControlRigBlueprint->CurveContainer;
 		// copy available rig units info, so that control rig can do things with it
 		ControlRig->AllowSourceAccessProperties = ControlRigBlueprint->AllowSourceAccessProperties;
 		ControlRigBlueprint->UpdateParametersOnControlRig(ControlRig);
@@ -471,6 +495,8 @@ void FControlRigBlueprintCompilerContext::CopyTermDefaultsToDefaultObject(UObjec
 
 void FControlRigBlueprintCompilerContext::EnsureProperGeneratedClass(UClass*& TargetUClass)
 {
+	DECLARE_SCOPE_HIERARCHICAL_COUNTER_FUNC()
+
 	if( TargetUClass && !((UObject*)TargetUClass)->IsA(UControlRigBlueprintGeneratedClass::StaticClass()) )
 	{
 		FKismetCompilerUtilities::ConsignToOblivion(TargetUClass, Blueprint->bIsRegeneratingOnLoad);
@@ -480,6 +506,8 @@ void FControlRigBlueprintCompilerContext::EnsureProperGeneratedClass(UClass*& Ta
 
 void FControlRigBlueprintCompilerContext::SpawnNewClass(const FString& NewClassName)
 {
+	DECLARE_SCOPE_HIERARCHICAL_COUNTER_FUNC()
+
 	NewControlRigBlueprintGeneratedClass = FindObject<UControlRigBlueprintGeneratedClass>(Blueprint->GetOutermost(), *NewClassName);
 
 	if (NewControlRigBlueprintGeneratedClass == nullptr)
@@ -496,11 +524,15 @@ void FControlRigBlueprintCompilerContext::SpawnNewClass(const FString& NewClassN
 
 void FControlRigBlueprintCompilerContext::OnNewClassSet(UBlueprintGeneratedClass* ClassToUse)
 {
+	DECLARE_SCOPE_HIERARCHICAL_COUNTER_FUNC()
+
 	NewControlRigBlueprintGeneratedClass = CastChecked<UControlRigBlueprintGeneratedClass>(ClassToUse);
 }
 
 void FControlRigBlueprintCompilerContext::CleanAndSanitizeClass(UBlueprintGeneratedClass* ClassToClean, UObject*& InOldCDO)
 {
+	DECLARE_SCOPE_HIERARCHICAL_COUNTER_FUNC()
+
 	FKismetCompilerContext::CleanAndSanitizeClass(ClassToClean, InOldCDO);
 
 	// Make sure our typed pointer is set
diff --git a/Engine/Plugins/Experimental/ControlRig/Source/ControlRigDeveloper/Private/ControlRigBlueprintUtils.cpp b/Engine/Plugins/Experimental/ControlRig/Source/ControlRigDeveloper/Private/ControlRigBlueprintUtils.cpp
--- a/Engine/Plugins/Experimental/ControlRig/Source/ControlRigDeveloper/Private/ControlRigBlueprintUtils.cpp
+++ b/Engine/Plugins/Experimental/ControlRig/Source/ControlRigDeveloper/Private/ControlRigBlueprintUtils.cpp
@@ -13,6 +13,8 @@
 
 FName FControlRigBlueprintUtils::GetNewUnitMemberName(UBlueprint* InBlueprint, const UStruct* InStructTemplate)
 {
+	DECLARE_SCOPE_HIERARCHICAL_COUNTER_FUNC()
+
 	FString VariableBaseName = InStructTemplate->GetName();
 	VariableBaseName.RemoveFromStart(TEXT("RigUnit_"));
 	return FBlueprintEditorUtils::FindUniqueKismetName(InBlueprint, VariableBaseName);
@@ -20,6 +22,8 @@ FName FControlRigBlueprintUtils::GetNewUnitMemberName(UBlueprint* InBlueprint, c
 
 FName FControlRigBlueprintUtils::AddUnitMember(UBlueprint* InBlueprint, const UStruct* InStructTemplate, const FName& InName)
 {
+	DECLARE_SCOPE_HIERARCHICAL_COUNTER_FUNC()
+
 	FName VarName = InName == NAME_None ? FControlRigBlueprintUtils::GetNewUnitMemberName(InBlueprint, InStructTemplate) : InName;
 
 	UScriptStruct* ScriptStruct = FindObjectChecked<UScriptStruct>(ANY_PACKAGE, *InStructTemplate->GetName());
@@ -39,11 +43,15 @@ FName FControlRigBlueprintUtils::AddUnitMember(UBlueprint* InBlueprint, const US
 
 FName FControlRigBlueprintUtils::GetNewPropertyMemberName(UBlueprint* InBlueprint, const FString& InVariableDesc)
 {
+	DECLARE_SCOPE_HIERARCHICAL_COUNTER_FUNC()
+
 	return FBlueprintEditorUtils::FindUniqueKismetName(InBlueprint, InVariableDesc);
 }
 
 FName FControlRigBlueprintUtils::AddPropertyMember(UBlueprint* InBlueprint, const FEdGraphPinType& InPinType, const FString& InVariableDesc)
 {
+	DECLARE_SCOPE_HIERARCHICAL_COUNTER_FUNC()
+
 	if(FBlueprintEditorUtils::AddMemberVariable(InBlueprint, *InVariableDesc, InPinType))
 	{
 		return InBlueprint->NewVariables.Last().VarName;
@@ -53,6 +61,8 @@ FName FControlRigBlueprintUtils::AddPropertyMember(UBlueprint* InBlueprint, cons
 
 FName FControlRigBlueprintUtils::ValidateName(UBlueprint* InBlueprint, const FString& InName)
 {
+	DECLARE_SCOPE_HIERARCHICAL_COUNTER_FUNC()
+
 	FString Name = InName;
 	if (Name.StartsWith(TEXT("RigUnit_")))
 	{
@@ -107,6 +117,8 @@ FName FControlRigBlueprintUtils::ValidateName(UBlueprint* InBlueprint, const FSt
 
 UControlRigGraphNode* FControlRigBlueprintUtils::InstantiateGraphNodeForProperty(UEdGraph* InGraph, const FName& InPropertyName, const FVector2D& InLocation, const FEdGraphPinType& InPinType)
 {
+	DECLARE_SCOPE_HIERARCHICAL_COUNTER_FUNC()
+
 	check(InGraph);
 
 	InGraph->Modify();
@@ -131,6 +143,8 @@ UControlRigGraphNode* FControlRigBlueprintUtils::InstantiateGraphNodeForProperty
 
 UControlRigGraphNode* FControlRigBlueprintUtils::InstantiateGraphNodeForStructPath(UEdGraph* InGraph, const FName& InPropertyName, const FVector2D& InLocation, const FString& InStructPath)
 {
+	DECLARE_SCOPE_HIERARCHICAL_COUNTER_FUNC()
+
 	check(InGraph);
 
 	InGraph->Modify();
@@ -155,6 +169,8 @@ UControlRigGraphNode* FControlRigBlueprintUtils::InstantiateGraphNodeForStructPa
 
 bool FControlRigBlueprintUtils::CanInstantiateGraphNodeForProperty(UEdGraph* InGraph, const FName& InPropertyName)
 {
+	DECLARE_SCOPE_HIERARCHICAL_COUNTER_FUNC()
+
 	for(UEdGraphNode* Node : InGraph->Nodes)
 	{
 		if(UControlRigGraphNode* ControlRigGraphNode = Cast<UControlRigGraphNode>(Node))
@@ -171,6 +187,8 @@ bool FControlRigBlueprintUtils::CanInstantiateGraphNodeForProperty(UEdGraph* InG
 
 void FControlRigBlueprintUtils::ForAllRigUnits(TFunction<void(UStruct*)> InFunction)
 {
+	DECLARE_SCOPE_HIERARCHICAL_COUNTER_FUNC()
+
 	// Run over all unit types
 	for(TObjectIterator<UStruct> StructIt; StructIt; ++StructIt)
 	{
@@ -183,6 +201,8 @@ void FControlRigBlueprintUtils::ForAllRigUnits(TFunction<void(UStruct*)> InFunct
 
 void FControlRigBlueprintUtils::HandleReconstructAllNodes(UBlueprint* InBlueprint)
 {
+	DECLARE_SCOPE_HIERARCHICAL_COUNTER_FUNC()
+
 	if(InBlueprint->IsA<UControlRigBlueprint>())
 	{
 		// TArray<UControlRigGraphNode*> AllNodes;
@@ -197,6 +217,8 @@ void FControlRigBlueprintUtils::HandleReconstructAllNodes(UBlueprint* InBlueprin
 
 void FControlRigBlueprintUtils::HandleRefreshAllNodes(UBlueprint* InBlueprint)
 {
+	DECLARE_SCOPE_HIERARCHICAL_COUNTER_FUNC()
+
 	if(InBlueprint->IsA<UControlRigBlueprint>())
 	{
 		TArray<UControlRigGraphNode*> AllNodes;
@@ -211,6 +233,8 @@ void FControlRigBlueprintUtils::HandleRefreshAllNodes(UBlueprint* InBlueprint)
 
 void FControlRigBlueprintUtils::HandleRenameVariableReferencesEvent(UBlueprint* InBlueprint, UClass* InVariableClass, const FName& InOldVarName, const FName& InNewVarName)
 {
+	DECLARE_SCOPE_HIERARCHICAL_COUNTER_FUNC()
+
 	if(UControlRigBlueprint* RigBlueprint = Cast<UControlRigBlueprint>(InBlueprint))
 	{
 		if (RigBlueprint->ModelController)
@@ -222,6 +246,8 @@ void FControlRigBlueprintUtils::HandleRenameVariableReferencesEvent(UBlueprint*
 
 void FControlRigBlueprintUtils::RemoveMemberVariableIfNotUsed(UBlueprint* Blueprint, const FName VarName, UControlRigGraphNode* ToBeDeleted)
 {
+	DECLARE_SCOPE_HIERARCHICAL_COUNTER_FUNC()
+
 	if (Blueprint->IsA<UControlRigBlueprint>())
 	{
 		bool bDeleteVariable = true;
diff --git a/Engine/Plugins/Experimental/ControlRig/Source/ControlRigDeveloper/Private/ControlRigController.cpp b/Engine/Plugins/Experimental/ControlRig/Source/ControlRigDeveloper/Private/ControlRigController.cpp
--- a/Engine/Plugins/Experimental/ControlRig/Source/ControlRigDeveloper/Private/ControlRigController.cpp
+++ b/Engine/Plugins/Experimental/ControlRig/Source/ControlRigDeveloper/Private/ControlRigController.cpp
@@ -3,6 +3,7 @@
 #include "ControlRigController.h"
 #include "EdGraphSchema_K2.h"
 #include "Logging/MessageLog.h"
+#include "Stats/StatsHierarchical.h"
 
 #define LOCTEXT_NAMESPACE "ControlRigController"
 
@@ -19,6 +20,8 @@ UControlRigController::~UControlRigController()
 
 void UControlRigController::SetModel(UControlRigModel* InModel)
 {
+	DECLARE_SCOPE_HIERARCHICAL_COUNTER_FUNC()
+
 	if (Model != nullptr)
 	{
 		Model->OnModified().Remove(_ModelModifiedHandle);
@@ -34,6 +37,8 @@ void UControlRigController::SetModel(UControlRigModel* InModel)
 
 bool UControlRigController::Clear()
 {
+	DECLARE_SCOPE_HIERARCHICAL_COUNTER_FUNC()
+
 	return Model->Clear();
 }
 
@@ -46,6 +51,8 @@ void UControlRigController::EnableMessageLog(bool bInEnabled)
 
 bool UControlRigController::OpenUndoBracket(const FString& Title)
 {
+	DECLARE_SCOPE_HIERARCHICAL_COUNTER_FUNC()
+
 	if (!EnsureModel())
 	{
 		return false;
@@ -62,6 +69,8 @@ bool UControlRigController::OpenUndoBracket(const FString& Title)
 
 bool UControlRigController::CloseUndoBracket()
 {
+	DECLARE_SCOPE_HIERARCHICAL_COUNTER_FUNC()
+
 	if (!EnsureModel())
 	{
 		return false;
@@ -88,6 +97,8 @@ UControlRigModel::FModifiedEvent& UControlRigController::OnModified()
 
 void UControlRigController::HandleModelModified(const UControlRigModel* InModel, EControlRigModelNotifType InType, const void* InPayload)
 {
+	DECLARE_SCOPE_HIERARCHICAL_COUNTER_FUNC()
+
 	_LastModelNotification = InType;
 
 	switch (InType)
@@ -280,6 +291,8 @@ void UControlRigController::LogError(const FString& Message) const
 
 bool UControlRigController::ConstructPreviewParameter(const FName& InDataType, EControlRigModelParameterType InParameterType, FControlRigModelNode& OutNode)
 {
+	DECLARE_SCOPE_HIERARCHICAL_COUNTER_FUNC()
+
 	ensure(InParameterType != EControlRigModelParameterType::None);
 
 	FEdGraphPinType PinType;
@@ -311,6 +324,8 @@ bool UControlRigController::ConstructPreviewParameter(const FName& InDataType, E
 
 bool UControlRigController::ConstructPreviewNode(const FName& InFunctionName, FControlRigModelNode& OutNode)
 {
+	DECLARE_SCOPE_HIERARCHICAL_COUNTER_FUNC()
+
 	FControlRigModelNode Node;
 	Node.Index = 0;
 	Node.Name = InFunctionName;
@@ -332,6 +347,8 @@ bool UControlRigController::ConstructPreviewNode(const FName& InFunctionName, FC
 
 bool UControlRigController::AddParameter(const FName& InName, const FName& InDataType, EControlRigModelParameterType InParameterType, const FVector2D& InPosition, bool bUndo)
 {
+	DECLARE_SCOPE_HIERARCHICAL_COUNTER_FUNC()
+
 	if (!EnsureModel())
 	{
 		return false;
@@ -360,6 +377,8 @@ bool UControlRigController::AddParameter(const FName& InName, const FName& InDat
 
 bool UControlRigController::AddComment(const FName& InName, const FString& InText, const FVector2D& InPosition, const FVector2D& InSize, const FLinearColor& InColor, bool bUndo)
 {
+	DECLARE_SCOPE_HIERARCHICAL_COUNTER_FUNC()
+
 	if (!EnsureModel())
 	{
 		return false;
@@ -375,6 +394,8 @@ bool UControlRigController::AddComment(const FName& InName, const FString& InTex
 
 bool UControlRigController::AddNode(const FName& InFunctionName, const FVector2D& InPosition, const FName& InName, bool bUndo)
 {
+	DECLARE_SCOPE_HIERARCHICAL_COUNTER_FUNC()
+
 	if (!EnsureModel())
 	{
 		return false;
@@ -391,6 +412,8 @@ bool UControlRigController::AddNode(const FName& InFunctionName, const FVector2D
 
 bool UControlRigController::RemoveNode(const FName& InName, bool bUndo)
 {
+	DECLARE_SCOPE_HIERARCHICAL_COUNTER_FUNC()
+
 	if (!EnsureModel())
 	{
 		return false;
@@ -401,6 +424,8 @@ bool UControlRigController::RemoveNode(const FName& InName, bool bUndo)
 
 bool UControlRigController::SetNodePosition(const FName& InName, const FVector2D& InPosition, bool bUndo)
 {
+	DECLARE_SCOPE_HIERARCHICAL_COUNTER_FUNC()
+
 	if (!EnsureModel())
 	{
 		return false;
@@ -411,6 +436,8 @@ bool UControlRigController::SetNodePosition(const FName& InName, const FVector2D
 
 bool UControlRigController::SetNodeSize(const FName& InName, const FVector2D& InSize, bool bUndo)
 {
+	DECLARE_SCOPE_HIERARCHICAL_COUNTER_FUNC()
+
 	if (!EnsureModel())
 	{
 		return false;
@@ -421,6 +448,8 @@ bool UControlRigController::SetNodeSize(const FName& InName, const FVector2D& In
 
 bool UControlRigController::SetNodeColor(const FName& InName, const FLinearColor& InColor, bool bUndo)
 {
+	DECLARE_SCOPE_HIERARCHICAL_COUNTER_FUNC()
+
 	if (!EnsureModel())
 	{
 		return false;
@@ -431,6 +460,8 @@ bool UControlRigController::SetNodeColor(const FName& InName, const FLinearColor
 
 bool UControlRigController::SetParameterType(const FName& InName, EControlRigModelParameterType InParameterType, bool bUndo)
 {
+	DECLARE_SCOPE_HIERARCHICAL_COUNTER_FUNC()
+
 	if (!EnsureModel())
 	{
 		return false;
@@ -441,6 +472,8 @@ bool UControlRigController::SetParameterType(const FName& InName, EControlRigMod
 
 bool UControlRigController::SetCommentText(const FName& InName, const FString& InText, bool bUndo)
 {
+	DECLARE_SCOPE_HIERARCHICAL_COUNTER_FUNC()
+
 	if (!EnsureModel())
 	{
 		return false;
@@ -451,6 +484,8 @@ bool UControlRigController::SetCommentText(const FName& InName, const FString& I
 
 bool UControlRigController::RenameNode(const FName& InOldNodeName, const FName& InNewNodeName, bool bUndo)
 {
+	DECLARE_SCOPE_HIERARCHICAL_COUNTER_FUNC()
+
 	if (!EnsureModel())
 	{
 		return false;
@@ -461,6 +496,8 @@ bool UControlRigController::RenameNode(const FName& InOldNodeName, const FName&
 
 bool UControlRigController::ClearSelection()
 {
+	DECLARE_SCOPE_HIERARCHICAL_COUNTER_FUNC()
+
 	if (!EnsureModel())
 	{
 		return false;
@@ -483,6 +520,8 @@ bool UControlRigController::ClearSelection()
 
 bool UControlRigController::SetSelection(const TArray<FName>& InNodeSelection)
 {
+	DECLARE_SCOPE_HIERARCHICAL_COUNTER_FUNC()
+
 	if (!EnsureModel())
 	{
 		return false;
@@ -527,6 +566,8 @@ bool UControlRigController::SetSelection(const TArray<FName>& InNodeSelection)
 
 bool UControlRigController::SelectNode(const FName& InName, bool bInSelected, bool bClearSelection)
 {
+	DECLARE_SCOPE_HIERARCHICAL_COUNTER_FUNC()
+
 	if (!EnsureModel())
 	{
 		return false;
@@ -558,11 +599,15 @@ bool UControlRigController::SelectNode(const FName& InName, bool bInSelected, bo
 
 bool UControlRigController::DeselectNode(const FName& InName)
 {
+	DECLARE_SCOPE_HIERARCHICAL_COUNTER_FUNC()
+
 	return SelectNode(InName, false, false);
 }
 
 bool UControlRigController::PrepareCycleCheckingForPin(const FName& InNodeName, const FName& InPinName, bool bIsInput)
 {
+	DECLARE_SCOPE_HIERARCHICAL_COUNTER_FUNC()
+
 	if (!EnsureModel())
 	{
 		return false;
@@ -585,6 +630,8 @@ bool UControlRigController::PrepareCycleCheckingForPin(const FName& InNodeName,
 
 bool UControlRigController::ResetCycleCheck()
 {
+	DECLARE_SCOPE_HIERARCHICAL_COUNTER_FUNC()
+
 	if (!EnsureModel())
 	{
 		return false;
@@ -595,6 +642,8 @@ bool UControlRigController::ResetCycleCheck()
 
 bool UControlRigController::CanLink(const FName& InSourceNodeName, const FName& InSourceOutputPinName, const FName& InTargetNodeName, const FName& InTargetInputPinName, FString* OutFailureReason, bool bReportError) const
 {
+	DECLARE_SCOPE_HIERARCHICAL_COUNTER_FUNC()
+
 	if (!EnsureModel())
 	{
 		return false;
@@ -642,6 +691,8 @@ bool UControlRigController::CanLink(const FName& InSourceNodeName, const FName&
 
 bool UControlRigController::MakeLink(const FName& InSourceNodeName, const FName& InSourceOutputPinName, const FName& InTargetNodeName, const FName& InTargetInputPinName, FString* OutFailureReason, bool bUndo)
 {
+	DECLARE_SCOPE_HIERARCHICAL_COUNTER_FUNC()
+
 	if (!EnsureModel())
 	{
 		return false;
@@ -703,6 +754,8 @@ bool UControlRigController::MakeLink(const FName& InSourceNodeName, const FName&
 
 bool UControlRigController::BreakLink(const FName& InSourceNodeName, const FName& InSourceOutputPinName, const FName& InTargetNodeName, const FName& InTargetInputPinName, bool bUndo)
 {
+	DECLARE_SCOPE_HIERARCHICAL_COUNTER_FUNC()
+
 	if (!EnsureModel())
 	{
 		return false;
@@ -738,6 +791,8 @@ bool UControlRigController::BreakLink(const FName& InSourceNodeName, const FName
 
 bool UControlRigController::BreakLinks(const FName& InNodeName, const FName& InPinName, bool bIsInput, bool bUndo)
 {
+	DECLARE_SCOPE_HIERARCHICAL_COUNTER_FUNC()
+
 	if (!EnsureModel())
 	{
 		return false;
@@ -760,6 +815,8 @@ bool UControlRigController::BreakLinks(const FName& InNodeName, const FName& InP
 
 bool UControlRigController::GetPinDefaultValue(const FName& InNodeName, const FName& InPinName, FString& OutDefaultValue) const
 {
+	DECLARE_SCOPE_HIERARCHICAL_COUNTER_FUNC()
+
 	if (!EnsureModel())
 	{
 		return false;
@@ -770,6 +827,8 @@ bool UControlRigController::GetPinDefaultValue(const FName& InNodeName, const FN
 
 bool UControlRigController::SetPinDefaultValue(const FName& InNodeName, const FName& InPinName, const FString& InDefaultValue, bool bLog, bool bUndo)
 {
+	DECLARE_SCOPE_HIERARCHICAL_COUNTER_FUNC()
+
 	if (!EnsureModel())
 	{
 		return false;
@@ -781,6 +840,8 @@ bool UControlRigController::SetPinDefaultValue(const FName& InNodeName, const FN
 
 bool UControlRigController::GetPinDefaultValueBool(const FName& InNodeName, const FName& InPinName, bool& OutDefaultValue) const
 {
+	DECLARE_SCOPE_HIERARCHICAL_COUNTER_FUNC()
+
 	FString DefaultValueString;
 	if(!GetPinDefaultValue(InNodeName, InPinName, DefaultValueString))
 	{
@@ -792,12 +853,16 @@ bool UControlRigController::GetPinDefaultValueBool(const FName& InNodeName, cons
 
 bool UControlRigController::SetPinDefaultValueBool(const FName& InNodeName, const FName& InPinName, bool InDefaultValue, bool bLog, bool bUndo)
 {
+	DECLARE_SCOPE_HIERARCHICAL_COUNTER_FUNC()
+
 	FString DefaultValueString = InDefaultValue ? TEXT("True") : TEXT("False");
 	return SetPinDefaultValue(InNodeName, InPinName, DefaultValueString, bLog, bUndo);
 }
 
 bool UControlRigController::GetPinDefaultValueFloat(const FName& InNodeName, const FName& InPinName, float& OutDefaultValue) const
 {
+	DECLARE_SCOPE_HIERARCHICAL_COUNTER_FUNC()
+
 	FString DefaultValueString;
 	if(!GetPinDefaultValue(InNodeName, InPinName, DefaultValueString))
 	{
@@ -809,12 +874,16 @@ bool UControlRigController::GetPinDefaultValueFloat(const FName& InNodeName, con
 
 bool UControlRigController::SetPinDefaultValueFloat(const FName& InNodeName, const FName& InPinName, float InDefaultValue, bool bLog, bool bUndo)
 {
+	DECLARE_SCOPE_HIERARCHICAL_COUNTER_FUNC()
+
 	FString DefaultValueString = FString::SanitizeFloat((double)InDefaultValue);
 	return SetPinDefaultValue(InNodeName, InPinName, DefaultValueString, bLog, bUndo);
 }
 
 bool UControlRigController::GetPinDefaultValueInt(const FName& InNodeName, const FName& InPinName, int32& OutDefaultValue) const
 {
+	DECLARE_SCOPE_HIERARCHICAL_COUNTER_FUNC()
+
 	FString DefaultValueString;
 	if(!GetPinDefaultValue(InNodeName, InPinName, DefaultValueString))
 	{
@@ -826,12 +895,16 @@ bool UControlRigController::GetPinDefaultValueInt(const FName& InNodeName, const
 
 bool UControlRigController::SetPinDefaultValueInt(const FName& InNodeName, const FName& InPinName, int32 InDefaultValue, bool bLog, bool bUndo)
 {
+	DECLARE_SCOPE_HIERARCHICAL_COUNTER_FUNC()
+
 	FString DefaultValueString = FString::FormatAsNumber(InDefaultValue);
 	return SetPinDefaultValue(InNodeName, InPinName, DefaultValueString, bLog, bUndo);
 }
 
 bool UControlRigController::GetPinDefaultValueName(const FName& InNodeName, const FName& InPinName, FName& OutDefaultValue) const
 {
+	DECLARE_SCOPE_HIERARCHICAL_COUNTER_FUNC()
+
 	FString DefaultValueString;
 	if(!GetPinDefaultValue(InNodeName, InPinName, DefaultValueString))
 	{
@@ -843,42 +916,58 @@ bool UControlRigController::GetPinDefaultValueName(const FName& InNodeName, cons
 
 bool UControlRigController::SetPinDefaultValueName(const FName& InNodeName, const FName& InPinName, const FName& InDefaultValue, bool bLog, bool bUndo)
 {
+	DECLARE_SCOPE_HIERARCHICAL_COUNTER_FUNC()
+
 	FString DefaultValueString = InDefaultValue.ToString();
 	return SetPinDefaultValue(InNodeName, InPinName, DefaultValueString, bLog, bUndo);
 }
 
 bool UControlRigController::GetPinDefaultValueVector(const FName& InNodeName, const FName& InPinName, FVector& OutDefaultValue) const
 {
+	DECLARE_SCOPE_HIERARCHICAL_COUNTER_FUNC()
+
 	return GetPinDefaultValueStruct<FVector>(InNodeName, InPinName, OutDefaultValue);
 }
 
 bool UControlRigController::SetPinDefaultValueVector(const FName& InNodeName, const FName& InPinName, const FVector& InDefaultValue, bool bLog, bool bUndo)
 {
+	DECLARE_SCOPE_HIERARCHICAL_COUNTER_FUNC()
+
 	return SetPinDefaultValueStruct<FVector>(InNodeName, InPinName, InDefaultValue, bLog, bUndo);
 }
 
 bool UControlRigController::GetPinDefaultValueQuat(const FName& InNodeName, const FName& InPinName, FQuat& OutDefaultValue) const
 {
+	DECLARE_SCOPE_HIERARCHICAL_COUNTER_FUNC()
+
 	return GetPinDefaultValueStruct<FQuat>(InNodeName, InPinName, OutDefaultValue);
 }
 
 bool UControlRigController::SetPinDefaultValueQuat(const FName& InNodeName, const FName& InPinName, const FQuat& InDefaultValue, bool bLog, bool bUndo)
 {
+	DECLARE_SCOPE_HIERARCHICAL_COUNTER_FUNC()
+
 	return SetPinDefaultValueStruct<FQuat>(InNodeName, InPinName, InDefaultValue, bLog, bUndo);
 }
 
 bool UControlRigController::GetPinDefaultValueTransform(const FName& InNodeName, const FName& InPinName, FTransform& OutDefaultValue) const
 {
+	DECLARE_SCOPE_HIERARCHICAL_COUNTER_FUNC()
+
 	return GetPinDefaultValueStruct<FTransform>(InNodeName, InPinName, OutDefaultValue);
 }
 
 bool UControlRigController::SetPinDefaultValueTransform(const FName& InNodeName, const FName& InPinName, const FTransform& InDefaultValue, bool bLog, bool bUndo)
 {
+	DECLARE_SCOPE_HIERARCHICAL_COUNTER_FUNC()
+
 	return SetPinDefaultValueStruct<FTransform>(InNodeName, InPinName, InDefaultValue, bLog, bUndo);
 }
 
 bool UControlRigController::AddArrayPin(const FName& InNodeName, const FName& InPinName, const FString& InDefaultValue, bool bUndo)
 {
+	DECLARE_SCOPE_HIERARCHICAL_COUNTER_FUNC()
+
 	if (!EnsureModel())
 	{
 		return false;
@@ -902,6 +991,8 @@ bool UControlRigController::AddArrayPin(const FName& InNodeName, const FName& In
 
 bool UControlRigController::PopArrayPin(const FName& InNodeName, const FName& InPinName, bool bUndo)
 {
+	DECLARE_SCOPE_HIERARCHICAL_COUNTER_FUNC()
+
 	if (!EnsureModel())
 	{
 		return false;
@@ -925,6 +1016,8 @@ bool UControlRigController::PopArrayPin(const FName& InNodeName, const FName& In
 
 bool UControlRigController::ClearArrayPin(const FName& InNodeName, const FName& InPinName, bool bUndo)
 {
+	DECLARE_SCOPE_HIERARCHICAL_COUNTER_FUNC()
+
 	if (!EnsureModel())
 	{
 		return false;
@@ -948,6 +1041,8 @@ bool UControlRigController::ClearArrayPin(const FName& InNodeName, const FName&
 
 bool UControlRigController::SetArrayPinSize(const FName& InNodeName, const FName& InPinName, int32 InSize, const FString& InDefaultValue, bool bUndo)
 {
+	DECLARE_SCOPE_HIERARCHICAL_COUNTER_FUNC()
+
 	if (!EnsureModel())
 	{
 		return false;
@@ -972,6 +1067,8 @@ bool UControlRigController::SetArrayPinSize(const FName& InNodeName, const FName
 
 bool UControlRigController::ExpandPin(const FName& InNodeName, const FName& InPinName, bool bIsInput, bool bExpanded, bool bUndo)
 {
+	DECLARE_SCOPE_HIERARCHICAL_COUNTER_FUNC()
+
 	if (!EnsureModel())
 	{
 		return false;
@@ -982,6 +1079,8 @@ bool UControlRigController::ExpandPin(const FName& InNodeName, const FName& InPi
 
 bool UControlRigController::FindPinTypeFromDataType(const FName& InDataType, FEdGraphPinType& OutPinType)
 {
+	DECLARE_SCOPE_HIERARCHICAL_COUNTER_FUNC()
+
 	TArray<FEdGraphPinType> PinTypes;
 	UControlRigModel::GetParameterPinTypes(PinTypes);
 
@@ -1008,6 +1107,8 @@ bool UControlRigController::FindPinTypeFromDataType(const FName& InDataType, FEd
 
 bool UControlRigController::ResendAllNotifications()
 {
+	DECLARE_SCOPE_HIERARCHICAL_COUNTER_FUNC()
+
 	if (!EnsureModel())
 	{
 		return false;
@@ -1018,6 +1119,8 @@ bool UControlRigController::ResendAllNotifications()
 
 bool UControlRigController::ResendAllPinDefaultNotifications()
 {
+	DECLARE_SCOPE_HIERARCHICAL_COUNTER_FUNC()
+
 	if (!EnsureModel())
 	{
 		return false;
@@ -1027,6 +1130,8 @@ bool UControlRigController::ResendAllPinDefaultNotifications()
 
 bool UControlRigController::Undo()
 {
+	DECLARE_SCOPE_HIERARCHICAL_COUNTER_FUNC()
+
 	if (!EnsureModel())
 	{
 		return false;
@@ -1037,6 +1142,8 @@ bool UControlRigController::Undo()
 
 bool UControlRigController::Redo()
 {
+	DECLARE_SCOPE_HIERARCHICAL_COUNTER_FUNC()
+
 	if (!EnsureModel())
 	{
 		return false;
diff --git a/Engine/Plugins/Experimental/ControlRig/Source/ControlRigDeveloper/Private/ControlRigDAG.cpp b/Engine/Plugins/Experimental/ControlRig/Source/ControlRigDeveloper/Private/ControlRigDAG.cpp
--- a/Engine/Plugins/Experimental/ControlRig/Source/ControlRigDeveloper/Private/ControlRigDAG.cpp
+++ b/Engine/Plugins/Experimental/ControlRig/Source/ControlRigDeveloper/Private/ControlRigDAG.cpp
@@ -10,9 +10,9 @@ FControlRigDAG::FControlRigDAG()
 {
 }
 
-void FControlRigDAG::AddNode(bool InIsMutable, const FName& InName)
+void FControlRigDAG::AddNode(bool InIsMutable, bool InIsOutputParameter, const FName& InName)
 {
-	Nodes.Add(FNode(InName, Nodes.Num(), InIsMutable));
+	Nodes.Add(FNode(InName, Nodes.Num(), InIsMutable, InIsOutputParameter));
 }
 
 void FControlRigDAG::AddLink(const int32 FromNode, const int32 ToNode, const int32 FromOrder, const int32 ToOrder)
@@ -39,22 +39,6 @@ bool FControlRigDAG::TopologicalSort(TArray<FNode>& OutOrder, TArray<FNode>& Out
 		return false;
 	}
 
-	// find all of the left mutable nodes on the left
-	TArray<FNode> MutableLeafNodes;
-	for (const FNode& Node : Nodes)
-	{
-		if (!Node.IsMutable)
-		{
-			continue;
-		}
-		if (Node.Inputs.Num() > 0)
-		{
-			continue;
-		}
-		MutableLeafNodes.Add(Node);
-	}
-	check(MutableLeafNodes.Num() > 0);
-
 	struct Local
 	{
 		static void VisitNode(const FNode& Node, TArray<bool>& Visited, TArray<FNode>& SortedNodes, const TArray<FNode>& Nodes)
@@ -83,10 +67,32 @@ bool FControlRigDAG::TopologicalSort(TArray<FNode>& OutOrder, TArray<FNode>& Out
 		}
 	};
 
+	// find all of the left mutable nodes on the left
+	TArray<FNode> LeafNodes;
+	TArray<FNode> OutputParameterNodes;
+	for (const FNode& Node : Nodes)
+	{
+		if (Node.IsOutputParameter)
+		{
+			OutputParameterNodes.Add(Node);
+		}
+		else if (Node.IsMutable)
+		{
+			if (Node.Inputs.Num() == 0)
+			{
+				LeafNodes.Add(Node);
+			}
+		}
+	}
+
+	LeafNodes.Append(OutputParameterNodes);
+
+	check(LeafNodes.Num() > 0);
+
 	TArray<bool> NodeVisited;
 	NodeVisited.SetNumZeroed(Nodes.Num());
 	OutOrder.Reset();
-	for (const FNode& LeafNode : MutableLeafNodes)
+	for (const FNode& LeafNode : LeafNodes)
 	{
 		Local::VisitNode(LeafNode, NodeVisited, OutOrder, Nodes);
 	}
diff --git a/Engine/Plugins/Experimental/ControlRig/Source/ControlRigDeveloper/Private/ControlRigGraphTraverser.cpp b/Engine/Plugins/Experimental/ControlRig/Source/ControlRigDeveloper/Private/ControlRigGraphTraverser.cpp
--- a/Engine/Plugins/Experimental/ControlRig/Source/ControlRigDeveloper/Private/ControlRigGraphTraverser.cpp
+++ b/Engine/Plugins/Experimental/ControlRig/Source/ControlRigDeveloper/Private/ControlRigGraphTraverser.cpp
@@ -39,6 +39,12 @@ bool FControlRigGraphTraverser::IsWiredToExecution(const FControlRigModelNode* N
 		return true;
 	}
 
+	if (Node->IsParameter() && Node->ParameterType == EControlRigModelParameterType::Output)
+	{
+		VisitedNodes.Add(Node->Name, true);
+		return true;
+	}
+
 	VisitedNodes.Add(Node->Name, false);
 
 	bool bFoundWiredPin = false;
diff --git a/Engine/Plugins/Experimental/ControlRig/Source/ControlRigDeveloper/Private/ControlRigModel.cpp b/Engine/Plugins/Experimental/ControlRig/Source/ControlRigDeveloper/Private/ControlRigModel.cpp
--- a/Engine/Plugins/Experimental/ControlRig/Source/ControlRigDeveloper/Private/ControlRigModel.cpp
+++ b/Engine/Plugins/Experimental/ControlRig/Source/ControlRigDeveloper/Private/ControlRigModel.cpp
@@ -278,6 +278,8 @@ const TArray<FControlRigModelNode>& UControlRigModel::Nodes() const
 
 TArray<FControlRigModelNode> UControlRigModel::SelectedNodes() const
 {
+	DECLARE_SCOPE_HIERARCHICAL_COUNTER_FUNC()
+
 	TArray<FControlRigModelNode> Nodes;
 	for (const FName& Name : _SelectedNodes)
 	{
@@ -290,6 +292,8 @@ TArray<FControlRigModelNode> UControlRigModel::SelectedNodes() const
 
 bool UControlRigModel::IsNodeSelected(const FName& InName) const
 {
+	DECLARE_SCOPE_HIERARCHICAL_COUNTER_FUNC()
+
 	return _SelectedNodes.Contains(InName);
 }
 
@@ -300,6 +304,8 @@ const TArray<FControlRigModelLink>& UControlRigModel::Links() const
 
 TArray<FControlRigModelPin> UControlRigModel::LinkedPins(const FControlRigModelPair& InPin) const
 {
+	DECLARE_SCOPE_HIERARCHICAL_COUNTER_FUNC()
+
 	ensure(InPin.Node >= 0 && InPin.Node < _Nodes.Num());
 	ensure(InPin.Pin >= 0 && InPin.Pin < _Nodes[InPin.Node].Pins.Num());
 
@@ -324,6 +330,8 @@ TArray<FControlRigModelPin> UControlRigModel::LinkedPins(const FControlRigModelP
 
 TArray<FControlRigModelPin> UControlRigModel::LinkedPins(const FName& InNodeName, const FName& InPinName, bool bLookForInput) const
 {
+	DECLARE_SCOPE_HIERARCHICAL_COUNTER_FUNC()
+
 	TArray<FControlRigModelPin> Pins;
 	const FControlRigModelPin* Pin = FindPin(InNodeName, InPinName, bLookForInput);
 	if(Pin == nullptr)
@@ -342,6 +350,8 @@ TArray<FControlRigModelPin> UControlRigModel::LinkedPins(const FName& InNodeName
 
 TArray<FControlRigModelNode> UControlRigModel::Parameters() const
 {
+	DECLARE_SCOPE_HIERARCHICAL_COUNTER_FUNC()
+
 	TArray<FControlRigModelNode> ParameterNodes;
 
 	for (const FControlRigModelNode& Node : _Nodes)
@@ -362,6 +372,8 @@ UControlRigModel::FModifiedEvent& UControlRigModel::OnModified()
 
 bool UControlRigModel::Clear()
 {
+	DECLARE_SCOPE_HIERARCHICAL_COUNTER_FUNC()
+
 	if (_Nodes.Num() == 0)
 	{
 		return false;
@@ -381,16 +393,22 @@ bool UControlRigModel::Clear()
 
 bool UControlRigModel::IsNodeNameAvailable(const FName& InName) const
 {
+	DECLARE_SCOPE_HIERARCHICAL_COUNTER_FUNC()
+
 	return FindNode(InName) == nullptr;
 }
 
 FName UControlRigModel::GetUniqueNodeName(const FName& InName) const
 {
+	DECLARE_SCOPE_HIERARCHICAL_COUNTER_FUNC()
+
 	return UtilityHelpers::CreateUniqueName(InName, [this](const FName& CurName) { return IsNodeNameAvailable(CurName); });
 }
 
 const FControlRigModelNode* UControlRigModel::FindNode(const FName& InName) const
 {
+	DECLARE_SCOPE_HIERARCHICAL_COUNTER_FUNC()
+
 	for (const FControlRigModelNode& Node : _Nodes)
 	{
 		if (Node.Name == InName)
@@ -409,6 +427,8 @@ const FControlRigModelNode* UControlRigModel::FindNode(int32 InNodeIndex) const
 
 bool UControlRigModel::AddNode(const FControlRigModelNode& InNode, bool bUndo)
 {
+	DECLARE_SCOPE_HIERARCHICAL_COUNTER_FUNC()
+
 	ensure(!InNode.IsParameter() && (InNode.FunctionName != NAME_None));
 
 	FName DesiredNodeName = InNode.Name == NAME_None ? InNode.FunctionName : InNode.Name;
@@ -652,6 +672,8 @@ bool UControlRigModel::AddNode(const FControlRigModelNode& InNode, bool bUndo)
 
 bool UControlRigModel::AddParameter(const FName& InName, const FEdGraphPinType& InDataType, EControlRigModelParameterType InParameterType, const FVector2D& InPosition, bool bUndo)
 {
+	DECLARE_SCOPE_HIERARCHICAL_COUNTER_FUNC()
+
 	FControlRigModelNode Parameter;
 	Parameter.Name = GetUniqueNodeName(InName);
 
@@ -701,6 +723,8 @@ bool UControlRigModel::AddParameter(const FName& InName, const FEdGraphPinType&
 
 bool UControlRigModel::AddComment(const FName& InName, const FString& InText, const FVector2D& InPosition, const FVector2D& InSize, const FLinearColor& InColor, bool bUndo)
 {
+	DECLARE_SCOPE_HIERARCHICAL_COUNTER_FUNC()
+
 	FControlRigModelNode Comment;
 	Comment.Name = GetUniqueNodeName(InName);
 
@@ -745,6 +769,8 @@ bool UControlRigModel::AddComment(const FName& InName, const FString& InText, co
 
 bool UControlRigModel::RemoveNode(const FName& InName, bool bUndo)
 {
+	DECLARE_SCOPE_HIERARCHICAL_COUNTER_FUNC()
+
 	const FControlRigModelNode* Node = FindNode(InName);
 	if(Node != nullptr)
 	{
@@ -866,6 +892,8 @@ bool UControlRigModel::RemoveNode(const FName& InName, bool bUndo)
 
 bool UControlRigModel::SetNodePosition(const FName& InName, const FVector2D& InPosition, bool bUndo)
 {
+	DECLARE_SCOPE_HIERARCHICAL_COUNTER_FUNC()
+
 	const FControlRigModelNode* Node = FindNode(InName);
 	if(Node != nullptr)
 	{
@@ -920,6 +948,8 @@ bool UControlRigModel::SetNodePosition(const FName& InName, const FVector2D& InP
 
 bool UControlRigModel::SetNodeSize(const FName& InName, const FVector2D& InSize, bool bUndo)
 {
+	DECLARE_SCOPE_HIERARCHICAL_COUNTER_FUNC()
+
 	const FControlRigModelNode* Node = FindNode(InName);
 	if (Node != nullptr)
 	{
@@ -974,6 +1004,8 @@ bool UControlRigModel::SetNodeSize(const FName& InName, const FVector2D& InSize,
 
 bool UControlRigModel::SetNodeColor(const FName& InName, const FLinearColor& InColor, bool bUndo)
 {
+	DECLARE_SCOPE_HIERARCHICAL_COUNTER_FUNC()
+
 	const FControlRigModelNode* Node = FindNode(InName);
 	if (Node != nullptr)
 	{
@@ -1029,6 +1061,8 @@ bool UControlRigModel::SetNodeColor(const FName& InName, const FLinearColor& InC
 
 bool UControlRigModel::SetParameterType(const FName& InName, EControlRigModelParameterType InParameterType, bool bUndo)
 {
+	DECLARE_SCOPE_HIERARCHICAL_COUNTER_FUNC()
+
 	ensure(InParameterType != EControlRigModelParameterType::None);
 
 	const FControlRigModelNode* Node = FindNode(InName);
@@ -1079,6 +1113,8 @@ bool UControlRigModel::SetParameterType(const FName& InName, EControlRigModelPar
 
 bool UControlRigModel::SetCommentText(const FName& InName, const FString& InText, bool bUndo)
 {
+	DECLARE_SCOPE_HIERARCHICAL_COUNTER_FUNC()
+
 	const FControlRigModelNode* Node = FindNode(InName);
 	if (Node != nullptr)
 	{
@@ -1127,6 +1163,8 @@ bool UControlRigModel::SetCommentText(const FName& InName, const FString& InText
 
 bool UControlRigModel::RenameNode(const FName& InOldNodeName, const FName& InNewNodeName, bool bUndo)
 {
+	DECLARE_SCOPE_HIERARCHICAL_COUNTER_FUNC()
+
 	FControlRigModelNode* Node = (FControlRigModelNode*)FindNode(InOldNodeName);
 	if (Node == nullptr)
 	{
@@ -1200,6 +1238,8 @@ bool UControlRigModel::RenameNode(const FName& InOldNodeName, const FName& InNew
 
 bool UControlRigModel::SelectNode(const FName& InName, bool bInSelected)
 {
+	DECLARE_SCOPE_HIERARCHICAL_COUNTER_FUNC()
+
 	if (bIsSelecting)
 	{
 		return false;
@@ -1246,6 +1286,8 @@ bool UControlRigModel::AreCompatibleTypes(const FEdGraphPinType& A, const FEdGra
 
 bool UControlRigModel::PrepareCycleCheckingForPin(int32 InNodeIndex, int32 InPinIndex)
 {
+	DECLARE_SCOPE_HIERARCHICAL_COUNTER_FUNC()
+
 	ensure(InNodeIndex >= 0 && InNodeIndex < _Nodes.Num());
 	ensure(InPinIndex >= 0 && InPinIndex < _Nodes[InNodeIndex].Pins.Num());
 
@@ -1299,6 +1341,8 @@ void UControlRigModel::ResetCycleCheck()
 
 bool UControlRigModel::CanLink(int32 InSourceNodeIndex, int32 InSourcePinIndex, int32 InTargetNodeIndex, int32 InTargetPinIndex, FString* OutFailureReason)
 {
+	DECLARE_SCOPE_HIERARCHICAL_COUNTER_FUNC()
+
 	ensure(InSourceNodeIndex >= 0 && InSourceNodeIndex < _Nodes.Num());
 	ensure(InTargetNodeIndex >= 0 && InTargetNodeIndex < _Nodes.Num());
 	ensure(InSourcePinIndex >= 0 && InSourcePinIndex < _Nodes[InSourceNodeIndex].Pins.Num());
@@ -1445,6 +1489,8 @@ bool UControlRigModel::CanLink(int32 InSourceNodeIndex, int32 InSourcePinIndex,
 
 bool UControlRigModel::MakeLink(int32 InSourceNodeIndex, int32 InSourcePinIndex, int32 InTargetNodeIndex, int32 InTargetPinIndex, bool bUndo)
 {
+	DECLARE_SCOPE_HIERARCHICAL_COUNTER_FUNC()
+
 	ensure(InSourceNodeIndex >= 0 && InSourceNodeIndex < _Nodes.Num());
 	ensure(InTargetNodeIndex >= 0 && InTargetNodeIndex < _Nodes.Num());
 	ensure(InSourcePinIndex >= 0 && InSourcePinIndex < _Nodes[InSourceNodeIndex].Pins.Num());
@@ -1519,6 +1565,8 @@ bool UControlRigModel::MakeLink(int32 InSourceNodeIndex, int32 InSourcePinIndex,
 
 bool UControlRigModel::BreakLink(int32 InSourceNodeIndex, int32 InSourcePinIndex, int32 InTargetNodeIndex, int32 InTargetPinIndex, bool bUndo)
 {
+	DECLARE_SCOPE_HIERARCHICAL_COUNTER_FUNC()
+
 	for (int32 LinkIndex = 0; LinkIndex < _Links.Num(); LinkIndex++)
 	{
 		FControlRigModelLink Link = _Links[LinkIndex];
@@ -1578,6 +1626,8 @@ bool UControlRigModel::BreakLink(int32 InSourceNodeIndex, int32 InSourcePinIndex
 
 bool UControlRigModel::BreakLinks(int32 InNodeIndex, int32 InPinIndex, bool bUndo)
 {
+	DECLARE_SCOPE_HIERARCHICAL_COUNTER_FUNC()
+
 #if CONTROLRIG_UNDO
 	FAction MainAction;
 	if (bUndo)
@@ -1658,6 +1708,8 @@ bool UControlRigModel::BreakLinks(int32 InNodeIndex, int32 InPinIndex, bool bUnd
 
 const FControlRigModelPin* UControlRigModel::FindPin(const FName& InNodeName, const FName& InPinName, bool bLookForInput) const
 {
+	DECLARE_SCOPE_HIERARCHICAL_COUNTER_FUNC()
+
 	const FControlRigModelNode* Node = FindNode(InNodeName);
 	if (Node != nullptr)
 	{
@@ -1675,6 +1727,8 @@ const FControlRigModelPin* UControlRigModel::FindPin(const FControlRigModelPair&
 
 const FControlRigModelPin* UControlRigModel::FindSubPin(const FControlRigModelPin* InParentPin, const FName& InSubPinName) const
 {
+	DECLARE_SCOPE_HIERARCHICAL_COUNTER_FUNC()
+
 	if (InParentPin->SubPins.Num() == 0)
 	{
 		return nullptr;
@@ -1707,6 +1761,8 @@ const FControlRigModelPin* UControlRigModel::FindParentPin(const FControlRigMode
 
 const FControlRigModelPin* UControlRigModel::FindPinFromPath(const FString& InPinPath, bool bLookForInput) const
 {
+	DECLARE_SCOPE_HIERARCHICAL_COUNTER_FUNC()
+
 	FString Left, Right;
 	if (!InPinPath.Split(TEXT("."), &Left, &Right))
 	{
@@ -1724,6 +1780,8 @@ const FControlRigModelLink* UControlRigModel::FindLink(int32 InLinkIndex) const
 
 void UControlRigModel::SplitPinPath(const FString& InPinPath, FString& OutLeft, FString& OutRight, bool bInSplitForNodeName)
 {
+	DECLARE_SCOPE_HIERARCHICAL_COUNTER_FUNC()
+
 	FString Left, Right;
 	int32 PeriodPos = INDEX_NONE;
 	int32 BracketPos = INDEX_NONE;
@@ -1776,6 +1834,8 @@ void UControlRigModel::SplitPinPath(const FString& InPinPath, FString& OutLeft,
 
 const FControlRigModelPin* UControlRigModel::GetParentPin(const FControlRigModelPair& InPin) const
 {
+	DECLARE_SCOPE_HIERARCHICAL_COUNTER_FUNC()
+
 	ensure(InPin.Node >= 0 && InPin.Node < _Nodes.Num());
 	ensure(InPin.Pin >= 0 && InPin.Pin < _Nodes[InPin.Node].Pins.Num());
 	const FControlRigModelNode& Node = _Nodes[InPin.Node];
@@ -1795,6 +1855,8 @@ const FControlRigModelPin* UControlRigModel::GetParentPin(const FControlRigModel
 
 FString UControlRigModel::GetPinPath(const FControlRigModelPair& InPin, bool bIncludeNodeName) const
 {
+	DECLARE_SCOPE_HIERARCHICAL_COUNTER_FUNC()
+
 	ensure(InPin.Node >= 0 && InPin.Node < _Nodes.Num());
 	ensure(InPin.Pin >= 0 && InPin.Pin < _Nodes[InPin.Node].Pins.Num());
 	const FControlRigModelNode& Node = _Nodes[InPin.Node];
@@ -1803,6 +1865,8 @@ FString UControlRigModel::GetPinPath(const FControlRigModelPair& InPin, bool bIn
 
 bool UControlRigModel::GetPinDefaultValue(const FName& InNodeName, const FName& InPinName, FString& OutValue) const
 {
+	DECLARE_SCOPE_HIERARCHICAL_COUNTER_FUNC()
+
 	const FControlRigModelPin* Pin = FindPin(InNodeName, InPinName, true);
 	if (Pin != nullptr)
 	{
@@ -1820,6 +1884,8 @@ bool UControlRigModel::GetPinDefaultValue(const FName& InNodeName, const FName&
 
 bool UControlRigModel::GetPinDefaultValue(const FControlRigModelPair& InPin, FString& OutValue) const
 {
+	DECLARE_SCOPE_HIERARCHICAL_COUNTER_FUNC()
+
 	ensure(InPin.Node >= 0 && InPin.Node < _Nodes.Num());
 	ensure(InPin.Pin >= 0 && InPin.Pin < _Nodes[InPin.Node].Pins.Num());
 	OutValue = _Nodes[InPin.Node].Pins[InPin.Pin].DefaultValue;
@@ -1828,6 +1894,8 @@ bool UControlRigModel::GetPinDefaultValue(const FControlRigModelPair& InPin, FSt
 
 bool UControlRigModel::SetPinDefaultValue(const FName& InNodeName, const FName& InPinName, const FString& InValue, bool bUndo)
 {
+	DECLARE_SCOPE_HIERARCHICAL_COUNTER_FUNC()
+
 	const FControlRigModelPin* Pin = FindPin(InNodeName, InPinName, true);
 	if (Pin != nullptr)
 	{
@@ -1845,6 +1913,8 @@ bool UControlRigModel::SetPinDefaultValue(const FName& InNodeName, const FName&
 
 bool UControlRigModel::SetPinDefaultValue(const FControlRigModelPair& InPin, const FString& InValue, bool bUndo)
 {
+	DECLARE_SCOPE_HIERARCHICAL_COUNTER_FUNC()
+
 	ensure(InPin.Node >= 0 && InPin.Node < _Nodes.Num());
 	ensure(InPin.Pin >= 0 && InPin.Pin < _Nodes[InPin.Node].Pins.Num());
 
@@ -1884,6 +1954,8 @@ bool UControlRigModel::SetPinDefaultValue(const FControlRigModelPair& InPin, con
 
 bool UControlRigModel::SetPinArraySize(const FControlRigModelPair& InPin, int32 InArraySize, const FString& InDefaultValue, bool bUndo)
 {
+	DECLARE_SCOPE_HIERARCHICAL_COUNTER_FUNC()
+
 	ensure(InPin.Node >= 0 && InPin.Node < _Nodes.Num());
 	ensure(InPin.Pin >= 0 && InPin.Pin < _Nodes[InPin.Node].Pins.Num());
 
@@ -2167,6 +2239,8 @@ bool UControlRigModel::SetPinArraySize(const FControlRigModelPair& InPin, int32
 
 bool UControlRigModel::ExpandPin(const FName& InNodeName, const FName& InPinName, bool bIsInput, bool bInExpanded, bool bUndo)
 {
+	DECLARE_SCOPE_HIERARCHICAL_COUNTER_FUNC()
+
 	const FControlRigModelPin* Pin = FindPin(InNodeName, InPinName, bIsInput);
 	if (Pin != nullptr)
 	{
@@ -2215,6 +2289,8 @@ bool UControlRigModel::ExpandPin(const FName& InNodeName, const FName& InPinName
 
 bool UControlRigModel::ResendAllNotifications()
 {
+	DECLARE_SCOPE_HIERARCHICAL_COUNTER_FUNC()
+
 	if (_ModifiedEvent.IsBound())
 	{
 		_ModifiedEvent.Broadcast(this, EControlRigModelNotifType::ModelCleared, nullptr);
@@ -2245,6 +2321,8 @@ bool UControlRigModel::ResendAllNotifications()
 
 bool UControlRigModel::ResendAllPinDefaultNotifications()
 {
+	DECLARE_SCOPE_HIERARCHICAL_COUNTER_FUNC()
+
 	if (_ModifiedEvent.IsBound())
 	{
 		for (FControlRigModelNode& Node : _Nodes)
@@ -2290,6 +2368,8 @@ bool UControlRigModel::ShouldStructBeUnfolded(const UStruct* Struct)
 
 FEdGraphPinType UControlRigModel::GetPinTypeFromField(UProperty* Property)
 {
+	DECLARE_SCOPE_HIERARCHICAL_COUNTER_FUNC()
+
 	FEdGraphPinType PinType;
 	GetDefault<UEdGraphSchema_K2>()->ConvertPropertyToPinType(Property, PinType);
 	return PinType;
@@ -2297,6 +2377,8 @@ FEdGraphPinType UControlRigModel::GetPinTypeFromField(UProperty* Property)
 
 void UControlRigModel::AddNodePinsForFunction(FControlRigModelNode& Node)
 {
+	DECLARE_SCOPE_HIERARCHICAL_COUNTER_FUNC()
+
 	Node.Pins.Reset();
 	int32 LastAddedIndex = -1;
 	const UStruct* Struct = Node.UnitStruct();
@@ -2333,6 +2415,8 @@ void UControlRigModel::AddNodePinsForFunction(FControlRigModelNode& Node)
 
 void UControlRigModel::SetNodePinDefaultsForFunction(FControlRigModelNode& Node)
 {
+	DECLARE_SCOPE_HIERARCHICAL_COUNTER_FUNC()
+
 	UScriptStruct* ScriptStruct = (UScriptStruct*)Cast<UScriptStruct>(Node.UnitStruct());
 	if (ScriptStruct)
 	{
@@ -2365,6 +2449,8 @@ void UControlRigModel::SetNodePinDefaultsForFunction(FControlRigModelNode& Node)
 
 void UControlRigModel::AddNodePinsForParameter(FControlRigModelNode& Node, const FEdGraphPinType& InDataType)
 {
+	DECLARE_SCOPE_HIERARCHICAL_COUNTER_FUNC()
+
 	int32 LastAddedIndex = Node.Pins.Num() - 1;
 
 	FControlRigModelPin InputPin;
@@ -2397,6 +2483,8 @@ void UControlRigModel::AddNodePinsForParameter(FControlRigModelNode& Node, const
 
 void UControlRigModel::SetNodePinDefaultsForParameter(FControlRigModelNode& Node, const FEdGraphPinType& InDataType)
 {
+	DECLARE_SCOPE_HIERARCHICAL_COUNTER_FUNC()
+
 	UScriptStruct* ScriptStruct = (UScriptStruct*)Cast<UScriptStruct>(InDataType.PinSubCategoryObject);
 	if (ScriptStruct)
 	{
@@ -2443,6 +2531,8 @@ void UControlRigModel::SetNodePinDefaultsForParameter(FControlRigModelNode& Node
 
 void UControlRigModel::ConfigurePinFromField(FControlRigModelPin& Pin, UProperty* Property, FControlRigModelNode& Node)
 {
+	DECLARE_SCOPE_HIERARCHICAL_COUNTER_FUNC()
+
 	Pin.Type = GetPinTypeFromField(Property);
 	Pin.Name = Property->GetFName();
 	Pin.DisplayNameText = Property->GetDisplayNameText();
@@ -2464,12 +2554,18 @@ void UControlRigModel::ConfigurePinFromField(FControlRigModelPin& Pin, UProperty
 	{
 		Pin.CustomWidgetName = UControlRig::BoneNameMetaName;
 	}
+	else if (Property->HasMetaData(UControlRig::CurveNameMetaName))
+	{
+		Pin.CustomWidgetName = UControlRig::CurveNameMetaName;
+	}
 
 	Pin.TooltipText = Property->GetToolTipText();
 }
 
  int32 UControlRigModel::AddPinsRecursive(FControlRigModelNode& Node, int32 ParentIndex, const UStruct* Struct, EEdGraphPinDirection PinDirection, int32& LastAddedIndex)
 {
+	DECLARE_SCOPE_HIERARCHICAL_COUNTER_FUNC()
+	 
 	if (!ShouldStructBeUnfolded(Struct))
 	{
 		return 0;
@@ -2515,6 +2611,8 @@ void UControlRigModel::ConfigurePinFromField(FControlRigModelPin& Pin, UProperty
 
 int32 UControlRigModel::RemovePinsRecursive(FControlRigModelNode& Node, int32 PinIndex, bool bUndo)
 {
+	DECLARE_SCOPE_HIERARCHICAL_COUNTER_FUNC()
+
 	int32 NumberOfPinsRemoved = 0;
 	while (Node.Pins[PinIndex].SubPins.Num() > 0)
 	{
@@ -2537,6 +2635,8 @@ int32 UControlRigModel::RemovePinsRecursive(FControlRigModelNode& Node, int32 Pi
 
 void UControlRigModel::ConfigurePinIndices(FControlRigModelNode& Node)
 {
+	DECLARE_SCOPE_HIERARCHICAL_COUNTER_FUNC()
+
 	for (int32 PinIndex = 0; PinIndex < Node.Pins.Num(); PinIndex++)
 	{
 		Node.Pins[PinIndex].SubPins.Reset();
@@ -2555,6 +2655,8 @@ void UControlRigModel::ConfigurePinIndices(FControlRigModelNode& Node)
 
 void UControlRigModel::GetParameterPinTypes(TArray<FEdGraphPinType>& PinTypes)
 {
+	DECLARE_SCOPE_HIERARCHICAL_COUNTER_FUNC()
+
 	PinTypes.Add(FEdGraphPinType(UEdGraphSchema_K2::PC_Boolean, FName(NAME_None), nullptr, EPinContainerType::None, false, FEdGraphTerminalType()));
 	PinTypes.Add(FEdGraphPinType(UEdGraphSchema_K2::PC_Float, FName(NAME_None), nullptr, EPinContainerType::None, false, FEdGraphTerminalType()));
 	PinTypes.Add(FEdGraphPinType(UEdGraphSchema_K2::PC_Int, FName(NAME_None), nullptr, EPinContainerType::None, false, FEdGraphTerminalType()));
@@ -2573,6 +2675,8 @@ void UControlRigModel::GetParameterPinTypes(TArray<FEdGraphPinType>& PinTypes)
 
 void UControlRigModel::PostTransacted(const FTransactionObjectEvent& TransactionEvent)
 {
+	DECLARE_SCOPE_HIERARCHICAL_COUNTER_FUNC()
+
 	if (TransactionEvent.GetEventType() == ETransactionObjectEventType::UndoRedo)
 	{
 		while (ActionCount < UndoActions.Num())
@@ -2602,6 +2706,8 @@ void UControlRigModel::PostTransacted(const FTransactionObjectEvent& Transaction
 
 bool UControlRigModel::Undo()
 {
+	DECLARE_SCOPE_HIERARCHICAL_COUNTER_FUNC()
+
 	if (UndoActions.Num() == 0)
 	{
 		return false;
@@ -2617,6 +2723,8 @@ bool UControlRigModel::Undo()
 
 bool UControlRigModel::Redo()
 {
+	DECLARE_SCOPE_HIERARCHICAL_COUNTER_FUNC()
+
 	if (RedoActions.Num() == 0)
 	{
 		return false;
@@ -2632,6 +2740,8 @@ bool UControlRigModel::Redo()
 
 void UControlRigModel::PushAction(const UControlRigModel::FAction& InAction)
 {
+	DECLARE_SCOPE_HIERARCHICAL_COUNTER_FUNC()
+
 	ensure(InAction.IsValid());
 
 	if (InAction.Type == EControlRigModelNotifType::Invalid && InAction.SubActions.Num() == 0)
@@ -2666,6 +2776,8 @@ void UControlRigModel::PushAction(const UControlRigModel::FAction& InAction)
 
 bool UControlRigModel::UndoAction(const UControlRigModel::FAction& InAction)
 {
+	DECLARE_SCOPE_HIERARCHICAL_COUNTER_FUNC()
+
 	for (int32 SubActionIndex = InAction.SubActions.Num() - 1; SubActionIndex >= 0; SubActionIndex--)
 	{
 		if (!UndoAction(InAction.SubActions[SubActionIndex]))
@@ -2783,6 +2895,8 @@ bool UControlRigModel::UndoAction(const UControlRigModel::FAction& InAction)
 
 bool UControlRigModel::RedoAction(const UControlRigModel::FAction& InAction)
 {
+	DECLARE_SCOPE_HIERARCHICAL_COUNTER_FUNC()
+
 	ensure(InAction.IsValid());
 
 	switch (InAction.Type)
diff --git a/Engine/Plugins/Experimental/ControlRig/Source/ControlRigDeveloper/Private/Graph/ControlRigGraph.cpp b/Engine/Plugins/Experimental/ControlRig/Source/ControlRigDeveloper/Private/Graph/ControlRigGraph.cpp
--- a/Engine/Plugins/Experimental/ControlRig/Source/ControlRigDeveloper/Private/Graph/ControlRigGraph.cpp
+++ b/Engine/Plugins/Experimental/ControlRig/Source/ControlRigDeveloper/Private/Graph/ControlRigGraph.cpp
@@ -29,6 +29,8 @@ UControlRigGraph::UControlRigGraph()
 
 void UControlRigGraph::Initialize(UControlRigBlueprint* InBlueprint)
 {
+	DECLARE_SCOPE_HIERARCHICAL_COUNTER_FUNC()
+
 	InBlueprint->OnModified().RemoveAll(this);
 	InBlueprint->OnModified().AddUObject(this, &UControlRigGraph::HandleModelModified);
 }
@@ -50,6 +52,8 @@ void UControlRigGraph::Serialize(FArchive& Ar)
 #if WITH_EDITOR
 void UControlRigGraph::PostLoad()
 {
+	DECLARE_SCOPE_HIERARCHICAL_COUNTER_FUNC()
+
 	FoundHierarchyRefVariableNodes.Reset();
 	FoundHierarchyRefMutableNodes.Reset();
 	FoundHierarchyRefConnections.Reset();
@@ -133,6 +137,8 @@ void UControlRigGraph::PostLoad()
 
 void UControlRigGraph::OnBlueprintCompiledPostLoad(UBlueprint* InCompiledBlueprint)
 {
+	DECLARE_SCOPE_HIERARCHICAL_COUNTER_FUNC()
+
 	if (GetLinkerCustomVersion(FControlRigObjectVersion::GUID) < FControlRigObjectVersion::RemovalOfHierarchyRefPins)
 	{
 		UControlRigBlueprint* RigBlueprint = Cast<UControlRigBlueprint>(GetOuter());
@@ -268,8 +274,11 @@ void UControlRigGraph::OnBlueprintCompiledPostLoad(UBlueprint* InCompiledBluepri
 
 void UControlRigGraph::CacheBoneNameList(const FRigHierarchy& Hierarchy)
 {
+	DECLARE_SCOPE_HIERARCHICAL_COUNTER_FUNC()
+
 	TArray<FString> Names;
-	for (const FRigBone& Bone : Hierarchy.Bones)
+	const TArray<FRigBone>& Bones = Hierarchy.GetBones();
+	for (const FRigBone& Bone : Bones )
 	{
 		Names.Add(Bone.Name.ToString());
 	}
@@ -291,6 +300,8 @@ const TArray<TSharedPtr<FString>>& UControlRigGraph::GetBoneNameList() const
 
 void UControlRigGraph::HandleModelModified(const UControlRigModel* InModel, EControlRigModelNotifType InType, const void* InPayload)
 {
+	DECLARE_SCOPE_HIERARCHICAL_COUNTER_FUNC()
+
 	if (bSuspendModelNotifications)
 	{
 		return;
@@ -390,8 +401,14 @@ void UControlRigGraph::HandleModelModified(const UControlRigModel* InModel, ECon
 
 					if (UControlRigGraphNode* RigNode = Cast<UControlRigGraphNode>(EdNode))
 					{
+						int32 PreviousParameterType = RigNode->ParameterType;
 						RigNode->ParameterType = (int32)Node->ParameterType;
 						RigNode->SetColorFromModel(Node->Color);
+
+						if (Node->IsParameter() && PreviousParameterType != RigNode->ParameterType)
+						{
+							RigNode->ReconstructNode();
+						}
 					}
 
 					if (UEdGraphNode_Comment* CommentNode = Cast<UEdGraphNode_Comment>(EdNode))
@@ -535,6 +552,8 @@ void UControlRigGraph::HandleModelModified(const UControlRigModel* InModel, ECon
 
 UEdGraphNode* UControlRigGraph::FindNodeFromPropertyName(const FName& InPropertyName)
 {
+	DECLARE_SCOPE_HIERARCHICAL_COUNTER_FUNC()
+
 	for (UEdGraphNode* EdNode : Nodes)
 	{
 		if (UControlRigGraphNode* RigNode = Cast<UControlRigGraphNode>(EdNode))
@@ -555,6 +574,30 @@ UEdGraphNode* UControlRigGraph::FindNodeFromPropertyName(const FName& InProperty
 	return nullptr;
 }
 
+void UControlRigGraph::CacheCurveNameList(const FRigCurveContainer& Container)
+{
+	DECLARE_SCOPE_HIERARCHICAL_COUNTER_FUNC()
+
+	TArray<FString> Names;
+	const TArray<FRigCurve>& Curves = Container.GetCurves();
+	for (const FRigCurve& Curve : Curves)
+	{
+		Names.Add(Curve.Name.ToString());
+	}
+	Names.Sort();
+
+	CurveNameList.Reset();
+	CurveNameList.Add(MakeShared<FString>(FName(NAME_None).ToString()));
+	for (const FString& Name : Names)
+	{
+		CurveNameList.Add(MakeShared<FString>(Name));
+	}
+}
+
+const TArray<TSharedPtr<FString>>& UControlRigGraph::GetCurveNameList() const
+{
+	return CurveNameList;
+}
 #endif
 
 #undef LOCTEXT_NAMESPACE
diff --git a/Engine/Plugins/Experimental/ControlRig/Source/ControlRigDeveloper/Private/Graph/ControlRigGraphNode.cpp b/Engine/Plugins/Experimental/ControlRig/Source/ControlRigDeveloper/Private/Graph/ControlRigGraphNode.cpp
--- a/Engine/Plugins/Experimental/ControlRig/Source/ControlRigDeveloper/Private/Graph/ControlRigGraphNode.cpp
+++ b/Engine/Plugins/Experimental/ControlRig/Source/ControlRigDeveloper/Private/Graph/ControlRigGraphNode.cpp
@@ -74,6 +74,8 @@ FText UControlRigGraphNode::GetNodeTitle(ENodeTitleType::Type TitleType) const
 
 void UControlRigGraphNode::ReconstructNode()
 {
+	DECLARE_SCOPE_HIERARCHICAL_COUNTER_FUNC()
+
 	UControlRigGraph* RigGraph = Cast<UControlRigGraph>(GetGraph());
 	if (RigGraph)
 	{
@@ -110,6 +112,8 @@ void UControlRigGraphNode::ReconstructNode()
 
 void UControlRigGraphNode::CacheHierarchyRefConnectionsOnPostLoad()
 {
+	DECLARE_SCOPE_HIERARCHICAL_COUNTER_FUNC()
+
 	if (HierarchyRefOutputConnections.Num() > 0)
 	{
 		return;
@@ -147,6 +151,8 @@ void UControlRigGraphNode::CacheHierarchyRefConnectionsOnPostLoad()
 
 void UControlRigGraphNode::PrepareForCopying()
 {
+	DECLARE_SCOPE_HIERARCHICAL_COUNTER_FUNC()
+
 	// cache the data we need for paste to work
 	// we fill up struct for rig unit
 	UScriptStruct* ScriptStruct = GetUnitScriptStruct();
@@ -207,6 +213,7 @@ void UControlRigGraphNode::ReallocatePinsDuringReconstruction(const TArray<UEdGr
 
 void UControlRigGraphNode::RewireOldPinsToNewPins(TArray<UEdGraphPin*>& InOldPins, TArray<UEdGraphPin*>& InNewPins)
 {
+	DECLARE_SCOPE_HIERARCHICAL_COUNTER_FUNC()
 	// @TODO: we should account for redirectors, orphaning etc. here too!
 
 	for(UEdGraphPin* OldPin : InOldPins)
@@ -226,6 +233,8 @@ void UControlRigGraphNode::RewireOldPinsToNewPins(TArray<UEdGraphPin*>& InOldPin
 
 void UControlRigGraphNode::DestroyPinList(TArray<UEdGraphPin*>& InPins)
 {
+	DECLARE_SCOPE_HIERARCHICAL_COUNTER_FUNC()
+
 	UBlueprint* Blueprint = GetBlueprint();
 	bool bNotify = false;
 	if (Blueprint != nullptr)
@@ -244,6 +253,8 @@ void UControlRigGraphNode::DestroyPinList(TArray<UEdGraphPin*>& InPins)
 
 void UControlRigGraphNode::PostReconstructNode()
 {
+	DECLARE_SCOPE_HIERARCHICAL_COUNTER_FUNC()
+
 	for (UEdGraphPin* Pin : Pins)
 	{
 		SetupPinDefaultsFromCDO(Pin);
@@ -281,6 +292,8 @@ void UControlRigGraphNode::PostLoad()
 
 void UControlRigGraphNode::CreateVariablePins(bool bAlwaysCreatePins)
 {
+	DECLARE_SCOPE_HIERARCHICAL_COUNTER_FUNC()
+
 	CacheVariableInfo();
 	CreateExecutionPins(bAlwaysCreatePins);
 	CreateInputPins(bAlwaysCreatePins);
@@ -290,6 +303,8 @@ void UControlRigGraphNode::CreateVariablePins(bool bAlwaysCreatePins)
 
 void UControlRigGraphNode::HandleClearArray(FString InPropertyPath)
 {
+	DECLARE_SCOPE_HIERARCHICAL_COUNTER_FUNC()
+
 	UControlRigBlueprint* Blueprint = Cast<UControlRigBlueprint>(GetOuter()->GetOuter());
 	if (Blueprint)
 	{
@@ -304,6 +319,8 @@ void UControlRigGraphNode::HandleClearArray(FString InPropertyPath)
 
 void UControlRigGraphNode::HandleAddArrayElement(FString InPropertyPath)
 {
+	DECLARE_SCOPE_HIERARCHICAL_COUNTER_FUNC()
+
 	UControlRigBlueprint* Blueprint = Cast<UControlRigBlueprint>(GetOuter()->GetOuter());
 	if (Blueprint)
 	{
@@ -318,6 +335,8 @@ void UControlRigGraphNode::HandleAddArrayElement(FString InPropertyPath)
 
 void UControlRigGraphNode::HandleRemoveArrayElement(FString InPropertyPath)
 {
+	DECLARE_SCOPE_HIERARCHICAL_COUNTER_FUNC()
+
 	UControlRigBlueprint* Blueprint = Cast<UControlRigBlueprint>(GetOuter()->GetOuter());
 	if (Blueprint)
 	{
@@ -364,6 +383,8 @@ static bool IsStructReference(const TSharedPtr<FControlRigField>& InputInfo)
 
 void UControlRigGraphNode::CreateExecutionPins(bool bAlwaysCreatePins)
 {
+	DECLARE_SCOPE_HIERARCHICAL_COUNTER_FUNC()
+
 	const TArray<TSharedRef<FControlRigField>>& LocalExecutionInfos = GetExecutionVariableInfo();
 
 	for (const TSharedRef<FControlRigField>& ExecutionInfo : LocalExecutionInfos)
@@ -396,6 +417,7 @@ void UControlRigGraphNode::CreateInputPins_Recursive(const TSharedPtr<FControlRi
 			ChildInfo->InputPin->PinType.bIsReference = IsStructReference(ChildInfo);
 			ChildInfo->InputPin->ParentPin = InputInfo->InputPin;
 			ChildInfo->InputPin->DefaultValue = ChildInfo->GetPin()->DefaultValue;
+			ChildInfo->OutputPin = nullptr;
 			InputInfo->InputPin->SubPins.Add(ChildInfo->InputPin);
 		}
 	}
@@ -408,6 +430,8 @@ void UControlRigGraphNode::CreateInputPins_Recursive(const TSharedPtr<FControlRi
 
 void UControlRigGraphNode::CreateInputPins(bool bAlwaysCreatePins)
 {
+	DECLARE_SCOPE_HIERARCHICAL_COUNTER_FUNC()
+
 	const TArray<TSharedRef<FControlRigField>>& LocalInputInfos = GetInputVariableInfo();
 
 	for (const TSharedRef<FControlRigField>& InputInfo : LocalInputInfos)
@@ -418,6 +442,11 @@ void UControlRigGraphNode::CreateInputPins(bool bAlwaysCreatePins)
 			InputInfo->InputPin->PinFriendlyName = InputInfo->GetDisplayNameText();
 			InputInfo->InputPin->PinType.bIsReference = IsStructReference(InputInfo);
 			InputInfo->InputPin->DefaultValue = InputInfo->GetPin()->DefaultValue;
+			InputInfo->OutputPin = nullptr;
+		}
+		else
+		{
+			InputInfo->InputPin->DefaultValue = InputInfo->GetPin()->DefaultValue;
 		}
 
 		CreateInputPins_Recursive(InputInfo, bAlwaysCreatePins);
@@ -437,6 +466,10 @@ void UControlRigGraphNode::CreateInputOutputPins_Recursive(const TSharedPtr<FCon
 			ChildInfo->InputPin->ParentPin = InputOutputInfo->InputPin;
 			InputOutputInfo->InputPin->SubPins.Add(ChildInfo->InputPin);
 		}
+		else
+		{
+			ChildInfo->InputPin->DefaultValue = ChildInfo->GetPin()->DefaultValue;
+		}
 
 		if (bAlwaysCreatePins || ChildInfo->OutputPin == nullptr)
 		{
@@ -456,6 +489,8 @@ void UControlRigGraphNode::CreateInputOutputPins_Recursive(const TSharedPtr<FCon
 
 void UControlRigGraphNode::CreateInputOutputPins(bool bAlwaysCreatePins)
 {
+	DECLARE_SCOPE_HIERARCHICAL_COUNTER_FUNC()
+
 	const TArray<TSharedRef<FControlRigField>>& LocalInputOutputInfos = GetInputOutputVariableInfo();
 
 	for (const TSharedRef<FControlRigField>& InputOutputInfo : LocalInputOutputInfos)
@@ -487,6 +522,7 @@ void UControlRigGraphNode::CreateOutputPins_Recursive(const TSharedPtr<FControlR
 			ChildInfo->OutputPin->PinFriendlyName = ChildInfo->GetDisplayNameText();
 			ChildInfo->OutputPin->PinType.bIsReference = IsStructReference(ChildInfo);
 			ChildInfo->OutputPin->ParentPin = OutputInfo->OutputPin;
+			ChildInfo->InputPin = nullptr;
 			OutputInfo->OutputPin->SubPins.Add(ChildInfo->OutputPin);
 		}
 	}
@@ -499,6 +535,8 @@ void UControlRigGraphNode::CreateOutputPins_Recursive(const TSharedPtr<FControlR
 
 void UControlRigGraphNode::CreateOutputPins(bool bAlwaysCreatePins)
 {
+	DECLARE_SCOPE_HIERARCHICAL_COUNTER_FUNC()
+
 	const TArray<TSharedRef<FControlRigField>>& LocalOutputInfos = GetOutputVariableInfo();
 
 	for (const TSharedRef<FControlRigField>& OutputInfo : LocalOutputInfos)
@@ -508,6 +546,7 @@ void UControlRigGraphNode::CreateOutputPins(bool bAlwaysCreatePins)
 			OutputInfo->OutputPin = CreatePin(EGPD_Output, OutputInfo->GetPinType(), FName(*OutputInfo->GetPinPath()));
 			OutputInfo->OutputPin->PinFriendlyName = OutputInfo->GetDisplayNameText();
 			OutputInfo->OutputPin->PinType.bIsReference = IsStructReference(OutputInfo);
+			OutputInfo->InputPin = nullptr;
 		}
 
 		CreateOutputPins_Recursive(OutputInfo, bAlwaysCreatePins);
@@ -516,6 +555,8 @@ void UControlRigGraphNode::CreateOutputPins(bool bAlwaysCreatePins)
 
 void UControlRigGraphNode::CacheVariableInfo()
 {
+	DECLARE_SCOPE_HIERARCHICAL_COUNTER_FUNC()
+
 	ExecutionInfos.Reset();
 	GetExecutionFields(ExecutionInfos);
 
@@ -593,6 +634,8 @@ TSharedPtr<FControlRigField> UControlRigGraphNode::CreateControlRigField(const F
 
 void UControlRigGraphNode::GetExecutionFields(TArray<TSharedRef<FControlRigField>>& OutFields) const
 {
+	DECLARE_SCOPE_HIERARCHICAL_COUNTER_FUNC()
+
 	GetFields([](const FControlRigModelPin* InPin, const FControlRigModelNode* InNode)
 	{
 		FString PinPath = InNode->GetPinPath(InPin->Index, false);
@@ -604,26 +647,56 @@ void UControlRigGraphNode::GetExecutionFields(TArray<TSharedRef<FControlRigField
 
 void UControlRigGraphNode::GetInputFields(TArray<TSharedRef<FControlRigField>>& OutFields) const
 {
+	DECLARE_SCOPE_HIERARCHICAL_COUNTER_FUNC()
+
 	GetFields([](const FControlRigModelPin* InPin, const FControlRigModelNode* InNode)
 	{
 		FString PinPath = InNode->GetPinPath(InPin->Index, false);
-		return InPin->Direction == EGPD_Input && InNode->FindPin(*PinPath, false) == nullptr;
+		if (InPin->Direction != EGPD_Input)
+		{
+			return false;
+		}
+
+		if (InNode->IsParameter() && InNode->ParameterType == EControlRigModelParameterType::Output)
+		{
+			return true;
+		}
+
+		return InNode->FindPin(*PinPath, false) == nullptr;
 	}, OutFields);
 }
 
 void UControlRigGraphNode::GetOutputFields(TArray<TSharedRef<FControlRigField>>& OutFields) const
 {
+	DECLARE_SCOPE_HIERARCHICAL_COUNTER_FUNC()
+
 	GetFields([](const FControlRigModelPin* InPin, const FControlRigModelNode* InNode)
 	{
 		FString PinPath = InNode->GetPinPath(InPin->Index, false);
-		return InPin->Direction == EGPD_Output && InNode->FindPin(*PinPath, true) == nullptr;
+		if (InPin->Direction != EGPD_Output)
+		{
+			return false;
+		}
+		
+		if (InNode->IsParameter() && InNode->ParameterType == EControlRigModelParameterType::Input)
+		{
+			return true;
+		}
+		
+		return InNode->FindPin(*PinPath, true) == nullptr;
 	}, OutFields);
 }
 
 void UControlRigGraphNode::GetInputOutputFields(TArray<TSharedRef<FControlRigField>>& OutFields) const
 {
+	DECLARE_SCOPE_HIERARCHICAL_COUNTER_FUNC()
+
 	GetFields([](const FControlRigModelPin* InPin, const FControlRigModelNode* InNode)
 	{
+		if (InNode->IsParameter())
+		{
+			return false;
+		}
 		FString PinPath = InNode->GetPinPath(InPin->Index, false);
 		return InPin->Direction == EGPD_Input && 
 			InNode->FindPin(*PinPath, false) != nullptr &&
@@ -794,6 +867,8 @@ bool UControlRigGraphNode::IsPinExpanded(const FString& InPinPropertyPath) const
 
 void UControlRigGraphNode::DestroyNode()
 {
+	DECLARE_SCOPE_HIERARCHICAL_COUNTER_FUNC()
+
 	if(UControlRigGraph* Graph = Cast<UControlRigGraph>(GetOuter()))
 	{
 		UControlRigBlueprint* ControlRigBlueprint = Cast<UControlRigBlueprint>(Graph->GetOuter());
@@ -819,6 +894,8 @@ TSharedPtr<INameValidatorInterface> UControlRigGraphNode::MakeNameValidator() co
 
 void UControlRigGraphNode::CopyPinDefaultsToModel(UEdGraphPin* Pin, bool bUndo)
 {
+	DECLARE_SCOPE_HIERARCHICAL_COUNTER_FUNC()
+
 	const UEdGraphSchema_K2* K2Schema = GetDefault<UEdGraphSchema_K2>();
 
 	if(UControlRigGraph* Graph = Cast<UControlRigGraph>(GetOuter()))
@@ -855,6 +932,8 @@ UControlRigBlueprint* UControlRigGraphNode::GetBlueprint() const
 
 void UControlRigGraphNode::SetupPinDefaultsFromCDO(UEdGraphPin* Pin)
 {
+	DECLARE_SCOPE_HIERARCHICAL_COUNTER_FUNC()
+
 	const UEdGraphSchema_K2* K2Schema = GetDefault<UEdGraphSchema_K2>();
 
 	if(UControlRigGraph* Graph = Cast<UControlRigGraph>(GetOuter()))
@@ -906,6 +985,8 @@ bool UControlRigGraphNode::CanCreateUnderSpecifiedSchema(const UEdGraphSchema* I
 
 void UControlRigGraphNode::AutowireNewNode(UEdGraphPin* FromPin)
 {
+	DECLARE_SCOPE_HIERARCHICAL_COUNTER_FUNC()
+
 	Super::AutowireNewNode(FromPin);
 
 	const UControlRigGraphSchema* Schema = GetDefault<UControlRigGraphSchema>();
@@ -949,6 +1030,8 @@ void ReplacePropertyName(TArray<TSharedRef<FControlRigField>>& InArray, const FS
 
 void UControlRigGraphNode::SetPropertyName(const FName& InPropertyName, bool bReplaceInnerProperties/*=false*/)
 { 
+	DECLARE_SCOPE_HIERARCHICAL_COUNTER_FUNC()
+
 	const FString OldPropertyName = PropertyName.ToString();
 	const FString NewPropertyName = InPropertyName.ToString();
 	PropertyName = InPropertyName;
diff --git a/Engine/Plugins/Experimental/ControlRig/Source/ControlRigDeveloper/Private/Graph/ControlRigGraphSchema.cpp b/Engine/Plugins/Experimental/ControlRig/Source/ControlRigDeveloper/Private/Graph/ControlRigGraphSchema.cpp
--- a/Engine/Plugins/Experimental/ControlRig/Source/ControlRigDeveloper/Private/Graph/ControlRigGraphSchema.cpp
+++ b/Engine/Plugins/Experimental/ControlRig/Source/ControlRigDeveloper/Private/Graph/ControlRigGraphSchema.cpp
@@ -153,10 +153,7 @@ void UControlRigGraphSchema::BreakPinLinks(UEdGraphPin& TargetPin, bool bSendsNo
 	{
 		FString Left, Right;
 		RigBlueprint->Model->SplitPinPath(TargetPin.GetName(), Left, Right);
-		if (RigBlueprint->ModelController->BreakLinks(*Left, *Right, TargetPin.Direction == EGPD_Input))
-		{
-			FBlueprintEditorUtils::MarkBlueprintAsModified(Blueprint);
-		}
+		RigBlueprint->ModelController->BreakLinks(*Left, *Right, TargetPin.Direction == EGPD_Input);
 	}
 }
 
@@ -178,10 +175,7 @@ void UControlRigGraphSchema::BreakSinglePinLink(UEdGraphPin* SourcePin, UEdGraph
 		FString SourceLeft, SourceRight, TargetLeft, TargetRight;
 		RigBlueprint->Model->SplitPinPath(SourcePin->GetName(), SourceLeft, SourceRight);
 		RigBlueprint->Model->SplitPinPath(TargetPin->GetName(), TargetLeft, TargetLeft);
-		if (RigBlueprint->ModelController->BreakLink(*SourceLeft, *SourceRight, *TargetLeft, *TargetRight))
-		{
-			FBlueprintEditorUtils::MarkBlueprintAsModified(Blueprint);
-		}
+		RigBlueprint->ModelController->BreakLink(*SourceLeft, *SourceRight, *TargetLeft, *TargetRight);
 	}
 }
 
@@ -212,7 +206,7 @@ UControlRigGraphNode* UControlRigGraphSchema::CreateGraphNode(UControlRigGraph*
 	return ControlRigGraphNode;
 }
 
-void UControlRigGraphSchema::TrySetDefaultValue(UEdGraphPin& InPin, const FString& InNewDefaultValue) const
+void UControlRigGraphSchema::TrySetDefaultValue(UEdGraphPin& InPin, const FString& InNewDefaultValue, bool bMarkAsModified) const
 {
 #if WITH_EDITOR
 	if (GEditor)
@@ -220,10 +214,10 @@ void UControlRigGraphSchema::TrySetDefaultValue(UEdGraphPin& InPin, const FStrin
 		GEditor->CancelTransaction(0);
 	}
 #endif
-	GetDefault<UEdGraphSchema_K2>()->TrySetDefaultValue(InPin, InNewDefaultValue);
+	GetDefault<UEdGraphSchema_K2>()->TrySetDefaultValue(InPin, InNewDefaultValue, false);
 }
 
-void UControlRigGraphSchema::TrySetDefaultObject(UEdGraphPin& InPin, UObject* InNewDefaultObject) const
+void UControlRigGraphSchema::TrySetDefaultObject(UEdGraphPin& InPin, UObject* InNewDefaultObject, bool bMarkAsModified) const
 {
 #if WITH_EDITOR
 	if (GEditor)
@@ -231,10 +225,10 @@ void UControlRigGraphSchema::TrySetDefaultObject(UEdGraphPin& InPin, UObject* In
 		GEditor->CancelTransaction(0);
 	}
 #endif
-	GetDefault<UEdGraphSchema_K2>()->TrySetDefaultObject(InPin, InNewDefaultObject);
+	GetDefault<UEdGraphSchema_K2>()->TrySetDefaultObject(InPin, InNewDefaultObject, false);
 }
 
-void UControlRigGraphSchema::TrySetDefaultText(UEdGraphPin& InPin, const FText& InNewDefaultText) const
+void UControlRigGraphSchema::TrySetDefaultText(UEdGraphPin& InPin, const FText& InNewDefaultText, bool bMarkAsModified) const
 {
 #if WITH_EDITOR
 	if (GEditor)
@@ -242,7 +236,7 @@ void UControlRigGraphSchema::TrySetDefaultText(UEdGraphPin& InPin, const FText&
 		GEditor->CancelTransaction(0);
 	}
 #endif
-	GetDefault<UEdGraphSchema_K2>()->TrySetDefaultText(InPin, InNewDefaultText);
+	GetDefault<UEdGraphSchema_K2>()->TrySetDefaultText(InPin, InNewDefaultText, false);
 }
 
 bool UControlRigGraphSchema::ArePinsCompatible(const UEdGraphPin* PinA, const UEdGraphPin* PinB, const UClass* CallingContext, bool bIgnoreArray /*= false*/) const
diff --git a/Engine/Plugins/Experimental/ControlRig/Source/ControlRigDeveloper/Public/ControlRigBlueprint.h b/Engine/Plugins/Experimental/ControlRig/Source/ControlRigDeveloper/Public/ControlRigBlueprint.h
--- a/Engine/Plugins/Experimental/ControlRig/Source/ControlRigDeveloper/Public/ControlRigBlueprint.h
+++ b/Engine/Plugins/Experimental/ControlRig/Source/ControlRigDeveloper/Public/ControlRigBlueprint.h
@@ -9,6 +9,7 @@
 #include "Misc/Crc.h"
 #include "ControlRigDefines.h"
 #include "Hierarchy.h"
+#include "CurveContainer.h"
 #include "Interfaces/Interface_PreviewMeshProvider.h"
 #include "ControlRigController.h"
 #include "ControlRigBlueprint.generated.h"
@@ -164,10 +165,20 @@ private:
 	UPROPERTY()
 	FRigHierarchy Hierarchy;
 
+	UPROPERTY()
+	FRigCurveContainer CurveContainer;
+
 	/** The default skeletal mesh to use when previewing this asset */
 	UPROPERTY(DuplicateTransient, AssetRegistrySearchable)
 	TSoftObjectPtr<USkeletalMesh> PreviewSkeletalMesh;
 
+	/** The default skeletal mesh to use when previewing this asset */
+	UPROPERTY(DuplicateTransient, AssetRegistrySearchable)
+	TSoftObjectPtr<UObject> SourceHierarchyImport;
+
+	UPROPERTY(DuplicateTransient, AssetRegistrySearchable)
+	TSoftObjectPtr<UObject> SourceCurveImport;
+
 	UControlRigModel::FModifiedEvent _ModifiedEvent;
 	void HandleModelModified(const UControlRigModel* InModel, EControlRigModelNotifType InType, const void* InPayload);
 	bool UpdateParametersOnControlRig(UControlRig* InRig = nullptr);
@@ -175,6 +186,7 @@ private:
 
 	friend class FControlRigBlueprintCompilerContext;
 	friend class SRigHierarchy;
+	friend class SRigCurveContainer;
 	friend class FControlRigEditor;
 	friend class UEngineTestControlRig;
 };
diff --git a/Engine/Plugins/Experimental/ControlRig/Source/ControlRigDeveloper/Public/ControlRigDAG.h b/Engine/Plugins/Experimental/ControlRig/Source/ControlRigDeveloper/Public/ControlRigDAG.h
--- a/Engine/Plugins/Experimental/ControlRig/Source/ControlRigDeveloper/Public/ControlRigDAG.h
+++ b/Engine/Plugins/Experimental/ControlRig/Source/ControlRigDeveloper/Public/ControlRigDAG.h
@@ -48,15 +48,21 @@ public:
 		FName Name;
 		int32 Index;
 		bool IsMutable;
+		bool IsOutputParameter;
 		FPinArray Inputs;
 		FPinArray Outputs;
 
-		FNode(const FName& InName, const int32 InIndex, const bool InIsMutable)
+		FNode(const FName& InName, const int32 InIndex, const bool InIsMutable, const bool InIsOutputParameter)
 			: Name(InName)
 			, Index(InIndex)
 			, IsMutable(InIsMutable)
+			, IsOutputParameter(InIsOutputParameter)
 		{
 		}
+
+		bool operator ==(const FNode& Other) const {
+			return Name == Other.Name;
+		}
 	};
 
 	TArray<FNode> Nodes;
@@ -65,7 +71,7 @@ public:
 	FControlRigDAG();
 
 	// add a node to this traverser
-	void AddNode(bool InIsMutable = false, const FName& InName = NAME_None);
+	void AddNode(bool InIsMutable = false, bool InIsOutputParameter = false, const FName& InName = NAME_None);
 
 	// add a link between two nodes given the node indices and the pin orders
 	void AddLink(const int32 FromNode, const int32 ToNode, const int32 FromOrder, const int32 ToOrder);
diff --git a/Engine/Plugins/Experimental/ControlRig/Source/ControlRigDeveloper/Public/Graph/ControlRigGraph.h b/Engine/Plugins/Experimental/ControlRig/Source/ControlRigDeveloper/Public/Graph/ControlRigGraph.h
--- a/Engine/Plugins/Experimental/ControlRig/Source/ControlRigDeveloper/Public/Graph/ControlRigGraph.h
+++ b/Engine/Plugins/Experimental/ControlRig/Source/ControlRigDeveloper/Public/Graph/ControlRigGraph.h
@@ -11,6 +11,7 @@
 class UControlRigBlueprint;
 class UControlRigGraphSchema;
 class UControlRig;
+struct FRigCurveContainer;
 
 UCLASS()
 class CONTROLRIGDEVELOPER_API UControlRigGraph : public UEdGraph
@@ -37,6 +38,8 @@ public:
 	FDelegateHandle BlueprintOnCompiledHandle;
 	void CacheBoneNameList(const FRigHierarchy& Hierarchy);
 	const TArray<TSharedPtr<FString>>& GetBoneNameList() const;
+	void CacheCurveNameList(const FRigCurveContainer& Container);
+	const TArray<TSharedPtr<FString>>& GetCurveNameList() const;
 
 	bool bSuspendModelNotifications;
 	bool bIsTemporaryGraphForCopyPaste;
@@ -51,6 +54,7 @@ private:
 	TArray<UControlRigGraphNode*> FoundHierarchyRefMutableNodes;
 	TMap<UControlRigGraphNode*, TArray<UControlRigGraphNode*>> FoundHierarchyRefConnections;
 	TArray<TSharedPtr<FString>> BoneNameList;
+	TArray<TSharedPtr<FString>> CurveNameList;
 #endif
 };
 
diff --git a/Engine/Plugins/Experimental/ControlRig/Source/ControlRigDeveloper/Public/Graph/ControlRigGraphSchema.h b/Engine/Plugins/Experimental/ControlRig/Source/ControlRigDeveloper/Public/Graph/ControlRigGraphSchema.h
--- a/Engine/Plugins/Experimental/ControlRig/Source/ControlRigDeveloper/Public/Graph/ControlRigGraphSchema.h
+++ b/Engine/Plugins/Experimental/ControlRig/Source/ControlRigDeveloper/Public/Graph/ControlRigGraphSchema.h
@@ -61,9 +61,9 @@ public:
 	virtual void BreakSinglePinLink(UEdGraphPin* SourcePin, UEdGraphPin* TargetPin) const override;
 	virtual class FConnectionDrawingPolicy* CreateConnectionDrawingPolicy(int32 InBackLayerID, int32 InFrontLayerID, float InZoomFactor, const FSlateRect& InClippingRect, class FSlateWindowElementList& InDrawElements, class UEdGraph* InGraphObj) const override;
 	virtual bool ShouldHidePinDefaultValue(UEdGraphPin* Pin) const override;
-	virtual void TrySetDefaultValue(UEdGraphPin& InPin, const FString& InNewDefaultValue) const override;
-	virtual void TrySetDefaultObject(UEdGraphPin& InPin, UObject* InNewDefaultObject) const override;
-	virtual void TrySetDefaultText(UEdGraphPin& InPin, const FText& InNewDefaultText) const override;
+	virtual void TrySetDefaultValue(UEdGraphPin& InPin, const FString& InNewDefaultValue, bool bMarkAsModified = true) const override;
+	virtual void TrySetDefaultObject(UEdGraphPin& InPin, UObject* InNewDefaultObject, bool bMarkAsModified) const override;
+	virtual void TrySetDefaultText(UEdGraphPin& InPin, const FText& InNewDefaultText, bool bMarkAsModified) const override;
 	virtual bool ShouldAlwaysPurgeOnModification() const override { return false; }
 	virtual bool ArePinsCompatible(const UEdGraphPin* PinA, const UEdGraphPin* PinB, const UClass* CallingContext, bool bIgnoreArray /*= false*/) const override;
 	virtual bool DoesSupportPinWatching() const	override { return true; }
diff --git a/Engine/Plugins/Experimental/ControlRig/Source/ControlRigEditor/ControlRigEditor.Build.cs b/Engine/Plugins/Experimental/ControlRig/Source/ControlRigEditor/ControlRigEditor.Build.cs
--- a/Engine/Plugins/Experimental/ControlRig/Source/ControlRigEditor/ControlRigEditor.Build.cs
+++ b/Engine/Plugins/Experimental/ControlRig/Source/ControlRigEditor/ControlRigEditor.Build.cs
@@ -49,6 +49,9 @@ namespace UnrealBuildTool.Rules
                     "UMG",
 					"TimeManagement",
                     "PropertyPath",
+					"WorkspaceMenuStructure",
+					"Json",
+					"DesktopPlatform",
                 }
             );
         }
diff --git a/Engine/Plugins/Experimental/ControlRig/Source/ControlRigEditor/Private/ControlRigEditorModule.cpp b/Engine/Plugins/Experimental/ControlRig/Source/ControlRigEditor/Private/ControlRigEditorModule.cpp
--- a/Engine/Plugins/Experimental/ControlRig/Source/ControlRigEditor/Private/ControlRigEditorModule.cpp
+++ b/Engine/Plugins/Experimental/ControlRig/Source/ControlRigEditor/Private/ControlRigEditorModule.cpp
@@ -17,6 +17,7 @@
 #include "ControlRigSequenceActions.h"
 #include "ControlRigEditorStyle.h"
 #include "Framework/Docking/LayoutExtender.h"
+#include "Framework/Application/SlateApplication.h"
 #include "LevelEditor.h"
 #include "MovieSceneToolsProjectSettings.h"
 #include "PropertyEditorModule.h"
@@ -69,6 +70,9 @@
 #include "ControlRigDetails.h"
 #include "Units/Deprecated/RigUnitEditor_TwoBoneIKFK.h"
 #include "Animation/AnimSequence.h"
+#include "Editor/SControlRigProfilingView.h"
+#include "WorkspaceMenuStructure.h"
+#include "WorkspaceMenuStructureModule.h"
 
 #define LOCTEXT_NAMESPACE "ControlRigEditorModule"
 
@@ -76,6 +80,15 @@ DEFINE_LOG_CATEGORY(LogControlRigEditor);
 
 TMap<FName, TSubclassOf<URigUnitEditor_Base>> FControlRigEditorModule::RigUnitEditorClasses;
 
+TSharedRef<SDockTab> SpawnRigProfiler( const FSpawnTabArgs& Args )
+{
+	return SNew(SDockTab)
+		.TabRole(ETabRole::NomadTab)
+		[
+			SNew(SControlRigProfilingView)
+		];
+}
+
 void FControlRigEditorModule::StartupModule()
 {
 	FControlRigEditModeCommands::Register();
@@ -302,10 +315,28 @@ void FControlRigEditorModule::StartupModule()
 
 	// register rig unit base editor class
 	RegisterRigUnitEditorClass("RigUnit_TwoBoneIKFK", URigUnitEditor_TwoBoneIKFK::StaticClass());
+
+#if WITH_EDITOR
+	if (FSlateApplication::IsInitialized())
+	{
+		FGlobalTabmanager::Get()->RegisterNomadTabSpawner("HierarchicalProfiler", FOnSpawnTab::CreateStatic(&SpawnRigProfiler))
+			.SetDisplayName(NSLOCTEXT("UnrealEditor", "HierarchicalProfilerTab", "Hierarchical Profiler"))
+			.SetTooltipText(NSLOCTEXT("UnrealEditor", "HierarchicalProfilerTooltip", "Open the Hierarchical Profiler tab."))
+			.SetGroup(WorkspaceMenu::GetMenuStructure().GetDeveloperToolsMiscCategory())
+			.SetIcon(FSlateIcon(TEXT("ControlRigEditorStyle"), TEXT("ControlRig.RigUnit")));
+	};
+#endif
 }
 
 void FControlRigEditorModule::ShutdownModule()
 {
+#if WITH_EDITOR
+	if (FSlateApplication::IsInitialized())
+	{
+		FGlobalTabmanager::Get()->UnregisterNomadTabSpawner("ControlRigProfiler");
+	}
+#endif
+
 	FBlueprintEditorUtils::OnRefreshAllNodesEvent.Remove(RefreshAllNodesDelegateHandle);
 	FBlueprintEditorUtils::OnReconstructAllNodesEvent.Remove(ReconstructAllNodesDelegateHandle);
 	FBlueprintEditorUtils::OnRenameVariableReferencesEvent.Remove(RenameVariableReferencesDelegateHandle);
diff --git a/Engine/Plugins/Experimental/ControlRig/Source/ControlRigEditor/Private/EditMode/ControlRigEditMode.cpp b/Engine/Plugins/Experimental/ControlRig/Source/ControlRigEditor/Private/EditMode/ControlRigEditMode.cpp
--- a/Engine/Plugins/Experimental/ControlRig/Source/ControlRigEditor/Private/EditMode/ControlRigEditMode.cpp
+++ b/Engine/Plugins/Experimental/ControlRig/Source/ControlRigEditor/Private/EditMode/ControlRigEditMode.cpp
@@ -381,9 +381,10 @@ void FControlRigEditMode::Render(const FSceneView* View, FViewport* Viewport, FP
 
 				// each base hierarchy Bone
 				const FRigHierarchy& BaseHierarchy = ControlRig->GetBaseHierarchy();
-				for (int32 BoneIndex = 0; BoneIndex < BaseHierarchy.Bones.Num(); ++BoneIndex)
+				const TArray<FRigBone>& Bones = BaseHierarchy.GetBones();
+				for (int32 BoneIndex = 0; BoneIndex < Bones.Num(); ++BoneIndex)
 				{
-					const FRigBone& CurrentBone = BaseHierarchy.Bones[BoneIndex];
+					const FRigBone& CurrentBone = Bones[BoneIndex];
 					const FTransform Transform = BaseHierarchy.GetGlobalTransform(BoneIndex);
 
 					if (CurrentBone.ParentIndex != INDEX_NONE)
@@ -844,7 +845,8 @@ bool FControlRigEditMode::InputDelta(FEditorViewportClient* InViewportClient, FV
 										UBlueprint* Blueprint = Cast<UBlueprint>(Class->ClassGeneratedBy);
 										if (Blueprint)
 										{
-											FBlueprintEditorUtils::MarkBlueprintAsModified(Blueprint);
+											Blueprint->MarkPackageDirty();
+											//FBlueprintEditorUtils::MarkBlueprintAsModified(Blueprint);
 										}
 									}
 								}
