Final copy-up for Control Rig Task Stream.

diff --git a/Engine/Plugins/Experimental/ControlRig/Source/ControlRigEditor/Private/Editor/ControlRigEditor.cpp b/Engine/Plugins/Experimental/ControlRig/Source/ControlRigEditor/Private/Editor/ControlRigEditor.cpp
--- a/Engine/Plugins/Experimental/ControlRig/Source/ControlRigEditor/Private/Editor/ControlRigEditor.cpp
+++ b/Engine/Plugins/Experimental/ControlRig/Source/ControlRigEditor/Private/Editor/ControlRigEditor.cpp
@@ -112,6 +112,8 @@ void FControlRigEditor::ExtendMenu()
 
 void FControlRigEditor::InitControlRigEditor(const EToolkitMode::Type Mode, const TSharedPtr< class IToolkitHost >& InitToolkitHost, UControlRigBlueprint* InControlRigBlueprint)
 {
+	DECLARE_SCOPE_HIERARCHICAL_COUNTER_FUNC()
+
 	FPersonaModule& PersonaModule = FModuleManager::GetModuleChecked<FPersonaModule>("Persona");
 
 	FPersonaToolkitArgs PersonaToolkitArgs;
@@ -121,6 +123,7 @@ void FControlRigEditor::InitControlRigEditor(const EToolkitMode::Type Mode, cons
 	// set delegate prior to setting mesh
 	// otherwise, you don't get delegate
 	PersonaToolkit->GetPreviewScene()->RegisterOnPreviewMeshChanged(FOnPreviewMeshChanged::CreateSP(this, &FControlRigEditor::HandlePreviewMeshChanged));
+
 	// Set a default preview mesh, if any
 	PersonaToolkit->SetPreviewMesh(InControlRigBlueprint->GetPreviewMesh(), false);
 
@@ -206,7 +209,7 @@ void FControlRigEditor::InitControlRigEditor(const EToolkitMode::Type Mode, cons
 		{
 			if (Graph->GetFName().IsEqual(UControlRigGraphSchema::GraphName_ControlRig))
 			{
-				OpenGraphAndBringToFront(Graph);
+				OpenGraphAndBringToFront(Graph, false);
 				break;
 			}
 		}
@@ -236,6 +239,8 @@ void FControlRigEditor::InitControlRigEditor(const EToolkitMode::Type Mode, cons
 
 void FControlRigEditor::BindCommands()
 {
+	DECLARE_SCOPE_HIERARCHICAL_COUNTER_FUNC()
+
 	GetToolkitCommands()->MapAction(
 		FControlRigBlueprintCommands::Get().ExecuteGraph,
 		FExecuteAction::CreateSP(this, &FControlRigEditor::ToggleExecuteGraph), 
@@ -258,6 +263,8 @@ bool FControlRigEditor::IsExecuteGraphOn() const
 
 void FControlRigEditor::ExtendToolbar()
 {
+	DECLARE_SCOPE_HIERARCHICAL_COUNTER_FUNC()
+
 	// If the ToolbarExtender is valid, remove it before rebuilding it
 	if(ToolbarExtender.IsValid())
 	{
@@ -539,52 +546,61 @@ void FControlRigEditor::OnCreateGraphEditorCommands(TSharedPtr<FUICommandList> G
 
 void FControlRigEditor::Compile()
 {
-	FString LastDebuggedObjectName = GetCustomDebugObjectLabel(GetBlueprintObj()->GetObjectBeingDebugged());
-
-	GetBlueprintObj()->SetObjectBeingDebugged(nullptr);
-	ClearDetailObject();
-
-	if (ControlRig)
 	{
-		ControlRig->OnInitialized().Clear();
-		ControlRig->OnExecuted().Clear();
-	}
+		DECLARE_SCOPE_HIERARCHICAL_COUNTER_FUNC()
 
-	FBlueprintEditor::Compile();
+		FString LastDebuggedObjectName = GetCustomDebugObjectLabel(GetBlueprintObj()->GetObjectBeingDebugged());
+		GetBlueprintObj()->SetObjectBeingDebugged(nullptr);
+		ClearDetailObject();
 
-	if (ControlRig)
-	{
-		ControlRig->ControlRigLog = &ControlRigLog;
-		ControlRig->DrawInterface = &DrawInterface;
+		if (ControlRig)
+		{
+			ControlRig->OnInitialized().Clear();
+			ControlRig->OnExecuted().Clear();
+		}
+
+		{
+			FBlueprintEditor::Compile();
+		}
 
-		UControlRigBlueprintGeneratedClass* GeneratedClass = Cast<UControlRigBlueprintGeneratedClass>(ControlRig->GetClass());
-		if (GeneratedClass)
+		if (ControlRig)
 		{
-			if (GeneratedClass->Operators.Num() == 1) // just the "done" operator
+			ControlRig->ControlRigLog = &ControlRigLog;
+			ControlRig->DrawInterface = &DrawInterface;
+
+			UControlRigBlueprintGeneratedClass* GeneratedClass = Cast<UControlRigBlueprintGeneratedClass>(ControlRig->GetClass());
+			if (GeneratedClass)
 			{
-				FNotificationInfo Info(LOCTEXT("ControlRigBlueprintCompilerEmptyRigMessage", "The Control Rig you compiled doesn't do anything. Did you forget to add a Begin_Execution node?"));
-				Info.bFireAndForget = true;
-				Info.FadeOutDuration = 10.0f;
-				Info.ExpireDuration = 0.0f;
-				TSharedPtr<SNotificationItem> NotificationPtr = FSlateNotificationManager::Get().AddNotification(Info);
-				NotificationPtr->SetCompletionState(SNotificationItem::CS_Success);
+				if (GeneratedClass->Operators.Num() == 1) // just the "done" operator
+				{
+					FNotificationInfo Info(LOCTEXT("ControlRigBlueprintCompilerEmptyRigMessage", "The Control Rig you compiled doesn't do anything. Did you forget to add a Begin_Execution node?"));
+					Info.bFireAndForget = true;
+					Info.FadeOutDuration = 10.0f;
+					Info.ExpireDuration = 0.0f;
+					TSharedPtr<SNotificationItem> NotificationPtr = FSlateNotificationManager::Get().AddNotification(Info);
+					NotificationPtr->SetCompletionState(SNotificationItem::CS_Success);
+				}
 			}
 		}
-	}
 
-	TArray<FCustomDebugObject> DebugList;
-	GetCustomDebugObjects(DebugList);
+		TArray<FCustomDebugObject> DebugList;
+		GetCustomDebugObjects(DebugList);
 
-	for (const FCustomDebugObject& DebugObject : DebugList)
-	{
-		if (DebugObject.NameOverride == LastDebuggedObjectName)
+		for (const FCustomDebugObject& DebugObject : DebugList)
 		{
-			GetBlueprintObj()->SetObjectBeingDebugged(DebugObject.Object);
+			if (DebugObject.NameOverride == LastDebuggedObjectName)
+			{
+				GetBlueprintObj()->SetObjectBeingDebugged(DebugObject.Object);
+			}
 		}
 	}
 
 	// enable this for creating a new unit test
 	// DumpUnitTestCode();
+
+	// FStatsHierarchical::EndMeasurements();
+	// FMessageLog LogForMeasurements("ControlRigLog");
+	// FStatsHierarchical::DumpMeasurements(LogForMeasurements);
 }
 
 FName FControlRigEditor::GetToolkitFName() const
@@ -614,6 +630,8 @@ FLinearColor FControlRigEditor::GetWorldCentricTabColorScale() const
 
 void FControlRigEditor::DeleteSelectedNodes()
 {
+	DECLARE_SCOPE_HIERARCHICAL_COUNTER_FUNC()
+
 	UControlRigBlueprint* RigBlueprint = Cast<UControlRigBlueprint>(GetBlueprintObj());
 	if (RigBlueprint == nullptr)
 	{
@@ -649,6 +667,8 @@ void FControlRigEditor::DeleteSelectedNodes()
 
 void FControlRigEditor::PasteNodesHere(class UEdGraph* DestinationGraph, const FVector2D& GraphLocation)
 {
+	DECLARE_SCOPE_HIERARCHICAL_COUNTER_FUNC()
+
 	UControlRigBlueprint* RigBlueprint = Cast<UControlRigBlueprint>(GetBlueprintObj());
 	if (RigBlueprint == nullptr)
 	{
@@ -897,6 +917,8 @@ void FControlRigEditor::NotifyPostChange(const FPropertyChangedEvent& PropertyCh
 
 void FControlRigEditor::HandleModelModified(const UControlRigModel* InModel, EControlRigModelNotifType InType, const void* InPayload)
 {
+	DECLARE_SCOPE_HIERARCHICAL_COUNTER_FUNC()
+
 	switch (InType)
 	{
 		case EControlRigModelNotifType::NodeSelected:
@@ -985,6 +1007,8 @@ TStatId FControlRigEditor::GetStatId() const
 
 void FControlRigEditor::OnSelectedNodesChangedImpl(const TSet<class UObject*>& NewSelection)
 {
+	DECLARE_SCOPE_HIERARCHICAL_COUNTER_FUNC()
+
 	if (bIsSelecting)
 	{
 		return;
@@ -1016,6 +1040,8 @@ void FControlRigEditor::OnSelectedNodesChangedImpl(const TSet<class UObject*>& N
 
 void FControlRigEditor::HandleHideItem()
 {
+	DECLARE_SCOPE_HIERARCHICAL_COUNTER_FUNC()
+
 	UControlRigBlueprint* ControlRigBlueprint = CastChecked<UControlRigBlueprint>(GetBlueprintObj());
 
 	TSet<UObject*> SelectedNodes = GetSelectedNodes();
@@ -1042,6 +1068,8 @@ bool FControlRigEditor::CanHideItem() const
 
 void FControlRigEditor::OnBlueprintChangedImpl(UBlueprint* InBlueprint, bool bIsJustBeingCompiled)
 {
+	DECLARE_SCOPE_HIERARCHICAL_COUNTER_FUNC()
+
 	if (!bControlRigEditorInitialized)
 	{
 		return;
@@ -1088,6 +1116,8 @@ void FControlRigEditor::OnBlueprintChangedImpl(UBlueprint* InBlueprint, bool bIs
 
 void FControlRigEditor::HandleViewportCreated(const TSharedRef<class IPersonaViewport>& InViewport)
 {
+	DECLARE_SCOPE_HIERARCHICAL_COUNTER_FUNC()
+
 	// TODO: this is duplicated code from FAnimBlueprintEditor, would be nice to consolidate. 
 	auto GetCompilationStateText = [this]()
 	{
@@ -1231,6 +1261,8 @@ void FControlRigEditor::HandleViewportCreated(const TSharedRef<class IPersonaVie
 
 void FControlRigEditor::HandlePreviewSceneCreated(const TSharedRef<IPersonaPreviewScene>& InPersonaPreviewScene)
 {
+	DECLARE_SCOPE_HIERARCHICAL_COUNTER_FUNC()
+
 	AAnimationEditorPreviewActor* Actor = InPersonaPreviewScene->GetWorld()->SpawnActor<AAnimationEditorPreviewActor>(AAnimationEditorPreviewActor::StaticClass(), FTransform::Identity);
 	InPersonaPreviewScene->SetActor(Actor);
 
@@ -1251,6 +1283,8 @@ void FControlRigEditor::HandlePreviewSceneCreated(const TSharedRef<IPersonaPrevi
 
 void FControlRigEditor::UpdateControlRig()
 {
+	DECLARE_SCOPE_HIERARCHICAL_COUNTER_FUNC()
+
 	if(UClass* Class = GetBlueprintObj()->GeneratedClass)
 	{
 		UControlRigSkeletalMeshComponent* EditorSkelComp = Cast<UControlRigSkeletalMeshComponent>(GetPersonaToolkit()->GetPreviewScene()->GetPreviewMeshComponent());
@@ -1269,6 +1303,7 @@ void FControlRigEditor::UpdateControlRig()
  			}
 
 			CacheBoneNameList();
+			CacheCurveNameList();
 
 			// When the control rig is re-instanced on compile, it loses its binding, so we refresh it here if needed
 			if (!ControlRig->GetObjectBinding().IsValid())
@@ -1296,6 +1331,8 @@ void FControlRigEditor::UpdateControlRig()
 
 void FControlRigEditor::CacheBoneNameList()
 {
+	DECLARE_SCOPE_HIERARCHICAL_COUNTER_FUNC()
+
 	if (ControlRig)
 	{
 		// make sure the bone name list is up 2 date for the editor graph
@@ -1312,6 +1349,26 @@ void FControlRigEditor::CacheBoneNameList()
 	}
 }
 
+void FControlRigEditor::CacheCurveNameList()
+{
+	DECLARE_SCOPE_HIERARCHICAL_COUNTER_FUNC()
+
+	if (ControlRig)
+	{
+		// make sure the Curve name list is up 2 date for the editor graph
+		for (UEdGraph* Graph : GetBlueprintObj()->UbergraphPages)
+		{
+			UControlRigGraph* RigGraph = Cast<UControlRigGraph>(Graph);
+			if (RigGraph == nullptr)
+			{
+				continue;
+			}
+
+			RigGraph->CacheCurveNameList(ControlRig->GetCurveContainer());
+		}
+	}
+}
+
 void FControlRigEditor::AddReferencedObjects( FReferenceCollector& Collector )
 {
 	FBlueprintEditor::AddReferencedObjects(Collector);
@@ -1326,6 +1383,8 @@ void FControlRigEditor::HandlePreviewMeshChanged(USkeletalMesh* InOldSkeletalMes
 
 void FControlRigEditor::RebindToSkeletalMeshComponent()
 {
+	DECLARE_SCOPE_HIERARCHICAL_COUNTER_FUNC()
+
 	UDebugSkelMeshComponent* MeshComponent = GetPersonaToolkit()->GetPreviewScene()->GetPreviewMeshComponent();
 	if (MeshComponent)
 	{
@@ -1336,6 +1395,8 @@ void FControlRigEditor::RebindToSkeletalMeshComponent()
 
 void FControlRigEditor::SetupGraphEditorEvents(UEdGraph* InGraph, SGraphEditor::FGraphEditorEvents& InEvents)
 {
+	DECLARE_SCOPE_HIERARCHICAL_COUNTER_FUNC()
+
 	FBlueprintEditor::SetupGraphEditorEvents(InGraph, InEvents);
 
 	InEvents.OnCreateActionMenu = SGraphEditor::FOnCreateActionMenu::CreateSP(this, &FControlRigEditor::HandleCreateGraphActionMenu);
@@ -1349,6 +1410,8 @@ FActionMenuContent FControlRigEditor::HandleCreateGraphActionMenu(UEdGraph* InGr
 
 void FControlRigEditor::OnNodeTitleCommitted(const FText& NewText, ETextCommit::Type CommitInfo, UEdGraphNode* NodeBeingChanged)
 {
+	DECLARE_SCOPE_HIERARCHICAL_COUNTER_FUNC()
+
 	if (UEdGraphNode_Comment* CommentBeingChanged = Cast<UEdGraphNode_Comment>(NodeBeingChanged))
 	{
 		if (UControlRigBlueprint* ControlRigBP = GetControlRigBlueprint())
@@ -1364,6 +1427,8 @@ void FControlRigEditor::OnNodeTitleCommitted(const FText& NewText, ETextCommit::
 
 void FControlRigEditor::SelectBone(const FName& InBone)
 {
+	DECLARE_SCOPE_HIERARCHICAL_COUNTER_FUNC()
+
 	// edit mode has to know
 	GetEditMode().SelectBone(InBone);
 	// copy locally, we use this for copying back to template when modified
@@ -1384,6 +1449,8 @@ void FControlRigEditor::SelectBone(const FName& InBone)
 
 FTransform FControlRigEditor::GetBoneTransform(const FName& InBone, bool bLocal) const
 {
+	DECLARE_SCOPE_HIERARCHICAL_COUNTER_FUNC()
+
 	// @todo: think about transform mode
 	if (bLocal)
 	{
@@ -1395,6 +1462,8 @@ FTransform FControlRigEditor::GetBoneTransform(const FName& InBone, bool bLocal)
 
 void FControlRigEditor::SetBoneTransform(const FName& InBone, const FTransform& InTransform)
 {
+	DECLARE_SCOPE_HIERARCHICAL_COUNTER_FUNC()
+
 	// execution should be off
 	ensure(!ControlRig->bExecutionOn);
 
@@ -1435,7 +1504,8 @@ void FControlRigEditor::SetBoneTransform(const FName& InBone, const FTransform&
 
 void FControlRigEditor::OnFinishedChangingProperties(const FPropertyChangedEvent& PropertyChangedEvent)
 {
-//	UE_LOG(LogControlRigEditor, Warning, TEXT("Current Property being modified : %s"), *GetNameSafe(PropertyChangedEvent.Property));
+	DECLARE_SCOPE_HIERARCHICAL_COUNTER_FUNC()
+	//	UE_LOG(LogControlRigEditor, Warning, TEXT("Current Property being modified : %s"), *GetNameSafe(PropertyChangedEvent.Property));
 
 	FName PropertyName = (PropertyChangedEvent.Property != NULL) ? PropertyChangedEvent.Property->GetFName() : NAME_None;
 	if (PropertyName == GET_MEMBER_NAME_CHECKED(FRigBone, InitialTransform))
@@ -1461,6 +1531,8 @@ void FControlRigEditor::OnFinishedChangingProperties(const FPropertyChangedEvent
 
 void FControlRigEditor::OnHierarchyChanged()
 {
+	DECLARE_SCOPE_HIERARCHICAL_COUNTER_FUNC()
+
 	ClearDetailObject();
 
 	FBlueprintEditorUtils::MarkBlueprintAsModified(GetControlRigBlueprint());
@@ -1486,6 +1558,8 @@ void FControlRigEditor::OnHierarchyChanged()
 
 void FControlRigEditor::OnBoneRenamed(const FName& OldName, const FName& NewName)
 {
+	DECLARE_SCOPE_HIERARCHICAL_COUNTER_FUNC()
+
 	UControlRigBlueprint* Blueprint = GetControlRigBlueprint();
 	for (UEdGraph* Graph : Blueprint->UbergraphPages)
 	{
@@ -1532,8 +1606,87 @@ void FControlRigEditor::OnBoneRenamed(const FName& OldName, const FName& NewName
 	}
 }
 
+void FControlRigEditor::OnCurveContainerChanged()
+{
+	DECLARE_SCOPE_HIERARCHICAL_COUNTER_FUNC()
+
+	ClearDetailObject();
+
+	FBlueprintEditorUtils::MarkBlueprintAsModified(GetControlRigBlueprint());
+
+	UControlRigSkeletalMeshComponent* EditorSkelComp = Cast<UControlRigSkeletalMeshComponent>(GetPersonaToolkit()->GetPreviewScene()->GetPreviewMeshComponent());
+	if (EditorSkelComp)
+	{
+		// restart animation 
+		EditorSkelComp->InitAnim(true);
+		UpdateControlRig();
+	}
+	CacheCurveNameList();
+
+	// notification
+	FNotificationInfo Info(LOCTEXT("CurveContainerChangeHelpMessage", "CurveContainer has been successfully modified."));
+	Info.bFireAndForget = true;
+	Info.FadeOutDuration = 10.0f;
+	Info.ExpireDuration = 0.0f;
+
+	TSharedPtr<SNotificationItem> NotificationPtr = FSlateNotificationManager::Get().AddNotification(Info);
+	NotificationPtr->SetCompletionState(SNotificationItem::CS_Success);
+}
+
+void FControlRigEditor::OnCurveRenamed(const FName& OldName, const FName& NewName)
+{
+	DECLARE_SCOPE_HIERARCHICAL_COUNTER_FUNC()
+
+	UControlRigBlueprint* Blueprint = GetControlRigBlueprint();
+	for (UEdGraph* Graph : Blueprint->UbergraphPages)
+	{
+		UControlRigGraph* RigGraph = Cast<UControlRigGraph>(Graph);
+		if (RigGraph == nullptr)
+		{
+			continue;
+		}
+
+		for (UEdGraphNode* Node : RigGraph->Nodes)
+		{
+			UControlRigGraphNode* RigNode = Cast<UControlRigGraphNode>(Node);
+			if (RigNode == nullptr)
+			{
+				continue;
+			}
+
+			UStructProperty* UnitProperty = RigNode->GetUnitProperty();
+			UStruct* UnitStruct = RigNode->GetUnitScriptStruct();
+			if (UnitProperty && UnitStruct)
+			{
+				for (TFieldIterator<UNameProperty> It(UnitStruct); It; ++It)
+				{
+					if (It->HasMetaData(UControlRig::CurveNameMetaName))
+					{
+						FString PinName = FString::Printf(TEXT("%s.%s"), *UnitProperty->GetName(), *It->GetName());
+						UEdGraphPin* Pin = Node->FindPin(PinName, EEdGraphPinDirection::EGPD_Input);
+						if (Pin)
+						{
+							FName CurrentCurve = FName(*Pin->GetDefaultAsString());
+							if (CurrentCurve == OldName)
+							{
+								const FScopedTransaction Transaction(NSLOCTEXT("ControlRigEditor", "ChangeCurveNamePinValue", "Change Curve Name Pin Value"));
+								Pin->Modify();
+								Pin->GetSchema()->TrySetDefaultValue(*Pin, NewName.ToString());
+							}
+						}
+					}
+				}
+			}
+		}
+
+		CacheCurveNameList();
+	}
+}
+
 void FControlRigEditor::OnGraphNodeDropToPerform(TSharedPtr<FGraphNodeDragDropOp> DragDropOp, UEdGraph* Graph, const FVector2D& NodePosition, const FVector2D& ScreenPosition)
 {
+	DECLARE_SCOPE_HIERARCHICAL_COUNTER_FUNC()
+
 	if (DragDropOp->IsOfType<FRigHierarchyDragDropOp>())
 	{
 		TSharedPtr<FRigHierarchyDragDropOp> RigHierarchyOp = StaticCastSharedPtr<FRigHierarchyDragDropOp>(DragDropOp);
@@ -1658,6 +1811,8 @@ void FControlRigEditor::OnGraphNodeDropToPerform(TSharedPtr<FGraphNodeDragDropOp
 
 void FControlRigEditor::HandleMakeBoneGetterSetter(EBoneGetterSetterType Type, bool bIsGetter, TArray<FName> BoneNames, UEdGraph* Graph, FVector2D NodePosition)
 {
+	DECLARE_SCOPE_HIERARCHICAL_COUNTER_FUNC()
+
 	UStruct* StructTemplate = nullptr;
 
 	if (bIsGetter)
@@ -1776,6 +1931,8 @@ void FControlRigEditor::HandleMakeBoneGetterSetter(EBoneGetterSetterType Type, b
 
 void FControlRigEditor::UpdateGraphCompilerErrors()
 {
+	DECLARE_SCOPE_HIERARCHICAL_COUNTER_FUNC()
+
 	UControlRigBlueprint* Blueprint = Cast<UControlRigBlueprint>(GetBlueprintObj());
 	if (Blueprint)
 	{
diff --git a/Engine/Plugins/Experimental/ControlRig/Source/ControlRigEditor/Private/Editor/ControlRigEditor.h b/Engine/Plugins/Experimental/ControlRig/Source/ControlRigEditor/Private/Editor/ControlRigEditor.h
--- a/Engine/Plugins/Experimental/ControlRig/Source/ControlRigEditor/Private/Editor/ControlRigEditor.h
+++ b/Engine/Plugins/Experimental/ControlRig/Source/ControlRigEditor/Private/Editor/ControlRigEditor.h
@@ -110,6 +110,8 @@ public:
 	// restart animation 
 	void OnHierarchyChanged();
 	void OnBoneRenamed(const FName& OldName, const FName& NewName);
+	void OnCurveContainerChanged();
+	void OnCurveRenamed(const FName& OldName, const FName& NewName);
 
 	void OnGraphNodeDropToPerform(TSharedPtr<FGraphNodeDragDropOp> DragDropOp, UEdGraph* Graph, const FVector2D& NodePosition, const FVector2D& ScreenPosition);
 
@@ -183,6 +185,9 @@ private:
 	/** Update the bone name list for use in bone name combo boxes */
 	void CacheBoneNameList();
 
+	/** Update the curve name list for use in curve name combo boxes */
+	void CacheCurveNameList();
+
 	/** Rebind our anim instance to the preview's skeletal mesh component */
 	void RebindToSkeletalMeshComponent();
 
diff --git a/Engine/Plugins/Experimental/ControlRig/Source/ControlRigEditor/Private/Editor/ControlRigEditorMode.cpp b/Engine/Plugins/Experimental/ControlRig/Source/ControlRigEditor/Private/Editor/ControlRigEditorMode.cpp
--- a/Engine/Plugins/Experimental/ControlRig/Source/ControlRigEditor/Private/Editor/ControlRigEditorMode.cpp
+++ b/Engine/Plugins/Experimental/ControlRig/Source/ControlRigEditor/Private/Editor/ControlRigEditorMode.cpp
@@ -8,6 +8,7 @@
 #include "PersonaTabs.h"
 #include "RigHierarchyTabSummoner.h"
 #include "RigStackTabSummoner.h"
+#include "RigCurveContainerTabSummoner.h"
 
 FControlRigEditorMode::FControlRigEditorMode(const TSharedRef<FControlRigEditor>& InControlRigEditor)
 	: FBlueprintEditorApplicationMode(InControlRigEditor, FControlRigEditorModes::ControlRigEditorMode, FControlRigEditorModes::GetLocalizedMode, false, false)
@@ -16,6 +17,7 @@ FControlRigEditorMode::FControlRigEditorMode(const TSharedRef<FControlRigEditor>
 
 	TabFactories.RegisterFactory(MakeShared<FRigHierarchyTabSummoner>(InControlRigEditor));
 	TabFactories.RegisterFactory(MakeShared<FRigStackTabSummoner>(InControlRigEditor));
+	TabFactories.RegisterFactory(MakeShared<FRigCurveContainerTabSummoner>(InControlRigEditor));
 
 	FPersonaModule& PersonaModule = FModuleManager::LoadModuleChecked<FPersonaModule>("Persona");
 
diff --git a/Engine/Plugins/Experimental/ControlRig/Source/ControlRigEditor/Private/Editor/ControlRigSkeletalMeshComponent.cpp b/Engine/Plugins/Experimental/ControlRig/Source/ControlRigEditor/Private/Editor/ControlRigSkeletalMeshComponent.cpp
--- a/Engine/Plugins/Experimental/ControlRig/Source/ControlRigEditor/Private/Editor/ControlRigSkeletalMeshComponent.cpp
+++ b/Engine/Plugins/Experimental/ControlRig/Source/ControlRigEditor/Private/Editor/ControlRigSkeletalMeshComponent.cpp
@@ -28,6 +28,7 @@ void UControlRigSkeletalMeshComponent::ShowReferencePose(bool bRefPose)
 	{
 		UControlRig* ControlRig = ControlRigInstance->CachedControlRig.Get();
 		ControlRig->bExecutionOn = !bRefPose;
+		ControlRigInstance->SetAnimationAsset(nullptr);
 	}
 }
 
@@ -91,4 +92,14 @@ FTransform UControlRigSkeletalMeshComponent::GetDrawTransform(int32 BoneIndex) c
 	}
 
 	return FTransform::Identity;
-}
\ No newline at end of file
+}
+
+void UControlRigSkeletalMeshComponent::EnablePreview(bool bEnable, UAnimationAsset* PreviewAsset)
+{
+	UControlRigSequencerAnimInstance* ControlRigInstance = Cast<UControlRigSequencerAnimInstance>(GetAnimInstance());
+
+	if (ControlRigInstance && ControlRigInstance->CachedControlRig.IsValid())
+	{
+		ControlRigInstance->SetAnimationAsset(PreviewAsset);
+	}
+}
diff --git a/Engine/Plugins/Experimental/ControlRig/Source/ControlRigEditor/Private/Editor/ControlRigSkeletalMeshComponent.h b/Engine/Plugins/Experimental/ControlRig/Source/ControlRigEditor/Private/Editor/ControlRigSkeletalMeshComponent.h
--- a/Engine/Plugins/Experimental/ControlRig/Source/ControlRigEditor/Private/Editor/ControlRigSkeletalMeshComponent.h
+++ b/Engine/Plugins/Experimental/ControlRig/Source/ControlRigEditor/Private/Editor/ControlRigSkeletalMeshComponent.h
@@ -33,6 +33,9 @@ class UControlRigSkeletalMeshComponent : public UDebugSkelMeshComponent
 	{
 		return DebugDrawBones.Num();
 	}
+
+	virtual void EnablePreview(bool bEnable, class UAnimationAsset * PreviewAsset) override;
+
 	// END UDebugSkeletalMeshComponent interface
 
 public:
diff --git a/Engine/Plugins/Experimental/ControlRig/Source/ControlRigEditor/Private/Editor/RigCurveContainerTabSummoner.cpp b/Engine/Plugins/Experimental/ControlRig/Source/ControlRigEditor/Private/Editor/RigCurveContainerTabSummoner.cpp
new file mode 100644
--- /dev/null
+++ b/Engine/Plugins/Experimental/ControlRig/Source/ControlRigEditor/Private/Editor/RigCurveContainerTabSummoner.cpp
@@ -0,0 +1,30 @@
+// Copyright 1998-2019 Epic Games, Inc. All Rights Reserved.
+
+#include "RigCurveContainerTabSummoner.h"
+#include "SRigCurveContainer.h"
+#include "ControlRigEditorStyle.h"
+#include "ControlRigEditor.h"
+
+#define LOCTEXT_NAMESPACE "RigCurveContainerTabSummoner"
+
+const FName FRigCurveContainerTabSummoner::TabID(TEXT("RigCurveContainer"));
+
+FRigCurveContainerTabSummoner::FRigCurveContainerTabSummoner(const TSharedRef<FControlRigEditor>& InControlRigEditor)
+	: FWorkflowTabFactory(TabID, InControlRigEditor)
+	, ControlRigEditor(InControlRigEditor)
+{
+	TabLabel = LOCTEXT("RigCurveContainerTabLabel", "Curve Container");
+	TabIcon = FSlateIcon(FControlRigEditorStyle::Get().GetStyleSetName(), "ControlRig.TabIcon");
+
+	bIsSingleton = true;
+
+	ViewMenuDescription = LOCTEXT("RigCurveContainer_ViewMenu_Desc", "Curve Container");
+	ViewMenuTooltip = LOCTEXT("RigCurveContainer_ViewMenu_ToolTip", "Show the Rig Curve Container tab");
+}
+
+TSharedRef<SWidget> FRigCurveContainerTabSummoner::CreateTabBody(const FWorkflowTabSpawnInfo& Info) const
+{
+	return SNew(SRigCurveContainer, ControlRigEditor.Pin().ToSharedRef());
+}
+
+#undef LOCTEXT_NAMESPACE 
diff --git a/Engine/Plugins/Experimental/ControlRig/Source/ControlRigEditor/Private/Editor/RigCurveContainerTabSummoner.h b/Engine/Plugins/Experimental/ControlRig/Source/ControlRigEditor/Private/Editor/RigCurveContainerTabSummoner.h
new file mode 100644
--- /dev/null
+++ b/Engine/Plugins/Experimental/ControlRig/Source/ControlRigEditor/Private/Editor/RigCurveContainerTabSummoner.h
@@ -0,0 +1,22 @@
+// Copyright 1998-2019 Epic Games, Inc. All Rights Reserved.
+
+#pragma once
+
+#include "CoreMinimal.h"
+#include "WorkflowOrientedApp/WorkflowTabFactory.h"
+
+class FControlRigEditor;
+
+struct FRigCurveContainerTabSummoner : public FWorkflowTabFactory
+{
+public:
+	static const FName TabID;
+	
+public:
+	FRigCurveContainerTabSummoner(const TSharedRef<FControlRigEditor>& InControlRigEditor);
+	
+	virtual TSharedRef<SWidget> CreateTabBody(const FWorkflowTabSpawnInfo& Info) const override;
+	
+protected:
+	TWeakPtr<FControlRigEditor> ControlRigEditor;
+};
diff --git a/Engine/Plugins/Experimental/ControlRig/Source/ControlRigEditor/Private/Editor/SControlRigProfilingView.cpp b/Engine/Plugins/Experimental/ControlRig/Source/ControlRigEditor/Private/Editor/SControlRigProfilingView.cpp
new file mode 100644
--- /dev/null
+++ b/Engine/Plugins/Experimental/ControlRig/Source/ControlRigEditor/Private/Editor/SControlRigProfilingView.cpp
@@ -0,0 +1,1041 @@
+// Copyright 1998-2019 Epic Games, Inc. All Rights Reserved.
+
+#include "SControlRigProfilingView.h"
+#include "Widgets/Layout/SBorder.h"
+#include "Widgets/SBoxPanel.h"
+#include "Widgets/Layout/SScrollBar.h"
+#include "Widgets/Layout/SSpacer.h"
+#include "Widgets/Input/SButton.h"
+#include "Widgets/Input/SCheckBox.h"
+#include "Widgets/Input/SComboBox.h"
+#include "Widgets/Input/SEditableTextBox.h"
+#include "ControlRigEditorStyle.h"
+#include "Framework/MultiBox/MultiBoxBuilder.h"
+#include "EditorStyleSet.h"
+#include "SlateOptMacros.h"
+#include "HAL/ConsoleManager.h"
+
+#include "Framework/Application/SlateApplication.h"
+#include "Dom/JsonValue.h"
+#include "Dom/JsonObject.h"
+#include "Policies/PrettyJsonPrintPolicy.h"
+#include "Serialization/JsonReader.h"
+#include "Serialization/JsonSerializer.h"
+#include "Misc/FileHelper.h"
+
+#include "DesktopPlatformModule.h"
+
+#define LOCTEXT_NAMESPACE "SControlRigProfilingView"
+
+class FStatsHierarchicalClient
+{
+public:
+
+	static TSharedPtr<FJsonValue> ToJson(const FStatsTreeElement* Element)
+	{
+		if (!Element)
+		{
+			return nullptr;
+		}
+
+		TSharedPtr<FJsonObject> Data = MakeShareable(new FJsonObject);
+		Data->SetStringField(TEXT("Name"), Element->GetName());
+		Data->SetStringField(TEXT("Path"), Element->GetPath());
+		Data->SetNumberField(TEXT("Cycles"), (double)Element->TotalCycles());
+		Data->SetNumberField(TEXT("Invocations"), (double)Element->Num());
+
+		if (Element->GetChildren().Num() > 0)
+		{
+			TArray<TSharedPtr<FJsonValue>> ChildArray;
+			for (const TSharedPtr<FStatsTreeElement>& Child : Element->GetChildren())
+			{
+				TSharedPtr<FJsonValue> ChildData = ToJson(Child.Get());
+				if (ChildData)
+				{
+					ChildArray.Add(ChildData);
+				}
+			}
+			Data->SetArrayField(TEXT("Children"), ChildArray);
+		}
+
+		return MakeShareable(new FJsonValueObject(Data));
+	}
+
+	static TSharedPtr<FStatsTreeElement> FromJson(const FJsonObject* Data)
+	{
+		TSharedPtr<FStatsTreeElement> Element = MakeShareable(new FStatsTreeElement);
+
+		const TArray<TSharedPtr<FJsonValue>>* ChildArray;
+		if (Data->TryGetArrayField(TEXT("Markers"), ChildArray))
+		{
+			Element->Invocations = 1;
+
+			for (const TSharedPtr<FJsonValue>& ChildData : *ChildArray)
+			{
+				const TSharedPtr<FJsonObject>* ChildObj = nullptr;
+				if (ChildData->TryGetObject(ChildObj))
+				{
+					Element->Children.Add(FStatsHierarchicalClient::FromJson(ChildObj->Get()));
+				}
+			}
+
+			Element->UpdatePostMeasurement();
+		}
+		else
+		{
+			Element->Name = *Data->GetStringField(TEXT("Name"));
+			Element->Path = Data->GetStringField(TEXT("Path"));
+			Element->Cycles = (uint32)Data->GetNumberField(TEXT("Cycles"));
+			Element->Invocations = (uint32)Data->GetNumberField(TEXT("Invocations"));
+
+			if (Data->TryGetArrayField(TEXT("Children"), ChildArray))
+			{
+				for (const TSharedPtr<FJsonValue>& ChildData : *ChildArray)
+				{
+					const TSharedPtr<FJsonObject>* ChildObj = nullptr;
+					if (ChildData->TryGetObject(ChildObj))
+					{
+						Element->Children.Add(FStatsHierarchicalClient::FromJson(ChildObj->Get()));
+					}
+				}
+			}
+		}
+		return Element;
+	}
+
+	static void GetAllElementsInline(const FStatsTreeElement* Element, TArray<TSharedPtr<FStatsTreeElement>>& OutElements)
+	{
+		for (const TSharedPtr<FStatsTreeElement>& Child : Element->Children)
+		{
+			OutElements.Add(Child);
+			GetAllElementsInline(Child.Get(), OutElements);
+		}
+	}
+
+	static void GetAllElementsCombined(const FStatsTreeElement* Element, TArray<TSharedPtr<FStatsTreeElement>>& OutElements, TMap<FName, TSharedPtr<FStatsTreeElement>>& InOutMap)
+	{
+		bool bIsRoot = InOutMap.Num() == 0;
+
+		for (const TSharedPtr<FStatsTreeElement>& Child : Element->Children)
+		{
+			TSharedPtr<FStatsTreeElement>* ExistingElementPtr = InOutMap.Find(Child->GetFName());
+			if (ExistingElementPtr)
+			{
+				TSharedPtr<FStatsTreeElement>& ExistingElement = *ExistingElementPtr;
+				ExistingElement->Cycles += Child->Cycles;
+				ExistingElement->CyclesOfChildren += Child->CyclesOfChildren;
+				ExistingElement->Invocations += Child->Invocations;
+			}
+			else
+			{
+				TSharedPtr<FStatsTreeElement> NewElement = MakeShareable(new FStatsTreeElement);
+				NewElement->Name = Child->GetFName();
+				NewElement->Path = Child->GetName();
+				NewElement->Cycles = Child->Cycles;
+				NewElement->CyclesOfChildren = Child->CyclesOfChildren;
+				NewElement->Invocations = Child->Invocations;
+				OutElements.Add(NewElement);
+				InOutMap.Add(NewElement->GetFName(), NewElement);
+			}
+			GetAllElementsCombined(Child.Get(), OutElements, InOutMap);
+		}
+
+		if (bIsRoot)
+		{
+			TSharedPtr<FStatsTreeElement> TempParentElement = MakeShareable(new FStatsTreeElement);
+			TempParentElement->Children = OutElements;
+			TempParentElement->UpdatePostMeasurement();
+		}
+	}
+
+	static bool LeafChildPathContains(const FStatsTreeElement* Element, const FString& BasePath, const FString& SearchText)
+	{
+		if (Element == nullptr)
+		{
+			return false;
+		}
+
+		if (Element->Children.Num() == 0)
+		{
+			FString Path = Element->GetPath().RightChop(BasePath.Len());
+			return Path.Contains(SearchText);
+		}
+
+		for (const TSharedPtr<FStatsTreeElement>& Child : Element->Children)
+		{
+			const FStatsTreeElement* ChildPtr = Child.Get();
+			if (LeafChildPathContains(ChildPtr, BasePath, SearchText))
+			{
+				return true;
+			}
+		}
+
+		return false;
+	}
+};
+
+//////////////////////////////////////////////////////////////
+/// SControlRigProfilingItem
+///////////////////////////////////////////////////////////
+
+FName SControlRigProfilingItem::NAME_MarkerName(TEXT("Marker"));
+FName SControlRigProfilingItem::NAME_TotalTimeInclusive(TEXT("Total"));
+FName SControlRigProfilingItem::NAME_TotalTimeExclusive(TEXT("Total (ex)"));
+FName SControlRigProfilingItem::NAME_AverageTimeInclusive(TEXT("Average"));
+FName SControlRigProfilingItem::NAME_AverageTimeExclusive(TEXT("Average (ex)"));
+FName SControlRigProfilingItem::NAME_Invocations(TEXT("Calls"));
+
+void SControlRigProfilingItem::Construct(const FArguments& InArgs, const TSharedRef<STableViewBase>& InOwnerTable, SControlRigProfilingView* InProfilingView, TSharedPtr<FStatsTreeElement> InTreeElement)
+{
+	WeakTreeElement = InTreeElement;
+	ProfilingView = InProfilingView;
+
+	SMultiColumnTableRow<TSharedPtr<FStatsTreeElement>>::Construct(
+		SMultiColumnTableRow<TSharedPtr<FStatsTreeElement>>::FArguments()
+			.Padding(0)	
+		, InOwnerTable);
+
+	UpdateVisibilityFromSearch(ProfilingView->SearchText);
+}
+
+BEGIN_SLATE_FUNCTION_BUILD_OPTIMIZATION
+TSharedRef<SWidget> SControlRigProfilingItem::GenerateWidgetForColumn(const FName& ColumnName)
+{
+	if (ColumnName == NAME_MarkerName)
+	{
+		return SNew(SHorizontalBox)
+
+		+ SHorizontalBox::Slot()
+		.AutoWidth()
+		[
+			SNew(SExpanderArrow, SharedThis(this))
+			.IndentAmount(16)
+			.ShouldDrawWires(true)
+		]
+
+		+ SHorizontalBox::Slot()
+		.AutoWidth()
+		.VAlign(VAlign_Center)
+		.HAlign(HAlign_Left)
+		[
+			SNew(STextBlock)
+			.Text(this, &SControlRigProfilingItem::GetLabelText)
+			.ToolTipText(this, &SControlRigProfilingItem::GetTooltipText)
+			.ColorAndOpacity(this, &SControlRigProfilingItem::GetTextColor)
+			.OnDoubleClicked(this, &SControlRigProfilingItem::OnDoubleClicked)
+		];
+	}
+	else if (ColumnName == NAME_TotalTimeInclusive)
+	{
+		return SNew(SHorizontalBox)
+
+			+ SHorizontalBox::Slot()
+			.AutoWidth()
+			.VAlign(VAlign_Center)
+			.HAlign(HAlign_Right)
+			[
+				SNew(STextBlock)
+				.Text(this, &SControlRigProfilingItem::GetTotalTimeText, true)
+			];
+	}
+	else if (ColumnName == NAME_TotalTimeExclusive)
+	{
+		return SNew(SHorizontalBox)
+
+			+ SHorizontalBox::Slot()
+			.AutoWidth()
+			.VAlign(VAlign_Center)
+			.HAlign(HAlign_Right)
+			[
+				SNew(STextBlock)
+				.Text(this, &SControlRigProfilingItem::GetTotalTimeText, false)
+			];
+	}
+	else if (ColumnName == NAME_AverageTimeInclusive)
+	{
+		return SNew(SHorizontalBox)
+
+		+ SHorizontalBox::Slot()
+		.AutoWidth()
+		.VAlign(VAlign_Center)
+		.HAlign(HAlign_Right)
+		[
+			SNew(STextBlock)
+			.Text(this, &SControlRigProfilingItem::GetAverageTimeText, true)
+		];
+	}
+	else if (ColumnName == NAME_AverageTimeExclusive)
+	{
+		return SNew(SHorizontalBox)
+
+			+ SHorizontalBox::Slot()
+			.AutoWidth()
+			.VAlign(VAlign_Center)
+			.HAlign(HAlign_Right)
+			[
+				SNew(STextBlock)
+				.Text(this, &SControlRigProfilingItem::GetAverageTimeText, false)
+			];
+	}
+	else if (ColumnName == NAME_Invocations)
+	{
+		return SNew(SHorizontalBox)
+
+		+ SHorizontalBox::Slot()
+		.AutoWidth()
+		.VAlign(VAlign_Center)
+		.HAlign(HAlign_Right)
+		[
+			SNew(STextBlock)
+			.Text(this, &SControlRigProfilingItem::GetInvocationsText)
+		];
+	}
+
+	return SNullWidget::NullWidget;
+}
+END_SLATE_FUNCTION_BUILD_OPTIMIZATION
+
+FText SControlRigProfilingItem::GetLabelText() const
+{
+	const FName& Name = WeakTreeElement.Pin()->GetFName();
+	if (Name == FStatsHierarchical::GetUntrackedTimeName())
+	{
+		return FText::FromString(TEXT("Unknown..."));
+	}
+	return (FText::FromName(Name));
+}
+
+FText SControlRigProfilingItem::GetTooltipText() const
+{
+	FString Path = WeakTreeElement.Pin()->GetPath();
+	
+	FString Tooltip;
+	FString Indent;
+
+	FString Left, Right;
+	while (Path.Split(TEXT("."), &Left, &Right))
+	{
+		Tooltip += Indent + Left + TEXT("\n");
+		Indent = Indent + TEXT("    ");
+		Path = Right;
+	}
+	Tooltip += Indent + Path;
+
+	return FText::FromString(Tooltip);
+}
+
+FSlateColor SControlRigProfilingItem::GetTextColor() const
+{
+	static FVector4 Green(FLinearColor::Green);
+	static FVector4 Red(FLinearColor::Red);
+	FVector4 Color = FMath::Lerp<FVector4>(Green, Red, FMath::Clamp<float>((float)WeakTreeElement.Pin()->Contribution(true /* against max */), 0.f, 1.f));
+	return FSlateColor(FLinearColor(Color));
+}
+
+FText SControlRigProfilingItem::GetTotalTimeText(bool bInclusive) const
+{
+	return FText::FromString(FString::Printf(TEXT("%.03f ms"), float(WeakTreeElement.Pin()->TotalSeconds(bInclusive) * 1000.0)));
+}
+
+FText SControlRigProfilingItem::GetAverageTimeText(bool bInclusive) const
+{
+	return FText::FromString(FString::Printf(TEXT("%.03f ms"), float(WeakTreeElement.Pin()->AverageSeconds(bInclusive) * 1000.0)));
+}
+
+FText SControlRigProfilingItem::GetInvocationsText() const
+{
+	return FText::FromString(FString::FormatAsNumber((int32)WeakTreeElement.Pin()->Num()));
+}
+
+FReply SControlRigProfilingItem::OnDoubleClicked(const FGeometry& MyGeometry, const FPointerEvent& PointerEvent)
+{
+	ProfilingView->ToggleItemExpansion(WeakTreeElement.Pin(), true /* recursive */);
+	return FReply::Handled();
+}
+
+void SControlRigProfilingItem::UpdateVisibilityFromSearch(const FString& InSearchText)
+{
+	if (InSearchText.IsEmpty())
+	{
+		SetVisibility(EVisibility::Visible);
+		return;
+	}
+
+	if (WeakTreeElement.Pin()->GetName().Contains(InSearchText))
+	{
+		SetVisibility(EVisibility::Visible);
+		return;
+	}
+
+	if (ProfilingView->DisplayMode != TEXT("Hierarchical"))
+	{
+		SetVisibility(EVisibility::Hidden);
+		return;
+	}
+
+	bool bFoundMatch = FStatsHierarchicalClient::LeafChildPathContains(WeakTreeElement.Pin().Get(), WeakTreeElement.Pin()->GetPath(), InSearchText);
+	SetVisibility(bFoundMatch ? EVisibility::Visible : EVisibility::Hidden);
+}
+
+//////////////////////////////////////////////////////////////
+/// SControlRigProfilingView
+///////////////////////////////////////////////////////////
+
+SControlRigProfilingView::~SControlRigProfilingView()
+{
+}
+
+void SControlRigProfilingView::Construct( const FArguments& InArgs)
+{
+	SortColumn = SControlRigProfilingItem::NAME_MarkerName;
+	bSortAscending = false;
+	bShowUntracked = false;
+
+	RecordTime = TEXT("1 s");
+	RecordTimeOptions.Add(MakeShareable(new FName(TEXT("1/120 s"))));
+	RecordTimeOptions.Add(MakeShareable(new FName(TEXT("1/60 s"))));
+	RecordTimeOptions.Add(MakeShareable(new FName(TEXT("1/2 s"))));
+	RecordTimeOptions.Add(MakeShareable(new FName(TEXT("1 s"))));
+	RecordTimeOptions.Add(MakeShareable(new FName(TEXT("5 s"))));
+	RecordTimeOptions.Add(MakeShareable(new FName(TEXT("10 s"))));
+
+	DisplayMode = TEXT("Hierarchical");
+	DisplayModeOptions.Add(MakeShareable(new FName(TEXT("Hierarchical"))));
+	DisplayModeOptions.Add(MakeShareable(new FName(TEXT("Flat (separate)"))));
+	DisplayModeOptions.Add(MakeShareable(new FName(TEXT("Flat (combined)"))));
+
+	ParallelAnimEvaluationVarPrevValue = 0;
+
+	ChildSlot
+		[
+			SNew(SVerticalBox)
+
+			+ SVerticalBox::Slot()
+			.Padding(4.0f, 4.0f)
+			.AutoHeight()
+			.MaxHeight(40)
+			[
+				SNew(SHorizontalBox)
+
+				+ SHorizontalBox::Slot()
+				.AutoWidth()
+				.VAlign(VAlign_Center)
+				.HAlign(HAlign_Left)
+				[
+					SNew(SButton)
+					.OnClicked(this, &SControlRigProfilingView::HandleProfilingButton)
+					.Text(this, &SControlRigProfilingView::GetProfilingButtonText)
+				]
+
+				+ SHorizontalBox::Slot()
+				.FillWidth(1)
+				.MaxWidth(10)
+				[
+					SNew(SSpacer)
+				]
+
+				+ SHorizontalBox::Slot()
+				.AutoWidth()
+				.VAlign(VAlign_Center)
+				.HAlign(HAlign_Left)
+				[
+					SNew(SComboBox<TSharedPtr<FName>>)
+					.OptionsSource(&RecordTimeOptions)
+					.OnSelectionChanged(this, &SControlRigProfilingView::OnComboBoxChanged, &RecordTime)
+					.OnGenerateWidget(this, &SControlRigProfilingView::OnGetComboBoxWidget)
+					.Content()
+					[
+						SNew(STextBlock).Text(this, &SControlRigProfilingView::GetComboBoxValueAsText, &RecordTime)
+					]
+				]
+
+				+ SHorizontalBox::Slot()
+				.FillWidth(1)
+				.MaxWidth(30)
+				[
+					SNew(SSpacer)
+				]
+
+				+ SHorizontalBox::Slot()
+				.VAlign(VAlign_Center)
+				.AutoWidth()
+				[
+					SNew(STextBlock)
+					.Text(LOCTEXT("ShowUntrackedMarkers", "Record unmarked time:"))
+				]
+
+				+ SHorizontalBox::Slot()
+				.AutoWidth()
+				.VAlign(VAlign_Center)
+				.HAlign(HAlign_Left)
+				[
+					SNew(SCheckBox)
+					.OnCheckStateChanged(this, &SControlRigProfilingView::UntrackedCheckBoxCheckStateChanged)
+					.IsChecked(this, &SControlRigProfilingView::GetUntrackedCheckBoxCheckState)
+				]
+
+				+ SHorizontalBox::Slot()
+				.FillWidth(1)
+				.MaxWidth(60)
+				[
+					SNew(SSpacer)
+				]
+
+				+ SHorizontalBox::Slot()
+				.AutoWidth()
+				.VAlign(VAlign_Center)
+				.HAlign(HAlign_Left)
+				[
+					SNew(SButton)
+					.OnClicked(this, &SControlRigProfilingView::HandleLoadFromFile)
+					.Text(LOCTEXT("ControlRigProfilingViewLoadFromFile", "Load"))
+					.IsEnabled(this, &SControlRigProfilingView::IsLoadingEnabled)
+				]
+
+				+ SHorizontalBox::Slot()
+				.AutoWidth()
+				.VAlign(VAlign_Center)
+				.HAlign(HAlign_Left)
+				[
+					SNew(SButton)
+					.OnClicked(this, &SControlRigProfilingView::HandleSaveToFile)
+					.Text(LOCTEXT("ControlRigProfilingViewSaveToFile", "Save"))
+					.IsEnabled(this, &SControlRigProfilingView::IsSavingEnabled)
+				]
+			]
+
+			+ SVerticalBox::Slot()
+			.Padding(4.0f, 4.0f)
+			.AutoHeight()
+			.MaxHeight(40)
+			[
+				SNew(SHorizontalBox)
+
+				+ SHorizontalBox::Slot()
+				.AutoWidth()
+				.MaxWidth(250)
+				.VAlign(VAlign_Center)
+				.HAlign(HAlign_Left)
+				[
+					SNew(SEditableTextBox)
+					.MinDesiredWidth(250)
+					.Font(FEditorStyle::GetFontStyle("ContentBrowser.AssetTileViewNameFont"))
+					.Text(GetSearchText())
+					.HintText(LOCTEXT("Search", "Search"))
+					.OnTextCommitted(this, &SControlRigProfilingView::OnSearchTextCommitted)
+				]
+
+				+ SHorizontalBox::Slot()
+				.FillWidth(1)
+				.MaxWidth(10)
+				[
+					SNew(SSpacer)
+				]
+
+				+ SHorizontalBox::Slot()
+				.VAlign(VAlign_Center)
+				.AutoWidth()
+				[
+					SNew(STextBlock)
+					.Text(LOCTEXT("DisplayMode", "DisplayMode:"))
+				]
+
+				+ SHorizontalBox::Slot()
+				.FillWidth(1)
+				.MaxWidth(3)
+				[
+					SNew(SSpacer)
+				]
+
+				+ SHorizontalBox::Slot()
+				.AutoWidth()
+				.VAlign(VAlign_Center)
+				.HAlign(HAlign_Left)
+				[
+					SNew(SComboBox<TSharedPtr<FName>>)
+					.OptionsSource(&DisplayModeOptions)
+					.OnSelectionChanged(this, &SControlRigProfilingView::OnComboBoxChanged, &DisplayMode)
+					.OnGenerateWidget(this, &SControlRigProfilingView::OnGetComboBoxWidget)
+					.Content()
+					[
+						SNew(STextBlock).Text(this, &SControlRigProfilingView::GetComboBoxValueAsText, &DisplayMode)
+					]
+				]
+			]
+
+			+ SVerticalBox::Slot()
+			.Padding(0.0f, 0.0f)
+			.FillHeight(1)
+			[
+				SNew(SBorder)
+				.Padding(2.0f)
+				.BorderBackgroundColor(this, &SControlRigProfilingView::GetTreeBackgroundColor)
+				.BorderImage(FEditorStyle::GetBrush("SCSEditor.TreePanel"))
+				[
+					SAssignNew(TreeView, STreeView<TSharedPtr<FStatsTreeElement>>)
+					.ItemHeight(28)
+					.Visibility(this, &SControlRigProfilingView::GetTreeVisibility)
+					.TreeItemsSource(&RootChildren)
+					.SelectionMode(ESelectionMode::Multi)
+					.OnGenerateRow(this, &SControlRigProfilingView::MakeTableRowWidget)
+					.OnGetChildren(this, &SControlRigProfilingView::HandleGetChildrenForTree)
+					.HighlightParentNodesForSelection(true)
+					.HeaderRow
+					(
+						SNew(SHeaderRow)
+
+						+ SHeaderRow::Column(SControlRigProfilingItem::NAME_MarkerName)
+						.DefaultLabel(LOCTEXT("MarkerName", "Marker Name"))
+						.OnSort(FOnSortModeChanged::CreateSP(this, &SControlRigProfilingView::OnSortColumnHeader))
+						.FillWidth(0.8f)
+
+						+ SHeaderRow::Column(SControlRigProfilingItem::NAME_TotalTimeInclusive)
+						.OnSort(FOnSortModeChanged::CreateSP(this, &SControlRigProfilingView::OnSortColumnHeader))
+						.FixedWidth(90.0f)
+						.VAlignHeader(VAlign_Center)
+						.HeaderContent()
+						[
+							SNew(STextBlock)
+							.Text(LOCTEXT("TotalMilisecondsInc", "Total"))
+							.ToolTipText(LOCTEXT("TotalMilisecondsIncTooltip", "Total time in Miliseconds (including children)"))
+						]
+
+						+ SHeaderRow::Column(SControlRigProfilingItem::NAME_TotalTimeExclusive)
+						.OnSort(FOnSortModeChanged::CreateSP(this, &SControlRigProfilingView::OnSortColumnHeader))
+						.FixedWidth(90)
+						.VAlignHeader(VAlign_Center)
+						.HeaderContent()
+						[
+							SNew(STextBlock)
+							.Text(LOCTEXT("TotalMilisecondsExc", "Total (ex)"))
+							.ToolTipText(LOCTEXT("TotalMilisecondsExcTooltip", "Total time in Miliseconds (excluding children)"))
+						]
+
+						+ SHeaderRow::Column(SControlRigProfilingItem::NAME_AverageTimeInclusive)
+						.OnSort(FOnSortModeChanged::CreateSP(this, &SControlRigProfilingView::OnSortColumnHeader))
+						.FixedWidth(90)
+						.VAlignHeader(VAlign_Center)
+						.HeaderContent()
+						[
+							SNew(STextBlock)
+							.Text(LOCTEXT("AverageMilisecondsInc", "Average"))
+							.ToolTipText(LOCTEXT("AverageMilisecondsIncTooltip", "Average time in Miliseconds (including children)"))
+						]
+
+						+ SHeaderRow::Column(SControlRigProfilingItem::NAME_AverageTimeExclusive)
+						.OnSort(FOnSortModeChanged::CreateSP(this, &SControlRigProfilingView::OnSortColumnHeader))
+						.FixedWidth(90)
+						.VAlignHeader(VAlign_Center)
+						.HeaderContent()
+						[
+							SNew(STextBlock)
+							.Text(LOCTEXT("AverageMilisecondsExc", "Average (ex)"))
+							.ToolTipText(LOCTEXT("AverageMilisecondsExcTooltip", "Average time in Miliseconds (excluding children)"))
+						]
+
+						+ SHeaderRow::Column(SControlRigProfilingItem::NAME_Invocations)
+						.OnSort(FOnSortModeChanged::CreateSP(this, &SControlRigProfilingView::OnSortColumnHeader))
+						.FixedWidth(70)
+						.VAlignHeader(VAlign_Center)
+						.HeaderContent()
+						[
+							SNew(STextBlock)
+							.Text(LOCTEXT("NumberOfCalls", "Calls"))
+							.ToolTipText(LOCTEXT("NumberOfCallsTooltip", "Total number of invocations"))
+						]
+					)
+				]
+			]
+		];
+
+	RefreshTreeView();
+}
+
+TSharedRef<ITableRow> SControlRigProfilingView::MakeTableRowWidget(TSharedPtr<FStatsTreeElement> InTreeElement, const TSharedRef<STableViewBase>& OwnerTable)
+{
+	return SNew(SControlRigProfilingItem, OwnerTable, this, InTreeElement);
+}
+
+void SControlRigProfilingView::OnSortColumnHeader(const EColumnSortPriority::Type SortPriority, const FName& ColumnId, const EColumnSortMode::Type NewSortMode)
+{
+	if (FStatsHierarchical::IsEnabled())
+	{
+		return;
+	}
+
+	if (SortColumn == ColumnId)
+	{
+		bSortAscending = !bSortAscending;
+	}
+
+	SortColumn = ColumnId;
+
+	RefreshTreeView();
+}
+
+void SControlRigProfilingView::HandleGetChildrenForTree(TSharedPtr<FStatsTreeElement> InItem, TArray<TSharedPtr<FStatsTreeElement>>& OutChildren)
+{
+	if (DisplayMode == TEXT("Hierarchical"))
+	{
+		OutChildren.Append(InItem->GetChildren());
+		SortChildren(OutChildren);
+	}
+}
+
+void SControlRigProfilingView::SortChildren(TArray<TSharedPtr<FStatsTreeElement>>& InOutChildren) const
+{
+	bool bRequiresSort = SortColumn != SControlRigProfilingItem::NAME_MarkerName;
+	if (!bRequiresSort)
+	{
+		bRequiresSort = DisplayMode != TEXT("Hierarchical");
+	}
+
+	if (bRequiresSort)
+	{
+		Algo::Sort(InOutChildren, [&](const TSharedPtr<FStatsTreeElement>& A, const TSharedPtr<FStatsTreeElement>& B)
+			{
+				if (SortColumn == SControlRigProfilingItem::NAME_MarkerName)
+				{
+					if (bSortAscending)
+					{
+						return A->GetName().Compare(B->GetName()) < 0;
+					}
+					return A->GetName().Compare(B->GetName()) > 0;
+				}
+				if (SortColumn == SControlRigProfilingItem::NAME_TotalTimeInclusive)
+				{
+					if (bSortAscending)
+					{
+						return A->TotalCycles(true) < B->TotalCycles(true);
+					}
+					return A->TotalCycles(true) > B->TotalCycles(true);
+				}
+				if (SortColumn == SControlRigProfilingItem::NAME_TotalTimeExclusive)
+				{
+					if (bSortAscending)
+					{
+						return A->TotalCycles(false) < B->TotalCycles(false);
+					}
+					return A->TotalCycles(false) > B->TotalCycles(false);
+				}
+				if (SortColumn == SControlRigProfilingItem::NAME_AverageTimeInclusive)
+				{
+					if (bSortAscending)
+					{
+						return A->AverageSeconds(true) < B->AverageSeconds(true);
+					}
+					return A->AverageSeconds(true) > B->AverageSeconds(true);
+				}
+				if (SortColumn == SControlRigProfilingItem::NAME_AverageTimeExclusive)
+				{
+					if (bSortAscending)
+					{
+						return A->AverageSeconds(false) < B->AverageSeconds(false);
+					}
+					return A->AverageSeconds(false) > B->AverageSeconds(false);
+				}
+
+				// fall through on invocations
+				if (bSortAscending)
+				{
+					return A->Num() < B->Num();
+				}
+				return A->Num() > B->Num();
+			}
+		);
+	}
+}
+
+EVisibility SControlRigProfilingView::GetTreeVisibility() const
+{
+	return FStatsHierarchical::IsEnabled() ? EVisibility::Hidden : EVisibility::Visible;
+}
+
+FSlateColor SControlRigProfilingView::GetTreeBackgroundColor() const
+{
+	return FSlateColor(FStatsHierarchical::IsEnabled() ? FLinearColor(3, 0, 0, 1) : FLinearColor::White);
+}
+
+FText SControlRigProfilingView::GetProfilingButtonText() const
+{
+	return FText::FromString(FStatsHierarchical::IsEnabled() ? TEXT("Stop") : TEXT("Record"));
+}
+
+FReply SControlRigProfilingView::HandleProfilingButton()
+{
+	if (FStatsHierarchical::IsEnabled())
+	{
+		// todo: aggregate if need be
+		RootElement = FStatsHierarchical::EndMeasurements(FStatsTreeElement(), bShowUntracked);
+
+		// optionally re-enable parallel anim eval
+		FConsoleManager& ConsoleManager = (FConsoleManager&)IConsoleManager::Get();
+		IConsoleVariable* ParallelAnimEvaluationVar = ConsoleManager.FindConsoleVariable(TEXT("a.ParallelAnimEvaluation"));
+		if (ParallelAnimEvaluationVar)
+		{
+			if (ParallelAnimEvaluationVarPrevValue != ParallelAnimEvaluationVar->AsVariableInt()->GetValueOnAnyThread())
+			{
+				*ParallelAnimEvaluationVar->AsVariableInt() = ParallelAnimEvaluationVarPrevValue;
+			}
+		}
+
+		RefreshTreeView();
+	}
+	else
+	{
+		// clear the Tree
+		RootChildren.Reset();
+		TreeView->RequestTreeRefresh();
+
+		// disable parallel evaluation of the animBP
+		FConsoleManager& ConsoleManager = (FConsoleManager&)IConsoleManager::Get();
+		IConsoleVariable* ParallelAnimEvaluationVar = ConsoleManager.FindConsoleVariable(TEXT("a.ParallelAnimEvaluation"));
+		if (ParallelAnimEvaluationVar)
+		{
+			ParallelAnimEvaluationVarPrevValue = ParallelAnimEvaluationVar->AsVariableInt()->GetValueOnAnyThread();
+			*ParallelAnimEvaluationVar->AsVariableInt() = 0;
+		}
+
+		FStatsHierarchical::BeginMeasurements();
+
+		double TimerSeconds = 0.0;
+		if(RecordTime.IsEqual(TEXT("1/120 s")))
+		{
+			TimerSeconds = 1.0 / 120.0;
+		}
+		if(RecordTime.IsEqual(TEXT("1/60 s")))
+		{
+			TimerSeconds = 1.0 / 60.0;
+		}
+		if(RecordTime.IsEqual(TEXT("1/2 s")))
+		{
+			TimerSeconds = 0.5;
+		}
+		if(RecordTime.IsEqual(TEXT("1 s")))
+		{
+			TimerSeconds = 1.0;
+		}
+		if(RecordTime.IsEqual(TEXT("5 s")))
+		{
+			TimerSeconds = 5.0;
+		}
+		if(RecordTime.IsEqual(TEXT("10 s")))
+		{
+			TimerSeconds = 10.0;
+		}
+
+		RegisterActiveTimer(TimerSeconds, FWidgetActiveTimerDelegate::CreateSP(this, &SControlRigProfilingView::OnRecordingTimerFinished));
+	}
+
+	return FReply::Handled();
+}
+
+FText SControlRigProfilingView::GetSearchText() const
+{
+	return FText::FromString(SearchText);
+}
+
+TSharedRef<SWidget> SControlRigProfilingView::OnGetComboBoxWidget(TSharedPtr<FName> InItem)
+{
+	return SNew(STextBlock).Text(FText::FromName(InItem.IsValid() ? *InItem : NAME_None));
+}
+
+void SControlRigProfilingView::OnComboBoxChanged(TSharedPtr<FName> InItem, ESelectInfo::Type InSeletionInfo, FName* OutValue)
+{
+	*OutValue = *InItem;
+
+	if (OutValue == &DisplayMode)
+	{
+		RefreshTreeView();
+	}
+}
+
+FText SControlRigProfilingView::GetComboBoxValueAsText(FName* InValue) const
+{
+	return FText::FromName(*InValue);
+}
+
+EActiveTimerReturnType SControlRigProfilingView::OnRecordingTimerFinished(double InCurrentTime, float InDeltaTime)
+{
+	if (FStatsHierarchical::IsEnabled())
+	{
+		HandleProfilingButton();
+	}
+	return EActiveTimerReturnType::Stop;
+}
+
+void SControlRigProfilingView::UntrackedCheckBoxCheckStateChanged(ECheckBoxState CheckState)
+{
+	bShowUntracked = CheckState == ECheckBoxState::Checked;
+}
+
+ECheckBoxState SControlRigProfilingView::GetUntrackedCheckBoxCheckState() const
+{
+	return bShowUntracked ? ECheckBoxState::Checked : ECheckBoxState::Unchecked;
+}
+
+bool SControlRigProfilingView::IsLoadingEnabled() const
+{
+	if (FStatsHierarchical::IsEnabled())
+	{
+		return false;
+	}
+	return true;
+}
+
+bool SControlRigProfilingView::IsSavingEnabled() const
+{
+	if (FStatsHierarchical::IsEnabled())
+	{
+		return false;
+	}
+	return RootElement.GetChildren().Num() > 0;
+}
+
+FReply SControlRigProfilingView::HandleLoadFromFile()
+{
+	IDesktopPlatform* DesktopPlatform = FDesktopPlatformModule::Get();
+	const void* ParentWindowWindowHandle = FSlateApplication::Get().FindBestParentWindowHandleForDialogs(nullptr);
+
+	const FText Title = LOCTEXT("ImportFromJson", "Import markers from JSON...");
+	const FString FileTypes = TEXT("Profiling JSON (*.json)|*.json");
+
+	TArray<FString> OutFilenames;
+	DesktopPlatform->OpenFileDialog(
+		ParentWindowWindowHandle,
+		Title.ToString(),
+		TEXT(""),
+		TEXT("Profiling.json"),
+		FileTypes,
+		EFileDialogFlags::None,
+		OutFilenames
+	);
+
+	if (OutFilenames.Num() == 0)
+	{
+		return FReply::Unhandled();
+	}
+
+	FString JsonText;
+	if (FFileHelper::LoadFileToString(JsonText, *OutFilenames[0]))
+	{
+
+		TSharedRef<TJsonReader<>> JsonReader = TJsonReaderFactory<>::Create(JsonText);
+		TSharedPtr<FJsonObject> RootData;
+		if (FJsonSerializer::Deserialize(JsonReader, RootData))
+		{
+			TSharedPtr<FStatsTreeElement> ParsedElement = FStatsHierarchicalClient::FromJson(RootData.Get());
+			if (ParsedElement)
+			{
+				RootElement = *ParsedElement.Get();
+				RefreshTreeView();
+			}
+			return FReply::Handled();
+		}
+	}
+
+	return FReply::Unhandled();
+}
+
+FReply SControlRigProfilingView::HandleSaveToFile()
+{
+	IDesktopPlatform* DesktopPlatform = FDesktopPlatformModule::Get();
+	const void* ParentWindowWindowHandle = FSlateApplication::Get().FindBestParentWindowHandleForDialogs(nullptr);
+
+	const FText Title = LOCTEXT("ExportToJson", "Export markers as JSON...");
+	const FString FileTypes = TEXT("Profiling JSON (*.json)|*.json");
+
+	TArray<FString> OutFilenames;
+	DesktopPlatform->SaveFileDialog(
+		ParentWindowWindowHandle,
+		Title.ToString(),
+		TEXT(""),
+		TEXT("Profiling.json"),
+		FileTypes,
+		EFileDialogFlags::None,
+		OutFilenames
+	);
+
+	if (OutFilenames.Num() == 0)
+	{
+		return FReply::Unhandled();
+	}
+
+	TSharedPtr<FJsonObject> RootData = MakeShareable(new FJsonObject);
+	TArray<TSharedPtr<FJsonValue>> RootArray;
+
+	for (const TSharedPtr<FStatsTreeElement>& Child : RootElement.GetChildren())
+	{
+		TSharedPtr<FJsonValue> ChildData = FStatsHierarchicalClient::ToJson(Child.Get());
+		if (ChildData)
+		{
+			RootArray.Add(ChildData);
+		}
+	}
+
+	RootData->SetArrayField(TEXT("Markers"), RootArray);
+
+	FString JsonText;
+	TSharedRef< TJsonWriter< TCHAR, TPrettyJsonPrintPolicy<TCHAR> > > JsonWriter = TJsonWriterFactory< TCHAR, TPrettyJsonPrintPolicy<TCHAR> >::Create(&JsonText);
+	if (FJsonSerializer::Serialize(RootData.ToSharedRef(), JsonWriter))
+	{
+		FFileHelper::SaveStringToFile(JsonText, *OutFilenames[0]);
+		return FReply::Handled();
+	}
+
+	return FReply::Unhandled();
+}
+
+void SControlRigProfilingView::RefreshTreeView()
+{
+	if (DisplayMode == TEXT("Hierarchical"))
+	{
+		RootChildren = RootElement.GetChildren();
+	}
+	else
+	{
+		RootChildren.Reset();
+
+		if (DisplayMode == TEXT("Flat (separate)"))
+		{
+			FStatsHierarchicalClient::GetAllElementsInline(&RootElement, RootChildren);
+		}
+		else if (DisplayMode == TEXT("Flat (combined)"))
+		{
+			TMap<FName, TSharedPtr<FStatsTreeElement>> Map;
+			FStatsHierarchicalClient::GetAllElementsCombined(&RootElement, RootChildren, Map);
+		}
+	}
+
+	SortChildren(RootChildren);
+
+	TreeView->RequestTreeRefresh();
+}
+
+void SControlRigProfilingView::SetItemExpansion(const TSharedPtr<FStatsTreeElement>& InItem, bool bExpand, bool bRecursive)
+{
+	TreeView->SetItemExpansion(InItem, bExpand);
+
+	if (bExpand && bRecursive)
+	{
+		for (const TSharedPtr<FStatsTreeElement>& Child : InItem->GetChildren())
+		{
+			SetItemExpansion(Child, bExpand, bRecursive);
+		}
+	}
+}
+
+void SControlRigProfilingView::ToggleItemExpansion(const TSharedPtr<FStatsTreeElement>& InItem, bool bRecursive)
+{
+	bool bExpand = !TreeView->IsItemExpanded(InItem);
+	SetItemExpansion(InItem, bExpand, bRecursive);
+}
+
+void SControlRigProfilingView::OnSearchTextCommitted(const FText& InNewText, ETextCommit::Type InCommitType)
+{
+	SearchText = InNewText.ToString();
+	TreeView->RebuildList();
+	//TreeView->RequestTreeRefresh();
+}
+
+#undef LOCTEXT_NAMESPACE
diff --git a/Engine/Plugins/Experimental/ControlRig/Source/ControlRigEditor/Private/Editor/SControlRigProfilingView.h b/Engine/Plugins/Experimental/ControlRig/Source/ControlRigEditor/Private/Editor/SControlRigProfilingView.h
new file mode 100644
--- /dev/null
+++ b/Engine/Plugins/Experimental/ControlRig/Source/ControlRigEditor/Private/Editor/SControlRigProfilingView.h
@@ -0,0 +1,112 @@
+// Copyright 1998-2019 Epic Games, Inc. All Rights Reserved.
+
+#pragma once
+
+#include "CoreMinimal.h"
+#include "Stats/StatsHierarchical.h"
+#include "Widgets/Views/STreeView.h"
+
+class SControlRigProfilingView;
+
+class SControlRigProfilingItem
+	: public SMultiColumnTableRow<TSharedPtr<FStatsTreeElement>>
+{
+	SLATE_BEGIN_ARGS(SControlRigProfilingItem) {}
+	SLATE_END_ARGS()
+
+	static FName NAME_MarkerName;
+	static FName NAME_TotalTimeInclusive;
+	static FName NAME_TotalTimeExclusive;
+	static FName NAME_AverageTimeInclusive;
+	static FName NAME_AverageTimeExclusive;
+	static FName NAME_Invocations;
+
+public:
+
+	void Construct(const FArguments& InArgs, const TSharedRef<STableViewBase>& InOwnerTable, SControlRigProfilingView* InProfilingView, const TSharedPtr<FStatsTreeElement> InTreeElement);
+
+	// SMultiColumnTableRow overrides
+	virtual TSharedRef<SWidget> GenerateWidgetForColumn(const FName& ColumnName) override;
+
+private:
+
+	SControlRigProfilingView* ProfilingView;
+	TWeakPtr<FStatsTreeElement> WeakTreeElement;
+
+	FText GetLabelText() const;
+	FText GetTooltipText() const;
+	FSlateColor GetTextColor() const;
+	FText GetTotalTimeText(bool bInclusive) const;
+	FText GetAverageTimeText(bool bInclusive) const;
+	FText GetInvocationsText() const;
+	FReply OnDoubleClicked(const FGeometry& MyGeometry, const FPointerEvent& PointerEvent);
+
+	void UpdateVisibilityFromSearch(const FString& InSearchText);
+
+	friend class SControlRigProfilingView;
+};
+
+class SControlRigProfilingView : public SCompoundWidget
+{
+public:
+	SLATE_BEGIN_ARGS(SControlRigProfilingView) {}
+	SLATE_END_ARGS()
+
+	~SControlRigProfilingView();
+
+	void Construct(const FArguments& InArgs);
+
+protected:
+
+	/** Rebuild the tree view */
+	void RefreshTreeView();
+
+	void SetItemExpansion(const TSharedPtr<FStatsTreeElement>& InItem, bool bExpand = true, bool bRecursive = false);
+	void ToggleItemExpansion(const TSharedPtr<FStatsTreeElement>& InItem, bool bRecursive = false);
+
+private:
+
+	/** Make a row widget for the table */
+	TSharedRef<ITableRow> MakeTableRowWidget(TSharedPtr<FStatsTreeElement> InTreeElement, const TSharedRef<STableViewBase>& OwnerTable);
+
+	void OnSortColumnHeader(const EColumnSortPriority::Type SortPriority, const FName& ColumnId, const EColumnSortMode::Type NewSortMode);
+
+	/** Get children for the tree */
+	void HandleGetChildrenForTree(TSharedPtr<FStatsTreeElement> InTreeElement, TArray<TSharedPtr<FStatsTreeElement>>& OutChildren);
+	void SortChildren(TArray<TSharedPtr<FStatsTreeElement>>& InOutChildren) const;
+
+	EVisibility GetTreeVisibility() const;
+	FSlateColor GetTreeBackgroundColor() const;
+	FText GetProfilingButtonText() const;
+	FReply HandleProfilingButton();
+	FText GetSearchText() const;
+
+	TSharedRef<SWidget> OnGetComboBoxWidget(TSharedPtr<FName> InItem);
+	void OnComboBoxChanged(TSharedPtr<FName> InItem, ESelectInfo::Type InSeletionInfo, FName* OutValue);
+	FText GetComboBoxValueAsText(FName* InValue) const;
+	EActiveTimerReturnType OnRecordingTimerFinished(double InCurrentTime, float InDeltaTime);
+	void OnSearchTextCommitted(const FText& InNewText, ETextCommit::Type InCommitType);
+
+	void UntrackedCheckBoxCheckStateChanged(ECheckBoxState CheckState);
+	ECheckBoxState GetUntrackedCheckBoxCheckState() const;
+	bool IsLoadingEnabled() const;
+	bool IsSavingEnabled() const;
+	FReply HandleLoadFromFile();
+	FReply HandleSaveToFile();
+
+	TSharedPtr<STreeView<TSharedPtr<FStatsTreeElement>>> TreeView;
+
+	FName SortColumn;
+	bool bSortAscending;
+	bool bShowUntracked;
+	FString SearchText;
+	FStatsTreeElement RootElement;
+	TArray<TSharedPtr<FStatsTreeElement>> RootChildren;
+	FName RecordTime;
+	TArray<TSharedPtr<FName>> RecordTimeOptions;
+	FName DisplayMode;
+	TArray<TSharedPtr<FName>> DisplayModeOptions;
+	int32 ParallelAnimEvaluationVarPrevValue;
+
+	friend class SControlRigProfilingItem;
+};
diff --git a/Engine/Plugins/Experimental/ControlRig/Source/ControlRigEditor/Private/Editor/SRigHierarchy.cpp b/Engine/Plugins/Experimental/ControlRig/Source/ControlRigEditor/Private/Editor/SRigHierarchy.cpp
--- a/Engine/Plugins/Experimental/ControlRig/Source/ControlRigEditor/Private/Editor/SRigHierarchy.cpp
+++ b/Engine/Plugins/Experimental/ControlRig/Source/ControlRigEditor/Private/Editor/SRigHierarchy.cpp
@@ -225,6 +225,7 @@ void SRigHierarchy::Construct(const FArguments& InArgs, TSharedRef<FControlRigEd
 				.OnGetChildren(this, &SRigHierarchy::HandleGetChildrenForTree)
 				.OnSelectionChanged(this, &SRigHierarchy::OnSelectionChanged)
 				.OnContextMenuOpening(this, &SRigHierarchy::CreateContextMenu)
+				.HighlightParentNodesForSelection(true)
 				.ItemHeight(24)
 			]
 		]
@@ -273,14 +274,15 @@ void SRigHierarchy::RefreshTreeView()
 
 		FString FilteredString = FilterText.ToString();
 		const bool bSearchOff = FilteredString.IsEmpty();
-		for (int32 BoneIndex = 0; BoneIndex < Hierarchy.Bones.Num(); ++BoneIndex)
+		const TArray<FRigBone>& Bones = Hierarchy.GetBones();
+		for (int32 BoneIndex = 0; BoneIndex < Bones.Num(); ++BoneIndex)
 		{
-			FRigBone& Bone = Hierarchy.Bones[BoneIndex];
+			const FRigBone& Bone = Bones[BoneIndex];
 
 			// create new item
 			if (bSearchOff)
 			{
-				TSharedPtr<FRigTreeBone> NewItem = MakeShared<FRigTreeBone>(Hierarchy.Bones[BoneIndex].Name, SharedThis(this));
+				TSharedPtr<FRigTreeBone> NewItem = MakeShared<FRigTreeBone>(Bones[BoneIndex].Name, SharedThis(this));
 				SearchTable.Add(Bone.Name, NewItem);
 
 				if (Bone.ParentName == NAME_None)
@@ -299,7 +301,7 @@ void SRigHierarchy::RefreshTreeView()
 			else if (Bone.Name.ToString().Contains(FilteredString))
 			{
 				// if contains, just list out everything to root
-				TSharedPtr<FRigTreeBone> NewItem = MakeShared<FRigTreeBone>(Hierarchy.Bones[BoneIndex].Name, SharedThis(this));
+				TSharedPtr<FRigTreeBone> NewItem = MakeShared<FRigTreeBone>(Bones[BoneIndex].Name, SharedThis(this));
 				// during search, everything is on root
 				RootBones.Add(NewItem);
 			}
@@ -348,7 +350,7 @@ void SRigHierarchy::OnSelectionChanged(TSharedPtr<FRigTreeBone> Selection, ESele
 			const int32 BoneIndex = RigHierarchy->GetIndex(Selection->CachedBone);
 			if (BoneIndex != INDEX_NONE)
 			{
-				ControlRigEditor.Pin()->SetDetailStruct(MakeShareable(new FStructOnScope(FRigBone::StaticStruct(), (uint8*)&RigHierarchy->Bones[BoneIndex])));
+				ControlRigEditor.Pin()->SetDetailStruct(MakeShareable(new FStructOnScope(FRigBone::StaticStruct(), (uint8*)&RigHierarchy->GetBones()[BoneIndex])));
 				ControlRigEditor.Pin()->SelectBone(Selection->CachedBone);
 				return;
 			}
@@ -416,81 +418,16 @@ TSharedPtr< SWidget > SRigHierarchy::CreateContextMenu()
 		MenuBuilder.AddMenuSeparator();
 		MenuBuilder.AddSubMenu(
 			LOCTEXT("ImportSubMenu", "Import"),
-			LOCTEXT("ImportSubMenu_ToolTip", "Import hierarchy to the current rig. This only imports non-existing node. For example, if there is hand_r, it won't import hand_r. \
-				If you want to reimport whole new hiearchy, delete all nodes, and use import hierarchy."),
+			LOCTEXT("ImportSubMenu_ToolTip", "Import hierarchy to the current rig. This overrides the data if it contains the existing node."),
 			FNewMenuDelegate::CreateSP(this, &SRigHierarchy::CreateImportMenu)
 		);
 
-		MenuBuilder.AddMenuSeparator();
-		MenuBuilder.AddSubMenu(
-			LOCTEXT("RefreshSubMenu", "Refresh"),
-			LOCTEXT("RefreshSubMenu_ToolTip", "Refresh the existing initial transform from the selected mesh. This only updates if the node is found."),
-			FNewMenuDelegate::CreateSP(this, &SRigHierarchy::CreateRefreshMenu)
-		);
-	
 		MenuBuilder.EndSection();
 	}
 
 	return MenuBuilder.MakeWidget();
 }
 
-void SRigHierarchy::CreateRefreshMenu(FMenuBuilder& MenuBuilder)
-{
-	MenuBuilder.AddWidget(
-		SNew(SVerticalBox)
-
-		+ SVerticalBox::Slot()
-		.AutoHeight()
-		.Padding(3)
-		[
-			SNew(STextBlock)
-			.Font(FEditorStyle::GetFontStyle("ControlRig.Hierarchy.Menu"))
-			.Text(LOCTEXT("RefreshMesh_Title", "Select Mesh"))
-			.ToolTipText(LOCTEXT("RefreshMesh_Tooltip", "Select Mesh to refresh transform from... It will refresh init transform from selected mesh. This doesn't change hierarchy. \
-				If you want to reimport hierarchy, please delete all nodes, and use import hierarchy."))
-		]
-		+ SVerticalBox::Slot()
-		.AutoHeight()
-		.Padding(3)
-		[
-			SNew(SObjectPropertyEntryBox)
-			.AllowedClass(USkeletalMesh::StaticClass())
-			.OnObjectChanged(this, &SRigHierarchy::RefreshHierarchy)
-		]
-		,
-		FText()
-	);
-}
-
-void SRigHierarchy::RefreshHierarchy(const FAssetData& InAssetData)
-{
-	FRigHierarchy* Hier = GetHierarchy();
-	USkeletalMesh* Mesh = Cast<USkeletalMesh>(InAssetData.GetAsset());
-	if (Mesh && Hier)
-	{
-		FScopedTransaction Transaction(LOCTEXT("HierarchyRefresh", "Refresh Transform"));
-		ControlRigBlueprint->Modify();
-
-		const FReferenceSkeleton& RefSkeleton = Mesh->RefSkeleton;
-		const TArray<FMeshBoneInfo>& BoneInfos = RefSkeleton.GetRefBoneInfo();
-		const TArray<FTransform>& BonePoses = RefSkeleton.GetRefBonePose();
-
-		for (int32 BoneIndex = 0; BoneIndex < RefSkeleton.GetNum(); ++BoneIndex)
-		{
-			// only add if you don't have it. This may change in the future
-			int32 RigIndex = Hier->GetIndex(BoneInfos[BoneIndex].Name);
-			if (RigIndex != INDEX_NONE)
-			{
-				// @todo: add optimized version without sorting, but if no sort, we should make sure not to use find index function
-				Hier->SetInitialTransform(RigIndex, FAnimationRuntime::GetComponentSpaceTransform(RefSkeleton, BonePoses, BoneIndex));
-			}
-		}
-
-		ControlRigEditor.Pin()->OnHierarchyChanged();
-		RefreshTreeView();
-		FSlateApplication::Get().DismissAllMenus();
-	}
-}
 void SRigHierarchy::CreateImportMenu(FMenuBuilder& MenuBuilder)
 {
 	MenuBuilder.AddWidget(
@@ -510,7 +447,7 @@ void SRigHierarchy::CreateImportMenu(FMenuBuilder& MenuBuilder)
 		.Padding(3)
 		[
 			SNew(SObjectPropertyEntryBox)
-			.AllowedClass(USkeletalMesh::StaticClass())
+			.OnShouldFilterAsset(this, &SRigHierarchy::ShouldFilterOnImport)
 			.OnObjectChanged(this, &SRigHierarchy::ImportHierarchy)
 		]
 		,
@@ -518,33 +455,65 @@ void SRigHierarchy::CreateImportMenu(FMenuBuilder& MenuBuilder)
 	);
 }
 
+bool SRigHierarchy::ShouldFilterOnImport(const FAssetData& AssetData) const
+{
+	return (AssetData.AssetClass != USkeletalMesh::StaticClass()->GetFName() &&
+		AssetData.AssetClass != USkeleton::StaticClass()->GetFName());
+}
+
 void SRigHierarchy::ImportHierarchy(const FAssetData& InAssetData)
 {
 	FRigHierarchy* Hier = GetHierarchy();
-	USkeletalMesh* Mesh = Cast<USkeletalMesh> (InAssetData.GetAsset());
-	if (Mesh && Hier)
+	if (Hier)
 	{
-		FScopedTransaction Transaction(LOCTEXT("HierarchyImport", "Import Hierarchy"));
-		ControlRigBlueprint->Modify();
-
-		const FReferenceSkeleton& RefSkeleton = Mesh->RefSkeleton;
-		const TArray<FMeshBoneInfo>& BoneInfos = RefSkeleton.GetRefBoneInfo();
-		const TArray<FTransform>& BonePoses = RefSkeleton.GetRefBonePose();
+		const FReferenceSkeleton* RefSkeleton = nullptr;
+		if (USkeletalMesh* Mesh = Cast<USkeletalMesh>(InAssetData.GetAsset()))
+		{
+			RefSkeleton = &Mesh->RefSkeleton;
+			ControlRigBlueprint->SourceHierarchyImport = Mesh;
+		}
+		else if (USkeleton* Skeleton = Cast<USkeleton>(InAssetData.GetAsset()))
+		{
+			RefSkeleton = &Skeleton->GetReferenceSkeleton();
+			ControlRigBlueprint->SourceHierarchyImport = Skeleton;
+		}
 
-		for (int32 BoneIndex = 0; BoneIndex < RefSkeleton.GetNum(); ++BoneIndex)
+		if (RefSkeleton)
 		{
-			// only add if you don't have it. This may change in the future
-			if (Hier->GetIndex(BoneInfos[BoneIndex].Name) == INDEX_NONE)
+			FScopedTransaction Transaction(LOCTEXT("HierarchyImport", "Import Hierarchy"));
+			ControlRigBlueprint->Modify();
+
+			const TArray<FMeshBoneInfo>& BoneInfos = RefSkeleton->GetRefBoneInfo();
+			const TArray<FTransform>& BonePoses = RefSkeleton->GetRefBonePose();
+
+			for (int32 BoneIndex = 0; BoneIndex < RefSkeleton->GetNum(); ++BoneIndex)
 			{
+				// only add if you don't have it. This may change in the future
+				int32 HierIndex = Hier->GetIndex(BoneInfos[BoneIndex].Name);
+				FTransform InitialTransform = FAnimationRuntime::GetComponentSpaceTransform(*RefSkeleton, BonePoses, BoneIndex);
 				// @todo: add optimized version without sorting, but if no sort, we should make sure not to use find index function
 				FName ParentName = (BoneInfos[BoneIndex].ParentIndex != INDEX_NONE) ? BoneInfos[BoneInfos[BoneIndex].ParentIndex].Name : NAME_None;
-				Hier->AddBone(BoneInfos[BoneIndex].Name, ParentName, FAnimationRuntime::GetComponentSpaceTransform(RefSkeleton, BonePoses, BoneIndex));
+				// if exists, see if we should change parents
+				if (HierIndex != INDEX_NONE)
+				{
+					if (ParentName != Hier->GetParentName(BoneInfos[BoneIndex].Name))
+					{
+						// reparent
+						Hier->Reparent(BoneInfos[BoneIndex].Name, ParentName);
+					}
+
+					Hier->SetInitialTransform(BoneInfos[BoneIndex].Name, InitialTransform);
+				}
+				else
+				{
+					Hier->AddBone(BoneInfos[BoneIndex].Name, ParentName, InitialTransform);
+				}
 			}
-		}
 
-		ControlRigEditor.Pin()->OnHierarchyChanged();
-		RefreshTreeView();
-		FSlateApplication::Get().DismissAllMenus();
+			ControlRigEditor.Pin()->OnHierarchyChanged();
+			RefreshTreeView();
+			FSlateApplication::Get().DismissAllMenus();
+		}
 	}
 }
 
diff --git a/Engine/Plugins/Experimental/ControlRig/Source/ControlRigEditor/Private/Editor/SRigHierarchy.h b/Engine/Plugins/Experimental/ControlRig/Source/ControlRigEditor/Private/Editor/SRigHierarchy.h
--- a/Engine/Plugins/Experimental/ControlRig/Source/ControlRigEditor/Private/Editor/SRigHierarchy.h
+++ b/Engine/Plugins/Experimental/ControlRig/Source/ControlRigEditor/Private/Editor/SRigHierarchy.h
@@ -181,8 +181,7 @@ private:
 
 	void ImportHierarchy(const FAssetData& InAssetData);
 	void CreateImportMenu(FMenuBuilder& MenuBuilder);
-	void CreateRefreshMenu(FMenuBuilder& MenuBuilder);
-	void RefreshHierarchy(const FAssetData& InAssetData);
+	bool ShouldFilterOnImport(const FAssetData& AssetData) const;
 
 	FName CreateUniqueName(const FName& InBaseName) const;
 
@@ -190,7 +189,6 @@ private:
 
 	void ClearDetailPanel() const;
 	void SelectBone(const FName& BoneName) const;
-
 public:
 	bool RenameBone(const FName& OldName, const FName& NewName);
 	bool OnVerifyNameChanged(const FName& OldName, const FName& NewName, FText& OutErrorMessage);
diff --git a/Engine/Plugins/Experimental/ControlRig/Source/ControlRigEditor/Private/Graph/ControlRigGraphPanelPinFactory.cpp b/Engine/Plugins/Experimental/ControlRig/Source/ControlRigEditor/Private/Graph/ControlRigGraphPanelPinFactory.cpp
--- a/Engine/Plugins/Experimental/ControlRig/Source/ControlRigEditor/Private/Graph/ControlRigGraphPanelPinFactory.cpp
+++ b/Engine/Plugins/Experimental/ControlRig/Source/ControlRigEditor/Private/Graph/ControlRigGraphPanelPinFactory.cpp
@@ -4,6 +4,7 @@
 #include "Graph/ControlRigGraphSchema.h"
 #include "Graph/ControlRigGraphNode.h"
 #include "Graph/SGraphPinBoneName.h"
+#include "Graph/SGraphPinCurveName.h"
 #include "Graph/SGraphPinCurveFloat.h"
 #include "KismetPins/SGraphPinExec.h"
 #include "ControlRig.h"
@@ -42,6 +43,10 @@ TSharedPtr<SGraphPin> FControlRigGraphPanelPinFactory::CreatePin(UEdGraphPin* In
 							{
 								return SNew(SGraphPinBoneName, InPin);
 							}
+							else if (Property->HasMetaData(UControlRig::CurveNameMetaName))
+							{
+								return SNew(SGraphPinCurveName, InPin);
+							}
 						}
 					}
 				}
diff --git a/Engine/Plugins/Experimental/ControlRig/Source/ControlRigEditor/Private/Graph/NodeSpawners/ControlRigCommentNodeSpawner.cpp b/Engine/Plugins/Experimental/ControlRig/Source/ControlRigEditor/Private/Graph/NodeSpawners/ControlRigCommentNodeSpawner.cpp
--- a/Engine/Plugins/Experimental/ControlRig/Source/ControlRigEditor/Private/Graph/NodeSpawners/ControlRigCommentNodeSpawner.cpp
+++ b/Engine/Plugins/Experimental/ControlRig/Source/ControlRigEditor/Private/Graph/NodeSpawners/ControlRigCommentNodeSpawner.cpp
@@ -11,6 +11,7 @@
 #include "ControlRigBlueprintUtils.h"
 #include "ControlRig.h"
 #include "ControlRigBlueprint.h"
+#include "Graph/ControlRigGraph.h"
 
 #if WITH_EDITOR
 #include "Editor.h"
@@ -25,7 +26,7 @@ UControlRigCommentNodeSpawner* UControlRigCommentNodeSpawner::Create()
 
 	FBlueprintActionUiSpec& MenuSignature = NodeSpawner->DefaultMenuSignature;
 	
-	MenuSignature.MenuName = LOCTEXT("Comment", "Comment");
+	MenuSignature.MenuName = LOCTEXT("Add Comment", "Add Comment for Selection...");
 	MenuSignature.Tooltip  = LOCTEXT("CommentTooltip", "Adds a comment box to the graph");
 	MenuSignature.Category = FText();
 
@@ -40,7 +41,7 @@ UControlRigCommentNodeSpawner* UControlRigCommentNodeSpawner::Create()
 	}
 
 	// @TODO: should use details customization-like extensibility system to provide editor only data like this
-	MenuSignature.Icon = FSlateIcon(TEXT("ControlRigEditorStyle"), TEXT("ControlRig.RigUnit"));
+	MenuSignature.Icon = FSlateIcon("EditorStyle", "GraphEditor.Comment_16x");
 
 	return NodeSpawner;
 }
@@ -84,7 +85,33 @@ UEdGraphNode* UControlRigCommentNodeSpawner::Invoke(UEdGraph* ParentGraph, FBind
 
 		if (!bIsTemplateNode)
 		{
-			if (RigBlueprint->ModelController->AddComment(TEXT("Comment"), TEXT("Comment"), Location, FVector2D(100.f, 100.f), FLinearColor::White, true))
+			FVector2D AverageLocation = Location;
+			FVector2D Size(400.f, 250.f);
+
+			TArray<FControlRigModelNode> SelectedNodes = RigBlueprint->Model->SelectedNodes();
+			if (SelectedNodes.Num() > 0)
+			{
+				FVector2D TopLeft(1000000.f, 1000000.f);
+				FVector2D BottomRight(-1000000.f, -1000000.f);
+				const float Margin = 20;
+
+				UControlRigGraph* RigGraph = Cast<UControlRigGraph>(ParentGraph);
+				for (const FControlRigModelNode& SelectedNode : SelectedNodes)
+				{
+					UEdGraphNode* EdNode = RigGraph->FindNodeFromPropertyName(SelectedNode.Name);
+					if (EdNode)
+					{
+						TopLeft.X = FMath::Min<float>(TopLeft.X, float(EdNode->NodePosX) - Margin);
+						TopLeft.Y = FMath::Min<float>(TopLeft.Y, float(EdNode->NodePosY) - Margin - 35.f);
+						BottomRight.X = FMath::Max<float>(BottomRight.X, float(EdNode->NodePosX) + float(EdNode->NodeWidth) + Margin);
+						BottomRight.Y = FMath::Max<float>(BottomRight.Y, float(EdNode->NodePosY) + float(EdNode->NodeHeight) + Margin);
+					}
+				}
+				AverageLocation = TopLeft;
+				Size = BottomRight - TopLeft;
+			}
+
+			if (RigBlueprint->ModelController->AddComment(TEXT("Comment"), TEXT("Comment"), AverageLocation, Size, FLinearColor::White, true))
 			{
 				MemberName = RigBlueprint->LastNameFromNotification;
 				for (UEdGraphNode* Node : ParentGraph->Nodes)
diff --git a/Engine/Plugins/Experimental/ControlRig/Source/ControlRigEditor/Private/Graph/SControlRigGraphNode.cpp b/Engine/Plugins/Experimental/ControlRig/Source/ControlRigEditor/Private/Graph/SControlRigGraphNode.cpp
--- a/Engine/Plugins/Experimental/ControlRig/Source/ControlRigEditor/Private/Graph/SControlRigGraphNode.cpp
+++ b/Engine/Plugins/Experimental/ControlRig/Source/ControlRigEditor/Private/Graph/SControlRigGraphNode.cpp
@@ -774,4 +774,15 @@ void SControlRigGraphNode::GetNodeInfoPopups(FNodeInfoContext* Context, TArray<F
 	}
 }
 
+void SControlRigGraphNode::Tick(const FGeometry& AllottedGeometry, const double InCurrentTime, const float InDeltaTime)
+{
+	SGraphNode::Tick(AllottedGeometry, InCurrentTime, InDeltaTime);
+
+	if (GraphNode)
+	{
+		GraphNode->NodeWidth = (int32)AllottedGeometry.Size.X;
+		GraphNode->NodeHeight = (int32)AllottedGeometry.Size.Y;
+	}
+}
+
 #undef LOCTEXT_NAMESPACE
\ No newline at end of file
diff --git a/Engine/Plugins/Experimental/ControlRig/Source/ControlRigEditor/Private/Graph/SControlRigGraphNode.h b/Engine/Plugins/Experimental/ControlRig/Source/ControlRigEditor/Private/Graph/SControlRigGraphNode.h
--- a/Engine/Plugins/Experimental/ControlRig/Source/ControlRigEditor/Private/Graph/SControlRigGraphNode.h
+++ b/Engine/Plugins/Experimental/ControlRig/Source/ControlRigEditor/Private/Graph/SControlRigGraphNode.h
@@ -42,6 +42,9 @@ public:
 	virtual TSharedRef<SWidget> CreateNodeContentArea() override;
 	virtual TSharedPtr<SGraphPin> GetHoveredPin( const FGeometry& MyGeometry, const FPointerEvent& MouseEvent ) const override;
 	virtual void GetNodeInfoPopups(FNodeInfoContext* Context, TArray<FGraphInformationPopupInfo>& Popups) const override;
+
+	virtual void Tick(const FGeometry& AllottedGeometry, const double InCurrentTime, const float InDeltaTime) override;
+
 private:
 	bool ParentUseLowDetailNodeTitles() const
 	{
diff --git a/Engine/Plugins/Experimental/ControlRig/Source/ControlRigEditor/Private/Graph/SGraphPinCurveName.cpp b/Engine/Plugins/Experimental/ControlRig/Source/ControlRigEditor/Private/Graph/SGraphPinCurveName.cpp
new file mode 100644
--- /dev/null
+++ b/Engine/Plugins/Experimental/ControlRig/Source/ControlRigEditor/Private/Graph/SGraphPinCurveName.cpp
@@ -0,0 +1,91 @@
+// Copyright 1998-2019 Epic Games, Inc. All Rights Reserved.
+
+
+#include "Graph/SGraphPinCurveName.h"
+#include "Graph/SGraphPinBoneNameValueWidget.h"
+#include "Graph/ControlRigGraph.h"
+#include "Widgets/Layout/SBox.h"
+#include "Widgets/Input/SEditableTextBox.h"
+#include "EdGraphSchema_K2.h"
+#include "ScopedTransaction.h"
+#include "DetailLayoutBuilder.h"
+
+void SGraphPinCurveName::Construct(const FArguments& InArgs, UEdGraphPin* InGraphPinObj)
+{
+	SGraphPin::Construct(SGraphPin::FArguments(), InGraphPinObj);
+}
+
+TSharedRef<SWidget>	SGraphPinCurveName::GetDefaultValueWidget()
+{
+	UControlRigGraph* RigGraph = Cast<UControlRigGraph>(GraphPinObj->GetOwningNode()->GetGraph());
+
+	TSharedPtr<FString> InitialSelected;
+	for (TSharedPtr<FString> Item : RigGraph->GetCurveNameList())
+	{
+		if (Item->Equals(GetCurveNameText().ToString()))
+		{
+			InitialSelected = Item;
+		}
+	}
+
+	return SNew(SBox)
+		.MinDesiredWidth(150)
+		.MaxDesiredWidth(400)
+		[
+			SAssignNew(CurveNameComboBox, SGraphPinBoneNameValueWidget)
+				.Visibility(this, &SGraphPin::GetDefaultValueVisibility)
+				.OptionsSource(&RigGraph->GetCurveNameList())
+				.OnGenerateWidget(this, &SGraphPinCurveName::MakeCurveNameItemWidget)
+				.OnSelectionChanged(this, &SGraphPinCurveName::OnCurveNameChanged)
+				.OnComboBoxOpening(this, &SGraphPinCurveName::OnCurveNameComboBox)
+				.InitiallySelectedItem(InitialSelected)
+				.Content()
+				[
+					SNew(STextBlock)
+					.Text(this, &SGraphPinCurveName::GetCurveNameText)
+				]
+		];
+}
+
+FText SGraphPinCurveName::GetCurveNameText() const
+{
+	return FText::FromString( GraphPinObj->GetDefaultAsString() );
+}
+
+void SGraphPinCurveName::SetCurveNameText(const FText& NewTypeInValue, ETextCommit::Type /*CommitInfo*/)
+{
+	if(!GraphPinObj->GetDefaultAsString().Equals(NewTypeInValue.ToString()))
+	{
+		const FScopedTransaction Transaction( NSLOCTEXT("GraphEditor", "ChangeCurveNamePinValue", "Change Bone Name Pin Value" ) );
+		GraphPinObj->Modify();
+		GraphPinObj->GetSchema()->TrySetDefaultValue(*GraphPinObj, NewTypeInValue.ToString());
+	}
+}
+
+TSharedRef<SWidget> SGraphPinCurveName::MakeCurveNameItemWidget(TSharedPtr<FString> InItem)
+{
+	return 	SNew(STextBlock).Text(FText::FromString(*InItem));// .Font(FEditorStyle::GetFontStyle(TEXT("PropertyWindow.NormalFont")));
+}
+
+void SGraphPinCurveName::OnCurveNameChanged(TSharedPtr<FString> NewSelection, ESelectInfo::Type SelectInfo)
+{
+	if (SelectInfo != ESelectInfo::Direct)
+	{
+		FString NewValue = *NewSelection.Get();
+		SetCurveNameText(FText::FromString(NewValue), ETextCommit::OnEnter);
+	}
+}
+
+void SGraphPinCurveName::OnCurveNameComboBox()
+{
+	UControlRigGraph* RigGraph = Cast<UControlRigGraph>(GraphPinObj->GetOwningNode()->GetGraph());
+	TSharedPtr<FString> CurrentlySelected;
+	for (TSharedPtr<FString> Item : RigGraph->GetCurveNameList())
+	{
+		if (Item->Equals(GetCurveNameText().ToString()))
+		{
+			CurrentlySelected = Item;
+		}
+	}
+	CurveNameComboBox->SetSelectedItem(CurrentlySelected);
+}
diff --git a/Engine/Plugins/Experimental/ControlRig/Source/ControlRigEditor/Private/Graph/SGraphPinCurveName.h b/Engine/Plugins/Experimental/ControlRig/Source/ControlRigEditor/Private/Graph/SGraphPinCurveName.h
new file mode 100644
--- /dev/null
+++ b/Engine/Plugins/Experimental/ControlRig/Source/ControlRigEditor/Private/Graph/SGraphPinCurveName.h
@@ -0,0 +1,33 @@
+// Copyright 1998-2019 Epic Games, Inc. All Rights Reserved.
+
+#pragma once
+
+#include "CoreMinimal.h"
+#include "Widgets/DeclarativeSyntaxSupport.h"
+#include "Widgets/SWidget.h"
+#include "SGraphPin.h"
+#include "SGraphPinBoneNameValueWidget.h"
+
+class SGraphPinCurveName : public SGraphPin
+{
+public:
+	SLATE_BEGIN_ARGS(SGraphPinCurveName) {}
+	SLATE_END_ARGS()
+
+		void Construct(const FArguments& InArgs, UEdGraphPin* InGraphPinObj);
+
+protected:
+	//~ Begin SGraphPin Interface
+	virtual TSharedRef<SWidget>	GetDefaultValueWidget() override;
+	//~ End SGraphPin Interface
+
+	FText GetCurveNameText() const;
+	virtual void SetCurveNameText(const FText& NewTypeInValue, ETextCommit::Type CommitInfo);
+
+	TSharedRef<SWidget> MakeCurveNameItemWidget(TSharedPtr<FString> InItem);
+	void OnCurveNameChanged(TSharedPtr<FString> NewSelection, ESelectInfo::Type SelectInfo);
+	void OnCurveNameComboBox();
+
+	TSharedPtr<SGraphPinBoneNameValueWidget> CurveNameComboBox;
+
+};
diff --git a/Engine/Source/Editor/BlueprintGraph/Classes/EdGraphSchema_K2.h b/Engine/Source/Editor/BlueprintGraph/Classes/EdGraphSchema_K2.h
--- a/Engine/Source/Editor/BlueprintGraph/Classes/EdGraphSchema_K2.h
+++ b/Engine/Source/Editor/BlueprintGraph/Classes/EdGraphSchema_K2.h
@@ -460,9 +460,9 @@ public:
 	virtual bool DoesSupportPinWatching() const	override;
 	virtual bool IsPinBeingWatched(UEdGraphPin const* Pin) const override;
 	virtual void ClearPinWatch(UEdGraphPin const* Pin) const override;
-	virtual void TrySetDefaultValue(UEdGraphPin& Pin, const FString& NewDefaultValue) const override;
-	virtual void TrySetDefaultObject(UEdGraphPin& Pin, UObject* NewDefaultObject) const override;
-	virtual void TrySetDefaultText(UEdGraphPin& InPin, const FText& InNewDefaultText) const override;
+	virtual void TrySetDefaultValue(UEdGraphPin& Pin, const FString& NewDefaultValue, bool bMarkAsModified = true) const override;
+	virtual void TrySetDefaultObject(UEdGraphPin& Pin, UObject* NewDefaultObject, bool bMarkAsModified = true) const override;
+	virtual void TrySetDefaultText(UEdGraphPin& InPin, const FText& InNewDefaultText, bool bMarkAsModified = true) const override;
 	virtual bool DoesDefaultValueMatchAutogenerated(const UEdGraphPin& InPin) const override;
 	virtual void ResetPinToAutogeneratedDefaultValue(UEdGraphPin* Pin, bool bCallModifyCallbacks = true) const override;
 	virtual bool ShouldHidePinDefaultValue(UEdGraphPin* Pin) const override;
diff --git a/Engine/Source/Editor/BlueprintGraph/Private/BlueprintActionDatabase.cpp b/Engine/Source/Editor/BlueprintGraph/Private/BlueprintActionDatabase.cpp
--- a/Engine/Source/Editor/BlueprintGraph/Private/BlueprintActionDatabase.cpp
+++ b/Engine/Source/Editor/BlueprintGraph/Private/BlueprintActionDatabase.cpp
@@ -836,13 +836,15 @@ static void BlueprintActionDatabaseImpl::GetNodeSpecificActions(TSubclassOf<UEdG
 		check(NodeCDO != nullptr);
 		NodeCDO->GetMenuActions(Registrar);
 	}
+
+
 	// unfortunately, UEdGraphNode_Comment is not a UK2Node and therefore cannot
 	// leverage UK2Node's GetMenuActions() function, so here we HACK it in
 	//
 	// @TODO: DO NOT follow this example! Do as I say, not as I do! If we need
 	//        to support other nodes in a similar way, then we should come up
 	//        with a better (more generalized) solution.
-	else if (NodeClass == UEdGraphNode_Comment::StaticClass())
+	if (NodeClass == UEdGraphNode_Comment::StaticClass())
 	{
 		Registrar.AddBlueprintAction(MakeCommentNodeSpawner());
 	}
@@ -1033,7 +1035,6 @@ static bool BlueprintActionDatabaseImpl::IsObjectValidForDatabase(UObject const*
  ******************************************************************************/
 
 static FBlueprintActionDatabase* DatabaseInst = nullptr;
-
 //------------------------------------------------------------------------------
 FBlueprintActionDatabase& FBlueprintActionDatabase::Get()
 {
diff --git a/Engine/Source/Editor/BlueprintGraph/Private/EdGraphSchema_K2.cpp b/Engine/Source/Editor/BlueprintGraph/Private/EdGraphSchema_K2.cpp
--- a/Engine/Source/Editor/BlueprintGraph/Private/EdGraphSchema_K2.cpp
+++ b/Engine/Source/Editor/BlueprintGraph/Private/EdGraphSchema_K2.cpp
@@ -4703,7 +4703,7 @@ void UEdGraphSchema_K2::GetPinDefaultValuesFromString(const FEdGraphPinType& Pin
 	}
 }
 
-void UEdGraphSchema_K2::TrySetDefaultValue(UEdGraphPin& Pin, const FString& NewDefaultValue) const
+void UEdGraphSchema_K2::TrySetDefaultValue(UEdGraphPin& Pin, const FString& NewDefaultValue, bool bMarkAsModified) const
 {
 	FString UseDefaultValue;
 	UObject* UseDefaultObject = nullptr;
@@ -4727,12 +4727,15 @@ void UEdGraphSchema_K2::TrySetDefaultValue(UEdGraphPin& Pin, const FString& NewD
 			Node->PinConnectionListChanged(&Pin);
 		}
 
-		UBlueprint* Blueprint = FBlueprintEditorUtils::FindBlueprintForNodeChecked(Node);
-		FBlueprintEditorUtils::MarkBlueprintAsModified(Blueprint);
+		if (bMarkAsModified)
+		{
+			UBlueprint* Blueprint = FBlueprintEditorUtils::FindBlueprintForNodeChecked(Node);
+			FBlueprintEditorUtils::MarkBlueprintAsModified(Blueprint);
+		}
 	}
 }
 
-void UEdGraphSchema_K2::TrySetDefaultObject(UEdGraphPin& Pin, UObject* NewDefaultObject) const
+void UEdGraphSchema_K2::TrySetDefaultObject(UEdGraphPin& Pin, UObject* NewDefaultObject, bool bMarkAsModified) const
 {
 	FText UseDefaultText;
 
@@ -4759,11 +4762,14 @@ void UEdGraphSchema_K2::TrySetDefaultObject(UEdGraphPin& Pin, UObject* NewDefaul
 		Node->PinConnectionListChanged(&Pin);
 	}
 
-	UBlueprint* Blueprint = FBlueprintEditorUtils::FindBlueprintForNodeChecked(Node);
-	FBlueprintEditorUtils::MarkBlueprintAsModified(Blueprint);
+	if (bMarkAsModified)
+	{
+		UBlueprint* Blueprint = FBlueprintEditorUtils::FindBlueprintForNodeChecked(Node);
+		FBlueprintEditorUtils::MarkBlueprintAsModified(Blueprint);
+	}
 }
 
-void UEdGraphSchema_K2::TrySetDefaultText(UEdGraphPin& InPin, const FText& InNewDefaultText) const
+void UEdGraphSchema_K2::TrySetDefaultText(UEdGraphPin& InPin, const FText& InNewDefaultText, bool bMarkAsModified) const
 {
 	// No reason to set the FText if it is not a PC_Text.
 	if(InPin.PinType.PinCategory == PC_Text)
@@ -4785,8 +4791,11 @@ void UEdGraphSchema_K2::TrySetDefaultText(UEdGraphPin& InPin, const FText& InNew
 			Node->PinConnectionListChanged(&InPin);
 		}
 
-		UBlueprint* Blueprint = FBlueprintEditorUtils::FindBlueprintForNodeChecked(Node);
-		FBlueprintEditorUtils::MarkBlueprintAsModified(Blueprint);
+		if (bMarkAsModified)
+		{
+			UBlueprint* Blueprint = FBlueprintEditorUtils::FindBlueprintForNodeChecked(Node);
+			FBlueprintEditorUtils::MarkBlueprintAsModified(Blueprint);
+		}
 	}
 }
 
diff --git a/Engine/Source/Editor/Kismet/Private/BlueprintCompilationManager.cpp b/Engine/Source/Editor/Kismet/Private/BlueprintCompilationManager.cpp
--- a/Engine/Source/Editor/Kismet/Private/BlueprintCompilationManager.cpp
+++ b/Engine/Source/Editor/Kismet/Private/BlueprintCompilationManager.cpp
@@ -33,6 +33,7 @@
 #include "UObject/UObjectHash.h"
 #include "Kismet2/KismetDebugUtilities.h"
 #include "BlueprintEditorModule.h"
+#include "Stats/StatsHierarchical.h"
 
 #define LOCTEXT_NAMESPACE "BlueprintCompilationManager"
 
@@ -197,6 +198,8 @@ void FBlueprintCompilationManagerImpl::QueueForCompilation(const FBPCompileReque
 
 void FBlueprintCompilationManagerImpl::CompileSynchronouslyImpl(const FBPCompileRequest& Request)
 {
+	DECLARE_SCOPE_HIERARCHICAL_COUNTER_FUNC()
+
 	Request.BPToCompile->bQueuedForCompilation = true;
 
 	const bool bIsRegeneratingOnLoad		= (Request.CompileOptions & EBlueprintCompileOptions::IsRegeneratingOnLoad		) != EBlueprintCompileOptions::None;
@@ -254,6 +257,8 @@ void FBlueprintCompilationManagerImpl::CompileSynchronouslyImpl(const FBPCompile
 
 	if ( GEditor && !bRegenerateSkeletonOnly)
 	{
+		DECLARE_SCOPE_HIERARCHICAL_COUNTER(BroadcastBlueprintReinstanced)
+
 		// Make sure clients know they're being reinstanced as part of blueprint compilation. After this point
 		// compilation is completely done:
 		TGuardValue<bool> GuardTemplateNameFlag(GCompilingBlueprint, true);
@@ -264,11 +269,15 @@ void FBlueprintCompilationManagerImpl::CompileSynchronouslyImpl(const FBPCompile
 
 	if(!bSkipGarbageCollection)
 	{
+		DECLARE_SCOPE_HIERARCHICAL_COUNTER(CollectGarbage)
+
 		CollectGarbage(GARBAGE_COLLECTION_KEEPFLAGS);
 	}
 
 	if (!bRegenerateSkeletonOnly)
 	{
+		DECLARE_SCOPE_HIERARCHICAL_COUNTER(BroadcastChanged)
+
 		for(UBlueprint* BP : SkeletonCompiledBlueprints)
 		{
 			BP->BroadcastChanged();
@@ -281,6 +290,8 @@ void FBlueprintCompilationManagerImpl::CompileSynchronouslyImpl(const FBPCompile
 
 	if (!bBatchCompile && !bRegenerateSkeletonOnly)
 	{
+		DECLARE_SCOPE_HIERARCHICAL_COUNTER(BroadCastCompiled)
+
 		for(UBlueprint* BP : SkeletonCompiledBlueprints)
 		{
 			BP->BroadcastCompiled();
@@ -434,6 +445,8 @@ FReinstancingJob::FReinstancingJob(TPair<UClass*, UClass*> InOldToNew)
 
 void FBlueprintCompilationManagerImpl::FlushCompilationQueueImpl(bool bSuppressBroadcastCompiled, TArray<UBlueprint*>* BlueprintsCompiled, TArray<UBlueprint*>* BlueprintsCompiledOrSkeletonCompiled, FUObjectSerializeContext* InLoadContext)
 {
+	DECLARE_SCOPE_HIERARCHICAL_COUNTER_FUNC()
+
 	TGuardValue<bool> GuardTemplateNameFlag(GCompilingBlueprint, true);
 	ensure(bGeneratedClassLayoutReady);
 
@@ -728,6 +741,8 @@ void FBlueprintCompilationManagerImpl::FlushCompilationQueueImpl(bool bSuppressB
 		
 				if(CompilerData.ShouldRegenerateSkeleton())
 				{
+					DECLARE_SCOPE_HIERARCHICAL_COUNTER(RecompileSkeleton)
+
 					if(BlueprintsCompiledOrSkeletonCompiled)
 					{
 						BlueprintsCompiledOrSkeletonCompiled->Add(BP);
@@ -763,6 +778,8 @@ void FBlueprintCompilationManagerImpl::FlushCompilationQueueImpl(bool bSuppressB
 				}
 				else
 				{
+					DECLARE_SCOPE_HIERARCHICAL_COUNTER(RelinkSkeleton)
+
 					// Just relink, note that UProperties that reference *other* types may be stale until
 					// we fixup below:
 					UClass* SkeletonToRelink = BP->SkeletonGeneratedClass;
@@ -787,6 +804,7 @@ void FBlueprintCompilationManagerImpl::FlushCompilationQueueImpl(bool bSuppressB
 					BP->bHasBeenRegenerated = true;
 					if (BP->GeneratedClass)
 					{
+						DECLARE_SCOPE_HIERARCHICAL_COUNTER(ClearFunctionMapsCaches)
 						BP->GeneratedClass->ClearFunctionMapsCaches();
 					}
 				}
@@ -796,6 +814,8 @@ void FBlueprintCompilationManagerImpl::FlushCompilationQueueImpl(bool bSuppressB
 			// that may have been created as part of skeleton generation:
 			for (FCompilerData& CompilerData : CurrentlyCompilingBPs)
 			{
+				DECLARE_SCOPE_HIERARCHICAL_COUNTER(FixUpDelegateParameters)
+
 				UBlueprint* BP = CompilerData.BP;
 				TArray< FSkeletonFixupData >& ParamsToFix = CompilerData.SkeletonFixupData;
 				for( const FSkeletonFixupData& FixupData : ParamsToFix )
@@ -817,6 +837,8 @@ void FBlueprintCompilationManagerImpl::FlushCompilationQueueImpl(bool bSuppressB
 			{
 				if(Data.ShouldSkipIfDependenciesAreUnchanged())
 				{
+					DECLARE_SCOPE_HIERARCHICAL_COUNTER(SkipIfDependenciesAreUnchanged)
+
 					// if our parent is still being compiled, then we still need to be compiled:
 					UClass* Iter = Data.BP->ParentClass;
 					while(Iter)
@@ -877,6 +899,8 @@ void FBlueprintCompilationManagerImpl::FlushCompilationQueueImpl(bool bSuppressB
 				continue;
 			}
 
+			DECLARE_SCOPE_HIERARCHICAL_COUNTER(ReconstructNodes)
+
 			UBlueprint* BP = CompilerData.BP;
 
 			// Some nodes are set up to do things during reconstruction only when this flag is NOT set.
@@ -951,6 +975,8 @@ void FBlueprintCompilationManagerImpl::FlushCompilationQueueImpl(bool bSuppressB
 					continue;
 				}
 
+				DECLARE_SCOPE_HIERARCHICAL_COUNTER(ReinstanceQueuedBlueprint)
+
 				// no need to reinstance skeleton or relink jobs that are not in a hierarchy that has had reinstancing initiated:
 				bool bRequiresReinstance = CompilerData.ShouldInitiateReinstancing();
 				if (!bRequiresReinstance)
@@ -1015,6 +1041,8 @@ void FBlueprintCompilationManagerImpl::FlushCompilationQueueImpl(bool bSuppressB
 			UBlueprint* BP = CompilerData.BP;
 			if(CompilerData.ShouldCompileClassLayout())
 			{
+				DECLARE_SCOPE_HIERARCHICAL_COUNTER(CompileClassLayout)
+
 				ensure( BP->GeneratedClass == nullptr ||
 						BP->GeneratedClass->ClassDefaultObject == nullptr || 
 						BP->GeneratedClass->ClassDefaultObject->GetClass() != BP->GeneratedClass);
@@ -1083,6 +1111,8 @@ void FBlueprintCompilationManagerImpl::FlushCompilationQueueImpl(bool bSuppressB
 			}
 			else
 			{
+				DECLARE_SCOPE_HIERARCHICAL_COUNTER(CompileClassFunctions)
+
 				// default value propagation occurs below:
 				if(BPGC)
 				{
@@ -1159,6 +1189,8 @@ void FBlueprintCompilationManagerImpl::FlushCompilationQueueImpl(bool bSuppressB
 	// STAGE XIV: Now we can finish the first stage of the reinstancing operation, moving old classes to new classes:
 	{
 		{
+			DECLARE_SCOPE_HIERARCHICAL_COUNTER(MoveOldClassesToNewClasses)
+
 			TArray<FReinstancingJob> Reinstancers;
 			// Set up reinstancing jobs - we need a reference to the compiler in order to honor 
 			// CopyTermDefaultsToDefaultObject
@@ -1181,6 +1213,8 @@ void FBlueprintCompilationManagerImpl::FlushCompilationQueueImpl(bool bSuppressB
 		// STAGE XV: POST CDO COMPILED
 		for (FCompilerData& CompilerData : CurrentlyCompilingBPs)
 		{
+			DECLARE_SCOPE_HIERARCHICAL_COUNTER(PostCDOCompiled)
+
 			if(!CompilerData.IsSkeletonOnly() && CompilerData.Compiler.IsValid())
 			{
 				CompilerData.Compiler->PostCDOCompiled();
@@ -1190,6 +1224,8 @@ void FBlueprintCompilationManagerImpl::FlushCompilationQueueImpl(bool bSuppressB
 		// STAGE XVI: CLEAR TEMPORARY FLAGS
 		for (FCompilerData& CompilerData : CurrentlyCompilingBPs)
 		{
+			DECLARE_SCOPE_HIERARCHICAL_COUNTER(ClearTemporaryFlags)
+
 			UBlueprint* BP = CompilerData.BP;
 
 			if (!CompilerData.IsSkeletonOnly())
@@ -1305,6 +1341,8 @@ void FBlueprintCompilationManagerImpl::FlushCompilationQueueImpl(bool bSuppressB
 
 		if (!bSuppressBroadcastCompiled)
 		{
+			DECLARE_SCOPE_HIERARCHICAL_COUNTER(BroadcastBlueprintCompiled)
+
 			if(GEditor)
 			{
 				GEditor->BroadcastBlueprintCompiled();	
@@ -1321,7 +1359,10 @@ void FBlueprintCompilationManagerImpl::FlushCompilationQueueImpl(bool bSuppressB
 		CompilerData.BP->bQueuedForCompilation = false;
 	}
 
-	UEdGraphPin::Purge();
+	{
+		DECLARE_SCOPE_HIERARCHICAL_COUNTER(UEdGraphPin::Purge)
+		UEdGraphPin::Purge();
+	}
 
 	UE_LOG(LogBlueprint, Display, TEXT("Time Compiling: %f, Time Reinstancing: %f"),  GTimeCompiling, GTimeReinstancing);
 	//GTimeCompiling = 0.0;
@@ -1331,6 +1372,8 @@ void FBlueprintCompilationManagerImpl::FlushCompilationQueueImpl(bool bSuppressB
 
 void FBlueprintCompilationManagerImpl::ProcessExtensions(const TArray<FCompilerData>& InCurrentlyCompilingBPs)
 {
+	DECLARE_SCOPE_HIERARCHICAL_COUNTER_FUNC()
+
 	if(CompilerExtensions.Num() == 0)
 	{
 		return;
@@ -1382,6 +1425,8 @@ void FBlueprintCompilationManagerImpl::ProcessExtensions(const TArray<FCompilerD
 
 void FBlueprintCompilationManagerImpl::FlushReinstancingQueueImpl()
 {
+	DECLARE_SCOPE_HIERARCHICAL_COUNTER_FUNC()
+
 	if(GCompilingBlueprint)
 	{
 		return;
diff --git a/Engine/Source/Editor/Kismet/Private/BlueprintEditor.cpp b/Engine/Source/Editor/Kismet/Private/BlueprintEditor.cpp
--- a/Engine/Source/Editor/Kismet/Private/BlueprintEditor.cpp
+++ b/Engine/Source/Editor/Kismet/Private/BlueprintEditor.cpp
@@ -95,6 +95,7 @@
 #include "Editor.h"
 #include "IDetailsView.h"
 #include "HAL/PlatformApplicationMisc.h"
+#include "Stats/StatsHierarchical.h"
 
 #include "BlueprintEditorTabs.h"
 
@@ -3346,6 +3347,8 @@ void FBlueprintEditor::OnBlueprintUnloaded(UBlueprint* InBlueprint)
 
 void FBlueprintEditor::Compile()
 {
+	DECLARE_SCOPE_HIERARCHICAL_COUNTER_FUNC()
+
 	UBlueprint* BlueprintObj = GetBlueprintObj();
 	if (BlueprintObj)
 	{
@@ -8028,7 +8031,7 @@ void FBlueprintEditor::OnEditTabClosed(TSharedRef<SDockTab> Tab)
 }
 
 // Tries to open the specified graph and bring it's document to the front
-TSharedPtr<SGraphEditor> FBlueprintEditor::OpenGraphAndBringToFront(UEdGraph* Graph)
+TSharedPtr<SGraphEditor> FBlueprintEditor::OpenGraphAndBringToFront(UEdGraph* Graph, bool bSetFocus)
 {
 	if (!Graph || Graph->IsPendingKill())
 	{
@@ -8045,7 +8048,10 @@ TSharedPtr<SGraphEditor> FBlueprintEditor::OpenGraphAndBringToFront(UEdGraph* Gr
 	TSharedRef<SGraphEditor> NewGraphEditor = StaticCastSharedRef<SGraphEditor>(TabWithGraph->GetContent());
 
 	// Handover the keyboard focus to the new graph editor widget.
-	NewGraphEditor->CaptureKeyboard();
+	if (bSetFocus)
+	{
+		NewGraphEditor->CaptureKeyboard();
+	}
 
 	return NewGraphEditor;
 }
diff --git a/Engine/Source/Editor/Kismet/Public/BlueprintEditor.h b/Engine/Source/Editor/Kismet/Public/BlueprintEditor.h
--- a/Engine/Source/Editor/Kismet/Public/BlueprintEditor.h
+++ b/Engine/Source/Editor/Kismet/Public/BlueprintEditor.h
@@ -217,7 +217,7 @@ public:
 	virtual void JumpToPin(const class UEdGraphPin* Pin) override;
 	virtual void SummonSearchUI(bool bSetFindWithinBlueprint, FString NewSearchTerms = FString(), bool bSelectFirstResult = false) override;
 	virtual void SummonFindAndReplaceUI() override;
-	virtual TSharedPtr<SGraphEditor> OpenGraphAndBringToFront(UEdGraph* Graph) override;
+	virtual TSharedPtr<SGraphEditor> OpenGraphAndBringToFront(UEdGraph* Graph, bool bSetFocus = true) override;
 	virtual TArray<TSharedPtr<class FSCSEditorTreeNode> > GetSelectedSCSEditorTreeNodes() const override;
 	virtual TSharedPtr<class FSCSEditorTreeNode> FindAndSelectSCSEditorTreeNode(const UActorComponent* InComponent, bool IsCntrlDown) override;
 	virtual int32 GetNumberOfSelectedNodes() const override;
diff --git a/Engine/Source/Editor/Kismet/Public/BlueprintEditorModule.h b/Engine/Source/Editor/Kismet/Public/BlueprintEditorModule.h
--- a/Engine/Source/Editor/Kismet/Public/BlueprintEditorModule.h
+++ b/Engine/Source/Editor/Kismet/Public/BlueprintEditorModule.h
@@ -68,7 +68,7 @@ public:
 	virtual void SummonFindAndReplaceUI() = 0;
 
 	/** Tries to open the specified graph and bring it's document to the front (note: this can return NULL) */
-	virtual TSharedPtr<class SGraphEditor> OpenGraphAndBringToFront(class UEdGraph* Graph) = 0;
+	virtual TSharedPtr<class SGraphEditor> OpenGraphAndBringToFront(class UEdGraph* Graph, bool bSetFocus = true) = 0;
 
 	virtual void RefreshEditors(ERefreshBlueprintEditorReason::Type Reason = ERefreshBlueprintEditorReason::UnknownReason) = 0;
 
diff --git a/Engine/Source/Editor/Persona/Private/AnimationEditorPreviewScene.cpp b/Engine/Source/Editor/Persona/Private/AnimationEditorPreviewScene.cpp
--- a/Engine/Source/Editor/Persona/Private/AnimationEditorPreviewScene.cpp
+++ b/Engine/Source/Editor/Persona/Private/AnimationEditorPreviewScene.cpp
@@ -572,7 +572,7 @@ void FAnimationEditorPreviewScene::SetPreviewAnimationAsset(UAnimationAsset* Ani
 			}
 
 			// Treat it as invalid if it's got a bogus skeleton pointer
-			if (AnimAsset->GetSkeleton() != Skeleton)
+			if (AnimAsset->GetSkeleton() != Skeleton && Skeleton != nullptr)
 			{
 				return;
 			}
diff --git a/Engine/Source/Editor/UnrealEd/Classes/Animation/DebugSkelMeshComponent.h b/Engine/Source/Editor/UnrealEd/Classes/Animation/DebugSkelMeshComponent.h
--- a/Engine/Source/Editor/UnrealEd/Classes/Animation/DebugSkelMeshComponent.h
+++ b/Engine/Source/Editor/UnrealEd/Classes/Animation/DebugSkelMeshComponent.h
@@ -317,7 +317,7 @@ class UNREALED_API UDebugSkelMeshComponent : public USkeletalMeshComponent
 
 	// @todo anim : you still need to give asset, so that we know which one to disable
 	// we can disable per asset, so that if some other window disabled before me, I don't accidently turn it off
-	void EnablePreview(bool bEnable, class UAnimationAsset * PreviewAsset);
+	virtual void EnablePreview(bool bEnable, class UAnimationAsset * PreviewAsset);
 
 	// reference pose for this component
 	// we don't want to use default refpose because you still want to move joint when this mode is on
