Fix deprecation warning
#rnx

diff --git a/Engine/Plugins/Experimental/ControlRig/Source/ControlRigEditor/Private/Editor/SRigCurveContainer.cpp b/Engine/Plugins/Experimental/ControlRig/Source/ControlRigEditor/Private/Editor/SRigCurveContainer.cpp
--- a/Engine/Plugins/Experimental/ControlRig/Source/ControlRigEditor/Private/Editor/SRigCurveContainer.cpp
+++ b/Engine/Plugins/Experimental/ControlRig/Source/ControlRigEditor/Private/Editor/SRigCurveContainer.cpp
@@ -290,7 +290,7 @@ TSharedPtr<SWidget> SRigCurveContainer::OnGetContextMenuContent() const
 	MenuBuilder.AddSubMenu(
 		LOCTEXT("ImportSubMenu", "Import"),
 		LOCTEXT("ImportSubMenu_ToolTip", "Import curves to the current rig. This only imports non-existing curve."),
-		FNewMenuDelegate::CreateSP(this, &SRigCurveContainer::CreateImportMenu)
+		FNewMenuDelegate::CreateSP(const_cast<SRigCurveContainer*>(this), &SRigCurveContainer::CreateImportMenu)
 	);
 
 
