#DUPEMERGE

Fixes many reference pose issues, where they want to add node, and how it updates, and how they could move

- mainly execution can be disabled by button, and  we will mark to compile when it needs recompile due to hierarchy change
- added quick warning to log when hierarchy is missing
- still lots of things need more iteration
- fixed assets

diff --git a/Engine/Plugins/Experimental/ControlRig/Source/ControlRigEditor/Private/Editor/ControlRigEditor.cpp b/Engine/Plugins/Experimental/ControlRig/Source/ControlRigEditor/Private/Editor/ControlRigEditor.cpp
--- a/Engine/Plugins/Experimental/ControlRig/Source/ControlRigEditor/Private/Editor/ControlRigEditor.cpp
+++ b/Engine/Plugins/Experimental/ControlRig/Source/ControlRigEditor/Private/Editor/ControlRigEditor.cpp
@@ -833,6 +833,12 @@ void FControlRigEditor::SetJointTransform(const FName& InJoint, const FTransform
 		EditorSkelComp->RebuildDebugDrawSkeleton();
 	}
 
+	// I don't think I have to mark dirty here. 
+	// FBlueprintEditorUtils::MarkBlueprintAsModified(GetControlRigBlueprint());
+	{
+		EditorSkelComp->RebuildDebugDrawSkeleton();
+	}
+
 	// I don't think I have to mark dirty here. 
 	// FBlueprintEditorUtils::MarkBlueprintAsModified(GetControlRigBlueprint());
 }
