Control Rig: Fix crash tearing down sequencer object spawner

diff --git a/Engine/Plugins/Experimental/ControlRig/Source/ControlRigEditor/Private/Sequencer/ControlRigEditorObjectSpawner.cpp b/Engine/Plugins/Experimental/ControlRig/Source/ControlRigEditor/Private/Sequencer/ControlRigEditorObjectSpawner.cpp
--- a/Engine/Plugins/Experimental/ControlRig/Source/ControlRigEditor/Private/Sequencer/ControlRigEditorObjectSpawner.cpp
+++ b/Engine/Plugins/Experimental/ControlRig/Source/ControlRigEditor/Private/Sequencer/ControlRigEditorObjectSpawner.cpp
@@ -25,7 +25,10 @@ FControlRigEditorObjectSpawner::FControlRigEditorObjectSpawner()
 FControlRigEditorObjectSpawner::~FControlRigEditorObjectSpawner()
 {
 #if WITH_EDITOR
-	GEditor->OnObjectsReplaced().RemoveAll(this);
+	if (GEditor)
+	{
+		GEditor->OnObjectsReplaced().RemoveAll(this);
+	}
 #endif
 }
 TSharedRef<IMovieSceneObjectSpawner> FControlRigEditorObjectSpawner::CreateObjectSpawner()
