Added Control Rig Execution stats

diff --git a/Engine/Plugins/Experimental/ControlRig/Source/ControlRig/Private/ControlRig.cpp b/Engine/Plugins/Experimental/ControlRig/Source/ControlRig/Private/ControlRig.cpp
--- a/Engine/Plugins/Experimental/ControlRig/Source/ControlRig/Private/ControlRig.cpp
+++ b/Engine/Plugins/Experimental/ControlRig/Source/ControlRig/Private/ControlRig.cpp
@@ -16,6 +16,10 @@
 
 DEFINE_LOG_CATEGORY_STATIC(LogControlRig, Log, All);
 
+DECLARE_STATS_GROUP(TEXT("ControlRig"), STATGROUP_ControlRig, STATCAT_Advanced);
+DECLARE_CYCLE_STAT_EXTERN(TEXT("Control Rig Execution"), STAT_RigExecution, STATGROUP_ControlRig, );
+DEFINE_STAT(STAT_RigExecution);
+
 const FName UControlRig::InputMetaName("Input");
 const FName UControlRig::OutputMetaName("Output");
 const FName UControlRig::AbstractMetaName("Abstract");
@@ -162,6 +166,8 @@ void UControlRig::InstantiateExecutor()
 
 void UControlRig::Execute(const EControlRigState InState)
 {
+	SCOPE_CYCLE_COUNTER(STAT_RigExecution);
+
 	FRigUnitContext Context;
 	Context.DeltaTime = DeltaTime;
 	Context.State = InState;
