Fixes many reference pose issues, where they want to add node, and how it updates, and how they could move

- mainly execution can be disabled by button, and  we will mark to compile when it needs recompile due to hierarchy change
- added quick warning to log when hierarchy is missing
- still lots of things need more iteration
- fixed assets

diff --git a/Engine/Plugins/Experimental/ControlRig/Source/ControlRig/Private/AnimNode_ControlRigBase.cpp b/Engine/Plugins/Experimental/ControlRig/Source/ControlRig/Private/AnimNode_ControlRigBase.cpp
--- a/Engine/Plugins/Experimental/ControlRig/Source/ControlRig/Private/AnimNode_ControlRigBase.cpp
+++ b/Engine/Plugins/Experimental/ControlRig/Source/ControlRig/Private/AnimNode_ControlRigBase.cpp
@@ -58,83 +58,65 @@ void FAnimNode_ControlRigBase::Update_AnyThread(const FAnimationUpdateContext& C
 
 void FAnimNode_ControlRigBase::UpdateInput(UControlRig* ControlRig, const FPoseContext& InOutput)
 {
-	// copy output of the rig
-	// @todo: find a better option to do this
-	// but for now, we use this flag to skip the execution and updating input
-	// this is used to skip when refpose is chosen to view
-#if WITH_EDITORONLY_DATA
-	if (ControlRig->bExecutionOn)
-#endif // WITH_EDITORONLY_DATA
-	{
-		const FBoneContainer& RequiredBones = InOutput.Pose.GetBoneContainer();
+	const FBoneContainer& RequiredBones = InOutput.Pose.GetBoneContainer();
 
-		// get component pose from control rig
-		FCSPose<FCompactPose> MeshPoses;
-		// first I need to convert to local pose
-		MeshPoses.InitPose(InOutput.Pose);
+	// get component pose from control rig
+	FCSPose<FCompactPose> MeshPoses;
+	// first I need to convert to local pose
+	MeshPoses.InitPose(InOutput.Pose);
 
-		// should we also update refpose? 
-		const int32 NumNodes = RigHierarchyItemNameMapping.Num();
-		for (int32 Index = 0; Index < NumNodes; ++Index)
+	// should we also update refpose? 
+	const int32 NumNodes = RigHierarchyItemNameMapping.Num();
+	for (int32 Index = 0; Index < NumNodes; ++Index)
+	{
+		if (RigHierarchyItemNameMapping[Index] != NAME_None)
 		{
-			if (RigHierarchyItemNameMapping[Index] != NAME_None)
+			FCompactPoseBoneIndex CompactPoseIndex(Index);
+			FTransform ComponentTransform = MeshPoses.GetComponentSpaceTransform(CompactPoseIndex);
+			if (NodeMappingContainer.IsValid())
 			{
-				FCompactPoseBoneIndex CompactPoseIndex(Index);
-				FTransform ComponentTransform = MeshPoses.GetComponentSpaceTransform(CompactPoseIndex);
-				if (NodeMappingContainer.IsValid())
-				{
-					// @todo : I should reverse retarget
-					//ComponentTransform = NodeMappingContainer->GetSourceToTargetTransform(RigHierarchyItemNameMapping[Index]) * ComponentTransform;
-				}
-
-				ControlRig->SetGlobalTransform(RigHierarchyItemNameMapping[Index], ComponentTransform);
+				ComponentTransform = NodeMappingContainer->GetSourceToTargetTransform(RigHierarchyItemNameMapping[Index]).GetRelativeTransformReverse(ComponentTransform);
 			}
+
+			ControlRig->SetGlobalTransform(RigHierarchyItemNameMapping[Index], ComponentTransform);
 		}
 	}
 }
 
 void FAnimNode_ControlRigBase::UpdateOutput(const UControlRig* ControlRig, FPoseContext& InOutput)
 {
-	// @todo: find a better option to do this
-	// but for now, we use this flag to skip the execution and updating output
-	// this is used to skip when refpose is chosen to view
-#if WITH_EDITORONLY_DATA
-	if (ControlRig->bExecutionOn)
-#endif // WITH_EDITORONLY_DATA
-	{
-		// copy output of the rig
-		const FBoneContainer& RequiredBones = InOutput.Pose.GetBoneContainer();
+	// copy output of the rig
+	const FBoneContainer& RequiredBones = InOutput.Pose.GetBoneContainer();
 
-		// get component pose from control rig
-		FCSPose<FCompactPose> MeshPoses;
-		MeshPoses.InitPose(InOutput.Pose);
+	// get component pose from control rig
+	FCSPose<FCompactPose> MeshPoses;
+	MeshPoses.InitPose(InOutput.Pose);
 
-		const int32 NumNodes = RigHierarchyItemNameMapping.Num();
-		for (int32 Index = 0; Index < NumNodes; ++Index)
+	const int32 NumNodes = RigHierarchyItemNameMapping.Num();
+	for (int32 Index = 0; Index < NumNodes; ++Index)
+	{
+		if (RigHierarchyItemNameMapping[Index] != NAME_None)
 		{
-			if (RigHierarchyItemNameMapping[Index] != NAME_None)
+			FCompactPoseBoneIndex CompactPoseIndex(Index);
+			FTransform ComponentTransform = ControlRig->GetGlobalTransform(RigHierarchyItemNameMapping[Index]);
+			if (NodeMappingContainer.IsValid())
 			{
-				FCompactPoseBoneIndex CompactPoseIndex(Index);
-				FTransform ComponentTransform = ControlRig->GetGlobalTransform(RigHierarchyItemNameMapping[Index]);
-				if (NodeMappingContainer.IsValid())
-				{
-					ComponentTransform = NodeMappingContainer->GetSourceToTargetTransform(RigHierarchyItemNameMapping[Index]) * ComponentTransform;
-				}
-
-				MeshPoses.SetComponentSpaceTransform(CompactPoseIndex, ComponentTransform);
+				ComponentTransform = NodeMappingContainer->GetSourceToTargetTransform(RigHierarchyItemNameMapping[Index]) * ComponentTransform;
 			}
+
+			MeshPoses.SetComponentSpaceTransform(CompactPoseIndex, ComponentTransform);
 		}
+	}
 
-		// now convert to what you'd like
-		// initialize with refpose
-	 	InOutput.ResetToRefPose();
-		for (int32 Index = 0; Index < NumNodes; ++Index)
+	// now convert to what you'd like
+	// initialize with refpose
+	InOutput.ResetToRefPose();
+	for (int32 Index = 0; Index < NumNodes; ++Index)
+	{
+		if (RigHierarchyItemNameMapping[Index] != NAME_None)
 		{
-			if (RigHierarchyItemNameMapping[Index] != NAME_None)
-			{
-				FCompactPoseBoneIndex CompactPoseIndex(Index);
-				InOutput.Pose[CompactPoseIndex] = MeshPoses.GetLocalSpaceTransform(CompactPoseIndex);
-			}
+			FCompactPoseBoneIndex CompactPoseIndex(Index);
+			InOutput.Pose[CompactPoseIndex] = MeshPoses.GetLocalSpaceTransform(CompactPoseIndex);
 		}
 	}
 }
diff --git a/Engine/Plugins/Experimental/ControlRig/Source/ControlRig/Private/Units/RigUnit_AimConstraint.cpp b/Engine/Plugins/Experimental/ControlRig/Source/ControlRig/Private/Units/RigUnit_AimConstraint.cpp
--- a/Engine/Plugins/Experimental/ControlRig/Source/ControlRig/Private/Units/RigUnit_AimConstraint.cpp
+++ b/Engine/Plugins/Experimental/ControlRig/Source/ControlRig/Private/Units/RigUnit_AimConstraint.cpp
@@ -35,6 +35,10 @@ void FRigUnit_AimConstraint::Execute(const FRigUnitContext& InContext)
 				}
 			}
 		}
+		else
+		{
+			UnitLogHelpers::PrintMissingHierarchy(RigUnitName);
+		}
 	}
 	else if (InContext.State == EControlRigState::Update)
 	{
diff --git a/Engine/Plugins/Experimental/ControlRig/Source/ControlRig/Private/Units/RigUnit_ApplyFK.cpp b/Engine/Plugins/Experimental/ControlRig/Source/ControlRig/Private/Units/RigUnit_ApplyFK.cpp
--- a/Engine/Plugins/Experimental/ControlRig/Source/ControlRig/Private/Units/RigUnit_ApplyFK.cpp
+++ b/Engine/Plugins/Experimental/ControlRig/Source/ControlRig/Private/Units/RigUnit_ApplyFK.cpp
@@ -6,7 +6,15 @@
 
 void FRigUnit_ApplyFK::Execute(const FRigUnitContext& InContext)
 {
-	if (InContext.State == EControlRigState::Update)
+	if (InContext.State == EControlRigState::Init)
+	{
+		FRigHierarchy* Hierarchy = HierarchyRef.Get();
+		if (!Hierarchy)
+		{
+			UnitLogHelpers::PrintMissingHierarchy(RigUnitName);
+		}
+	}
+	else if (InContext.State == EControlRigState::Update)
 	{
 		FRigHierarchy* Hierarchy = HierarchyRef.Get();
 		if (Hierarchy)
diff --git a/Engine/Plugins/Experimental/ControlRig/Source/ControlRig/Private/Units/RigUnit_FABRIK.cpp b/Engine/Plugins/Experimental/ControlRig/Source/ControlRig/Private/Units/RigUnit_FABRIK.cpp
--- a/Engine/Plugins/Experimental/ControlRig/Source/ControlRig/Private/Units/RigUnit_FABRIK.cpp
+++ b/Engine/Plugins/Experimental/ControlRig/Source/ControlRig/Private/Units/RigUnit_FABRIK.cpp
@@ -39,13 +39,16 @@ void FRigUnit_FABRIK::Execute(const FRigUnitContext& InContext)
 				}
 			}
 		}
+		else
+		{
+			UnitLogHelpers::PrintMissingHierarchy(RigUnitName);
+		}
 	}
 	else  if (InContext.State == EControlRigState::Update)
 	{
 		if (FullLimbLength > 0.f)
 		{
-			
-	
+			UnitLogHelpers::PrintUnimplemented(RigUnitName);
 		}
 	}
 }
diff --git a/Engine/Plugins/Experimental/ControlRig/Source/ControlRig/Private/Units/RigUnit_FABRIK.h b/Engine/Plugins/Experimental/ControlRig/Source/ControlRig/Private/Units/RigUnit_FABRIK.h
--- a/Engine/Plugins/Experimental/ControlRig/Source/ControlRig/Private/Units/RigUnit_FABRIK.h
+++ b/Engine/Plugins/Experimental/ControlRig/Source/ControlRig/Private/Units/RigUnit_FABRIK.h
@@ -12,7 +12,8 @@
  * Spec define: TBD
  */
 
-USTRUCT(meta=(DisplayName="FABRIK", Category="Transforms"))
+// make it abstract since it's not working yet
+USTRUCT(meta=(DisplayName="FABRIK", Category="Transforms", Abstract))
 struct CONTROLRIG_API FRigUnit_FABRIK : public FRigUnit
 {
 	GENERATED_BODY()
diff --git a/Engine/Plugins/Experimental/ControlRig/Source/ControlRig/Private/Units/RigUnit_GetJointTransform.cpp b/Engine/Plugins/Experimental/ControlRig/Source/ControlRig/Private/Units/RigUnit_GetJointTransform.cpp
--- a/Engine/Plugins/Experimental/ControlRig/Source/ControlRig/Private/Units/RigUnit_GetJointTransform.cpp
+++ b/Engine/Plugins/Experimental/ControlRig/Source/ControlRig/Private/Units/RigUnit_GetJointTransform.cpp
@@ -42,4 +42,11 @@ void FRigUnit_GetJointTransform::Execute(const FRigUnitContext& InContext)
 			}
 		}
 	}
+	else
+	{
+		if (InContext.State == EControlRigState::Init)
+		{
+			UnitLogHelpers::PrintMissingHierarchy(RigUnitName);
+		}
+	}
 }
diff --git a/Engine/Plugins/Experimental/ControlRig/Source/ControlRig/Private/Units/RigUnit_TwoBoneIKFK.cpp b/Engine/Plugins/Experimental/ControlRig/Source/ControlRig/Private/Units/RigUnit_TwoBoneIKFK.cpp
--- a/Engine/Plugins/Experimental/ControlRig/Source/ControlRig/Private/Units/RigUnit_TwoBoneIKFK.cpp
+++ b/Engine/Plugins/Experimental/ControlRig/Source/ControlRig/Private/Units/RigUnit_TwoBoneIKFK.cpp
@@ -49,6 +49,10 @@ void FRigUnit_TwoBoneIKFK::Execute(const FRigUnitContext& InContext)
 				}
 			}
 		}
+		else
+		{
+			UnitLogHelpers::PrintMissingHierarchy(RigUnitName);
+		}
 	}
 	else  if (InContext.State == EControlRigState::Update)
 	{
diff --git a/Engine/Plugins/Experimental/ControlRig/Source/ControlRig/Public/Units/RigUnit.h b/Engine/Plugins/Experimental/ControlRig/Source/ControlRig/Public/Units/RigUnit.h
--- a/Engine/Plugins/Experimental/ControlRig/Source/ControlRig/Public/Units/RigUnit.h
+++ b/Engine/Plugins/Experimental/ControlRig/Source/ControlRig/Public/Units/RigUnit.h
@@ -48,4 +48,11 @@ struct CONTROLRIG_API FRigUnit
 
 	UPROPERTY(EditAnywhere, Category = FRigUnit)
 	EUnitExecutionType ExecutionType;
-};
\ No newline at end of file
+};
+
+// this will have to change in the future and move to editor, I assume the errors will be saved in the rig unit and it will print fromthe editor module
+namespace UnitLogHelpers
+{
+	void PrintMissingHierarchy(const FName& InputName);
+	void PrintUnimplemented(const FName& InputName);
+}
diff --git a/Engine/Plugins/Experimental/ControlRig/Source/ControlRigEditor/Private/ControlRigBlueprintCommands.cpp b/Engine/Plugins/Experimental/ControlRig/Source/ControlRigEditor/Private/ControlRigBlueprintCommands.cpp
--- a/Engine/Plugins/Experimental/ControlRig/Source/ControlRigEditor/Private/ControlRigBlueprintCommands.cpp
+++ b/Engine/Plugins/Experimental/ControlRig/Source/ControlRigEditor/Private/ControlRigBlueprintCommands.cpp
@@ -7,6 +7,7 @@
 void FControlRigBlueprintCommands::RegisterCommands()
 {
 	UI_COMMAND(DeleteItem, "Delete", "Deletes the selected items and removes their nodes from the graph.", EUserInterfaceActionType::Button, FInputChord(EKeys::Delete));
+	UI_COMMAND(ExecuteGraph, "Execute", "Execute the rig graph if On.", EUserInterfaceActionType::ToggleButton, FInputChord());
 }
 
 #undef LOCTEXT_NAMESPACE
diff --git a/Engine/Plugins/Experimental/ControlRig/Source/ControlRigEditor/Private/ControlRigBlueprintCommands.h b/Engine/Plugins/Experimental/ControlRig/Source/ControlRigEditor/Private/ControlRigBlueprintCommands.h
--- a/Engine/Plugins/Experimental/ControlRig/Source/ControlRigEditor/Private/ControlRigBlueprintCommands.h
+++ b/Engine/Plugins/Experimental/ControlRig/Source/ControlRigEditor/Private/ControlRigBlueprintCommands.h
@@ -21,6 +21,8 @@ public:
 	/** Deletes the selected items and removes their nodes from the graph. */
 	TSharedPtr< FUICommandInfo > DeleteItem;
 
+	/** Toggle Execute the Graph */
+	TSharedPtr< FUICommandInfo > ExecuteGraph;
 	/**
 	 * Initialize commands
 	 */
diff --git a/Engine/Plugins/Experimental/ControlRig/Source/ControlRigEditor/Private/ControlRigEditorStyle.h b/Engine/Plugins/Experimental/ControlRig/Source/ControlRigEditor/Private/ControlRigEditorStyle.h
--- a/Engine/Plugins/Experimental/ControlRig/Source/ControlRigEditor/Private/ControlRigEditorStyle.h
+++ b/Engine/Plugins/Experimental/ControlRig/Source/ControlRigEditor/Private/ControlRigEditorStyle.h
@@ -60,6 +60,9 @@ public:
 			// icons for control units
 			Set("ControlRig.ControlUnitOn", new IMAGE_BRUSH(TEXT("ControlUnit_On"), Icon32x32));
 			Set("ControlRig.ControlUnitOff", new IMAGE_BRUSH(TEXT("ControlUnit_Off"), Icon32x32));
+
+			Set("ControlRig.ExecuteGraph", new IMAGE_BRUSH("ExecuteGraph", Icon40x40));
+			Set("ControlRig.ExecuteGraph.Small", new IMAGE_BRUSH("ExecuteGraph", Icon20x20));
 		}
 
 		// Graph styles
diff --git a/Engine/Plugins/Experimental/ControlRig/Source/ControlRigEditor/Private/EditMode/ControlRigEditMode.cpp b/Engine/Plugins/Experimental/ControlRig/Source/ControlRigEditor/Private/EditMode/ControlRigEditMode.cpp
--- a/Engine/Plugins/Experimental/ControlRig/Source/ControlRigEditor/Private/EditMode/ControlRigEditMode.cpp
+++ b/Engine/Plugins/Experimental/ControlRig/Source/ControlRigEditor/Private/EditMode/ControlRigEditMode.cpp
@@ -487,7 +487,7 @@ bool FControlRigEditMode::UsesTransformWidget() const
 		}
 	}
 
-	if (AreJointSelected())
+	if (AreJointSelectedAndMovable())
 	{
 		return true;
 	}
@@ -509,11 +509,11 @@ bool FControlRigEditMode::UsesTransformWidget(FWidget::EWidgetMode CheckMode) co
 				}
 			}
 		}
-	}
 
-	if (AreJointSelected())
-	{
-		return true;
+		if (AreJointSelectedAndMovable())
+		{
+			return true;
+		}
 	}
 
 	return FEdMode::UsesTransformWidget(CheckMode);
@@ -536,7 +536,7 @@ FVector FControlRigEditMode::GetWidgetLocation() const
 
 		// @todo: we only supports the first ast one for now
 		// later we support multi select
-		if (AreJointSelected())
+		if (AreJointSelectedAndMovable())
 		{
 			return ComponentTransform.TransformPosition(OnGetJointTransformDelegate.Execute(SelectedJoints[0], false).GetLocation());
 		}
@@ -558,7 +558,7 @@ bool FControlRigEditMode::GetCustomDrawingCoordinateSystem(FMatrix& OutMatrix, v
 			}
 		}
 
-		if (AreJointSelected())
+		if (AreJointSelectedAndMovable())
 		{
 			USceneComponent* Component = Cast<USkeletalMeshComponent>(ControlRig->GetObjectBinding()->GetBoundObject());
 			FTransform ComponentTransform = Component ? Component->GetComponentTransform() : FTransform::Identity;
@@ -791,7 +791,7 @@ bool FControlRigEditMode::InputDelta(FEditorViewportClient* InViewportClient, FV
 
 				return true;
 			}
-			else if (AreJointSelected())
+			else if (AreJointSelectedAndMovable())
 			{
 				// set joint transform
 				// that will set initial joint transform
@@ -838,7 +838,7 @@ bool FControlRigEditMode::InputDelta(FEditorViewportClient* InViewportClient, FV
 
 bool FControlRigEditMode::ShouldDrawWidget() const
 {
-	if (AreControlsSelected() || AreJointSelected())
+	if (AreControlsSelected() || AreJointSelectedAndMovable())
 	{
 		return true;
 	}
@@ -1362,9 +1362,19 @@ void FControlRigEditMode::OnObjectsReplaced(const TMap<UObject*, UObject*>& OldT
 	}
 }
 
+bool FControlRigEditMode::AreJointSelectedAndMovable() const
+{
+	if (UControlRig* ControlRig = WeakControlRig.Get())
+	{
+		return (!ControlRig->bExecutionOn && OnGetJointTransformDelegate.IsBound() && OnSetJointTransformDelegate.IsBound() && SelectedJoints.Num() > 0);
+	}
+
+	return false;
+}
+
 bool FControlRigEditMode::AreJointSelected() const
 {
-	return (OnGetJointTransformDelegate.IsBound() && OnSetJointTransformDelegate.IsBound() && SelectedJoints.Num() > 0);
+	return (SelectedJoints.Num() > 0);
 }
 
 void FControlRigEditMode::SelectJoint(const FName& InJoint)
@@ -1372,6 +1382,9 @@ void FControlRigEditMode::SelectJoint(const FName& InJoint)
 	ClearControlSelection();
 
 	SelectedJoints.Reset();
-	SelectedJoints.Add(InJoint);
+	if (InJoint != NAME_None)
+	{
+		SelectedJoints.Add(InJoint);
+	}
 }
 #undef LOCTEXT_NAMESPACE
diff --git a/Engine/Plugins/Experimental/ControlRig/Source/ControlRigEditor/Private/EditMode/ControlRigEditMode.h b/Engine/Plugins/Experimental/ControlRig/Source/ControlRigEditor/Private/EditMode/ControlRigEditMode.h
--- a/Engine/Plugins/Experimental/ControlRig/Source/ControlRigEditor/Private/EditMode/ControlRigEditMode.h
+++ b/Engine/Plugins/Experimental/ControlRig/Source/ControlRigEditor/Private/EditMode/ControlRigEditMode.h
@@ -224,4 +224,6 @@ protected:
 	FOnSetJointTransform OnSetJointTransformDelegate;
 
 	bool AreJointSelected() const;
+
+	bool AreJointSelectedAndMovable() const;
 };
diff --git a/Engine/Plugins/Experimental/ControlRig/Source/ControlRigEditor/Private/Editor/ControlRigEditor.cpp b/Engine/Plugins/Experimental/ControlRig/Source/ControlRigEditor/Private/Editor/ControlRigEditor.cpp
--- a/Engine/Plugins/Experimental/ControlRig/Source/ControlRigEditor/Private/Editor/ControlRigEditor.cpp
+++ b/Engine/Plugins/Experimental/ControlRig/Source/ControlRigEditor/Private/Editor/ControlRigEditor.cpp
@@ -31,6 +31,13 @@
 #include "IPersonaViewport.h"
 #include "EditorViewportClient.h"
 #include "AnimationEditorPreviewActor.h"
+#include "Misc/MessageDialog.h"
+#include "Framework/MultiBox/MultiBoxBuilder.h"
+#include "ControlRigEditorStyle.h"
+#include "EditorFontGlyphs.h"
+#include "Widgets/Input/SButton.h"
+#include "Widgets/Notifications/SNotificationList.h"
+#include "Framework/Notifications/NotificationManager.h"
 
 #define LOCTEXT_NAMESPACE "ControlRigEditor"
 
@@ -141,7 +148,24 @@ void FControlRigEditor::InitControlRigEditor(const EToolkitMode::Type Mode, cons
 
 void FControlRigEditor::BindCommands()
 {
+	GetToolkitCommands()->MapAction(
+		FControlRigBlueprintCommands::Get().ExecuteGraph,
+		FExecuteAction::CreateSP(this, &FControlRigEditor::ToggleExecuteGraph), 
+		FCanExecuteAction(), 
+		FIsActionChecked::CreateSP(this, &FControlRigEditor::IsExecuteGraphOn));
+}
 
+void FControlRigEditor::ToggleExecuteGraph()
+{
+	if (ControlRig)
+	{
+		ControlRig->bExecutionOn = !ControlRig->bExecutionOn;
+	}
+}
+
+bool FControlRigEditor::IsExecuteGraphOn() const
+{
+	return (ControlRig)? ControlRig->bExecutionOn : false;
 }
 
 void FControlRigEditor::ExtendToolbar()
@@ -169,6 +193,26 @@ void FControlRigEditor::ExtendToolbar()
 			AddToolbarExtender(ToolbarExtenderDelegate.Execute(GetToolkitCommands(), SharedThis(this)));
 		}
 	}
+
+	struct Local
+	{
+		static void FillToolbar(FToolBarBuilder& ToolbarBuilder)
+		{
+			ToolbarBuilder.BeginSection("Toolbar");
+			{
+				ToolbarBuilder.AddToolBarButton(FControlRigBlueprintCommands::Get().ExecuteGraph, 
+					NAME_None, TAttribute<FText>(), TAttribute<FText>(), FSlateIcon(FControlRigEditorStyle::Get().GetStyleSetName(), "ControlRig.ExecuteGraph"));
+			}
+			ToolbarBuilder.EndSection();
+		}
+	};
+
+	ToolbarExtender->AddToolBarExtension(
+		"Asset",
+		EExtensionHook::After,
+		GetToolkitCommands(),
+		FToolBarExtensionDelegate::CreateStatic(&Local::FillToolbar)
+	);
 }
 
 UBlueprint* FControlRigEditor::GetBlueprintObj() const
@@ -294,6 +338,8 @@ void FControlRigEditor::PostUndo(bool bSuccess)
 	DocumentManager->CleanInvalidTabs();
 	DocumentManager->RefreshAllTabs();
 
+	OnHierarchyChanged();
+
 	FBlueprintEditor::PostUndo(bSuccess);
 }
 
@@ -502,8 +548,145 @@ void FControlRigEditor::OnBlueprintChangedImpl(UBlueprint* InBlueprint, bool bIs
 
 void FControlRigEditor::HandleViewportCreated(const TSharedRef<class IPersonaViewport>& InViewport)
 {
-	// set default to be local 
-	InViewport->GetViewportClient().SetWidgetCoordSystemSpace(COORD_Local);
+	// TODO: this is duplicated code from FAnimBlueprintEditor, would be nice to consolidate. 
+	auto GetCompilationStateText = [this]()
+	{
+		if (UBlueprint* Blueprint = GetBlueprintObj())
+		{
+			switch (Blueprint->Status)
+			{
+			case BS_UpToDate:
+			case BS_UpToDateWithWarnings:
+				// Fall thru and return empty string
+				break;
+			case BS_Dirty:
+				return LOCTEXT("ControlRigBP_Dirty", "Preview out of date");
+			case BS_Error:
+				return LOCTEXT("ControlRigBP_CompileError", "Compile Error");
+			default:
+				return LOCTEXT("ControlRigBP_UnknownStatus", "Unknown Status");
+			}
+		}
+
+		return FText::GetEmpty();
+	};
+
+	auto GetCompilationStateVisibility = [this]()
+	{
+		if (UBlueprint* Blueprint = GetBlueprintObj())
+		{
+			const bool bUpToDate = (Blueprint->Status == BS_UpToDate) || (Blueprint->Status == BS_UpToDateWithWarnings);
+			return bUpToDate ? EVisibility::Collapsed : EVisibility::Visible;
+		}
+
+		return EVisibility::Collapsed;
+	};
+
+	auto GetCompileButtonVisibility = [this]()
+	{
+		if (UBlueprint* Blueprint = GetBlueprintObj())
+		{
+			return (Blueprint->Status == BS_Dirty) ? EVisibility::Visible : EVisibility::Collapsed;
+		}
+
+		return EVisibility::Collapsed;
+	};
+
+	auto CompileBlueprint = [this]()
+	{
+		if (UBlueprint* Blueprint = GetBlueprintObj())
+		{
+			if (!Blueprint->IsUpToDate())
+			{
+				Compile();
+			}
+		}
+
+		return FReply::Handled();
+	};
+
+	auto GetErrorSeverity = [this]()
+	{
+		if (UBlueprint* Blueprint = GetBlueprintObj())
+		{
+			return (Blueprint->Status == BS_Error) ? EMessageSeverity::Error : EMessageSeverity::Warning;
+		}
+
+		return EMessageSeverity::Warning;
+	};
+
+	auto GetIcon = [this]()
+	{
+		if (UBlueprint* Blueprint = GetBlueprintObj())
+		{
+			return (Blueprint->Status == BS_Error) ? FEditorFontGlyphs::Exclamation_Triangle : FEditorFontGlyphs::Eye;
+		}
+
+		return FEditorFontGlyphs::Eye;
+	};
+
+	InViewport->AddNotification(MakeAttributeLambda(GetErrorSeverity),
+		false,
+		SNew(SHorizontalBox)
+		.Visibility_Lambda(GetCompilationStateVisibility)
+		+SHorizontalBox::Slot()
+		.FillWidth(1.0f)
+		.Padding(4.0f, 4.0f)
+		[
+			SNew(SHorizontalBox)
+			.ToolTipText_Lambda(GetCompilationStateText)
+			+SHorizontalBox::Slot()
+			.AutoWidth()
+			.VAlign(VAlign_Center)
+			.Padding(0.0f, 0.0f, 4.0f, 0.0f)
+			[
+				SNew(STextBlock)
+				.TextStyle(FEditorStyle::Get(), "AnimViewport.MessageText")
+				.Font(FEditorStyle::Get().GetFontStyle("FontAwesome.9"))
+				.Text_Lambda(GetIcon)
+			]
+			+SHorizontalBox::Slot()
+			.VAlign(VAlign_Center)
+			.FillWidth(1.0f)
+			[
+				SNew(STextBlock)
+				.Text_Lambda(GetCompilationStateText)
+				.TextStyle(FEditorStyle::Get(), "AnimViewport.MessageText")
+			]
+		]
+		+SHorizontalBox::Slot()
+		.AutoWidth()
+		.Padding(2.0f, 0.0f)
+		[
+			SNew(SButton)
+			.ForegroundColor(FSlateColor::UseForeground())
+			.ButtonStyle(FEditorStyle::Get(), "FlatButton.Success")
+			.Visibility_Lambda(GetCompileButtonVisibility)
+			.ToolTipText(LOCTEXT("ControlRigBPViewportCompileButtonToolTip", "Compile this Animation Blueprint to update the preview to reflect any recent changes."))
+			.OnClicked_Lambda(CompileBlueprint)
+			[
+				SNew(SHorizontalBox)
+				+SHorizontalBox::Slot()
+				.AutoWidth()
+				.VAlign(VAlign_Center)
+				.Padding(0.0f, 0.0f, 4.0f, 0.0f)
+				[
+					SNew(STextBlock)
+					.TextStyle(FEditorStyle::Get(), "AnimViewport.MessageText")
+					.Font(FEditorStyle::Get().GetFontStyle("FontAwesome.9"))
+					.Text(FEditorFontGlyphs::Cog)
+				]
+				+SHorizontalBox::Slot()
+				.VAlign(VAlign_Center)
+				.AutoWidth()
+				[
+					SNew(STextBlock)
+					.TextStyle(FEditorStyle::Get(), "AnimViewport.MessageText")
+					.Text(LOCTEXT("ControlRigBPViewportCompileButtonLabel", "Compile"))
+				]
+			]
+		]
+	);
 }
 
 void FControlRigEditor::HandlePreviewSceneCreated(const TSharedRef<IPersonaPreviewScene>& InPersonaPreviewScene)
@@ -594,6 +777,7 @@ FActionMenuContent FControlRigEditor::HandleCreateGraphActionMenu(UEdGraph* InGr
 
 void FControlRigEditor::SelectJoint(const FName& InJoint)
 {
+	// edit mode has to know
 	GetEditMode().SelectJoint(InJoint);
 	// copy locally, we use this for copying back to template when modified
 
@@ -624,15 +808,33 @@ FTransform FControlRigEditor::GetJointTransform(const FName& InJoint, bool bLoca
 
 void FControlRigEditor::SetJointTransform(const FName& InJoint, const FTransform& InTransform)
 {
+	// execution should be off
+	ensure(!ControlRig->bExecutionOn);
+
+	FScopedTransaction Transaction(LOCTEXT("Move Joint", "Move joint transform"));
+	UControlRigBlueprint* ControlRigBP = GetControlRigBlueprint();
+	ControlRigBP->Modify();
+
+	// moving ref pose warning
 	// update init/global transform
 	// @todo: this needs revision once we decide how we allow users to modify init/global transform
 	// for now, updating init/global of the joint from instances, but only modify init transform for archetype
+	// get local transform of current
+	// apply init based on parent init * current local 
+
 	ControlRig->Hierarchy.BaseHierarchy.SetInitialTransform(InJoint, InTransform);
 	ControlRig->Hierarchy.BaseHierarchy.SetGlobalTransform(InJoint, InTransform);
 
-	// update CDO  @todo - re-think about how we wrap around this nicer
-	UControlRigBlueprint* ControlRigBP = GetControlRigBlueprint();
 	ControlRigBP->Hierarchy.SetInitialTransform(InJoint, InTransform);
+	
+	UControlRigSkeletalMeshComponent* EditorSkelComp = Cast<UControlRigSkeletalMeshComponent>(GetPersonaToolkit()->GetPreviewScene()->GetPreviewMeshComponent());
+	if (EditorSkelComp)
+	{
+		EditorSkelComp->RebuildDebugDrawSkeleton();
+	}
+
+	// I don't think I have to mark dirty here. 
+	// FBlueprintEditorUtils::MarkBlueprintAsModified(GetControlRigBlueprint());
 }
 
 void FControlRigEditor::PostPasteNode(TSet<UEdGraphNode*>& PastedNodes)
@@ -691,4 +893,27 @@ void FControlRigEditor::OnFinishedChangingProperties(const FPropertyChangedEvent
 	}
 }
 
+void FControlRigEditor::OnHierarchyChanged()
+{
+	ClearDetailObject();
+
+	FBlueprintEditorUtils::MarkBlueprintAsModified(GetControlRigBlueprint());
+
+	UControlRigSkeletalMeshComponent* EditorSkelComp = Cast<UControlRigSkeletalMeshComponent>(GetPersonaToolkit()->GetPreviewScene()->GetPreviewMeshComponent());
+	if (EditorSkelComp)
+	{
+		// restart animation 
+		EditorSkelComp->InitAnim(true);
+		UpdateControlRig();
+	}
+
+	// notification
+	FNotificationInfo Info(LOCTEXT("HierarchyChangeHelpMessage", "Hierarchy has been successfully modified. If you want to move the joint, compile and turn off execution mode."));
+	Info.bFireAndForget = true;
+	Info.FadeOutDuration = 10.0f;
+	Info.ExpireDuration = 0.0f;
+
+	TSharedPtr<SNotificationItem> NotificationPtr = FSlateNotificationManager::Get().AddNotification(Info);
+	NotificationPtr->SetCompletionState(SNotificationItem::CS_Success);
+}
 #undef LOCTEXT_NAMESPACE
diff --git a/Engine/Plugins/Experimental/ControlRig/Source/ControlRigEditor/Private/Editor/ControlRigEditor.h b/Engine/Plugins/Experimental/ControlRig/Source/ControlRigEditor/Private/Editor/ControlRigEditor.h
--- a/Engine/Plugins/Experimental/ControlRig/Source/ControlRigEditor/Private/Editor/ControlRigEditor.h
+++ b/Engine/Plugins/Experimental/ControlRig/Source/ControlRigEditor/Private/Editor/ControlRigEditor.h
@@ -104,6 +104,8 @@ public:
 	void SelectJoint(const FName& InJoint);
 	// this changes everytime you compile, so don't cache it expecting it will last. 
 	UControlRig* GetInstanceRig() const { return ControlRig;  }
+	// restart animation 
+	void OnHierarchyChanged();
 protected:
 	// FBlueprintEditor Interface
 	virtual void CreateDefaultCommands() override;
@@ -166,6 +168,9 @@ private:
 	/** Wraps the normal blueprint editor's action menu creation callback */
 	FActionMenuContent HandleCreateGraphActionMenu(UEdGraph* InGraph, const FVector2D& InNodePosition, const TArray<UEdGraphPin*>& InDraggedPins, bool bAutoExpand, SGraphEditor::FActionMenuClosed InOnMenuClosed);
 
+	void ToggleExecuteGraph();
+	bool IsExecuteGraphOn() const;
+
 protected:
 
 	/** Toolbox hosting widget */
@@ -188,6 +193,9 @@ protected:
 	/** Our currently running control rig instance */
 	UControlRig* ControlRig;
 
+	/** preview scene */
+	TSharedPtr<IPersonaPreviewScene> PreviewScene;
+
 	/** Recursion guard for selection */
 	bool bSelecting;
 
diff --git a/Engine/Plugins/Experimental/ControlRig/Source/ControlRigEditor/Private/Editor/SRigHierarchy.cpp b/Engine/Plugins/Experimental/ControlRig/Source/ControlRigEditor/Private/Editor/SRigHierarchy.cpp
--- a/Engine/Plugins/Experimental/ControlRig/Source/ControlRigEditor/Private/Editor/SRigHierarchy.cpp
+++ b/Engine/Plugins/Experimental/ControlRig/Source/ControlRigEditor/Private/Editor/SRigHierarchy.cpp
@@ -322,6 +322,12 @@ void SRigHierarchy::OnSelectionChanged(TSharedPtr<FRigTreeJoint> Selection, ESel
 				ControlRigEditor.Pin()->SelectJoint(Selection->CachedJoint);
 				return;
 			}
+			else
+			{
+				// clear the current selection
+				ControlRigEditor.Pin()->ClearDetailObject();
+				ControlRigEditor.Pin()->SelectJoint(NAME_None);
+			}
 		}
 
 		// if failed, try BP hierarhcy? Todo:
@@ -450,6 +456,7 @@ void SRigHierarchy::RefreshHierarchy(const FAssetData& InAssetData)
 			}
 		}
 
+		ControlRigEditor.Pin()->OnHierarchyChanged();
 		RefreshTreeView();
 		FSlateApplication::Get().DismissAllMenus();
 	}
@@ -505,6 +512,7 @@ void SRigHierarchy::ImportHierarchy(const FAssetData& InAssetData)
 			}
 		}
 
+		ControlRigEditor.Pin()->OnHierarchyChanged();
 		RefreshTreeView();
 		FSlateApplication::Get().DismissAllMenus();
 	}
@@ -543,6 +551,7 @@ void SRigHierarchy::HandleDeleteItem()
 			}
 		}
 
+		ControlRigEditor.Pin()->OnHierarchyChanged();
 		RefreshTreeView();
  	}
 }
@@ -578,7 +587,7 @@ void SRigHierarchy::HandleNewItem()
 		Hierarchy->AddJoint(NewJointName, ParentName, ParentTransform);
 
 		RefreshTreeView();
-
+		ControlRigEditor.Pin()->OnHierarchyChanged();
 		// reselect current selected item
 		SelectJoint(NewJointName);
 	}
@@ -616,6 +625,7 @@ void SRigHierarchy::HandleDuplicateItem()
 		}
 
 		RefreshTreeView();
+		ControlRigEditor.Pin()->OnHierarchyChanged();
 		
 		for (int32 Index = 0; Index < NewNames.Num(); ++Index)
 		{
@@ -709,6 +719,7 @@ bool SRigHierarchy::RenameJoint(const FName& OldName, const FName& NewName)
 		Hierarchy->Rename(OldName, NewName);
 		SelectJoint(NewName);
 
+		ControlRigEditor.Pin()->OnHierarchyChanged();
 		return true;
 	}
 
diff --git a/Engine/Source/Editor/Persona/Private/AnimationEditorViewportClient.cpp b/Engine/Source/Editor/Persona/Private/AnimationEditorViewportClient.cpp
--- a/Engine/Source/Editor/Persona/Private/AnimationEditorViewportClient.cpp
+++ b/Engine/Source/Editor/Persona/Private/AnimationEditorViewportClient.cpp
@@ -1267,9 +1267,9 @@ void FAnimationViewportClient::DrawMeshBones(UDebugSkelMeshComponent * MeshCompo
 void FAnimationViewportClient::DrawBones(const TArray<FBoneIndexType>& RequiredBones, const FReferenceSkeleton& RefSkeleton, const TArray<FTransform> & WorldTransforms, const TArray<int32>& InSelectedBones, FPrimitiveDrawInterface* PDI, const TArray<FLinearColor>& BoneColours, float LineThickness/*=0.f*/, bool bForceDraw/*=false*/) const
 {
 	TArray<int32> SelectedBones = InSelectedBones;
-	if(GetBoneDrawMode() == EBoneDrawMode::SelectedAndParents)
+	if(InSelectedBones.Num() > 0 && GetBoneDrawMode() == EBoneDrawMode::SelectedAndParents)
 	{
-		int32 BoneIndex = GetAnimPreviewScene()->GetSelectedBoneIndex();
+		int32 BoneIndex = InSelectedBones[0];
 		while (BoneIndex != INDEX_NONE)
 		{
 			int32 ParentIndex = RefSkeleton.GetParentIndex(BoneIndex);
@@ -1306,9 +1306,11 @@ void FAnimationViewportClient::DrawBones(const TArray<FBoneIndexType>& RequiredB
 				End = WorldTransforms[BoneIndex].GetLocation();
 			}
 
+			const float BoneLength = (End - Start).Size();
+			const float AxisLength = FMath::Max(BoneLength * 0.1f, 1.f);
 			//Render Sphere for bone end point and a cone between it and its parent.
 			PDI->SetHitProxy(new HPersonaBoneProxy(RefSkeleton.GetBoneName(BoneIndex)));
-			SkeletalDebugRendering::DrawWireBone(PDI, Start, End, LineColor, SDPG_Foreground);
+			SkeletalDebugRendering::DrawWireBone(PDI, Start, End, LineColor, SDPG_Foreground, AxisLength);
 			PDI->SetHitProxy(NULL);
 
 			// draw gizmo
@@ -1316,7 +1318,8 @@ void FAnimationViewportClient::DrawBones(const TArray<FBoneIndexType>& RequiredB
 				((GetLocalAxesMode() == ELocalAxesMode::Selected) && SelectedBones.Contains(BoneIndex))
 				)
 			{
-				SkeletalDebugRendering::DrawAxes(PDI, WorldTransforms[BoneIndex], SDPG_Foreground);
+				// we want to say 10 % of bone length is good
+				SkeletalDebugRendering::DrawAxes(PDI, WorldTransforms[BoneIndex], SDPG_Foreground, 0.f , AxisLength);
 			}
 		}
 	}
diff --git a/Engine/Source/Runtime/Engine/Private/SkeletalDebugRendering.cpp b/Engine/Source/Runtime/Engine/Private/SkeletalDebugRendering.cpp
--- a/Engine/Source/Runtime/Engine/Private/SkeletalDebugRendering.cpp
+++ b/Engine/Source/Runtime/Engine/Private/SkeletalDebugRendering.cpp
@@ -7,10 +7,9 @@
 namespace SkeletalDebugRendering
 {
 
-void DrawWireBone(FPrimitiveDrawInterface* PDI, const FVector& InStart, const FVector& InEnd, const FLinearColor& InColor, ESceneDepthPriorityGroup InDepthPriority)
+void DrawWireBone(FPrimitiveDrawInterface* PDI, const FVector& InStart, const FVector& InEnd, const FLinearColor& InColor, ESceneDepthPriorityGroup InDepthPriority, const float SphereRadius)
 {
 #if ENABLE_DRAW_DEBUG
-	static const float SphereRadius = 1.0f;
 	static const int32 NumSphereSides = 10;
 	static const int32 NumConeSides = 4;
 
@@ -27,27 +26,26 @@ void DrawWireBone(FPrimitiveDrawInterface* PDI, const FVector& InStart, const FV
 #endif
 }
 
-void DrawAxes(FPrimitiveDrawInterface* PDI, const FTransform& Transform, ESceneDepthPriorityGroup InDepthPriority)
+void DrawAxes(FPrimitiveDrawInterface* PDI, const FTransform& Transform, ESceneDepthPriorityGroup InDepthPriority, const float Thickness, const float AxisLength)
 {
 #if ENABLE_DRAW_DEBUG
 	// Display colored coordinate system axes for this joint.
-	const float AxisLength = 4.0f;
 	const FVector Origin = Transform.GetLocation();
 
 	// Red = X
 	FVector XAxis = Transform.TransformVector(FVector(1.0f, 0.0f, 0.0f));
 	XAxis.Normalize();
-	PDI->DrawLine(Origin, Origin + XAxis * AxisLength, FColor(255, 80, 80), InDepthPriority, 0.0f, 1.0f);
+	PDI->DrawLine(Origin, Origin + XAxis * AxisLength, FColor(255, 80, 80), InDepthPriority, Thickness, 1.0f);
 
 	// Green = Y
 	FVector YAxis = Transform.TransformVector(FVector(0.0f, 1.0f, 0.0f));
 	YAxis.Normalize();
-	PDI->DrawLine(Origin, Origin + YAxis * AxisLength, FColor(80, 255, 80), InDepthPriority, 0.0f, 1.0f); 
+	PDI->DrawLine(Origin, Origin + YAxis * AxisLength, FColor(80, 255, 80), InDepthPriority, Thickness, 1.0f);
 
 	// Blue = Z
 	FVector ZAxis = Transform.TransformVector(FVector(0.0f, 0.0f, 1.0f));
 	ZAxis.Normalize();
-	PDI->DrawLine(Origin, Origin + ZAxis * AxisLength, FColor(80, 80, 255), InDepthPriority, 0.0f, 1.0f);
+	PDI->DrawLine(Origin, Origin + ZAxis * AxisLength, FColor(80, 80, 255), InDepthPriority, Thickness, 1.0f);
 #endif
 }
 
diff --git a/Engine/Source/Runtime/Engine/Public/SkeletalDebugRendering.h b/Engine/Source/Runtime/Engine/Public/SkeletalDebugRendering.h
--- a/Engine/Source/Runtime/Engine/Public/SkeletalDebugRendering.h
+++ b/Engine/Source/Runtime/Engine/Public/SkeletalDebugRendering.h
@@ -17,15 +17,18 @@ namespace SkeletalDebugRendering
  * @param	InEnd				The end location of the bone
  * @param	InColor				The color to draw the bone with
  * @param	InDepthPriority		The scene depth priority group to use
+ * @param	SphereRadius		SphereRadius
  */
-ENGINE_API void DrawWireBone(FPrimitiveDrawInterface* PDI, const FVector& InStart, const FVector& InEnd, const FLinearColor& InColor, ESceneDepthPriorityGroup InDepthPriority);
+ENGINE_API void DrawWireBone(FPrimitiveDrawInterface* PDI, const FVector& InStart, const FVector& InEnd, const FLinearColor& InColor, ESceneDepthPriorityGroup InDepthPriority, const float SphereRadius = 1.f);
 
 /**
  * Draw a set of axes to represent a transform
  * @param	PDI					Primitive draw interface to use
  * @param	InTransform			The transform to represent
  * @param	InDepthPriority		The scene depth priority group to use
+ * @param	Thickness			Thickness of Axes
+ * @param	AxisLength			AxisLength
  */
-ENGINE_API void DrawAxes(FPrimitiveDrawInterface* PDI, const FTransform& InTransform, ESceneDepthPriorityGroup InDepthPriority);
+ENGINE_API void DrawAxes(FPrimitiveDrawInterface* PDI, const FTransform& InTransform, ESceneDepthPriorityGroup InDepthPriority, const float Thickness = 0.f, const float AxisLength = 4.f);
 
 }
\ No newline at end of file
