fx the build error

diff --git a/Engine/Plugins/Experimental/ControlRig/Source/ControlRig/Public/Units/RigUnit.h b/Engine/Plugins/Experimental/ControlRig/Source/ControlRig/Public/Units/RigUnit.h
--- a/Engine/Plugins/Experimental/ControlRig/Source/ControlRig/Public/Units/RigUnit.h
+++ b/Engine/Plugins/Experimental/ControlRig/Source/ControlRig/Public/Units/RigUnit.h
@@ -53,6 +53,6 @@ struct CONTROLRIG_API FRigUnit
 // this will have to change in the future and move to editor, I assume the errors will be saved in the rig unit and it will print fromthe editor module
 namespace UnitLogHelpers
 {
-	void PrintMissingHierarchy(const FName& InputName);
-	void PrintUnimplemented(const FName& InputName);
+	CONTROLRIG_API void PrintMissingHierarchy(const FName& InputName);
+	CONTROLRIG_API void PrintUnimplemented(const FName& InputName);
 }
