fix build error - missing file

diff --git a/Engine/Plugins/Experimental/ControlRig/Source/ControlRig/Private/Units/RigUnit.cpp b/Engine/Plugins/Experimental/ControlRig/Source/ControlRig/Private/Units/RigUnit.cpp
new file mode 100644
--- /dev/null
+++ b/Engine/Plugins/Experimental/ControlRig/Source/ControlRig/Private/Units/RigUnit.cpp
@@ -0,0 +1,21 @@
+// Copyright 1998-2018 Epic Games, Inc. All Rights Reserved.
+
+#pragma once
+
+#include "Units/RigUnit.h"
+
+DEFINE_LOG_CATEGORY_STATIC(LogRigUnit, Log, All);
+
+// this will have to change in the future and move to editor, I assume the errors will be saved in the rig unit and it will print fromthe editor module
+namespace UnitLogHelpers
+{
+	void PrintMissingHierarchy(const FName& InputName)
+	{
+		UE_LOG(LogRigUnit, Warning, TEXT("%s: Input Hierarch Link is missing"), *InputName.ToString());
+	}
+
+	void PrintUnimplemented(const FName& InputName)
+	{
+		UE_LOG(LogRigUnit, Warning, TEXT("%s: Not implemented"), *InputName.ToString());
+	}
+}
