fix missing image file, and fix linux error

diff --git a/Engine/Plugins/Experimental/ControlRig/Source/ControlRig/Private/Units/RigUnit.cpp b/Engine/Plugins/Experimental/ControlRig/Source/ControlRig/Private/Units/RigUnit.cpp
--- a/Engine/Plugins/Experimental/ControlRig/Source/ControlRig/Private/Units/RigUnit.cpp
+++ b/Engine/Plugins/Experimental/ControlRig/Source/ControlRig/Private/Units/RigUnit.cpp
@@ -1,7 +1,5 @@
 // Copyright 1998-2018 Epic Games, Inc. All Rights Reserved.
 
-#pragma once
-
 #include "Units/RigUnit.h"
 
 DEFINE_LOG_CATEGORY_STATIC(LogRigUnit, Log, All);
