CIS fixes after the last merge from Main

diff --git a/Engine/Plugins/Experimental/ControlRig/Source/ControlRigDeveloper/Private/ControlRigBlueprintCompiler.cpp b/Engine/Plugins/Experimental/ControlRig/Source/ControlRigDeveloper/Private/ControlRigBlueprintCompiler.cpp
--- a/Engine/Plugins/Experimental/ControlRig/Source/ControlRigDeveloper/Private/ControlRigBlueprintCompiler.cpp
+++ b/Engine/Plugins/Experimental/ControlRig/Source/ControlRigDeveloper/Private/ControlRigBlueprintCompiler.cpp
@@ -22,9 +22,9 @@ bool FControlRigBlueprintCompiler::CanCompile(const UBlueprint* Blueprint)
 	return false;
 }
 
-void FControlRigBlueprintCompiler::Compile(UBlueprint* Blueprint, const FKismetCompilerOptions& CompileOptions, FCompilerResultsLog& Results, TArray<UObject*>* ObjLoaded)
+void FControlRigBlueprintCompiler::Compile(UBlueprint* Blueprint, const FKismetCompilerOptions& CompileOptions, FCompilerResultsLog& Results)
 {
-	FControlRigBlueprintCompilerContext Compiler(Blueprint, Results, CompileOptions, ObjLoaded);
+	FControlRigBlueprintCompilerContext Compiler(Blueprint, Results, CompileOptions);
 	Compiler.Compile();
 }
 
diff --git a/Engine/Plugins/Experimental/ControlRig/Source/ControlRigDeveloper/Public/ControlRigBlueprintCompiler.h b/Engine/Plugins/Experimental/ControlRig/Source/ControlRigDeveloper/Public/ControlRigBlueprintCompiler.h
--- a/Engine/Plugins/Experimental/ControlRig/Source/ControlRigDeveloper/Public/ControlRigBlueprintCompiler.h
+++ b/Engine/Plugins/Experimental/ControlRig/Source/ControlRigDeveloper/Public/ControlRigBlueprintCompiler.h
@@ -13,14 +13,14 @@ class CONTROLRIGDEVELOPER_API FControlRigBlueprintCompiler : public IBlueprintCo
 public:
 	/** IBlueprintCompiler interface */
 	virtual bool CanCompile(const UBlueprint* Blueprint) override;
-	virtual void Compile(UBlueprint* Blueprint, const FKismetCompilerOptions& CompileOptions, FCompilerResultsLog& Results, TArray<UObject*>* ObjLoaded) override;
+	virtual void Compile(UBlueprint* Blueprint, const FKismetCompilerOptions& CompileOptions, FCompilerResultsLog& Results) override;
 };
 
 class CONTROLRIGDEVELOPER_API FControlRigBlueprintCompilerContext : public FKismetCompilerContext
 {
 public:
-	FControlRigBlueprintCompilerContext(UBlueprint* SourceSketch, FCompilerResultsLog& InMessageLog, const FKismetCompilerOptions& InCompilerOptions, TArray<UObject*>* InObjLoaded)
-		: FKismetCompilerContext(SourceSketch, InMessageLog, InCompilerOptions, InObjLoaded)
+	FControlRigBlueprintCompilerContext(UBlueprint* SourceSketch, FCompilerResultsLog& InMessageLog, const FKismetCompilerOptions& InCompilerOptions)
+		: FKismetCompilerContext(SourceSketch, InMessageLog, InCompilerOptions)
 		, NewControlRigBlueprintGeneratedClass(nullptr)
 	{
 	}
diff --git a/Engine/Source/Editor/BlueprintGraph/Private/K2Node.cpp b/Engine/Source/Editor/BlueprintGraph/Private/K2Node.cpp
--- a/Engine/Source/Editor/BlueprintGraph/Private/K2Node.cpp
+++ b/Engine/Source/Editor/BlueprintGraph/Private/K2Node.cpp
@@ -184,7 +184,7 @@ void UK2Node::FixupPinDefaultValues()
 				if (GIsEditor && !Pin->DefaultValue.IsEmpty())
 				{
 					FSoftObjectPath TempRef(Pin->DefaultValue);
-					TempRef.PostLoadPath();
+					TempRef.PostLoadPath(GetLinker());
 					TempRef.PreSavePath();
 					Pin->DefaultValue = TempRef.ToString();
 				}
diff --git a/Engine/Source/Editor/BlueprintGraph/Private/K2Node_FunctionEntry.cpp b/Engine/Source/Editor/BlueprintGraph/Private/K2Node_FunctionEntry.cpp
--- a/Engine/Source/Editor/BlueprintGraph/Private/K2Node_FunctionEntry.cpp
+++ b/Engine/Source/Editor/BlueprintGraph/Private/K2Node_FunctionEntry.cpp
@@ -350,7 +350,7 @@ void UK2Node_FunctionEntry::Serialize(FArchive& Ar)
 				if (!LocalVariable.DefaultValue.IsEmpty() && (LocalVariable.VarType.PinCategory == UEdGraphSchema_K2::PC_SoftObject || LocalVariable.VarType.PinCategory == UEdGraphSchema_K2::PC_SoftClass))
 				{
 					FSoftObjectPath TempRef(LocalVariable.DefaultValue);
-					TempRef.PostLoadPath();
+					TempRef.PostLoadPath(&Ar);
 					TempRef.PreSavePath();
 					LocalVariable.DefaultValue = TempRef.ToString();
 				}
diff --git a/Engine/Source/Editor/UnrealEd/Private/Factories/SkeletalMeshImport.cpp b/Engine/Source/Editor/UnrealEd/Private/Factories/SkeletalMeshImport.cpp
--- a/Engine/Source/Editor/UnrealEd/Private/Factories/SkeletalMeshImport.cpp
+++ b/Engine/Source/Editor/UnrealEd/Private/Factories/SkeletalMeshImport.cpp
@@ -867,8 +867,6 @@ void RestoreExistingSkelMeshData(ExistingSkelMeshData* MeshData, USkeletalMesh*
 					}
 					FSkeletalMeshLODModel& LODModel = MeshData->ExistingLODModels[Index];
 					FSkeletalMeshLODInfo& LODInfo = MeshData->ExistingLODInfo[Index];
-					FSkeletalMeshLODModel* NewLODModel = new FSkeletalMeshLODModel(LODModel);
-					SkeletalMeshImportedModel->LODModels.Add(NewLODModel);
 
 					FSkeletalMeshLODModel* NewLODModel = nullptr;
 					if (RealIndex >= SkeletalMesh->GetLODNum())
