Control Rig: Now() rig unit

diff --git a/Engine/Plugins/Experimental/ControlRig/Source/ControlRig/Private/Units/Animation/RigUnit_GetWorldTime.cpp b/Engine/Plugins/Experimental/ControlRig/Source/ControlRig/Private/Units/Animation/RigUnit_GetWorldTime.cpp
new file mode 100644
--- /dev/null
+++ b/Engine/Plugins/Experimental/ControlRig/Source/ControlRig/Private/Units/Animation/RigUnit_GetWorldTime.cpp
@@ -0,0 +1,18 @@
+// Copyright 1998-2019 Epic Games, Inc. All Rights Reserved.
+
+#include "Units/Animation/RigUnit_GetWorldTime.h"
+#include "Units/RigUnitContext.h"
+
+void FRigUnit_GetWorldTime::Execute(const FRigUnitContext& Context)
+{
+	FDateTime Now = FDateTime::Now();
+	Year = (float)Now.GetYear();
+	Month = (float)Now.GetMonth();
+	Day = (float)Now.GetDay();
+	WeekDay = (float)Now.GetDayOfWeek();
+	Hours = (float)Now.GetHour();
+	Minutes = (float)Now.GetMinute();
+	Seconds = float(Now.GetSecond()) + float(Now.GetMillisecond()) * 0.001f;
+	OverallSeconds = (float)FPlatformTime::Seconds();
+}
+
diff --git a/Engine/Plugins/Experimental/ControlRig/Source/ControlRig/Private/Units/Animation/RigUnit_GetWorldTime.h b/Engine/Plugins/Experimental/ControlRig/Source/ControlRig/Private/Units/Animation/RigUnit_GetWorldTime.h
new file mode 100644
--- /dev/null
+++ b/Engine/Plugins/Experimental/ControlRig/Source/ControlRig/Private/Units/Animation/RigUnit_GetWorldTime.h
@@ -0,0 +1,47 @@
+// Copyright 1998-2019 Epic Games, Inc. All Rights Reserved.
+
+#pragma once
+
+#include "Units/Animation/RigUnit_AnimBase.h"
+#include "RigUnit_GetWorldTime.generated.h"
+
+/**
+ * A constant value of an easing type
+ */
+USTRUCT(meta = (DisplayName = "Now", Keywords = "Time,Clock"))
+struct FRigUnit_GetWorldTime : public FRigUnit_AnimBase
+{
+	GENERATED_BODY()
+
+	FRigUnit_GetWorldTime()
+	{
+		Year = Month = Day = WeekDay = Hours = Minutes = Seconds = OverallSeconds = 0.f;
+	}
+
+	virtual void Execute(const FRigUnitContext& Context) override;
+
+	UPROPERTY(meta = (Output))
+	float Year;
+
+	UPROPERTY(meta = (Output))
+	float Month;
+
+	UPROPERTY(meta = (Output))
+	float Day;
+
+	UPROPERTY(meta = (Output))
+	float WeekDay;
+
+	UPROPERTY(meta = (Output))
+	float Hours;
+
+	UPROPERTY(meta = (Output))
+	float Minutes;
+
+	UPROPERTY(meta = (Output))
+	float Seconds;
+
+	UPROPERTY(meta = (Output))
+	float OverallSeconds;
+};
+
