Control Rig: Add profiling bracket units to do simple profiling.

diff --git a/Engine/Plugins/Experimental/ControlRig/Source/ControlRig/Private/Units/Debug/RigUnit_ProfilingBracket.cpp b/Engine/Plugins/Experimental/ControlRig/Source/ControlRig/Private/Units/Debug/RigUnit_ProfilingBracket.cpp
new file mode 100644
--- /dev/null
+++ b/Engine/Plugins/Experimental/ControlRig/Source/ControlRig/Private/Units/Debug/RigUnit_ProfilingBracket.cpp
@@ -0,0 +1,45 @@
+// Copyright 1998-2019 Epic Games, Inc. All Rights Reserved.
+
+#include "Units/Debug/RigUnit_ProfilingBracket.h"
+#include "Units/RigUnitContext.h"
+#include "KismetAnimationLibrary.h"
+
+void FRigUnit_StartProfilingTimer::Execute(const FRigUnitContext& Context)
+{
+	if (Context.State == EControlRigState::Update)
+	{
+		UKismetAnimationLibrary::K2_StartProfilingTimer();
+	}
+}
+
+void FRigUnit_EndProfilingTimer::Execute(const FRigUnitContext& Context)
+{
+	if (Context.State == EControlRigState::Init)
+	{
+		AccumulatedTime = 0.f;
+		MeasurementsLeft = FMath::Max(NumberOfMeasurements, 1);
+		return;
+	}
+
+	if (Context.State == EControlRigState::Update)
+	{
+		float Delta = UKismetAnimationLibrary::K2_EndProfilingTimer(false);
+		if(MeasurementsLeft > 0)
+		{
+			AccumulatedTime += Delta / float(FMath::Max(NumberOfMeasurements, 1));
+			MeasurementsLeft--;
+
+			if (MeasurementsLeft == 0)
+			{
+				if (Prefix.IsEmpty())
+				{
+					UE_CONTROLRIG_RIGUNIT_LOG_MESSAGE(TEXT("%0.3f ms (%d runs)."), AccumulatedTime, FMath::Max(NumberOfMeasurements, 1));
+				}
+				else
+				{
+					UE_CONTROLRIG_RIGUNIT_LOG_MESSAGE(TEXT("[%s] %0.3f ms (%d runs)."), *Prefix, AccumulatedTime, FMath::Max(NumberOfMeasurements, 1));
+				}
+			}
+		}
+	}
+}
\ No newline at end of file
diff --git a/Engine/Plugins/Experimental/ControlRig/Source/ControlRig/Private/Units/Debug/RigUnit_ProfilingBracket.h b/Engine/Plugins/Experimental/ControlRig/Source/ControlRig/Private/Units/Debug/RigUnit_ProfilingBracket.h
new file mode 100644
--- /dev/null
+++ b/Engine/Plugins/Experimental/ControlRig/Source/ControlRig/Private/Units/Debug/RigUnit_ProfilingBracket.h
@@ -0,0 +1,46 @@
+// Copyright 1998-2019 Epic Games, Inc. All Rights Reserved.
+
+#pragma once
+
+#include "Units/Debug/RigUnit_DebugBase.h"
+#include "RigUnit_ProfilingBracket.generated.h"
+
+USTRUCT(meta=(DisplayName="Start Profiling Timer", Keywords="Measure,BeginProfiling,Profile"))
+struct FRigUnit_StartProfilingTimer : public FRigUnit_DebugBaseMutable
+{
+	GENERATED_BODY()
+
+	FRigUnit_StartProfilingTimer()
+	{
+	}
+
+	virtual void Execute(const FRigUnitContext& Context) override;
+};
+
+USTRUCT(meta = (DisplayName = "End Profiling Timer", Keywords = "Measure,StopProfiling,Meter,Profile"))
+struct FRigUnit_EndProfilingTimer : public FRigUnit_DebugBaseMutable
+{
+	GENERATED_BODY()
+
+	FRigUnit_EndProfilingTimer()
+	{
+		NumberOfMeasurements = 1;
+		AccumulatedTime = 0.f;
+		MeasurementsLeft = 0;
+		Prefix = TEXT("Timer");
+	}
+
+	virtual void Execute(const FRigUnitContext& Context) override;
+
+	UPROPERTY(meta = (Input, Constant))
+	int32 NumberOfMeasurements;
+
+	UPROPERTY(meta = (Input, Constant))
+	FString Prefix;
+
+	UPROPERTY()
+	float AccumulatedTime;
+
+	UPROPERTY()
+	int32 MeasurementsLeft;
+};
