Control Rig MVC Refactoring

diff --git a/Engine/Plugins/Experimental/ControlRig/Source/ControlRigDeveloper/Private/ControlRigController.cpp b/Engine/Plugins/Experimental/ControlRig/Source/ControlRigDeveloper/Private/ControlRigController.cpp
new file mode 100644
--- /dev/null
+++ b/Engine/Plugins/Experimental/ControlRig/Source/ControlRigDeveloper/Private/ControlRigController.cpp
@@ -0,0 +1,1169 @@
+// Copyright 1998-2019 Epic Games, Inc. All Rights Reserved.
+
+#include "ControlRigController.h"
+#include "EdGraphSchema_K2.h"
+#include "Logging/MessageLog.h"
+
+#define LOCTEXT_NAMESPACE "ControlRigController"
+
+UControlRigController::UControlRigController()
+	: Model(nullptr)
+	, bSuspendLog(false)
+{
+}
+
+UControlRigController::~UControlRigController()
+{
+	_ModifiedEvent.Clear();
+}
+
+void UControlRigController::SetModel(UControlRigModel* InModel)
+{
+	if (Model != nullptr)
+	{
+		Model->OnModified().Remove(_ModelModifiedHandle);
+	}
+
+	Model = InModel;
+
+	if (Model != nullptr)
+	{
+		_ModelModifiedHandle = Model->OnModified().AddUObject(this, &UControlRigController::HandleModelModified);
+	}
+}
+
+bool UControlRigController::Clear()
+{
+	return Model->Clear();
+}
+
+void UControlRigController::EnableMessageLog(bool bInEnabled)
+{
+	bSuspendLog = !bInEnabled;
+}
+
+#if CONTROLRIG_UNDO
+
+bool UControlRigController::OpenUndoBracket(const FString& Title)
+{
+	if (!EnsureModel())
+	{
+		return false;
+	}
+
+	TSharedPtr<UControlRigModel::FAction> Action = MakeShared<UControlRigModel::FAction>();
+	Action->Title = Title;
+	Action->Type = EControlRigModelNotifType::Invalid;
+
+	_UndoBrackets.Add(Action);
+	Model->CurrentActions.Add(Action.Get());
+	return true;
+}
+
+bool UControlRigController::CloseUndoBracket()
+{
+	if (!EnsureModel())
+	{
+		return false;
+	}
+
+	if (_UndoBrackets.Num() == 0)
+	{
+		return false;
+	}
+
+	Model->CurrentActions.Pop();
+	TSharedPtr<UControlRigModel::FAction> Action = _UndoBrackets.Pop();
+	Model->PushAction(*Action);
+	return true;
+}
+
+#endif
+
+
+UControlRigModel::FModifiedEvent& UControlRigController::OnModified()
+{
+	return _ModifiedEvent;
+}
+
+void UControlRigController::HandleModelModified(const UControlRigModel* InModel, EControlRigModelNotifType InType, const void* InPayload)
+{
+	_LastModelNotification = InType;
+
+	switch (InType)
+	{
+		case EControlRigModelNotifType::NodeAdded:
+		{
+			const FControlRigModelNode* Node = (const FControlRigModelNode*)InPayload;
+			if (Node)
+			{
+				if (Node->IsParameter())
+				{
+					LogMessage(FString::Printf(TEXT("Added parameter '%s' at (%.01f, %.01f)"), *Node->Name.ToString(), Node->Position.X, Node->Position.Y));
+				}
+				else
+				{
+					LogMessage(FString::Printf(TEXT("Added node '%s' of type '%s' at (%.01f, %.01f)"), *Node->Name.ToString(), *Node->FunctionName.ToString(), Node->Position.X, Node->Position.Y));
+				}
+			}
+			break;
+		}
+		case EControlRigModelNotifType::NodeRemoved:
+		{
+			const FControlRigModelNode* Node = (const FControlRigModelNode*)InPayload;
+			if (Node)
+			{
+				LogMessage(FString::Printf(TEXT("Removed node '%s'."), *Node->Name.ToString()));
+			}
+			break;
+		}
+		case EControlRigModelNotifType::NodeRenamed:
+		{
+			const FControlRigModelNodeRenameInfo* Info= (const FControlRigModelNodeRenameInfo*)InPayload;
+			if (Info)
+			{
+				LogMessage(FString::Printf(TEXT("Renamed node '%s' to '%s'."), *Info->OldName.ToString(), *Info->NewName.ToString()));
+			}
+			break;
+		}
+		case EControlRigModelNotifType::NodeSelected:
+		{
+			const FControlRigModelNode* Node = (const FControlRigModelNode*)InPayload;
+			if (Node)
+			{
+				LogMessage(FString::Printf(TEXT("Selected node '%s'."), *Node->Name.ToString()));
+			}
+			break;
+		}
+		case EControlRigModelNotifType::NodeDeselected:
+		{
+			const FControlRigModelNode* Node = (const FControlRigModelNode*)InPayload;
+			if (Node)
+			{
+				LogMessage(FString::Printf(TEXT("Deselected node '%s'."), *Node->Name.ToString()));
+			}
+			break;
+		}
+		case EControlRigModelNotifType::NodeChanged:
+		{
+// 			const FControlRigModelNode* Node = (const FControlRigModelNode*)InPayload;
+// 			if (Node)
+// 			{
+// 				LogMessage(FString::Printf(TEXT("Changed node '%s'."), *Node->Name.ToString()));
+// 			}
+			break;
+		}
+		case EControlRigModelNotifType::LinkAdded:
+		{
+			const FControlRigModelLink* Link = (const FControlRigModelLink*)InPayload;
+			if (Link)
+			{
+				FString SourcePinPath = InModel->GetPinPath(Link->Source);
+				FString TargetPinPath = InModel->GetPinPath(Link->Target);
+				LogMessage(FString::Printf(TEXT("Added link '%s' to '%s'."), *SourcePinPath, *TargetPinPath));
+			}
+			break;
+		}
+		case EControlRigModelNotifType::LinkRemoved:
+		{
+			const FControlRigModelLink* Link = (const FControlRigModelLink*)InPayload;
+			if (Link)
+			{
+				FString SourcePinPath = InModel->GetPinPath(Link->Source);
+				FString TargetPinPath = InModel->GetPinPath(Link->Target);
+				LogMessage(FString::Printf(TEXT("Removed link '%s' to '%s'."), *SourcePinPath, *TargetPinPath));
+			}
+			break;
+		}
+		case EControlRigModelNotifType::PinAdded:
+		{
+			const FControlRigModelPin* Pin = (const FControlRigModelPin*)InPayload;
+			if (Pin)
+			{
+				FString PinPath = InModel->GetPinPath(Pin->GetPair());
+				LogMessage(FString::Printf(TEXT("Added pin '%s'."), *PinPath));
+			}
+			break;
+		}
+		case EControlRigModelNotifType::PinRemoved:
+		{
+			const FControlRigModelPin* Pin = (const FControlRigModelPin*)InPayload;
+			if (Pin)
+			{
+				const FControlRigModelNode& Node = InModel->Nodes()[Pin->Node];
+				const FControlRigModelPin& ParentPin = Node.Pins[Pin->ParentIndex];
+				FString PinPath = InModel->GetPinPath(ParentPin.GetPair(), true);
+				LogMessage(FString::Printf(TEXT("Removed pin '%s.%s'."), *PinPath, *Pin->Name.ToString()));
+			}
+			break;
+		}
+		case EControlRigModelNotifType::PinChanged:
+		{
+			const FControlRigModelPin* Pin = (const FControlRigModelPin*)InPayload;
+			if (Pin)
+			{
+				FString PinPath = InModel->GetPinPath(Pin->GetPair());
+				if (Pin->DefaultValue.IsEmpty() || Pin->SubPins.Num() > 0)
+				{
+					LogMessage(FString::Printf(TEXT("Changed pin '%s'."), *PinPath));
+				}
+				else
+				{
+					LogMessage(FString::Printf(TEXT("Changed pin '%s', default '%s'."), *PinPath, *Pin->DefaultValue));
+				}
+			}
+			break;
+		}
+		case EControlRigModelNotifType::ModelError:
+		{
+			const FControlRigModelError* Error = (const FControlRigModelError*)InPayload;
+			if (Error)
+			{
+				LogError(Error->Message);
+			}
+			break;
+		}
+		default:
+		{
+			break;
+		}
+	}
+	
+	if (_ModifiedEvent.IsBound())
+	{
+		_ModifiedEvent.Broadcast(InModel, InType, InPayload);
+	}
+}
+
+bool UControlRigController::EnsureModel() const
+{
+	if (Model == nullptr)
+	{
+		if (_ModifiedEvent.IsBound())
+		{
+			FControlRigModelError Error;
+			Error.Message = TEXT("No model set on the controller.");
+			_ModifiedEvent.Broadcast(nullptr, EControlRigModelNotifType::ModelError, &Error);
+			LogError(Error.Message);
+		}
+		return false;
+	}
+	return true;
+}
+
+void UControlRigController::LogMessage(const FString& Message) const
+{
+	if (!bSuspendLog)
+	{
+		FMessageLog ControlRigLog("ControlRigLog");
+		ControlRigLog.Info(FText::FromString(Message));
+	}
+}
+
+void UControlRigController::LogWarning(const FString& Message) const
+{
+	if (!bSuspendLog)
+	{
+		FMessageLog ControlRigLog("ControlRigLog");
+		ControlRigLog.Warning(FText::FromString(Message));
+	}
+}
+
+void UControlRigController::LogError(const FString& Message) const
+{
+	if (!bSuspendLog)
+	{
+		FMessageLog ControlRigLog("ControlRigLog");
+		ControlRigLog.Error(FText::FromString(Message));
+	}
+}
+
+bool UControlRigController::ConstructPreviewParameter(const FName& InDataType, EControlRigModelParameterType InParameterType, FControlRigModelNode& OutNode)
+{
+	ensure(InParameterType != EControlRigModelParameterType::None);
+
+	FEdGraphPinType PinType;
+	if (!FindPinTypeFromDataType(InDataType, PinType))
+	{
+		return false;
+	}
+
+	FControlRigModelNode Node;
+	Node.Index = 0;
+	Node.Name = TEXT("Parameter");
+	Node.FunctionName = NAME_None;
+	Node.ParameterType = InParameterType;
+	Node.Position = FVector2D::ZeroVector;
+
+	UControlRigModel::AddNodePinsForParameter(Node, PinType);
+	if (Node.Pins.Num() == 0)
+	{
+		return false;
+	}
+
+	UControlRigModel::ConfigurePinIndices(Node);
+	UControlRigModel::SetNodePinDefaultsForParameter(Node, PinType);
+
+	OutNode = Node;
+	return true;
+}
+
+bool UControlRigController::ConstructPreviewNode(const FName& InFunctionName, FControlRigModelNode& OutNode)
+{
+	FControlRigModelNode Node;
+	Node.Index = 0;
+	Node.Name = InFunctionName;
+	Node.FunctionName = InFunctionName;
+	Node.ParameterType = EControlRigModelParameterType::None;
+	Node.Position = FVector2D::ZeroVector;
+	UControlRigModel::AddNodePinsForFunction(Node);
+	if (Node.Pins.Num() == 0)
+	{
+		return false;
+	}
+	UControlRigModel::ConfigurePinIndices(Node);
+	UControlRigModel::SetNodePinDefaultsForFunction(Node);
+	
+	OutNode = Node;
+	return true;
+}
+
+bool UControlRigController::AddParameter(const FName& InName, const FName& InDataType, EControlRigModelParameterType InParameterType, const FVector2D& InPosition, bool bUndo)
+{
+	if (!EnsureModel())
+	{
+		return false;
+	}
+
+	ensure(InParameterType != EControlRigModelParameterType::None);
+
+	FEdGraphPinType PinType;
+	if (FindPinTypeFromDataType(InDataType, PinType))
+	{
+		if (Model->AddParameter(InName, PinType, InParameterType, InPosition, bUndo))
+		{
+			return true;
+		}
+	}
+
+	if (_ModifiedEvent.IsBound())
+	{
+		FControlRigModelError Error;
+		Error.Message = FString::Printf(TEXT("Parameter data type '%s' not supported."), *InDataType.ToString());
+		_ModifiedEvent.Broadcast(nullptr, EControlRigModelNotifType::ModelError, &Error);
+		LogError(Error.Message);
+	}
+	return false;
+}
+
+bool UControlRigController::AddNode(const FName& InFunctionName, const FVector2D& InPosition, const FName& InName, bool bUndo)
+{
+	if (!EnsureModel())
+	{
+		return false;
+	}
+	
+	FControlRigModelNode Node;
+	Node.Name = InName;
+	Node.FunctionName = InFunctionName;
+	Node.ParameterType = EControlRigModelParameterType::None;
+	Node.Position = InPosition;
+	return Model->AddNode(Node, bUndo);
+}
+
+bool UControlRigController::RemoveNode(const FName& InName, bool bUndo)
+{
+	if (!EnsureModel())
+	{
+		return false;
+	}
+
+	return Model->RemoveNode(InName, bUndo);
+}
+
+bool UControlRigController::SetNodePosition(const FName& InName, const FVector2D& InPosition, bool bUndo)
+{
+	if (!EnsureModel())
+	{
+		return false;
+	}
+
+	return Model->SetNodePosition(InName, InPosition, bUndo);
+}
+
+bool UControlRigController::SetParameterType(const FName& InName, EControlRigModelParameterType InParameterType, bool bUndo)
+{
+	if (!EnsureModel())
+	{
+		return false;
+	}
+
+	return Model->SetParameterType(InName, InParameterType, bUndo);
+}
+
+bool UControlRigController::RenameNode(const FName& InOldNodeName, const FName& InNewNodeName, bool bUndo)
+{
+	if (!EnsureModel())
+	{
+		return false;
+	}
+
+	return Model->RenameNode(InOldNodeName, InNewNodeName, bUndo);
+}
+
+bool UControlRigController::ClearSelection()
+{
+	if (!EnsureModel())
+	{
+		return false;
+	}
+
+	TArray<FControlRigModelNode> SelectedNodes = Model->SelectedNodes();
+	if (SelectedNodes.Num() == 0)
+	{
+		return false;
+	}
+	for (const FControlRigModelNode& SelectedNode : SelectedNodes)
+	{
+		if (!Model->SelectNode(SelectedNode.Name, false))
+		{
+			return false;
+		}
+	}
+	return true;
+}
+
+bool UControlRigController::SetSelection(const TArray<FName>& InNodeSelection)
+{
+	if (!EnsureModel())
+	{
+		return false;
+	}
+
+	TArray<FControlRigModelNode> SelectedNodes = Model->SelectedNodes();
+	if (InNodeSelection.Num() == SelectedNodes.Num())
+	{
+		bool bSelectionMatches = true;
+		for (const FControlRigModelNode& SelectedNode : SelectedNodes)
+		{
+			if (!InNodeSelection.Contains(SelectedNode.Name))
+			{
+				bSelectionMatches = false;
+				break;
+			}
+		}
+
+		if (bSelectionMatches)
+		{
+			return false;
+		}
+	}
+
+	if (Model->SelectedNodes().Num() > 0)
+	{
+		if (!ClearSelection())
+		{
+			return false;
+		}
+	}
+
+	for (const FName& NodeToSelect : InNodeSelection)
+	{
+		if (!Model->SelectNode(NodeToSelect, true))
+		{
+			return false;
+		}
+	}
+	return true;
+}
+
+bool UControlRigController::SelectNode(const FName& InName, bool bInSelected, bool bClearSelection)
+{
+	if (!EnsureModel())
+	{
+		return false;
+	}
+
+	if (!Model->SelectNode(InName, bInSelected))
+	{
+		return false;
+	}
+
+	if (bClearSelection)
+	{
+		TArray<FControlRigModelNode> SelectedNodes = Model->SelectedNodes();
+		for (const FControlRigModelNode& SelectedNode : SelectedNodes)
+		{
+			if (SelectedNode.Name == InName)
+			{
+				continue;
+			}
+			if (!Model->SelectNode(SelectedNode.Name, false))
+			{
+				return false;
+			}
+		}
+	}
+
+	return true;
+}
+
+bool UControlRigController::DeselectNode(const FName& InName)
+{
+	return SelectNode(InName, false, false);
+}
+
+bool UControlRigController::PrepareCycleCheckingForPin(const FName& InNodeName, const FName& InPinName, bool bIsInput)
+{
+	if (!EnsureModel())
+	{
+		return false;
+	}
+
+	const FControlRigModelPin* Pin = Model->FindPin(InNodeName, InPinName, bIsInput);
+	if (Pin == nullptr)
+	{
+		if (_ModifiedEvent.IsBound())
+		{
+			FControlRigModelError Error;
+			Error.Message = FString::Printf(TEXT("Cannot find pin '%s.%s'."), *InNodeName.ToString(), *InPinName.ToString());
+			_ModifiedEvent.Broadcast(nullptr, EControlRigModelNotifType::ModelError, &Error);
+			LogError(Error.Message);
+		}
+		return false;
+	}
+	return Model->PrepareCycleCheckingForPin(Pin->Node, Pin->Index);
+}
+
+bool UControlRigController::ResetCycleCheck()
+{
+	if (!EnsureModel())
+	{
+		return false;
+	}
+	Model->ResetCycleCheck();
+	return true;
+}
+
+bool UControlRigController::CanLink(const FName& InSourceNodeName, const FName& InSourceOutputPinName, const FName& InTargetNodeName, const FName& InTargetInputPinName, FString* OutFailureReason, bool bReportError) const
+{
+	if (!EnsureModel())
+	{
+		return false;
+	}
+
+	const FControlRigModelPin* SourcePin = Model->FindPin(InSourceNodeName, InSourceOutputPinName, false);
+	if (SourcePin == nullptr)
+	{
+		if (OutFailureReason != nullptr)
+		{
+			*OutFailureReason = TEXT("Pins have the same direction.");
+		}
+
+		if (_ModifiedEvent.IsBound() && bReportError)
+		{
+			FControlRigModelError Error;
+			Error.Message = FString::Printf(TEXT("Cannot find source pin '%s.%s'."), *InSourceNodeName.ToString(), *InSourceOutputPinName.ToString());
+			_ModifiedEvent.Broadcast(nullptr, EControlRigModelNotifType::ModelError, &Error);
+			LogError(Error.Message);
+		}
+
+		return false;
+	}
+	const FControlRigModelPin* TargetPin = Model->FindPin(InTargetNodeName, InTargetInputPinName, true);
+	if (TargetPin == nullptr)
+	{
+		if (OutFailureReason != nullptr)
+		{
+			*OutFailureReason = TEXT("Pins have the same direction.");
+		}
+
+		if (_ModifiedEvent.IsBound() && bReportError)
+		{
+			FControlRigModelError Error;
+			Error.Message = FString::Printf(TEXT("Cannot find target pin '%s.%s'."), *InTargetNodeName.ToString(), *InTargetInputPinName.ToString());
+			_ModifiedEvent.Broadcast(nullptr, EControlRigModelNotifType::ModelError, &Error);
+			LogError(Error.Message);
+		}
+
+		return false;
+	}
+
+	return Model->CanLink(SourcePin->Node, SourcePin->Index, TargetPin->Node, TargetPin->Index, OutFailureReason);
+}
+
+bool UControlRigController::MakeLink(const FName& InSourceNodeName, const FName& InSourceOutputPinName, const FName& InTargetNodeName, const FName& InTargetInputPinName, FString* OutFailureReason, bool bUndo)
+{
+	if (!EnsureModel())
+	{
+		return false;
+	}
+
+	const FControlRigModelPin* SourcePin = Model->FindPin(InSourceNodeName, InSourceOutputPinName, false);
+	if (SourcePin == nullptr)
+	{
+		FControlRigModelError Error;
+		Error.Message = FString::Printf(TEXT("Cannot find source pin '%s.%s'."), *InSourceNodeName.ToString(), *InSourceOutputPinName.ToString());
+		if (OutFailureReason != nullptr)
+		{
+			*OutFailureReason = Error.Message;
+		}
+		if (_ModifiedEvent.IsBound())
+		{
+			_ModifiedEvent.Broadcast(nullptr, EControlRigModelNotifType::ModelError, &Error);
+		}
+		LogError(Error.Message);
+		return false;
+	}
+	const FControlRigModelPin* TargetPin = Model->FindPin(InTargetNodeName, InTargetInputPinName, true);
+	if (TargetPin == nullptr)
+	{
+		FControlRigModelError Error;
+		Error.Message = FString::Printf(TEXT("Cannot find target pin '%s.%s'."), *InTargetNodeName.ToString(), *InTargetInputPinName.ToString());
+		if (OutFailureReason != nullptr)
+		{
+			*OutFailureReason = Error.Message;
+		}
+		if (_ModifiedEvent.IsBound())
+		{
+			_ModifiedEvent.Broadcast(nullptr, EControlRigModelNotifType::ModelError, &Error);
+		}
+		LogError(Error.Message);
+		return false;
+	}
+
+	FString FailureReason;
+	if (!Model->CanLink(SourcePin->Node, SourcePin->Index, TargetPin->Node, TargetPin->Index, &FailureReason))
+	{
+		if (OutFailureReason != nullptr)
+		{
+			*OutFailureReason = FailureReason;
+		}
+
+		if (_ModifiedEvent.IsBound())
+		{
+			FControlRigModelError Error;
+			Error.Message = FString::Printf(TEXT("Cannot link '%s.%s' to '%s.%s': %s"), *InSourceNodeName.ToString(), *InSourceOutputPinName.ToString(), *InTargetNodeName.ToString(), *InTargetInputPinName.ToString(), *FailureReason);
+			_ModifiedEvent.Broadcast(nullptr, EControlRigModelNotifType::ModelError, &Error);
+			LogError(Error.Message);
+		}
+		return false;
+	}
+
+	return Model->MakeLink(SourcePin->Node, SourcePin->Index, TargetPin->Node, TargetPin->Index, bUndo);
+}
+
+bool UControlRigController::BreakLink(const FName& InSourceNodeName, const FName& InSourceOutputPinName, const FName& InTargetNodeName, const FName& InTargetInputPinName, bool bUndo)
+{
+	if (!EnsureModel())
+	{
+		return false;
+	}
+
+	const FControlRigModelPin* SourcePin = Model->FindPin(InSourceNodeName, InSourceOutputPinName, false);
+	if (SourcePin == nullptr)
+	{
+		if (_ModifiedEvent.IsBound())
+		{
+			FControlRigModelError Error;
+			Error.Message = FString::Printf(TEXT("Cannot find source pin '%s.%s'."), *InSourceNodeName.ToString(), *InSourceOutputPinName.ToString());
+			_ModifiedEvent.Broadcast(nullptr, EControlRigModelNotifType::ModelError, &Error);
+			LogError(Error.Message);
+		}
+		return false;
+	}
+	const FControlRigModelPin* TargetPin = Model->FindPin(InTargetNodeName, InTargetInputPinName, true);
+	if (TargetPin == nullptr)
+	{
+		if (_ModifiedEvent.IsBound())
+		{
+			FControlRigModelError Error;
+			Error.Message = FString::Printf(TEXT("Cannot find target pin '%s.%s'."), *InTargetNodeName.ToString(), *InTargetInputPinName.ToString());
+			_ModifiedEvent.Broadcast(nullptr, EControlRigModelNotifType::ModelError, &Error);
+			LogError(Error.Message);
+		}
+		return false;
+	}
+
+	return Model->BreakLink(SourcePin->Node, SourcePin->Index, TargetPin->Node, TargetPin->Index, bUndo);
+}
+
+bool UControlRigController::BreakLinks(const FName& InNodeName, const FName& InPinName, bool bIsInput, bool bUndo)
+{
+	if (!EnsureModel())
+	{
+		return false;
+	}
+
+	const FControlRigModelPin* Pin = Model->FindPin(InNodeName, InPinName, bIsInput);
+	if (Pin == nullptr)
+	{
+		if (_ModifiedEvent.IsBound())
+		{
+			FControlRigModelError Error;
+			Error.Message = FString::Printf(TEXT("Cannot find pin '%s.%s'."), *InNodeName.ToString(), *InPinName.ToString());
+			_ModifiedEvent.Broadcast(nullptr, EControlRigModelNotifType::ModelError, &Error);
+			LogError(Error.Message);
+		}
+		return false;
+	}
+	return Model->BreakLinks(Pin->Node, Pin->Index, bUndo);
+}
+
+bool UControlRigController::GetPinDefaultValue(const FName& InNodeName, const FName& InPinName, FString& OutDefaultValue) const
+{
+	if (!EnsureModel())
+	{
+		return false;
+	}
+
+	return Model->GetPinDefaultValue(InNodeName, InPinName, OutDefaultValue);
+}
+
+bool UControlRigController::SetPinDefaultValue(const FName& InNodeName, const FName& InPinName, const FString& InDefaultValue, bool bLog, bool bUndo)
+{
+	if (!EnsureModel())
+	{
+		return false;
+	}
+
+	TGuardValue<bool> SuspendLog(bSuspendLog, !bLog);
+	return Model->SetPinDefaultValue(InNodeName, InPinName, InDefaultValue, bUndo);
+}
+
+bool UControlRigController::GetPinDefaultValueBool(const FName& InNodeName, const FName& InPinName, bool& OutDefaultValue) const
+{
+	FString DefaultValueString;
+	if(!GetPinDefaultValue(InNodeName, InPinName, DefaultValueString))
+	{
+		return false;
+	}
+	OutDefaultValue = DefaultValueString.ToBool();
+	return true;
+}
+
+bool UControlRigController::SetPinDefaultValueBool(const FName& InNodeName, const FName& InPinName, bool InDefaultValue, bool bLog, bool bUndo)
+{
+	FString DefaultValueString = InDefaultValue ? TEXT("True") : TEXT("False");
+	return SetPinDefaultValue(InNodeName, InPinName, DefaultValueString, bLog, bUndo);
+}
+
+bool UControlRigController::GetPinDefaultValueFloat(const FName& InNodeName, const FName& InPinName, float& OutDefaultValue) const
+{
+	FString DefaultValueString;
+	if(!GetPinDefaultValue(InNodeName, InPinName, DefaultValueString))
+	{
+		return false;
+	}
+	OutDefaultValue = FCString::Atof(*DefaultValueString);
+	return true;
+}
+
+bool UControlRigController::SetPinDefaultValueFloat(const FName& InNodeName, const FName& InPinName, float InDefaultValue, bool bLog, bool bUndo)
+{
+	FString DefaultValueString = FString::SanitizeFloat((double)InDefaultValue);
+	return SetPinDefaultValue(InNodeName, InPinName, DefaultValueString, bLog, bUndo);
+}
+
+bool UControlRigController::GetPinDefaultValueInt(const FName& InNodeName, const FName& InPinName, int32& OutDefaultValue) const
+{
+	FString DefaultValueString;
+	if(!GetPinDefaultValue(InNodeName, InPinName, DefaultValueString))
+	{
+		return false;
+	}
+	OutDefaultValue = FCString::Atoi(*DefaultValueString);
+	return true;
+}
+
+bool UControlRigController::SetPinDefaultValueInt(const FName& InNodeName, const FName& InPinName, int32 InDefaultValue, bool bLog, bool bUndo)
+{
+	FString DefaultValueString = FString::FormatAsNumber(InDefaultValue);
+	return SetPinDefaultValue(InNodeName, InPinName, DefaultValueString, bLog, bUndo);
+}
+
+bool UControlRigController::GetPinDefaultValueName(const FName& InNodeName, const FName& InPinName, FName& OutDefaultValue) const
+{
+	FString DefaultValueString;
+	if(!GetPinDefaultValue(InNodeName, InPinName, DefaultValueString))
+	{
+		return false;
+	}
+	OutDefaultValue = *DefaultValueString;
+	return true;
+}
+
+bool UControlRigController::SetPinDefaultValueName(const FName& InNodeName, const FName& InPinName, const FName& InDefaultValue, bool bLog, bool bUndo)
+{
+	FString DefaultValueString = InDefaultValue.ToString();
+	return SetPinDefaultValue(InNodeName, InPinName, DefaultValueString, bLog, bUndo);
+}
+
+bool UControlRigController::GetPinDefaultValueVector(const FName& InNodeName, const FName& InPinName, FVector& OutDefaultValue) const
+{
+	return GetPinDefaultValueStruct<FVector>(InNodeName, InPinName, OutDefaultValue);
+}
+
+bool UControlRigController::SetPinDefaultValueVector(const FName& InNodeName, const FName& InPinName, const FVector& InDefaultValue, bool bLog, bool bUndo)
+{
+	return SetPinDefaultValueStruct<FVector>(InNodeName, InPinName, InDefaultValue, bLog, bUndo);
+}
+
+bool UControlRigController::GetPinDefaultValueQuat(const FName& InNodeName, const FName& InPinName, FQuat& OutDefaultValue) const
+{
+	return GetPinDefaultValueStruct<FQuat>(InNodeName, InPinName, OutDefaultValue);
+}
+
+bool UControlRigController::SetPinDefaultValueQuat(const FName& InNodeName, const FName& InPinName, const FQuat& InDefaultValue, bool bLog, bool bUndo)
+{
+	return SetPinDefaultValueStruct<FQuat>(InNodeName, InPinName, InDefaultValue, bLog, bUndo);
+}
+
+bool UControlRigController::GetPinDefaultValueTransform(const FName& InNodeName, const FName& InPinName, FTransform& OutDefaultValue) const
+{
+	return GetPinDefaultValueStruct<FTransform>(InNodeName, InPinName, OutDefaultValue);
+}
+
+bool UControlRigController::SetPinDefaultValueTransform(const FName& InNodeName, const FName& InPinName, const FTransform& InDefaultValue, bool bLog, bool bUndo)
+{
+	return SetPinDefaultValueStruct<FTransform>(InNodeName, InPinName, InDefaultValue, bLog, bUndo);
+}
+
+bool UControlRigController::AddArrayPin(const FName& InNodeName, const FName& InPinName, const FString& InDefaultValue, bool bUndo)
+{
+	if (!EnsureModel())
+	{
+		return false;
+	}
+
+	const FControlRigModelPin* Pin = Model->FindPin(InNodeName, InPinName, true);
+	if (Pin == nullptr)
+	{
+		if (_ModifiedEvent.IsBound())
+		{
+			FControlRigModelError Error;
+			Error.Message = FString::Printf(TEXT("Cannot find pin '%s.%s'."), *InNodeName.ToString(), *InPinName.ToString());
+			_ModifiedEvent.Broadcast(nullptr, EControlRigModelNotifType::ModelError, &Error);
+			LogError(Error.Message);
+		}
+		return false;
+	}
+
+	return Model->SetPinArraySize(Pin->GetPair(), Pin->ArraySize() + 1, InDefaultValue, bUndo);
+}
+
+bool UControlRigController::PopArrayPin(const FName& InNodeName, const FName& InPinName, bool bUndo)
+{
+	if (!EnsureModel())
+	{
+		return false;
+	}
+
+	const FControlRigModelPin* Pin = Model->FindPin(InNodeName, InPinName, true);
+	if (Pin == nullptr)
+	{
+		if (_ModifiedEvent.IsBound())
+		{
+			FControlRigModelError Error;
+			Error.Message = FString::Printf(TEXT("Cannot find pin '%s.%s'."), *InNodeName.ToString(), *InPinName.ToString());
+			_ModifiedEvent.Broadcast(nullptr, EControlRigModelNotifType::ModelError, &Error);
+			LogError(Error.Message);
+		}
+		return false;
+	}
+
+	return Model->SetPinArraySize(Pin->GetPair(), Pin->ArraySize() - 1, FString(), bUndo);
+}
+
+bool UControlRigController::ClearArrayPin(const FName& InNodeName, const FName& InPinName, bool bUndo)
+{
+	if (!EnsureModel())
+	{
+		return false;
+	}
+
+	const FControlRigModelPin* Pin = Model->FindPin(InNodeName, InPinName, true);
+	if (Pin == nullptr)
+	{
+		if (_ModifiedEvent.IsBound())
+		{
+			FControlRigModelError Error;
+			Error.Message = FString::Printf(TEXT("Cannot find pin '%s.%s'."), *InNodeName.ToString(), *InPinName.ToString());
+			_ModifiedEvent.Broadcast(nullptr, EControlRigModelNotifType::ModelError, &Error);
+			LogError(Error.Message);
+		}
+		return false;
+	}
+
+	return Model->SetPinArraySize(Pin->GetPair(), 0, FString(), bUndo);
+}
+
+bool UControlRigController::SetArrayPinSize(const FName& InNodeName, const FName& InPinName, int32 InSize, const FString& InDefaultValue, bool bUndo)
+{
+	if (!EnsureModel())
+	{
+		return false;
+	}
+
+	const FControlRigModelPin* Pin = Model->FindPin(InNodeName, InPinName, true);
+	if (Pin == nullptr)
+	{
+		if (_ModifiedEvent.IsBound())
+		{
+			FControlRigModelError Error;
+			Error.Message = FString::Printf(TEXT("Cannot find pin '%s.%s'."), *InNodeName.ToString(), *InPinName.ToString());
+			_ModifiedEvent.Broadcast(nullptr, EControlRigModelNotifType::ModelError, &Error);
+			LogError(Error.Message);
+		}
+		return false;
+	}
+
+	return Model->SetPinArraySize(Pin->GetPair(), InSize, InDefaultValue, bUndo);
+}
+
+
+bool UControlRigController::ExpandPin(const FName& InNodeName, const FName& InPinName, bool bIsInput, bool bExpanded, bool bUndo)
+{
+	if (!EnsureModel())
+	{
+		return false;
+	}
+
+	return Model->ExpandPin(InNodeName, InPinName, bIsInput, bExpanded, bUndo);
+}
+
+bool UControlRigController::FindPinTypeFromDataType(const FName& InDataType, FEdGraphPinType& OutPinType)
+{
+	TArray<FEdGraphPinType> PinTypes;
+	UControlRigModel::GetParameterPinTypes(PinTypes);
+
+	for (const FEdGraphPinType& PinType : PinTypes)
+	{
+		FName PinTypeDataType = PinType.PinCategory;
+		if (PinTypeDataType == UEdGraphSchema_K2::PC_Struct)
+		{
+			if (UStruct* Struct = Cast<UStruct>(PinType.PinSubCategoryObject))
+			{
+				PinTypeDataType = Struct->GetFName();
+			}
+		}
+
+		if (PinTypeDataType == InDataType)
+		{
+			OutPinType = PinType;
+			return true;
+		}
+	}
+
+	return false;
+}
+
+bool UControlRigController::ResendAllNotifications()
+{
+	if (!EnsureModel())
+	{
+		return false;
+	}
+
+	return Model->ResendAllNotifications();
+}
+
+bool UControlRigController::ResendAllPinDefaultNotifications()
+{
+	if (!EnsureModel())
+	{
+		return false;
+	}
+	return Model->ResendAllPinDefaultNotifications();
+}
+
+bool UControlRigController::Undo()
+{
+	if (!EnsureModel())
+	{
+		return false;
+	}
+
+	return Model->Undo();
+}
+
+bool UControlRigController::Redo()
+{
+	if (!EnsureModel())
+	{
+		return false;
+	}
+
+	return Model->Redo();
+}
+
+bool UControlRigController::GetPinDefaultValueRecursiveStruct(const FControlRigModelPin* InPin, FString& OutValue) const
+{
+	if (InPin->SubPins.Num() == 0)
+	{
+		OutValue = InPin->DefaultValue;
+	}
+	else
+	{
+		FString DefaultValue;
+		for (int32 SubPinIndex : InPin->SubPins)
+		{
+			const FControlRigModelPin* SubPin = Model->FindPin(FControlRigModelPair(InPin->Node, SubPinIndex));
+			if (SubPin != nullptr)
+			{
+				FString SubDefaultValue;
+				ensure(SubPin->Type.ContainerType == EPinContainerType::None);
+				if (Cast<UScriptStruct>(SubPin->Type.PinSubCategoryObject))
+				{
+					if (!GetPinDefaultValueRecursiveStruct(SubPin, SubDefaultValue))
+					{
+						return false;
+					}
+				}
+				else
+				{
+					SubDefaultValue = SubPin->DefaultValue;
+					if (SubPin->Type.PinCategory == UEdGraphSchema_K2::PC_Name || SubPin->Type.PinCategory == UEdGraphSchema_K2::PC_String)
+					{
+						SubDefaultValue = TEXT("\"") + SubDefaultValue + TEXT("\"");
+					}
+				}
+
+				if (!DefaultValue.IsEmpty())
+				{
+					DefaultValue += FString::Printf(TEXT(",%s=%s"), *SubPin->Name.ToString(), *SubDefaultValue);
+				}
+				else
+				{
+					DefaultValue = FString::Printf(TEXT("%s=%s"), *SubPin->Name.ToString(), *SubDefaultValue);
+				}
+			}
+		}
+
+		OutValue = TEXT("(") + DefaultValue + TEXT(")");
+	}
+	return true;
+}
+
+bool UControlRigController::SetPinDefaultValueRecursiveStruct(const FControlRigModelPin* OutPin, const FString& InValue, bool bUndo)
+{
+	if (OutPin->SubPins.Num() == 0)
+	{
+		if (OutPin->DefaultValue == InValue)
+		{
+			return true;
+		}
+		return Model->SetPinDefaultValue(OutPin->GetPair(), InValue, bUndo);
+	}
+
+	FString Data = InValue;
+	if (!Data.RemoveFromStart(TEXT("(")))
+	{
+		return false;
+	}
+	if (!Data.RemoveFromEnd(TEXT(")")))
+	{
+		return false;
+	}
+
+	TArray<FString> Parts;
+	int32 BraceCount = 0;
+	int32 QuoteCount = 0;
+	int32 LastCommaPos = -1;
+	int32 CharIndex = 1;
+
+	while (CharIndex < Data.Len())
+	{
+		if (Data[CharIndex] == ',')
+		{
+			if (BraceCount == 0 && QuoteCount == 0)
+			{
+				FString Part = Data.Mid(LastCommaPos + 1, CharIndex - LastCommaPos - 1);
+				Parts.Add(Part);
+				LastCommaPos = CharIndex;
+			}
+		}
+		else if (Data[CharIndex] == '(')
+		{
+			if (QuoteCount == 0)
+			{
+				BraceCount++;
+			}
+		}
+		else if (Data[CharIndex] == ')')
+		{
+			if (QuoteCount == 0)
+			{
+				BraceCount--;
+				if (BraceCount < 0)
+				{
+					return false;
+				}
+			}
+		}
+		else if (Data[CharIndex] == '\"')
+		{
+			if (QuoteCount > 0)
+			{
+				QuoteCount--;
+			}
+			else
+			{
+				QuoteCount++;
+			}
+		}
+		CharIndex++;
+	}
+
+	if (CharIndex - LastCommaPos - 1 > 0)
+	{
+		FString Part = Data.Mid(LastCommaPos + 1, CharIndex - LastCommaPos - 1);
+		Parts.Add(Part);
+	}
+
+	for (const FString& Part : Parts)
+	{
+		FString Name, Value;
+		if(!Part.Split(TEXT("="), &Name, &Value))
+		{
+			return false;
+		}
+
+		const FControlRigModelPin* SubPin = Model->FindSubPin(OutPin, *Name);
+		if (SubPin)
+		{
+			ensure(SubPin->Type.ContainerType == EPinContainerType::None);
+
+			if (Cast<UScriptStruct>(SubPin->Type.PinSubCategoryObject))
+			{
+				if (!SetPinDefaultValueRecursiveStruct(SubPin, Value, bUndo))
+				{
+					return false;
+				}
+			}
+			else
+			{
+				if (SubPin->DefaultValue != Value)
+				{
+					if (!Model->SetPinDefaultValue(SubPin->GetPair(), Value, bUndo))
+					{
+						return false;
+					}
+				}
+			}
+		}
+	}
+
+	if (OutPin->DefaultValue != InValue)
+	{
+		if (!Model->SetPinDefaultValue(OutPin->GetPair(), InValue, bUndo))
+		{
+			return false;
+		}
+	}
+
+	return Parts.Num() > 0;
+}
+
+#undef LOCTEXT_NAMESPACE
+
diff --git a/Engine/Plugins/Experimental/ControlRig/Source/ControlRigDeveloper/Private/ControlRigDeveloper.cpp b/Engine/Plugins/Experimental/ControlRig/Source/ControlRigDeveloper/Private/ControlRigDeveloper.cpp
--- a/Engine/Plugins/Experimental/ControlRig/Source/ControlRigDeveloper/Private/ControlRigDeveloper.cpp
+++ b/Engine/Plugins/Experimental/ControlRig/Source/ControlRigDeveloper/Private/ControlRigDeveloper.cpp
@@ -4,6 +4,9 @@
 #include "ControlRigBlueprint.h"
 #include "ControlRigBlueprintCompiler.h"
 #include "Modules/ModuleManager.h"
+#include "Developer/MessageLog/Public/MessageLogModule.h"
+
+#define LOCTEXT_NAMESPACE "ControlRigDeveloperModule"
 
 class FControlRigDeveloperModule : public IControlRigDeveloperModule
 {
@@ -25,6 +28,12 @@ void FControlRigDeveloperModule::StartupModule()
 	IKismetCompilerInterface& KismetCompilerModule = FModuleManager::LoadModuleChecked<IKismetCompilerInterface>("KismetCompiler");
 	KismetCompilerModule.GetCompilers().Add(&ControlRigBlueprintCompiler);
 
+	FMessageLogModule& MessageLogModule = FModuleManager::LoadModuleChecked<FMessageLogModule>("MessageLog");
+	FMessageLogInitializationOptions InitOptions;
+	InitOptions.bShowFilters = true;
+	InitOptions.bShowPages = false;
+	InitOptions.bAllowClear = true;
+	MessageLogModule.RegisterLogListing("ControlRigLog", LOCTEXT("ControlRigLog", "Control Rig Log"), InitOptions);
 }
 
 void FControlRigDeveloperModule::ShutdownModule()
@@ -34,6 +43,9 @@ void FControlRigDeveloperModule::ShutdownModule()
 	{
 		KismetCompilerModule->GetCompilers().Remove(&ControlRigBlueprintCompiler);
 	}
+
+	FMessageLogModule& MessageLogModule = FModuleManager::LoadModuleChecked<FMessageLogModule>("MessageLog");
+	MessageLogModule.UnregisterLogListing("ControlRigLog");
 }
 
 TSharedPtr<FKismetCompilerContext> FControlRigDeveloperModule::GetControlRigCompiler(UBlueprint* BP, FCompilerResultsLog& InMessageLog, const FKismetCompilerOptions& InCompileOptions)
@@ -43,3 +55,5 @@ TSharedPtr<FKismetCompilerContext> FControlRigDeveloperModule::GetControlRigComp
 
 
 IMPLEMENT_MODULE(FControlRigDeveloperModule, ControlRigDeveloper)
+
+#undef LOCTEXT_NAMESPACE
diff --git a/Engine/Plugins/Experimental/ControlRig/Source/ControlRigDeveloper/Private/ControlRigGraphTraverser.cpp b/Engine/Plugins/Experimental/ControlRig/Source/ControlRigDeveloper/Private/ControlRigGraphTraverser.cpp
--- a/Engine/Plugins/Experimental/ControlRig/Source/ControlRigDeveloper/Private/ControlRigGraphTraverser.cpp
+++ b/Engine/Plugins/Experimental/ControlRig/Source/ControlRigDeveloper/Private/ControlRigGraphTraverser.cpp
@@ -1,152 +1,153 @@
 // Copyright 1998-2019 Epic Games, Inc. All Rights Reserved.
 
 #include "ControlRigGraphTraverser.h"
-#include "Units/RigUnit.h"
-#include "ControlRig/Private/Units/Execution/RigUnit_BeginExecution.h"
+#include "EdGraphSchema_K2.h"
 
-FControlRigGraphTraverser::FControlRigGraphTraverser(UControlRigBlueprint* InBlueprint, UControlRigGraph* InGraph)
-	: Blueprint(InBlueprint)
-	, Graph(InGraph)
+FControlRigGraphTraverser::FControlRigGraphTraverser(UControlRigModel* InModel)
+	: Model(InModel)
 {
 }
 
 #if WITH_EDITORONLY_DATA
-bool FControlRigGraphTraverser::IsWiredToExecution(const FName& UnitName)
+bool FControlRigGraphTraverser::IsWiredToExecution(const FName& NodeName)
 {
-	for (UEdGraphNode* Node : Graph->Nodes)
+	const FControlRigModelNode* Node = Model->FindNode(NodeName);
+	if(Node)
 	{
-		if (UControlRigGraphNode* RigNode = Cast<UControlRigGraphNode>(Node))
-		{
-			if (RigNode->GetPropertyName().IsEqual(UnitName))
-			{
-				return IsWiredToExecution(RigNode);
-			}
-		}
+		return IsWiredToExecution(Node);
 	}
 	return false;
 }
 #endif
 
-bool FControlRigGraphTraverser::IsWiredToExecution(UControlRigGraphNode* Node)
+bool FControlRigGraphTraverser::IsWiredToExecution(const FControlRigModelNode* Node)
 {
 	if (Node == nullptr)
 	{
 		return false;
 	}
 
-	const bool* Found = VisitedNodes.Find(Node->PropertyName);
+	const bool* Found = VisitedNodes.Find(Node->Name);
 	if (Found)
 	{
 		return *Found;
 	}
 
-	if (Node->GetUnitScriptStruct() == FRigUnit_BeginExecution::StaticStruct())
+	if (Node->IsBeginExecution())
 	{
-		VisitedNodes.Add(Node->PropertyName, true);
+		VisitedNodes.Add(Node->Name, true);
 		return true;
 	}
 
-	VisitedNodes.Add(Node->PropertyName, false);
+	VisitedNodes.Add(Node->Name, false);
 
 	bool bFoundWiredPin = false;
 
 	// is this an execution node,
 	// walk upwards (to the left) to find a proper begin execution node
-	if (Node->GetUnitScriptStruct()->IsChildOf(FRigUnitMutable::StaticStruct()))
+	if (Node->IsMutable())
 	{
-		for (const UEdGraphPin* Pin : Node->Pins)
+		for (const FControlRigModelPin& Pin : Node->Pins)
 		{
-			if (Pin->Direction != EEdGraphPinDirection::EGPD_Input)
+			if (Pin.Direction != EEdGraphPinDirection::EGPD_Input)
 			{
 				continue;
 			}
-			if (Pin->PinType.PinCategory != UEdGraphSchema_K2::PC_Struct)
+			if (Pin.Type.PinCategory != UEdGraphSchema_K2::PC_Struct)
 			{
 				continue;
 			}
-			if (Pin->PinType.PinSubCategoryObject != FControlRigExecuteContext::StaticStruct())
+			if (Pin.Type.PinSubCategoryObject != FControlRigExecuteContext::StaticStruct())
 			{
 				continue;
 			}
 
-			for (UEdGraphPin* LinkedPin : Pin->LinkedTo)
+			for (int32 LinkIndex : Pin.Links)
 			{
-				UControlRigGraphNode* LinkedNode = Cast<UControlRigGraphNode>(LinkedPin->GetOwningNode());
-				if (LinkedNode)
+				const FControlRigModelLink* Link = Model->FindLink(LinkIndex);
+				if (Link)
 				{
-					bool bIsLinkedNodeWired = IsWiredToExecution(LinkedNode);
-					if (bIsLinkedNodeWired)
+					const FControlRigModelNode* LinkedNode = Model->FindNode(Link->Source.Node);
+					if (LinkedNode)
 					{
-						bFoundWiredPin = true;
+						bool bIsLinkedNodeWired = IsWiredToExecution(LinkedNode);
+						if (bIsLinkedNodeWired)
+						{
+							bFoundWiredPin = true;
+						}
 					}
 				}
 			}
 		}
 
-		VisitedNodes.FindChecked(Node->PropertyName) = bFoundWiredPin;
+		VisitedNodes.FindChecked(Node->Name) = bFoundWiredPin;
 		return bFoundWiredPin;
 	}
 
 	// for all other nodes walk  to the right...
-	for (const UEdGraphPin* Pin : Node->Pins)
+	for (const FControlRigModelPin& Pin : Node->Pins)
 	{
-		if (Pin->Direction != EEdGraphPinDirection::EGPD_Output)
+		if (Pin.Direction != EEdGraphPinDirection::EGPD_Output)
 		{
 			continue;
 		}
 
-		for (UEdGraphPin* LinkedPin : Pin->LinkedTo)
+		for (int32 LinkIndex : Pin.Links)
 		{
-			UControlRigGraphNode* LinkedNode = Cast<UControlRigGraphNode>(LinkedPin->GetOwningNode());
-			if (LinkedNode)
+			const FControlRigModelLink* Link = Model->FindLink(LinkIndex);
+			if (Link)
 			{
-				bool bIsLinkedNodeWired = IsWiredToExecution(LinkedNode);
-				if (bIsLinkedNodeWired)
+				const FControlRigModelNode* LinkedNode = Model->FindNode(Link->Target.Node);
+				if (LinkedNode)
 				{
-					bFoundWiredPin = true;
+					bool bIsLinkedNodeWired = IsWiredToExecution(LinkedNode);
+					if (bIsLinkedNodeWired)
+					{
+						bFoundWiredPin = true;
+					}
 				}
 			}
 		}
-
 	}
 
-	VisitedNodes.FindChecked(Node->PropertyName) = bFoundWiredPin;
+	VisitedNodes.FindChecked(Node->Name) = bFoundWiredPin;
 	return bFoundWiredPin;
 }
 
-void FControlRigGraphTraverser::TraverseAndBuildPropertyLinks()
+void FControlRigGraphTraverser::TraverseAndBuildPropertyLinks(UControlRigBlueprint* Blueprint)
 {
-	for (UEdGraphNode* Node : Graph->Nodes)
+	for(const FControlRigModelNode& Node : Model->Nodes())
 	{
-		UControlRigGraphNode* RigNode = Cast< UControlRigGraphNode>(Node);
-		if (RigNode == nullptr)
-		{
-			continue;
-		}
-
-		if (!IsWiredToExecution(RigNode))
+		if (!IsWiredToExecution(&Node))
 		{
 			continue;
 		}
 
-		for (UEdGraphPin* Pin : RigNode->Pins)
+		for (const FControlRigModelPin& Pin : Node.Pins)
 		{
-			if (Pin->Direction == EGPD_Output)
+			if (Pin.Direction != EEdGraphPinDirection::EGPD_Output)
+			{
+				continue;
+			}
+			for (int32 LinkIndex : Pin.Links)
 			{
-				for (UEdGraphPin* LinkedPin : Pin->LinkedTo)
+				const FControlRigModelLink* Link = Model->FindLink(LinkIndex);
+				if (Link)
 				{
-					UControlRigGraphNode* LinkedRigNode = Cast< UControlRigGraphNode>(LinkedPin->GetOwningNode());
-					if (LinkedRigNode == nullptr)
-					{
-						continue;
-					}
-					if (!IsWiredToExecution(LinkedRigNode))
+					const FControlRigModelNode* LinkedNode = Model->FindNode(Link->Target.Node);
+					if (LinkedNode)
 					{
-						continue;
+						if (!IsWiredToExecution(LinkedNode))
+						{
+							continue;
+						}
+
+						int32 PinIndex = Link->Source.Pin;
+						int32 LinkedPinIndex = Link->Target.Pin;
+						FString PinPath = Model->GetPinPath(Link->Source);
+						FString LinkedPinPath = Model->GetPinPath(Link->Target);
+						Blueprint->MakePropertyLink(PinPath, LinkedPinPath, PinIndex, LinkedPinIndex);
 					}
-					int32 PinIndex = RigNode->Pins.IndexOfByKey(Pin);
-					int32 LinkedPinIndex = LinkedRigNode->Pins.IndexOfByKey(LinkedPin);
-					Blueprint->MakePropertyLink(Pin->PinName.ToString(), LinkedPin->PinName.ToString(), PinIndex, LinkedPinIndex);
 				}
 			}
 		}
