Control Rig MVC Refactoring

diff --git a/Engine/Plugins/Experimental/ControlRig/Source/ControlRigDeveloper/Private/ControlRigModel.cpp b/Engine/Plugins/Experimental/ControlRig/Source/ControlRigDeveloper/Private/ControlRigModel.cpp
new file mode 100644
--- /dev/null
+++ b/Engine/Plugins/Experimental/ControlRig/Source/ControlRigDeveloper/Private/ControlRigModel.cpp
@@ -0,0 +1,2486 @@
+// Copyright 1998-2019 Epic Games, Inc. All Rights Reserved.
+
+#include "ControlRigModel.h"
+#include "Units/RigUnit.h"
+#include "Units/Execution/RigUnit_BeginExecution.h"
+#include "HelperUtil.h"
+#include "ControlRig.h"
+#include "EdGraphSchema_K2.h"
+#include "UObject/PropertyPortFlags.h"
+#include "Curves/CurveFloat.h"
+#include "Materials/MaterialInterface.h"
+
+#if CONTROLRIG_UNDO
+#include "ScopedTransaction.h"
+#endif
+
+#define LOCTEXT_NAMESPACE "ControlRigModel"
+
+#if CONTROLRIG_UNDO
+
+void FControlRigModelPair::AppendArgumentsForAction(TArray<FString>& InOutArguments, const UControlRigModel* InModel) const
+{
+	InOutArguments.Add(InModel->Nodes()[Node].Pins[Pin].Direction == EGPD_Input ? TEXT("true") : TEXT("false"));
+	InOutArguments.Add(InModel->GetPinPath(*this));
+}
+
+void FControlRigModelPair::ConfigureFromActionArguments(const TArray<FString>& InOutArguments, int32 ArgumentIndex, const UControlRigModel* InModel)
+{
+	ensure(InOutArguments.Num() >= ArgumentIndex + 2);
+	bool bIsInput = InOutArguments[ArgumentIndex++] == TEXT("true");
+	const FControlRigModelPin* ExistingPin = InModel->FindPinFromPath(InOutArguments[ArgumentIndex++], bIsInput);
+	if (ExistingPin)
+	{
+		Node = ExistingPin->Node;
+		Pin = ExistingPin->Index;
+	}
+}
+
+void FControlRigModelLink::AppendArgumentsForAction(TArray<FString>& InOutArguments, const UControlRigModel* InModel) const
+{
+	Source.AppendArgumentsForAction(InOutArguments, InModel);
+	Target.AppendArgumentsForAction(InOutArguments, InModel);
+}
+
+void FControlRigModelLink::ConfigureFromActionArguments(const TArray<FString>& InOutArguments, int32 ArgumentIndex, const UControlRigModel* InModel)
+{
+	Source.ConfigureFromActionArguments(InOutArguments, ArgumentIndex, InModel);
+	Target.ConfigureFromActionArguments(InOutArguments, ArgumentIndex + 2, InModel);
+}
+
+void FControlRigModelPin::AppendArgumentsForAction(TArray<FString>& InOutArguments) const
+{
+	InOutArguments.Add(Name.ToString());
+	InOutArguments.Add(DefaultValue);
+	InOutArguments.Add(bExpanded ? TEXT("true") : TEXT("false"));
+	InOutArguments.Add(FString::FormatAsNumber((int32)Direction));
+}
+
+void FControlRigModelPin::ConfigureFromActionArguments(const TArray<FString>& InOutArguments, int32 ArgumentIndex)
+{
+	ensure(InOutArguments.Num() >= ArgumentIndex + 3);
+
+	Name = *InOutArguments[ArgumentIndex++];
+	DefaultValue = InOutArguments[ArgumentIndex++];
+	bExpanded = InOutArguments[ArgumentIndex++] == TEXT("true");
+	Direction = (EEdGraphPinDirection)FCString::Atoi(*InOutArguments[ArgumentIndex++]);
+}
+
+#endif
+
+FString FControlRigModelNode::GetPinPath(int32 InPinIndex, bool bIncludeNodeName) const
+{
+	ensure(InPinIndex >= 0 && InPinIndex < Pins.Num());
+
+	const FControlRigModelPin& Pin = Pins[InPinIndex];
+	if (Pin.ParentIndex != INDEX_NONE)
+	{
+		ensure(Pin.Index != Pin.ParentIndex);
+
+		const FControlRigModelPin& ParentPin = Pins[Pin.ParentIndex];
+		if (ParentPin.IsArray())
+		{
+			return FString::Printf(TEXT("%s[%s]"), *GetPinPath(Pin.ParentIndex, bIncludeNodeName), *Pin.Name.ToString());
+		}
+		return FString::Printf(TEXT("%s.%s"), *GetPinPath(Pin.ParentIndex, bIncludeNodeName), *Pin.Name.ToString());
+	}
+	if (bIncludeNodeName)
+	{
+		// this is done for backwards compatibility.
+		// on variable nodes the pin paths are short of the "value" pin name
+		if (IsParameter())
+		{
+			return Name.ToString();
+		}
+		return FString::Printf(TEXT("%s.%s"), *Name.ToString(), *Pin.Name.ToString());
+	}
+	return Pin.Name.ToString();
+}
+
+bool FControlRigModelNode::IsMutable() const
+{
+	if (IsParameter())
+	{
+		return false;
+	}
+
+	ensure(FunctionName != NAME_None);
+
+	const UStruct* Struct = UnitStruct();
+	if (Struct)
+	{
+		return Struct->IsChildOf(FRigUnitMutable::StaticStruct());
+	}
+	return false;
+}
+
+bool FControlRigModelNode::IsBeginExecution() const
+{
+	if (IsParameter())
+	{
+		return false;
+	}
+
+	ensure(FunctionName != NAME_None);
+
+	const UStruct* Struct = UnitStruct();
+	if (Struct)
+	{
+		return Struct->IsChildOf(FRigUnit_BeginExecution::StaticStruct());
+	}
+	return false;
+}
+
+const UStruct* FControlRigModelNode::UnitStruct() const
+{
+	ensure(ParameterType == EControlRigModelParameterType::None);
+	ensure(FunctionName != NAME_None);
+	return FindObject<UStruct>(ANY_PACKAGE, *(FunctionName.ToString()));
+}
+
+const FControlRigModelPin* FControlRigModelNode::FindPin(const FName& InName, bool bLookForInput) const
+{
+	FString Left, Right;
+	FString PinPath = InName.ToString();
+	UControlRigModel::SplitPinPath(PinPath, Left, Right, false);
+
+	if (Left == Name.ToString())
+	{
+		return FindPin(*Right, bLookForInput);
+	}
+
+	for (const FControlRigModelPin& Pin : Pins)
+	{
+		if (Pin.Name.ToString() == Left)
+		{
+			if (bLookForInput && Pin.Direction != EGPD_Input)
+			{
+				continue;
+			}
+			if (!bLookForInput && Pin.Direction != EGPD_Output)
+			{
+				continue;
+			}
+
+			int32 PinIndex = Pin.Index;
+			while (!Right.IsEmpty())
+			{
+				PinPath = Right;
+				UControlRigModel::SplitPinPath(PinPath, Left, Right, false);
+
+				for (int32 ChildIndex : Pins[PinIndex].SubPins)
+				{
+					if (Pins[ChildIndex].Name.ToString() == Left)
+					{
+						PinIndex = ChildIndex;
+						break;
+					}
+				}
+			}
+
+			return &Pins[PinIndex];
+		}
+	}
+	return nullptr;
+}
+
+#if CONTROLRIG_UNDO
+
+void FControlRigModelNode::AppendArgumentsForAction(TArray<FString>& InOutArguments) const
+{
+	InOutArguments.Add(Name.ToString());
+	InOutArguments.Add(FunctionName.ToString());
+	InOutArguments.Add(FString::FormatAsNumber((int32)ParameterType));
+	FString PositionStr;
+	TBaseStructure<FVector2D>::Get()->ExportText(PositionStr, &Position, nullptr, nullptr, PPF_None, nullptr);
+	InOutArguments.Add(PositionStr);
+	if (IsParameter() && Pins.Num() > 0)
+	{
+		FString DataTypeStr;
+		FEdGraphPinType::StaticStruct()->ExportText(DataTypeStr, &Pins[0].Type, nullptr, nullptr, PPF_None, nullptr);
+		InOutArguments.Add(DataTypeStr);
+	}
+	else
+	{
+		InOutArguments.Add(FString() /* parameter data type */);
+	}
+}
+
+void FControlRigModelNode::ConfigureFromActionArguments(const TArray<FString>& InOutArguments, int32 ArgumentIndex)
+{
+	ensure(InOutArguments.Num() >= ArgumentIndex + 4);
+
+	Name = *InOutArguments[ArgumentIndex++];
+	const FString& InFunctionName = InOutArguments[ArgumentIndex++];
+	if (InFunctionName == FName(NAME_None).ToString())
+	{
+		FunctionName = NAME_None;
+	}
+	else
+	{
+		FunctionName = *InFunctionName;
+	}
+
+	ParameterType = (EControlRigModelParameterType)FCString::Atoi(*InOutArguments[ArgumentIndex++]);
+	TBaseStructure<FVector2D>::Get()->ImportText(*InOutArguments[ArgumentIndex++], &Position, nullptr, EPropertyPortFlags::PPF_None, nullptr, TEXT("Vector2D"), true);
+}
+
+#endif
+
+
+const FName UControlRigModel::ValueName("Value");
+
+UControlRigModel::UControlRigModel()
+{
+	bIsSelecting = false;
+}
+
+UControlRigModel::~UControlRigModel()
+{
+	_ModifiedEvent.Clear();
+}
+
+const TArray<FControlRigModelNode>& UControlRigModel::Nodes() const
+{
+	return _Nodes;
+}
+
+TArray<FControlRigModelNode> UControlRigModel::SelectedNodes() const
+{
+	TArray<FControlRigModelNode> Nodes;
+	for (const FName& Name : _SelectedNodes)
+	{
+		const FControlRigModelNode* Node = FindNode(Name);
+		check(Node);
+		Nodes.Add(*Node);
+	}
+	return Nodes;
+}
+
+const TArray<FControlRigModelLink>& UControlRigModel::Links() const
+{
+	return _Links;
+}
+
+TArray<FControlRigModelPin> UControlRigModel::LinkedPins(const FControlRigModelPair& InPin) const
+{
+	ensure(InPin.Node >= 0 && InPin.Node < _Nodes.Num());
+	ensure(InPin.Pin >= 0 && InPin.Pin < _Nodes[InPin.Node].Pins.Num());
+
+	TArray<FControlRigModelPin> Pins;
+	const FControlRigModelPin& Pin = _Nodes[InPin.Node].Pins[InPin.Pin];
+
+	for (int32 LinkIndex : Pin.Links)
+	{
+		const FControlRigModelLink& Link = _Links[LinkIndex];
+		if (Pin.Direction == EGPD_Input)
+		{
+			Pins.Add(_Nodes[Link.Source.Node].Pins[Link.Source.Pin]);
+		}
+		else
+		{
+			Pins.Add(_Nodes[Link.Target.Node].Pins[Link.Target.Pin]);
+		}
+	}
+
+	return Pins;
+}
+
+TArray<FControlRigModelPin> UControlRigModel::LinkedPins(const FName& InNodeName, const FName& InPinName, bool bLookForInput) const
+{
+	TArray<FControlRigModelPin> Pins;
+	const FControlRigModelPin* Pin = FindPin(InNodeName, InPinName, bLookForInput);
+	if(Pin == nullptr)
+	{
+		if (_ModifiedEvent.IsBound())
+		{
+			FControlRigModelError Error;
+			Error.Message = FString::Printf(TEXT("Pin '%s.%s' cannot be found."), *InNodeName.ToString(), *InPinName.ToString());
+			_ModifiedEvent.Broadcast(this, EControlRigModelNotifType::ModelError, &Error);
+		}
+		return Pins;
+	}
+
+	return LinkedPins(Pin->GetPair());
+}
+
+TArray<FControlRigModelNode> UControlRigModel::Parameters() const
+{
+	TArray<FControlRigModelNode> ParameterNodes;
+
+	for (const FControlRigModelNode& Node : _Nodes)
+	{
+		if (Node.ParameterType != EControlRigModelParameterType::None)
+		{
+			ParameterNodes.Add(Node);
+		}
+	}
+
+	return ParameterNodes;
+}
+
+UControlRigModel::FModifiedEvent& UControlRigModel::OnModified()
+{
+	return _ModifiedEvent;
+}
+
+bool UControlRigModel::Clear()
+{
+	if (_Nodes.Num() == 0)
+	{
+		return false;
+	}
+
+	if (_ModifiedEvent.IsBound())
+	{
+		_ModifiedEvent.Broadcast(this, EControlRigModelNotifType::ModelCleared, nullptr);
+	}
+
+	_Nodes.Reset();
+	_Links.Reset();
+	_SelectedNodes.Reset();
+	_LastNodePositions.Reset();
+
+	return true;
+}
+
+bool UControlRigModel::IsNodeNameAvailable(const FName& InName) const
+{
+	return FindNode(InName) == nullptr;
+}
+
+FName UControlRigModel::GetUniqueNodeName(const FName& InName) const
+{
+	return UtilityHelpers::CreateUniqueName(InName, [this](const FName& CurName) { return IsNodeNameAvailable(CurName); });
+}
+
+const FControlRigModelNode* UControlRigModel::FindNode(const FName& InName) const
+{
+	for (const FControlRigModelNode& Node : _Nodes)
+	{
+		if (Node.Name == InName)
+		{
+			return &Node;
+		}
+	}
+	return nullptr;
+}
+
+const FControlRigModelNode* UControlRigModel::FindNode(int32 InNodeIndex) const
+{
+	ensure(InNodeIndex >= 0 && InNodeIndex < _Nodes.Num());
+	return &_Nodes[InNodeIndex];
+}
+
+bool UControlRigModel::AddNode(const FControlRigModelNode& InNode, bool bUndo)
+{
+	ensure(!InNode.IsParameter() && (InNode.FunctionName != NAME_None));
+
+	FName DesiredNodeName = InNode.Name == NAME_None ? InNode.FunctionName : InNode.Name;
+	if (InNode.UnitStruct() == nullptr)
+	{
+		if (_ModifiedEvent.IsBound())
+		{
+			FControlRigModelError Error;
+			Error.Message = FString::Printf(TEXT("Node '%s' has no function specified. Cannot add node."), *DesiredNodeName.ToString());
+			_ModifiedEvent.Broadcast(this, EControlRigModelNotifType::ModelError, &Error);
+		}
+		return false;
+	}
+
+	if (DesiredNodeName.ToString().StartsWith(TEXT("RigUnit_")))
+	{
+		DesiredNodeName = *DesiredNodeName.ToString().RightChop(8);
+	}
+
+	FControlRigModelNode NodeToAdd = InNode;
+	NodeToAdd.Name = GetUniqueNodeName(DesiredNodeName);
+
+	AddNodePinsForFunction(NodeToAdd);
+
+#if CONTROLRIG_UNDO
+
+	FAction MainAction;
+	FAction AddNodeAction;
+	if (bUndo)
+	{
+		MainAction.Title = FString::Printf(TEXT("Added Node '%s'"), *NodeToAdd.Name.ToString());
+		MainAction.Type = EControlRigModelNotifType::Invalid;
+		CurrentActions.Add(&MainAction);
+
+		AddNodeAction.Type = EControlRigModelNotifType::NodeAdded;
+		AddNodeAction.Title = MainAction.Title;
+		NodeToAdd.AppendArgumentsForAction(AddNodeAction.Arguments);
+	}
+
+#endif
+
+	FName PreviousExecuteNode = NAME_None;
+	int32 PreviousExecutePin = INDEX_NONE;
+
+	// if this is a mutable node let's add the begin execution and / or wire it up
+	if (NodeToAdd.UnitStruct()->IsChildOf(FRigUnitMutable::StaticStruct()))
+	{
+		float ClosestDistance = FLT_MAX;
+
+		// try to hook up the rig execution pin automatically for the user
+		for (const FControlRigModelNode& ExistingNode : _Nodes)
+		{
+			for (const FControlRigModelPin& ExistingPin : ExistingNode.Pins)
+			{
+				if (ExistingPin.Direction != EGPD_Output)
+				{
+					continue;
+				}
+				if (ExistingPin.Type.PinSubCategoryObject != FControlRigExecuteContext::StaticStruct())
+				{
+					continue;
+				}
+				if (ExistingPin.Links.Num() > 0)
+				{
+					continue;
+				}
+
+				float Distance = (NodeToAdd.Position - ExistingNode.Position).SizeSquared();
+				if (Distance < ClosestDistance)
+				{
+					ClosestDistance = Distance;
+					PreviousExecuteNode = ExistingNode.Name;
+					PreviousExecutePin = ExistingPin.Index;
+				}
+			}
+		}
+
+		// if we didn't find a closest rig node with an execution pin - let's create one
+		if (PreviousExecuteNode == NAME_None)
+		{
+			FControlRigModelNode BeginExecutionNode;
+			BeginExecutionNode.FunctionName = FRigUnit_BeginExecution::StaticStruct()->GetFName();
+			BeginExecutionNode.Position = NodeToAdd.Position - FVector2D(200.f, 0.f);
+			BeginExecutionNode.Index = _Nodes.Num();
+
+			if (AddNode(BeginExecutionNode, bUndo))
+			{
+				PreviousExecuteNode = _Nodes[BeginExecutionNode.Index].Name;
+				PreviousExecutePin = 0;
+			}
+		}
+	}
+
+	NodeToAdd.Index = _Nodes.Num();
+	ConfigurePinIndices(NodeToAdd);
+
+	_Nodes.Add(NodeToAdd);
+	FControlRigModelNode& AddedNode = _Nodes[NodeToAdd.Index];
+	_LastNodePositions.Add(AddedNode.Name, AddedNode.Position);
+
+	SetNodePinDefaultsForFunction(AddedNode);
+
+	ResetCycleCheck();
+
+	if (_ModifiedEvent.IsBound())
+	{
+		_ModifiedEvent.Broadcast(this, EControlRigModelNotifType::NodeAdded, &AddedNode);
+	}
+
+#if CONTROLRIG_UNDO
+
+	if (bUndo)
+	{
+		PushAction(AddNodeAction);
+	}
+
+#endif
+
+	// only hook up the node automatically if we are currently
+	// performing for undo. Otherwise it will be taken care of by the nested actions.
+	if (PreviousExecuteNode != NAME_None && bUndo)
+	{
+		for (const FControlRigModelPin& AddedPin : AddedNode.Pins)
+		{
+			if (AddedPin.Direction != EGPD_Input)
+			{
+				continue;
+			}
+			if (AddedPin.Type.PinSubCategoryObject != FControlRigExecuteContext::StaticStruct())
+			{
+				continue;
+			}
+
+			const FControlRigModelNode* PreviousNode = FindNode(PreviousExecuteNode);
+			if (PreviousNode)
+			{
+				MakeLink(PreviousNode->Index, PreviousExecutePin, AddedNode.Index, AddedPin.Index, bUndo);
+			}
+			break;
+		}
+	}
+
+#if CONTROLRIG_UNDO
+
+	if (bUndo)
+	{
+		CurrentActions.Pop();
+		PushAction(MainAction);
+	}
+
+#endif
+
+	return true;
+}
+
+bool UControlRigModel::AddParameter(const FName& InName, const FEdGraphPinType& InDataType, EControlRigModelParameterType InParameterType, const FVector2D& InPosition, bool bUndo)
+{
+	FControlRigModelNode Parameter;
+	Parameter.Name = GetUniqueNodeName(InName);
+
+	Parameter.ParameterType = InParameterType;
+	Parameter.Position = InPosition;
+
+	AddNodePinsForParameter(Parameter, InDataType);
+
+	Parameter.Index = _Nodes.Num();
+	ConfigurePinIndices(Parameter);
+
+#if CONTROLRIG_UNDO
+	FAction Action;
+	if (bUndo)
+	{
+		CurrentActions.Add(&Action);
+		Action.Type = EControlRigModelNotifType::NodeAdded;
+		Action.Title = FString::Printf(TEXT("Added Parameter '%s'"), *Parameter.Name.ToString());
+		Parameter.AppendArgumentsForAction(Action.Arguments);
+	}
+#endif
+
+	_Nodes.Add(Parameter);
+	FControlRigModelNode& AddedNode = _Nodes.Last();
+
+	SetNodePinDefaultsForParameter(AddedNode, InDataType);
+
+	ResetCycleCheck();
+
+	if (_ModifiedEvent.IsBound())
+	{
+		_ModifiedEvent.Broadcast(this, EControlRigModelNotifType::NodeAdded, &AddedNode);
+	}
+
+#if CONTROLRIG_UNDO
+	if (bUndo)
+	{
+		CurrentActions.Pop();
+		PushAction(Action);
+	}
+#endif
+
+	return true;
+}
+
+bool UControlRigModel::RemoveNode(const FName& InName, bool bUndo)
+{
+	const FControlRigModelNode* Node = FindNode(InName);
+	if(Node != nullptr)
+	{
+#if CONTROLRIG_UNDO
+		FAction MainAction;
+		if (bUndo)
+		{
+			CurrentActions.Add(&MainAction);
+			MainAction.Type = EControlRigModelNotifType::Invalid;
+			MainAction.Title = FString::Printf(TEXT("Removed Node '%s'"), *Node->Name.ToString());
+		}
+#endif
+
+		int32 NodeIndex = Node->Index;
+		BreakLinks(NodeIndex, INDEX_NONE, bUndo);
+
+		FControlRigModelNode RemovedNode = _Nodes[NodeIndex];
+
+#if CONTROLRIG_UNDO
+		if (bUndo)
+		{
+			for (const FControlRigModelPin& Pin : RemovedNode.Pins)
+			{
+				if (Pin.Direction == EGPD_Input)
+				{
+					FAction SetPinDefaultAction;
+					SetPinDefaultAction.Title = TEXT("Set pin default");
+					SetPinDefaultAction.Type = EControlRigModelNotifType::PinChanged;
+					SetPinDefaultAction.Arguments.Add(GetPinPath(Pin.GetPair()));
+					Pin.AppendArgumentsForAction(SetPinDefaultAction.Arguments);
+					Pin.AppendArgumentsForAction(SetPinDefaultAction.Arguments);
+					PushAction(SetPinDefaultAction);
+				}
+			}
+
+			for (const FControlRigModelPin& Pin : RemovedNode.Pins)
+			{
+				if (Pin.IsArray() && Pin.ArraySize() > 0)
+				{
+					FAction ResizePinAction;
+					ResizePinAction.Title = TEXT("Resize Pin.");
+					ResizePinAction.Type = EControlRigModelNotifType::PinAdded;
+					ResizePinAction.Arguments.Add(GetPinPath(Pin.GetPair()));
+					ResizePinAction.Arguments.Add(FString());
+					ResizePinAction.Arguments.Add(FString::FormatAsNumber(Pin.ArraySize()));
+					ResizePinAction.Arguments.Add(FString::FormatAsNumber(0));
+					PushAction(ResizePinAction);
+				}
+			}
+		}
+#endif
+
+		_Nodes.RemoveAt(NodeIndex);
+
+		int32 SelectedIndex = _SelectedNodes.Find(RemovedNode.Name);
+		if (SelectedIndex != INDEX_NONE)
+		{
+			_SelectedNodes.RemoveAt(SelectedIndex);
+		}
+
+		for (FControlRigModelNode& OtherNode : _Nodes)
+		{
+			if (OtherNode.Index > NodeIndex)
+			{
+				OtherNode.Index--;
+
+				for (FControlRigModelPin& Pin : OtherNode.Pins)
+				{
+					Pin.Node = OtherNode.Index;
+				}
+			}
+		}
+
+		for (FControlRigModelLink& Link : _Links)
+		{
+			if (Link.Source.Node > NodeIndex)
+			{
+				Link.Source.Node--;
+			}
+			if (Link.Target.Node > NodeIndex)
+			{
+				Link.Target.Node--;
+			}
+		}
+
+		ResetCycleCheck();
+
+		if (SelectedIndex != INDEX_NONE)
+		{
+			if (_ModifiedEvent.IsBound())
+			{
+				_ModifiedEvent.Broadcast(this, EControlRigModelNotifType::NodeDeselected, &RemovedNode);
+			}
+		}
+
+		if (_ModifiedEvent.IsBound())
+		{
+			_ModifiedEvent.Broadcast(this, EControlRigModelNotifType::NodeRemoved, &RemovedNode);
+		}
+
+#if CONTROLRIG_UNDO
+		if (bUndo)
+		{
+			FAction RemoveAction;
+			RemoveAction.Type = EControlRigModelNotifType::NodeRemoved;
+			RemoveAction.Title = MainAction.Title;
+			RemovedNode.AppendArgumentsForAction(RemoveAction.Arguments);
+			PushAction(RemoveAction);
+
+			CurrentActions.Pop();
+			PushAction(MainAction);
+		}
+#endif
+
+		return true;
+	}
+	return false;
+}
+
+bool UControlRigModel::SetNodePosition(const FName& InName, const FVector2D& InPosition, bool bUndo)
+{
+	const FControlRigModelNode* Node = FindNode(InName);
+	if(Node != nullptr)
+	{
+#if CONTROLRIG_UNDO
+		FAction Action;
+		if (bUndo)
+		{
+			CurrentActions.Add(&Action);
+			Action.Title = FString::Printf(TEXT("Moved Node '%s'"), *Node->Name.ToString());
+			Action.Type = EControlRigModelNotifType::NodeChanged;
+
+			const FVector2D* PreviousPosition = _LastNodePositions.Find(Node->Name);
+			if (PreviousPosition)
+			{
+				_Nodes[Node->Index].Position = *PreviousPosition;
+			}
+			Node->AppendArgumentsForAction(Action.Arguments);
+		}
+#endif
+
+		_Nodes[Node->Index].Position = InPosition;
+
+#if CONTROLRIG_UNDO
+		if (bUndo)
+		{
+			Node->AppendArgumentsForAction(Action.Arguments);
+			_LastNodePositions.Add(Node->Name, Node->Position);
+		}
+#endif
+
+		if (_ModifiedEvent.IsBound())
+		{
+			_ModifiedEvent.Broadcast(this, EControlRigModelNotifType::NodeChanged, Node);
+		}
+
+#if CONTROLRIG_UNDO
+		if (bUndo)
+		{
+			CurrentActions.Pop();
+			PushAction(Action);
+		}
+#endif
+
+		return true;
+	}
+	return false;
+}
+
+bool UControlRigModel::SetParameterType(const FName& InName, EControlRigModelParameterType InParameterType, bool bUndo)
+{
+	ensure(InParameterType != EControlRigModelParameterType::None);
+
+	const FControlRigModelNode* Node = FindNode(InName);
+	if (Node != nullptr)
+	{
+		ensure(Node->IsParameter());
+
+		if (_Nodes[Node->Index].ParameterType != InParameterType)
+		{
+#if CONTROLRIG_UNDO
+			FAction Action;
+			if (bUndo)
+			{
+				CurrentActions.Add(&Action);
+				Action.Title = FString::Printf(TEXT("Set Parameter Type for Node '%s'"), *Node->Name.ToString());
+				Action.Type = EControlRigModelNotifType::NodeChanged;
+				Node->AppendArgumentsForAction(Action.Arguments);
+			}
+#endif
+
+			_Nodes[Node->Index].ParameterType = InParameterType;
+
+#if CONTROLRIG_UNDO
+			if (bUndo)
+			{
+				Node->AppendArgumentsForAction(Action.Arguments);
+			}
+#endif
+
+			if (_ModifiedEvent.IsBound())
+			{
+				_ModifiedEvent.Broadcast(this, EControlRigModelNotifType::NodeChanged, Node);
+			}
+
+#if CONTROLRIG_UNDO
+			if (bUndo)
+			{
+				CurrentActions.Pop();
+				PushAction(Action);
+			}
+#endif
+
+			return true;
+		}
+	}
+	return false;
+}
+
+bool UControlRigModel::RenameNode(const FName& InOldNodeName, const FName& InNewNodeName, bool bUndo)
+{
+	FControlRigModelNode* Node = (FControlRigModelNode*)FindNode(InOldNodeName);
+	if (Node == nullptr)
+	{
+		if (_ModifiedEvent.IsBound())
+		{
+			FControlRigModelError Error;
+			Error.Message = FString::Printf(TEXT("Node '%s' cannot be found."), *InOldNodeName.ToString());
+			_ModifiedEvent.Broadcast(this, EControlRigModelNotifType::ModelError, &Error);
+		}
+		return false;
+	}
+
+	if (InOldNodeName == InNewNodeName)
+	{
+		return false;
+	}
+
+#if CONTROLRIG_UNDO
+
+	FAction Action;
+	if (bUndo)
+	{
+		CurrentActions.Add(&Action);
+		Action.Type = EControlRigModelNotifType::NodeChanged;
+		Action.Title = FString::Printf(TEXT("Renamed Node '%s'"), *Node->Name.ToString());
+		Action.Arguments.Add(Node->Name.ToString());
+	}
+
+#endif
+
+	Node->Name = GetUniqueNodeName(InNewNodeName);
+
+#if CONTROLRIG_UNDO
+
+	if (bUndo)
+	{
+		Action.Arguments.Add(Node->Name.ToString());
+	}
+
+#endif
+
+	for (FName& SelectedNode : _SelectedNodes)
+	{
+		if (SelectedNode == InOldNodeName)
+		{
+			SelectedNode = Node->Name;
+		}
+	}
+
+	if (_ModifiedEvent.IsBound())
+	{
+		FControlRigModelNodeRenameInfo Info;
+		Info.OldName = InOldNodeName;
+		Info.NewName = Node->Name;
+		Info.Node = *Node;
+		_ModifiedEvent.Broadcast(this, EControlRigModelNotifType::NodeRenamed, &Info);
+	}
+
+#if CONTROLRIG_UNDO
+
+	if (bUndo)
+	{
+		CurrentActions.Pop();
+		PushAction(Action);
+	}
+
+#endif
+
+	return true;
+}
+
+bool UControlRigModel::SelectNode(const FName& InName, bool bInSelected)
+{
+	if (bIsSelecting)
+	{
+		return false;
+	}
+	TGuardValue<bool> ReentrantGuard(bIsSelecting, true);
+
+	const FControlRigModelNode* Node = FindNode(InName);
+	if(Node != nullptr)
+	{
+		int32 SelectedIndex = _SelectedNodes.Find(InName);
+		if(bInSelected)
+		{
+			if(SelectedIndex == INDEX_NONE)
+			{
+				_SelectedNodes.Add(InName);
+				_LastNodePositions.Add(InName, Node->Position);
+				if (_ModifiedEvent.IsBound())
+				{
+					_ModifiedEvent.Broadcast(this, EControlRigModelNotifType::NodeSelected, Node);
+				}
+				return true;
+			}
+		}
+		else
+		{
+			if(SelectedIndex != INDEX_NONE)
+			{
+				_SelectedNodes.RemoveAt(SelectedIndex);
+				if (_ModifiedEvent.IsBound())
+				{
+					_ModifiedEvent.Broadcast(this, EControlRigModelNotifType::NodeDeselected, Node);
+				}
+				return true;
+			}
+		}
+	}
+	return false;
+}
+
+bool UControlRigModel::AreCompatibleTypes(const FEdGraphPinType& A, const FEdGraphPinType& B) const
+{
+	// support casts in the future
+	return A == B;
+}
+
+bool UControlRigModel::PrepareCycleCheckingForPin(int32 InNodeIndex, int32 InPinIndex)
+{
+	ensure(InNodeIndex >= 0 && InNodeIndex < _Nodes.Num());
+	ensure(InPinIndex >= 0 && InPinIndex < _Nodes[InNodeIndex].Pins.Num());
+
+	if (_CycleCheckSubject.Node == InNodeIndex && _CycleCheckSubject.Pin == InPinIndex)
+	{
+		return true;
+	}
+
+	_CycleCheckSubject.Node = InNodeIndex;
+	_CycleCheckSubject.Pin = InPinIndex;
+	_NodeIsOnCycle.SetNumZeroed(_Nodes.Num());
+
+	struct Local
+	{
+		static void VisitNode(int32 Index, const TArray<FControlRigModelNode>& Nodes, const TArray<FControlRigModelLink>& Links, TArray<bool>& Visited, bool bWalkInputs)
+		{
+			if (Visited[Index])
+			{
+				return;
+			}
+
+			Visited[Index] = true;
+
+			const FControlRigModelNode& Node = Nodes[Index];
+			for (const FControlRigModelPin& Pin : Node.Pins)
+			{
+				if ((Pin.Direction == EGPD_Input) != bWalkInputs)
+				{
+					continue;
+				}
+				for (int32 LinkIndex : Pin.Links)
+				{
+					int32 OtherIndex = bWalkInputs ? Links[LinkIndex].Source.Node : Links[LinkIndex].Target.Node;
+					VisitNode(OtherIndex, Nodes, Links, Visited, bWalkInputs);
+				}
+			}
+		}
+	};
+
+	bool bWalkInputs = _Nodes[InNodeIndex].Pins[InPinIndex].Direction != EGPD_Input;
+	Local::VisitNode(InNodeIndex, _Nodes, _Links, _NodeIsOnCycle, bWalkInputs);
+
+	return true;
+}
+
+void UControlRigModel::ResetCycleCheck()
+{
+	_CycleCheckSubject = FControlRigModelPair();
+	_NodeIsOnCycle.Reset();
+}
+
+bool UControlRigModel::CanLink(int32 InSourceNodeIndex, int32 InSourcePinIndex, int32 InTargetNodeIndex, int32 InTargetPinIndex, FString* OutFailureReason)
+{
+	ensure(InSourceNodeIndex >= 0 && InSourceNodeIndex < _Nodes.Num());
+	ensure(InTargetNodeIndex >= 0 && InTargetNodeIndex < _Nodes.Num());
+	ensure(InSourcePinIndex >= 0 && InSourcePinIndex < _Nodes[InSourceNodeIndex].Pins.Num());
+	ensure(InTargetPinIndex >= 0 && InTargetPinIndex < _Nodes[InTargetNodeIndex].Pins.Num());
+
+	if (InSourceNodeIndex == InTargetNodeIndex)
+	{
+		if (OutFailureReason != nullptr)
+		{
+			*OutFailureReason = TEXT("Source and target pins are on the same node.");
+		}
+		return false;
+	}
+
+	const FControlRigModelNode& SourceNode = _Nodes[InSourceNodeIndex];
+	const FControlRigModelPin& SourcePin = SourceNode.Pins[InSourcePinIndex];
+	const FControlRigModelNode& TargetNode = _Nodes[InTargetNodeIndex];
+	const FControlRigModelPin& TargetPin = TargetNode.Pins[InTargetPinIndex];
+
+	if (SourcePin.Direction == TargetPin.Direction)
+	{
+		if (OutFailureReason != nullptr)
+		{
+			*OutFailureReason = TEXT("Source and target have the same direction.");
+		}
+		return false;
+	}
+
+	if (SourcePin.Direction != EGPD_Output)
+	{
+		if (OutFailureReason != nullptr)
+		{
+			*OutFailureReason = TEXT("Source pin is not an output.");
+		}
+		return false;
+	}
+	if (TargetPin.Direction != EGPD_Input)
+	{
+		if (OutFailureReason != nullptr)
+		{
+			*OutFailureReason = TEXT("Target pin is not an input.");
+		}
+		return false;
+	}
+	if (!AreCompatibleTypes(SourcePin.Type, TargetPin.Type))
+	{
+		if (OutFailureReason != nullptr)
+		{
+			*OutFailureReason = TEXT("Types are not compatible.");
+		}
+		return false;
+	}
+
+	if (TargetPin.bIsConstant && !SourcePin.bIsConstant)
+	{
+		if (OutFailureReason != nullptr)
+		{
+			*OutFailureReason = TEXT("Only constant values can be connected to constants.");
+		}
+		return false;
+	}
+	
+	for (int32 LinkIndex : TargetPin.Links)
+	{
+		if (_Links[LinkIndex].Source.Node == InSourceNodeIndex &&
+			_Links[LinkIndex].Source.Pin == InSourcePinIndex)
+		{
+			if (OutFailureReason != nullptr)
+			{
+				*OutFailureReason = TEXT("Pins already linked.");
+			}
+			return false;
+		}
+	}
+
+	bool bCycleCheckWasSetup = _CycleCheckSubject.IsValid() && (_CycleCheckSubject.Node == InSourceNodeIndex || _CycleCheckSubject.Node == InTargetNodeIndex);
+	if (!bCycleCheckWasSetup)
+	{
+		PrepareCycleCheckingForPin(InSourceNodeIndex, InSourcePinIndex);
+	}
+
+	struct Local
+	{
+		static bool TestNodeOnCycle(int32 Index, const TArray<FControlRigModelNode>& Nodes, const TArray<FControlRigModelLink>& Links, TArray<bool>& Visited, bool bWalkInputs)
+		{
+			if (Visited[Index])
+			{
+				return true;
+			}
+
+			const FControlRigModelNode& Node = Nodes[Index];
+			for (const FControlRigModelPin& Pin : Node.Pins)
+			{
+				if ((Pin.Direction == EGPD_Input) != bWalkInputs)
+				{
+					continue;
+				}
+				for (int32 LinkIndex : Pin.Links)
+				{
+					int32 OtherIndex = bWalkInputs ? Links[LinkIndex].Source.Node : Links[LinkIndex].Target.Node;
+					if (TestNodeOnCycle(OtherIndex, Nodes, Links, Visited, bWalkInputs))
+					{
+						Visited[Index] = true;
+						return true;
+					}
+				}
+			}
+
+			return false;
+		}
+	};
+
+	bool bWalkInputs = _CycleCheckSubject.Node == InSourceNodeIndex;
+	if (bWalkInputs)
+	{
+		if (Local::TestNodeOnCycle(InTargetNodeIndex, _Nodes, _Links, _NodeIsOnCycle, false))
+		{
+			if (OutFailureReason != nullptr)
+			{
+				*OutFailureReason = TEXT("Cannot create a cycle.");
+			}
+			return false;
+		}
+	}
+	else
+	{
+		if (Local::TestNodeOnCycle(InSourceNodeIndex, _Nodes, _Links, _NodeIsOnCycle, true))
+		{
+			if (OutFailureReason != nullptr)
+			{
+				*OutFailureReason = TEXT("Cannot create a cycle.");
+			}
+			return false;
+		}
+	}
+
+	if (!bCycleCheckWasSetup)
+	{
+		ResetCycleCheck();
+	}
+
+	return true;
+}
+
+bool UControlRigModel::MakeLink(int32 InSourceNodeIndex, int32 InSourcePinIndex, int32 InTargetNodeIndex, int32 InTargetPinIndex, bool bUndo)
+{
+	ensure(InSourceNodeIndex >= 0 && InSourceNodeIndex < _Nodes.Num());
+	ensure(InTargetNodeIndex >= 0 && InTargetNodeIndex < _Nodes.Num());
+	ensure(InSourcePinIndex >= 0 && InSourcePinIndex < _Nodes[InSourceNodeIndex].Pins.Num());
+	ensure(InTargetPinIndex >= 0 && InTargetPinIndex < _Nodes[InTargetNodeIndex].Pins.Num());
+	ensure(CanLink(InSourceNodeIndex, InSourcePinIndex, InTargetNodeIndex, InTargetPinIndex, nullptr));
+
+	TArray<int32> PinsToDisconnect;
+	PinsToDisconnect.Add(InTargetPinIndex);
+
+	for (int32 PinIndex = 0; PinIndex < PinsToDisconnect.Num(); PinIndex++)
+	{
+		for (int32 SubPinIndex : _Nodes[InTargetNodeIndex].Pins[PinsToDisconnect[PinIndex]].SubPins)
+		{
+			PinsToDisconnect.Add(SubPinIndex);
+		}
+	}
+
+#if CONTROLRIG_UNDO
+	FAction Action;
+	if (bUndo)
+	{
+		CurrentActions.Add(&Action);
+	}
+#endif
+
+	int32 ParentPinIndex = InTargetPinIndex;
+	while (ParentPinIndex != INDEX_NONE)
+	{
+		ParentPinIndex = _Nodes[InTargetNodeIndex].Pins[ParentPinIndex].ParentIndex;
+		if (ParentPinIndex != INDEX_NONE)
+		{
+			PinsToDisconnect.Add(ParentPinIndex);
+		}
+	}
+
+	for (int32 PinToDisconnect : PinsToDisconnect)
+	{
+		BreakLinks(InTargetNodeIndex, PinToDisconnect, bUndo);
+	}
+
+	FControlRigModelLink Link;
+	Link.Index = _Links.Num();
+	Link.Source.Node = InSourceNodeIndex;
+	Link.Source.Pin = InSourcePinIndex;
+	Link.Target.Node = InTargetNodeIndex;
+	Link.Target.Pin = InTargetPinIndex;
+	_Links.Add(Link);
+
+	_Nodes[InSourceNodeIndex].Pins[InSourcePinIndex].Links.Add(Link.Index);
+	_Nodes[InTargetNodeIndex].Pins[InTargetPinIndex].Links.Add(Link.Index);
+
+	ResetCycleCheck();
+
+	if (_ModifiedEvent.IsBound())
+	{
+		_ModifiedEvent.Broadcast(this, EControlRigModelNotifType::LinkAdded, &Link);
+	}
+
+#if CONTROLRIG_UNDO
+	if (bUndo)
+	{
+		CurrentActions.Pop();
+		Action.Type = EControlRigModelNotifType::LinkAdded;
+		Action.Title = TEXT("Added Link.");
+		Link.AppendArgumentsForAction(Action.Arguments, this);
+		PushAction(Action);
+	}
+#endif
+
+	return true;
+}
+
+bool UControlRigModel::BreakLink(int32 InSourceNodeIndex, int32 InSourcePinIndex, int32 InTargetNodeIndex, int32 InTargetPinIndex, bool bUndo)
+{
+	for (int32 LinkIndex = 0; LinkIndex < _Links.Num(); LinkIndex++)
+	{
+		FControlRigModelLink Link = _Links[LinkIndex];
+		if (Link.Source.Node == InSourceNodeIndex && Link.Source.Pin == InSourcePinIndex &&
+			Link.Target.Node == InTargetNodeIndex && Link.Target.Pin == InTargetPinIndex)
+		{
+			_Links.RemoveAt(LinkIndex);
+
+			for (int32 OtherLinkIndex = LinkIndex; OtherLinkIndex < _Links.Num(); OtherLinkIndex++)
+			{
+				if (_Links[OtherLinkIndex].Index > LinkIndex)
+				{
+					_Links[OtherLinkIndex].Index--;
+				}
+			}
+
+			ensure(_Nodes[Link.Source.Node].Pins[Link.Source.Pin].Links.Remove(LinkIndex) == 1);
+			ensure(_Nodes[Link.Target.Node].Pins[Link.Target.Pin].Links.Remove(LinkIndex) == 1);
+
+			for (FControlRigModelNode& Node : _Nodes)
+			{
+				for (FControlRigModelPin& Pin : Node.Pins)
+				{
+					for (int32 PinLinkIndex = 0; PinLinkIndex < Pin.Links.Num(); PinLinkIndex++)
+					{
+						if (Pin.Links[PinLinkIndex] > LinkIndex)
+						{
+							Pin.Links[PinLinkIndex]--;
+						}
+					}
+				}
+			}
+
+			if (_ModifiedEvent.IsBound())
+			{
+				_ModifiedEvent.Broadcast(this, EControlRigModelNotifType::LinkRemoved, &Link);
+			}
+
+#if CONTROLRIG_UNDO
+			if (bUndo)
+			{
+				FAction Action;
+				Action.Type = EControlRigModelNotifType::LinkRemoved;
+				Action.Title = TEXT("Broke Link.");
+				Link.AppendArgumentsForAction(Action.Arguments, this);
+				PushAction(Action);
+			}
+#endif
+
+			ResetCycleCheck();
+
+			return true;
+		}
+	}
+	return false;
+}
+
+bool UControlRigModel::BreakLinks(int32 InNodeIndex, int32 InPinIndex, bool bUndo)
+{
+#if CONTROLRIG_UNDO
+	FAction MainAction;
+	if (bUndo)
+	{
+		CurrentActions.Add(&MainAction);
+		MainAction.Type = EControlRigModelNotifType::Invalid;
+	}
+#endif
+
+	TArray<int32> LinksRemoved;
+	for (int32 LinkIndex = _Links.Num() - 1; LinkIndex >= 0; LinkIndex--)
+	{
+		FControlRigModelLink Link = _Links[LinkIndex];
+		if ((Link.Source.Node == InNodeIndex && (Link.Source.Pin == InPinIndex || InPinIndex == INDEX_NONE)) ||
+			(Link.Target.Node == InNodeIndex && (Link.Target.Pin == InPinIndex || InPinIndex == INDEX_NONE)))
+		{
+			LinksRemoved.Add(LinkIndex);
+			_Links.RemoveAt(LinkIndex);
+
+			for (int32 OtherLinkIndex = LinkIndex; OtherLinkIndex < _Links.Num(); OtherLinkIndex++)
+			{
+				if (_Links[OtherLinkIndex].Index > LinkIndex)
+				{
+					_Links[OtherLinkIndex].Index--;
+				}
+			}
+
+			ensure(_Nodes[Link.Source.Node].Pins[Link.Source.Pin].Links.Remove(LinkIndex) == 1);
+			ensure(_Nodes[Link.Target.Node].Pins[Link.Target.Pin].Links.Remove(LinkIndex) == 1);
+
+			for (FControlRigModelNode& Node : _Nodes)
+			{
+				for (FControlRigModelPin& Pin : Node.Pins)
+				{
+					for (int32 PinLinkIndex = 0; PinLinkIndex < Pin.Links.Num(); PinLinkIndex++)
+					{
+						if (Pin.Links[PinLinkIndex] > LinkIndex)
+						{
+							Pin.Links[PinLinkIndex]--;
+						}
+					}
+				}
+			}
+
+			if (_ModifiedEvent.IsBound())
+			{
+				_ModifiedEvent.Broadcast(this, EControlRigModelNotifType::LinkRemoved, &Link);
+			}
+
+#if CONTROLRIG_UNDO
+			if (bUndo)
+			{
+				FAction Action;
+				Action.Type = EControlRigModelNotifType::LinkRemoved;
+				Action.Title = TEXT("Broke all Links for Pin.");
+				Link.AppendArgumentsForAction(Action.Arguments, this);
+				PushAction(Action);
+			}
+#endif
+		}
+	}
+
+	if (LinksRemoved.Num() > 0)
+	{
+		ResetCycleCheck();
+	}
+
+#if CONTROLRIG_UNDO
+	if (bUndo)
+	{
+		CurrentActions.Pop();
+		PushAction(MainAction);
+	}
+#endif
+
+	return LinksRemoved.Num() > 0;
+}
+
+const FControlRigModelPin* UControlRigModel::FindPin(const FName& InNodeName, const FName& InPinName, bool bLookForInput) const
+{
+	const FControlRigModelNode* Node = FindNode(InNodeName);
+	if (Node != nullptr)
+	{
+		return Node->FindPin(InPinName, bLookForInput);
+	}
+	return nullptr;
+}
+
+const FControlRigModelPin* UControlRigModel::FindPin(const FControlRigModelPair& InPin) const
+{
+	ensure(InPin.Node >= 0 && InPin.Node < _Nodes.Num());
+	ensure(InPin.Pin >= 0 && InPin.Pin < _Nodes[InPin.Node].Pins.Num());
+	return &_Nodes[InPin.Node].Pins[InPin.Pin];
+}
+
+const FControlRigModelPin* UControlRigModel::FindSubPin(const FControlRigModelPin* InParentPin, const FName& InSubPinName) const
+{
+	if (InParentPin->SubPins.Num() == 0)
+	{
+		return nullptr;
+	}
+
+	ensure(InParentPin->Node >= 0 && InParentPin->Node < _Nodes.Num());
+
+	for (int32 SubPinIndex : _Nodes[InParentPin->Node].Pins[InParentPin->Index].SubPins)
+	{
+		const FControlRigModelPin* SubPin = &_Nodes[InParentPin->Node].Pins[SubPinIndex];
+		if (SubPin->Name == InSubPinName)
+		{
+			return SubPin;
+		}
+	}
+	return nullptr;
+}
+
+const FControlRigModelPin* UControlRigModel::FindParentPin(const FControlRigModelPin* InSubPin) const
+{
+	if (InSubPin->ParentIndex == INDEX_NONE)
+	{
+		return nullptr;
+	}
+
+	ensure(InSubPin->Node >= 0 && InSubPin->Node < _Nodes.Num());
+	ensure(InSubPin->ParentIndex >= 0 && InSubPin->ParentIndex < _Nodes[InSubPin->Node].Pins.Num());
+	return &_Nodes[InSubPin->Node].Pins[InSubPin->ParentIndex];
+}
+
+const FControlRigModelPin* UControlRigModel::FindPinFromPath(const FString& InPinPath, bool bLookForInput) const
+{
+	FString Left, Right;
+	if (!InPinPath.Split(TEXT("."), &Left, &Right))
+	{
+		Left = InPinPath;
+		Right = TEXT("Value");
+	}
+	return FindPin(*Left, *Right, bLookForInput);
+}
+
+const FControlRigModelLink* UControlRigModel::FindLink(int32 InLinkIndex) const
+{
+	ensure(InLinkIndex >= 0 && InLinkIndex < _Links.Num());
+	return &_Links[InLinkIndex];
+}
+
+void UControlRigModel::SplitPinPath(const FString& InPinPath, FString& OutLeft, FString& OutRight, bool bInSplitForNodeName)
+{
+	FString Left, Right;
+	int32 PeriodPos = INDEX_NONE;
+	int32 BracketPos = INDEX_NONE;
+	InPinPath.FindChar('.', PeriodPos);
+	InPinPath.FindChar('[', BracketPos);
+
+	if (PeriodPos != INDEX_NONE && BracketPos != INDEX_NONE)
+	{
+		if (PeriodPos < BracketPos)
+		{
+			BracketPos = INDEX_NONE;
+		}
+		else
+		{
+			PeriodPos = INDEX_NONE;
+		}
+	}
+
+	if (PeriodPos != INDEX_NONE)
+	{
+		InPinPath.Split(TEXT("."), &Left, &Right);
+	}
+	else if (BracketPos != INDEX_NONE)
+	{
+		InPinPath.Split(TEXT("["), &Left, &Right);
+
+		FString RightLeft, RightRight;
+		if (Right.Split(TEXT("]"), &RightLeft, &RightRight))
+		{
+			Right = RightLeft + RightRight;
+		}
+	}
+	else
+	{
+		Left = InPinPath;
+
+		if (bInSplitForNodeName)
+		{
+			Right = TEXT("Value");
+		}
+		else
+		{
+			Right.Empty();
+		}
+	}
+
+	OutLeft = Left;
+	OutRight = Right;
+}
+
+const FControlRigModelPin* UControlRigModel::GetParentPin(const FControlRigModelPair& InPin) const
+{
+	ensure(InPin.Node >= 0 && InPin.Node < _Nodes.Num());
+	ensure(InPin.Pin >= 0 && InPin.Pin < _Nodes[InPin.Node].Pins.Num());
+	const FControlRigModelNode& Node = _Nodes[InPin.Node];
+	if (Node.Pins[InPin.Pin].ParentIndex == INDEX_NONE)
+	{
+		if (_ModifiedEvent.IsBound())
+		{
+			FControlRigModelError Error;
+			FString PinPath = GetPinPath(InPin);
+			Error.Message = FString::Printf(TEXT("Pin '%s' has no parent pin."), *PinPath);
+			_ModifiedEvent.Broadcast(this, EControlRigModelNotifType::ModelError, &Error);
+		}
+		return nullptr;
+	}
+	return &Node.Pins[Node.Pins[InPin.Pin].ParentIndex];
+}
+
+FString UControlRigModel::GetPinPath(const FControlRigModelPair& InPin, bool bIncludeNodeName) const
+{
+	ensure(InPin.Node >= 0 && InPin.Node < _Nodes.Num());
+	ensure(InPin.Pin >= 0 && InPin.Pin < _Nodes[InPin.Node].Pins.Num());
+	const FControlRigModelNode& Node = _Nodes[InPin.Node];
+	return Node.GetPinPath(InPin.Pin, bIncludeNodeName);
+}
+
+bool UControlRigModel::GetPinDefaultValue(const FName& InNodeName, const FName& InPinName, FString& OutValue) const
+{
+	const FControlRigModelPin* Pin = FindPin(InNodeName, InPinName, true);
+	if (Pin != nullptr)
+	{
+		return GetPinDefaultValue(Pin->GetPair(), OutValue);
+	}
+
+	if (_ModifiedEvent.IsBound())
+	{
+		FControlRigModelError Error;
+		Error.Message = FString::Printf(TEXT("Pin '%s.%s' cannot be found."), *InNodeName.ToString(), *InPinName.ToString());
+		_ModifiedEvent.Broadcast(this, EControlRigModelNotifType::ModelError, &Error);
+	}
+	return false;
+}
+
+bool UControlRigModel::GetPinDefaultValue(const FControlRigModelPair& InPin, FString& OutValue) const
+{
+	ensure(InPin.Node >= 0 && InPin.Node < _Nodes.Num());
+	ensure(InPin.Pin >= 0 && InPin.Pin < _Nodes[InPin.Node].Pins.Num());
+	OutValue = _Nodes[InPin.Node].Pins[InPin.Pin].DefaultValue;
+	return true;
+}
+
+bool UControlRigModel::SetPinDefaultValue(const FName& InNodeName, const FName& InPinName, const FString& InValue, bool bUndo)
+{
+	const FControlRigModelPin* Pin = FindPin(InNodeName, InPinName, true);
+	if (Pin != nullptr)
+	{
+		return SetPinDefaultValue(Pin->GetPair(), InValue, bUndo);
+	}
+
+	if (_ModifiedEvent.IsBound())
+	{
+		FControlRigModelError Error;
+		Error.Message = FString::Printf(TEXT("Pin '%s.%s' cannot be found."), *InNodeName.ToString(), *InPinName.ToString());
+		_ModifiedEvent.Broadcast(this, EControlRigModelNotifType::ModelError, &Error);
+	}
+	return false;
+}
+
+bool UControlRigModel::SetPinDefaultValue(const FControlRigModelPair& InPin, const FString& InValue, bool bUndo)
+{
+	ensure(InPin.Node >= 0 && InPin.Node < _Nodes.Num());
+	ensure(InPin.Pin >= 0 && InPin.Pin < _Nodes[InPin.Node].Pins.Num());
+
+	if (_Nodes[InPin.Node].Pins[InPin.Pin].DefaultValue == InValue)
+	{
+		return false;
+	}
+
+#if CONTROLRIG_UNDO
+	FAction Action;
+	if (bUndo)
+	{
+		Action.Type = EControlRigModelNotifType::PinChanged;
+		Action.Title = TEXT("Set Pin Default.");
+		Action.Arguments.Add(GetPinPath(InPin));
+		_Nodes[InPin.Node].Pins[InPin.Pin].AppendArgumentsForAction(Action.Arguments);
+	}
+#endif
+
+	_Nodes[InPin.Node].Pins[InPin.Pin].DefaultValue = InValue;
+
+	if (_ModifiedEvent.IsBound())
+	{
+		FControlRigModelPin Pin = _Nodes[InPin.Node].Pins[InPin.Pin];
+		_ModifiedEvent.Broadcast(this, EControlRigModelNotifType::PinChanged, &Pin);
+	}
+
+#if CONTROLRIG_UNDO
+	if (bUndo)
+	{
+		_Nodes[InPin.Node].Pins[InPin.Pin].AppendArgumentsForAction(Action.Arguments);
+		PushAction(Action);
+	}
+#endif
+	return true;
+}
+
+bool UControlRigModel::SetPinArraySize(const FControlRigModelPair& InPin, int32 InArraySize, const FString& InDefaultValue, bool bUndo)
+{
+	ensure(InPin.Node >= 0 && InPin.Node < _Nodes.Num());
+	ensure(InPin.Pin >= 0 && InPin.Pin < _Nodes[InPin.Node].Pins.Num());
+
+	const FControlRigModelPin* Pin = &_Nodes[InPin.Node].Pins[InPin.Pin];
+	if (!Pin->IsArray())
+	{
+		if (_ModifiedEvent.IsBound())
+		{
+			FControlRigModelError Error;
+			FString PinPath = _Nodes[InPin.Node].GetPinPath(InPin.Pin, true);
+			Error.Message = FString::Printf(TEXT("Pin '%s' is not an array pin."), *PinPath);
+			_ModifiedEvent.Broadcast(this, EControlRigModelNotifType::ModelError, &Error);
+		}
+		return false;
+	}
+
+
+	if (Pin->ArraySize() == InArraySize || InArraySize < 0)
+	{
+		return false;
+	}
+
+	TArray<int32> AddedPins;
+	TArray<FControlRigModelPin> RemovedPins;
+
+#if CONTROLRIG_UNDO
+	FAction MainAction;
+	FAction ResizeAction;
+	if (bUndo)
+	{
+		MainAction.Title = TEXT("Resized Array Pin.");
+		MainAction.Type = EControlRigModelNotifType::Invalid;
+		CurrentActions.Add(&MainAction);
+
+		if (Pin->ArraySize() < InArraySize)
+		{
+			ResizeAction.Type = EControlRigModelNotifType::PinAdded;
+		}
+		else
+		{
+			ResizeAction.Type = EControlRigModelNotifType::PinRemoved;
+		}
+		ResizeAction.Title = MainAction.Title;
+		ResizeAction.Arguments.Add(GetPinPath(Pin->GetPair()));
+		ResizeAction.Arguments.Add(InDefaultValue);
+		ResizeAction.Arguments.Add(FString::FormatAsNumber(Pin->ArraySize()));
+		ResizeAction.Arguments.Add(FString::FormatAsNumber(InArraySize));
+
+
+		if (ResizeAction.Type == EControlRigModelNotifType::PinAdded)
+		{
+			PushAction(ResizeAction);
+		}
+		else
+		{
+			const FControlRigModelNode& Node = _Nodes[Pin->Node];
+			TArray<int32> PinsToVisit;
+			PinsToVisit.Add(Pin->Index);
+			for (int32 PinToVisitIndex = 0; PinToVisitIndex < PinsToVisit.Num(); PinToVisitIndex++)
+			{
+				const FControlRigModelPin& PinToVisit = Node.Pins[PinsToVisit[PinToVisitIndex]];
+				for (int32 SubPinIndex : PinToVisit.SubPins)
+				{
+					const FControlRigModelPin& SubPin = Node.Pins[SubPinIndex];
+
+					FAction SetPinDefaultAction;
+					SetPinDefaultAction.Title = TEXT("Set pin default");
+					SetPinDefaultAction.Type = EControlRigModelNotifType::PinChanged;
+					SetPinDefaultAction.Arguments.Add(GetPinPath(SubPin.GetPair()));
+					SubPin.AppendArgumentsForAction(SetPinDefaultAction.Arguments);
+					SubPin.AppendArgumentsForAction(SetPinDefaultAction.Arguments);
+					PushAction(SetPinDefaultAction);
+
+					PinsToVisit.Add(SubPin.Index);
+				}
+			}
+		}
+	}
+#endif
+
+	int32 PreviousArraySize = Pin->ArraySize();
+	int32 PinIndex = Pin->Index;
+	FControlRigModelNode& Node = _Nodes[Pin->Node];
+	int32 IndexShift = 0;
+	int32 PinIndexAfterArray = PinIndex;
+	while (Node.Pins[PinIndexAfterArray].SubPins.Num() > 0)
+	{
+		PinIndexAfterArray = Node.Pins[PinIndexAfterArray].SubPins.Last();
+	}
+
+	while (Node.Pins[PinIndex].ArraySize() > InArraySize)
+	{
+		int32 PinIndexToRemove = Node.Pins[PinIndex].SubPins.Last();
+		FControlRigModelPin PinToRemove = Node.Pins[PinIndexToRemove];
+		RemovedPins.Add(PinToRemove);
+
+		int32 NummberOfRemovedPins = RemovePinsRecursive(Node, PinIndexToRemove, bUndo);
+
+		IndexShift -= NummberOfRemovedPins;
+		PinIndexAfterArray -= NummberOfRemovedPins;
+	}
+
+	while (Node.Pins[PinIndex].ArraySize() < InArraySize)
+	{
+		FControlRigModelPin PinToAdd = Node.Pins[PinIndex];
+		PinToAdd.Name = *FString::FormatAsNumber(Node.Pins[PinIndex].ArraySize());
+		PinToAdd.Index = ++PinIndexAfterArray;
+		PinToAdd.ParentIndex = PinIndex;
+		PinToAdd.Type.ContainerType = EPinContainerType::None;
+		PinToAdd.DefaultValue =  InDefaultValue;
+
+		Node.Pins[PinIndex].SubPins.Add(PinToAdd.Index);
+		if (PinToAdd.Index == Node.Pins.Num())
+		{
+			Node.Pins.Add(PinToAdd);
+		}
+		else
+		{
+			Node.Pins.Insert(PinToAdd, PinToAdd.Index);
+		}
+		IndexShift++;
+
+		if (UStruct* Struct = Cast<UStruct>(PinToAdd.Type.PinSubCategoryObject))
+		{
+			IndexShift += AddPinsRecursive(Node, PinToAdd.Index, Struct, PinToAdd.Direction, PinIndexAfterArray);
+		}
+
+		AddedPins.Add(PinToAdd.Index);
+	}
+
+	TMap<int32, int32> RemappedIndices;
+	for (int32 OtherPinIndex = PinIndexAfterArray + 1; OtherPinIndex < Node.Pins.Num(); OtherPinIndex++)
+	{
+		FControlRigModelPin& OtherPin = Node.Pins[OtherPinIndex];
+		RemappedIndices.Add(OtherPin.Index, OtherPinIndex);
+
+		// shift parent index
+		if (OtherPin.ParentIndex != INDEX_NONE)
+		{
+			int32* MappedParentIndex = RemappedIndices.Find(OtherPin.ParentIndex);
+			if (MappedParentIndex)
+			{
+				OtherPin.ParentIndex = *MappedParentIndex;
+			}
+		}
+
+		// shift all of the links in terms of indices
+		for (int32 LinkIndex = 0; LinkIndex < OtherPin.Links.Num(); LinkIndex++)
+		{
+			FControlRigModelLink& Link = _Links[LinkIndex];
+			if (Link.Source.Node == Node.Index)
+			{
+				int32* MappedLinkIndex = RemappedIndices.Find(Link.Source.Pin);
+				if (MappedLinkIndex)
+				{
+					Link.Source.Pin = *MappedLinkIndex;
+				}
+			}
+			else if (Link.Target.Node == Node.Index)
+			{
+				int32* MappedLinkIndex = RemappedIndices.Find(Link.Target.Pin);
+				if (MappedLinkIndex)
+				{
+					Link.Target.Pin = *MappedLinkIndex;
+				}
+			}
+		}
+	}
+
+	ConfigurePinIndices(Node);
+
+	if (_ModifiedEvent.IsBound())
+	{
+		for (int32 AddedPinIndex : AddedPins)
+		{
+			FControlRigModelPin AddedPin = Node.Pins[AddedPinIndex];
+			_ModifiedEvent.Broadcast(this, EControlRigModelNotifType::PinAdded, &AddedPin);
+		}
+
+		for(FControlRigModelPin PinToRemove : RemovedPins)
+		{
+			_ModifiedEvent.Broadcast(this, EControlRigModelNotifType::PinRemoved, &PinToRemove);
+		}
+	}
+
+#if CONTROLRIG_UNDO
+	if (bUndo)
+	{
+		if (ResizeAction.Type == EControlRigModelNotifType::PinRemoved)
+		{
+			PushAction(ResizeAction);
+		}
+
+		CurrentActions.Pop();
+		PushAction(MainAction);
+	}
+#endif
+
+	return true;
+}
+
+bool UControlRigModel::ExpandPin(const FName& InNodeName, const FName& InPinName, bool bIsInput, bool bInExpanded, bool bUndo)
+{
+	const FControlRigModelPin* Pin = FindPin(InNodeName, InPinName, bIsInput);
+	if (Pin != nullptr)
+	{
+		if (Pin->SubPins.Num() == 0)
+		{
+			return false;
+		}
+		if (_Nodes[Pin->Node].Pins[Pin->Index].bExpanded != bInExpanded)
+		{
+#if CONTROLRIG_UNDO
+			FAction Action;
+			if (bUndo)
+			{
+				Action.Type = EControlRigModelNotifType::PinChanged;
+				if (bInExpanded)
+				{
+					Action.Title = TEXT("Expanded Pin.");
+				}
+				else
+				{
+					Action.Title = TEXT("Collapsed Pin.");
+				}
+				Action.Arguments.Add(GetPinPath(Pin->GetPair()));
+				Pin->AppendArgumentsForAction(Action.Arguments);
+			}
+#endif
+
+			_Nodes[Pin->Node].Pins[Pin->Index].bExpanded = bInExpanded;
+			if (_ModifiedEvent.IsBound())
+			{
+				_ModifiedEvent.Broadcast(this, EControlRigModelNotifType::PinChanged, Pin);
+			}
+
+#if CONTROLRIG_UNDO
+			if (bUndo)
+			{
+				Pin->AppendArgumentsForAction(Action.Arguments);
+				PushAction(Action);
+			}
+#endif
+			return true;
+		}
+	}
+	return false;
+}
+
+bool UControlRigModel::ResendAllNotifications()
+{
+	if (_ModifiedEvent.IsBound())
+	{
+		_ModifiedEvent.Broadcast(this, EControlRigModelNotifType::ModelCleared, nullptr);
+
+		for (FControlRigModelNode& Node : _Nodes)
+		{
+			_ModifiedEvent.Broadcast(this, EControlRigModelNotifType::NodeAdded, &Node);
+			for (FControlRigModelPin& Pin : Node.Pins)
+			{
+				// todo: we might need to send infos around array sizes
+			}
+		}
+
+		for (FControlRigModelLink& Link : _Links)
+		{
+			_ModifiedEvent.Broadcast(this, EControlRigModelNotifType::LinkAdded, &Link);
+		}
+
+		for (const FName& Name : _SelectedNodes)
+		{
+			const FControlRigModelNode* Node = FindNode(Name);
+			check(Node);
+			_ModifiedEvent.Broadcast(this, EControlRigModelNotifType::NodeSelected, Node);
+		}
+	}
+	return ResendAllPinDefaultNotifications();
+}
+
+bool UControlRigModel::ResendAllPinDefaultNotifications()
+{
+	if (_ModifiedEvent.IsBound())
+	{
+		for (FControlRigModelNode& Node : _Nodes)
+		{
+			for (FControlRigModelPin& Pin : Node.Pins)
+			{
+				_ModifiedEvent.Broadcast(this, EControlRigModelNotifType::PinChanged, &Pin);
+			}
+		}
+	}
+	return true;
+}
+
+bool UControlRigModel::ShouldStructBeUnfolded(const UStruct* Struct)
+{
+	if (Struct == nullptr)
+	{
+		return false;
+	}
+	if (Struct->IsChildOf(UClass::StaticClass()))
+	{
+		return false;
+	}
+	if (Struct == TBaseStructure<FQuat>::Get())
+	{
+		return false;
+	}
+	if (Struct == FControlRigExecuteContext::StaticStruct())
+	{
+		return false;
+	}
+	if (Struct == FRuntimeFloatCurve::StaticStruct())
+	{
+		return false;
+	}
+	if (Struct == UMaterialInterface::StaticClass())
+	{
+		return false;
+	}
+
+	return true;
+}
+
+FEdGraphPinType UControlRigModel::GetPinTypeFromField(UProperty* Property)
+{
+	FEdGraphPinType PinType;
+	GetDefault<UEdGraphSchema_K2>()->ConvertPropertyToPinType(Property, PinType);
+	return PinType;
+}
+
+void UControlRigModel::AddNodePinsForFunction(FControlRigModelNode& Node)
+{
+	Node.Pins.Reset();
+	int32 LastAddedIndex = -1;
+	const UStruct* Struct = Node.UnitStruct();
+	if (Struct)
+	{
+		for (TFieldIterator<UProperty> It(Struct); It; ++It)
+		{
+			FControlRigModelPin Pin;
+			ConfigurePinFromField(Pin, *It);
+
+			if (It->HasMetaData(UControlRig::InputMetaName))
+			{
+				Pin.Index = ++LastAddedIndex;
+				Pin.Direction = EGPD_Input;
+				Node.Pins.Add(Pin);
+				if (UStructProperty* StructProp = Cast<UStructProperty>(*It))
+				{
+					AddPinsRecursive(Node, Pin.Index, StructProp->Struct, Pin.Direction, LastAddedIndex);
+				}
+			}
+			if (It->HasMetaData(UControlRig::OutputMetaName))
+			{
+				Pin.Index = ++LastAddedIndex;
+				Pin.Direction = EGPD_Output;
+				Node.Pins.Add(Pin);
+				if (UStructProperty* StructProp = Cast<UStructProperty>(*It))
+				{
+					AddPinsRecursive(Node, Pin.Index, StructProp->Struct, Pin.Direction, LastAddedIndex);
+				}
+			}
+		}
+	}
+}
+
+void UControlRigModel::SetNodePinDefaultsForFunction(FControlRigModelNode& Node)
+{
+	UScriptStruct* ScriptStruct = (UScriptStruct*)Cast<UScriptStruct>(Node.UnitStruct());
+	if (ScriptStruct)
+	{
+		TArray<uint8> TempBuffer;
+		TempBuffer.AddUninitialized(ScriptStruct->GetStructureSize());
+		ScriptStruct->InitializeDefaultValue(TempBuffer.GetData());
+
+		for (FControlRigModelPin& Pin : Node.Pins)
+		{
+			if (Pin.Direction != EGPD_Input)
+			{
+				continue;
+			}
+			FString DefaultValueString;
+			FControlRigModelPin ParentPin = Pin;
+			FString PinPath = ParentPin.Name.ToString();
+			while (ParentPin.ParentIndex != INDEX_NONE)
+			{
+				ParentPin = Node.Pins[ParentPin.ParentIndex];
+				PinPath = ParentPin.Name.ToString() + TEXT(".") + PinPath;
+			}
+			FCachedPropertyPath PropertyPath(PinPath);
+			if (PropertyPathHelpers::GetPropertyValueAsString(TempBuffer.GetData(), ScriptStruct, PropertyPath, DefaultValueString))
+			{
+				Pin.DefaultValue = DefaultValueString;
+			}
+		}
+	}
+}
+
+void UControlRigModel::AddNodePinsForParameter(FControlRigModelNode& Node, const FEdGraphPinType& InDataType)
+{
+	int32 LastAddedIndex = Node.Pins.Num() - 1;
+
+	FControlRigModelPin InputPin;
+	InputPin.Name = ValueName;
+	InputPin.Type = InDataType;
+	InputPin.Direction = EGPD_Input;
+
+	FControlRigModelPin OutputPin;
+	OutputPin.Name = ValueName;
+	OutputPin.Type = InDataType;
+	OutputPin.Direction = EGPD_Output;
+
+	InputPin.Index = ++LastAddedIndex;
+	Node.Pins.Add(InputPin);
+
+	UScriptStruct* Struct = Cast<UScriptStruct>(InDataType.PinSubCategoryObject);
+	if (Struct != nullptr)
+	{
+		AddPinsRecursive(Node, InputPin.Index, Struct, InputPin.Direction, LastAddedIndex);
+	}
+
+	OutputPin.Index = ++LastAddedIndex;
+	Node.Pins.Add(OutputPin);
+
+	if (Struct != nullptr)
+	{
+		AddPinsRecursive(Node, OutputPin.Index, Struct, OutputPin.Direction, LastAddedIndex);
+	}
+}
+
+void UControlRigModel::SetNodePinDefaultsForParameter(FControlRigModelNode& Node, const FEdGraphPinType& InDataType)
+{
+	UScriptStruct* ScriptStruct = (UScriptStruct*)Cast<UScriptStruct>(InDataType.PinSubCategoryObject);
+	if (ScriptStruct)
+	{
+		TArray<uint8> TempBuffer;
+		TempBuffer.AddUninitialized(ScriptStruct->GetStructureSize());
+		ScriptStruct->InitializeDefaultValue(TempBuffer.GetData());
+
+		for (FControlRigModelPin& Pin : Node.Pins)
+		{
+			if (Pin.Direction != EGPD_Input)
+			{
+				continue;
+			}
+
+			FString DefaultValueString;
+			FControlRigModelPin ParentPin = Pin;
+			FString PinPath = ParentPin.Name.ToString();
+			while (ParentPin.ParentIndex != INDEX_NONE)
+			{
+				ParentPin = Node.Pins[ParentPin.ParentIndex];
+				PinPath = ParentPin.Name.ToString() + TEXT(".") + PinPath;
+			}
+			FString Left, Right;
+			PinPath.Split(TEXT("."), &Left, &Right);
+			if (Right.IsEmpty())
+			{
+				ScriptStruct->ExportText(DefaultValueString, TempBuffer.GetData(), nullptr, nullptr, PPF_None, nullptr);
+				if (!DefaultValueString.IsEmpty())
+				{
+					Pin.DefaultValue = DefaultValueString;
+				}
+			}
+			else
+			{
+				FCachedPropertyPath PropertyPath(Right);
+				if (PropertyPathHelpers::GetPropertyValueAsString(TempBuffer.GetData(), ScriptStruct, PropertyPath, DefaultValueString))
+				{
+					Pin.DefaultValue = DefaultValueString;
+				}
+			}
+		}
+	}
+}
+
+void UControlRigModel::ConfigurePinFromField(FControlRigModelPin& Pin, UProperty* Property)
+{
+	Pin.Type = GetPinTypeFromField(Property);
+	Pin.Name = Property->GetFName();
+
+	Pin.bIsConstant = Property->HasMetaData(UControlRig::ConstantMetaName);
+
+	if (Property->HasMetaData(UControlRig::BoneNameMetaName))
+	{
+		Pin.CustomWidgetName = UControlRig::BoneNameMetaName;
+	}
+
+	Pin.TooltipText = Property->GetToolTipText();
+}
+
+ int32 UControlRigModel::AddPinsRecursive(FControlRigModelNode& Node, int32 ParentIndex, const UStruct* Struct, EEdGraphPinDirection PinDirection, int32& LastAddedIndex)
+{
+	if (!ShouldStructBeUnfolded(Struct))
+	{
+		return 0;
+	}
+
+	if (!Node.Pins[ParentIndex].IsSingleValue())
+	{
+		return 0;
+	}
+
+	int32 NumberOfPinsAdded = 0;
+	for (TFieldIterator<UProperty> It(Struct); It; ++It)
+	{
+		FControlRigModelPin Pin;
+		Pin.Index = ++LastAddedIndex;
+		Pin.ParentIndex = ParentIndex;
+		Pin.Direction = PinDirection;
+		ConfigurePinFromField(Pin, *It);
+
+		if (Pin.Index == Node.Pins.Num())
+		{
+			Node.Pins.Add(Pin);
+		}
+		else
+		{
+			Node.Pins.Insert(Pin, Pin.Index);
+		}
+
+		NumberOfPinsAdded++;
+
+		if (UStructProperty* StructProp = Cast<UStructProperty>(*It))
+		{
+			if (!ShouldStructBeUnfolded(StructProp->Struct))
+			{
+				continue;
+			}
+			NumberOfPinsAdded += AddPinsRecursive(Node, Pin.Index, StructProp->Struct, Pin.Direction, LastAddedIndex);
+		}
+	}
+
+	return NumberOfPinsAdded;
+}
+
+int32 UControlRigModel::RemovePinsRecursive(FControlRigModelNode& Node, int32 PinIndex, bool bUndo)
+{
+	int32 NumberOfPinsRemoved = 0;
+	while (Node.Pins[PinIndex].SubPins.Num() > 0)
+	{
+		NumberOfPinsRemoved += RemovePinsRecursive(Node, Node.Pins[PinIndex].SubPins.Last(), bUndo);
+	}
+
+	if (Node.Pins[PinIndex].ParentIndex != INDEX_NONE)
+	{
+		Node.Pins[Node.Pins[PinIndex].ParentIndex].SubPins.Remove(PinIndex);
+	}
+
+	BreakLinks(Node.Index, PinIndex, bUndo);
+
+	FControlRigModelPin PinToRemove = Node.Pins[PinIndex];
+	Node.Pins.RemoveAt(PinIndex);
+
+	NumberOfPinsRemoved++;
+	return NumberOfPinsRemoved;
+}
+
+void UControlRigModel::ConfigurePinIndices(FControlRigModelNode& Node)
+{
+	for (int32 PinIndex = 0; PinIndex < Node.Pins.Num(); PinIndex++)
+	{
+		Node.Pins[PinIndex].SubPins.Reset();
+	}
+	for (int32 PinIndex = 0; PinIndex < Node.Pins.Num(); PinIndex++)
+	{
+		Node.Pins[PinIndex].Index = PinIndex;
+		Node.Pins[PinIndex].Node = Node.Index;
+
+		if (Node.Pins[PinIndex].ParentIndex != INDEX_NONE)
+		{
+			Node.Pins[Node.Pins[PinIndex].ParentIndex].SubPins.Add(PinIndex);
+		}
+	}
+}
+
+void UControlRigModel::GetParameterPinTypes(TArray<FEdGraphPinType>& PinTypes)
+{
+	PinTypes.Add(FEdGraphPinType(UEdGraphSchema_K2::PC_Boolean, FName(NAME_None), nullptr, EPinContainerType::None, false, FEdGraphTerminalType()));
+	PinTypes.Add(FEdGraphPinType(UEdGraphSchema_K2::PC_Float, FName(NAME_None), nullptr, EPinContainerType::None, false, FEdGraphTerminalType()));
+	PinTypes.Add(FEdGraphPinType(UEdGraphSchema_K2::PC_Int, FName(NAME_None), nullptr, EPinContainerType::None, false, FEdGraphTerminalType()));
+	PinTypes.Add(FEdGraphPinType(UEdGraphSchema_K2::PC_Int64, FName(NAME_None), nullptr, EPinContainerType::None, false, FEdGraphTerminalType()));
+	PinTypes.Add(FEdGraphPinType(UEdGraphSchema_K2::PC_Byte, FName(NAME_None), nullptr, EPinContainerType::None, false, FEdGraphTerminalType()));
+	PinTypes.Add(FEdGraphPinType(UEdGraphSchema_K2::PC_Name, FName(NAME_None), nullptr, EPinContainerType::None, false, FEdGraphTerminalType()));
+	PinTypes.Add(FEdGraphPinType(UEdGraphSchema_K2::PC_Struct, FName(NAME_None), TBaseStructure<FVector>::Get(), EPinContainerType::None, false, FEdGraphTerminalType()));
+	PinTypes.Add(FEdGraphPinType(UEdGraphSchema_K2::PC_Struct, FName(NAME_None), TBaseStructure<FVector2D>::Get(), EPinContainerType::None, false, FEdGraphTerminalType()));
+	PinTypes.Add(FEdGraphPinType(UEdGraphSchema_K2::PC_Struct, FName(NAME_None), TBaseStructure<FRotator>::Get(), EPinContainerType::None, false, FEdGraphTerminalType()));
+	PinTypes.Add(FEdGraphPinType(UEdGraphSchema_K2::PC_Struct, FName(NAME_None), TBaseStructure<FTransform>::Get(), EPinContainerType::None, false, FEdGraphTerminalType()));
+	PinTypes.Add(FEdGraphPinType(UEdGraphSchema_K2::PC_Struct, FName(NAME_None), TBaseStructure<FEulerTransform>::Get(), EPinContainerType::None, false, FEdGraphTerminalType()));
+	PinTypes.Add(FEdGraphPinType(UEdGraphSchema_K2::PC_Struct, FName(NAME_None), TBaseStructure<FLinearColor>::Get(), EPinContainerType::None, false, FEdGraphTerminalType()));
+}
+
+#if CONTROLRIG_UNDO
+
+void UControlRigModel::PostTransacted(const FTransactionObjectEvent& TransactionEvent)
+{
+	if (TransactionEvent.GetEventType() == ETransactionObjectEventType::UndoRedo)
+	{
+		while (ActionCount < UndoActions.Num())
+		{
+			if (UndoActions.Num() == 0)
+			{
+				break;
+			}
+			if (!Undo())
+			{
+				return;
+			}
+		}
+		while (ActionCount > UndoActions.Num())
+		{
+			if (RedoActions.Num() == 0)
+			{
+				break;
+			}
+			if (!Redo())
+			{
+				return;
+			}
+		}
+	}
+}
+
+bool UControlRigModel::Undo()
+{
+	if (UndoActions.Num() == 0)
+	{
+		return false;
+	}
+
+	FAction Action = UndoActions.Pop();
+	if (UndoAction(Action))
+	{
+		RedoActions.Add(Action);
+	}
+	return true;
+}
+
+bool UControlRigModel::Redo()
+{
+	if (RedoActions.Num() == 0)
+	{
+		return false;
+	}
+
+	FAction Action = RedoActions.Pop();
+	if (RedoAction(Action))
+	{
+		UndoActions.Add(Action);
+	}
+	return true;
+}
+
+void UControlRigModel::PushAction(const UControlRigModel::FAction& InAction)
+{
+	ensure(InAction.IsValid());
+
+	if (InAction.Type == EControlRigModelNotifType::Invalid && InAction.SubActions.Num() == 0)
+	{
+		return;
+	}
+
+	if (CurrentActions.Num() > 0)
+	{
+		CurrentActions.Last()->SubActions.Add(InAction);
+	}
+	else
+	{
+		UndoActions.Add(InAction);
+		RedoActions.Reset();
+
+		// todo: where does this number come from?
+		if (UndoActions.Num() > 100)
+		{
+			FAction& ActionToErase = UndoActions[UndoActions.Num() - 100];
+			ActionToErase.Type = EControlRigModelNotifType::Invalid;
+			ActionToErase.Arguments.Empty();
+			ActionToErase.SubActions.Empty();
+		}
+
+		const FScopedTransaction Transaction(FText::FromString(InAction.Title));
+		SetFlags(RF_Transactional);
+		Modify();
+		ActionCount = ActionCount + 1;
+	}
+}
+
+bool UControlRigModel::UndoAction(const UControlRigModel::FAction& InAction)
+{
+	for (int32 SubActionIndex = InAction.SubActions.Num() - 1; SubActionIndex >= 0; SubActionIndex--)
+	{
+		if (!UndoAction(InAction.SubActions[SubActionIndex]))
+		{
+			return false;
+		}
+	}
+
+	switch (InAction.Type)
+	{
+		case EControlRigModelNotifType::NodeAdded:
+		{
+			FControlRigModelNode Node;
+			Node.ConfigureFromActionArguments(InAction.Arguments);
+			if (!RemoveNode(Node.Name, false /* undo */))
+			{
+				return false;
+			}
+			break;
+		}
+		case EControlRigModelNotifType::NodeRemoved:
+		{
+			UControlRigModel::FAction Action;
+			Action.Type = EControlRigModelNotifType::NodeAdded;
+			Action.Arguments = InAction.Arguments;
+			if (!RedoAction(Action))
+			{
+				return false;
+			}
+			break;
+		}
+		case EControlRigModelNotifType::NodeChanged:
+		{
+			FControlRigModelNode Node;
+			Node.ConfigureFromActionArguments(InAction.Arguments, 0 /* offset */);
+			if (Node.IsParameter())
+			{
+				SetParameterType(Node.Name, Node.ParameterType, false /* undo */);
+			}
+			if (!SetNodePosition(Node.Name, Node.Position, false /* undo */))
+			{
+				return false;
+			}
+			_LastNodePositions.Add(Node.Name, Node.Position);
+			break;
+		}
+		case EControlRigModelNotifType::NodeRenamed:
+		{
+			return RenameNode(*InAction.Arguments[1], *InAction.Arguments[0], false /* undo */);
+		}
+		case EControlRigModelNotifType::LinkAdded:
+		{
+			FControlRigModelLink Link;
+			Link.ConfigureFromActionArguments(InAction.Arguments, 0, this);
+			return BreakLink(Link.Source.Node, Link.Source.Pin, Link.Target.Node, Link.Target.Pin, false /* undo */);
+		}
+		case EControlRigModelNotifType::LinkRemoved:
+		{
+			FControlRigModelLink Link;
+			Link.ConfigureFromActionArguments(InAction.Arguments, 0, this);
+			if (!MakeLink(Link.Source.Node, Link.Source.Pin, Link.Target.Node, Link.Target.Pin, false /* undo */))
+			{
+				return false;
+			}
+			break;
+		}
+		case EControlRigModelNotifType::PinAdded:
+		case EControlRigModelNotifType::PinRemoved:
+		{
+			const FString& ArrayPinPath = InAction.Arguments[0];
+			const FString& DefaultValue = InAction.Arguments[1];
+			int32 OldArraySize = FCString::Atoi(*InAction.Arguments[2]);
+			int32 NewArraySize = FCString::Atoi(*InAction.Arguments[3]);
+
+			const FControlRigModelPin* Pin = FindPinFromPath(ArrayPinPath);
+			if (Pin == nullptr)
+			{
+				return false;
+			}
+			if (!SetPinArraySize(Pin->GetPair(), OldArraySize, DefaultValue, false /* undo */))
+			{
+				return false;
+			}
+			break;
+		}
+		case EControlRigModelNotifType::PinChanged:
+		{
+			const FString& PinPath = InAction.Arguments[0];
+			FString Left, Right;
+			SplitPinPath(PinPath, Left, Right);
+
+			FControlRigModelPin Pin;
+			Pin.ConfigureFromActionArguments(InAction.Arguments, 1);
+			
+			SetPinDefaultValue(*Left, *Right, Pin.DefaultValue, false /* undo */);
+			ExpandPin(*Left, *Right, Pin.Direction == EGPD_Input, Pin.bExpanded, false /* undo */);
+			break;
+		}
+		case EControlRigModelNotifType::Invalid:
+		{
+			ensure(InAction.Arguments.Num() == 0);
+			break;
+		}
+		default:
+		{
+			ensure(false);
+			break;
+		}
+	}
+
+	return true;
+}
+
+bool UControlRigModel::RedoAction(const UControlRigModel::FAction& InAction)
+{
+	ensure(InAction.IsValid());
+
+	switch (InAction.Type)
+	{
+		case EControlRigModelNotifType::NodeAdded:
+		{
+			FControlRigModelNode Node;
+			Node.ConfigureFromActionArguments(InAction.Arguments);
+
+			if (Node.IsParameter())
+			{
+				FEdGraphPinType PinType;
+				FEdGraphPinType::StaticStruct()->ImportText(*InAction.Arguments[4], &PinType, nullptr, EPropertyPortFlags::PPF_None, nullptr, FEdGraphPinType::StaticStruct()->GetFName().ToString(), true);
+				return AddParameter(Node.Name, PinType, Node.ParameterType, Node.Position, false /* undo */);
+			}
+
+			Node.FunctionName = *InAction.Arguments[1];
+			if (!AddNode(Node, false /* undo */))
+			{
+				return false;
+			}
+			break;
+		}
+		case EControlRigModelNotifType::NodeRemoved:
+		{
+			FControlRigModelNode Node;
+			Node.ConfigureFromActionArguments(InAction.Arguments);
+			if (!RemoveNode(Node.Name, false /* undo */))
+			{
+				return false;
+			}
+			break;
+		}
+		case EControlRigModelNotifType::NodeChanged:
+		{
+			FControlRigModelNode Node;
+			Node.ConfigureFromActionArguments(InAction.Arguments, 5 /* offset */);
+			if (Node.IsParameter())
+			{
+				SetParameterType(Node.Name, Node.ParameterType, false /* undo */);
+			}
+			if (!SetNodePosition(Node.Name, Node.Position, false /* undo */))
+			{
+				return false;
+			}
+			_LastNodePositions.Add(Node.Name, Node.Position);
+			break;
+		}
+		case EControlRigModelNotifType::NodeRenamed:
+		{
+			return RenameNode(*InAction.Arguments[0], *InAction.Arguments[1], false /* undo */);
+		}
+		case EControlRigModelNotifType::LinkAdded:
+		{
+			FControlRigModelLink Link;
+			Link.ConfigureFromActionArguments(InAction.Arguments, 0, this);
+			return MakeLink(Link.Source.Node, Link.Source.Pin, Link.Target.Node, Link.Target.Pin, false /* undo */);
+		}
+		case EControlRigModelNotifType::LinkRemoved:
+		{
+			FControlRigModelLink Link;
+			Link.ConfigureFromActionArguments(InAction.Arguments, 0, this);
+			return BreakLink(Link.Source.Node, Link.Source.Pin, Link.Target.Node, Link.Target.Pin, false /* undo */);
+		}
+		case EControlRigModelNotifType::PinAdded:
+		case EControlRigModelNotifType::PinRemoved:
+		{
+			const FString& ArrayPinPath = InAction.Arguments[0];
+			const FString& DefaultValue = InAction.Arguments[1];
+			int32 OldArraySize = FCString::Atoi(*InAction.Arguments[2]);
+			int32 NewArraySize = FCString::Atoi(*InAction.Arguments[3]);
+
+			const FControlRigModelPin* Pin = FindPinFromPath(ArrayPinPath);
+			if (Pin == nullptr)
+			{
+				return false;
+			}
+			if (!SetPinArraySize(Pin->GetPair(), NewArraySize, DefaultValue, false /* undo */))
+			{
+				return false;
+			}
+			break;
+		}
+		case EControlRigModelNotifType::PinChanged:
+		{
+			const FString& PinPath = InAction.Arguments[0];
+			FString Left, Right;
+			SplitPinPath(PinPath, Left, Right);
+
+			FControlRigModelPin Pin;
+			Pin.ConfigureFromActionArguments(InAction.Arguments, 5);
+
+			SetPinDefaultValue(*Left, *Right, Pin.DefaultValue, false /* undo */);
+			ExpandPin(*Left, *Right, Pin.Direction == EGPD_Input, Pin.bExpanded, false /* undo */);
+			break;
+		}
+		case EControlRigModelNotifType::Invalid:
+		{
+			ensure(InAction.Arguments.Num() == 0);
+			break;
+		}
+		default:
+		{
+			ensure(false);
+			break;
+		}
+	}
+
+	for (int32 SubActionIndex = 0; SubActionIndex < InAction.SubActions.Num(); SubActionIndex++)
+	{
+		if (!RedoAction(InAction.SubActions[SubActionIndex]))
+		{
+			return false;
+		}
+	}
+
+	return true;
+}
+
+#endif
+#undef LOCTEXT_NAMESPACE
+
diff --git a/Engine/Plugins/Experimental/ControlRig/Source/ControlRigDeveloper/Private/Graph/ControlRigGraph.cpp b/Engine/Plugins/Experimental/ControlRig/Source/ControlRigDeveloper/Private/Graph/ControlRigGraph.cpp
--- a/Engine/Plugins/Experimental/ControlRig/Source/ControlRigDeveloper/Private/Graph/ControlRigGraph.cpp
+++ b/Engine/Plugins/Experimental/ControlRig/Source/ControlRigDeveloper/Private/Graph/ControlRigGraph.cpp
@@ -5,6 +5,7 @@
 #include "ControlRigBlueprintGeneratedClass.h"
 #include "Graph/ControlRigGraphSchema.h"
 #include "ControlRig.h"
+#include "ControlRigModel.h"
 #include "ControlRigObjectVersion.h"
 #include "Units/RigUnit.h"
 #include "ControlRig/Private/Units/Execution/RigUnit_BeginExecution.h"
@@ -21,11 +22,14 @@
 
 UControlRigGraph::UControlRigGraph()
 {
+	bSuspendModelNotifications = false;
+	bIsTemporaryGraphForCopyPaste = false;
 }
 
 void UControlRigGraph::Initialize(UControlRigBlueprint* InBlueprint)
 {
-
+	InBlueprint->OnModified().RemoveAll(this);
+	InBlueprint->OnModified().AddUObject(this, &UControlRigGraph::HandleModelModified);
 }
 
 const UControlRigGraphSchema* UControlRigGraph::GetControlRigGraphSchema()
@@ -94,7 +98,8 @@ void UControlRigGraph::PostLoad()
 
 	Super::PostLoad();
 
-	if (Blueprint)
+	UControlRigBlueprint* RigBlueprint = Cast<UControlRigBlueprint>(Blueprint);
+	if (RigBlueprint)
 	{
 		if (GetLinkerCustomVersion(FControlRigObjectVersion::GUID) < FControlRigObjectVersion::RemovalOfHierarchyRefPins)
 		{
@@ -102,8 +107,10 @@ void UControlRigGraph::PostLoad()
 			{
 				Blueprint->OnCompiled().Remove(BlueprintOnCompiledHandle);
 			}
-			Blueprint->OnCompiled().AddUObject(this, &UControlRigGraph::OnBlueprintCompiledPostLoad);
+			BlueprintOnCompiledHandle = Blueprint->OnCompiled().AddUObject(this, &UControlRigGraph::OnBlueprintCompiledPostLoad);
 		}
+
+		RigBlueprint->PopulateModelFromGraph(this);
 	}
 }
 
@@ -224,6 +231,225 @@ const TArray<TSharedPtr<FString>>& UControlRigGraph::GetBoneNameList() const
 	return BoneNameList;
 }
 
+
+void UControlRigGraph::HandleModelModified(const UControlRigModel* InModel, EControlRigModelNotifType InType, const void* InPayload)
+{
+	if (bSuspendModelNotifications)
+	{
+		return;
+	}
+
+	switch (InType)
+	{
+		case EControlRigModelNotifType::ModelCleared:
+		{
+			for (const FControlRigModelNode& Node : InModel->Nodes())
+			{
+				UControlRigGraphNode* RigNode = FindNodeFromPropertyName(Node.Name);
+				if (RigNode != nullptr)
+				{
+					RemoveNode(RigNode);
+				}
+			}
+			Modify();
+			break;
+		}
+		case EControlRigModelNotifType::NodeAdded:
+		{
+			const FControlRigModelNode* Node = (const FControlRigModelNode*)InPayload;
+			if (Node != nullptr)
+			{
+				FEdGraphPinType PinType;
+				if (Node->IsParameter())
+				{
+					PinType = Node->Pins[0].Type;
+				}
+				UEdGraphNode* EdNode = FControlRigBlueprintUtils::InstantiateGraphNodeForProperty(this, Node->Name, Node->Position, PinType);
+				if (EdNode != nullptr)
+				{
+					EdNode->CreateNewGuid();
+				}
+			}
+			break;
+		}
+		case EControlRigModelNotifType::NodeRemoved:
+		{
+			const FControlRigModelNode* Node = (const FControlRigModelNode*)InPayload;
+			if (Node != nullptr)
+			{
+				UControlRigGraphNode* RigNode = FindNodeFromPropertyName(Node->Name);
+				if (RigNode != nullptr)
+				{
+					RemoveNode(RigNode);
+				}
+			}
+			break;
+		}
+		case EControlRigModelNotifType::NodeChanged:
+		{
+			const FControlRigModelNode* Node = (const FControlRigModelNode*)InPayload;
+			if (Node != nullptr)
+			{
+				UControlRigGraphNode* RigNode = FindNodeFromPropertyName(Node->Name);
+				if (RigNode != nullptr)
+				{
+					RigNode->NodePosX = (int32)Node->Position.X;
+					RigNode->NodePosY = (int32)Node->Position.Y;
+					RigNode->Modify();
+				}
+			}
+			break;
+		}
+		case EControlRigModelNotifType::NodeRenamed:
+		{
+			const FControlRigModelNodeRenameInfo* Info = (const FControlRigModelNodeRenameInfo*)InPayload;
+			if (Info != nullptr)
+			{
+				UControlRigGraphNode* RigNode = FindNodeFromPropertyName(Info->OldName);
+				if (RigNode != nullptr)
+				{
+					RigNode->SetPropertyName(Info->NewName, true);
+					RigNode->InvalidateNodeTitle();
+					RigNode->Modify();
+				}
+			}
+			break;
+		}
+		case EControlRigModelNotifType::PinAdded:
+		case EControlRigModelNotifType::PinRemoved:
+		{
+			const FControlRigModelPin* Pin = (const FControlRigModelPin*)InPayload;
+			if (Pin!= nullptr)
+			{
+				const FControlRigModelNode& Node = InModel->Nodes()[Pin->Node];
+				UControlRigGraphNode* RigNode = FindNodeFromPropertyName(Node.Name);
+				if (RigNode != nullptr)
+				{
+					RigNode->ReconstructNode();
+				}
+			}
+			break;
+		}
+		case EControlRigModelNotifType::LinkAdded:
+		case EControlRigModelNotifType::LinkRemoved:
+		{
+			bool AddLink = InType == EControlRigModelNotifType::LinkAdded;
+
+			const FControlRigModelLink* Link = (const FControlRigModelLink*)InPayload;
+			if (Link != nullptr)
+			{
+				const FControlRigModelNode& SourceNode = InModel->Nodes()[Link->Source.Node];
+				const FControlRigModelPin& SourcePin = SourceNode.Pins[Link->Source.Pin];
+				const FControlRigModelNode& TargetNode = InModel->Nodes()[Link->Target.Node];
+				const FControlRigModelPin& TargetPin = TargetNode.Pins[Link->Target.Pin];
+
+				UControlRigGraphNode* SourceRigNode = FindNodeFromPropertyName(SourceNode.Name);
+				UControlRigGraphNode* TargetRigNode = FindNodeFromPropertyName(TargetNode.Name);
+
+				if (SourceRigNode != nullptr && TargetRigNode != nullptr)
+				{
+					FString SourcePinPath = InModel->GetPinPath(Link->Source, true);
+					FString TargetPinPath = InModel->GetPinPath(Link->Target, true);
+
+					UEdGraphPin* SourceRigPin = SourceRigNode->FindPin(*SourcePinPath, EGPD_Output);
+					UEdGraphPin* TargetRigPin = TargetRigNode->FindPin(*TargetPinPath, EGPD_Input);
+
+					if (SourceRigPin != nullptr && TargetRigPin != nullptr)
+					{
+						if (AddLink)
+						{
+							SourceRigPin->MakeLinkTo(TargetRigPin);
+							SourceRigPin->Modify();
+							TargetRigPin->Modify();
+						}
+						else
+						{
+							SourceRigPin->BreakLinkTo(TargetRigPin);
+							SourceRigPin->Modify();
+							TargetRigPin->Modify();
+						}
+					}
+				}
+
+			}
+			break;
+		}
+		case EControlRigModelNotifType::PinChanged:
+		{
+			const FControlRigModelPin* Pin = (const FControlRigModelPin*)InPayload;
+			if (Pin != nullptr)
+			{
+				const FControlRigModelNode& Node = InModel->Nodes()[Pin->Node];
+				UControlRigGraphNode* EdNode = Cast<UControlRigGraphNode>(FindNodeFromPropertyName(Node.Name));
+				if (EdNode != nullptr)
+				{
+					FString PinPath = InModel->GetPinPath(Pin->GetPair());
+					UEdGraphPin* EdPin = EdNode->FindPin(*PinPath, Pin->Direction);
+					if (EdPin)
+					{
+						bool bShouldExpand = EdNode->IsPinExpanded(PinPath) != Pin->bExpanded;
+						if (bShouldExpand)
+						{
+							// check if this is an input / output pin,
+							// and in that case: don't expand (since it is handled by the input variant)
+							if (Pin->Direction == EGPD_Output)
+							{
+								if (InModel->FindPinFromPath(PinPath, true /*input*/))
+								{
+									bShouldExpand = false;
+								}
+							}
+						}
+						if (bShouldExpand)
+						{
+							EdNode->SetPinExpansion(PinPath, Pin->bExpanded);
+						}
+						if (Pin->Direction == EGPD_Input)
+						{
+							if (!Pin->DefaultValue.IsEmpty())
+							{
+								if (Pin->Type.PinCategory == UEdGraphSchema_K2::PC_Object)
+								{
+									UClass* Class = Cast<UClass>(Pin->Type.PinSubCategoryObject);
+									if (Class)
+									{
+										EdPin->DefaultObject = StaticFindObject(Class, ANY_PACKAGE, *Pin->DefaultValue);
+									}
+								}
+
+								EdPin->DefaultValue = Pin->DefaultValue;
+							}
+						}
+						EdPin->Modify();
+					}
+				}
+			}
+			break;
+		}
+		default:
+		{
+			// todo... other cases
+			break;
+		}
+	}
+}
+
+UControlRigGraphNode* UControlRigGraph::FindNodeFromPropertyName(const FName& InPropertyName)
+{
+	for (UEdGraphNode* EdNode : Nodes)
+	{
+		UControlRigGraphNode* RigNode = Cast<UControlRigGraphNode>(EdNode);
+		if (RigNode != nullptr)
+		{
+			if (RigNode->PropertyName == InPropertyName)
+			{
+				return RigNode;
+			}
+		}
+	}
+	return nullptr;
+}
+
 #endif
 
 #undef LOCTEXT_NAMESPACE
