Control Rig MVC Refactoring

diff --git a/Engine/Plugins/Experimental/ControlRig/Source/ControlRigDeveloper/Private/Graph/ControlRigGraphNode.cpp b/Engine/Plugins/Experimental/ControlRig/Source/ControlRigDeveloper/Private/Graph/ControlRigGraphNode.cpp
--- a/Engine/Plugins/Experimental/ControlRig/Source/ControlRigDeveloper/Private/Graph/ControlRigGraphNode.cpp
+++ b/Engine/Plugins/Experimental/ControlRig/Source/ControlRigDeveloper/Private/Graph/ControlRigGraphNode.cpp
@@ -75,6 +75,15 @@ FText UControlRigGraphNode::GetNodeTitle(ENodeTitleType::Type TitleType) const
 
 void UControlRigGraphNode::ReconstructNode()
 {
+	UControlRigGraph* RigGraph = Cast<UControlRigGraph>(GetGraph());
+	if (RigGraph)
+	{
+		if (RigGraph->bIsTemporaryGraphForCopyPaste)
+		{
+			return;
+		}
+	}
+
 #if WITH_EDITORONLY_DATA
 	// store the nodes connected to outputs of hierarchy refs.
 	// this is done for backwards compatibility
@@ -103,6 +112,12 @@ void UControlRigGraphNode::ReconstructNode()
 	PostReconstructNode();
 
 	GetGraph()->NotifyGraphChanged();
+
+	UScriptStruct* ScriptStruct = GetUnitScriptStruct();
+	if (ScriptStruct)
+	{
+		StructPath = ScriptStruct->GetPathName();
+	}
 }
 
 void UControlRigGraphNode::CacheHierarchyRefConnectionsOnPostLoad()
@@ -162,39 +177,6 @@ void UControlRigGraphNode::PrepareForCopying()
 	}
 }
 
-void UControlRigGraphNode::PostPasteNode()
-{
-	// this reads if it's struct, or normal property
-	UStruct* Struct = FindObject<UStruct>(ANY_PACKAGE, *(StructPath));
-
-	if (UControlRigGraph* Graph = Cast<UControlRigGraph>(GetOuter()))
-	{
-		UControlRigBlueprint* ControlRigBlueprint = Cast<UControlRigBlueprint>(Graph->GetOuter());
-		if (ControlRigBlueprint)
-		{
-			// if struct property
-			if (Struct)
-			{
-				// if it's struct, it is rig unit
-				FName NewPropName = FControlRigBlueprintUtils::AddUnitMember(ControlRigBlueprint, Struct);
-				if (ensure(NewPropName != NAME_None))
-				{
-					SetPropertyName(NewPropName, true);
-				}
-			}
-			else
-			{
-				// if not, it's normal property
-				FName NewPropName = FControlRigBlueprintUtils::AddPropertyMember(ControlRigBlueprint, PinType, PropertyName.ToString());
-				if (ensure(NewPropName != NAME_None))
-				{
-					SetPropertyName(NewPropName, true);
-				}
-			}
-		}
-	}
-}
-
 bool UControlRigGraphNode::IsDeprecated() const
 {
 	UScriptStruct* ScriptStruct = GetUnitScriptStruct();
@@ -257,11 +239,17 @@ void UControlRigGraphNode::RewireOldPinsToNewPins(TArray<UEdGraphPin*>& InOldPin
 void UControlRigGraphNode::DestroyPinList(TArray<UEdGraphPin*>& InPins)
 {
 	UBlueprint* Blueprint = GetBlueprint();
+	bool bNotify = false;
+	if (Blueprint != nullptr)
+	{
+		bNotify = !Blueprint->bIsRegeneratingOnLoad;
+	}
+
 	// Throw away the original pins
 	for (UEdGraphPin* Pin : InPins)
 	{
 		Pin->Modify();
-		Pin->BreakAllPinLinks(!Blueprint->bIsRegeneratingOnLoad);
+		Pin->BreakAllPinLinks(bNotify);
 
 		UEdGraphNode::DestroyPin(Pin);
 	}
@@ -274,9 +262,7 @@ void UControlRigGraphNode::PostReconstructNode()
 		SetupPinDefaultsFromCDO(Pin);
 	}
 
-	UScriptStruct* ScriptStruct = GetUnitScriptStruct();
-	bCanRenameNode = (ScriptStruct == nullptr) || (ScriptStruct && ScriptStruct->HasMetaData(UControlRig::DisplayNameMetaName) && ScriptStruct->HasMetaData(UControlRig::ShowVariableNameInTitleMetaName));
-
+	bCanRenameNode = false;
 	UpdateNodeColorFromMetadata();
 }
 
@@ -288,28 +274,6 @@ void UControlRigGraphNode::PostLoad()
 }
 #endif
 
-void UControlRigGraphNode::HandleVariableRenamed(UBlueprint* InBlueprint, UClass* InVariableClass, UEdGraph* InGraph, const FName& InOldVarName, const FName& InNewVarName)
-{
-	if(InBlueprint == GetBlueprint() && InGraph == GetGraph())
-	{
-		if(InOldVarName == PropertyName)
-		{
-			Modify();
-
-			PropertyName = InNewVarName;
-			InvalidateNodeTitle();
-
-			for(UEdGraphPin* Pin : Pins)
-			{
-				FString OldPinName = Pin->PinName.ToString();
-				bool bRemoved = OldPinName.RemoveFromStart(InOldVarName.ToString());
-				check(bRemoved);
-				Pin->PinName = FName(*FString(InNewVarName.ToString() + OldPinName));
-			}
-		}
-	}
-}
-
 void UControlRigGraphNode::CreateVariablePins(bool bAlwaysCreatePins)
 {
 	CacheVariableInfo();
@@ -321,54 +285,62 @@ void UControlRigGraphNode::CreateVariablePins(bool bAlwaysCreatePins)
 
 void UControlRigGraphNode::HandleClearArray(FString InPropertyPath)
 {
-	const FScopedTransaction Transaction(LOCTEXT("ClearArrayTransaction", "Clear Array"));
-
-	if(PerformArrayOperation(InPropertyPath, [](FScriptArrayHelper& InArrayHelper, int32 InArrayIndex)
-	{
-		InArrayHelper.EmptyValues();
-		return true;
-	}, true, true))
+	UControlRigBlueprint* Blueprint = Cast<UControlRigBlueprint>(GetOuter()->GetOuter());
+	if (Blueprint)
 	{
-		ReconstructNode();
+		if (Blueprint->ModelController && Blueprint->Model)
+		{
+			FString Left, Right;
+			Blueprint->Model->SplitPinPath(InPropertyPath, Left, Right);
+			Blueprint->ModelController->ClearArrayPin(*Left, *Right);
+		}
 	}
 }
 
-void UControlRigGraphNode::HandleRemoveArrayElement(FString InPropertyPath)
+void UControlRigGraphNode::HandleAddArrayElement(FString InPropertyPath)
 {
-	const FScopedTransaction Transaction(LOCTEXT("RemoveArrayElementTransaction", "Remove Array Element"));
-
-	if(PerformArrayOperation(InPropertyPath, [](FScriptArrayHelper& InArrayHelper, int32 InArrayIndex)
+	UControlRigBlueprint* Blueprint = Cast<UControlRigBlueprint>(GetOuter()->GetOuter());
+	if (Blueprint)
 	{
-		if(InArrayIndex != INDEX_NONE)
+		if (Blueprint->ModelController && Blueprint->Model)
 		{
-			InArrayHelper.RemoveValues(InArrayIndex);
-			return true;
+			FString Left, Right;
+			Blueprint->Model->SplitPinPath(InPropertyPath, Left, Right);
+			Blueprint->ModelController->AddArrayPin(*Left, *Right, FString());
 		}
-		return false;
-	}, true, true))
-	{
-		ReconstructNode();
 	}
 }
 
-void UControlRigGraphNode::HandleInsertArrayElement(FString InPropertyPath)
+void UControlRigGraphNode::HandleRemoveArrayElement(FString InPropertyPath)
 {
-	const FScopedTransaction Transaction(LOCTEXT("InsertArrayElementTransaction", "Insert Array Element"));
-
-	if(PerformArrayOperation(InPropertyPath, [](FScriptArrayHelper& InArrayHelper, int32 InArrayIndex)
+	UControlRigBlueprint* Blueprint = Cast<UControlRigBlueprint>(GetOuter()->GetOuter());
+	if (Blueprint)
 	{
-		if(InArrayIndex != INDEX_NONE)
+		if (Blueprint->ModelController && Blueprint->Model)
 		{
-			InArrayHelper.InsertValues(InArrayIndex);
-			return true;
+			const FControlRigModelPin* ChildPin = Blueprint->Model->FindPinFromPath(*InPropertyPath);
+			if (ChildPin)
+			{
+				const FControlRigModelPin* ParentPin = Blueprint->Model->GetParentPin(ChildPin->GetPair());
+				if (ParentPin)
+				{
+					FString ParentPinPath = Blueprint->Model->GetPinPath(ParentPin->GetPair());
+					FString Left, Right;
+					Blueprint->Model->SplitPinPath(ParentPinPath, Left, Right);
+					// todo: really should be remove at index
+					Blueprint->ModelController->PopArrayPin(*Left, *Right);
+				}
+			}
 		}
-		return false;
-	}, true, true))
-	{
-		ReconstructNode();
 	}
 }
 
+void UControlRigGraphNode::HandleInsertArrayElement(FString InPropertyPath)
+{
+	// todo: really should be insert
+	HandleAddArrayElement(InPropertyPath);
+}
+
 void UControlRigGraphNode::AllocateDefaultPins()
 {
 	CreateVariablePins(true);
@@ -393,15 +365,15 @@ void UControlRigGraphNode::CreateExecutionPins(bool bAlwaysCreatePins)
 	{
 		if (bAlwaysCreatePins || ExecutionInfo->InputPin == nullptr)
 		{
-			ExecutionInfo->InputPin = CreatePin(EGPD_Input, ExecutionInfo->GetPinType(), FName(*ExecutionInfo->GetPropertyPath()));
+			ExecutionInfo->InputPin = CreatePin(EGPD_Input, ExecutionInfo->GetPinType(), FName(*ExecutionInfo->GetPinPath()));
 			ExecutionInfo->InputPin->PinFriendlyName = ExecutionInfo->GetDisplayNameText();
 			ExecutionInfo->InputPin->PinType.bIsReference = IsStructReference(ExecutionInfo);
-			SetupPinAutoGeneratedDefaults(ExecutionInfo->InputPin);
+			ExecutionInfo->InputPin->DefaultValue = ExecutionInfo->GetPin()->DefaultValue;
 		}
 
 		if (bAlwaysCreatePins || ExecutionInfo->OutputPin == nullptr)
 		{
-			ExecutionInfo->OutputPin = CreatePin(EGPD_Output, ExecutionInfo->GetPinType(), FName(*ExecutionInfo->GetPropertyPath()));
+			ExecutionInfo->OutputPin = CreatePin(EGPD_Output, ExecutionInfo->GetPinType(), FName(*ExecutionInfo->GetPinPath()));
 		}
 
 		// note: no recursion for execution pins
@@ -414,12 +386,12 @@ void UControlRigGraphNode::CreateInputPins_Recursive(const TSharedPtr<FControlRi
 	{
 		if (bAlwaysCreatePins || ChildInfo->InputPin == nullptr)
 		{
-			ChildInfo->InputPin = CreatePin(EGPD_Input, ChildInfo->GetPinType(), FName(*ChildInfo->GetPropertyPath()));
+			ChildInfo->InputPin = CreatePin(EGPD_Input, ChildInfo->GetPinType(), FName(*ChildInfo->GetPinPath()));
 			ChildInfo->InputPin->PinFriendlyName = ChildInfo->GetDisplayNameText();
 			ChildInfo->InputPin->PinType.bIsReference = IsStructReference(ChildInfo);
 			ChildInfo->InputPin->ParentPin = InputInfo->InputPin;
+			ChildInfo->InputPin->DefaultValue = ChildInfo->GetPin()->DefaultValue;
 			InputInfo->InputPin->SubPins.Add(ChildInfo->InputPin);
-			SetupPinAutoGeneratedDefaults(ChildInfo->InputPin);
 		}
 	}
 
@@ -437,10 +409,10 @@ void UControlRigGraphNode::CreateInputPins(bool bAlwaysCreatePins)
 	{
 		if (bAlwaysCreatePins || InputInfo->InputPin == nullptr)
 		{
-			InputInfo->InputPin = CreatePin(EGPD_Input, InputInfo->GetPinType(), FName(*InputInfo->GetPropertyPath()));
+			InputInfo->InputPin = CreatePin(EGPD_Input, InputInfo->GetPinType(), FName(*InputInfo->GetPinPath()));
 			InputInfo->InputPin->PinFriendlyName = InputInfo->GetDisplayNameText();
 			InputInfo->InputPin->PinType.bIsReference = IsStructReference(InputInfo);
-			SetupPinAutoGeneratedDefaults(InputInfo->InputPin);
+			InputInfo->InputPin->DefaultValue = InputInfo->GetPin()->DefaultValue;
 		}
 
 		CreateInputPins_Recursive(InputInfo, bAlwaysCreatePins);
@@ -453,17 +425,17 @@ void UControlRigGraphNode::CreateInputOutputPins_Recursive(const TSharedPtr<FCon
 	{
 		if (bAlwaysCreatePins || ChildInfo->InputPin == nullptr)
 		{
-			ChildInfo->InputPin = CreatePin(EGPD_Input, ChildInfo->GetPinType(), FName(*ChildInfo->GetPropertyPath()));
+			ChildInfo->InputPin = CreatePin(EGPD_Input, ChildInfo->GetPinType(), FName(*ChildInfo->GetPinPath()));
 			ChildInfo->InputPin->PinFriendlyName = ChildInfo->GetDisplayNameText();
 			ChildInfo->InputPin->PinType.bIsReference = IsStructReference(ChildInfo);
+			ChildInfo->InputPin->DefaultValue = ChildInfo->GetPin()->DefaultValue;
 			ChildInfo->InputPin->ParentPin = InputOutputInfo->InputPin;
 			InputOutputInfo->InputPin->SubPins.Add(ChildInfo->InputPin);
-			SetupPinAutoGeneratedDefaults(ChildInfo->InputPin);
 		}
 
 		if (bAlwaysCreatePins || ChildInfo->OutputPin == nullptr)
 		{
-			ChildInfo->OutputPin = CreatePin(EGPD_Output, ChildInfo->GetPinType(), FName(*ChildInfo->GetPropertyPath()));
+			ChildInfo->OutputPin = CreatePin(EGPD_Output, ChildInfo->GetPinType(), FName(*ChildInfo->GetPinPath()));
 			ChildInfo->OutputPin->PinFriendlyName = ChildInfo->GetDisplayNameText();
 			ChildInfo->OutputPin->ParentPin = InputOutputInfo->OutputPin;
 			ChildInfo->OutputPin->PinType.bIsReference = IsStructReference(ChildInfo);
@@ -485,15 +457,15 @@ void UControlRigGraphNode::CreateInputOutputPins(bool bAlwaysCreatePins)
 	{
 		if (bAlwaysCreatePins || InputOutputInfo->InputPin == nullptr)
 		{
-			InputOutputInfo->InputPin = CreatePin(EGPD_Input, InputOutputInfo->GetPinType(), FName(*InputOutputInfo->GetPropertyPath()));
+			InputOutputInfo->InputPin = CreatePin(EGPD_Input, InputOutputInfo->GetPinType(), FName(*InputOutputInfo->GetPinPath()));
 			InputOutputInfo->InputPin->PinFriendlyName = InputOutputInfo->GetDisplayNameText();
 			InputOutputInfo->InputPin->PinType.bIsReference = IsStructReference(InputOutputInfo);
-			SetupPinAutoGeneratedDefaults(InputOutputInfo->InputPin);
+			InputOutputInfo->InputPin->DefaultValue = InputOutputInfo->GetPin()->DefaultValue;
 		}
 
 		if (bAlwaysCreatePins || InputOutputInfo->OutputPin == nullptr)
 		{
-			InputOutputInfo->OutputPin = CreatePin(EGPD_Output, InputOutputInfo->GetPinType(), FName(*InputOutputInfo->GetPropertyPath()));
+			InputOutputInfo->OutputPin = CreatePin(EGPD_Output, InputOutputInfo->GetPinType(), FName(*InputOutputInfo->GetPinPath()));
 		}
 
 		CreateInputOutputPins_Recursive(InputOutputInfo, bAlwaysCreatePins);
@@ -506,7 +478,7 @@ void UControlRigGraphNode::CreateOutputPins_Recursive(const TSharedPtr<FControlR
 	{
 		if (bAlwaysCreatePins || ChildInfo->OutputPin == nullptr)
 		{
-			ChildInfo->OutputPin = CreatePin(EGPD_Output, ChildInfo->GetPinType(), FName(*ChildInfo->GetPropertyPath()));
+			ChildInfo->OutputPin = CreatePin(EGPD_Output, ChildInfo->GetPinType(), FName(*ChildInfo->GetPinPath()));
 			ChildInfo->OutputPin->PinFriendlyName = ChildInfo->GetDisplayNameText();
 			ChildInfo->OutputPin->PinType.bIsReference = IsStructReference(ChildInfo);
 			ChildInfo->OutputPin->ParentPin = OutputInfo->OutputPin;
@@ -528,7 +500,7 @@ void UControlRigGraphNode::CreateOutputPins(bool bAlwaysCreatePins)
 	{
 		if (bAlwaysCreatePins || OutputInfo->OutputPin == nullptr)
 		{
-			OutputInfo->OutputPin = CreatePin(EGPD_Output, OutputInfo->GetPinType(), FName(*OutputInfo->GetPropertyPath()));
+			OutputInfo->OutputPin = CreatePin(EGPD_Output, OutputInfo->GetPinType(), FName(*OutputInfo->GetPinPath()));
 			OutputInfo->OutputPin->PinFriendlyName = OutputInfo->GetDisplayNameText();
 			OutputInfo->OutputPin->PinType.bIsReference = IsStructReference(OutputInfo);
 		}
@@ -599,204 +571,127 @@ FSlateIcon UControlRigGraphNode::GetIconAndTint(FLinearColor& OutColor) const
 	return Icon;
 }
 
-// void UControlRigGraphNode::HandleVariableRenamed(UBlueprint* InBlueprint, UClass* InVariableClass, UEdGraph* InGraph, const FName& InOldVarName, const FName& InNewVarName)
-// {
-// 	const UEdGraphSchema_K2* K2Schema = GetDefault<UEdGraphSchema_K2>();
-// 
-// 	Modify();
-// 
-// 	// first rename any disabled inputs/outputs
-// 	for (FName& DisabledInput : DisabledInputs)
-// 	{
-// 		if (DisabledInput == InOldVarName)
-// 		{
-// 			DisabledInput = InNewVarName;
-// 		}
-// 	}
-// 
-// 	for (FName& DisabledOutput : DisabledOutputs)
-// 	{
-// 		if (DisabledOutput == InOldVarName)
-// 		{
-// 			DisabledOutput = InNewVarName;
-// 		}
-// 	}
-// 
-// 	RenameUserDefinedPin(InOldVarName.ToString(), InNewVarName.ToString());
-// }
-
 // bool UControlRigGraphNode::ReferencesVariable(const FName& InVarName, const UStruct* InScope) const
 // {
 // 	return InVarName == PropertyName;
 // }
 
-TSharedPtr<FControlRigField> UControlRigGraphNode::CreateControlRigField(UField* Field, const FString& PropertyPath, int32 InArrayIndex) const
-{
-	if (UProperty* Property = Cast<UProperty>(Field))
-	{
-		TSharedPtr<FControlRigField> NewField = MakeShareable(new FControlRigProperty(Property, PropertyPath, InArrayIndex));
-		NewField->TooltipText = Field->GetToolTipText();
-		NewField->InputPin = FindPin(PropertyPath, EGPD_Input);
-		NewField->OutputPin = FindPin(PropertyPath, EGPD_Output);
-		return NewField;
-	}
-
-	return nullptr;
-}
-
-/** Helper function used to prevent us from creating sub-pins for certain field types we want to be 'atomic' */
-static bool CanExpandPinsForField(UField* InField)
+TSharedPtr<FControlRigField> UControlRigGraphNode::CreateControlRigField(const FControlRigModelPin* InPin, const FString& InPinPath, int32 InArrayIndex) const
 {
-	if(UStructProperty* StructProperty = Cast<UStructProperty>(InField))
-	{
-		if (StructProperty->Struct == TBaseStructure<FQuat>::Get())
-		{
-			return false;
-		}
-		if (StructProperty->Struct == FControlRigExecuteContext::StaticStruct())
-		{
-			return false;
-		}
-		if (StructProperty->Struct == FRuntimeFloatCurve::StaticStruct())
-		{
-			return false;
-		}
-	}
-
-	return true;
+	TSharedPtr<FControlRigField> NewField = MakeShareable(new FControlRigPin(InPin, InPinPath, InArrayIndex));
+	NewField->TooltipText = InPin->TooltipText;
+	NewField->InputPin = FindPin(InPinPath, EGPD_Input);
+	NewField->OutputPin = FindPin(InPinPath, EGPD_Output);
+	return NewField;
 }
 
 void UControlRigGraphNode::GetExecutionFields(TArray<TSharedRef<FControlRigField>>& OutFields) const
 {
-	GetFields([](UProperty* InProperty) 
+	GetFields([](const FControlRigModelPin* InPin, const FControlRigModelNode* InNode)
 	{
-		return InProperty->HasMetaData(UControlRig::InputMetaName) &&
-				InProperty->HasMetaData(UControlRig::OutputMetaName) && 
-				InProperty->GetCPPType().Equals(FControlRigExecuteContext::StaticStruct()->GetStructCPPName()); 
+		FString PinPath = InNode->GetPinPath(InPin->Index, false);
+		return InPin->Direction == EGPD_Output &&
+			InNode->FindPin(*PinPath, true) != nullptr &&
+			InPin->Type.PinSubCategoryObject == FControlRigExecuteContext::StaticStruct();
 	}, OutFields);
 }
 
 void UControlRigGraphNode::GetInputFields(TArray<TSharedRef<FControlRigField>>& OutFields) const
 {
-	GetFields([](UProperty* InProperty){ return InProperty->HasMetaData(UControlRig::InputMetaName) && !InProperty->HasMetaData(UControlRig::OutputMetaName); }, OutFields);
+	GetFields([](const FControlRigModelPin* InPin, const FControlRigModelNode* InNode)
+	{
+		FString PinPath = InNode->GetPinPath(InPin->Index, false);
+		return InPin->Direction == EGPD_Input && InNode->FindPin(*PinPath, false) == nullptr;
+	}, OutFields);
 }
 
 void UControlRigGraphNode::GetOutputFields(TArray<TSharedRef<FControlRigField>>& OutFields) const
 {
-	GetFields([](UProperty* InProperty){ return InProperty->HasMetaData(UControlRig::OutputMetaName) && !InProperty->HasMetaData(UControlRig::InputMetaName); }, OutFields);
+	GetFields([](const FControlRigModelPin* InPin, const FControlRigModelNode* InNode)
+	{
+		FString PinPath = InNode->GetPinPath(InPin->Index, false);
+		return InPin->Direction == EGPD_Output && InNode->FindPin(*PinPath, true) == nullptr;
+	}, OutFields);
 }
 
 void UControlRigGraphNode::GetInputOutputFields(TArray<TSharedRef<FControlRigField>>& OutFields) const
 {
-	GetFields([](UProperty* InProperty){
-		return InProperty->HasMetaData(UControlRig::InputMetaName) && 
-				InProperty->HasMetaData(UControlRig::OutputMetaName) &&
-				!InProperty->GetCPPType().Equals(FControlRigExecuteContext::StaticStruct()->GetStructCPPName());
-	}, OutFields);
-	
-	// Handle properties as in-outs
-	if (UClass* MyControlRigClass = GetControlRigSkeletonGeneratedClass())
+	GetFields([](const FControlRigModelPin* InPin, const FControlRigModelNode* InNode)
 	{
-		UScriptStruct* ScriptStruct = GetUnitScriptStruct();
-		if(ScriptStruct == nullptr)
-		{
-			FString PropertyPath = PropertyName.ToString();
-			if(UProperty* Property = MyControlRigClass->FindPropertyByName(PropertyName))
-			{
-				// We don't care here whether we are dealing with input/output fields as we want a pin to be created for both
-				TSharedPtr<FControlRigField> ControlRigField = CreateControlRigField(Property, PropertyPath);
-				if(ControlRigField.IsValid())
-				{
-					ControlRigField->DisplayNameText = LOCTEXT("Value", "Value");
-					OutFields.Add(ControlRigField.ToSharedRef());
-					GetFields_Recursive(ControlRigField.ToSharedRef(), PropertyPath);
-				}
-			}
-			else // we might be on a variable template node
-			{
-				if (PinType.PinCategory != NAME_None)
-				{
-					TSharedPtr<FControlRigField> TemplateField = MakeShareable(new FControlRigField(PinType, PropertyPath, LOCTEXT("Value", "Value"), -1));
-					if (TemplateField.IsValid())
-					{
-						OutFields.Add(TemplateField.ToSharedRef());
-					}
-				}
-			}
-		}
-	}
+		FString PinPath = InNode->GetPinPath(InPin->Index, false);
+		return InPin->Direction == EGPD_Input && 
+			InNode->FindPin(*PinPath, false) != nullptr &&
+			InPin->Type.PinSubCategoryObject != FControlRigExecuteContext::StaticStruct();
+	}, OutFields);
 }
 
-static const FString PropertyPathDelimiter(TEXT("."));
-
-void UControlRigGraphNode::GetFields(TFunction<bool(UProperty*)> InPropertyCheckFunction, TArray<TSharedRef<FControlRigField>>& OutFields) const
+void UControlRigGraphNode::GetFields(TFunction<bool(const FControlRigModelPin*, const FControlRigModelNode*)> InPinCheckFunction, TArray<TSharedRef<FControlRigField>>& OutFields) const
 {
 	OutFields.Reset();
 
-	if(UScriptStruct* ScriptStruct = GetUnitScriptStruct())
+	FControlRigModelNode Node;
+
+	if (UControlRigBlueprint* RigBlueprint = Cast<UControlRigBlueprint>(GetBlueprint()))
 	{
-		for (TFieldIterator<UProperty> PropertyIt(ScriptStruct, EFieldIteratorFlags::IncludeSuper); PropertyIt; ++PropertyIt)
+		if (RigBlueprint->Model)
 		{
-			UProperty* Property = *PropertyIt;
-			if (InPropertyCheckFunction(Property))
+			const FControlRigModelNode* FoundNode = RigBlueprint->Model->FindNode(GetPropertyName());
+			if (FoundNode != nullptr)
 			{
-				FString PropertyPath = PropertyName.ToString() + PropertyPathDelimiter + Property->GetName();
-				TSharedPtr<FControlRigField> ControlRigField = CreateControlRigField(Property, PropertyPath);
-				if(ControlRigField.IsValid())
-				{
-					OutFields.Add(ControlRigField.ToSharedRef());
-					GetFields_RecursiveHelper(Property, ControlRigField.ToSharedRef(), PropertyPath);
-				}
+				Node = *FoundNode;
 			}
 		}
 	}
-}
 
-void UControlRigGraphNode::GetFields_RecursiveHelper(UProperty* InProperty, const TSharedRef<FControlRigField>& InControlRigField, const FString& InPropertyPath) const
-{
-	if(UArrayProperty* ArrayProperty = Cast<UArrayProperty>(InProperty))
+	if (!Node.IsValid())
 	{
-		// if this is an array property, add sub-fields for each element
-		// Note we can only do this for nodes that are present in the CDO
-		int32 ElementCount = 0;
-		PerformArrayOperation(InPropertyPath, [&ElementCount](FScriptArrayHelper& InArrayHelper, int32 InArrayIndex)
+		if (IsVariable())
 		{
-			ElementCount = InArrayHelper.Num();
-			return true;
-		}, false, false);
-
-		for(int32 ElementIndex = 0; ElementIndex < ElementCount; ++ElementIndex)
+			FName DataType = PinType.PinCategory;
+			if (UStruct* Struct = Cast<UStruct>(PinType.PinSubCategoryObject))
+			{
+				DataType = Struct->GetFName();
+			}
+			UControlRigController::ConstructPreviewParameter(DataType, EControlRigModelParameterType::Input, Node);
+		}
+		else
 		{
-			FString SubPropertyPath = InPropertyPath + FString::Printf(TEXT("[%d]"), ElementIndex);
-			TSharedPtr<FControlRigField> ControlRigSubField = CreateControlRigField(ArrayProperty->Inner, SubPropertyPath, ElementIndex);
-			if(ControlRigSubField.IsValid())
+			if (UStruct* Struct = GetUnitScriptStruct())
 			{
-				InControlRigField->Children.Add(ControlRigSubField.ToSharedRef());
-				GetFields_Recursive(ControlRigSubField.ToSharedRef(), SubPropertyPath);
+				FName FunctionName = Struct->GetFName();
+				UControlRigController::ConstructPreviewNode(FunctionName, Node);
 			}
 		}
 	}
-	else
+
+	if (!Node.IsValid())
 	{
-		GetFields_Recursive(InControlRigField, InPropertyPath);
+		return;
 	}
-}
 
-void UControlRigGraphNode::GetFields_Recursive(const TSharedRef<FControlRigField>& ParentControlRigField, const FString& ParentPropertyPath) const
-{
-	if(CanExpandPinsForField(ParentControlRigField->GetField()))
+	Node.Name = GetPropertyName();
+
+	TMap<int32, TSharedRef<FControlRigField>> AllFields;
+	for(int32 PinIndex=0;PinIndex < Node.Pins.Num(); PinIndex++)
 	{
-		if(UStructProperty* StructProperty = Cast<UStructProperty>(ParentControlRigField->GetField()))
+		const FControlRigModelPin& Pin = Node.Pins[PinIndex];
+		if (InPinCheckFunction(&Pin, &Node))
 		{
-			for (TFieldIterator<UProperty> PropertyIt(StructProperty->Struct, EFieldIteratorFlags::IncludeSuper); PropertyIt; ++PropertyIt)
+			FString PinPath = Node.GetPinPath(Pin.Index, true);
+			TSharedPtr<FControlRigField> NewField = CreateControlRigField(&Pin, PinPath);
+			if (NewField.IsValid())
 			{
-				UProperty* Property = *PropertyIt;
-				FString PropertyPath = ParentPropertyPath + PropertyPathDelimiter + PropertyIt->GetName();
-				TSharedPtr<FControlRigField> ControlRigField = CreateControlRigField(*PropertyIt, PropertyPath);
-				if(ControlRigField.IsValid())
+				TSharedRef<FControlRigField> NewFieldRef = NewField.ToSharedRef();
+				AllFields.Add(Pin.Index, NewFieldRef);
+
+				TSharedRef<FControlRigField>* ParentField = AllFields.Find(Pin.ParentIndex);
+				if (ParentField)
 				{
-					ParentControlRigField->Children.Add(ControlRigField.ToSharedRef());
-					GetFields_RecursiveHelper(Property, ControlRigField.ToSharedRef(), PropertyPath);
+					(*ParentField)->Children.Add(NewFieldRef);
+				}
+				else
+				{
+					OutFields.Add(NewFieldRef);
 				}
 			}
 		}
@@ -838,6 +733,17 @@ UScriptStruct* UControlRigGraphNode::GetUnitScriptStruct() const
 		{
 			return Struct;
 		}
+
+		// if this doesn't work we can still fall back on the struct path
+		FString Prefix, StructName;
+		if (StructPath.Split(TEXT("."), &Prefix, &StructName))
+		{
+			Struct = FindObject<UScriptStruct>(ANY_PACKAGE, *StructName);
+			if (Struct && Struct->IsChildOf(FRigUnit::StaticStruct()))
+			{
+				return Struct;
+			}
+		}
 	}
 	return nullptr;
 }
@@ -900,7 +806,7 @@ void UControlRigGraphNode::DestroyNode()
 
 void UControlRigGraphNode::PinDefaultValueChanged(UEdGraphPin* Pin)
 {
-	CopyPinDefaultsToProperties(Pin, true, true);
+	CopyPinDefaultsToModel(Pin, true);
 }
 
 TSharedPtr<INameValidatorInterface> UControlRigGraphNode::MakeNameValidator() const
@@ -908,7 +814,7 @@ TSharedPtr<INameValidatorInterface> UControlRigGraphNode::MakeNameValidator() co
 	return MakeShared<FKismetNameValidator>(GetBlueprint(), PropertyName);
 }
 
-void UControlRigGraphNode::CopyPinDefaultsToProperties(UEdGraphPin* Pin, bool bCallModify, bool bPropagateToInstances)
+void UControlRigGraphNode::CopyPinDefaultsToModel(UEdGraphPin* Pin, bool bUndo)
 {
 	const UEdGraphSchema_K2* K2Schema = GetDefault<UEdGraphSchema_K2>();
 
@@ -917,40 +823,18 @@ void UControlRigGraphNode::CopyPinDefaultsToProperties(UEdGraphPin* Pin, bool bC
 		UControlRigBlueprint* ControlRigBlueprint = Cast<UControlRigBlueprint>(Graph->GetOuter());
 		if(ControlRigBlueprint)
 		{
-			// Note we need the actual generated class here
-			if (UClass* MyControlRigClass = GetControlRigGeneratedClass())
+			if (ControlRigBlueprint->Model)
 			{
-				if(UObject* DefaultObject = MyControlRigClass->GetDefaultObject(false))
+				if (Pin->Direction == EGPD_Input)
 				{
-					if(bCallModify)
-					{
-						DefaultObject->SetFlags(RF_Transactional);
-						DefaultObject->Modify();
-					}
-
-					FString DefaultValueString = Pin->GetDefaultAsString();
-					if(DefaultValueString.Len() > 0)
+					FString DefaultValue = Pin->DefaultValue;
+					FString Left, Right;
+					ControlRigBlueprint->Model->SplitPinPath(Pin->GetName(), Left, Right);
+					if (DefaultValue.IsEmpty() && Pin->DefaultObject != nullptr)
 					{
-						FCachedPropertyPath PropertyPath(Pin->PinName.ToString());
-						if(PropertyPathHelpers::SetPropertyValueFromString(DefaultObject, PropertyPath, DefaultValueString))
-						{
-							if(bCallModify)
-							{
-								FBlueprintEditorUtils::MarkBlueprintAsModified(GetBlueprint());
-							}
-						}
-
-						if(bPropagateToInstances)
-						{
-							TArray<UObject*> ArchetypeInstances;
-							DefaultObject->GetArchetypeInstances(ArchetypeInstances);
-							
-							for (UObject* ArchetypeInstance : ArchetypeInstances)
-							{
-								PropertyPathHelpers::SetPropertyValueFromString(ArchetypeInstance, PropertyPath, DefaultValueString);
-							}
-						}
+						DefaultValue = Pin->DefaultObject->GetPathName();
 					}
+					ControlRigBlueprint->ModelController->SetPinDefaultValue(*Left, *Right, DefaultValue, false, bUndo);
 				}
 			}
 		}
@@ -966,75 +850,6 @@ UControlRigBlueprint* UControlRigGraphNode::GetBlueprint() const
 	return nullptr;
 }
 
-void UControlRigGraphNode::SetupPinAutoGeneratedDefaults(UEdGraphPin* Pin)
-{
-	const UEdGraphSchema_K2* K2Schema = GetDefault<UEdGraphSchema_K2>();
-
-	if(UControlRigGraph* Graph = Cast<UControlRigGraph>(GetOuter()))
-	{
-		UControlRigBlueprint* ControlRigBlueprint = Cast<UControlRigBlueprint>(Graph->GetOuter());
-		if(ControlRigBlueprint)
-		{
-			FString DefaultValueString;
-			FCachedPropertyPath PropertyPath(Pin->PinName.ToString());
-				
-			// If GetPropertyValueAsString fails then the generated class doesn't have the property yet (likely just the skeleton class has been compiled)
-			UScriptStruct* ScriptStruct = GetUnitScriptStruct();
-			UStructProperty* UnitStructProperty = GetUnitProperty();
-			if(ScriptStruct && UnitStructProperty)
-			{
-				TArray<uint8> TempBuffer;
-				TempBuffer.AddUninitialized(UnitStructProperty->ElementSize);
-				ScriptStruct->InitializeDefaultValue(TempBuffer.GetData());
-							
-				// Trim the property path to look at the members of this struct
-				PropertyPath.RemoveFromStart();
-				PropertyPathHelpers::GetPropertyValueAsString(TempBuffer.GetData(), ScriptStruct, PropertyPath, DefaultValueString);
-								
-				K2Schema->GetPinDefaultValuesFromString(Pin->PinType, Pin->GetOwningNodeUnchecked(), DefaultValueString, Pin->AutogeneratedDefaultValue, Pin->DefaultObject, Pin->DefaultTextValue);
-				Pin->DefaultValue = Pin->AutogeneratedDefaultValue;
-			}
-			else if(UProperty* Property = GetProperty())
-			{
-				if(UStructProperty* StructProperty = Cast<UStructProperty>(Property))
-				{
-					TArray<uint8> TempBuffer;
-					TempBuffer.AddUninitialized(StructProperty->ElementSize);
-					StructProperty->Struct->InitializeDefaultValue(TempBuffer.GetData());
-								
-					// Fill in the root defaults from the struct itself
-					if(Pin->ParentPin == nullptr)
-					{
-						Property->ExportTextItem(DefaultValueString, TempBuffer.GetData(), nullptr, nullptr, PPF_None);
-					}
-					else
-					{
-						// Trim the property path to look at the members of this struct
-						PropertyPath.RemoveFromStart();
-
-						PropertyPathHelpers::GetPropertyValueAsString(TempBuffer.GetData(), StructProperty->Struct, PropertyPath, DefaultValueString);
-					}
-
-					K2Schema->GetPinDefaultValuesFromString(Pin->PinType, Pin->GetOwningNodeUnchecked(), DefaultValueString, Pin->AutogeneratedDefaultValue, Pin->DefaultObject, Pin->DefaultTextValue);
-					Pin->DefaultValue = Pin->AutogeneratedDefaultValue;
-				}
-				else
-				{
-					// Plain ol' properties are simpler to set up
-					TArray<uint8> TempBuffer;
-					TempBuffer.AddUninitialized(Property->ElementSize);
-					Property->InitializeValue(TempBuffer.GetData());
-
-					Property->ExportTextItem(DefaultValueString, TempBuffer.GetData(), nullptr, nullptr, PPF_None);
-
-					K2Schema->GetPinDefaultValuesFromString(Pin->PinType, Pin->GetOwningNodeUnchecked(), DefaultValueString, Pin->AutogeneratedDefaultValue, Pin->DefaultObject, Pin->DefaultTextValue);
-					Pin->DefaultValue = Pin->AutogeneratedDefaultValue;
-				}
-			}
-		}
-	}
-}
-
 void UControlRigGraphNode::SetupPinDefaultsFromCDO(UEdGraphPin* Pin)
 {
 	const UEdGraphSchema_K2* K2Schema = GetDefault<UEdGraphSchema_K2>();
@@ -1061,55 +876,6 @@ void UControlRigGraphNode::SetupPinDefaultsFromCDO(UEdGraphPin* Pin)
 	}
 }
 
-bool UControlRigGraphNode::PerformArrayOperation(const FString& InPropertyPath, TFunctionRef<bool(FScriptArrayHelper&,int32)> InOperation, bool bCallModify, bool bPropagateToInstances) const
-{
-	if(UStructProperty* StructProperty = GetUnitProperty())
-	{
-		if (UClass* MyControlRigClass = GetControlRigGeneratedClass())
-		{
-			if(UObject* DefaultObject = MyControlRigClass->GetDefaultObject(false))
-			{
-				if(bCallModify)
-				{
-					DefaultObject->SetFlags(RF_Transactional);
-					DefaultObject->Modify();
-				}
-
-				FCachedPropertyPath CachedPropertyPath(InPropertyPath);
-				if(PropertyPathHelpers::PerformArrayOperation(DefaultObject, CachedPropertyPath, InOperation))
-				{
-					if(bCallModify)
-					{
-						FBlueprintEditorUtils::MarkBlueprintAsModified(GetBlueprint());
-
-						if(bPropagateToInstances)
-						{
-							TArray<UObject*> ArchetypeInstances;
-							DefaultObject->GetArchetypeInstances(ArchetypeInstances);
-							
-							for (UObject* ArchetypeInstance : ArchetypeInstances)
-							{
-								PropertyPathHelpers::PerformArrayOperation(ArchetypeInstance, CachedPropertyPath, InOperation);
-							}
-						}
-					}
-					return true;
-				}
-			}
-		}
-	}
-
-	return false;
-}
-
-void UControlRigGraphNode::OnRenameNode(const FString& InNewName)
-{
-	FBlueprintEditorUtils::RenameMemberVariable(GetBlueprint(), PropertyName, *InNewName);
-	PropertyName = *InNewName;
-	InvalidateNodeTitle();
-	FBlueprintEditorUtils::ReconstructAllNodes(GetBlueprint());
-}
-
 FText UControlRigGraphNode::GetTooltipText() const
 {
 	if(GetUnitScriptStruct())
@@ -1148,8 +914,8 @@ void UControlRigGraphNode::AutowireNewNode(UEdGraphPin* FromPin)
 			continue;
 		}
 
-		FControlRigPinConnectionResponse ConnectResponse = Schema->CanCreateConnection_Extended(FromPin, Pin);
-		if(ConnectResponse.Response.Response != ECanCreateConnectionResponse::CONNECT_RESPONSE_DISALLOW)
+		FPinConnectionResponse ConnectResponse = Schema->CanCreateConnection(FromPin, Pin);
+		if(ConnectResponse.Response != ECanCreateConnectionResponse::CONNECT_RESPONSE_DISALLOW)
 		{
 			if(Schema->TryCreateConnection(FromPin, Pin))
 			{
@@ -1169,26 +935,11 @@ void UControlRigGraphNode::AutowireNewNode(UEdGraphPin* FromPin)
 	}
 }
 
-void UControlRigGraphNode::HandleAddArrayElement(FString InPropertyPath)
-{
-	const FScopedTransaction Transaction(LOCTEXT("AddArrayElementTransaction", "Add Array Element"));
-
-	if(PerformArrayOperation(InPropertyPath, [](FScriptArrayHelper& InArrayHelper, int32 InArrayIndex)
-	{
-		InArrayHelper.AddValues(1);
-		return true;
-	}, true, true))
-	{
-		ReconstructNode();
-	}
-}
-
 void ReplacePropertyName(TArray<TSharedRef<FControlRigField>>& InArray, const FString& OldPropName, const FString& NewPropName)
 {
 	for (int32 Index = 0; Index < InArray.Num(); ++Index)
 	{
-		InArray[Index]->PropertyPath = InArray[Index]->PropertyPath.Replace(*OldPropName, *NewPropName);
-
+		InArray[Index]->PinPath = InArray[Index]->PinPath.Replace(*OldPropName, *NewPropName);
 		ReplacePropertyName(InArray[Index]->Children, OldPropName, NewPropName);
 	}
 };
@@ -1276,4 +1027,9 @@ void UControlRigGraphNode::UpdateNodeColorFromMetadata()
 	}
 }
 
+bool UControlRigGraphNode::IsVariable() const
+{
+	return GetUnitScriptStruct() == nullptr;
+}
+
 #undef LOCTEXT_NAMESPACE
diff --git a/Engine/Plugins/Experimental/ControlRig/Source/ControlRigDeveloper/Private/Graph/ControlRigGraphSchema.cpp b/Engine/Plugins/Experimental/ControlRig/Source/ControlRigDeveloper/Private/Graph/ControlRigGraphSchema.cpp
--- a/Engine/Plugins/Experimental/ControlRig/Source/ControlRigDeveloper/Private/Graph/ControlRigGraphSchema.cpp
+++ b/Engine/Plugins/Experimental/ControlRig/Source/ControlRigDeveloper/Private/Graph/ControlRigGraphSchema.cpp
@@ -16,6 +16,7 @@
 #include "ControlRigBlueprintGeneratedClass.h"
 #include "Widgets/Notifications/SNotificationList.h"
 #include "Framework/Notifications/NotificationManager.h"
+#include "EulerTransform.h"
 
 #define LOCTEXT_NAMESPACE "ControlRigGraphSchema"
 
@@ -39,175 +40,32 @@ void UControlRigGraphSchema::GetContextMenuActions(const UEdGraph* CurrentGraph,
 #endif
 }
 
-bool UControlRigGraphSchema::TryCreateConnection_DetectCycle(UEdGraphPin* PinA, UEdGraphPin* PinB) const
+bool UControlRigGraphSchema::TryCreateConnection(UEdGraphPin* PinA, UEdGraphPin* PinB) const
 {
-	check(PinA);
-	check(PinB);
-	if (PinA->GetOwningNode() == PinB->GetOwningNode())
-	{
-		return true;
-	}
-
-	const UEdGraphNode* NodeA = PinA->GetOwningNode();
-	const UEdGraphNode* NodeB = PinB->GetOwningNode();
-
-	FControlRigDAG DAG;
-	const UEdGraph* Graph = NodeA->GetGraph();
-	for (const UEdGraphNode* Node : Graph->Nodes)
-	{
-		DAG.AddNode();
-	}
-	for (int32 NodeIndex = 0; NodeIndex < Graph->Nodes.Num(); NodeIndex++)
-	{
-		const UEdGraphNode* Node = Graph->Nodes[NodeIndex];
-		for (int32 PinIndex = 0; PinIndex < Node->Pins.Num(); PinIndex++)
-		{
-			const UEdGraphPin* Pin = Node->Pins[PinIndex];
-			if (Pin->Direction != EGPD_Output)
-			{
-				continue;
-			}
-			for (const UEdGraphPin* LinkedPin : Pin->LinkedTo)
-			{
-				const UEdGraphNode* LinkedNode = LinkedPin->GetOwningNode();
-				int32 LinkedNodeIndex = Graph->Nodes.IndexOfByKey(LinkedNode);
-				int32 LinkedPinIndex = LinkedNode->Pins.IndexOfByKey(LinkedPin);
-				DAG.AddLink(NodeIndex, LinkedNodeIndex, PinIndex, LinkedPinIndex);
-			}
-		}
-	}
-
-	// finally add the link we are going to make
-	int32 NodeAIndex = Graph->Nodes.IndexOfByKey(NodeA);
-	int32 NodeBIndex = Graph->Nodes.IndexOfByKey(NodeB);
-	int32 PinAIndex = NodeA->Pins.IndexOfByKey(PinA);
-	int32 PinBIndex = NodeA->Pins.IndexOfByKey(PinB);
-	if (PinA->Direction == EGPD_Output)
-	{
-		DAG.AddLink(NodeAIndex, NodeBIndex, PinAIndex, PinBIndex);
-	}
-	else
+#if WITH_EDITOR
+	if (GEditor)
 	{
-		DAG.AddLink(NodeBIndex, NodeAIndex, PinBIndex, PinAIndex);
+		GEditor->CancelTransaction(0);
 	}
-
-	return DAG.FindCycle().Num() > 0;
-}
-
-bool UControlRigGraphSchema::TryCreateConnection_Extended(UEdGraphPin* PinA, UEdGraphPin* PinB) const
-{
-	const FControlRigPinConnectionResponse Response = CanCreateConnection_Extended(PinA, PinB);
-	bool bModified = false;
-
-	struct Local
-	{
-		static void BreakParentConnections_Recursive(UEdGraphPin* InPin, const UControlRigGraphSchema* This)
-		{
-			This->ResetPinDefaultsRecursive(InPin);
-
-			if(InPin->ParentPin)
-			{
-				InPin->ParentPin->Modify();
-				InPin->ParentPin->BreakAllPinLinks();
-				InPin->GetOwningNode()->PinConnectionListChanged(InPin->ParentPin);
-				BreakParentConnections_Recursive(InPin->ParentPin, This);
-			}
-		}
-
-		static void BreakChildConnections_Recursive(UEdGraphPin* InPin, const UControlRigGraphSchema* This)
-		{
-			for(UEdGraphPin* SubPin : InPin->SubPins)
-			{
-				if(SubPin->LinkedTo.Num() > 0)
-				{
-					SubPin->Modify();
-					SubPin->BreakAllPinLinks();
-					SubPin->GetOwningNode()->PinConnectionListChanged(SubPin);
-				}
-
-				BreakChildConnections_Recursive(SubPin, This);
-			}
-		}
-	};
-
-	// build a temporary dag to disallow cycles.
-	// we do this here only once since it's a costly calculation
-	// and we don't want to do it for every possible pin.
-	if (Response.Response.Response == CONNECT_RESPONSE_MAKE)
-	{
-		if (TryCreateConnection_DetectCycle(PinA, PinB))
-		{
-#if WITH_EDITOR
-			FNotificationInfo Info(LOCTEXT("ConnectResponse_Disallowed_Cycle", "Connection not allowed to avoid cycle."));
-			Info.Image = FCoreStyle::Get().GetBrush(TEXT("MessageLog.Warning"));
-			Info.bFireAndForget = true;
-			Info.FadeOutDuration = 5.0f;
-			Info.ExpireDuration = 0.0f;
-			TSharedPtr<SNotificationItem> NotificationPtr = FSlateNotificationManager::Get().AddNotification(Info);
-			NotificationPtr->SetCompletionState(SNotificationItem::CS_Success);
 #endif
-			//PinA->GetOwningNode()->PinConnectionListChanged(PinA);
-			//PinB->GetOwningNode()->PinConnectionListChanged(PinB);
-			return false;
-		}
-	}
 
-	switch (Response.Response.Response)
+	UBlueprint* Blueprint = FBlueprintEditorUtils::FindBlueprintForNodeChecked(PinA->GetOwningNode());
+	UControlRigBlueprint* RigBlueprint = Cast<UControlRigBlueprint>(Blueprint);
+	if (RigBlueprint != nullptr)
 	{
-	case CONNECT_RESPONSE_MAKE:
-		PinA->Modify();
-		PinB->Modify();
-		PinA->MakeLinkTo(PinB);
-		bModified = true;
-		switch(Response.ExtendedResponse)
+		if (PinA->Direction == EGPD_Input)
 		{
-		case ECanCreateConnectionResponse_Extended::None:
-			break;
-		case ECanCreateConnectionResponse_Extended::BreakChildren:
-			if(PinA->Direction == EGPD_Input)
-			{
-				Local::BreakChildConnections_Recursive(PinA, this);
-			}
-			else if(PinB->Direction == EGPD_Input)
-			{
-				Local::BreakChildConnections_Recursive(PinB, this);
-			}
-			break;
-		case ECanCreateConnectionResponse_Extended::BreakParent:
-			if(PinA->Direction == EGPD_Input)
-			{
-				Local::BreakParentConnections_Recursive(PinA, this);
-			}
-			else if(PinB->Direction == EGPD_Input)
-			{
-				Local::BreakParentConnections_Recursive(PinB, this);
-			}
-			break;
+			UEdGraphPin* Temp = PinA;
+			PinA = PinB;
+			PinB = Temp;
 		}
-		PinA->GetOwningNode()->PinConnectionListChanged(PinA);
-		PinB->GetOwningNode()->PinConnectionListChanged(PinB);
-		break;
 
-	default:
-		bModified = UEdGraphSchema::TryCreateConnection(PinA, PinB);
-		break;
+		FString LeftA, LeftB, RightA, RightB;
+		RigBlueprint->Model->SplitPinPath(PinA->GetName(), LeftA, RightA);
+		RigBlueprint->Model->SplitPinPath(PinB->GetName(), LeftB, RightB);
+		return RigBlueprint->ModelController->MakeLink(*LeftA, *RightA, *LeftB, *RightB);
 	}
-
-	return bModified;
-}
-
-bool UControlRigGraphSchema::TryCreateConnection(UEdGraphPin* PinA, UEdGraphPin* PinB) const
-{
-	UBlueprint* Blueprint = FBlueprintEditorUtils::FindBlueprintForNodeChecked(PinA->GetOwningNode());
-
-	bool bModified = TryCreateConnection_Extended(PinA, PinB);
-
-	if (bModified && !PinA->IsPendingKill())
-	{
-		FBlueprintEditorUtils::MarkBlueprintAsModified(Blueprint);
-	}
-
-	return bModified;	
+	return false;
 }
 
 static bool HasParentConnection_Recursive(const UEdGraphPin* InPin)
@@ -233,121 +91,33 @@ static bool HasChildConnection_Recursive(const UEdGraphPin* InPin)
 	return false;
 }
 
-const FControlRigPinConnectionResponse UControlRigGraphSchema::CanCreateConnection_Extended(const UEdGraphPin* A, const UEdGraphPin* B) const
+const FPinConnectionResponse UControlRigGraphSchema::CanCreateConnection(const UEdGraphPin* A, const UEdGraphPin* B) const
 {
-	check(A != nullptr);
-	check(B != nullptr);
-
-	// Deal with basic connections (same pins, same node, differing types etc.)
-	if(A == B)
-	{
-		return FControlRigPinConnectionResponse(CONNECT_RESPONSE_DISALLOW, LOCTEXT("ConnectResponse_Disallowed_Self", "Cannot link a pin to itself"));
-	}
-
-	if(A->Direction == B->Direction)
-	{
-		return FControlRigPinConnectionResponse(CONNECT_RESPONSE_DISALLOW, A->Direction == EGPD_Input ? LOCTEXT("ConnectResponse_Disallowed_Direction_Input", "Cannot link input pin to input pin") : LOCTEXT("ConnectResponse_Disallowed_Direction_Output", "Cannot link output pin to output pin"));
-	}
-
-	if(A->GetOwningNode() == B->GetOwningNode())
-	{
-		return FControlRigPinConnectionResponse(CONNECT_RESPONSE_DISALLOW, LOCTEXT("ConnectResponse_Disallowed_SameNode", "Cannot link two pins on the same node"));
-	}
-
-	if(A->PinType != B->PinType)
-	{
-		return FControlRigPinConnectionResponse(CONNECT_RESPONSE_DISALLOW, LOCTEXT("ConnectResponse_Disallowed_Different_Types", "Cannot link pins of differing types"));
-	}
-
-	struct Local
+	UBlueprint* Blueprint = FBlueprintEditorUtils::FindBlueprintForNodeChecked(A->GetOwningNode());
+	UControlRigBlueprint* RigBlueprint = Cast<UControlRigBlueprint>(Blueprint);
+	if (RigBlueprint != nullptr)
 	{
-		static bool HasNotConnectableMetaData(const UEdGraphPin* PinToCheck)
+		if (A->Direction == EGPD_Input)
 		{
-			if (PinToCheck->ParentPin)
-			{
-				return HasNotConnectableMetaData(PinToCheck->ParentPin);
-			}
-
-			UControlRigGraphNode* RigNode = CastChecked<UControlRigGraphNode>(PinToCheck->GetOwningNode());
-			UScriptStruct* UnitStruct = RigNode->GetUnitScriptStruct();
-			if (UnitStruct)
-			{
-				FString PropertyName = PinToCheck->GetName();
-				int32 PeriodIndex = PropertyName.Find(TEXT("."));
-				if (PeriodIndex != INDEX_NONE)
-				{
-					PropertyName = PropertyName.Mid(PeriodIndex + 1);
-				}
-				UProperty* TargetProperty = UnitStruct->FindPropertyByName(*PropertyName);
-				if (TargetProperty)
-				{
-					return TargetProperty->HasMetaData(UControlRig::ConstantMetaName);
-				}
-			}
-
-			return false;
+			const UEdGraphPin* Temp = A;
+			A = B;
+			B = Temp;
 		}
-	};
 
-	// check if this property can be connected to based on metadata
-	if (A->Direction == EGPD_Input)
-	{
-		if (Local::HasNotConnectableMetaData(A))
-		{
-			if (!Local::HasNotConnectableMetaData(B))
-			{
-				return FControlRigPinConnectionResponse(CONNECT_RESPONSE_DISALLOW, LOCTEXT("ConnectResponse_Disallowed_Constant", "This pin is defined as constant."));
-			}
-		}
-	}
-	if (B->Direction == EGPD_Input)
-	{
-		if (Local::HasNotConnectableMetaData(B))
+		FString LeftA, LeftB, RightA, RightB;
+		RigBlueprint->Model->SplitPinPath(A->GetName(), LeftA, RightA);
+		RigBlueprint->Model->SplitPinPath(B->GetName(), LeftB, RightB);
+
+		FString FailureReason;
+		bool bResult = RigBlueprint->ModelController->CanLink(*LeftA, *RightA, *LeftB, *RightB, &FailureReason);
+		if (!bResult)
 		{
-			if (!Local::HasNotConnectableMetaData(A))
-			{
-				return FControlRigPinConnectionResponse(CONNECT_RESPONSE_DISALLOW, LOCTEXT("ConnectResponse_Disallowed_Constant", "This pin is defined as constant."));
-			}
+			return FPinConnectionResponse(CONNECT_RESPONSE_DISALLOW, FText::FromString(FailureReason));
 		}
+		return FPinConnectionResponse(CONNECT_RESPONSE_MAKE, LOCTEXT("ConnectResponse_Allowed", "Connect"));
 	}
 
-	// Deal with many-to-one and one to many connections
-	if(A->Direction == EGPD_Input && A->LinkedTo.Num() > 0)
-	{
-		return FControlRigPinConnectionResponse(CONNECT_RESPONSE_BREAK_OTHERS_A, LOCTEXT("ConnectResponse_Replace_Input", "Replace connection"));
-	}
-	else if(B->Direction == EGPD_Input && B->LinkedTo.Num() > 0)
-	{
-		return FControlRigPinConnectionResponse(CONNECT_RESPONSE_BREAK_OTHERS_B, LOCTEXT("ConnectResponse_Replace_Input", "Replace connection"));
-	}
-
-	// Deal with sub-struct pins
-
-	if(A->Direction == EGPD_Input && HasParentConnection_Recursive(A))
-	{
-		return FControlRigPinConnectionResponse(CONNECT_RESPONSE_MAKE, LOCTEXT("ConnectResponse_Replace_Parent", "Replace parent connection"), ECanCreateConnectionResponse_Extended::BreakParent);
-	}
-	else if(B->Direction == EGPD_Input && HasParentConnection_Recursive(B))
-	{
-		return FControlRigPinConnectionResponse(CONNECT_RESPONSE_MAKE, LOCTEXT("ConnectResponse_Replace_Parent", "Replace parent connection"), ECanCreateConnectionResponse_Extended::BreakParent);
-	}
-
-	if(A->Direction == EGPD_Input && HasChildConnection_Recursive(A))
-	{
-		return FControlRigPinConnectionResponse(CONNECT_RESPONSE_MAKE, LOCTEXT("ConnectResponse_Replace_Child", "Replace child connection(s)"), ECanCreateConnectionResponse_Extended::BreakChildren);
-	}
-	else if(B->Direction == EGPD_Input && HasChildConnection_Recursive(B))
-	{
-		return FControlRigPinConnectionResponse(CONNECT_RESPONSE_MAKE, LOCTEXT("ConnectResponse_Replace_Child", "Replace child connection(s)"), ECanCreateConnectionResponse_Extended::BreakChildren);
-	}
-
-	return FControlRigPinConnectionResponse(CONNECT_RESPONSE_MAKE, LOCTEXT("ConnectResponse_Allowed", "Connect"));	
-}
-
-const FPinConnectionResponse UControlRigGraphSchema::CanCreateConnection(const UEdGraphPin* A, const UEdGraphPin* B) const
-{
-	const FControlRigPinConnectionResponse Response = CanCreateConnection_Extended(A, B);
-	return Response.Response;
+	return FPinConnectionResponse(CONNECT_RESPONSE_DISALLOW, LOCTEXT("ConnectResponse_Disallowed_Unexpected", "Unexpected error"));
 }
 
 FLinearColor UControlRigGraphSchema::GetPinTypeColor(const FEdGraphPinType& PinType) const
@@ -365,51 +135,55 @@ FLinearColor UControlRigGraphSchema::GetPinTypeColor(const FEdGraphPinType& PinT
 
 void UControlRigGraphSchema::BreakPinLinks(UEdGraphPin& TargetPin, bool bSendsNodeNotifcation) const
 {
-	const FScopedTransaction Transaction( LOCTEXT("GraphEd_BreakPinLinks", "Break Pin Links") );
+	//const FScopedTransaction Transaction( LOCTEXT("GraphEd_BreakPinLinks", "Break Pin Links") );
 
 	// cache this here, as BreakPinLinks can trigger a node reconstruction invalidating the TargetPin referenceS
-	UBlueprint* const Blueprint = FBlueprintEditorUtils::FindBlueprintForNodeChecked(TargetPin.GetOwningNode());
-
-	TArray<UEdGraphPin*> LinkedTo;
-	LinkedTo.Append(TargetPin.LinkedTo);
-
-	Super::BreakPinLinks(TargetPin, bSendsNodeNotifcation);
-
-	if (TargetPin.Direction == EEdGraphPinDirection::EGPD_Input)
+	UBlueprint* Blueprint = FBlueprintEditorUtils::FindBlueprintForNodeChecked(TargetPin.GetOwningNode());
+	const UControlRigBlueprint* RigBlueprint = Cast<UControlRigBlueprint>(Blueprint);
+	if (RigBlueprint != nullptr)
 	{
-		ResetPinDefaultsRecursive(&TargetPin);
-	}
-	else
-	{
-		for (UEdGraphPin* LinkedPin : LinkedTo)
+		FString Left, Right;
+		if (!TargetPin.GetName().Split(TEXT("."), &Left, &Right))
+		{
+			return;
+		}
+		if (RigBlueprint->ModelController->BreakLinks(*Left, *Right, TargetPin.Direction == EGPD_Input))
 		{
-			if (LinkedPin->Direction == EEdGraphPinDirection::EGPD_Input)
-			{
-				ResetPinDefaultsRecursive(LinkedPin);
-			}
+			FBlueprintEditorUtils::MarkBlueprintAsModified(Blueprint);
 		}
 	}
-
-	FBlueprintEditorUtils::MarkBlueprintAsModified(Blueprint);	
 }
 
 void UControlRigGraphSchema::BreakSinglePinLink(UEdGraphPin* SourcePin, UEdGraphPin* TargetPin) const
 {
-	const FScopedTransaction Transaction(LOCTEXT("GraphEd_BreakSinglePinLink", "Break Pin Link") );
+	//const FScopedTransaction Transaction(LOCTEXT("GraphEd_BreakSinglePinLink", "Break Pin Link") );
 
 	UBlueprint* Blueprint = FBlueprintEditorUtils::FindBlueprintForNodeChecked(TargetPin->GetOwningNode());
-
-	Super::BreakSinglePinLink(SourcePin, TargetPin);
-	if (SourcePin->Direction == EEdGraphPinDirection::EGPD_Input)
-	{
-		ResetPinDefaultsRecursive(SourcePin);
-	}
-	else if (TargetPin->Direction == EEdGraphPinDirection::EGPD_Input)
+	const UControlRigBlueprint* RigBlueprint = Cast<UControlRigBlueprint>(Blueprint);
+	if (RigBlueprint != nullptr)
 	{
-		ResetPinDefaultsRecursive(TargetPin);
+		if (SourcePin->Direction == EGPD_Input)
+		{
+			UEdGraphPin* Temp = TargetPin;
+			TargetPin = SourcePin;
+			SourcePin = Temp;
+		}
+		
+		FString SourceLeft, SourceRight;
+		if (!SourcePin->GetName().Split(TEXT("."), &SourceLeft, &SourceRight))
+		{
+			return;
+		}
+		FString TargetLeft, TargetRight;
+		if (!TargetPin->GetName().Split(TEXT("."), &TargetLeft, &TargetRight))
+		{
+			return;
+		}
+		if (RigBlueprint->ModelController->BreakLink(*SourceLeft, *SourceRight, *TargetLeft, *TargetRight))
+		{
+			FBlueprintEditorUtils::MarkBlueprintAsModified(Blueprint);
+		}
 	}
-
-	FBlueprintEditorUtils::MarkBlueprintAsModified(Blueprint);
 }
 
 FConnectionDrawingPolicy* UControlRigGraphSchema::CreateConnectionDrawingPolicy(int32 InBackLayerID, int32 InFrontLayerID, float InZoomFactor, const FSlateRect& InClippingRect, class FSlateWindowElementList& InDrawElements, class UEdGraph* InGraphObj) const
@@ -441,16 +215,34 @@ UControlRigGraphNode* UControlRigGraphSchema::CreateGraphNode(UControlRigGraph*
 
 void UControlRigGraphSchema::TrySetDefaultValue(UEdGraphPin& InPin, const FString& InNewDefaultValue) const
 {
+#if WITH_EDITOR
+	if (GEditor)
+	{
+		GEditor->CancelTransaction(0);
+	}
+#endif
 	GetDefault<UEdGraphSchema_K2>()->TrySetDefaultValue(InPin, InNewDefaultValue);
 }
 
 void UControlRigGraphSchema::TrySetDefaultObject(UEdGraphPin& InPin, UObject* InNewDefaultObject) const
 {
+#if WITH_EDITOR
+	if (GEditor)
+	{
+		GEditor->CancelTransaction(0);
+	}
+#endif
 	GetDefault<UEdGraphSchema_K2>()->TrySetDefaultObject(InPin, InNewDefaultObject);
 }
 
 void UControlRigGraphSchema::TrySetDefaultText(UEdGraphPin& InPin, const FText& InNewDefaultText) const
 {
+#if WITH_EDITOR
+	if (GEditor)
+	{
+		GEditor->CancelTransaction(0);
+	}
+#endif
 	GetDefault<UEdGraphSchema_K2>()->TrySetDefaultText(InPin, InNewDefaultText);
 }
 
@@ -478,8 +270,7 @@ void UControlRigGraphSchema::ResetPinDefaultsRecursive(UEdGraphPin* InPin) const
 		return;
 	}
 
-	RigNode->CopyPinDefaultsToProperties(InPin, true, false);
-
+	RigNode->CopyPinDefaultsToModel(InPin);
 	for (UEdGraphPin* SubPin : InPin->SubPins)
 	{
 		ResetPinDefaultsRecursive(SubPin);
@@ -495,6 +286,7 @@ void UControlRigGraphSchema::GetVariablePinTypes(TArray<FEdGraphPinType>& PinTyp
 	PinTypes.Add(FEdGraphPinType(UEdGraphSchema_K2::PC_Struct, FName(NAME_None), TBaseStructure<FVector2D>::Get(), EPinContainerType::None, false, FEdGraphTerminalType()));
 	PinTypes.Add(FEdGraphPinType(UEdGraphSchema_K2::PC_Struct, FName(NAME_None), TBaseStructure<FRotator>::Get(), EPinContainerType::None, false, FEdGraphTerminalType()));
 	PinTypes.Add(FEdGraphPinType(UEdGraphSchema_K2::PC_Struct, FName(NAME_None), TBaseStructure<FTransform>::Get(), EPinContainerType::None, false, FEdGraphTerminalType()));
+	PinTypes.Add(FEdGraphPinType(UEdGraphSchema_K2::PC_Struct, FName(NAME_None), TBaseStructure<FEulerTransform>::Get(), EPinContainerType::None, false, FEdGraphTerminalType()));
 	PinTypes.Add(FEdGraphPinType(UEdGraphSchema_K2::PC_Struct, FName(NAME_None), TBaseStructure<FLinearColor>::Get(), EPinContainerType::None, false, FEdGraphTerminalType()));
 }
 
diff --git a/Engine/Plugins/Experimental/ControlRig/Source/ControlRigDeveloper/Private/SVariableMappingWidget.cpp b/Engine/Plugins/Experimental/ControlRig/Source/ControlRigDeveloper/Private/SVariableMappingWidget.cpp
new file mode 100644
--- /dev/null
+++ b/Engine/Plugins/Experimental/ControlRig/Source/ControlRigDeveloper/Private/SVariableMappingWidget.cpp
@@ -0,0 +1,380 @@
+// Copyright 1998-2019 Epic Games, Inc. All Rights Reserved.
+
+
+#include "SVariableMappingWidget.h"
+#include "Widgets/Text/SInlineEditableTextBlock.h"
+#include "Widgets/Input/SSearchBox.h"
+
+#define LOCTEXT_NAMESPACE "SVariableMappingWidget"
+
+static const FName ColumnId_VarLabel("Variable");
+static const FName ColumnID_MappingLabel("Mapping");
+
+void SVariableMappingTreeRow::Construct(const FArguments& InArgs, const TSharedRef<STableViewBase>& InOwnerTableView)
+{
+	Item = InArgs._Item;
+	OnVariableMappingChanged = InArgs._OnVariableMappingChanged;
+	OnGetVariableMapping = InArgs._OnGetVariableMapping;
+	OnGetAvailableMapping = InArgs._OnGetAvailableMapping;
+	OnGetFilteredText = InArgs._OnGetFilteredText;
+	OnVariableOptionAvailable = InArgs._OnVariableOptionAvailable;
+	OnPinCheckStateChanged = InArgs._OnPinCheckStateChanged;
+	OnPinGetCheckState = InArgs._OnPinGetCheckState;
+	OnPinIsEnabledCheckState = InArgs._OnPinIsEnabledCheckState;
+
+	check(Item.IsValid());
+
+	VariableOptionList.Reset();
+	if (OnGetAvailableMapping.IsBound())
+	{
+		TArray<FName> ListOfMappingVars;
+
+		OnGetAvailableMapping.Execute(Item->GetPathName(), ListOfMappingVars);
+		for (auto& MappingVar : ListOfMappingVars)
+		{
+			VariableOptionList.Add(MakeShareable(new FString(MappingVar.ToString())));
+		}
+	}
+	//SMultiColumnTableRow< FVariableMappingInfoPtr >::Construct(FSuperRowType::FArguments(), InOwnerTableView);
+	STableRow<TSharedPtr<FVariableMappingInfoPtr>>::Construct(
+		STableRow<TSharedPtr<FVariableMappingInfoPtr>>::FArguments()
+		.Padding(FMargin(3.0f, 2.0f))
+		.Content()
+		[
+			SNew(SHorizontalBox)
+
+			+ SHorizontalBox::Slot()
+			.FillWidth(0.5f)
+			.Padding(0.0f, 4.0f)
+			.HAlign(HAlign_Left)
+			[
+				SNew(STextBlock)
+				.Text(FText::FromString(Item->GetDisplayName()))
+				//.TextStyle(&FCoreStyle::Get().GetWidgetStyle<FTextBlockStyle>("BoldText"))
+				.HighlightText(this, &SVariableMappingTreeRow::GetFilterText)
+				//.IsReadOnly(true)
+				//.IsSelected(this, &SMultiColumnTableRow< FVariableMappingInfoPtr >::IsSelectedExclusively)
+			]
+// 			+ SHorizontalBox::Slot()
+// 			.Padding(2.0f, 4.0f)
+// 			.HAlign(HAlign_Right)
+// 			[
+// 				SNew(SHorizontalBox)
+// // 				+ SHorizontalBox::Slot()
+// // 				.FillWidth(0.5f)
+// // 				[
+// // 					SNew(STextBlock)
+// // 					.Text(LOCTEXT("EmptyString", ""))
+// // 				]
+				+ SHorizontalBox::Slot()
+				.Padding(0.0f, 4.0f)
+				.HAlign(HAlign_Right)
+				.AutoWidth()
+				[
+					SNew(STextBlock)
+					.Text(LOCTEXT("ExposeAsPinPropertyValue", "Use Pin "))
+				]
+				+ SHorizontalBox::Slot()
+				.Padding(0.0f, 4.0f)
+				.AutoWidth()
+				[
+					SNew(SCheckBox)
+					.IsChecked(this, &SVariableMappingTreeRow::IsPinChecked)
+					.OnCheckStateChanged(this, &SVariableMappingTreeRow::OnPinCheckStatusChanged)
+					.IsEnabled(this, &SVariableMappingTreeRow::IsPinEnabled)
+				]
+				+ SHorizontalBox::Slot()
+				.Padding(0.0f, 4.0f)
+				.HAlign(HAlign_Right)
+				.AutoWidth()
+				[
+					SNew(STextBlock)
+					.Text(LOCTEXT("Optional", " OR Use Curve "))
+				]
+				+ SHorizontalBox::Slot()
+				.FillWidth(0.5f)
+				[
+					SAssignNew(VarOptionComboBox, SSearchableComboBox)
+					.OptionsSource(&VariableOptionList)
+					.OnGenerateWidget(this, &SVariableMappingTreeRow::MakeVarOptionComboWidget)
+					.OnSelectionChanged(this, &SVariableMappingTreeRow::OnVarOptionSourceChanged)
+					.OnComboBoxOpening(this, &SVariableMappingTreeRow::OnVarOptionComboOpening)
+					.IsEnabled(this, &SVariableMappingTreeRow::IsVarOptionEnabled)
+					.ContentPadding(2)
+					.Content()
+					[
+						SNew(STextBlock)
+						.Text(this, &SVariableMappingTreeRow::GetVarOptionComboBoxContent)
+						//.Font(FEditorStyle::GetFontStyle(TEXT("PropertyWindow.NormalFont")))
+						.ToolTipText(this, &SVariableMappingTreeRow::GetVarOptionComboBoxToolTip)
+					]
+				]
+				+ SHorizontalBox::Slot()
+				.AutoWidth()
+				[
+					SNew(SButton)
+					.OnClicked(FOnClicked::CreateSP(this, &SVariableMappingTreeRow::OnClearButtonClicked))
+					.Text(FText::FromString(TEXT("x")))
+				]
+
+			//]
+		],
+		InOwnerTableView);
+}
+
+bool SVariableMappingTreeRow::IsVarOptionEnabled() const
+{
+	// if pin is not checked
+	if (IsPinChecked()==ECheckBoxState::Unchecked && 
+		OnVariableOptionAvailable.IsBound())
+	{
+		return OnVariableOptionAvailable.Execute(Item->GetPathName());
+	}
+
+	return false;
+}
+
+ECheckBoxState SVariableMappingTreeRow::IsPinChecked() const
+{
+	if (OnPinGetCheckState.IsBound())
+	{
+		return OnPinGetCheckState.Execute(Item->GetPathName());
+	}
+
+	return ECheckBoxState::Unchecked;
+}
+
+void SVariableMappingTreeRow::OnPinCheckStatusChanged(ECheckBoxState NewState)
+{
+	if (OnPinCheckStateChanged.IsBound())
+	{
+		OnPinCheckStateChanged.Execute(NewState, Item->GetPathName());
+	}
+}
+
+bool SVariableMappingTreeRow::IsPinEnabled() const
+{
+	if (OnPinIsEnabledCheckState.IsBound())
+	{
+		return OnPinIsEnabledCheckState.Execute(Item->GetPathName());
+	}
+
+	return false;
+}
+
+
+FReply SVariableMappingTreeRow::OnClearButtonClicked()
+{
+	if (OnVariableMappingChanged.IsBound())
+	{
+		OnVariableMappingChanged.Execute(Item->GetPathName(), NAME_None);
+	}
+
+	return FReply::Handled();
+}
+
+FText SVariableMappingTreeRow::GetFilterText() const
+{
+	if (OnGetFilteredText.IsBound())
+	{
+		return OnGetFilteredText.Execute();
+	}
+
+	return FText::GetEmpty();
+}
+
+void SVariableMappingTreeRow::OnVarOptionComboOpening()
+{
+	VariableOptionList.Reset();
+	if (OnGetAvailableMapping.IsBound())
+	{
+		TArray<FName> ListOfMappingVars;
+
+		OnGetAvailableMapping.Execute(Item->GetPathName(), ListOfMappingVars);
+		for (auto& MappingVar : ListOfMappingVars)
+		{
+			VariableOptionList.Add(MakeShareable(new FString(MappingVar.ToString())));
+		}
+	}
+
+	if (OnGetVariableMapping.IsBound())
+	{
+		FName MappedVar = OnGetVariableMapping.Execute(Item->GetPathName());
+		if (MappedVar!= NAME_None)
+		{
+			TSharedPtr<FString> ComboStringPtr = GetVarOptionString(MappedVar);
+			if (ComboStringPtr.IsValid())
+			{
+				VarOptionComboBox->SetSelectedItem(ComboStringPtr);
+			}
+		}
+		else
+		{
+			VarOptionComboBox->ClearSelection();
+		}
+	}
+}
+
+void SVariableMappingTreeRow::OnVarOptionSourceChanged(TSharedPtr<FString> NewSelection, ESelectInfo::Type SelectInfo)
+{
+	// if it's set from code, we did that on purpose
+	if (SelectInfo != ESelectInfo::Direct)
+	{
+		FString NewValue = *NewSelection.Get();
+		FName NewName = FName(*NewValue);
+
+		OnVariableMappingChanged.ExecuteIfBound(Item->GetPathName(), NewName);
+	}
+}
+
+FText SVariableMappingTreeRow::GetVarOptionComboBoxContent() const
+{
+	if (OnGetVariableMapping.IsBound())
+	{
+		return FText::FromName(OnGetVariableMapping.Execute(Item->GetPathName()));
+	}
+
+	return FText::FromString(TEXT("Invalid"));
+}
+
+FText SVariableMappingTreeRow::GetVarOptionComboBoxToolTip() const
+{
+	return LOCTEXT("VarOptionComboToolTip", "Map input/output variable to available options.");
+}
+
+TSharedPtr<FString> SVariableMappingTreeRow::GetVarOptionString(FName VarOptionName) const
+{
+	return TSharedPtr<FString>();
+	FString VarOptionString = VarOptionName.ToString();
+
+	// go through profile and see if it has mine
+	for (int32 Index = 1; Index < VariableOptionList.Num(); ++Index)
+	{
+		if (VarOptionString == *VariableOptionList[Index])
+		{
+			return VariableOptionList[Index];
+		}
+	}
+
+	return TSharedPtr<FString>();
+}
+
+TSharedRef<SWidget> SVariableMappingTreeRow::MakeVarOptionComboWidget(TSharedPtr<FString> InItem)
+{
+	return SNew(STextBlock).Text(FText::FromString(*InItem));// .Font(FEditorStyle::GetFontStyle(TEXT("PropertyWindow.NormalFont")));
+}
+//////////////////////////////////////////////////////////////////////////
+// SVariableMappingWidget
+
+void SVariableMappingWidget::Construct(const FArguments& InArgs/*, FSimpleMulticastDelegate& InOnPostUndo*/)
+{
+	OnGetAvailableMappingDelegate = InArgs._OnGetAvailableMapping;
+	OnGetVariableMappingDelegate = InArgs._OnGetVariableMapping;
+	OnVariableMappingChangedDelegate = InArgs._OnVariableMappingChanged;
+	OnCreateVariableMappingDelegate = InArgs._OnCreateVariableMapping;
+	OnVariableOptionAvailableDelegate = InArgs._OnVariableOptionAvailable;
+	OnPinCheckStateChangedDelegate = InArgs._OnPinCheckStateChanged;
+	OnPinGetCheckStateDelegate = InArgs._OnPinGetCheckState;
+	OnPinIsEnabledCheckStateDelegate = InArgs._OnPinIsEnabledCheckState;
+
+	//InOnPostUndo.Add(FSimpleDelegate::CreateSP(this, &SVariableMappingWidget::PostUndo));
+
+	ChildSlot
+		[
+			SNew(SVerticalBox)
+
+			// now show bone mapping
+			+ SVerticalBox::Slot()
+			.AutoHeight()
+			.Padding(0, 2)
+			[
+				SNew(SHorizontalBox)
+				// Filter entry
+				+ SHorizontalBox::Slot()
+				.FillWidth(1)
+				[
+					SAssignNew(NameFilterBox, SSearchBox)
+					.SelectAllTextWhenFocused(true)
+					.OnTextChanged(this, &SVariableMappingWidget::OnFilterTextChanged)
+					.OnTextCommitted(this, &SVariableMappingWidget::OnFilterTextCommitted)
+				]
+			]
+
+			+ SVerticalBox::Slot()
+			.FillHeight(1.0f)		// This is required to make the scrollbar work, as content overflows Slate containers by default
+			[
+				SAssignNew(VariableMappingTreeView, SVariableMappingTreeView)
+				.TreeItemsSource(&VariableMappingList)
+				.OnGenerateRow(this, &SVariableMappingWidget::MakeTreeRowWidget)
+				.OnGetChildren(this, &SVariableMappingWidget::GetChildrenForInfo)
+				.ItemHeight(22.0f)
+			]
+		];
+
+	RefreshVariableMappingList();
+}
+
+TSharedRef<ITableRow> SVariableMappingWidget::MakeTreeRowWidget(TSharedPtr<FVariableMappingInfo> InInfo, const TSharedRef<STableViewBase>& OwnerTable)
+{
+	return SNew(SVariableMappingTreeRow, OwnerTable)
+		.Item(InInfo)
+		.OnVariableMappingChanged(OnVariableMappingChangedDelegate)
+		.OnGetVariableMapping(OnGetVariableMappingDelegate)
+		.OnGetAvailableMapping(OnGetAvailableMappingDelegate)
+		.OnGetFilteredText(this, &SVariableMappingWidget::GetFilterText)
+		.OnVariableOptionAvailable(OnVariableOptionAvailableDelegate)
+		.OnPinCheckStateChanged(OnPinCheckStateChangedDelegate)
+		.OnPinGetCheckState(OnPinGetCheckStateDelegate)
+		.OnPinIsEnabledCheckState(OnPinIsEnabledCheckStateDelegate);
+}
+
+void SVariableMappingWidget::GetChildrenForInfo(TSharedPtr<FVariableMappingInfo> InInfo, TArray< TSharedPtr<FVariableMappingInfo> >& OutChildren)
+{
+	OutChildren = InInfo->Children;
+}
+
+void SVariableMappingWidget::OnFilterTextChanged(const FText& SearchText)
+{
+	// need to make sure not to have the same text go
+	// otherwise, the widget gets recreated multiple times causing 
+	// other issue
+	if (FilterText.CompareToCaseIgnored(SearchText) != 0)
+	{
+		FilterText = SearchText;
+		RefreshVariableMappingList();
+	}
+}
+
+void SVariableMappingWidget::OnFilterTextCommitted(const FText& SearchText, ETextCommit::Type CommitInfo)
+{
+	// Just do the same as if the user typed in the box
+	OnFilterTextChanged(SearchText);
+}
+
+TSharedRef<ITableRow> SVariableMappingWidget::GenerateVariableMappingRow(TSharedPtr<FVariableMappingInfo> InInfo, const TSharedRef<STableViewBase>& OwnerTable)
+{
+	check(InInfo.IsValid());
+
+	return
+		SNew(SVariableMappingTreeRow, OwnerTable)
+		.Item(InInfo)
+		.OnVariableMappingChanged(OnVariableMappingChangedDelegate)
+		.OnGetVariableMapping(OnGetVariableMappingDelegate)
+		.OnGetAvailableMapping(OnGetAvailableMappingDelegate)
+		.OnGetFilteredText(this, &SVariableMappingWidget::GetFilterText);
+}
+
+void SVariableMappingWidget::RefreshVariableMappingList()
+{
+	OnCreateVariableMappingDelegate.ExecuteIfBound(FilterText.ToString(), VariableMappingList);
+
+	VariableMappingTreeView->RequestListRefresh();
+}
+
+// void SVariableMappingWidget::PostUndo()
+// {
+// 	RefreshVariableMappingList();
+// }
+
+#undef LOCTEXT_NAMESPACE
+
diff --git a/Engine/Plugins/Experimental/ControlRig/Source/ControlRigDeveloper/Private/SVariableMappingWidget.h b/Engine/Plugins/Experimental/ControlRig/Source/ControlRigDeveloper/Private/SVariableMappingWidget.h
new file mode 100644
--- /dev/null
+++ b/Engine/Plugins/Experimental/ControlRig/Source/ControlRigDeveloper/Private/SVariableMappingWidget.h
@@ -0,0 +1,231 @@
+// Copyright 1998-2019 Epic Games, Inc. All Rights Reserved.
+
+#pragma once
+
+#include "CoreMinimal.h"
+#include "SSearchableComboBox.h"
+
+//=====================================================================================
+// UI for displaying property info
+//=====================================================================================
+
+struct FVariableMappingInfo
+{
+public:
+	// This is the property that is the most shallow type
+	// It will be Transform.Translation.X
+	FName PropertyName; 
+	// Display Name
+	FString DisplayName;
+	// List of Children
+	// in theory, this actually shouldn't be active if you have children
+	// but it represent each raw nonetheless
+	// this is to map curve which is always float
+	TArray<TSharedPtr<FVariableMappingInfo>> Children;
+
+	/** Static function for creating a new item, but ensures that you can only have a TSharedRef to one */
+	static TSharedRef<FVariableMappingInfo> Make(const FName InPathName)
+	{
+		return MakeShareable(new FVariableMappingInfo(InPathName));
+	}
+
+	FName GetPathName() const
+	{
+		return PropertyName;
+	}
+
+	FString GetDisplayName() const
+	{
+		return DisplayName;
+	}
+
+protected:
+	/** Hidden constructor, always use Make above */
+	FVariableMappingInfo(const FName InPathName)
+		: PropertyName(InPathName)
+	{
+		FString PathString = InPathName.ToString();
+		int32 Found = PathString.Find(TEXT("."), ESearchCase::IgnoreCase, ESearchDir::FromEnd);
+		DisplayName = (Found != INDEX_NONE) ? PathString.RightChop(Found + 1) : PathString;
+	}
+
+	/** Hidden constructor, always use Make above */
+	FVariableMappingInfo() {}
+};
+
+//////////////////////////////////////////////////////////////////////////
+// SVariableMappingListRow
+
+typedef TSharedPtr< FVariableMappingInfo > FVariableMappingInfoPtr;
+
+DECLARE_DELEGATE_TwoParams(FOnVariableMappingChanged, const FName& /** PathName */, const FName&/** Curve Name**/)
+DECLARE_DELEGATE_RetVal_OneParam(FName, FOnGetVariableMapping, const FName& /** PathName **/)
+// type verified, and make sure it's not double booked
+DECLARE_DELEGATE_TwoParams(FOnGetAvailableMapping, const FName& /** PathName **/, TArray<FName>& /** List of available mappings**/)
+// to highlight item
+DECLARE_DELEGATE_RetVal(FText&, FOnGetFilteredText)
+// caller to create list
+DECLARE_DELEGATE_TwoParams(FOnCreateVariableMapping, const FString&, TArray< TSharedPtr<FVariableMappingInfo> >&)
+DECLARE_DELEGATE_RetVal_OneParam(bool, FOnVarOptionAvailable, const FName& /** Property Name */);
+
+// pin exposure option
+DECLARE_DELEGATE_TwoParams(FOnPinCheckStateChanged, ECheckBoxState, FName /** Property Name */);
+DECLARE_DELEGATE_RetVal_OneParam(ECheckBoxState, FOnPinGetCheckState, FName /** Property Name */);
+DECLARE_DELEGATE_RetVal_OneParam(bool, FOnPinIsCheckEnabled, FName /** Property Name */);
+
+class SVariableMappingTreeRow
+	: public STableRow< TSharedPtr<FVariableMappingInfoPtr> >
+{
+public:
+	
+	SLATE_BEGIN_ARGS(SVariableMappingTreeRow) {}
+
+	/** The item for this row **/
+	SLATE_ARGUMENT(FVariableMappingInfoPtr, Item)
+	SLATE_EVENT(FOnVariableMappingChanged, OnVariableMappingChanged)
+	SLATE_EVENT(FOnGetVariableMapping, OnGetVariableMapping)
+	SLATE_EVENT(FOnGetAvailableMapping, OnGetAvailableMapping)
+	SLATE_EVENT(FOnGetFilteredText, OnGetFilteredText)
+	SLATE_EVENT(FOnVarOptionAvailable, OnVariableOptionAvailable);
+	SLATE_EVENT(FOnPinCheckStateChanged, OnPinCheckStateChanged)
+	SLATE_EVENT(FOnPinGetCheckState, OnPinGetCheckState)
+	SLATE_EVENT(FOnPinIsCheckEnabled, OnPinIsEnabledCheckState)
+	SLATE_END_ARGS()
+
+	void Construct(const FArguments& InArgs, const TSharedRef<STableViewBase>& OwnerTableView);
+
+private:
+
+	/** Widget used to display the list of variable option list*/
+	TSharedPtr<SSearchableComboBox>	VarOptionComboBox;
+	TArray< TSharedPtr< FString > >	VariableOptionList;
+
+	/** The name and weight of the variable option*/
+	FVariableMappingInfoPtr	Item;
+
+	// Curve combo box options
+	FReply OnClearButtonClicked();
+	FText GetFilterText() const;
+
+	TSharedRef<SWidget> MakeVarOptionComboWidget(TSharedPtr<FString> InItem);
+	void OnVarOptionSourceChanged(TSharedPtr<FString> NewSelection, ESelectInfo::Type SelectInfo);
+	FText GetVarOptionComboBoxContent() const;
+	FText GetVarOptionComboBoxToolTip() const;
+	void OnVarOptionComboOpening();
+	bool IsVarOptionEnabled() const;
+	TSharedPtr<FString> GetVarOptionString(FName VarOptionName) const;
+
+	ECheckBoxState IsPinChecked() const;
+	void OnPinCheckStatusChanged(ECheckBoxState NewState);
+	bool IsPinEnabled() const;
+
+	FOnVariableMappingChanged	OnVariableMappingChanged;
+	FOnGetVariableMapping		OnGetVariableMapping;
+	FOnGetAvailableMapping		OnGetAvailableMapping;
+	FOnGetFilteredText			OnGetFilteredText;
+	FOnVarOptionAvailable		OnVariableOptionAvailable;
+	FOnPinCheckStateChanged		OnPinCheckStateChanged;
+	FOnPinGetCheckState			OnPinGetCheckState;
+	FOnPinIsCheckEnabled		OnPinIsEnabledCheckState;
+};
+
+typedef STreeView< TSharedPtr<FVariableMappingInfo> > SVariableMappingTreeView;
+
+class SVariableMappingWidget : public SCompoundWidget
+{
+public:
+	SLATE_BEGIN_ARGS(SVariableMappingWidget)
+	{}
+
+	SLATE_EVENT(FOnVariableMappingChanged, OnVariableMappingChanged)
+	SLATE_EVENT(FOnGetVariableMapping, OnGetVariableMapping)
+	SLATE_EVENT(FOnGetAvailableMapping, OnGetAvailableMapping)
+	SLATE_EVENT(FOnCreateVariableMapping, OnCreateVariableMapping)
+	SLATE_EVENT(FOnGetFilteredText, OnGetFilteredText)
+	SLATE_EVENT(FOnVarOptionAvailable, OnVariableOptionAvailable);
+	SLATE_EVENT(FOnPinCheckStateChanged, OnPinCheckStateChanged)
+	SLATE_EVENT(FOnPinGetCheckState, OnPinGetCheckState)
+	SLATE_EVENT(FOnPinIsCheckEnabled, OnPinIsEnabledCheckState)
+	SLATE_END_ARGS()
+
+	/**
+	* Slate construction function
+	*
+	* @param InArgs - Arguments passed from Slate
+	*
+	*/
+	void Construct(const FArguments& InArgs/*, FSimpleMulticastDelegate& InOnPostUndo*/);
+
+	/**
+	* Filters the SListView when the user changes the search text box (NameFilterBox)
+	*
+	* @param SearchText - The text the user has typed
+	*
+	*/
+	void OnFilterTextChanged(const FText& SearchText);
+
+	/**
+	* Filters the SListView when the user hits enter or clears the search box
+	* Simply calls OnFilterTextChanged
+	*
+	* @param SearchText - The text the user has typed
+	* @param CommitInfo - Not used
+	*
+	*/
+	void OnFilterTextCommitted(const FText& SearchText, ETextCommit::Type CommitInfo);
+
+	/**
+	* Create a widget for an entry in the tree from an info
+	*
+	* @param InInfo - Shared pointer to the morph target we're generating a row for
+	* @param OwnerTable - The table that owns this row
+	*
+	* @return A new Slate widget, containing the UI for this row
+	*/
+	TSharedRef<ITableRow> GenerateVariableMappingRow(TSharedPtr<FVariableMappingInfo> InInfo, const TSharedRef<STableViewBase>& OwnerTable);
+
+	/**
+	* Handler for the delete of retarget source
+	*/
+	void RefreshVariableMappingList();
+
+private:
+
+	/**
+	* Accessor so our rows can grab the filtertext for highlighting
+	*
+	*/
+	FText& GetFilterText() { return FilterText; }
+
+	/** Box to filter to a specific morph target name */
+	TSharedPtr<SSearchBox>	NameFilterBox;
+
+	/** Widget used to display the list of retarget sources */
+	TSharedPtr<SVariableMappingTreeView> VariableMappingTreeView;
+
+	/** A list of variable mapping list. Used by the VariableMappingListView. */
+	TArray< TSharedPtr<FVariableMappingInfo> > VariableMappingList;
+
+	/** Current text typed into NameFilterBox */
+	FText FilterText;
+
+	/** Delegate for undo/redo transaction **/
+	void PostUndo();
+
+	// Make a single tree row widget
+	TSharedRef<ITableRow> MakeTreeRowWidget(TSharedPtr<FVariableMappingInfo> InInfo, const TSharedRef<STableViewBase>& OwnerTable);
+
+	// Get the children for the provided bone info
+	void GetChildrenForInfo(TSharedPtr<FVariableMappingInfo> InInfo, TArray< TSharedPtr<FVariableMappingInfo> >& OutChildren);
+
+	FOnGetAvailableMapping		OnGetAvailableMappingDelegate;
+	FOnVariableMappingChanged	OnVariableMappingChangedDelegate;
+	FOnGetVariableMapping		OnGetVariableMappingDelegate;
+	FOnCreateVariableMapping	OnCreateVariableMappingDelegate;
+	FOnVarOptionAvailable		OnVariableOptionAvailableDelegate;
+	FOnPinCheckStateChanged		OnPinCheckStateChangedDelegate;
+	FOnPinGetCheckState			OnPinGetCheckStateDelegate;
+	FOnPinIsCheckEnabled		OnPinIsEnabledCheckStateDelegate;
+
+};
+
diff --git a/Engine/Plugins/Experimental/ControlRig/Source/ControlRigDeveloper/Public/ControlRigBlueprint.h b/Engine/Plugins/Experimental/ControlRig/Source/ControlRigDeveloper/Public/ControlRigBlueprint.h
--- a/Engine/Plugins/Experimental/ControlRig/Source/ControlRigDeveloper/Public/ControlRigBlueprint.h
+++ b/Engine/Plugins/Experimental/ControlRig/Source/ControlRigDeveloper/Public/ControlRigBlueprint.h
@@ -10,10 +10,12 @@
 #include "ControlRigDefines.h"
 #include "Hierarchy.h"
 #include "Interfaces/Interface_PreviewMeshProvider.h"
+#include "ControlRigController.h"
 #include "ControlRigBlueprint.generated.h"
 
 class UControlRigBlueprintGeneratedClass;
 class USkeletalMesh;
+class UControlRigGraph;
 
 /** 
   * Source data used by the FControlRigBlueprintCompiler, can't be an editor plugin
@@ -104,6 +106,8 @@ class CONTROLRIGDEVELOPER_API UControlRigBlueprint : public UBlueprint, public I
 public:
 	UControlRigBlueprint();
 
+	void InitializeModel();
+
 	/** Get the (full) generated class for this control rig blueprint */
 	UControlRigBlueprintGeneratedClass* GetControlRigBlueprintGeneratedClass() const;
 
@@ -127,6 +131,20 @@ public:
 	virtual void SetPreviewMesh(USkeletalMesh* PreviewMesh, bool bMarkAsDirty = true) override;
 	virtual USkeletalMesh* GetPreviewMesh() const override;
 
+	UPROPERTY(transient)
+	UControlRigModel* Model;
+
+	UPROPERTY(transient)
+	UControlRigController* ModelController;
+
+	bool bSuspendModelNotificationsForSelf;
+	bool bSuspendModelNotificationsForOthers;
+	FName LastNameFromNotification;
+
+	void PopulateModelFromGraph(const UControlRigGraph* InGraph);
+
+	UControlRigModel::FModifiedEvent& OnModified();
+
 private:
 	/** Links between the various properties we have */
 	UPROPERTY()
@@ -148,6 +166,11 @@ private:
 	UPROPERTY(DuplicateTransient, AssetRegistrySearchable)
 	TSoftObjectPtr<USkeletalMesh> PreviewSkeletalMesh;
 
+	UControlRigModel::FModifiedEvent _ModifiedEvent;
+	void HandleModelModified(const UControlRigModel* InModel, EControlRigModelNotifType InType, const void* InPayload);
+	bool UpdateParametersOnControlRig(UControlRig* InRig = nullptr);
+	bool PerformArrayOperation(const FString& InPropertyPath, TFunctionRef<bool(FScriptArrayHelper&, int32)> InOperation, bool bCallModify, bool bPropagateToInstances);
+
 	friend class FControlRigBlueprintCompilerContext;
 	friend class SRigHierarchy;
 	friend class FControlRigEditor;
diff --git a/Engine/Plugins/Experimental/ControlRig/Source/ControlRigDeveloper/Public/ControlRigBlueprintUtils.h b/Engine/Plugins/Experimental/ControlRig/Source/ControlRigDeveloper/Public/ControlRigBlueprintUtils.h
--- a/Engine/Plugins/Experimental/ControlRig/Source/ControlRigDeveloper/Public/ControlRigBlueprintUtils.h
+++ b/Engine/Plugins/Experimental/ControlRig/Source/ControlRigDeveloper/Public/ControlRigBlueprintUtils.h
@@ -21,7 +21,7 @@ struct CONTROLRIGDEVELOPER_API FControlRigBlueprintUtils
  * @param	InStructTemplate	The struct template we want to use
  * @return the name of the new member
  */
-static FName GetNewUnitMemberName(UBlueprint* InBlueprint, UStruct* InStructTemplate);
+static FName GetNewUnitMemberName(UBlueprint* InBlueprint, const UStruct* InStructTemplate);
 
 /** 
  * Adds a new unit member struct. 
@@ -29,7 +29,7 @@ static FName GetNewUnitMemberName(UBlueprint* InBlueprint, UStruct* InStructTemp
  * @param	InStructTemplate	The struct template we want to use
  * @return the name of the new member, or NAME_None if the member was not created
  */
-static FName AddUnitMember(UBlueprint* InBlueprint, UStruct* InStructTemplate);
+static FName AddUnitMember(UBlueprint* InBlueprint, const UStruct* InStructTemplate, const FName& InName= NAME_None);
 
 /** 
  * Gets a new name for a member
@@ -56,6 +56,7 @@ static FName AddPropertyMember(UBlueprint* InBlueprint, const FEdGraphPinType& I
  * @return the new graph node
  */
 static UControlRigGraphNode* InstantiateGraphNodeForProperty(UEdGraph* InGraph, const FName& InPropertyName, const FVector2D& InLocation = FVector2D::ZeroVector, const FEdGraphPinType& InPinType = FEdGraphPinType());
+static UControlRigGraphNode* InstantiateGraphNodeForStructPath(UEdGraph* InGraph, const FName& InPropertyName, const FVector2D& InLocation, const FString& InStructPath);
 
 /**
  * Check whether we can instantiate a node in the specified graph for the specified property
@@ -80,4 +81,6 @@ static void HandleRenameVariableReferencesEvent(UBlueprint* InBlueprint, UClass*
 
 /** remove the variable if not used by anybody else but ToBeDeleted*/
 static void RemoveMemberVariableIfNotUsed(UBlueprint* Blueprint, const FName VarName, UControlRigGraphNode* ToBeDeleted);
+
+static FName ValidateName(UBlueprint* InBlueprint, const FString& InName);
 };
\ No newline at end of file
diff --git a/Engine/Plugins/Experimental/ControlRig/Source/ControlRigDeveloper/Public/ControlRigController.h b/Engine/Plugins/Experimental/ControlRig/Source/ControlRigDeveloper/Public/ControlRigController.h
new file mode 100644
--- /dev/null
+++ b/Engine/Plugins/Experimental/ControlRig/Source/ControlRigDeveloper/Public/ControlRigController.h
@@ -0,0 +1,180 @@
+// Copyright 1998-2019 Epic Games, Inc. All Rights Reserved.
+
+#pragma once
+
+#include "CoreMinimal.h"
+#include "ControlRigModel.h"
+#include "UObject/PropertyPortFlags.h"
+#include "ControlRigController.generated.h"
+
+/**
+ * The ControlRigController is the sole authority to perform changes
+ * in a ControlRigModel.
+ */
+UCLASS()
+class CONTROLRIGDEVELOPER_API UControlRigController : public UObject
+{
+	GENERATED_BODY()
+
+public:
+
+	UControlRigController();
+	virtual ~UControlRigController();
+
+	UPROPERTY()
+	UControlRigModel* Model;
+
+	void SetModel(UControlRigModel* InModel);
+	virtual bool Clear();
+	static bool ConstructPreviewParameter(const FName& InDataType, EControlRigModelParameterType InParameterType, FControlRigModelNode& OutNode);
+	static bool ConstructPreviewNode(const FName& InFunctionName, FControlRigModelNode& OutNode);
+	virtual bool AddParameter(const FName& InName, const FName& InDataType, EControlRigModelParameterType InParameterType = EControlRigModelParameterType::Input, const FVector2D& InPosition = FVector2D::ZeroVector, bool bUndo = true);
+	virtual bool AddNode(const FName& InFunctionName, const FVector2D& InPosition = FVector2D::ZeroVector, const FName& InName = NAME_None, bool bUndo = true);
+	virtual bool RemoveNode(const FName& InName, bool bUndo = true);
+	virtual bool SetNodePosition(const FName& InName, const FVector2D& InPosition, bool bUndo = true);
+	virtual bool SetParameterType(const FName& InName, EControlRigModelParameterType InParameterType, bool bUndo = true);
+	virtual bool RenameNode(const FName& InOldNodeName, const FName& InNewNodeName, bool bUndo = true);
+	virtual bool ClearSelection();
+	virtual bool SetSelection(const TArray<FName>& InNodeSelection);
+	virtual bool SelectNode(const FName& InName, bool bInSelected = true, bool bClearSelection = true);
+	bool DeselectNode(const FName& InName);
+	bool PrepareCycleCheckingForPin(const FName& InNodeName, const FName& InPinName, bool bIsInput);
+	bool ResetCycleCheck();
+	bool CanLink(const FName& InSourceNodeName, const FName& InSourceOutputPinName, const FName& InTargetNodeName, const FName& InTargetInputPinName, FString* OutFailureReason, bool bReportError = false) const;
+	bool MakeLink(const FName& InSourceNodeName, const FName& InSourceOutputPinName, const FName& InTargetNodeName, const FName& InTargetInputPinName, FString* OutFailureReason = nullptr, bool bUndo = true);
+	bool BreakLink(const FName& InSourceNodeName, const FName& InSourceOutputPinName, const FName& InTargetNodeName, const FName& InTargetInputPinName, bool bUndo = true);
+	bool BreakLinks(const FName& InNodeName, const FName& InPinName, bool bIsInput, bool bUndo = true);
+
+	virtual bool GetPinDefaultValue(const FName& InNodeName, const FName& InPinName, FString& OutDefaultValue) const;
+	virtual bool SetPinDefaultValue(const FName& InNodeName, const FName& InPinName, const FString& InDefaultValue, bool bLog = true, bool bUndo = true);
+	virtual bool GetPinDefaultValueBool(const FName& InNodeName, const FName& InPinName, bool& OutDefaultValue) const;
+	virtual bool SetPinDefaultValueBool(const FName& InNodeName, const FName& InPinName, bool InDefaultValue, bool bLog = true, bool bUndo = true);
+	virtual bool GetPinDefaultValueFloat(const FName& InNodeName, const FName& InPinName, float& OutDefaultValue) const;
+	virtual bool SetPinDefaultValueFloat(const FName& InNodeName, const FName& InPinName, float InDefaultValue, bool bLog = true, bool bUndo = true);
+	virtual bool GetPinDefaultValueInt(const FName& InNodeName, const FName& InPinName, int32& OutDefaultValue) const;
+	virtual bool SetPinDefaultValueInt(const FName& InNodeName, const FName& InPinName, int32 InDefaultValue, bool bLog = true, bool bUndo = true);
+	virtual bool GetPinDefaultValueName(const FName& InNodeName, const FName& InPinName, FName& OutDefaultValue) const;
+	virtual bool SetPinDefaultValueName(const FName& InNodeName, const FName& InPinName, const FName& InDefaultValue, bool bLog = true, bool bUndo = true);
+	virtual bool GetPinDefaultValueVector(const FName& InNodeName, const FName& InPinName, FVector& OutDefaultValue) const;
+	virtual bool SetPinDefaultValueVector(const FName& InNodeName, const FName& InPinName, const FVector& InDefaultValue, bool bLog = true, bool bUndo = true);
+	virtual bool GetPinDefaultValueQuat(const FName& InNodeName, const FName& InPinName, FQuat& OutDefaultValue) const;
+	virtual bool SetPinDefaultValueQuat(const FName& InNodeName, const FName& InPinName, const FQuat& InDefaultValue, bool bLog = true, bool bUndo = true);
+	virtual bool GetPinDefaultValueTransform(const FName& InNodeName, const FName& InPinName, FTransform& OutDefaultValue) const;
+	virtual bool SetPinDefaultValueTransform(const FName& InNodeName, const FName& InPinName, const FTransform& InDefaultValue, bool bLog = true, bool bUndo = true);
+	template<typename T>
+	bool GetPinDefaultValueStruct(const FName& InNodeName, const FName& InPinName, T& OutDefaultValue) const;
+	template<typename T>
+	bool SetPinDefaultValueStruct(const FName& InNodeName, const FName& InPinName, const T& InDefaultValue, bool bLog = true, bool bUndo = true);
+
+	virtual	bool AddArrayPin(const FName& InNodeName, const FName& InPinName, const FString& InDefaultValue = FString(), bool bUndo = true);
+	virtual	bool PopArrayPin(const FName& InNodeName, const FName& InPinName, bool bUndo = true);
+	virtual	bool ClearArrayPin(const FName& InNodeName, const FName& InPinNam, bool bUndo = true);
+	virtual	bool SetArrayPinSize(const FName& InNodeName, const FName& InPinName, int32 InSize, const FString& InDefaultValue = FString(), bool bUndo = true);
+	virtual bool ExpandPin(const FName& InNodeName, const FName& InPinName, bool bIsInput, bool bExpanded = true, bool bUndo = true);
+
+	virtual bool Undo();
+	virtual bool Redo();
+
+	virtual bool ResendAllNotifications();
+	virtual bool ResendAllPinDefaultNotifications();
+	virtual void EnableMessageLog(bool bInEnabled = true);
+
+#if CONTROLRIG_UNDO
+	virtual bool OpenUndoBracket(const FString& Title);
+	virtual bool CloseUndoBracket();
+#endif
+
+	UControlRigModel::FModifiedEvent& OnModified();
+
+private:
+
+	bool GetPinDefaultValueRecursiveStruct(const FControlRigModelPin* InPin, FString& OutValue) const;
+	bool SetPinDefaultValueRecursiveStruct(const FControlRigModelPin* OutPin, const FString& InValue, bool bUndo);
+
+	bool bSuspendLog;
+	void LogMessage(const FString& Message) const;
+	void LogWarning(const FString& Message) const;
+	void LogError(const FString& Message) const;
+
+	static bool FindPinTypeFromDataType(const FName& InDataType, FEdGraphPinType& OutPinType);
+	bool EnsureModel() const;
+	void HandleModelModified(const UControlRigModel* InModel, EControlRigModelNotifType InType, const void* InPayload);
+
+	EControlRigModelNotifType _LastModelNotification;
+	FDelegateHandle _ModelModifiedHandle;
+	UControlRigModel::FModifiedEvent _ModifiedEvent;
+
+#if CONTROLRIG_UNDO
+	TArray<TSharedPtr<UControlRigModel::FAction>> _UndoBrackets;
+#endif
+};
+
+template<typename T>
+inline bool UControlRigController::GetPinDefaultValueStruct(const FName& InNodeName, const FName& InPinName, T& OutDefaultValue) const
+{
+	if (!EnsureModel())
+	{
+		return false;
+	}
+	const FControlRigModelPin* Pin = Model->FindPin(InNodeName, InPinName);
+	if (Pin == nullptr)
+	{
+		return false;
+	}
+	if (UScriptStruct* Struct = Cast<UScriptStruct>(Pin->Type.PinSubCategoryObject))
+	{
+		FString DefaultValueStr;
+		if (!GetPinDefaultValueRecursiveStruct(Pin, DefaultValueStr))
+		{
+			return false;
+		}
+		Struct->ImportText(*DefaultValueStr, &OutDefaultValue, nullptr, EPropertyPortFlags::PPF_None, nullptr, Struct->GetName(), true);
+		return true;
+	}
+	return false;
+}
+
+template<typename T>
+inline bool UControlRigController::SetPinDefaultValueStruct(const FName& InNodeName, const FName& InPinName, const T& InDefaultValue, bool bLog, bool bUndo)
+{
+	if (!EnsureModel())
+	{
+		return false;
+	}
+	const FControlRigModelPin* Pin = Model->FindPin(InNodeName, InPinName);
+	if (Pin == nullptr)
+	{
+		return false;
+	}
+
+	if (UScriptStruct* Struct = Cast<UScriptStruct>(Pin->Type.PinSubCategoryObject))
+	{
+#if CONTROLRIG_UNDO
+		UControlRigModel::FAction Action;
+		if (bUndo)
+		{
+			Action.Type = EControlRigModelNotifType::Invalid;
+			Action.Title = TEXT("Set Pin Default from Struct");
+			Model->CurrentActions.Push(&Action);
+		}
+#endif
+
+		FString DefaultValueStr;
+		Struct->ExportText(DefaultValueStr, &InDefaultValue, nullptr, nullptr, PPF_None, nullptr);
+		bool bSucceeded = SetPinDefaultValueRecursiveStruct(Pin, DefaultValueStr, bUndo);
+
+#if CONTROLRIG_UNDO
+		if (bUndo)
+		{
+			Model->CurrentActions.Pop();
+			if (bSucceeded)
+			{
+				Model->PushAction(Action);
+			}
+		}
+#endif
+
+		return bSucceeded;
+	}
+	return false;
+}
diff --git a/Engine/Plugins/Experimental/ControlRig/Source/ControlRigDeveloper/Public/ControlRigGraphTraverser.h b/Engine/Plugins/Experimental/ControlRig/Source/ControlRigDeveloper/Public/ControlRigGraphTraverser.h
--- a/Engine/Plugins/Experimental/ControlRig/Source/ControlRigDeveloper/Public/ControlRigGraphTraverser.h
+++ b/Engine/Plugins/Experimental/ControlRig/Source/ControlRigDeveloper/Public/ControlRigGraphTraverser.h
@@ -3,8 +3,7 @@
 #pragma once
 
 #include "ControlRigBlueprint.h"
-#include "Graph/ControlRigGraph.h"
-#include "Graph/ControlRigGraphNode.h"
+#include "ControlRigModel.h"
 
 /*
  * The ControlRigTraverser is used to walk the UI graph (UEdGraph).
@@ -16,23 +15,22 @@ class CONTROLRIGDEVELOPER_API FControlRigGraphTraverser
 {
 public:
 
-	FControlRigGraphTraverser(UControlRigBlueprint* InBlueprint, UControlRigGraph* InGraph);
+	FControlRigGraphTraverser(UControlRigModel* InModel);
 
 #if WITH_EDITORONLY_DATA
 	// Returns true if a given unit is part of a valid execution graph
-	bool IsWiredToExecution(const FName& UnitName);
+	bool IsWiredToExecution(const FName& NodeName);
 #endif
 
 	// walks the UEdGraph, finds all valid nodes and builds property links
-	void TraverseAndBuildPropertyLinks();
+	void TraverseAndBuildPropertyLinks(UControlRigBlueprint* Blueprint);
 
 private:
 
 	// Returns true if a given node is part of a valid execution graph
-	bool IsWiredToExecution(UControlRigGraphNode* Node);
+	bool IsWiredToExecution(const FControlRigModelNode* Node);
 
-	UControlRigBlueprint* Blueprint;
-	UControlRigGraph* Graph;
+	UControlRigModel* Model;
 	TMap<FName, bool> VisitedNodes;
 
 	friend class FControlRigBlueprintCompilerContext;
