Control Rig MVC Refactoring

diff --git a/Engine/Plugins/Experimental/ControlRig/Source/ControlRigDeveloper/Public/ControlRigModel.h b/Engine/Plugins/Experimental/ControlRig/Source/ControlRigDeveloper/Public/ControlRigModel.h
new file mode 100644
--- /dev/null
+++ b/Engine/Plugins/Experimental/ControlRig/Source/ControlRigDeveloper/Public/ControlRigModel.h
@@ -0,0 +1,394 @@
+// Copyright 1998-2019 Epic Games, Inc. All Rights Reserved.
+
+#pragma once
+
+#include "CoreMinimal.h"
+#include "EdGraph/EdGraphPin.h"
+#include "ControlRigModel.generated.h"
+
+#define CONTROLRIG_UNDO WITH_EDITOR
+
+class UControlRigModel;
+
+enum class EControlRigModelNotifType : uint8
+{
+	ModelError,
+	ModelCleared,
+	NodeAdded,
+	NodeRemoved,
+	NodeRenamed,
+	NodeChanged,
+	NodeSelected,
+	NodeDeselected,
+	LinkAdded,
+	LinkRemoved,
+	PinAdded,
+	PinRemoved,
+	PinChanged,
+	Invalid
+};
+
+/**
+ * Describes a single error within the model
+ */
+struct CONTROLRIGDEVELOPER_API FControlRigModelError
+{
+	FString Message;
+
+	FControlRigModelError()
+		: Message(FString())
+	{
+	}
+};
+
+/**
+ * A pair of node + pin index
+ */
+struct CONTROLRIGDEVELOPER_API FControlRigModelPair
+{
+	int32 Node;
+	int32 Pin;
+
+	FControlRigModelPair()
+		: Node(INDEX_NONE)
+		, Pin(INDEX_NONE)
+	{
+	}
+
+	FControlRigModelPair(int32 InNode, int32 InPin)
+		: Node(InNode)
+		, Pin(InPin)
+	{
+	}
+
+	bool operator==(const FControlRigModelPair& Other) const
+	{
+		return Node == Other.Node && Pin == Other.Pin;
+	}
+
+	bool IsValid() const
+	{
+		return Node != INDEX_NONE && Pin != INDEX_NONE;
+	}
+
+#if CONTROLRIG_UNDO
+
+	void AppendArgumentsForAction(TArray<FString>& InOutArguments, const UControlRigModel* InModel) const;
+	void ConfigureFromActionArguments(const TArray<FString>& InOutArguments, int32 ArgumentIndex, const UControlRigModel* InModel);
+
+#endif
+};
+
+/**
+ * A link between two pins.
+ */
+struct CONTROLRIGDEVELOPER_API FControlRigModelLink
+{
+	int32 Index;
+	FControlRigModelPair Source;
+	FControlRigModelPair Target;
+
+	FControlRigModelLink()
+		: Index(INDEX_NONE)
+		, Source(FControlRigModelPair())
+		, Target(FControlRigModelPair())
+	{
+	}
+
+	bool IsValid() const
+	{
+		return Index != INDEX_NONE && Source.IsValid() && Target.IsValid();
+	}
+
+#if CONTROLRIG_UNDO
+
+	void AppendArgumentsForAction(TArray<FString>& InOutArguments, const UControlRigModel* InModel) const;
+	void ConfigureFromActionArguments(const TArray<FString>& InOutArguments, int32 ArgumentIndex, const UControlRigModel* InModel);
+
+#endif
+};
+
+/**
+ * A single node within the control rig's model
+ */
+struct CONTROLRIGDEVELOPER_API FControlRigModelPin
+{
+	FName Name;
+	int32 Node;
+	int32 Index;
+	int32 ParentIndex;
+	TArray<int32> SubPins;
+	TEnumAsByte<enum EEdGraphPinDirection> Direction;
+	FEdGraphPinType Type;
+	FString DefaultValue;
+	bool bExpanded;
+	bool bIsConstant;
+	FName CustomWidgetName;
+	TArray<int32> Links;
+	FText TooltipText;
+
+	FControlRigModelPin()
+		: Name(NAME_None)
+		, Node(INDEX_NONE)
+		, Index(INDEX_NONE)
+		, ParentIndex(INDEX_NONE)
+		, Direction(EGPD_Input)
+		, Type(FEdGraphPinType())
+		, DefaultValue(FString())
+		, bExpanded(false)
+		, bIsConstant(false)
+		, CustomWidgetName(NAME_None)
+	{
+	}
+
+	FControlRigModelPair GetPair() const
+	{
+		FControlRigModelPair Pair;
+		Pair.Node = Node;
+		Pair.Pin = Index;
+		return Pair;
+	}
+
+	bool IsValid() const
+	{
+		return Index != INDEX_NONE && Node != INDEX_NONE && Name != NAME_None;
+	}
+
+	bool IsSingleValue() const
+	{
+		return Type.ContainerType == EPinContainerType::None;
+	}
+
+	bool IsArray() const
+	{
+		return Type.IsArray();
+	}
+
+	int32 ArraySize() const
+	{
+		return SubPins.Num();
+	}
+
+#if CONTROLRIG_UNDO
+
+	void AppendArgumentsForAction(TArray<FString>& InOutArguments) const;
+	void ConfigureFromActionArguments(const TArray<FString>& InOutArguments, int32 ArgumentIndex = 0);
+
+#endif
+};
+
+enum class EControlRigModelParameterType : uint8
+{
+	None,
+	Input,
+	Output,
+	Hidden
+};
+
+/**
+ * A single node within the model.
+ */
+struct CONTROLRIGDEVELOPER_API FControlRigModelNode
+{
+	FName Name;
+	int32 Index;
+	FName FunctionName;
+	FVector2D Position;
+	TArray<FControlRigModelPin> Pins;
+	EControlRigModelParameterType ParameterType;
+
+	FControlRigModelNode()
+		: Name(NAME_None)
+		, Index(INDEX_NONE)
+		, FunctionName(NAME_None)
+		, Position(FVector2D::ZeroVector)
+		, ParameterType(EControlRigModelParameterType::None)
+	{
+	}
+
+	bool IsValid() const
+	{
+		return Index != INDEX_NONE && Name != NAME_None;
+	}
+
+	bool IsParameter() const
+	{
+		return ParameterType != EControlRigModelParameterType::None;
+	}
+
+	FString GetPinPath(int32 InPinIndex, bool bIncludeNodeName = true) const;
+
+	bool IsMutable() const;
+	bool IsBeginExecution() const;
+	const UStruct* UnitStruct() const;
+	const FControlRigModelPin* FindPin(const FName& InName, bool bLookForInput = true) const;
+
+#if CONTROLRIG_UNDO
+
+	void AppendArgumentsForAction(TArray<FString>& InOutArguments) const;
+	void ConfigureFromActionArguments(const TArray<FString>& InOutArguments, int32 ArgumentIndex = 0);
+
+#endif
+};
+
+// A struct used for passing on information about a rename operation
+struct CONTROLRIGDEVELOPER_API FControlRigModelNodeRenameInfo
+{
+	FName OldName;
+	FName NewName;
+	FControlRigModelNode Node;
+
+	FControlRigModelNodeRenameInfo()
+		: OldName(NAME_None)
+		, NewName(NAME_None)
+		, Node(FControlRigModelNode())
+	{
+	}
+};
+
+/**
+ * The ControlRigModel represents the low level data required to compile
+ * a Control Rig. This is similar to the UI aspects of the Control Rig,
+ * such as the ControlRigGraph. The model centralizes all of the relevant
+ * data in one place.
+ * The only class which can mutate the model's data is the ControlRigController.
+ */
+UCLASS()
+class CONTROLRIGDEVELOPER_API UControlRigModel : public UObject
+{
+	GENERATED_BODY()
+
+public:
+
+	static const FName ValueName;
+
+	UControlRigModel();
+	virtual ~UControlRigModel();
+
+	const TArray<FControlRigModelNode>& Nodes() const;
+	TArray<FControlRigModelNode> SelectedNodes() const;
+	const TArray<FControlRigModelLink>& Links() const;
+	TArray<FControlRigModelPin> LinkedPins(const FControlRigModelPair& InPin) const;
+	TArray<FControlRigModelPin> LinkedPins(const FName& InNodeName, const FName& InPinName, bool bLookForInput = true) const;
+	TArray<FControlRigModelNode> Parameters() const;
+
+	DECLARE_EVENT_ThreeParams(UControlRigModel, FModifiedEvent, const UControlRigModel*, EControlRigModelNotifType, const void*);
+
+	const FControlRigModelNode* FindNode(const FName& InName) const;
+	const FControlRigModelNode* FindNode(int32 InNodeIndex) const;
+	const FControlRigModelPin* FindPin(const FName& InNodeName, const FName& InPinName, bool bLookForInput = true) const;
+	const FControlRigModelPin* FindPin(const FControlRigModelPair& InPin) const;
+	const FControlRigModelPin* FindSubPin(const FControlRigModelPin* InParentPin, const FName& InSubPinName) const;
+	const FControlRigModelPin* FindParentPin(const FControlRigModelPin* InSubPin) const;
+	const FControlRigModelPin* FindPinFromPath(const FString& InPinPath, bool bLookForInput = true) const;
+	const FControlRigModelLink* FindLink(int32 InLinkIndex) const;
+
+	FName GetUniqueNodeName(const FName& InName) const;
+	static void SplitPinPath(const FString& InPinPath, FString& OutLeft, FString& OutRight, bool bInSplitForNodeName = true);
+	const FControlRigModelPin* GetParentPin(const FControlRigModelPair& InPin) const;
+	FString GetPinPath(const FControlRigModelPair& InPin, bool bIncludeNodeName = true) const;
+
+#if WITH_EDITORONLY_DATA
+
+	UPROPERTY()
+	int32 ActionCount;
+
+#endif
+
+#if CONTROLRIG_UNDO
+
+	virtual void PostTransacted(const FTransactionObjectEvent& TransactionEvent) override;
+
+	bool Undo();
+	bool Redo();
+
+#endif
+
+private:
+
+	FModifiedEvent& OnModified();
+
+	bool Clear();
+	bool IsNodeNameAvailable(const FName& InName) const;
+	bool AddNode(const FControlRigModelNode& InNode, bool bUndo = true);
+	bool AddParameter(const FName& InName, const FEdGraphPinType& InDataType, EControlRigModelParameterType InParameterType, const FVector2D& InPosition, bool bUndo = true);
+	bool RemoveNode(const FName& InName, bool bUndo = true);
+	bool SetNodePosition(const FName& InName, const FVector2D& InPosition, bool bUndo = true);
+	bool SetParameterType(const FName& InName, EControlRigModelParameterType InParameterType, bool bUndo = true);
+	bool RenameNode(const FName& InOldNodeName, const FName& InNewNodeName, bool bUndo = true);
+	bool SelectNode(const FName& InName, bool bInSelected);
+	bool AreCompatibleTypes(const FEdGraphPinType& A, const FEdGraphPinType& B) const;
+	bool PrepareCycleCheckingForPin(int32 InNodeIndex, int32 InPinIndex);
+	void ResetCycleCheck();
+	bool CanLink(int32 InSourceNodeIndex, int32 InSourcePinIndex, int32 InTargetNodeIndex, int32 InTargetPinIndex, FString* OutFailureReason);
+	bool MakeLink(int32 InSourceNodeIndex, int32 InSourcePinIndex, int32 InTargetNodeIndex, int32 InTargetPinIndex, bool bUndo = true);
+	bool BreakLink(int32 InSourceNodeIndex, int32 InSourcePinIndex, int32 InTargetNodeIndex, int32 InTargetPinIndex, bool bUndo = true);
+	bool BreakLinks(int32 InNodeIndex, int32 InPinIndex, bool bUndo = true);
+	bool GetPinDefaultValue(const FName& InNodeName, const FName& InPinName, FString& OutValue) const;
+	bool GetPinDefaultValue(const FControlRigModelPair& InPin, FString& OutValue) const;
+	bool SetPinDefaultValue(const FName& InNodeName, const FName& InPinName, const FString& InValue, bool bUndo = true);
+	bool SetPinDefaultValue(const FControlRigModelPair& InPin, const FString& InValue, bool bUndo = true);
+
+	bool SetPinArraySize(const FControlRigModelPair& InPin, int32 InArraySize, const FString& InDefaultValue, bool bUndo = true);
+	bool ExpandPin(const FName& InNodeName, const FName& InPinName, bool bIsInput, bool bInExpanded, bool bUndo = true);
+	bool ResendAllNotifications();
+	bool ResendAllPinDefaultNotifications();
+
+	static bool ShouldStructBeUnfolded(const UStruct* Struct);
+	static FEdGraphPinType GetPinTypeFromField(UProperty* Property);
+	static void AddNodePinsForFunction(FControlRigModelNode& Node);
+	static void SetNodePinDefaultsForFunction(FControlRigModelNode& Node);
+	static void AddNodePinsForParameter(FControlRigModelNode& Node, const FEdGraphPinType& InDataType);
+	static void SetNodePinDefaultsForParameter(FControlRigModelNode& Node, const FEdGraphPinType& InDataType);
+	static void ConfigurePinFromField(FControlRigModelPin& Pin, UProperty* Property);
+	static int32 AddPinsRecursive(FControlRigModelNode& Node, int32 ParentIndex, const UStruct* Struct, EEdGraphPinDirection PinDirection, int32& LastAddedIndex);
+	int32 RemovePinsRecursive(FControlRigModelNode& Node, int32 ParentIndex, bool bUndo = true);
+	static void ConfigurePinIndices(FControlRigModelNode& Node);
+	static void GetParameterPinTypes(TArray<FEdGraphPinType>& PinTypes);
+
+	TArray<FControlRigModelNode> _Nodes;
+	TArray<FControlRigModelLink> _Links;
+
+	bool bIsSelecting;
+	TArray<FName> _SelectedNodes;
+	TMap<FName, FVector2D> _LastNodePositions;
+	FModifiedEvent _ModifiedEvent;
+
+	FControlRigModelPair _CycleCheckSubject;
+	TArray<bool> _NodeIsOnCycle;
+
+#if CONTROLRIG_UNDO
+
+	/**
+	 * An action performed with additional context / arguments
+	 */
+	struct FAction
+	{
+		EControlRigModelNotifType Type;
+		FString Title;
+		TArray<FString> Arguments;
+		TArray<FAction> SubActions;
+
+		FAction()
+		{
+			Type = EControlRigModelNotifType::Invalid;
+		}
+
+		bool IsValid() const
+		{
+			return (Type != EControlRigModelNotifType::Invalid) != (Arguments.Num() == 0);
+		}
+	};
+
+	TArray<FAction*> CurrentActions;
+	TArray<FAction> UndoActions;
+	TArray<FAction> RedoActions;
+
+	void PushAction(const FAction& InAction);
+	bool UndoAction(const FAction& InAction);
+	bool RedoAction(const FAction& InAction);
+
+#endif
+
+	friend class UControlRigController;
+};
diff --git a/Engine/Plugins/Experimental/ControlRig/Source/ControlRigDeveloper/Public/Graph/ControlRigGraph.h b/Engine/Plugins/Experimental/ControlRig/Source/ControlRigDeveloper/Public/Graph/ControlRigGraph.h
--- a/Engine/Plugins/Experimental/ControlRig/Source/ControlRigDeveloper/Public/Graph/ControlRigGraph.h
+++ b/Engine/Plugins/Experimental/ControlRig/Source/ControlRigDeveloper/Public/Graph/ControlRigGraph.h
@@ -5,6 +5,7 @@
 #include "EdGraph/EdGraph.h"
 #include "Graph/ControlRigGraphNode.h"
 #include "Hierarchy.h"
+#include "ControlRigModel.h"
 #include "ControlRigGraph.generated.h"
 
 class UControlRigBlueprint;
@@ -30,14 +31,22 @@ public:
 	virtual void Serialize(FArchive& Ar) override;
 #endif
 #if WITH_EDITOR
+
 	virtual void PostLoad() override;
 	void OnBlueprintCompiledPostLoad(UBlueprint*);
 	FDelegateHandle BlueprintOnCompiledHandle;
 	void CacheBoneNameList(const FRigHierarchy& Hierarchy);
 	const TArray<TSharedPtr<FString>>& GetBoneNameList() const;
 
+	bool bSuspendModelNotifications;
+	bool bIsTemporaryGraphForCopyPaste;
+
+	UControlRigGraphNode* FindNodeFromPropertyName(const FName& InPropertyName);
+
 private:
 
+	void HandleModelModified(const UControlRigModel* InModel, EControlRigModelNotifType InType, const void* InPayload);
+
 	TArray<UControlRigGraphNode*> FoundHierarchyRefVariableNodes;
 	TArray<UControlRigGraphNode*> FoundHierarchyRefMutableNodes;
 	TMap<UControlRigGraphNode*, TArray<UControlRigGraphNode*>> FoundHierarchyRefConnections;
diff --git a/Engine/Plugins/Experimental/ControlRig/Source/ControlRigDeveloper/Public/Graph/ControlRigGraphNode.h b/Engine/Plugins/Experimental/ControlRig/Source/ControlRigDeveloper/Public/Graph/ControlRigGraphNode.h
--- a/Engine/Plugins/Experimental/ControlRig/Source/ControlRigDeveloper/Public/Graph/ControlRigGraphNode.h
+++ b/Engine/Plugins/Experimental/ControlRig/Source/ControlRigDeveloper/Public/Graph/ControlRigGraphNode.h
@@ -9,6 +9,7 @@
 #include "K2Node.h"
 #include "UObject/SoftObjectPath.h"
 #include "EdGraphSchema_K2.h"
+#include "ControlRigModel.h"
 #include "ControlRigGraphNode.generated.h"
 
 class FBlueprintActionDatabaseRegistrar;
@@ -20,11 +21,11 @@ class UControlRigBlueprint;
 class FControlRigField
 {
 public:
-	FControlRigField(const FEdGraphPinType& InPinType, const FString& InPropertyPath, const FText& InDisplayNameText, int32 InArrayIndex = INDEX_NONE)
+	FControlRigField(const FEdGraphPinType& InPinType, const FString& InPinPath, const FText& InDisplayNameText, int32 InArrayIndex = INDEX_NONE)
 		: InputPin(nullptr)
 		, OutputPin(nullptr)
 		, PinType(InPinType)
-		, PropertyPath(InPropertyPath)
+		, PinPath(InPinPath)
 		, DisplayNameText(InDisplayNameText)
 		, ArrayIndex(InArrayIndex)
 	{
@@ -36,13 +37,13 @@ public:
 	virtual UField* GetField() const { return nullptr; }
 
 	/** Get the input pin for this item */
-	virtual UEdGraphPin* GetPin() const { return InputPin; }
-
+	virtual const FControlRigModelPin* GetPin() const { return nullptr; }
+	
 	/** Get the output pin for this item */
 	virtual UEdGraphPin* GetOutputPin() const { return OutputPin; }
 
 	/** Get the name of this field */
-	virtual FString GetPropertyPath() const { return PropertyPath; }
+	virtual FString GetPinPath() const { return PinPath; }
 
 	/** Get the name to display for this field */
 	virtual FText GetDisplayNameText() const { return DisplayNameText; }
@@ -63,7 +64,7 @@ public:
 	FEdGraphPinType PinType;
 
 	/** Cached name to display */
-	FString PropertyPath;
+	FString PinPath;
 
 	/** The name to display for this field  */
 	FText DisplayNameText;
@@ -78,40 +79,32 @@ public:
 	TArray<TSharedRef<FControlRigField>> Children;
 };
 
-/** Information about an input/output property */
-class FControlRigProperty : public FControlRigField
+/** Information about an input/output pin */
+class FControlRigPin : public FControlRigField
 {
 private:
-	static FEdGraphPinType GetPinTypeFromProperty(UProperty* InProperty)
+	static FEdGraphPinType GetPinTypeFromPin(const FControlRigModelPin* InPin)
 	{
-		FEdGraphPinType OutPinType; 
-		GetDefault<UEdGraphSchema_K2>()->ConvertPropertyToPinType(InProperty, OutPinType);
-
-		return OutPinType;
+		return InPin->Type;
 	}
 
-	static FText GetDisplayNameForProperty(UProperty* InProperty, int32 InArrayIndex)
+	static FText GetDisplayNameForPin(const FControlRigModelPin* InPin)
 	{
-		if(InArrayIndex != INDEX_NONE)
-		{
-			return FText::Format(NSLOCTEXT("ControlRigGraphNode", "ArrayPinFormat", "[{0}]"), FText::AsNumber(InArrayIndex));
-		}
-
-		return InProperty->GetDisplayNameText();
+		return FText::FromString(InPin->Name.ToString());
 	}
 
 public:
-	FControlRigProperty(UProperty* InProperty, const FString& InPropertyPath, int32 InArrayIndex = INDEX_NONE)
-		: FControlRigField(GetPinTypeFromProperty(InProperty), *InPropertyPath, GetDisplayNameForProperty(InProperty, InArrayIndex), InArrayIndex)
-		, Property(InProperty)
+	FControlRigPin(const FControlRigModelPin* InPin, const FString& InPinPath, int32 InArrayIndex = INDEX_NONE)
+		: FControlRigField(GetPinTypeFromPin(InPin), *InPinPath, GetDisplayNameForPin(InPin), InArrayIndex)
+		, Pin(*InPin)
 	{
 	}
 
-	virtual UField* GetField() const { return Property; }
+	virtual const FControlRigModelPin* GetPin() const override { return &Pin; }
 
 private:
 	/** The field that we use as input/output */
-	UProperty* Property;
+	FControlRigModelPin Pin;
 };
 
 /** Base class for animation ControlRig-related nodes */
@@ -124,6 +117,7 @@ class CONTROLRIGDEVELOPER_API UControlRigGraphNode : public UEdGraphNode
 	friend class FControlRigBlueprintCompilerContext;
 	friend class UControlRigGraph;
 	friend class UControlRigGraphSchema;
+	friend class UControlRigBlueprint;
 	friend class FControlRigGraphTraverser;
 	friend class FControlRigGraphPanelPinFactory;
 	friend class FControlRigEditor;
@@ -174,25 +168,15 @@ public:
 	virtual void DestroyNode() override;
 	virtual void PinDefaultValueChanged(UEdGraphPin* Pin) override;
 	virtual TSharedPtr<INameValidatorInterface> MakeNameValidator() const override;
-	virtual void OnRenameNode(const FString& InNewName) override;
 	virtual FText GetTooltipText() const override;
 	virtual bool CanCreateUnderSpecifiedSchema(const UEdGraphSchema* InSchema) const override;
 	virtual void AutowireNewNode(UEdGraphPin* FromPin) override;	
 	virtual void PrepareForCopying() override;
-	virtual void PostPasteNode() override;
 
 	virtual bool IsDeprecated() const;
 	virtual bool ShouldWarnOnDeprecation() const;
 	virtual FString GetDeprecationMessage() const;
 
-	// UK2Node Interface
-	// @TODO: need to find a better way to handle the following functions! 
-	// We cant derive from UK2Node as we don't want most of its functionality
-// 	virtual bool ShouldShowNodeProperties() const { return true; }
-
-	/** Handle a variable being renamed */
-	virtual void HandleVariableRenamed(UBlueprint* InBlueprint, UClass* InVariableClass, UEdGraph* InGraph, const FName& InOldVarName, const FName& InNewVarName);
-
 	/** Set the cached dimensions of this node */
 	void SetDimensions(const FVector2D& InDimensions) { Dimensions = InDimensions; }
 
@@ -224,7 +208,7 @@ public:
 	bool IsPinExpanded(const FString& InPinPropertyPath) const;
 
 	/** Propagate pin defaults to underlying properties if they have changed */
-	void CopyPinDefaultsToProperties(UEdGraphPin* Pin, bool bCallModify, bool bPropagateToInstances);
+	void CopyPinDefaultsToModel(UEdGraphPin* Pin, bool bUndo = false);
 
 	/** Check whether we are a property accessor */
 	bool IsPropertyAccessor() const { return GetUnitScriptStruct() == nullptr; }
@@ -281,7 +265,7 @@ protected:
 	UClass* GetControlRigSkeletonGeneratedClass() const;
 
 	/** Create a ControlRig field from a field on the ControlRig class, if possible */
-	TSharedPtr<FControlRigField> CreateControlRigField(UField* Field, const FString& PropertyPath, int32 InArrayIndex = INDEX_NONE) const;
+	TSharedPtr<FControlRigField> CreateControlRigField(const FControlRigModelPin* InPin, const FString& InPinPath, int32 InArrayIndex = INDEX_NONE) const;
 
 	/** Get all fields that act as execution pins for this node */
 	void GetExecutionFields(TArray<TSharedRef<FControlRigField>>& OutFields) const;
@@ -296,9 +280,7 @@ protected:
 	void GetInputOutputFields(TArray<TSharedRef<FControlRigField>>& OutFields) const;
 
 	/** Helper function for GetInputFields/GetOutputFields */
-	void GetFields(TFunction<bool(UProperty*)> InPropertyCheckFunction, TArray<TSharedRef<FControlRigField>>& OutFields) const;
-	void GetFields_Recursive(const TSharedRef<FControlRigField>& ParentControlRigField, const FString& ParentPropertyPath) const;
-	void GetFields_RecursiveHelper(UProperty* InProperty, const TSharedRef<FControlRigField>& InControlRigField, const FString& InPropertyPath) const;
+	void GetFields(TFunction<bool(const FControlRigModelPin*, const FControlRigModelNode*)> InPinCheckFunction, TArray<TSharedRef<FControlRigField>>& OutFields) const;
 
 	/** Get the struct property for the unit we represent, if any (we could just be a property accessor) */
 	UStructProperty* GetUnitProperty() const;
@@ -309,9 +291,6 @@ protected:
 	/** Get the property for the unit we represent */
 	UProperty* GetProperty() const;
 
-	/** Creates auto generated defaults for pins */
-	void SetupPinAutoGeneratedDefaults(UEdGraphPin* Pin);
-
 	/** Copies default values from underlying properties into pin defaults, for editing */
 	void SetupPinDefaultsFromCDO(UEdGraphPin* Pin);
 
@@ -330,14 +309,10 @@ protected:
 	/** Destroy all pins in an array */
 	void DestroyPinList(TArray<UEdGraphPin*>& InPins);
 
-	/** 
-	 * Perform the specified operation on the array described by the passed-in property path. 
-	 * If bCallModify is true then it is assumed that the array will be mutated.
-	 */
-	bool PerformArrayOperation(const FString& InPropertyPath, TFunctionRef<bool(FScriptArrayHelper&,int32)> InOperation, bool bCallModify, bool bPropagateToInstances) const;
-
 private:
 
+	bool IsVariable() const;
+
 #if WITH_EDITORONLY_DATA
 	TArray<UEdGraphNode*> HierarchyRefOutputConnections;
 #endif
diff --git a/Engine/Plugins/Experimental/ControlRig/Source/ControlRigDeveloper/Public/Graph/ControlRigGraphSchema.h b/Engine/Plugins/Experimental/ControlRig/Source/ControlRigDeveloper/Public/Graph/ControlRigGraphSchema.h
--- a/Engine/Plugins/Experimental/ControlRig/Source/ControlRigDeveloper/Public/Graph/ControlRigGraphSchema.h
+++ b/Engine/Plugins/Experimental/ControlRig/Source/ControlRigDeveloper/Public/Graph/ControlRigGraphSchema.h
@@ -74,15 +74,6 @@ public:
 	/** Automatically layout the passed-in nodes */
 	void LayoutNodes(UControlRigGraph* InGraph, const TArray<UControlRigGraphNode*>& InNodes) const;
 
-	/** Helper function to detect if a connection would cause a cycle */
-	bool TryCreateConnection_DetectCycle(UEdGraphPin* PinA, UEdGraphPin* PinB) const;
-
-	/** Helper function to allow us to apply extended logic to TryCreateConnection */
-	bool TryCreateConnection_Extended(UEdGraphPin* PinA, UEdGraphPin* PinB) const;
-
-	/** Helper function to allow us to apply extended logic to CanCreateConnection */
-	const FControlRigPinConnectionResponse CanCreateConnection_Extended(const UEdGraphPin* A, const UEdGraphPin* B) const;
-
 	/** Helper function to rename a node */
 	void RenameNode(UControlRigGraphNode* Node, const FName& InNewNodeName) const;
 
diff --git a/Engine/Plugins/Experimental/ControlRig/Source/ControlRigEditor/ControlRigEditor.Build.cs b/Engine/Plugins/Experimental/ControlRig/Source/ControlRigEditor/ControlRigEditor.Build.cs
--- a/Engine/Plugins/Experimental/ControlRig/Source/ControlRigEditor/ControlRigEditor.Build.cs
+++ b/Engine/Plugins/Experimental/ControlRig/Source/ControlRigEditor/ControlRigEditor.Build.cs
@@ -28,6 +28,7 @@ namespace UnrealBuildTool.Rules
                     "ControlRigDeveloper",
                     "Kismet",
                     "EditorStyle",
+                    "ApplicationCore",
                     "AnimationCore",
                     "PropertyEditor",
                     "AnimGraph",
diff --git a/Engine/Plugins/Experimental/ControlRig/Source/ControlRigEditor/Private/ControlRigVariableDetailsCustomization.cpp b/Engine/Plugins/Experimental/ControlRig/Source/ControlRigEditor/Private/ControlRigVariableDetailsCustomization.cpp
--- a/Engine/Plugins/Experimental/ControlRig/Source/ControlRigEditor/Private/ControlRigVariableDetailsCustomization.cpp
+++ b/Engine/Plugins/Experimental/ControlRig/Source/ControlRigEditor/Private/ControlRigVariableDetailsCustomization.cpp
@@ -10,12 +10,11 @@
 #include "Widgets/Text/STextBlock.h"
 #include "Widgets/Input/SCheckBox.h"
 #include "Engine/BlueprintGeneratedClass.h"
+#include "ControlRigBlueprint.h"
+#include "ControlRigController.h"
 
 #define LOCTEXT_NAMESPACE "ControlRigVariableDetailsCustomization"
 
-static FName AnimationOutputMetadataName("AnimationOutput");
-static FName AnimationInputMetadataName("AnimationInput");
-
 TSharedPtr<IDetailCustomization> FControlRigVariableDetailsCustomization::MakeInstance(TSharedPtr<IBlueprintEditor> InBlueprintEditor)
 {
 	const TArray<UObject*>* Objects = (InBlueprintEditor.IsValid() ? InBlueprintEditor->GetObjectsCurrentlyBeingEdited() : nullptr);
@@ -52,7 +51,7 @@ void FControlRigVariableDetailsCustomization::CustomizeDetails(IDetailLayoutBuil
 				.NameContent()
 				[
 					SNew(STextBlock)
-					.IsEnabled(IsAnimationFlagEnabled(PropertyBeingCustomized))
+					.Visibility(IsAnimationFlagEnabled(PropertyBeingCustomized) ? EVisibility::Visible : EVisibility::Hidden)
 					.Font(DetailLayout.GetDetailFont())
 					.Text(AnimationOutputText)
 					.ToolTipText(AnimationOutputTooltipText)
@@ -60,7 +59,7 @@ void FControlRigVariableDetailsCustomization::CustomizeDetails(IDetailLayoutBuil
 				.ValueContent()
 				[
 					SNew(SCheckBox)
-					.IsEnabled(IsAnimationFlagEnabled(PropertyBeingCustomized))
+					.Visibility(IsAnimationFlagEnabled(PropertyBeingCustomized) ? EVisibility::Visible : EVisibility::Hidden)
 					.IsChecked_Raw(this, &FControlRigVariableDetailsCustomization::IsAnimationOutputChecked, PropertyBeingCustomized)
 					.OnCheckStateChanged_Raw(this, &FControlRigVariableDetailsCustomization::HandleAnimationOutputCheckStateChanged, PropertyBeingCustomized)
 					.ToolTipText(AnimationOutputTooltipText)
@@ -71,7 +70,7 @@ void FControlRigVariableDetailsCustomization::CustomizeDetails(IDetailLayoutBuil
 				.NameContent()
 				[
 					SNew(STextBlock)
-					.IsEnabled(IsAnimationFlagEnabled(PropertyBeingCustomized))
+					.Visibility(IsAnimationFlagEnabled(PropertyBeingCustomized) ? EVisibility::Visible : EVisibility::Hidden)
 					.Font(DetailLayout.GetDetailFont())
 					.Text(AnimationInputText)
 					.ToolTipText(AnimationInputTooltipText)
@@ -79,7 +78,7 @@ void FControlRigVariableDetailsCustomization::CustomizeDetails(IDetailLayoutBuil
 				.ValueContent()
 				[
 					SNew(SCheckBox)
-					.IsEnabled(IsAnimationFlagEnabled(PropertyBeingCustomized))
+					.Visibility(IsAnimationFlagEnabled(PropertyBeingCustomized) ? EVisibility::Visible : EVisibility::Hidden)
 					.IsChecked_Raw(this, &FControlRigVariableDetailsCustomization::IsAnimationInputChecked, PropertyBeingCustomized)
 					.OnCheckStateChanged_Raw(this, &FControlRigVariableDetailsCustomization::HandleAnimationInputCheckStateChanged, PropertyBeingCustomized)
 					.ToolTipText(AnimationInputTooltipText)
@@ -90,29 +89,44 @@ void FControlRigVariableDetailsCustomization::CustomizeDetails(IDetailLayoutBuil
 
 bool FControlRigVariableDetailsCustomization::IsAnimationFlagEnabled(TWeakObjectPtr<UProperty> PropertyBeingCustomized) const
 {
-	if (PropertyBeingCustomized.IsValid())
+	if (PropertyBeingCustomized.IsValid() && BlueprintPtr.IsValid())
 	{
-		if (UBlueprintGeneratedClass* GeneratedClass = Cast<UBlueprintGeneratedClass>(PropertyBeingCustomized.Get()->GetOwnerClass()))
+		UControlRigBlueprint* RigBlueprint = Cast<UControlRigBlueprint>(BlueprintPtr.Get());
+		if (RigBlueprint)
 		{
-			UBlueprint* PropertyOwnerBlueprint = Cast<UBlueprint>(GeneratedClass->ClassGeneratedBy);
-			return PropertyOwnerBlueprint == BlueprintPtr.Get();
+			if (RigBlueprint->Model)
+			{
+				UProperty* PropertyPtrBeingCustomized = PropertyBeingCustomized.Get();
+				const FControlRigModelNode* Node = RigBlueprint->Model->FindNode(PropertyPtrBeingCustomized->GetFName());
+				if (Node)
+				{
+					return Node->IsParameter();
+				}
+			}
 		}
-
 	}
-
 	return false;
 }
 
 ECheckBoxState FControlRigVariableDetailsCustomization::IsAnimationOutputChecked(TWeakObjectPtr<UProperty> PropertyBeingCustomized) const
 {
-	if (PropertyBeingCustomized.IsValid())
+	if (PropertyBeingCustomized.IsValid() && BlueprintPtr.IsValid())
 	{
-		UProperty* PropertyPtrBeingCustomized = PropertyBeingCustomized.Get();
-		FString Value;
-		if (PropertyPtrBeingCustomized->HasMetaData(AnimationOutputMetadataName) ||
-			FBlueprintEditorUtils::GetBlueprintVariableMetaData(BlueprintPtr.Get(), PropertyPtrBeingCustomized->GetFName(), nullptr, AnimationOutputMetadataName, Value))
+		UControlRigBlueprint* RigBlueprint = Cast<UControlRigBlueprint>(BlueprintPtr.Get());
+		if (RigBlueprint)
 		{
-			return ECheckBoxState::Checked;
+			if (RigBlueprint->Model)
+			{
+				UProperty* PropertyPtrBeingCustomized = PropertyBeingCustomized.Get();
+				const FControlRigModelNode* Node = RigBlueprint->Model->FindNode(PropertyPtrBeingCustomized->GetFName());
+				if (Node)
+				{
+					if (Node->ParameterType == EControlRigModelParameterType::Output)
+					{
+						return ECheckBoxState::Checked;
+					}
+				}
+			}
 		}
 	}
 
@@ -121,27 +135,40 @@ ECheckBoxState FControlRigVariableDetailsCustomization::IsAnimationOutputChecked
 
 void FControlRigVariableDetailsCustomization::HandleAnimationOutputCheckStateChanged(ECheckBoxState CheckBoxState, TWeakObjectPtr<UProperty> PropertyBeingCustomized)
 {
-	if (CheckBoxState == ECheckBoxState::Checked)
-	{
-		FBlueprintEditorUtils::SetBlueprintVariableMetaData(BlueprintPtr.Get(), PropertyBeingCustomized->GetFName(), nullptr, AnimationOutputMetadataName, TEXT("true"));
-	}
-	else
+	if (PropertyBeingCustomized.IsValid() && BlueprintPtr.IsValid())
 	{
-		FBlueprintEditorUtils::RemoveBlueprintVariableMetaData(BlueprintPtr.Get(), PropertyBeingCustomized->GetFName(), nullptr, AnimationOutputMetadataName);
+		UControlRigBlueprint* RigBlueprint = Cast<UControlRigBlueprint>(BlueprintPtr.Get());
+		if (RigBlueprint)
+		{
+			if (RigBlueprint->ModelController)
+			{
+				UProperty* PropertyPtrBeingCustomized = PropertyBeingCustomized.Get();
+				RigBlueprint->ModelController->SetParameterType(PropertyPtrBeingCustomized->GetFName(), CheckBoxState == ECheckBoxState::Checked ? EControlRigModelParameterType::Output : EControlRigModelParameterType::Hidden);
+				FBlueprintEditorUtils::ReconstructAllNodes(BlueprintPtr.Get());
+			}
+		}
 	}
-	FBlueprintEditorUtils::ReconstructAllNodes(BlueprintPtr.Get());
 }
 
 ECheckBoxState FControlRigVariableDetailsCustomization::IsAnimationInputChecked(TWeakObjectPtr<UProperty> PropertyBeingCustomized) const
 {
-	if (PropertyBeingCustomized.IsValid())
+	if (PropertyBeingCustomized.IsValid() && BlueprintPtr.IsValid())
 	{
-		UProperty* PropertyPtrBeingCustomized = PropertyBeingCustomized.Get();
-		FString Value;
-		if (PropertyPtrBeingCustomized->HasMetaData(AnimationInputMetadataName) ||
-			FBlueprintEditorUtils::GetBlueprintVariableMetaData(BlueprintPtr.Get(), PropertyPtrBeingCustomized->GetFName(), nullptr, AnimationInputMetadataName, Value))
+		UControlRigBlueprint* RigBlueprint = Cast<UControlRigBlueprint>(BlueprintPtr.Get());
+		if (RigBlueprint)
 		{
-			return ECheckBoxState::Checked;
+			if (RigBlueprint->Model)
+			{
+				UProperty* PropertyPtrBeingCustomized = PropertyBeingCustomized.Get();
+				const FControlRigModelNode* Node = RigBlueprint->Model->FindNode(PropertyPtrBeingCustomized->GetFName());
+				if (Node)
+				{
+					if (Node->ParameterType == EControlRigModelParameterType::Input)
+					{
+						return ECheckBoxState::Checked;
+					}
+				}
+			}
 		}
 	}
 
@@ -150,15 +177,19 @@ ECheckBoxState FControlRigVariableDetailsCustomization::IsAnimationInputChecked(
 
 void FControlRigVariableDetailsCustomization::HandleAnimationInputCheckStateChanged(ECheckBoxState CheckBoxState, TWeakObjectPtr<UProperty> PropertyBeingCustomized)
 {
-	if (CheckBoxState == ECheckBoxState::Checked)
-	{
-		FBlueprintEditorUtils::SetBlueprintVariableMetaData(BlueprintPtr.Get(), PropertyBeingCustomized->GetFName(), nullptr, AnimationInputMetadataName, TEXT("true"));
-	}
-	else
+	if (PropertyBeingCustomized.IsValid() && BlueprintPtr.IsValid())
 	{
-		FBlueprintEditorUtils::RemoveBlueprintVariableMetaData(BlueprintPtr.Get(), PropertyBeingCustomized->GetFName(), nullptr, AnimationInputMetadataName);
+		UControlRigBlueprint* RigBlueprint = Cast<UControlRigBlueprint>(BlueprintPtr.Get());
+		if (RigBlueprint)
+		{
+			if (RigBlueprint->ModelController)
+			{
+				UProperty* PropertyPtrBeingCustomized = PropertyBeingCustomized.Get();
+				RigBlueprint->ModelController->SetParameterType(PropertyPtrBeingCustomized->GetFName(), CheckBoxState == ECheckBoxState::Checked ? EControlRigModelParameterType::Input : EControlRigModelParameterType::Hidden);
+				FBlueprintEditorUtils::ReconstructAllNodes(BlueprintPtr.Get());
+			}
+		}
 	}
-	FBlueprintEditorUtils::ReconstructAllNodes(BlueprintPtr.Get());
 }
 
 #undef LOCTEXT_NAMESPACE
diff --git a/Engine/Plugins/Experimental/ControlRig/Source/ControlRigEditor/Private/EditMode/ControlRigEditMode.cpp b/Engine/Plugins/Experimental/ControlRig/Source/ControlRigEditor/Private/EditMode/ControlRigEditMode.cpp
--- a/Engine/Plugins/Experimental/ControlRig/Source/ControlRigEditor/Private/EditMode/ControlRigEditMode.cpp
+++ b/Engine/Plugins/Experimental/ControlRig/Source/ControlRigEditor/Private/EditMode/ControlRigEditMode.cpp
@@ -33,6 +33,8 @@
 #include "IControlRigObjectBinding.h"
 #include "Kismet2/BlueprintEditorUtils.h"
 #include "Drawing/ControlRigDrawInterface.h"
+#include "ControlRigBlueprint.h"
+#include "ControlRigController.h"
 
 FName FControlRigEditMode::ModeName("EditMode.ControlRig");
 
@@ -60,8 +62,6 @@ FControlRigEditMode::FControlRigEditMode()
 {
 	Settings = NewObject<UControlRigEditModeSettings>(GetTransientPackage(), TEXT("Settings"));
 
-	OnControlsSelectedDelegate.AddRaw(this, &FControlRigEditMode::HandleSelectionChanged);
-
 	CommandBindings = MakeShareable(new FUICommandList);
 	BindCommands();
 
@@ -918,10 +918,27 @@ void FControlRigEditMode::ClearControlSelection()
 		for(FControlUnitProxy& UnitProxy : ControlUnits)
 		{
 			UnitProxy.SetSelected(false);
+
+			if (UControlRig* ControlRig = WeakControlRig.Get())
+			{
+				if (UControlRigBlueprintGeneratedClass* Class = Cast<UControlRigBlueprintGeneratedClass>(ControlRig->GetClass()))
+				{
+					if (UControlRigBlueprint* RigBlueprint = Cast<UControlRigBlueprint>(Class->ClassGeneratedBy))
+					{
+						if (RigBlueprint)
+						{
+							if (RigBlueprint->ModelController)
+							{
+								RigBlueprint->ModelController->DeselectNode(*UnitProxy.PropertyPath.ToString());
+							}
+						}
+					}
+				}
+			}
 		}
 
 		bSelectedBone = true;
-		OnControlsSelectedDelegate.Broadcast(TArray<FString>());
+		HandleSelectionChanged(TArray<FString>());
 	}
 }
 
@@ -938,11 +955,31 @@ void FControlRigEditMode::SetControlSelection(const FString& InControlPropertyPa
 			{
 				UnitProxy.SetSelected(bSelected);
 				SelectedPropertyPaths.Add(UnitProxy.TransformPropertyPathString);
+
+				if (UControlRig* ControlRig = WeakControlRig.Get())
+				{
+					if (UControlRigBlueprintGeneratedClass* Class = Cast<UControlRigBlueprintGeneratedClass>(ControlRig->GetClass()))
+					{
+						if (UControlRigBlueprint* RigBlueprint = Cast<UControlRigBlueprint>(Class->ClassGeneratedBy))
+						{
+							if (RigBlueprint)
+							{
+								if (RigBlueprint->ModelController)
+								{
+									TArray<FName> NodeNames;
+									NodeNames.Add(*UnitProxy.PropertyPath.ToString());
+									RigBlueprint->ModelController->SetSelection(NodeNames);
+								}
+							}
+						}
+					}
+				}
 			}
 		}
 
 		bSelectedBone = true;
-		OnControlsSelectedDelegate.Broadcast(SelectedPropertyPaths);
+
+		HandleSelectionChanged(SelectedPropertyPaths);
 	}
 }
 
@@ -961,13 +998,33 @@ void FControlRigEditMode::SetControlSelection(const TArray<FString>& InControlPr
 				{
 					UnitProxy.SetSelected(bSelected);
 					SelectedPropertyPaths.Add(UnitProxy.TransformPropertyPathString);
+
+					if (UControlRig* ControlRig = WeakControlRig.Get())
+					{
+						if (UControlRigBlueprintGeneratedClass* Class = Cast<UControlRigBlueprintGeneratedClass>(ControlRig->GetClass()))
+						{
+							if (UControlRigBlueprint* RigBlueprint = Cast<UControlRigBlueprint>(Class->ClassGeneratedBy))
+							{
+								if (RigBlueprint)
+								{
+									if (RigBlueprint->ModelController)
+									{
+										TArray<FName> NodeNames;
+										NodeNames.Add(*UnitProxy.PropertyPath.ToString());
+										RigBlueprint->ModelController->SetSelection(NodeNames);
+									}
+								}
+							}
+						}
+					}
+
 					break;
 				}
 			}
 		}
 
 		bSelectedBone = true;
-		OnControlsSelectedDelegate.Broadcast(SelectedPropertyPaths);
+		HandleSelectionChanged(SelectedPropertyPaths);
 	}
 }
 
@@ -1017,6 +1074,23 @@ int32 FControlRigEditMode::GetNumSelectedControls() const
 	return NumSelected;
 }
 
+TArray<FString> FControlRigEditMode::GetSelectedControls() const
+{
+	TArray<FString> Paths;
+	if (UControlRig* ControlRig = WeakControlRig.Get())
+	{
+		for (const FControlUnitProxy& UnitProxy : ControlUnits)
+		{
+			if (UnitProxy.IsSelected())
+			{
+				Paths.Add(UnitProxy.PropertyPath.ToString());
+			}
+		}
+	}
+
+	return Paths;
+}
+
 void FControlRigEditMode::SetControlEnabled(const FString& InControlPropertyPath, bool bEnabled)
 {
 	for(FControlUnitProxy& UnitProxy : ControlUnits)
@@ -1172,12 +1246,6 @@ void FControlRigEditMode::RecalcPivotTransform()
 
 void FControlRigEditMode::HandleSelectionChanged(const TArray<FString>& InSelectedPropertyPaths)
 {
-	if(!bSelecting)
-	{
-		ClearControlSelection();
-		SetControlSelection(InSelectedPropertyPaths, true);
-	}
-
 	if (WeakSequencer.IsValid())
 	{
 		if (InSelectedPropertyPaths.Num() > 0)
@@ -1348,6 +1416,8 @@ void FControlRigEditMode::RefreshControlProxies()
 			}
 		}
 	}
+
+	HandleSelectionChanged(SelectedPropertyPaths);
 }
 
 FControlRigEditMode* FControlRigEditMode::GetEditModeFromWorldContext(UWorld* InWorldContext)
@@ -1431,4 +1501,41 @@ void FControlRigEditMode::SelectBone(const FName& InBone)
 		SelectedBones.Add(InBone);
 	}
 }
+
+void FControlRigEditMode::HandleModelModified(const UControlRigModel* InModel, EControlRigModelNotifType InType, const void* InPayload)
+{
+	switch (InType)
+	{
+		case EControlRigModelNotifType::NodeSelected:
+		{
+			const FControlRigModelNode* Node = (const FControlRigModelNode*)InPayload;
+			if (Node)
+			{
+				SetControlSelection(*Node->Name.ToString(), true);
+			}
+			break;
+		}
+		case EControlRigModelNotifType::NodeDeselected:
+		{
+			const FControlRigModelNode* Node = (const FControlRigModelNode*)InPayload;
+			if (Node)
+			{
+				SetControlSelection(*Node->Name.ToString(), false);
+			}
+			break;
+		}
+		default:
+		{
+			break;
+		}
+	}
+
+	_ModifiedEvent.Broadcast(InModel, InType, InPayload);
+}
+
+UControlRigModel::FModifiedEvent& FControlRigEditMode::OnModified()
+{
+	return _ModifiedEvent;
+}
+
 #undef LOCTEXT_NAMESPACE
diff --git a/Engine/Plugins/Experimental/ControlRig/Source/ControlRigEditor/Private/EditMode/ControlRigEditMode.h b/Engine/Plugins/Experimental/ControlRig/Source/ControlRigEditor/Private/EditMode/ControlRigEditMode.h
--- a/Engine/Plugins/Experimental/ControlRig/Source/ControlRigEditor/Private/EditMode/ControlRigEditMode.h
+++ b/Engine/Plugins/Experimental/ControlRig/Source/ControlRigEditor/Private/EditMode/ControlRigEditMode.h
@@ -7,6 +7,7 @@
 #include "IPersonaEditMode.h"
 #include "ControlRigTrajectoryCache.h"
 #include "ControlUnitProxy.h"
+#include "ControlRigModel.h"
 
 class FEditorViewportClient;
 class FViewport;
@@ -24,7 +25,6 @@ class IMovieScenePlayer;
 struct FRigUnit_Control;
 
 /** Delegate fired when controls are selected */
-DECLARE_MULTICAST_DELEGATE_OneParam(FOnControlsSelected, const TArray<FString>& /*SelectedControlPropertyPaths*/);
 DECLARE_DELEGATE_RetVal_TwoParams(FTransform, FOnGetBoneTransform, const FName& /*BoneName*/, bool /*bLocal*/);
 DECLARE_DELEGATE_TwoParams(FOnSetBoneTransform, const FName& /*BoneName*/, const FTransform& /*Transform*/);
 
@@ -95,6 +95,9 @@ public:
 	/** Get the number of selected controls */
 	int32 GetNumSelectedControls() const;
 
+	/** returns all property paths for all selected controls */
+	TArray<FString> GetSelectedControls() const;
+
 	/** Set a control's enabled state */
 	void SetControlEnabled(const FString& InControlPropertyPath, bool bEnabled);
 
@@ -118,9 +121,6 @@ public:
 	/** Refresh our internal object list (they may have changed) */
 	void RefreshObjects();
 
-	/** Delegate fired when controls are selected */
-	FOnControlsSelected& OnControlsSelected() { return OnControlsSelectedDelegate; }
-
 	/** Refresh our trajectory cache */
 	void RefreshTrajectoryCache();
 
@@ -140,6 +140,10 @@ public:
 	void SelectBone(const FName& InBone);
 	FOnGetBoneTransform& OnGetBoneTransform() { return OnGetBoneTransformDelegate; }
 	FOnSetBoneTransform& OnSetBoneTransform() { return OnSetBoneTransformDelegate; }
+
+	UControlRigModel::FModifiedEvent& OnModified();
+	void HandleModelModified(const UControlRigModel* InModel, EControlRigModelNotifType InType, const void* InPayload);
+
 protected:
 	/** Helper function: set ControlRigs array to the details panel */
 	void SetObjects_Internal();
@@ -199,9 +203,6 @@ protected:
 	/** As we cannot cycle widget mode during tracking, we defer cycling until after a click with this flag */
 	bool bSelectedBone;
 
-	/** Delegate fired when controls are selected */
-	FOnControlsSelected OnControlsSelectedDelegate;
-
 	/** Guard value for selection */
 	bool bSelecting;
 
@@ -226,4 +227,6 @@ protected:
 	bool AreBoneSelected() const;
 
 	bool AreBoneSelectedAndMovable() const;
+
+	UControlRigModel::FModifiedEvent _ModifiedEvent;
 };
diff --git a/Engine/Plugins/Experimental/ControlRig/Source/ControlRigEditor/Private/EditMode/SControlRigEditModeTools.cpp b/Engine/Plugins/Experimental/ControlRig/Source/ControlRigEditor/Private/EditMode/SControlRigEditModeTools.cpp
--- a/Engine/Plugins/Experimental/ControlRig/Source/ControlRigEditor/Private/EditMode/SControlRigEditModeTools.cpp
+++ b/Engine/Plugins/Experimental/ControlRig/Source/ControlRigEditor/Private/EditMode/SControlRigEditModeTools.cpp
@@ -100,7 +100,7 @@ void SControlRigEditModeTools::Construct(const FArguments& InArgs, UWorld* InWor
 	FControlRigEditMode* ControlRigEditMode = static_cast<FControlRigEditMode*>(GLevelEditorModeTools().GetActiveMode(FControlRigEditMode::ModeName));
 	if (ControlRigEditMode)
 	{
-		ControlRigEditMode->OnControlsSelected().AddSP(this, &SControlRigEditModeTools::OnSelectionSetChanged);
+		ControlRigEditMode->OnModified().AddSP(this, &SControlRigEditModeTools::HandleModelModified);
 	}	
 }
 
@@ -255,22 +255,41 @@ void SControlRigEditModeTools::OnManipulatorsPicked(const TArray<FString>& Manip
 	FControlRigEditMode* ControlRigEditMode = static_cast<FControlRigEditMode*>(GLevelEditorModeTools().GetActiveMode(FControlRigEditMode::ModeName));
 	if (ControlRigEditMode)
 	{
-		bPickerChangingSelection = true;
-		ControlRigEditMode->ClearControlSelection();
-		ControlRigEditMode->SetControlSelection(Manipulators, true);
-		bPickerChangingSelection = false;
+		if (!bPickerChangingSelection)
+		{
+			TGuardValue<bool> SelectGuard(bPickerChangingSelection, true);
+			ControlRigEditMode->ClearControlSelection();
+			ControlRigEditMode->SetControlSelection(Manipulators, true);
+		}
 	}
 }
 
-void SControlRigEditModeTools::OnSelectionSetChanged(const TArray<FString>& SelectedManipulators)
+void SControlRigEditModeTools::HandleModelModified(const UControlRigModel* InModel, EControlRigModelNotifType InType, const void* InPayload)
 {
-/*
-	// Don't want to udpate picker selection set if its the picker causing the change
-	if (!bPickerChangingSelection)
+	if (bPickerChangingSelection)
 	{
-		ControlPicker->SetSelectedManipulators(SelectedManipulators);
+		return;
+	}
+
+	TGuardValue<bool> SelectGuard(bPickerChangingSelection, true);
+	switch (InType)
+	{
+		case EControlRigModelNotifType::NodeSelected:
+		case EControlRigModelNotifType::NodeDeselected:
+		{
+			const FControlRigModelNode* Node = (const FControlRigModelNode*)InPayload;
+			if (Node)
+			{
+				// those are not yet implemented yet
+				// ControlPicker->SelectManipulator(Node->Name, InType == EControlRigModelNotifType::NodeSelected);
+			}
+			break;
+		}
+		default:
+		{
+			break;
+		}
 	}
-*/
 }
 
 #undef LOCTEXT_NAMESPACE
diff --git a/Engine/Plugins/Experimental/ControlRig/Source/ControlRigEditor/Private/EditMode/SControlRigEditModeTools.h b/Engine/Plugins/Experimental/ControlRig/Source/ControlRigEditor/Private/EditMode/SControlRigEditModeTools.h
--- a/Engine/Plugins/Experimental/ControlRig/Source/ControlRigEditor/Private/EditMode/SControlRigEditModeTools.h
+++ b/Engine/Plugins/Experimental/ControlRig/Source/ControlRigEditor/Private/EditMode/SControlRigEditModeTools.h
@@ -5,6 +5,7 @@
 #include "CoreMinimal.h"
 #include "Widgets/SCompoundWidget.h"
 #include "IDetailKeyframeHandler.h"
+#include "ControlRigModel.h"
 #include "Widgets/DeclarativeSyntaxSupport.h"
 
 class IDetailsView;
@@ -49,6 +50,6 @@ private:
 
 	/** Called when a manipulator is selected in the picker */
 	void OnManipulatorsPicked(const TArray<FString>& Manipulators);
-	/** Called when edit mode selection set changes */
-	void OnSelectionSetChanged(const TArray<FString>& SelectedManipulators);
+
+	void HandleModelModified(const UControlRigModel* InModel, EControlRigModelNotifType InType, const void* InPayload);
 };
diff --git a/Engine/Plugins/Experimental/ControlRig/Source/ControlRigEditor/Private/Editor/ControlRigEditor.cpp b/Engine/Plugins/Experimental/ControlRig/Source/ControlRigEditor/Private/Editor/ControlRigEditor.cpp
--- a/Engine/Plugins/Experimental/ControlRig/Source/ControlRigEditor/Private/Editor/ControlRigEditor.cpp
+++ b/Engine/Plugins/Experimental/ControlRig/Source/ControlRigEditor/Private/Editor/ControlRigEditor.cpp
@@ -55,6 +55,9 @@
 #include "Graph/NodeSpawners/ControlRigUnitNodeSpawner.h"
 #include "Graph/ControlRigGraphSchema.h"
 #include "ControlRigObjectVersion.h"
+#include "EdGraphUtilities.h"
+#include "HAL/PlatformApplicationMisc.h"
+#include "SNodePanel.h"
 
 #define LOCTEXT_NAMESPACE "ControlRigEditor"
 
@@ -71,8 +74,8 @@ namespace ControlRigEditorTabs
 
 FControlRigEditor::FControlRigEditor()
 	: ControlRig(nullptr)
-	, bSelecting(false)
 	, bControlRigEditorInitialized(false)
+	, bIsSelecting(false)
 {
 }
 
@@ -136,6 +139,8 @@ void FControlRigEditor::InitControlRigEditor(const EToolkitMode::Type Mode, cons
 	TArray<UBlueprint*> ControlRigBlueprints;
 	ControlRigBlueprints.Add(InControlRigBlueprint);
 
+	InControlRigBlueprint->InitializeModel();
+
 	CommonInitialization(ControlRigBlueprints);
 
 	for (UBlueprint* Blueprint : ControlRigBlueprints)
@@ -148,6 +153,8 @@ void FControlRigEditor::InitControlRigEditor(const EToolkitMode::Type Mode, cons
 				continue;
 			}
 
+			RigGraph->Initialize(InControlRigBlueprint);
+
 			if (RigGraph->GetLinkerCustomVersion(FControlRigObjectVersion::GUID) < FControlRigObjectVersion::RemovalOfHierarchyRefPins)
 			{
 				// recompile in case this control rig requires a rebuild
@@ -155,8 +162,11 @@ void FControlRigEditor::InitControlRigEditor(const EToolkitMode::Type Mode, cons
 				Compile();
 			}
 		}
+
 	}
 
+	InControlRigBlueprint->OnModified().AddSP(this, &FControlRigEditor::HandleModelModified);
+
 	BindCommands();
 
 	AddApplicationMode(
@@ -174,9 +184,10 @@ void FControlRigEditor::InitControlRigEditor(const EToolkitMode::Type Mode, cons
 //	GetAssetEditorModeManager()->SetToolkitHost(GetToolkitHost());
 	GetAssetEditorModeManager()->SetDefaultMode(FControlRigEditorEditMode::ModeName);
 	GetAssetEditorModeManager()->ActivateMode(FControlRigEditorEditMode::ModeName);
-	GetEditMode().OnControlsSelected().AddSP(this, &FControlRigEditor::SetSelectedNodes);
 	GetEditMode().OnGetBoneTransform() = FOnGetBoneTransform::CreateSP(this, &FControlRigEditor::GetBoneTransform);
 	GetEditMode().OnSetBoneTransform() = FOnSetBoneTransform::CreateSP(this, &FControlRigEditor::SetBoneTransform);
+	InControlRigBlueprint->OnModified().AddSP(&GetEditMode(), &FControlRigEditMode::HandleModelModified);
+
 	UpdateControlRig();
 
 	// Post-layout initialization
@@ -388,6 +399,179 @@ FLinearColor FControlRigEditor::GetWorldCentricTabColorScale() const
 	return FLinearColor( 0.5f, 0.25f, 0.35f, 0.5f );
 }
 
+void FControlRigEditor::DeleteSelectedNodes()
+{
+	UControlRigBlueprint* RigBlueprint = Cast<UControlRigBlueprint>(GetBlueprintObj());
+	if (RigBlueprint == nullptr)
+	{
+		return;
+	}
+
+	const FGraphPanelSelectionSet SelectedNodes = GetSelectedNodes();
+	SetUISelectionState(NAME_None);
+
+	for (FGraphPanelSelectionSet::TConstIterator NodeIt(SelectedNodes); NodeIt; ++NodeIt)
+	{
+		if (UEdGraphNode* Node = Cast<UEdGraphNode>(*NodeIt))
+		{
+			if (Node->CanUserDeleteNode())
+			{
+				UControlRigGraphNode* RigNode = Cast<UControlRigGraphNode>(Node);
+				AnalyticsTrackNodeEvent(GetBlueprintObj(), Node, true);
+				RigBlueprint->ModelController->RemoveNode(RigNode->PropertyName);
+			}
+		}
+	}
+}
+
+void FControlRigEditor::PasteNodesHere(class UEdGraph* DestinationGraph, const FVector2D& GraphLocation)
+{
+	UControlRigBlueprint* RigBlueprint = Cast<UControlRigBlueprint>(GetBlueprintObj());
+	if (RigBlueprint == nullptr)
+	{
+		return;
+	}
+
+	UControlRigGraph* TempGraph = NewObject<UControlRigGraph>(GetTransientPackage());
+	TempGraph->bIsTemporaryGraphForCopyPaste = true;
+	TempGraph->Schema = DestinationGraph->Schema;
+
+	// Grab the text to paste from the clipboard.
+	FString TextToImport;
+	FPlatformApplicationMisc::ClipboardPaste(TextToImport);
+
+	// Import the nodes
+	TSet<UEdGraphNode*> PastedNodes;
+	FEdGraphUtilities::ImportNodesFromText(TempGraph, TextToImport, PastedNodes);
+	if (PastedNodes.Num() == 0)
+	{
+		return;
+	}
+
+	FVector2D AvgNodePosition(0.0f, 0.0f);
+	for (TSet<UEdGraphNode*>::TIterator It(PastedNodes); It; ++It)
+	{
+		UEdGraphNode* Node = *It;
+		AvgNodePosition.X += Node->NodePosX;
+		AvgNodePosition.Y += Node->NodePosY;
+	}
+
+	float InvNumNodes = 1.0f / float(PastedNodes.Num());
+	AvgNodePosition.X *= InvNumNodes;
+	AvgNodePosition.Y *= InvNumNodes;
+
+	RigBlueprint->ModelController->OpenUndoBracket(TEXT("Pasted Nodes."));
+
+	TMap<FString, FString> NameMap;
+	for (TSet<UEdGraphNode*>::TIterator It(PastedNodes); It; ++It)
+	{
+		UEdGraphNode* Node = *It;
+		Node->NodePosX = (Node->NodePosX - AvgNodePosition.X) + GraphLocation.X;
+		Node->NodePosY = (Node->NodePosY - AvgNodePosition.Y) + GraphLocation.Y;
+		Node->SnapToGrid(SNodePanel::GetSnapGridSize());
+
+		UControlRigGraphNode* RigNode = Cast<UControlRigGraphNode>(Node);
+		if (RigNode != nullptr)
+		{
+			FVector2D NodePosition = FVector2D(Node->NodePosX, Node->NodePosY);
+			UScriptStruct* ScriptStruct = RigNode->GetUnitScriptStruct();
+			bool bAddedNode = false;
+			if (ScriptStruct)
+			{
+				bAddedNode = RigBlueprint->ModelController->AddNode(ScriptStruct->GetFName(), NodePosition, RigNode->GetPropertyName());
+			}
+			else
+			{
+				FName DataType = RigNode->PinType.PinCategory;
+				if (UStruct* Struct = Cast<UStruct>(RigNode->PinType.PinSubCategoryObject))
+				{
+					DataType = Struct->GetFName();
+				}
+
+				bAddedNode = RigBlueprint->ModelController->AddParameter(RigNode->GetPropertyName(), DataType, EControlRigModelParameterType::Hidden, NodePosition);
+			}
+
+			if(bAddedNode && RigBlueprint->LastNameFromNotification != NAME_None)
+			{
+				FName AddedNodeName = RigBlueprint->LastNameFromNotification;
+				NameMap.Add(Node->GetName(), AddedNodeName.ToString());
+
+				for (UEdGraphPin* Pin : Node->Pins)
+				{
+					FString Left, Right;
+					RigBlueprint->Model->SplitPinPath(Pin->GetName(), Left, Right);
+					Left = *AddedNodeName.ToString();
+
+					if (Pin->PinType.ContainerType == EPinContainerType::Array)
+					{
+						RigBlueprint->ModelController->SetArrayPinSize(*Left, *Right, Pin->SubPins.Num());
+					}
+
+					if (RigNode->ExpandedPins.Contains(Pin->GetName()))
+					{
+						RigBlueprint->ModelController->ExpandPin(*Left, *Right, Pin->Direction == EGPD_Input, true);
+					}
+
+					if (Pin->Direction == EGPD_Input)
+					{
+						if (!Pin->DefaultValue.IsEmpty())
+						{
+							RigBlueprint->ModelController->SetPinDefaultValue(*Left, *Right, Pin->DefaultValue, false);
+						}
+					}
+				}
+			}
+		}
+	}
+
+	for (TSet<UEdGraphNode*>::TIterator It(PastedNodes); It; ++It)
+	{
+		UEdGraphNode* Node = *It;
+		UControlRigGraphNode* RigNode = Cast<UControlRigGraphNode>(Node);
+		if (RigNode != nullptr)
+		{
+			for (UEdGraphPin* Pin : RigNode->Pins)
+			{
+				if (Pin->Direction == EGPD_Input)
+				{
+					continue;
+				}
+
+				for (UEdGraphPin* OtherPin : Pin->LinkedTo)
+				{
+					UControlRigGraphNode* OtherRigNode = Cast<UControlRigGraphNode>(OtherPin->GetOwningNode());
+					if (OtherRigNode)
+					{
+						FString* RemappedNodeName = NameMap.Find(RigNode->GetName());
+						FString* OtherRemappedNodeName = NameMap.Find(OtherRigNode->GetName());
+
+						if (RemappedNodeName != nullptr && OtherRemappedNodeName != nullptr)
+						{
+							FString PinPath = Pin->GetName();
+							if (PinPath.StartsWith(RigNode->GetPropertyName().ToString()))
+							{
+								PinPath = *RemappedNodeName + PinPath.RightChop(RigNode->GetPropertyName().ToString().Len());
+							}
+							FString OtherPinPath = OtherPin->GetName();
+							if (OtherPinPath.StartsWith(OtherRigNode->GetPropertyName().ToString()))
+							{
+								OtherPinPath = *OtherRemappedNodeName + OtherPinPath.RightChop(OtherRigNode->GetPropertyName().ToString().Len());
+							}
+
+							FString SourceLeft, SourceRight, TargetLeft, TargetRight;
+							RigBlueprint->Model->SplitPinPath(PinPath, SourceLeft, SourceRight);
+							RigBlueprint->Model->SplitPinPath(OtherPinPath, TargetLeft, TargetRight);
+							RigBlueprint->ModelController->MakeLink(*SourceLeft, *SourceRight, *TargetLeft, *TargetRight);
+						}
+					}
+				}
+			}
+		}
+	}
+
+	RigBlueprint->ModelController->CloseUndoBracket();
+}
+
 void FControlRigEditor::OnToolkitHostingStarted(const TSharedRef<class IToolkit>& Toolkit)
 {
 	TSharedPtr<SWidget> InlineContent = Toolkit->GetInlineContent();
@@ -417,17 +601,15 @@ void FControlRigEditor::OnActiveTabChanged( TSharedPtr<SDockTab> PreviouslyActiv
 
 void FControlRigEditor::PostUndo(bool bSuccess)
 {
-	DocumentManager->CleanInvalidTabs();
-	DocumentManager->RefreshAllTabs();
-
-	OnHierarchyChanged();
+	//DocumentManager->CleanInvalidTabs();
+	//DocumentManager->RefreshAllTabs();
 
 	FBlueprintEditor::PostUndo(bSuccess);
 }
 
 void FControlRigEditor::PostRedo(bool bSuccess)
 {
-	DocumentManager->RefreshAllTabs();
+	//DocumentManager->RefreshAllTabs();
 
 	FBlueprintEditor::PostRedo(bSuccess);
 }
@@ -464,6 +646,59 @@ void FControlRigEditor::NotifyPostChange(const FPropertyChangedEvent& PropertyCh
 	FBlueprintEditor::NotifyPostChange(PropertyChangedEvent, PropertyThatChanged);
 }
 
+void FControlRigEditor::HandleModelModified(const UControlRigModel* InModel, EControlRigModelNotifType InType, const void* InPayload)
+{
+	switch (InType)
+	{
+		case EControlRigModelNotifType::NodeSelected:
+		case EControlRigModelNotifType::NodeDeselected:
+		{
+			const FControlRigModelNode* Node = (const FControlRigModelNode*)InPayload;
+			if (Node != nullptr)
+			{
+				if (!bIsSelecting)
+				{
+					TGuardValue<bool> SelectingGuard(bIsSelecting, true);
+					if (FocusedGraphEdPtr.IsValid())
+					{
+						TSharedPtr<SGraphEditor> FocusedGraphEd = FocusedGraphEdPtr.Pin();
+						UControlRigGraph* RigGraph = Cast<UControlRigGraph>(FocusedGraphEd->GetCurrentGraph());
+						if (RigGraph != nullptr)
+						{
+							UControlRigGraphNode* RigNode = RigGraph->FindNodeFromPropertyName(Node->Name);
+							if (RigNode != nullptr)
+							{
+								FocusedGraphEd->SetNodeSelection(RigNode, InType == EControlRigModelNotifType::NodeSelected);
+							}
+						}
+					}
+					break;
+				}
+
+				if (InType == EControlRigModelNotifType::NodeSelected)
+				{
+					UClass* Class = GetBlueprintObj()->GeneratedClass.Get();
+					if (Class)
+					{
+						UProperty* Property = Class->FindPropertyByName(Node->Name);
+						if (Property)
+						{
+							TSet<class UObject*> SelectedObjects;
+							SelectedObjects.Add(Property);
+							FBlueprintEditor::OnSelectedNodesChangedImpl(SelectedObjects);
+						}
+					}
+				}
+			}
+			break;
+		}
+		default:
+		{
+			break;
+		}
+	}
+}
+
 void FControlRigEditor::Tick(float DeltaTime)
 {
 	FBlueprintEditor::Tick(DeltaTime);
@@ -488,101 +723,34 @@ TStatId FControlRigEditor::GetStatId() const
 
 void FControlRigEditor::OnSelectedNodesChangedImpl(const TSet<class UObject*>& NewSelection)
 {
-	if(!bSelecting)
+	if (bIsSelecting)
 	{
-		TGuardValue<bool> GuardValue(bSelecting, true);
-		// Substitute any control rig nodes for their properties, so we display details for them instead
-		TSet<class UObject*> SelectedObjects;
-		TArray<FString> PropertyPathStrings;
-		for(UObject* Object : NewSelection)
-		{
-			UClass* ClassUsed = nullptr;
-			UClass* Class = GetBlueprintObj()->GeneratedClass.Get();
-			UClass* SkeletonClass = GetBlueprintObj()->SkeletonGeneratedClass.Get();
-			UControlRigGraphNode* ControlRigGraphNode = Cast<UControlRigGraphNode>(Object);
-			if(ControlRigGraphNode)
-			{
-				UProperty* Property = nullptr;
-
-				if(Class && ControlRigGraphNode)
-				{
-					Property = Class->FindPropertyByName(ControlRigGraphNode->GetPropertyName());
-					ClassUsed = Class;
-				}
-
-				if(Property == nullptr)
-				{
-					if(SkeletonClass && ControlRigGraphNode)
-					{
-						Property = SkeletonClass->FindPropertyByName(ControlRigGraphNode->GetPropertyName());
-						ClassUsed = SkeletonClass;
-					}
-				}
-
-				if(Property)
-				{
-					SelectedObjects.Add(Property);
+		return;
+	}
 
-					check(ClassUsed);
+	TGuardValue<bool> SelectingGuard(bIsSelecting, true);
 
-					// @TODO: if we ever want to support sub-graphs, we will need a full property path here
-					PropertyPathStrings.Add(Property->GetName());
-				}
-			}
-			else
-			{
-				SelectedObjects.Add(Object);
-			}
+	TArray<UControlRigGraphNode*> SelectedNodes;
+	for(UObject* Object : NewSelection)
+	{
+		UControlRigGraphNode* ControlRigGraphNode = Cast<UControlRigGraphNode>(Object);
+		if(ControlRigGraphNode)
+		{
+			SelectedNodes.Add(ControlRigGraphNode);
 		}
-
-		OnGraphNodeSelectionChangedDelegate.Broadcast(NewSelection);
-
-		// Let the edit mode know about selection
-		FControlRigEditMode& EditMode = GetEditMode();
-		EditMode.ClearControlSelection();
-		EditMode.SetControlSelection(PropertyPathStrings, true);
-
-		FBlueprintEditor::OnSelectedNodesChangedImpl(SelectedObjects);
 	}
-}
 
-void FControlRigEditor::SetSelectedNodes(const TArray<FString>& InSelectedPropertyPaths)
-{
-	if(!bSelecting)
+	UControlRigBlueprint* ControlRigBlueprint = CastChecked<UControlRigBlueprint>(GetBlueprintObj());
+	if (ControlRigBlueprint)
 	{
-		TGuardValue<bool> GuardValue(bSelecting, true);
-
-		UControlRigBlueprint* ControlRigBlueprint = CastChecked<UControlRigBlueprint>(GetBlueprintObj());
-		if(UEdGraph* Graph = GetFocusedGraph())
+		if (ControlRigBlueprint->ModelController)
 		{
-			TSet<const UEdGraphNode*> Nodes;
-			TSet<UObject*> Objects;
-
-			for(UEdGraphNode* GraphNode : Graph->Nodes)
+			TArray<FName> NodeNamesToSelect;
+			for (UControlRigGraphNode* SelectedNode : SelectedNodes)
 			{
-				if(UControlRigGraphNode* ControlRigGraphNode = Cast<UControlRigGraphNode>(GraphNode))
-				{
-					for(const FString& SelectedPropertyPath : InSelectedPropertyPaths)
-					{
-						if(ControlRigGraphNode->GetPropertyName().ToString() == SelectedPropertyPath)
-						{
-							Nodes.Add(GraphNode);
-							Objects.Add(GraphNode);
-							break;
-						}
-					}
-				}
+				NodeNamesToSelect.Add(SelectedNode->GetPropertyName());
 			}
-
-			FocusedGraphEdPtr.Pin()->ClearSelectionSet();
-			Graph->SelectNodeSet(Nodes);
-
-			OnGraphNodeSelectionChangedDelegate.Broadcast(Objects);
-
-			// Let the edit mode know about selection
-			FControlRigEditMode& EditMode = GetEditMode();
-			EditMode.ClearControlSelection();
-			EditMode.SetControlSelection(InSelectedPropertyPaths, true);
+			ControlRigBlueprint->ModelController->SetSelection(NodeNamesToSelect);
 		}
 	}
 }
@@ -630,6 +798,30 @@ void FControlRigEditor::OnBlueprintChangedImpl(UBlueprint* InBlueprint, bool bIs
 		}
 
 		OnSelectedNodesChangedImpl(GetSelectedNodes());
+
+		UClass* Class = GetBlueprintObj()->GeneratedClass.Get();
+		if (Class)
+		{
+			TSet<class UObject*> SelectedObjects;
+			FGraphPanelSelectionSet SelectedNodes = GetSelectedNodes();
+			for (UObject* SelecteNode : SelectedNodes)
+			{
+				UControlRigGraphNode* RigNode = Cast<UControlRigGraphNode>(SelecteNode);
+				if (RigNode == nullptr)
+				{
+					continue;
+				}
+				UProperty* Property = Class->FindPropertyByName(RigNode->GetPropertyName());
+				if (Property)
+				{
+					SelectedObjects.Add(Property);
+				}
+			}
+			if (SelectedObjects.Num() > 0)
+			{
+				FBlueprintEditor::OnSelectedNodesChangedImpl(SelectedObjects);
+			}
+		}
 	}
 }
 
@@ -1260,43 +1452,46 @@ void FControlRigEditor::HandleMakeBoneGetterSetter(EBoneGetterSetterType Type, b
 		return;
 	}
 
-	UControlRigUnitNodeSpawner* Spawner = NewObject<UControlRigUnitNodeSpawner>(GetTransientPackage());
-	Spawner->StructTemplate = StructTemplate;
-	Spawner->NodeClass = UControlRigGraphNode::StaticClass();
-	IBlueprintNodeBinder::FBindingSet Bindings;
+	UControlRigBlueprint* Blueprint = Cast<UControlRigBlueprint>(GetBlueprintObj());
+	if (Blueprint == nullptr)
+	{
+		return;
+	}
+	if (Blueprint->ModelController == nullptr)
+	{
+		return;
+	}
 
-	const FScopedTransaction Transaction(LOCTEXT("DroppedHierarchyItems", "Add Rig Units from Drag & Drop"));
+	Blueprint->ModelController->OpenUndoBracket(TEXT("Adding Nodes from Hierarchy"));
 
-	TSet<const UEdGraphNode*> NewNodes;
+	TArray<FName> NewNodeNames;
 	for (const FName& BoneName : BoneNames)
 	{
-		const TCHAR* BoneNameSuffix = TEXT(".Bone");
 		FVector2D NodePositionIncrement(0.f, 120.f);
 		if(!bIsGetter)
 		{
 			NodePositionIncrement = FVector2D(380.f, 0.f);
 		}
 
-		UControlRigGraphNode* Node = Cast<UControlRigGraphNode>(Spawner->Invoke(Graph, Bindings, NodePosition));
-		if (Node != nullptr)
+		FName Name = FControlRigBlueprintUtils::ValidateName(Blueprint, StructTemplate->GetName());
+		if (Blueprint->ModelController->AddNode(StructTemplate->GetFName(), NodePosition, Name))
 		{
-			NewNodes.Add(Node);
-
-			for (UEdGraphPin* Pin : Node->Pins)
+			const FControlRigModelNode* Node = Blueprint->Model->FindNode(Blueprint->LastNameFromNotification);
+			if (Node)
 			{
-				if (Pin->GetName().EndsWith(BoneNameSuffix))
-				{
-					Pin->DefaultValue = BoneName.ToString();
-				}
+				NewNodeNames.Add(Node->Name);
+				Blueprint->ModelController->SetPinDefaultValueName(Node->Name, TEXT("Bone"), BoneName, true);
 			}
 		}
 
 		NodePosition += NodePositionIncrement;
 	}
 
-	if (NewNodes.Num() > 0)
+	Blueprint->ModelController->CloseUndoBracket();
+
+	if (NewNodeNames.Num() > 0)
 	{
-		Graph->SelectNodeSet(NewNodes);
+		Blueprint->ModelController->SetSelection(NewNodeNames);
 	}
 }
 
@@ -1421,6 +1616,7 @@ void FControlRigEditor::UpdateGraphCompilerErrors()
 
 void FControlRigEditor::DumpUnitTestCode()
 {
+	/*
 	if (UEdGraph* Graph = GetFocusedGraph())
 	{
 		TArray<FString> Code;
@@ -1525,6 +1721,7 @@ void FControlRigEditor::DumpUnitTestCode()
 
 		UE_LOG(LogControlRigEditor, Display, TEXT("\n%s\n"), *FString::Join(Code, TEXT("\n")));
 	}
+	*/
 }
 
 #undef LOCTEXT_NAMESPACE
diff --git a/Engine/Plugins/Experimental/ControlRig/Source/ControlRigEditor/Private/Editor/ControlRigEditor.h b/Engine/Plugins/Experimental/ControlRig/Source/ControlRigEditor/Private/Editor/ControlRigEditor.h
--- a/Engine/Plugins/Experimental/ControlRig/Source/ControlRigEditor/Private/Editor/ControlRigEditor.h
+++ b/Engine/Plugins/Experimental/ControlRig/Source/ControlRigEditor/Private/Editor/ControlRigEditor.h
@@ -9,6 +9,7 @@
 #include "ControlRigDefines.h"
 #include "ControlRigLog.h"
 #include "Drawing/ControlRigDrawInterface.h"
+#include "ControlRigModel.h"
 
 class UControlRigBlueprint;
 class IPersonaToolkit;
@@ -17,9 +18,6 @@ class SBorder;
 class USkeletalMesh;
 class FStructOnScope;
 
-/** Delegate fired when node selection has changed */
-DECLARE_MULTICAST_DELEGATE_OneParam(FOnGraphNodeSelectionChanged, const TSet<UObject*>& /*Selection*/);
-
 struct FControlRigEditorModes
 {
 	// Mode constants
@@ -75,7 +73,9 @@ public:
 
 	// BlueprintEditor interface
 	virtual bool CanAddNewLocalVariable() const override { return false; }
-	
+	virtual void DeleteSelectedNodes();
+	virtual void PasteNodesHere(class UEdGraph* DestinationGraph, const FVector2D& GraphLocation) override;
+
 	// IToolkitHost Interface
 	virtual void OnToolkitHostingStarted(const TSharedRef<class IToolkit>& Toolkit) override;
 	virtual void OnToolkitHostingFinished(const TSharedRef<class IToolkit>& Toolkit) override;
@@ -95,8 +95,6 @@ public:
 
 	void ClearDetailObject();
 
-	FOnGraphNodeSelectionChanged& OnGraphNodeSelectionChanged() { return OnGraphNodeSelectionChangedDelegate; }
-
 	/** Get the persona toolkit */
 	TSharedRef<IPersonaToolkit> GetPersonaToolkit() const { return PersonaToolkit.ToSharedRef(); }
 
@@ -106,8 +104,6 @@ public:
 	/** Get the edit mode */
 	FControlRigEditorEditMode& GetEditMode() { return *static_cast<FControlRigEditorEditMode*>(GetAssetEditorModeManager()->GetActiveMode(FControlRigEditorEditMode::ModeName)); }
 
-	/** Try to set the selected nodes form some external source */
-	void SetSelectedNodes(const TArray<FString>& InSelectedPropertyPaths);
 	void SelectBone(const FName& InBone);
 	// this changes everytime you compile, so don't cache it expecting it will last. 
 	UControlRig* GetInstanceRig() const { return ControlRig;  }
@@ -139,6 +135,8 @@ protected:
 	// FNotifyHook Interface
 	virtual void NotifyPostChange(const FPropertyChangedEvent& PropertyChangedEvent, UProperty* PropertyThatChanged) override;
 
+	void HandleModelModified(const UControlRigModel* InModel, EControlRigModelNotifType InType, const void* InPayload);
+
 	// FGCObject Interface
 	virtual void AddReferencedObjects( FReferenceCollector& Collector ) override;
 
@@ -214,17 +212,12 @@ protected:
 	/** Preview instance inspector widget */
 	TSharedPtr<SWidget> PreviewEditor;
 
-	FOnGraphNodeSelectionChanged OnGraphNodeSelectionChangedDelegate;
-
 	/** Our currently running control rig instance */
 	UControlRig* ControlRig;
 
 	/** preview scene */
 	TSharedPtr<IPersonaPreviewScene> PreviewScene;
 
-	/** Recursion guard for selection */
-	bool bSelecting;
-
 	/** Bone Selection related */
 	FTransform GetBoneTransform(const FName& InBone, bool bLocal) const;
 	void SetBoneTransform(const FName& InBone, const FTransform& InTransform);
@@ -236,6 +229,7 @@ protected:
 	FName SelectedBone;
 
 	bool bControlRigEditorInitialized;
+	bool bIsSelecting;
 
 	/** The log to use for errors resulting from the init phase of the units */
 	FControlRigLog ControlRigLog;
diff --git a/Engine/Plugins/Experimental/ControlRig/Source/ControlRigEditor/Private/Editor/SControlRigStackView.cpp b/Engine/Plugins/Experimental/ControlRig/Source/ControlRigEditor/Private/Editor/SControlRigStackView.cpp
--- a/Engine/Plugins/Experimental/ControlRig/Source/ControlRigEditor/Private/Editor/SControlRigStackView.cpp
+++ b/Engine/Plugins/Experimental/ControlRig/Source/ControlRigEditor/Private/Editor/SControlRigStackView.cpp
@@ -124,9 +124,9 @@ SControlRigStackView::~SControlRigStackView()
 {
 	if (ControlRigEditor.IsValid())
 	{
-		if (OnGraphNodeSelectionChanged.IsValid())
+		if (OnModelModified.IsValid() && ControlRigEditor.Pin()->GetControlRigBlueprint())
 		{
-			ControlRigEditor.Pin()->OnGraphNodeSelectionChanged().Remove(OnGraphNodeSelectionChanged);
+			ControlRigEditor.Pin()->GetControlRigBlueprint()->OnModified().Remove(OnModelModified);
 		}
 		if (OnControlRigInitializedHandle.IsValid() && ControlRigEditor.Pin()->ControlRig)
 		{
@@ -145,9 +145,7 @@ void SControlRigStackView::Construct( const FArguments& InArgs, TSharedRef<FCont
 	ControlRigBlueprint = ControlRigEditor.Pin()->GetControlRigBlueprint();
 	Graph = Cast<UControlRigGraph>(ControlRigBlueprint->GetLastEditedUberGraph());
 	CommandList = MakeShared<FUICommandList>();
-	bSelecting = false;
-
-	OnGraphNodeSelectionChanged = InControlRigEditor->OnGraphNodeSelectionChanged().AddSP(this, &SControlRigStackView::HandleGraphSelectionChanged);
+	bSuspendModelNotifications = false;
 
 	BindCommands();
 
@@ -182,61 +180,17 @@ void SControlRigStackView::Construct( const FArguments& InArgs, TSharedRef<FCont
 			ControlRigBlueprint->OnCompiled().Remove(OnBlueprintCompiledHandle);
 		}
 		OnBlueprintCompiledHandle = ControlRigBlueprint->OnCompiled().AddSP(this, &SControlRigStackView::OnBlueprintCompiled);
-	}
-}
-
-void SControlRigStackView::HandleGraphSelectionChanged(const TSet<UObject*>& SelectedNodes)
-{
-	if (bSelecting)
-	{
-		return;
-	}
-
-	{
-		TGuardValue<bool> SelectingGuard(bSelecting, true);
-		TreeView->ClearSelection();
-	}
-
-	if (SelectedNodes.Num() > 0)
-	{
-		if (!ControlRigBlueprint.IsValid())
-		{
-			return;
-		}
-		UControlRigBlueprintGeneratedClass* GeneratedClass = ControlRigBlueprint->GetControlRigBlueprintGeneratedClass();
-		if (GeneratedClass == nullptr)
-		{
-			return;
-		}
-
-		TMap<FName, const UControlRigGraphNode*> SelectedNodeNameMap;
-		for (const UObject* SelectedNode : SelectedNodes)
-		{
-			const UControlRigGraphNode* RigNode = Cast<const UControlRigGraphNode>(SelectedNode);
-			if (RigNode)
-			{
-				SelectedNodeNameMap.Add(RigNode->GetPropertyName(), RigNode);
-			}
-		}
-
-
-		TGuardValue<bool> SelectingGuard(bSelecting, true);
-		for (const TSharedPtr<FRigStackEntry>& Entry : Operators)
-		{
-			if (SelectedNodeNameMap.Contains(Entry->Name))
-			{
-				TreeView->SetItemSelection(Entry, true, ESelectInfo::Direct);
-			}
-		}
+		OnModelModified = ControlRigBlueprint->OnModified().AddSP(this, &SControlRigStackView::HandleModelModified);
 	}
 }
 
 void SControlRigStackView::OnSelectionChanged(TSharedPtr<FRigStackEntry> Selection, ESelectInfo::Type SelectInfo)
 {
-	if (bSelecting)
+	if (bSuspendModelNotifications)
 	{
 		return;
 	}
+	TGuardValue<bool> SuspendNotifs(bSuspendModelNotifications, true);
 
 	TArray<TSharedPtr<FRigStackEntry>> SelectedItems = TreeView->GetSelectedItems();
 	if (SelectedItems.Num() > 0)
@@ -245,13 +199,14 @@ void SControlRigStackView::OnSelectionChanged(TSharedPtr<FRigStackEntry> Selecti
 		{
 			return;
 		}
+
 		UControlRigBlueprintGeneratedClass* GeneratedClass = ControlRigBlueprint->GetControlRigBlueprintGeneratedClass();
 		if (GeneratedClass == nullptr)
 		{
 			return;
 		}
 
-		TArray<FString> SelectedNodes;
+		TArray<FName> SelectedNodes;
 		for (TSharedPtr<FRigStackEntry>& Entry : SelectedItems)
 		{
 			if (Entry->OpIndex >= GeneratedClass->Operators.Num())
@@ -263,12 +218,11 @@ void SControlRigStackView::OnSelectionChanged(TSharedPtr<FRigStackEntry> Selecti
 			if (Operator.OpCode == EControlRigOpCode::Exec)
 			{
 				FString PropertyPath = Operator.CachedPropertyPath1.ToString();
-				SelectedNodes.Add(PropertyPath);
+				SelectedNodes.Add(*PropertyPath);
 			}
 		}
 
-		TGuardValue<bool> SelectingGuard(bSelecting, true);
-		ControlRigEditor.Pin()->SetSelectedNodes(SelectedNodes);
+		ControlRigBlueprint->ModelController->SetSelection(SelectedNodes);
 	}
 }
 
@@ -461,4 +415,49 @@ void SControlRigStackView::OnControlRigInitialized(UControlRig* ControlRig, ECon
 	}
 }
 
+void SControlRigStackView::HandleModelModified(const UControlRigModel* InModel, EControlRigModelNotifType InType, const void* InPayload)
+{
+	if (bSuspendModelNotifications)
+	{
+		return;
+	}
+
+	switch (InType)
+	{
+		case EControlRigModelNotifType::NodeSelected:
+		case EControlRigModelNotifType::NodeDeselected:
+		{
+			const FControlRigModelNode* Node = (const FControlRigModelNode*)InPayload;
+			if (Node != nullptr)
+			{
+				if (!ControlRigBlueprint.IsValid())
+				{
+					return;
+				}
+
+				UControlRigBlueprintGeneratedClass* GeneratedClass = ControlRigBlueprint->GetControlRigBlueprintGeneratedClass();
+				if (GeneratedClass == nullptr)
+				{
+					return;
+				}
+
+				TGuardValue<bool> SuspendNotifs(bSuspendModelNotifications, true);
+				for (const TSharedPtr<FRigStackEntry>& Entry : Operators)
+				{
+					if (Node->Name == Entry->Name)
+					{
+						TreeView->SetItemSelection(Entry, InType == EControlRigModelNotifType::NodeSelected, ESelectInfo::Direct);
+					}
+				}
+			}
+			break;
+		}
+		default:
+		{
+			// todo... other cases
+			break;
+		}
+	}
+}
+
 #undef LOCTEXT_NAMESPACE
diff --git a/Engine/Plugins/Experimental/ControlRig/Source/ControlRigEditor/Private/Editor/SControlRigStackView.h b/Engine/Plugins/Experimental/ControlRig/Source/ControlRigEditor/Private/Editor/SControlRigStackView.h
--- a/Engine/Plugins/Experimental/ControlRig/Source/ControlRigEditor/Private/Editor/SControlRigStackView.h
+++ b/Engine/Plugins/Experimental/ControlRig/Source/ControlRigEditor/Private/Editor/SControlRigStackView.h
@@ -65,9 +65,6 @@ public:
 
 	void Construct(const FArguments& InArgs, TSharedRef<FControlRigEditor> InControlRigEditor);
 
-	/** Sync up selection with the graph */
-	void HandleGraphSelectionChanged(const TSet<UObject*>& SelectedJoints);
-
 	/** Set Selection Changed */
 	void OnSelectionChanged(TSharedPtr<FRigStackEntry> Selection, ESelectInfo::Type SelectInfo);
 
@@ -96,6 +93,9 @@ private:
 
 	void OnControlRigInitialized(UControlRig* ControlRig, EControlRigState State);
 
+	bool bSuspendModelNotifications;
+	void HandleModelModified(const UControlRigModel* InModel, EControlRigModelNotifType InType, const void* InPayload);
+
 	TSharedPtr<STreeView<TSharedPtr<FRigStackEntry>>> TreeView;
 
 	/** Command list we bind to */
@@ -107,9 +107,7 @@ private:
 
 	TArray<TSharedPtr<FRigStackEntry>> Operators;
 
-	bool bSelecting;
-
-	FDelegateHandle OnGraphNodeSelectionChanged;
+	FDelegateHandle OnModelModified;
 	FDelegateHandle OnBlueprintCompiledHandle;
 	FDelegateHandle OnControlRigInitializedHandle;
 };
diff --git a/Engine/Plugins/Experimental/ControlRig/Source/ControlRigEditor/Private/Editor/SRigHierarchy.cpp b/Engine/Plugins/Experimental/ControlRig/Source/ControlRigEditor/Private/Editor/SRigHierarchy.cpp
--- a/Engine/Plugins/Experimental/ControlRig/Source/ControlRigEditor/Private/Editor/SRigHierarchy.cpp
+++ b/Engine/Plugins/Experimental/ControlRig/Source/ControlRigEditor/Private/Editor/SRigHierarchy.cpp
@@ -174,8 +174,6 @@ void SRigHierarchy::Construct(const FArguments& InArgs, TSharedRef<FControlRigEd
 	// for deleting, renaming, dragging
 	CommandList = MakeShared<FUICommandList>();
 
-	InControlRigEditor->OnGraphNodeSelectionChanged().AddSP(this, &SRigHierarchy::HandleGraphSelectionChanged);
-
 	UEditorEngine* Editor = Cast<UEditorEngine>(GEngine);
 	if (Editor != nullptr)
 	{
@@ -338,11 +336,6 @@ void SRigHierarchy::HandleGetChildrenForTree(TSharedPtr<FRigTreeBone> InItem, TA
 	OutChildren = InItem.Get()->Children;
 }
 
-void SRigHierarchy::HandleGraphSelectionChanged(const TSet<UObject*>& SelectedBones)
-{
-
-}
-
 void SRigHierarchy::OnSelectionChanged(TSharedPtr<FRigTreeBone> Selection, ESelectInfo::Type SelectInfo)
 {
 	// need dummy object
diff --git a/Engine/Plugins/Experimental/ControlRig/Source/ControlRigEditor/Private/Editor/SRigHierarchy.h b/Engine/Plugins/Experimental/ControlRig/Source/ControlRigEditor/Private/Editor/SRigHierarchy.h
--- a/Engine/Plugins/Experimental/ControlRig/Source/ControlRigEditor/Private/Editor/SRigHierarchy.h
+++ b/Engine/Plugins/Experimental/ControlRig/Source/ControlRigEditor/Private/Editor/SRigHierarchy.h
@@ -137,9 +137,6 @@ private:
 	/** Delete Item */
 	void HandleRenameItem();
 
-	/** Sync up selection with the graph */
-	void HandleGraphSelectionChanged(const TSet<UObject*>& SelectedBones);
-
 	/** Set Selection Changed */
 	void OnSelectionChanged(TSharedPtr<FRigTreeBone> Selection, ESelectInfo::Type SelectInfo);
 
diff --git a/Engine/Plugins/Experimental/ControlRig/Source/ControlRigEditor/Private/Graph/ControlRigConnectionDrawingPolicy.cpp b/Engine/Plugins/Experimental/ControlRig/Source/ControlRigEditor/Private/Graph/ControlRigConnectionDrawingPolicy.cpp
--- a/Engine/Plugins/Experimental/ControlRig/Source/ControlRigEditor/Private/Graph/ControlRigConnectionDrawingPolicy.cpp
+++ b/Engine/Plugins/Experimental/ControlRig/Source/ControlRigEditor/Private/Graph/ControlRigConnectionDrawingPolicy.cpp
@@ -1,6 +1,48 @@
 // Copyright 1998-2019 Epic Games, Inc. All Rights Reserved.
 
 #include "ControlRigConnectionDrawingPolicy.h"
+#include "Graph/ControlRigGraph.h"
+#include "Graph/ControlRigGraphNode.h"
+#include "ControlRigBlueprint.h"
+#include "ControlRigController.h"
+#include "Kismet2/BlueprintEditorUtils.h"
+
+void FControlRigConnectionDrawingPolicy::SetIncompatiblePinDrawState(const TSharedPtr<SGraphPin>& StartPin, const TSet< TSharedRef<SWidget> >& VisiblePins)
+{
+	UEdGraphPin* Pin = StartPin->GetPinObj();
+	if (Pin != nullptr)
+	{
+		UBlueprint* Blueprint = FBlueprintEditorUtils::FindBlueprintForNodeChecked(Pin->GetOwningNode());
+		UControlRigBlueprint* RigBlueprint = Cast<UControlRigBlueprint>(Blueprint);
+		if (RigBlueprint != nullptr)
+		{
+			FString Left, Right;
+			RigBlueprint->Model->SplitPinPath(Pin->GetName(), Left, Right);
+			RigBlueprint->ModelController->PrepareCycleCheckingForPin(*Left, *Right, Pin->Direction == EGPD_Input);
+		}
+	}
+	FKismetConnectionDrawingPolicy::SetIncompatiblePinDrawState(StartPin, VisiblePins);
+}
+
+void FControlRigConnectionDrawingPolicy::ResetIncompatiblePinDrawState(const TSet< TSharedRef<SWidget> >& VisiblePins)
+{
+	if (VisiblePins.Num() > 0)
+	{
+		TSharedRef<SWidget> WidgetRef = *VisiblePins.begin();
+		const SGraphPin* PinWidget = (const SGraphPin*)&WidgetRef.Get();
+		UEdGraphPin* Pin = PinWidget->GetPinObj();
+		if (Pin != nullptr)
+		{
+			UBlueprint* Blueprint = FBlueprintEditorUtils::FindBlueprintForNodeChecked(Pin->GetOwningNode());
+			UControlRigBlueprint* RigBlueprint = Cast<UControlRigBlueprint>(Blueprint);
+			if (RigBlueprint != nullptr)
+			{
+				RigBlueprint->ModelController->ResetCycleCheck();
+			}
+		}
+	}
+	FKismetConnectionDrawingPolicy::ResetIncompatiblePinDrawState(VisiblePins);
+}
 
 void FControlRigConnectionDrawingPolicy::BuildPinToPinWidgetMap(TMap<TSharedRef<SWidget>, FArrangedWidget>& InPinGeometries)
 {
diff --git a/Engine/Plugins/Experimental/ControlRig/Source/ControlRigEditor/Private/Graph/ControlRigConnectionDrawingPolicy.h b/Engine/Plugins/Experimental/ControlRig/Source/ControlRigEditor/Private/Graph/ControlRigConnectionDrawingPolicy.h
--- a/Engine/Plugins/Experimental/ControlRig/Source/ControlRigEditor/Private/Graph/ControlRigConnectionDrawingPolicy.h
+++ b/Engine/Plugins/Experimental/ControlRig/Source/ControlRigEditor/Private/Graph/ControlRigConnectionDrawingPolicy.h
@@ -12,6 +12,8 @@ public:
 	{
 	}
 
+	virtual void SetIncompatiblePinDrawState(const TSharedPtr<SGraphPin>& StartPin, const TSet< TSharedRef<SWidget> >& VisiblePins) override;
+	virtual void ResetIncompatiblePinDrawState(const TSet< TSharedRef<SWidget> >& VisiblePins) override;
 	virtual void BuildPinToPinWidgetMap(TMap<TSharedRef<SWidget>, FArrangedWidget>& InPinGeometries) override;
 	virtual void DrawPinGeometries(TMap<TSharedRef<SWidget>, FArrangedWidget>& InPinGeometries, FArrangedChildren& ArrangedNodes) override;
 	virtual void DetermineLinkGeometry(
diff --git a/Engine/Plugins/Experimental/ControlRig/Source/ControlRigEditor/Private/Graph/NodeSpawners/ControlRigUnitNodeSpawner.cpp b/Engine/Plugins/Experimental/ControlRig/Source/ControlRigEditor/Private/Graph/NodeSpawners/ControlRigUnitNodeSpawner.cpp
--- a/Engine/Plugins/Experimental/ControlRig/Source/ControlRigEditor/Private/Graph/NodeSpawners/ControlRigUnitNodeSpawner.cpp
+++ b/Engine/Plugins/Experimental/ControlRig/Source/ControlRigEditor/Private/Graph/NodeSpawners/ControlRigUnitNodeSpawner.cpp
@@ -16,6 +16,7 @@
 #include "ScopedTransaction.h"
 #include "ControlRig/Private/Units/Execution/RigUnit_BeginExecution.h"
 #include "ControlRig.h"
+#include "ControlRigBlueprint.h"
 
 #define LOCTEXT_NAMESPACE "ControlRigUnitNodeSpawner"
 
@@ -88,73 +89,15 @@ UEdGraphNode* UControlRigUnitNodeSpawner::Invoke(UEdGraph* ParentGraph, FBinding
 
 	if(StructTemplate)
 	{
-		const FScopedTransaction Transaction(LOCTEXT("AddRigUnitNode", "Add Rig Unit Node"));
+#if WITH_EDITOR
+		if (GEditor)
+		{
+			GEditor->CancelTransaction(0);
+		}
+#endif
 
 		UBlueprint* Blueprint = CastChecked<UBlueprint>(ParentGraph->GetOuter());
 		NewNode = SpawnNode(ParentGraph, Blueprint, StructTemplate, Location);
-
-		bool const bIsTemplateNode = FBlueprintNodeTemplateCache::IsTemplateOuter(ParentGraph);
-		if (NewNode != nullptr && !bIsTemplateNode)
-		{
-			const TArray<TSharedRef<FControlRigField>>& NewExecutionInfos = NewNode->GetExecutionVariableInfo();
-			if(NewExecutionInfos.Num() > 0)
-			{
-				float ClosestDistance = FLT_MAX;
-				UControlRigGraphNode* ClosestRigNode = nullptr;
-				UEdGraphPin* ClosestExecutionPin = nullptr;
-
-				// try to hook up the rig execution pin automatically for the user
-				for (UEdGraphNode* Node : ParentGraph->Nodes)
-				{
-					if (UControlRigGraphNode* RigNode = Cast<UControlRigGraphNode>(Node))
-					{
-						if (RigNode == NewNode)
-						{
-							continue;
-						}
-
-						TArray<TSharedRef<FControlRigField>> CurrentExecutionInfos = RigNode->GetExecutionVariableInfo();
-						CurrentExecutionInfos.Append(RigNode->GetInputOutputVariableInfo());
-						CurrentExecutionInfos.Append(RigNode->GetOutputVariableInfo());
-						for(const TSharedRef<FControlRigField>& ExecutionInfo : CurrentExecutionInfos)
-						{
-							if (ExecutionInfo->OutputPin != nullptr)
-							{
-								if (ExecutionInfo->OutputPin->LinkedTo.Num() == 0 &&
-									ExecutionInfo->OutputPin->PinType.PinCategory == UEdGraphSchema_K2::PC_Struct &&
-									ExecutionInfo->OutputPin->PinType.PinSubCategoryObject == FControlRigExecuteContext::StaticStruct())
-								{
-									float Distance = (Location - FVector2D((float)RigNode->NodePosX, (float)RigNode->NodePosY)).SizeSquared();
-									if (Distance < ClosestDistance)
-									{
-										ClosestDistance = Distance;
-										ClosestRigNode = RigNode;
-										ClosestExecutionPin = CurrentExecutionInfos[0]->OutputPin;
-									}
-								}
-							}
-						}
-					}
-				}
-
-				// if we didn't find a closest rig node with an execution pin - let's create one
-				if (ClosestRigNode == nullptr)
-				{
-					ClosestRigNode = SpawnNode(ParentGraph, Blueprint, FRigUnit_BeginExecution::StaticStruct(), Location - FVector2D(200.f, 0.f));
-					ClosestExecutionPin = ClosestRigNode->Pins[0];
-				}
-
-				UControlRigGraph* RigGraph = CastChecked<UControlRigGraph>(ParentGraph);
-				const UControlRigGraphSchema* ControlRigSchema = RigGraph->GetControlRigGraphSchema();
-				for (const TSharedRef<FControlRigField>& NewExecutionInfo : NewExecutionInfos)
-				{
-					if (NewExecutionInfo->InputPin != nullptr)
-					{
-						ControlRigSchema->TryCreateConnection(ClosestExecutionPin, NewExecutionInfo->InputPin);
-					}
-				}
-			}
-		}
 	}
 
 	return NewNode;
@@ -177,26 +120,38 @@ bool UControlRigUnitNodeSpawner::IsTemplateNodeFilteredOut(FBlueprintActionFilte
 UControlRigGraphNode* UControlRigUnitNodeSpawner::SpawnNode(UEdGraph* ParentGraph, UBlueprint* Blueprint, UStruct* StructTemplate, FVector2D const Location)
 {
 	UControlRigGraphNode* NewNode = nullptr;
-
-	bool const bIsTemplateNode = FBlueprintNodeTemplateCache::IsTemplateOuter(ParentGraph);
-
-	// First create a backing member for our node
-	FName MemberName = NAME_None;
-
-	if (!bIsTemplateNode)
+	UControlRigBlueprint* RigBlueprint = Cast<UControlRigBlueprint>(Blueprint);
+	if (RigBlueprint != nullptr)
 	{
-		MemberName = FControlRigBlueprintUtils::AddUnitMember(Blueprint, StructTemplate);
-	}
-	else
-	{
-		MemberName = StructTemplate->GetFName();
-	}
+		bool const bIsTemplateNode = FBlueprintNodeTemplateCache::IsTemplateOuter(ParentGraph);
 
-	if (MemberName != NAME_None)
-	{
-		NewNode = FControlRigBlueprintUtils::InstantiateGraphNodeForProperty(ParentGraph, MemberName, Location);
-	}
+		// First create a backing member for our node
+		FName MemberName = NAME_None;
 
+		if (!bIsTemplateNode)
+		{
+			FName Name = FControlRigBlueprintUtils::ValidateName(RigBlueprint, StructTemplate->GetFName().ToString());
+			if (RigBlueprint->ModelController->AddNode(StructTemplate->GetFName(), Location, Name))
+			{
+				MemberName = RigBlueprint->LastNameFromNotification;
+				for (UEdGraphNode* Node : ParentGraph->Nodes)
+				{
+					if (UControlRigGraphNode* RigNode = Cast<UControlRigGraphNode>(Node))
+					{
+						if (RigNode->GetPropertyName() == MemberName)
+						{
+							NewNode = RigNode;
+							break;
+						}
+					}
+				}
+			}
+		}
+		else
+		{
+			NewNode = FControlRigBlueprintUtils::InstantiateGraphNodeForStructPath(ParentGraph, *StructTemplate->GetDisplayNameText().ToString(), Location, StructTemplate->GetPathName());
+		}
+	}
 	return NewNode;
 }
 
diff --git a/Engine/Plugins/Experimental/ControlRig/Source/ControlRigEditor/Private/Graph/NodeSpawners/ControlRigVariableNodeSpawner.cpp b/Engine/Plugins/Experimental/ControlRig/Source/ControlRigEditor/Private/Graph/NodeSpawners/ControlRigVariableNodeSpawner.cpp
--- a/Engine/Plugins/Experimental/ControlRig/Source/ControlRigEditor/Private/Graph/NodeSpawners/ControlRigVariableNodeSpawner.cpp
+++ b/Engine/Plugins/Experimental/ControlRig/Source/ControlRigEditor/Private/Graph/NodeSpawners/ControlRigVariableNodeSpawner.cpp
@@ -10,6 +10,8 @@
 #include "EditorCategoryUtils.h"
 #include "K2Node_Variable.h"
 #include "BlueprintNodeTemplateCache.h"
+#include "ControlRigController.h"
+#include "ControlRigBlueprint.h"
 #include "ControlRigBlueprintUtils.h"
 #include "Units/RigUnit.h"
 
@@ -29,7 +31,8 @@ const TArray<FString> GControlRigVariableNodeSpawnerAllowedStructTypes = {
 	TEXT("FMatrix"),
 	TEXT("FRotationMatrix"),
 	TEXT("FScaleMatrix"),
-	TEXT("FTransform")
+	TEXT("FTransform"),
+	TEXT("FEulerTransform")
 };
 
 const TArray<FString> GControlRigVariableNodeSpawnerAllowedEnumTypes = {
@@ -76,8 +79,6 @@ UEdGraphNode* UControlRigVariableNodeSpawner::Invoke(UEdGraph* ParentGraph, FBin
 {
 	UControlRigGraphNode* NewNode = nullptr;
 
-//	const FScopedTransaction Transaction(LOCTEXT("AddRigPropertyNode", "Add Rig Property Node"));
-
 	bool const bIsTemplateNode = FBlueprintNodeTemplateCache::IsTemplateOuter(ParentGraph);
 
 	// First create a backing member for our node
@@ -85,16 +86,42 @@ UEdGraphNode* UControlRigVariableNodeSpawner::Invoke(UEdGraph* ParentGraph, FBin
 	FName MemberName = NAME_None;
 	if(!bIsTemplateNode)
 	{
-		MemberName = FControlRigBlueprintUtils::AddPropertyMember(Blueprint, EdGraphPinType, DefaultMenuSignature.MenuName.ToString());
+#if WITH_EDITOR
+		if (GEditor)
+		{
+			GEditor->CancelTransaction(0);
+		}
+#endif
+
+		if (UControlRigBlueprint* RigBlueprint = Cast<UControlRigBlueprint>(Blueprint))
+		{
+			FName DataType = EdGraphPinType.PinCategory;
+			if (UStruct* Struct = Cast<UStruct>(EdGraphPinType.PinSubCategoryObject))
+			{
+				DataType = Struct->GetFName();
+			}
+
+			FName Name = FControlRigBlueprintUtils::ValidateName(RigBlueprint, DefaultMenuSignature.MenuName.ToString());
+			if (RigBlueprint->ModelController->AddParameter(*Name.ToString(), DataType, EControlRigModelParameterType::Hidden, Location))
+			{
+				MemberName = RigBlueprint->LastNameFromNotification;
+				for (UEdGraphNode* Node : ParentGraph->Nodes)
+				{
+					if (UControlRigGraphNode* RigNode = Cast<UControlRigGraphNode>(Node))
+					{
+						if (RigNode->GetPropertyName() == MemberName)
+						{
+							NewNode = RigNode;
+							break;
+						}
+					}
+				}
+			}
+		}
 	}
 	else
 	{
-		MemberName = FControlRigBlueprintUtils::GetNewPropertyMemberName(Blueprint, DefaultMenuSignature.MenuName.ToString());
-	}
-
-	if(MemberName != NAME_None)
-	{
-		NewNode = FControlRigBlueprintUtils::InstantiateGraphNodeForProperty(ParentGraph, MemberName, Location, EdGraphPinType);
+		NewNode = FControlRigBlueprintUtils::InstantiateGraphNodeForProperty(ParentGraph, *DefaultMenuSignature.MenuName.ToString(), Location, EdGraphPinType);
 	}
 
 	return NewNode;
diff --git a/Engine/Plugins/Experimental/ControlRig/Source/ControlRigEditor/Private/Graph/SControlRigGraphNode.cpp b/Engine/Plugins/Experimental/ControlRig/Source/ControlRigEditor/Private/Graph/SControlRigGraphNode.cpp
--- a/Engine/Plugins/Experimental/ControlRig/Source/ControlRigEditor/Private/Graph/SControlRigGraphNode.cpp
+++ b/Engine/Plugins/Experimental/ControlRig/Source/ControlRigEditor/Private/Graph/SControlRigGraphNode.cpp
@@ -18,6 +18,8 @@
 #include "Kismet2/KismetDebugUtilities.h"
 #include "PropertyPathHelpers.h"
 #include "UObject/PropertyPortFlags.h"
+#include "ControlRigBlueprint.h"
+#include "ControlRigController.h"
 
 #define LOCTEXT_NAMESPACE "SControlRigGraphNode"
 
@@ -106,7 +108,7 @@ void SControlRigGraphNode::Construct( const FArguments& InArgs )
 		{
 			for(const TSharedRef<FControlRigField>& Field : InItems)
 			{
-				if(InControlRigGraphNode->IsPinExpanded(Field->GetPropertyPath()))
+				if(InControlRigGraphNode->IsPinExpanded(Field->GetPinPath()))
 				{
 					TreeWidget->SetItemExpansion(Field, true);
 
@@ -181,6 +183,36 @@ TSharedPtr<SGraphPin> SControlRigGraphNode::GetHoveredPin(const FGeometry& MyGeo
 	return HoveredPin;
 }
 
+/** @param NewPosition  The Node should be relocated to this position in the graph panel */
+void SControlRigGraphNode::MoveTo( const FVector2D& NewPosition, FNodeSet& NodeFilter )
+{
+	if (!NodeFilter.Find(SharedThis(this)))
+	{
+		if (GraphNode && !RequiresSecondPassLayout())
+		{
+			UControlRigGraphNode* ControlRigGraphNode = CastChecked<UControlRigGraphNode>(GraphNode);
+			ControlRigGraphNode->GetBlueprint()->ModelController->SetNodePosition(ControlRigGraphNode->GetPropertyName(), NewPosition, false);
+		}
+	}
+}
+
+void SControlRigGraphNode::EndUserInteraction() const
+{
+#if WITH_EDITOR
+	if (GEditor)
+	{
+		GEditor->CancelTransaction(0);
+	}
+#endif
+
+	if (GraphNode)
+	{
+		UControlRigGraphNode* ControlRigGraphNode = CastChecked<UControlRigGraphNode>(GraphNode);
+		FVector2D Position(ControlRigGraphNode->NodePosX, ControlRigGraphNode->NodePosY);
+		ControlRigGraphNode->GetBlueprint()->ModelController->SetNodePosition(ControlRigGraphNode->GetPropertyName(), Position, true);
+	}
+}
+
 void SControlRigGraphNode::AddPin(const TSharedRef<SGraphPin>& PinToAdd) 
 {
 	// We show our own label
@@ -638,7 +670,17 @@ void SControlRigGraphNode::HandleExpansionChanged(TSharedRef<FControlRigField> I
 {
 	if (GraphNode)
 	{
-		CastChecked<UControlRigGraphNode>(GraphNode)->SetPinExpansion(InItem->GetPropertyPath(), bExpanded);
+		UControlRigBlueprint* ControlRigBlueprint = Cast<UControlRigBlueprint>(GraphNode->GetGraph()->GetOuter());
+		if (ControlRigBlueprint)
+		{
+			if (ControlRigBlueprint->ModelController)
+			{
+				FString PinPath = InItem->GetPinPath();
+				FString Left, Right;
+				ControlRigBlueprint->Model->SplitPinPath(PinPath, Left, Right);
+				ControlRigBlueprint->ModelController->ExpandPin(*Left, *Right, InItem->GetPin()->Direction == EGPD_Input, bExpanded);
+			}
+		}
 	}
 }
 
@@ -679,7 +721,8 @@ FReply SControlRigGraphNode::HandleAddArrayElement(TWeakPtr<FControlRigField> In
 	{
 		if (UControlRigGraphNode* ControlRigGraphNode = Cast<UControlRigGraphNode>(GraphNode))
 		{
-			ControlRigGraphNode->HandleAddArrayElement(Item->PropertyPath);
+			// todo ?
+			ControlRigGraphNode->HandleAddArrayElement(Item->GetPinPath());
 		}
 	}
 
diff --git a/Engine/Plugins/Experimental/ControlRig/Source/ControlRigEditor/Private/Graph/SControlRigGraphNode.h b/Engine/Plugins/Experimental/ControlRig/Source/ControlRigEditor/Private/Graph/SControlRigGraphNode.h
--- a/Engine/Plugins/Experimental/ControlRig/Source/ControlRigEditor/Private/Graph/SControlRigGraphNode.h
+++ b/Engine/Plugins/Experimental/ControlRig/Source/ControlRigEditor/Private/Graph/SControlRigGraphNode.h
@@ -31,6 +31,8 @@ public:
 	// SGraphNode interface
 	virtual TSharedRef<SWidget> CreateTitleWidget(TSharedPtr<SNodeTitle> NodeTitle) override;
 	virtual bool UseLowDetailNodeTitles() const override;
+	virtual void MoveTo( const FVector2D& NewPosition, FNodeSet& NodeFilter ) override;
+	virtual void EndUserInteraction() const override;
 	virtual void AddPin( const TSharedRef<SGraphPin>& PinToAdd ) override;
 	virtual void SetDefaultTitleAreaWidget(TSharedRef<SOverlay> DefaultTitleAreaWidget) override
 	{
diff --git a/Engine/Plugins/Experimental/ControlRig/Source/ControlRigEditor/Private/Graph/SGraphPinCurveFloat.cpp b/Engine/Plugins/Experimental/ControlRig/Source/ControlRigEditor/Private/Graph/SGraphPinCurveFloat.cpp
--- a/Engine/Plugins/Experimental/ControlRig/Source/ControlRigEditor/Private/Graph/SGraphPinCurveFloat.cpp
+++ b/Engine/Plugins/Experimental/ControlRig/Source/ControlRigEditor/Private/Graph/SGraphPinCurveFloat.cpp
@@ -121,7 +121,7 @@ void SGraphPinCurveFloat::OnCurveChanged(const TArray<FRichCurveEditInfo>& Chang
 		{
 			Pin->DefaultValue.Empty();
 			FRuntimeFloatCurve::StaticStruct()->ExportText(Pin->DefaultValue, &Curve, nullptr, nullptr, EPropertyPortFlags::PPF_None, nullptr, true);
-			Node->CopyPinDefaultsToProperties(Pin, true, true);
+			Node->CopyPinDefaultsToModel(Pin);
 		}
 	}
 	ModifyOwner();
diff --git a/Engine/Source/Editor/AnimGraph/Classes/AnimGraphNode_CustomProperty.h b/Engine/Source/Editor/AnimGraph/Classes/AnimGraphNode_CustomProperty.h
--- a/Engine/Source/Editor/AnimGraph/Classes/AnimGraphNode_CustomProperty.h
+++ b/Engine/Source/Editor/AnimGraph/Classes/AnimGraphNode_CustomProperty.h
@@ -14,8 +14,8 @@ class FCompilerResultsLog;
 class IDetailLayoutBuilder;
 class IPropertyHandle;
 
-UCLASS(MinimalAPI, Abstract)
-class UAnimGraphNode_CustomProperty : public UAnimGraphNode_Base
+UCLASS(Abstract)
+class ANIMGRAPH_API UAnimGraphNode_CustomProperty : public UAnimGraphNode_Base
 {
 	GENERATED_BODY()
 
@@ -29,13 +29,13 @@ public:
 	//~ End UEdGraphNode Interface.
 
 	// Gets the property on InOwnerInstanceClass that corresponds to InInputPin
-	ANIMGRAPH_API void GetInstancePinProperty(const UClass* InOwnerInstanceClass, UEdGraphPin* InInputPin, UProperty*& OutProperty);
+	void GetInstancePinProperty(const UClass* InOwnerInstanceClass, UEdGraphPin* InInputPin, UProperty*& OutProperty);
 	// Gets the unique name for the property linked to a given pin
-	ANIMGRAPH_API FString GetPinTargetVariableName(const UEdGraphPin* InPin) const;
+	FString GetPinTargetVariableName(const UEdGraphPin* InPin) const;
 	// Gets Target Class this properties to link
-	ANIMGRAPH_API UClass* GetTargetClass() const;
+	UClass* GetTargetClass() const;
 	// Add Source and Target Properties - Check FAnimNode_CustomProperty
-	ANIMGRAPH_API void AddSourceTargetProperties(const FName& InSourcePropertyName, const FName& InTargetPropertyName);
+	void AddSourceTargetProperties(const FName& InSourcePropertyName, const FName& InTargetPropertyName);
 
 	// return true if this pin name is property exposed
 	// return false if this pin doesn't belong to property
@@ -44,6 +44,13 @@ public:
 		return (PinName != FName(TEXT("Pose"), FNAME_Find));
 	}
 
+	// ----- UI CALLBACKS ----- //
+	// If given property exposed on this node
+	virtual ECheckBoxState IsPropertyExposed(FName PropertyName) const;
+	// User chose to expose, or unexpose a property
+	virtual void OnPropertyExposeCheckboxChanged(ECheckBoxState NewState, FName PropertyName);
+	// User changed the instance class
+	void OnInstanceClassChanged(IDetailLayoutBuilder* DetailBuilder);
 protected:
 
 	/** List of property names we know to exist on the target class, so we can detect when
@@ -57,23 +64,14 @@ protected:
 	TArray<FName> ExposedPropertyNames;
 
 	// Searches the instance class for properties that we can expose (public and BP visible)
-	void GetExposableProperties(TArray<UProperty*>& OutExposableProperties) const;
+	virtual void GetExposableProperties(TArray<UProperty*>& OutExposableProperties) const;
 	// Gets a property's type as FText (for UI)
 	FText GetPropertyTypeText(UProperty* Property);
 	// Given a new class, rebuild the known property list (for tracking class changes and moving pins)
-	void RebuildExposedProperties(UClass* InNewClass);
-
-	// ----- UI CALLBACKS ----- //
-	// If given property exposed on this node
-	ECheckBoxState IsPropertyExposed(FName PropertyName) const;
-	// User chose to expose, or unexpose a property
-	void OnPropertyExposeCheckboxChanged(ECheckBoxState NewState, FName PropertyName);
-	// User changed the instance class
-	void OnInstanceClassChanged(IDetailLayoutBuilder* DetailBuilder);
+	virtual void RebuildExposedProperties();
 
 	// internal node accessor
 	virtual FAnimNode_CustomProperty* GetInternalNode() PURE_VIRTUAL(UAnimGraphNode_CustomProperty::GetInternalNode, return nullptr;);
 	virtual const FAnimNode_CustomProperty* GetInternalNode() const PURE_VIRTUAL(UAnimGraphNode_CustomProperty::GetInternalNode, return nullptr;);
 
-	friend class UAnimGraphNode_SubInstance;
 };
diff --git a/Engine/Source/Editor/AnimGraph/Private/AnimGraphNode_CustomProperty.cpp b/Engine/Source/Editor/AnimGraph/Private/AnimGraphNode_CustomProperty.cpp
--- a/Engine/Source/Editor/AnimGraph/Private/AnimGraphNode_CustomProperty.cpp
+++ b/Engine/Source/Editor/AnimGraph/Private/AnimGraphNode_CustomProperty.cpp
@@ -136,19 +136,15 @@ FText UAnimGraphNode_CustomProperty::GetPropertyTypeText(UProperty* Property)
 	return PropertyTypeText;
 }
 
-void UAnimGraphNode_CustomProperty::RebuildExposedProperties(UClass* InNewClass)
+void UAnimGraphNode_CustomProperty::RebuildExposedProperties()
 {
 	ExposedPropertyNames.Empty();
 	KnownExposableProperties.Empty();
-	if(InNewClass)
+	TArray<UProperty*> ExposableProperties;
+	GetExposableProperties(ExposableProperties);
+	for(UProperty* Property : ExposableProperties)
 	{
-		TArray<UProperty*> ExposableProperties;
-		GetExposableProperties(ExposableProperties);
-
-		for(UProperty* Property : ExposableProperties)
-		{
-			KnownExposableProperties.Add(Property->GetFName());
-		}
+		KnownExposableProperties.Add(Property->GetFName());
 	}
 }
 
diff --git a/Engine/Source/Editor/AnimGraph/Private/AnimGraphNode_SubInstance.cpp b/Engine/Source/Editor/AnimGraph/Private/AnimGraphNode_SubInstance.cpp
--- a/Engine/Source/Editor/AnimGraph/Private/AnimGraphNode_SubInstance.cpp
+++ b/Engine/Source/Editor/AnimGraph/Private/AnimGraphNode_SubInstance.cpp
@@ -179,7 +179,7 @@ void UAnimGraphNode_SubInstance::PostEditChangeProperty(FPropertyChangedEvent& P
 		if(ChangedProperty->GetFName() == GET_MEMBER_NAME_CHECKED(FAnimNode_SubInstance, InstanceClass))
 		{
 			bRequiresNodeReconstruct = true;
-			RebuildExposedProperties(*Node.InstanceClass);
+			RebuildExposedProperties();
 		}
 	}
 
