- fix issue with mesh change delegate not coming first time

diff --git a/Engine/Plugins/Experimental/ControlRig/Source/ControlRigEditor/Private/Editor/ControlRigEditor.cpp b/Engine/Plugins/Experimental/ControlRig/Source/ControlRigEditor/Private/Editor/ControlRigEditor.cpp
--- a/Engine/Plugins/Experimental/ControlRig/Source/ControlRigEditor/Private/Editor/ControlRigEditor.cpp
+++ b/Engine/Plugins/Experimental/ControlRig/Source/ControlRigEditor/Private/Editor/ControlRigEditor.cpp
@@ -115,9 +115,11 @@ void FControlRigEditor::InitControlRigEditor(const EToolkitMode::Type Mode, cons
 	PersonaToolkitArgs.OnPreviewSceneCreated = FOnPreviewSceneCreated::FDelegate::CreateSP(this, &FControlRigEditor::HandlePreviewSceneCreated);
 	PersonaToolkit = PersonaModule.CreatePersonaToolkit(InControlRigBlueprint, PersonaToolkitArgs);
 
+	// set delegate prior to setting mesh
+	// otherwise, you don't get delegate
+	PersonaToolkit->GetPreviewScene()->RegisterOnPreviewMeshChanged(FOnPreviewMeshChanged::CreateSP(this, &FControlRigEditor::HandlePreviewMeshChanged));
 	// Set a default preview mesh, if any
 	PersonaToolkit->SetPreviewMesh(InControlRigBlueprint->GetPreviewMesh(), false);
-	PersonaToolkit->GetPreviewScene()->RegisterOnPreviewMeshChanged(FOnPreviewMeshChanged::CreateSP(this, &FControlRigEditor::HandlePreviewMeshChanged));
 
 	Toolbox = SNew(SBorder)
 		.BorderImage(FEditorStyle::GetBrush("ToolPanel.GroupBorder"))
