Fix UE Main Build break in non PCH non unity

diff --git a/Engine/Plugins/Experimental/ControlRig/Source/ControlRigDeveloper/Private/SVariableMappingWidget.cpp b/Engine/Plugins/Experimental/ControlRig/Source/ControlRigDeveloper/Private/SVariableMappingWidget.cpp
--- a/Engine/Plugins/Experimental/ControlRig/Source/ControlRigDeveloper/Private/SVariableMappingWidget.cpp
+++ b/Engine/Plugins/Experimental/ControlRig/Source/ControlRigDeveloper/Private/SVariableMappingWidget.cpp
@@ -2,9 +2,8 @@
 
 
 #include "SVariableMappingWidget.h"
+#include "Widgets/Input/SCheckBox.h"
 #include "Widgets/Text/SInlineEditableTextBlock.h"
-#include "Widgets/Input/SSearchBox.h"
-#include "Widgets/Input/SButton.h"
 
 #define LOCTEXT_NAMESPACE "SVariableMappingWidget"
 
diff --git a/Engine/Plugins/Experimental/ControlRig/Source/ControlRigDeveloper/Private/SVariableMappingWidget.h b/Engine/Plugins/Experimental/ControlRig/Source/ControlRigDeveloper/Private/SVariableMappingWidget.h
--- a/Engine/Plugins/Experimental/ControlRig/Source/ControlRigDeveloper/Private/SVariableMappingWidget.h
+++ b/Engine/Plugins/Experimental/ControlRig/Source/ControlRigDeveloper/Private/SVariableMappingWidget.h
@@ -4,6 +4,8 @@
 
 #include "CoreMinimal.h"
 #include "SSearchableComboBox.h"
+#include "Widgets/Input/SSearchBox.h"
+#include "Widgets/Input/SButton.h"
 #include "Widgets/Views/STreeView.h"
 
 //=====================================================================================
diff --git a/Engine/Plugins/Experimental/ControlRig/Source/ControlRigEditor/Private/Graph/SControlRigGraphNode.cpp b/Engine/Plugins/Experimental/ControlRig/Source/ControlRigEditor/Private/Graph/SControlRigGraphNode.cpp
--- a/Engine/Plugins/Experimental/ControlRig/Source/ControlRigEditor/Private/Graph/SControlRigGraphNode.cpp
+++ b/Engine/Plugins/Experimental/ControlRig/Source/ControlRigEditor/Private/Graph/SControlRigGraphNode.cpp
@@ -20,6 +20,9 @@
 #include "UObject/PropertyPortFlags.h"
 #include "ControlRigBlueprint.h"
 #include "ControlRigController.h"
+#if WITH_EDITOR
+#include "Editor.h"
+#endif
 
 #define LOCTEXT_NAMESPACE "SControlRigGraphNode"
 
diff --git a/Engine/Source/Editor/DetailCustomizations/Private/LandscapeUIDetails.cpp b/Engine/Source/Editor/DetailCustomizations/Private/LandscapeUIDetails.cpp
--- a/Engine/Source/Editor/DetailCustomizations/Private/LandscapeUIDetails.cpp
+++ b/Engine/Source/Editor/DetailCustomizations/Private/LandscapeUIDetails.cpp
@@ -10,6 +10,7 @@
 #include "DetailWidgetRow.h"
 #include "Widgets/Input/SButton.h"
 #include "Widgets/Text/STextBlock.h"
+#include "Widgets/Input/SCheckBox.h"
 #include "Misc/MessageDialog.h"
 #include "Editor.h"
 
