Remove CacheBones

diff --git a/Engine/Plugins/Experimental/ControlRig/Source/ControlRig/Private/Sequencer/ControlRigSequencerAnimInstanceProxy.cpp b/Engine/Plugins/Experimental/ControlRig/Source/ControlRig/Private/Sequencer/ControlRigSequencerAnimInstanceProxy.cpp
--- a/Engine/Plugins/Experimental/ControlRig/Source/ControlRig/Private/Sequencer/ControlRigSequencerAnimInstanceProxy.cpp
+++ b/Engine/Plugins/Experimental/ControlRig/Source/ControlRig/Private/Sequencer/ControlRigSequencerAnimInstanceProxy.cpp
@@ -43,20 +43,6 @@ FAnimNode_Base* FControlRigSequencerAnimInstanceProxy::GetCustomRootNode()
 	return &SequencerRootNode;
 }
 
-void FControlRigSequencerAnimInstanceProxy::CacheBones()
-{
-	// As we dont use the RootNode (this is not using anim blueprint), 
-	// we just cache the nodes we know about
-	if (bBoneCachesInvalidated)
-	{
-		bBoneCachesInvalidated = false;
-
-		CachedBonesCounter.Increment();
-		FAnimationCacheBonesContext Context(this);
-		SequencerRootNode.CacheBones_AnyThread(Context);
-	}
-}
-
 void FControlRigSequencerAnimInstanceProxy::ResetNodes()
 {
 	FAnimSequencerInstanceProxy::ResetNodes();
diff --git a/Engine/Plugins/Experimental/ControlRig/Source/ControlRig/Private/Sequencer/ControlRigSequencerAnimInstanceProxy.h b/Engine/Plugins/Experimental/ControlRig/Source/ControlRig/Private/Sequencer/ControlRigSequencerAnimInstanceProxy.h
--- a/Engine/Plugins/Experimental/ControlRig/Source/ControlRig/Private/Sequencer/ControlRigSequencerAnimInstanceProxy.h
+++ b/Engine/Plugins/Experimental/ControlRig/Source/ControlRig/Private/Sequencer/ControlRigSequencerAnimInstanceProxy.h
@@ -39,7 +39,6 @@ public:
 	}
 
 	// FAnimInstanceProxy interface
-	virtual void CacheBones() override;
 	virtual void Initialize(UAnimInstance* InAnimInstance) override;
 	virtual void Update(float DeltaSeconds) override;
 	virtual FAnimNode_Base* GetCustomRootNode() override;
