Control Rig Automation

diff --git a/Engine/Plugins/Experimental/ControlRig/Source/ControlRig/Private/ControlRig.cpp b/Engine/Plugins/Experimental/ControlRig/Source/ControlRig/Private/ControlRig.cpp
--- a/Engine/Plugins/Experimental/ControlRig/Source/ControlRig/Private/ControlRig.cpp
+++ b/Engine/Plugins/Experimental/ControlRig/Source/ControlRig/Private/ControlRig.cpp
@@ -108,7 +108,7 @@ void UControlRig::PreEvaluate_GameThread()
 {
 	// this won't work with procedural rigging
 	// so I wonder this should be just bool option for control rig?
-	Hierarchy.Reset();
+	Hierarchy.ResetTransforms();
 
 	// input delegates
 	OnPreEvaluateGatherInput.ExecuteIfBound(this);
diff --git a/Engine/Plugins/Experimental/ControlRig/Source/ControlRig/Private/Rigs/Hierarchy.cpp b/Engine/Plugins/Experimental/ControlRig/Source/ControlRig/Private/Rigs/Hierarchy.cpp
--- a/Engine/Plugins/Experimental/ControlRig/Source/ControlRig/Private/Rigs/Hierarchy.cpp
+++ b/Engine/Plugins/Experimental/ControlRig/Source/ControlRig/Private/Rigs/Hierarchy.cpp
@@ -100,6 +100,11 @@ void FRigHierarchy::Initialize()
 }
 
 void FRigHierarchy::Reset()
+{
+	Joints.Reset();
+}
+
+void FRigHierarchy::ResetTransforms()
 {
 	// initialize transform
 	for (int32 Index = 0; Index < Joints.Num(); ++Index)
diff --git a/Engine/Plugins/Experimental/ControlRig/Source/ControlRig/Private/Units/RigUnit_ApplyFK.h b/Engine/Plugins/Experimental/ControlRig/Source/ControlRig/Private/Units/RigUnit_ApplyFK.h
--- a/Engine/Plugins/Experimental/ControlRig/Source/ControlRig/Private/Units/RigUnit_ApplyFK.h
+++ b/Engine/Plugins/Experimental/ControlRig/Source/ControlRig/Private/Units/RigUnit_ApplyFK.h
@@ -30,7 +30,7 @@ struct FRigUnit_ApplyFK : public FRigUnit
 {
 	GENERATED_BODY()
 
-	virtual void Execute(const FRigUnitContext& InContext) override;
+	virtual CONTROLRIG_API void Execute(const FRigUnitContext& InContext) override;
 
 	UPROPERTY(meta = (Input, Output))
 	FRigHierarchyRef HierarchyRef;
diff --git a/Engine/Plugins/Experimental/ControlRig/Source/ControlRig/Public/ControlRig.h b/Engine/Plugins/Experimental/ControlRig/Source/ControlRig/Public/ControlRig.h
--- a/Engine/Plugins/Experimental/ControlRig/Source/ControlRig/Public/ControlRig.h
+++ b/Engine/Plugins/Experimental/ControlRig/Source/ControlRig/Public/ControlRig.h
@@ -185,4 +185,5 @@ private:
 	friend class URigUnitEditor_Base;
 	friend class FControlRigEditor;
 	friend class SRigHierarchy;
+	friend class UEngineTestControlRig;
 };
\ No newline at end of file
diff --git a/Engine/Plugins/Experimental/ControlRig/Source/ControlRig/Public/Hierarchy.h b/Engine/Plugins/Experimental/ControlRig/Source/ControlRig/Public/Hierarchy.h
--- a/Engine/Plugins/Experimental/ControlRig/Source/ControlRig/Public/Hierarchy.h
+++ b/Engine/Plugins/Experimental/ControlRig/Source/ControlRig/Public/Hierarchy.h
@@ -338,9 +338,15 @@ struct CONTROLRIG_API FRigHierarchy
 		}
 	}
 
+	// updates all of the internal caches
 	void Initialize();
+
+	// clears the hierarchy and removes all content
 	void Reset();
 
+	// resets all of the transforms back to the initial transform
+	void ResetTransforms();
+
 	int32 GetNum() const
 	{
 		return Joints.Num();
@@ -412,6 +418,13 @@ struct FRigHierarchyContainer
 
 		// @todo reset whole hierarhcy
 	}
+
+	void ResetTransforms()
+	{
+		BaseHierarchy.ResetTransforms();
+
+		// @todo reset whole hierarhcy
+	}
 };
 
 USTRUCT(BlueprintType)
@@ -478,4 +491,5 @@ private:
 		return nullptr;
 	}
 	friend class UControlRig;
+	friend class FControlRigUnitTestBase;
 };
\ No newline at end of file
diff --git a/Engine/Plugins/Experimental/ControlRig/Source/ControlRigDeveloper/Public/ControlRigBlueprint.h b/Engine/Plugins/Experimental/ControlRig/Source/ControlRigDeveloper/Public/ControlRigBlueprint.h
--- a/Engine/Plugins/Experimental/ControlRig/Source/ControlRigDeveloper/Public/ControlRigBlueprint.h
+++ b/Engine/Plugins/Experimental/ControlRig/Source/ControlRigDeveloper/Public/ControlRigBlueprint.h
@@ -121,4 +121,5 @@ private:
 	friend class FControlRigBlueprintCompilerContext;
 	friend class SRigHierarchy;
 	friend class FControlRigEditor;
+	friend class UEngineTestControlRig;
 };
diff --git a/Engine/Plugins/Experimental/ControlRig/Source/ControlRigEditor/Private/Graph/NodeSpawners/ControlRigPropertyNodeSpawner.h b/Engine/Plugins/Experimental/ControlRig/Source/ControlRigEditor/Private/Graph/NodeSpawners/ControlRigPropertyNodeSpawner.h
--- a/Engine/Plugins/Experimental/ControlRig/Source/ControlRigEditor/Private/Graph/NodeSpawners/ControlRigPropertyNodeSpawner.h
+++ b/Engine/Plugins/Experimental/ControlRig/Source/ControlRigEditor/Private/Graph/NodeSpawners/ControlRigPropertyNodeSpawner.h
@@ -17,7 +17,7 @@
 class UControlRigGraphNode;
 
 UCLASS(Transient)
-class UControlRigPropertyNodeSpawner : public UBlueprintFieldNodeSpawner
+class CONTROLRIGEDITOR_API UControlRigPropertyNodeSpawner : public UBlueprintFieldNodeSpawner
 {
 	GENERATED_BODY()
 
diff --git a/Engine/Plugins/Experimental/ControlRig/Source/ControlRigEditor/Private/Graph/NodeSpawners/ControlRigUnitNodeSpawner.h b/Engine/Plugins/Experimental/ControlRig/Source/ControlRigEditor/Private/Graph/NodeSpawners/ControlRigUnitNodeSpawner.h
--- a/Engine/Plugins/Experimental/ControlRig/Source/ControlRigEditor/Private/Graph/NodeSpawners/ControlRigUnitNodeSpawner.h
+++ b/Engine/Plugins/Experimental/ControlRig/Source/ControlRigEditor/Private/Graph/NodeSpawners/ControlRigUnitNodeSpawner.h
@@ -17,7 +17,7 @@
 class UControlRigGraphNode;
 
 UCLASS(Transient)
-class UControlRigUnitNodeSpawner : public UBlueprintNodeSpawner
+class CONTROLRIGEDITOR_API UControlRigUnitNodeSpawner : public UBlueprintNodeSpawner
 {
 	GENERATED_BODY()
 
@@ -41,4 +41,6 @@ private:
 	/** The unit type we will spawn */
 	UPROPERTY(Transient)
 	UStruct* StructTemplate;
+
+	friend class UEngineTestControlRig;
 };
diff --git a/Engine/Plugins/Experimental/ControlRig/Source/ControlRigEditor/Private/Graph/NodeSpawners/ControlRigVariableNodeSpawner.h b/Engine/Plugins/Experimental/ControlRig/Source/ControlRigEditor/Private/Graph/NodeSpawners/ControlRigVariableNodeSpawner.h
--- a/Engine/Plugins/Experimental/ControlRig/Source/ControlRigEditor/Private/Graph/NodeSpawners/ControlRigVariableNodeSpawner.h
+++ b/Engine/Plugins/Experimental/ControlRig/Source/ControlRigEditor/Private/Graph/NodeSpawners/ControlRigVariableNodeSpawner.h
@@ -17,7 +17,7 @@
 class UControlRigGraphNode;
 
 UCLASS(Transient)
-class UControlRigVariableNodeSpawner : public UBlueprintNodeSpawner
+class CONTROLRIGEDITOR_API UControlRigVariableNodeSpawner : public UBlueprintNodeSpawner
 {
 	GENERATED_BODY()
 
@@ -49,4 +49,6 @@ public:
 private:
 	/** The pin type we will spawn */
 	FEdGraphPinType EdGraphPinType;
+
+	friend class UEngineTestControlRig;
 };
