Moved code for rigunit testing to the actual rig unit with a helper macro.

This way we can implement tests in the same file and make things easy.

diff --git a/Engine/Plugins/Experimental/ControlRig/Source/ControlRig/Private/ControlRig.cpp b/Engine/Plugins/Experimental/ControlRig/Source/ControlRig/Private/ControlRig.cpp
--- a/Engine/Plugins/Experimental/ControlRig/Source/ControlRig/Private/ControlRig.cpp
+++ b/Engine/Plugins/Experimental/ControlRig/Source/ControlRig/Private/ControlRig.cpp
@@ -177,6 +177,8 @@ void UControlRig::Execute(const EControlRigState InState)
 	FRigUnitContext Context;
 	Context.DeltaTime = DeltaTime;
 	Context.State = InState;
+	Context.HierarchyReference.Container = &Hierarchy;
+	Context.HierarchyReference.bUseBaseHierarchy = true;
 
 #if WITH_EDITORONLY_DATA
 	if (!bExecutionOn)
diff --git a/Engine/Plugins/Experimental/ControlRig/Source/ControlRig/Private/Units/RigUnit_ApplyFK.cpp b/Engine/Plugins/Experimental/ControlRig/Source/ControlRig/Private/Units/RigUnit_ApplyFK.cpp
--- a/Engine/Plugins/Experimental/ControlRig/Source/ControlRig/Private/Units/RigUnit_ApplyFK.cpp
+++ b/Engine/Plugins/Experimental/ControlRig/Source/ControlRig/Private/Units/RigUnit_ApplyFK.cpp
@@ -1,6 +1,7 @@
 // Copyright 1998-2019 Epic Games, Inc. All Rights Reserved.
 
 #include "RigUnit_ApplyFK.h"
+#include "Units/RigUnitTest.h"
 #include "Units/RigUnitContext.h"
 #include "HelperUtil.h"
 
@@ -53,4 +54,41 @@ FTransform FRigUnit_ApplyFK::GetBaseTransform(int32 JointIndex, const FRigHierar
 {
 	return UtilityHelpers::GetBaseTransformByMode(ApplyTransformSpace, [CurrentHierarchy](const FName& JointName) { return CurrentHierarchy->GetGlobalTransform(JointName); },
 		CurrentHierarchy->Joints[JointIndex].ParentName, BaseJoint, BaseTransform);
-}
\ No newline at end of file
+}
+
+#if WITH_DEV_AUTOMATION_TESTS
+IMPLEMENT_RIGUNIT_AUTOMATION_TEST(FRigUnit_ApplyFK)
+{
+	Context.State = EControlRigState::Update;
+
+	Hierarchy.AddJoint(TEXT("Root"), NAME_None, FTransform(FVector(1.f, 0.f, 0.f)));
+	Hierarchy.AddJoint(TEXT("JointA"), TEXT("Root"), FTransform(FVector(1.f, 2.f, 3.f)));
+
+	Unit.HierarchyRef = HierarchyRef;
+	Unit.Joint = TEXT("JointA");
+	Unit.ApplyTransformMode = EApplyTransformMode::Override;
+	Unit.ApplyTransformSpace = ETransformSpaceMode::GlobalSpace;
+	Unit.Transform = FTransform(FVector(0.f, 5.f, 0.f));
+
+	Hierarchy.Initialize();
+	Unit.Execute(Context);
+	AddErrorIfFalse(Hierarchy.GetGlobalTransform(1).GetTranslation().Equals(FVector(0.f, 5.f, 0.f)), TEXT("unexpected global transform"));
+	AddErrorIfFalse(Hierarchy.GetLocalTransform(1).GetTranslation().Equals(FVector(-1.f, 5.f, 0.f)), TEXT("unexpected local transform"));
+
+	Unit.ApplyTransformMode = EApplyTransformMode::Override;
+	Unit.ApplyTransformSpace = ETransformSpaceMode::LocalSpace;
+
+	Hierarchy.Initialize();
+	Unit.Execute(Context);
+	AddErrorIfFalse(Hierarchy.GetGlobalTransform(1).GetTranslation().Equals(FVector(1.f, 5.f, 0.f)), TEXT("unexpected global transform"));
+	AddErrorIfFalse(Hierarchy.GetLocalTransform(1).GetTranslation().Equals(FVector(0.f, 5.f, 0.f)), TEXT("unexpected local transform"));
+
+	Unit.ApplyTransformMode = EApplyTransformMode::Additive;
+
+	Hierarchy.Initialize();
+	Unit.Execute(Context);
+	AddErrorIfFalse(Hierarchy.GetGlobalTransform(1).GetTranslation().Equals(FVector(1.f, 7.f, 3.f)), TEXT("unexpected global transform"));
+	AddErrorIfFalse(Hierarchy.GetLocalTransform(1).GetTranslation().Equals(FVector(0.f, 7.f, 3.f)), TEXT("unexpected local transform"));
+	return true;
+}
+#endif
\ No newline at end of file
diff --git a/Engine/Plugins/Experimental/ControlRig/Source/ControlRig/Private/Units/RigUnit_ApplyFK.h b/Engine/Plugins/Experimental/ControlRig/Source/ControlRig/Private/Units/RigUnit_ApplyFK.h
--- a/Engine/Plugins/Experimental/ControlRig/Source/ControlRig/Private/Units/RigUnit_ApplyFK.h
+++ b/Engine/Plugins/Experimental/ControlRig/Source/ControlRig/Private/Units/RigUnit_ApplyFK.h
@@ -30,7 +30,7 @@ struct FRigUnit_ApplyFK : public FRigUnit
 {
 	GENERATED_BODY()
 
-	virtual CONTROLRIG_API void Execute(const FRigUnitContext& InContext) override;
+	virtual void Execute(const FRigUnitContext& InContext) override;
 
 	UPROPERTY(meta = (Input, Output))
 	FRigHierarchyRef HierarchyRef;
diff --git a/Engine/Plugins/Experimental/ControlRig/Source/ControlRig/Public/Units/RigUnitContext.h b/Engine/Plugins/Experimental/ControlRig/Source/ControlRig/Public/Units/RigUnitContext.h
--- a/Engine/Plugins/Experimental/ControlRig/Source/ControlRig/Public/Units/RigUnitContext.h
+++ b/Engine/Plugins/Experimental/ControlRig/Source/ControlRig/Public/Units/RigUnitContext.h
@@ -3,6 +3,7 @@
 #pragma once
 
 #include "CoreMinimal.h"
+#include "Hierarchy.h"
 
 /** Current state of rig
 *	What  state Control Rig currently is
@@ -23,5 +24,8 @@ struct FRigUnitContext
 
 	/** Current execution context */
 	EControlRigState State;
+
+	/** The current hierarchy being executed */
+	FRigHierarchyRef HierarchyReference;
 };
 
