Fix build break in dev-anim

diff --git a/Engine/Plugins/Experimental/ControlRig/Source/ControlRig/Public/Units/RigUnitTest.h b/Engine/Plugins/Experimental/ControlRig/Source/ControlRig/Public/Units/RigUnitTest.h
new file mode 100644
--- /dev/null
+++ b/Engine/Plugins/Experimental/ControlRig/Source/ControlRig/Public/Units/RigUnitTest.h
@@ -0,0 +1,62 @@
+// Copyright 1998-2019 Epic Games, Inc. All Rights Reserved.
+
+#pragma once
+
+#include "Misc/AutomationTest.h"
+#include "RigUnit.h"
+#include "RigUnitContext.h"
+#include "Hierarchy.h"
+
+class FControlRigUnitTestBase : public FAutomationTestBase
+{
+public:
+	FControlRigUnitTestBase(const FString& InName, bool bIsComplex)
+		: FAutomationTestBase(InName, bIsComplex)
+		, HierarchyContainer(FRigHierarchyContainer())
+		, HierarchyRef(FRigHierarchyRef())
+		, Hierarchy(HierarchyContainer.BaseHierarchy)
+	{
+		HierarchyRef.Container = &HierarchyContainer;
+		HierarchyRef.bUseBaseHierarchy = true;
+		Context.HierarchyReference = HierarchyRef;
+	}
+
+	FRigHierarchyContainer HierarchyContainer;
+	FRigHierarchyRef HierarchyRef;
+	FRigHierarchy& Hierarchy;
+	FRigUnitContext Context;
+};
+
+#define CONTROLRIG_RIGUNIT_STRINGIFY(Content) #Content
+#define IMPLEMENT_RIGUNIT_AUTOMATION_TEST(TUnitStruct) \
+	class TUnitStruct##Test : public FControlRigUnitTestBase \
+	{ \
+	public: \
+		TUnitStruct##Test( const FString& InName ) \
+		:FControlRigUnitTestBase( InName, false ) {\
+		} \
+		virtual uint32 GetTestFlags() const override { return EAutomationTestFlags::EditorContext | EAutomationTestFlags::EngineFilter; } \
+		virtual bool IsStressTest() const { return false; } \
+		virtual uint32 GetRequiredDeviceNum() const override { return 1; } \
+		virtual FString GetTestSourceFileName() const override { return __FILE__; } \
+		virtual int32 GetTestSourceFileLine() const override { return __LINE__; } \
+	protected: \
+		virtual void GetTests(TArray<FString>& OutBeautifiedNames, TArray <FString>& OutTestCommands) const override \
+		{ \
+			OutBeautifiedNames.Add(TEXT(CONTROLRIG_RIGUNIT_STRINGIFY(ControlRig.Units.##TUnitStruct))); \
+			OutTestCommands.Add(FString()); \
+		} \
+		virtual bool RunTest(const FString& Parameters) override \
+		{ \
+			Hierarchy.Reset(); \
+			TUnitStruct Unit; \
+			return RunControlRigUnitTest(Unit, Parameters); \
+		} \
+		virtual bool RunControlRigUnitTest(TUnitStruct& Unit, const FString& Parameters); \
+		virtual FString GetBeautifiedTestName() const override { return TEXT(CONTROLRIG_RIGUNIT_STRINGIFY(ControlRig.Units.##TUnitStruct)); } \
+	}; \
+	namespace\
+	{\
+		TUnitStruct##Test TUnitStruct##AutomationTestInstance(TEXT(CONTROLRIG_RIGUNIT_STRINGIFY(TUnitStruct##Test))); \
+	} \
+	bool TUnitStruct##Test::RunControlRigUnitTest(TUnitStruct& Unit, const FString& Parameters)
\ No newline at end of file
