Remove use pole vector because it doesn't support and this is deprecated, so it won't show up but for the existing nodes, they won't see the input anymore.

diff --git a/Engine/Plugins/Experimental/ControlRig/Source/ControlRig/Private/Units/Deprecated/RigUnit_TwoBoneIKFK.cpp b/Engine/Plugins/Experimental/ControlRig/Source/ControlRig/Private/Units/Deprecated/RigUnit_TwoBoneIKFK.cpp
--- a/Engine/Plugins/Experimental/ControlRig/Source/ControlRig/Private/Units/Deprecated/RigUnit_TwoBoneIKFK.cpp
+++ b/Engine/Plugins/Experimental/ControlRig/Source/ControlRig/Private/Units/Deprecated/RigUnit_TwoBoneIKFK.cpp
@@ -118,18 +118,7 @@ void FRigUnit_TwoBoneIKFK::Execute(const FRigUnitContext& Context)
 
 void FRigUnit_TwoBoneIKFK::SolveIK()
 {
-	if (bUsePoleTarget)
-	{
-		AnimationCore::SolveTwoBoneIK(StartJointIKTransform, MidJointIKTransform, EndJointIKTransform, PoleTarget, EndEffector.GetLocation(), UpperLimbLength, LowerLimbLength, false, 1.0f, 1.05f);
-	}
-	else
-	{
-		// default IK
-		// @todo: we don't have that yet, so default to using pole vector
-		ensure(false);
-		AnimationCore::SolveTwoBoneIK(StartJointIKTransform, MidJointIKTransform, EndJointIKTransform, PoleTarget, EndEffector.GetLocation(), UpperLimbLength, LowerLimbLength, false, 1.0f, 1.05f);
-	}
-
+	AnimationCore::SolveTwoBoneIK(StartJointIKTransform, MidJointIKTransform, EndJointIKTransform, PoleTarget, EndEffector.GetLocation(), UpperLimbLength, LowerLimbLength, false, 1.0f, 1.05f);
 	// set end effector rotation to current rotation
 	EndJointIKTransform.SetRotation(EndEffector.GetRotation());
 }
diff --git a/Engine/Plugins/Experimental/ControlRig/Source/ControlRig/Public/Units/Deprecated/RigUnit_TwoBoneIKFK.h b/Engine/Plugins/Experimental/ControlRig/Source/ControlRig/Public/Units/Deprecated/RigUnit_TwoBoneIKFK.h
--- a/Engine/Plugins/Experimental/ControlRig/Source/ControlRig/Public/Units/Deprecated/RigUnit_TwoBoneIKFK.h
+++ b/Engine/Plugins/Experimental/ControlRig/Source/ControlRig/Public/Units/Deprecated/RigUnit_TwoBoneIKFK.h
@@ -19,8 +19,7 @@ struct CONTROLRIG_API FRigUnit_TwoBoneIKFK : public FRigUnitMutable
 	virtual void Execute(const FRigUnitContext& Context) override;
 
 	FRigUnit_TwoBoneIKFK()
-		: bUsePoleTarget(true)
-		, PoleTarget(ForceInitToZero)
+		: PoleTarget(ForceInitToZero)
 		, Spin(0.f)
 		, IKBlend(1.f)
 		, PreviousFKIKBlend(1.f)
@@ -37,10 +36,6 @@ struct CONTROLRIG_API FRigUnit_TwoBoneIKFK : public FRigUnitMutable
 	UPROPERTY(EditAnywhere, Category = "TwoBoneIKFK", meta = (Input))
 	FName EndJoint;
 
-	// # Whether or not to use the pole target matrix.If false, use a minimum energy solution
-	UPROPERTY(EditAnywhere, Category = "TwoBoneIKFK", meta = (Input))
-	bool bUsePoleTarget;
-
 	// # Transform to use as the pole target(specifies the plane of solution)
 	UPROPERTY(EditAnywhere, Category = "TwoBoneIKFK", meta = (Input))
 	FVector PoleTarget;
