Fix build error

diff --git a/Engine/Plugins/Experimental/ControlRig/Source/ControlRig/Private/Units/Highlevel/Hierarchy/RigUnit_ModifyBoneTransforms.h b/Engine/Plugins/Experimental/ControlRig/Source/ControlRig/Private/Units/Highlevel/Hierarchy/RigUnit_ModifyBoneTransforms.h
--- a/Engine/Plugins/Experimental/ControlRig/Source/ControlRig/Private/Units/Highlevel/Hierarchy/RigUnit_ModifyBoneTransforms.h
+++ b/Engine/Plugins/Experimental/ControlRig/Source/ControlRig/Private/Units/Highlevel/Hierarchy/RigUnit_ModifyBoneTransforms.h
@@ -81,19 +81,19 @@ struct FRigUnit_ModifyBoneTransforms : public FRigUnit_HighlevelBaseMutable
 	/**
 	 * At 1 this sets the transform, between 0 and 1 the transform is blended with previous results.
 	 */
-	UPROPERTY(meta = (Input))
+	UPROPERTY(meta = (Input, ClampMin=0.f, ClampMax=1.f, UIMin = 0.f, UIMax = 1.f))
 	float Weight;
 
 	/**
 	 * The minimum of the weight - defaults to 0.0
 	 */
-	UPROPERTY(meta = (Input, Constant))
+	UPROPERTY(meta = (Input, Constant, ClampMin = 0.f, ClampMax = 1.f, UIMin = 0.f, UIMax = 1.f))
 	float WeightMinimum;
 
 	/**
 	 * The maximum of the weight - defaults to 1.0
 	 */
-	UPROPERTY(meta = (Input, Constant))
+	UPROPERTY(meta = (Input, Constant, ClampMin = 0.f, ClampMax = 1.f, UIMin = 0.f, UIMax = 1.f))
 	float WeightMaximum;
 
 	/**
