Control Rig Slide Setup for Mech

diff --git a/Engine/Plugins/Experimental/ControlRig/Source/ControlRig/Private/Units/Highlevel/Hierarchy/RigUnit_SlideChain.cpp b/Engine/Plugins/Experimental/ControlRig/Source/ControlRig/Private/Units/Highlevel/Hierarchy/RigUnit_SlideChain.cpp
new file mode 100644
--- /dev/null
+++ b/Engine/Plugins/Experimental/ControlRig/Source/ControlRig/Private/Units/Highlevel/Hierarchy/RigUnit_SlideChain.cpp
@@ -0,0 +1,133 @@
+// Copyright 1998-2019 Epic Games, Inc. All Rights Reserved.
+
+#include "Units/Highlevel/Hierarchy/RigUnit_SlideChain.h"
+#include "Math/ControlRigMathLibrary.h"
+#include "Units/RigUnitContext.h"
+
+void FRigUnit_SlideChain::Execute(const FRigUnitContext& Context)
+{
+	FRigHierarchy* Hierarchy = (FRigHierarchy*)(Context.HierarchyReference.Get());
+	if (Hierarchy == nullptr)
+	{
+		return;
+	}
+
+	if (Context.State == EControlRigState::Init)
+	{
+		BoneSegments.Reset();
+		BoneIndices.Reset();
+		Transforms.Reset();
+		BlendedTransforms.Reset();
+
+		ChainLength = 0.f;
+
+		int32 EndBoneIndex = Hierarchy->GetIndex(EndBone);
+		if (EndBoneIndex != INDEX_NONE)
+		{
+			int32 StartBoneIndex = Hierarchy->GetIndex(StartBone);
+			if (StartBoneIndex == EndBoneIndex)
+			{
+				return;
+			}
+
+			while (EndBoneIndex != INDEX_NONE)
+			{
+				BoneIndices.Add(EndBoneIndex);
+				if (EndBoneIndex == StartBoneIndex)
+				{
+					break;
+				}
+				EndBoneIndex = Hierarchy->GetParentIndex(EndBoneIndex);
+			}
+		}
+
+		if (BoneIndices.Num() < 2)
+		{
+			UE_CONTROLRIG_RIGUNIT_REPORT_WARNING(TEXT("Didn't find enough bones. You need at least two in the chain!"));
+			return;
+		}
+
+		Algo::Reverse(BoneIndices);
+
+		BoneSegments.SetNumZeroed(BoneIndices.Num());
+		BoneSegments[0] = 0;
+		for (int32 Index = 1; Index < BoneIndices.Num(); Index++)
+		{
+			FVector A = Hierarchy->GetGlobalTransform(BoneIndices[Index - 1]).GetLocation();
+			FVector B = Hierarchy->GetGlobalTransform(BoneIndices[Index]).GetLocation();
+			BoneSegments[Index] = (A - B).Size();
+			ChainLength += BoneSegments[Index];
+		}
+
+		Transforms.SetNum(BoneIndices.Num());
+		BlendedTransforms.SetNum(BoneIndices.Num());
+		return;
+	}
+
+	if (BoneIndices.Num() == 0 || ChainLength < SMALL_NUMBER)
+	{
+		return;
+	}
+
+	for (int32 Index = 0; Index < BoneIndices.Num(); Index++)
+	{
+		Transforms[Index] = Hierarchy->GetGlobalTransform(BoneIndices[Index]);
+	}
+
+	for (int32 Index = 0; Index < Transforms.Num(); Index++)
+	{
+		int32 TargetIndex = Index;
+		float Ratio = 0.f;
+		float SlidePerBone = -SlideAmount * ChainLength;
+
+		if (SlidePerBone > 0)
+		{
+			while (SlidePerBone > SMALL_NUMBER && TargetIndex < Transforms.Num() - 1)
+			{
+				TargetIndex++;
+				SlidePerBone -= BoneSegments[TargetIndex];
+			}
+		}
+		else
+		{
+			while (SlidePerBone < -SMALL_NUMBER && TargetIndex > 0)
+			{
+				SlidePerBone += BoneSegments[TargetIndex];
+				TargetIndex--;
+			}
+		}
+
+		if (TargetIndex < Transforms.Num() - 1)
+		{
+			if (BoneSegments[TargetIndex + 1] > SMALL_NUMBER)
+			{
+				if (SlideAmount < -SMALL_NUMBER)
+				{
+					Ratio = FMath::Clamp<float>(1.f - FMath::Abs<float>(SlidePerBone / BoneSegments[TargetIndex + 1]), 0.f, 1.f);
+				}
+				else
+				{
+					Ratio = FMath::Clamp<float>(SlidePerBone / BoneSegments[TargetIndex + 1], 0.f, 1.f);
+				}
+			}
+		}
+
+		BlendedTransforms[Index] = Transforms[TargetIndex];
+		if (TargetIndex < Transforms.Num() - 1 && Ratio > SMALL_NUMBER && Ratio < 1.f - SMALL_NUMBER)
+		{
+			BlendedTransforms[Index] = FControlRigMathLibrary::LerpTransform(BlendedTransforms[Index], Transforms[TargetIndex + 1], Ratio);
+		}
+	}
+
+	for (int32 Index = 0; Index < BoneIndices.Num(); Index++)
+	{
+		if (Index < BoneIndices.Num() - 1)
+		{
+			FVector CurrentX = BlendedTransforms[Index].GetRotation().GetAxisX();
+			FVector DesiredX = BlendedTransforms[Index + 1].GetLocation() - BlendedTransforms[Index].GetLocation();
+			FQuat OffsetQuat = FQuat::FindBetweenVectors(CurrentX, DesiredX);
+			BlendedTransforms[Index].SetRotation(OffsetQuat * BlendedTransforms[Index].GetRotation());
+		}
+		Hierarchy->SetGlobalTransform(BoneIndices[Index], BlendedTransforms[Index], bPropagateToChildren);
+	}
+}
diff --git a/Engine/Plugins/Experimental/ControlRig/Source/ControlRig/Private/Units/Highlevel/Hierarchy/RigUnit_SlideChain.h b/Engine/Plugins/Experimental/ControlRig/Source/ControlRig/Private/Units/Highlevel/Hierarchy/RigUnit_SlideChain.h
new file mode 100644
--- /dev/null
+++ b/Engine/Plugins/Experimental/ControlRig/Source/ControlRig/Private/Units/Highlevel/Hierarchy/RigUnit_SlideChain.h
@@ -0,0 +1,67 @@
+// Copyright 1998-2019 Epic Games, Inc. All Rights Reserved.
+
+#pragma once
+
+#include "Units/Highlevel/RigUnit_HighlevelBase.h"
+#include "Math/ControlRigMathLibrary.h"
+#include "RigUnit_SlideChain.generated.h"
+
+/**
+ * Slides an existing chain along itself with control over extrapolation.
+ */
+USTRUCT(meta=(DisplayName="Slide Chain", Category="Hierarchy", Keywords="Fit,Refit"))
+struct FRigUnit_SlideChain: public FRigUnit_HighlevelBaseMutable
+{
+	GENERATED_BODY()
+
+	FRigUnit_SlideChain()
+	{
+		StartBone = EndBone = NAME_None;
+		SlideAmount = 0.f;
+		bPropagateToChildren = false;
+		ChainLength = 0.f;
+	}
+
+	virtual void Execute(const FRigUnitContext& Context) override;
+
+	/** 
+	 * The name of the first bone to slide
+	 */
+	UPROPERTY(meta = (Input, Constant, BoneName))
+	FName StartBone;
+
+	/** 
+	 * The name of the last bone to slide
+	 */
+	UPROPERTY(meta = (Input, Constant, BoneName))
+	FName EndBone;
+
+	/** 
+	 * The amount of sliding. This unit is multiple of the chain length.
+	 */
+	UPROPERTY(meta = (Input))
+	float SlideAmount;
+
+	/**
+	 * If set to true all of the global transforms of the children
+	 * of this bone will be recalculated based on their local transforms.
+	 * Note: This is computationally more expensive than turning it off.
+	 */
+	UPROPERTY(meta = (Input))
+	bool bPropagateToChildren;
+
+	UPROPERTY()
+	float ChainLength;
+
+	UPROPERTY()
+	TArray<float> BoneSegments;
+
+	UPROPERTY()
+	TArray<int32> BoneIndices;
+
+	UPROPERTY()
+	TArray<FTransform> Transforms;
+
+	UPROPERTY()
+	TArray<FTransform> BlendedTransforms;
+};
