Prevent the ControlRig plugin from hijacking comment widget creation for other graph types
#rnx

diff --git a/Engine/Plugins/Experimental/ControlRig/Source/ControlRigEditor/Private/Graph/ControlRigGraphPanelNodeFactory.cpp b/Engine/Plugins/Experimental/ControlRig/Source/ControlRigEditor/Private/Graph/ControlRigGraphPanelNodeFactory.cpp
--- a/Engine/Plugins/Experimental/ControlRig/Source/ControlRigEditor/Private/Graph/ControlRigGraphPanelNodeFactory.cpp
+++ b/Engine/Plugins/Experimental/ControlRig/Source/ControlRigEditor/Private/Graph/ControlRigGraphPanelNodeFactory.cpp
@@ -6,6 +6,7 @@
 #include "EdGraphNode_Comment.h"
 #include "SControlRigGraphNode.h"
 #include "SControlRigGraphNodeComment.h"
+#include "Graph/ControlRigGraphSchema.h"
 
 TSharedPtr<SGraphNode> FControlRigGraphPanelNodeFactory::CreateNode(UEdGraphNode* Node) const
 {
@@ -19,13 +20,17 @@ TSharedPtr<SGraphNode> FControlRigGraphPanelNodeFactory::CreateNode(UEdGraphNode
 		ControlRigGraphNode->SetDimensions(GraphNode->GetDesiredSize());
 		return GraphNode;
 	}
+	
 	if (UEdGraphNode_Comment* CommentNode = Cast<UEdGraphNode_Comment>(Node))
 	{
-		TSharedRef<SGraphNode> GraphNode =
-			SNew(SControlRigGraphNodeComment, CommentNode);
+		if (CommentNode->GetSchema()->IsA(UControlRigGraphSchema::StaticClass()))
+		{
+			TSharedRef<SGraphNode> GraphNode =
+				SNew(SControlRigGraphNodeComment, CommentNode);
 
-		GraphNode->SlatePrepass();
-		return GraphNode;
+			GraphNode->SlatePrepass();
+			return GraphNode;
+		}
 	}
 
 	return nullptr;
