ControlRig: Support for animation data source + unit tests

diff --git a/Engine/Plugins/Experimental/ControlRig/Source/ControlRig/Private/ControlRig.cpp b/Engine/Plugins/Experimental/ControlRig/Source/ControlRig/Private/ControlRig.cpp
--- a/Engine/Plugins/Experimental/ControlRig/Source/ControlRig/Private/ControlRig.cpp
+++ b/Engine/Plugins/Experimental/ControlRig/Source/ControlRig/Private/ControlRig.cpp
@@ -53,6 +53,7 @@ UControlRig::UControlRig()
 	, ControlRigLog(nullptr)
 	, bEnableControlRigLogging(true)
 	, DrawInterface(nullptr)
+	, DataSourceRegistry(nullptr)
 #endif
 {
 #if DEBUG_CONTROLRIG_PROPERTYCHANGE
@@ -302,6 +303,12 @@ void UControlRig::Execute(const EControlRigState InState)
 
 	FRigUnitContext Context;
 	Context.DrawInterface = DrawInterface;
+
+	if (InState == EControlRigState::Init)
+	{
+		Context.DataSourceRegistry = DataSourceRegistry;
+	}
+
 	Context.DeltaTime = DeltaTime;
 	Context.State = InState;
 	Context.HierarchyReference.Container = &Hierarchy;
diff --git a/Engine/Plugins/Experimental/ControlRig/Source/ControlRig/Public/ControlRig.h b/Engine/Plugins/Experimental/ControlRig/Source/ControlRig/Public/ControlRig.h
--- a/Engine/Plugins/Experimental/ControlRig/Source/ControlRig/Public/ControlRig.h
+++ b/Engine/Plugins/Experimental/ControlRig/Source/ControlRig/Public/ControlRig.h
@@ -237,6 +237,10 @@ private:
 	/** The draw interface for the units to use */
 	FControlRigDrawInterface* DrawInterface;
 
+	/** The registry to access data source */
+	UPROPERTY(transient)
+	UAnimationDataSourceRegistry* DataSourceRegistry;
+
 #if DEBUG_CONTROLRIG_PROPERTYCHANGE
 	// This is to debug class size when constructed and destroyed to verify match
 	// if this size changes, that implies more problem, where properties have been changed and layout has been modified
diff --git a/Engine/Plugins/Experimental/ControlRig/Source/ControlRig/Public/Units/RigUnitContext.h b/Engine/Plugins/Experimental/ControlRig/Source/ControlRig/Public/Units/RigUnitContext.h
--- a/Engine/Plugins/Experimental/ControlRig/Source/ControlRig/Public/Units/RigUnitContext.h
+++ b/Engine/Plugins/Experimental/ControlRig/Source/ControlRig/Public/Units/RigUnitContext.h
@@ -5,6 +5,7 @@
 #include "CoreMinimal.h"
 #include "Hierarchy.h"
 #include "ControlRigLog.h"
+#include "AnimationDataSource.h"
 #include "Drawing/ControlRigDrawInterface.h"
 
 /** Current state of rig
@@ -24,6 +25,7 @@ struct FRigUnitContext
 	/** default constructor */
 	FRigUnitContext()
 		: DrawInterface(nullptr)
+		, DataSourceRegistry(nullptr)
 #if WITH_EDITOR
 		, Log(nullptr)
 #endif
@@ -34,6 +36,9 @@ struct FRigUnitContext
 	/** The draw interface for the units to use */
 	FControlRigDrawInterface* DrawInterface;
 
+	/** The registry to access data source */
+	const UAnimationDataSourceRegistry* DataSourceRegistry;
+
 	/** The current delta time */
 	float DeltaTime;
 
@@ -47,6 +52,22 @@ struct FRigUnitContext
 	/** A handle to the compiler log */
 	FControlRigLog* Log;
 #endif
+
+	/**
+	 * Returns a given data source and cast it to the expected class.
+	 *
+	 * @param InName The name of the data source to look up.
+	 * @return The requested data source
+	 */
+	template<class T>
+	T* RequestDataSource(const FName& InName) const
+	{
+		if (DataSourceRegistry == nullptr)
+		{
+			return nullptr;
+		}
+		return DataSourceRegistry->RequestSource<T>(InName);
+	}
 };
 
 #if WITH_EDITOR
