Rename/move file(s)

Moved hierarchy into rigs subfolder

diff --git a/Engine/Plugins/Experimental/ControlRig/Source/ControlRig/Private/Rigs/Hierarchy.cpp b/Engine/Plugins/Experimental/ControlRig/Source/ControlRig/Private/Rigs/Hierarchy.cpp
--- a/Engine/Plugins/Experimental/ControlRig/Source/ControlRig/Private/Rigs/Hierarchy.cpp
+++ b/Engine/Plugins/Experimental/ControlRig/Source/ControlRig/Private/Rigs/Hierarchy.cpp
@@ -1,6 +1,6 @@
 // Copyright 1998-2019 Epic Games, Inc. All Rights Reserved.
 
-#include "Hierarchy.h"
+#include "Rigs/Hierarchy.h"
 #include "ControlRig.h"
 #include "HelperUtil.h"
 
diff --git a/Engine/Plugins/Experimental/ControlRig/Source/ControlRig/Private/Units/Deprecated/RigUnit_AimConstraint.h b/Engine/Plugins/Experimental/ControlRig/Source/ControlRig/Private/Units/Deprecated/RigUnit_AimConstraint.h
--- a/Engine/Plugins/Experimental/ControlRig/Source/ControlRig/Private/Units/Deprecated/RigUnit_AimConstraint.h
+++ b/Engine/Plugins/Experimental/ControlRig/Source/ControlRig/Private/Units/Deprecated/RigUnit_AimConstraint.h
@@ -3,7 +3,7 @@
 #pragma once
 
 #include "Units/RigUnit.h"
-#include "Hierarchy.h"
+#include "Rigs/Hierarchy.h"
 #include "Constraint.h"
 #include "ControlRigDefines.h"
 #include "RigUnit_AimConstraint.generated.h"
diff --git a/Engine/Plugins/Experimental/ControlRig/Source/ControlRig/Private/Units/Deprecated/RigUnit_ApplyFK.h b/Engine/Plugins/Experimental/ControlRig/Source/ControlRig/Private/Units/Deprecated/RigUnit_ApplyFK.h
--- a/Engine/Plugins/Experimental/ControlRig/Source/ControlRig/Private/Units/Deprecated/RigUnit_ApplyFK.h
+++ b/Engine/Plugins/Experimental/ControlRig/Source/ControlRig/Private/Units/Deprecated/RigUnit_ApplyFK.h
@@ -3,7 +3,7 @@
 #pragma once
 
 #include "Units/RigUnit.h"
-#include "Hierarchy.h"
+#include "Rigs/Hierarchy.h"
 #include "Constraint.h"
 #include "ControlRigDefines.h"
 #include "RigUnit_ApplyFK.generated.h"
diff --git a/Engine/Plugins/Experimental/ControlRig/Source/ControlRig/Private/Units/Deprecated/RigUnit_CreateHierarchy.h b/Engine/Plugins/Experimental/ControlRig/Source/ControlRig/Private/Units/Deprecated/RigUnit_CreateHierarchy.h
--- a/Engine/Plugins/Experimental/ControlRig/Source/ControlRig/Private/Units/Deprecated/RigUnit_CreateHierarchy.h
+++ b/Engine/Plugins/Experimental/ControlRig/Source/ControlRig/Private/Units/Deprecated/RigUnit_CreateHierarchy.h
@@ -3,7 +3,7 @@
 #pragma once
 
 #include "Units/RigUnit.h"
-#include "Hierarchy.h"
+#include "Rigs/Hierarchy.h"
 #include "RigUnit_CreateHierarchy.generated.h"
 
 USTRUCT(meta=(DisplayName="Create Hierarchy", Category="Hierarchy", Deprecated = "4.23.0"))
diff --git a/Engine/Plugins/Experimental/ControlRig/Source/ControlRig/Private/Units/Deprecated/RigUnit_GetJointTransform.h b/Engine/Plugins/Experimental/ControlRig/Source/ControlRig/Private/Units/Deprecated/RigUnit_GetJointTransform.h
--- a/Engine/Plugins/Experimental/ControlRig/Source/ControlRig/Private/Units/Deprecated/RigUnit_GetJointTransform.h
+++ b/Engine/Plugins/Experimental/ControlRig/Source/ControlRig/Private/Units/Deprecated/RigUnit_GetJointTransform.h
@@ -3,7 +3,7 @@
 #pragma once
 
 #include "Units/RigUnit.h"
-#include "Hierarchy.h"
+#include "Rigs/Hierarchy.h"
 #include "Constraint.h"
 #include "ControlRigDefines.h"
 #include "RigUnit_GetJointTransform.generated.h"
diff --git a/Engine/Plugins/Experimental/ControlRig/Source/ControlRig/Private/Units/Deprecated/RigUnit_MergeHierarchy.h b/Engine/Plugins/Experimental/ControlRig/Source/ControlRig/Private/Units/Deprecated/RigUnit_MergeHierarchy.h
--- a/Engine/Plugins/Experimental/ControlRig/Source/ControlRig/Private/Units/Deprecated/RigUnit_MergeHierarchy.h
+++ b/Engine/Plugins/Experimental/ControlRig/Source/ControlRig/Private/Units/Deprecated/RigUnit_MergeHierarchy.h
@@ -3,7 +3,7 @@
 #pragma once
 
 #include "Units/RigUnit.h"
-#include "Hierarchy.h"
+#include "Rigs/Hierarchy.h"
 #include "RigUnit_MergeHierarchy.generated.h"
 
 USTRUCT(meta=(DisplayName="Merge Hierarchy", Category="Hierarchy", Deprecated = "4.23.0"))
diff --git a/Engine/Plugins/Experimental/ControlRig/Source/ControlRig/Private/Units/Highlevel/Hierarchy/RigUnit_TransformConstraint.h b/Engine/Plugins/Experimental/ControlRig/Source/ControlRig/Private/Units/Highlevel/Hierarchy/RigUnit_TransformConstraint.h
--- a/Engine/Plugins/Experimental/ControlRig/Source/ControlRig/Private/Units/Highlevel/Hierarchy/RigUnit_TransformConstraint.h
+++ b/Engine/Plugins/Experimental/ControlRig/Source/ControlRig/Private/Units/Highlevel/Hierarchy/RigUnit_TransformConstraint.h
@@ -4,7 +4,7 @@
 
 #include "Units/RigUnit.h"
 #include "Units/Highlevel/RigUnit_HighlevelBase.h"
-#include "Hierarchy.h"
+#include "Rigs/Hierarchy.h"
 #include "Constraint.h"
 #include "ControlRigDefines.h"
 #include "RigUnit_TransformConstraint.generated.h"
diff --git a/Engine/Plugins/Experimental/ControlRig/Source/ControlRig/Public/ControlRig.h b/Engine/Plugins/Experimental/ControlRig/Source/ControlRig/Public/ControlRig.h
--- a/Engine/Plugins/Experimental/ControlRig/Source/ControlRig/Public/ControlRig.h
+++ b/Engine/Plugins/Experimental/ControlRig/Source/ControlRig/Public/ControlRig.h
@@ -9,7 +9,7 @@
 #include "Engine/EngineBaseTypes.h"
 #include "Templates/SubclassOf.h"
 #include "ControlRigDefines.h"
-#include "Hierarchy.h"
+#include "Rigs/Hierarchy.h"
 #include "Units/RigUnitContext.h"
 #include "Animation/NodeMappingProviderInterface.h"
 #include "Units/RigUnit.h"
diff --git a/Engine/Plugins/Experimental/ControlRig/Source/ControlRig/Public/ControlRigDefines.h b/Engine/Plugins/Experimental/ControlRig/Source/ControlRig/Public/ControlRigDefines.h
--- a/Engine/Plugins/Experimental/ControlRig/Source/ControlRig/Public/ControlRigDefines.h
+++ b/Engine/Plugins/Experimental/ControlRig/Source/ControlRig/Public/ControlRigDefines.h
@@ -3,7 +3,7 @@
 #pragma once
 
 #include "PropertyPathHelpers.h"
-#include "Hierarchy.h"
+#include "Rigs/Hierarchy.h"
 #include "ControlRigDefines.generated.h"
 
 USTRUCT()
diff --git a/Engine/Plugins/Experimental/ControlRig/Source/ControlRig/Public/Drawing/ControlRigDrawInterface.h b/Engine/Plugins/Experimental/ControlRig/Source/ControlRig/Public/Drawing/ControlRigDrawInterface.h
--- a/Engine/Plugins/Experimental/ControlRig/Source/ControlRig/Public/Drawing/ControlRigDrawInterface.h
+++ b/Engine/Plugins/Experimental/ControlRig/Source/ControlRig/Public/Drawing/ControlRigDrawInterface.h
@@ -3,7 +3,7 @@
 #pragma once
 
 #include "CoreMinimal.h"
-#include "Hierarchy.h"
+#include "Rigs/Hierarchy.h"
 #include "Math/Simulation/CRSimPointContainer.h"
 #include "ControlRigDrawInterface.generated.h"
 
diff --git a/Engine/Plugins/Experimental/ControlRig/Source/ControlRig/Public/Hierarchy.h b/Engine/Plugins/Experimental/ControlRig/Source/ControlRig/Public/Hierarchy.h
deleted file mode 100644
diff --git a/Engine/Plugins/Experimental/ControlRig/Source/ControlRig/Public/Rigs/Hierarchy.h b/Engine/Plugins/Experimental/ControlRig/Source/ControlRig/Public/Rigs/Hierarchy.h
new file mode 100644
--- /dev/null
+++ b/Engine/Plugins/Experimental/ControlRig/Source/ControlRig/Public/Rigs/Hierarchy.h
@@ -0,0 +1,495 @@
+// Copyright 1998-2019 Epic Games, Inc. All Rights Reserved.
+
+#pragma once
+
+#include "CoreMinimal.h"
+#include "Hierarchy.generated.h"
+
+class UControlRig;
+
+USTRUCT(BlueprintType)
+struct FRigBone
+{
+	GENERATED_BODY()
+
+	UPROPERTY(VisibleAnywhere, Category = FRigHierarchy)
+	FName Name;
+
+	UPROPERTY(VisibleAnywhere, Category = FRigHierarchy)
+	FName ParentName;
+
+	UPROPERTY(transient)
+	int32 ParentIndex;
+
+	/* Initial global transform that is saved in this rig */
+	UPROPERTY(EditAnywhere, Category = FRigHierarchy)
+	FTransform InitialTransform;
+
+	UPROPERTY(transient, VisibleAnywhere, Category = FRigHierarchy)
+	FTransform GlobalTransform;
+
+	UPROPERTY(transient, VisibleAnywhere, Category = FRigHierarchy)
+	FTransform LocalTransform;
+
+	/** dependent list - direct dependent for child or anything that needs to update due to this */
+	UPROPERTY(transient)
+	TArray<int32> Dependents;
+};
+
+USTRUCT()
+struct CONTROLRIG_API FRigHierarchy
+{
+	GENERATED_BODY()
+
+	UPROPERTY(EditAnywhere, Category = FRigHierarchy)
+	TArray<FRigBone> Bones;
+
+	// can serialize fine? 
+	UPROPERTY()
+	TMap<FName, int32> NameToIndexMapping;
+
+	void AddBone(const FName& NewBoneName, const FName& Parent, const FTransform& InitTransform)
+	{
+		int32 ParentIndex = GetIndex(Parent);
+		bool bHasParent = (ParentIndex != INDEX_NONE);
+
+		FRigBone NewBone;
+		NewBone.Name = NewBoneName;
+		NewBone.ParentIndex = ParentIndex;
+		NewBone.ParentName = bHasParent? Parent : NAME_None;
+		NewBone.InitialTransform = InitTransform;
+		NewBone.GlobalTransform = InitTransform;
+		RecalculateLocalTransform(NewBone);
+
+		Bones.Add(NewBone);
+		RefreshMapping();
+	}
+
+	void AddBone(const FName& NewBoneName, const FName& Parent, const FTransform& InitTransform, const FTransform& LocalTransform, const FTransform& GlobalTransform)
+	{
+		AddBone(NewBoneName, Parent, InitTransform);
+
+		int32 NewIndex = GetIndex(NewBoneName);
+		Bones[NewIndex].LocalTransform = LocalTransform;
+		Bones[NewIndex].GlobalTransform = GlobalTransform;
+	}
+
+	void Reparent(const FName& InBone, const FName& NewParent)
+	{
+		int32 Index = GetIndex(InBone);
+		// can't parent to itself
+		if (Index != INDEX_NONE && InBone != NewParent)
+		{
+			// should allow reparent to none (no parent)
+			// if invalid, we consider to be none
+			int32 ParentIndex = GetIndex(NewParent);
+			bool bHasParent = (ParentIndex != INDEX_NONE);
+			FRigBone& CurBone = Bones[Index];
+			CurBone.ParentIndex = ParentIndex;
+			CurBone.ParentName = (bHasParent)? NewParent : NAME_None;
+			RecalculateLocalTransform(CurBone);
+
+			// we want to make sure parent is before the child
+			RefreshMapping();
+		}
+	}
+
+	void DeleteBone(const FName& BoneToDelete, bool bIncludeChildren)
+	{
+		TArray<int32> Children;
+		if (GetChildren(BoneToDelete, Children, true) > 0)
+		{
+			// sort by child index
+			Children.Sort([](const int32& A, const int32& B) { return A < B; });
+
+			// want to delete from end to the first 
+			for (int32 ChildIndex = Children.Num() - 1; ChildIndex >= 0; --ChildIndex)
+			{
+				Bones.RemoveAt(Children[ChildIndex]);
+			}
+		}
+
+		// in theory, since I'm not adding new element here
+		// it is safe to do search using FindIndex, but it may cause
+		// difficulty in the future, so I'm changing to FindIndexSlow
+		// note that we're removing units here, but the index is removed from later
+		// to begin, so in theory it should be fine
+		int32 IndexToDelete = GetIndex(BoneToDelete);
+		Bones.RemoveAt(IndexToDelete);
+
+		RefreshMapping();
+	}
+
+	FName GetParentName(const FName& InBone) const
+	{
+		int32 Index = GetIndex(InBone);
+		if (Index != INDEX_NONE)
+		{
+			return Bones[Index].ParentName;
+		}
+
+		return NAME_None;
+	}
+
+	int32 GetParentIndex(const int32 BoneIndex) const
+	{
+		if (BoneIndex != INDEX_NONE)
+		{
+			return Bones[BoneIndex].ParentIndex;
+		}
+
+		return INDEX_NONE;
+	}
+	// list of names of children - this is not cheap, and is supposed to be used only for one time set up
+	int32 GetChildren(const FName& InBone, TArray<int32>& OutChildren, bool bRecursively) const
+	{
+		return GetChildren(GetIndex(InBone), OutChildren, bRecursively);
+	}
+
+	int32 GetChildren(const int32 InBoneIndex, TArray<int32>& OutChildren, bool bRecursively) const
+	{
+		OutChildren.Reset();
+
+		if (InBoneIndex != INDEX_NONE)
+		{
+			GetChildrenRecursive(InBoneIndex, OutChildren, bRecursively);
+		}
+
+		return OutChildren.Num();
+	}
+
+	FName GetName(int32 Index) const
+	{
+		if (Bones.IsValidIndex(Index))
+		{
+			return Bones[Index].Name;
+		}
+
+		return NAME_None;
+	}
+
+	int32 GetIndex(const FName& Bone) const
+	{
+		// ensure if it does not match
+		//ensureAlways(Bones.Num() == NameToIndexMapping.Num());
+
+		const int32* Index = NameToIndexMapping.Find(Bone);
+		if (Index)
+		{
+			return *Index;
+		}
+
+		return INDEX_NONE;
+	}
+
+//#if WITH_EDITOR
+	// @FIXMELINA: figure out how to remove this outside of editor
+	// ignore mapping, run slow search
+	// this is useful in editor while editing
+	// we don't want to build mapping data every time
+	int32 GetIndexSlow(const FName& Bone) const
+	{
+		for (int32 BoneId = 0; BoneId < Bones.Num(); ++BoneId)
+		{
+			if (Bones[BoneId].Name == Bone)
+			{
+				return BoneId;
+			}
+		}
+
+		return INDEX_NONE;
+	}
+//#endif // WITH_EDITOR
+	void SetGlobalTransform(const FName& Bone, const FTransform& InTransform, bool bPropagateTransform = true)
+	{
+		SetGlobalTransform(GetIndex(Bone), InTransform, bPropagateTransform);
+	}
+
+	void SetGlobalTransform(int32 Index, const FTransform& InTransform, bool bPropagateTransform = true)
+	{
+		if (Bones.IsValidIndex(Index))
+		{
+			FRigBone& Bone = Bones[Index];
+			Bone.GlobalTransform = InTransform;
+			Bone.GlobalTransform.NormalizeRotation();
+			RecalculateLocalTransform(Bone);
+
+			if (bPropagateTransform)
+			{
+				PropagateTransform(Index);
+			}
+		}
+	}
+
+	FTransform GetGlobalTransform(const FName& Bone) const
+	{
+		return GetGlobalTransform(GetIndex(Bone));
+	}
+
+	FTransform GetGlobalTransform(int32 Index) const
+	{
+		if (Bones.IsValidIndex(Index))
+		{
+			return Bones[Index].GlobalTransform;
+		}
+
+		return FTransform::Identity;
+	}
+
+	void SetLocalTransform(const FName& Bone, const FTransform& InTransform, bool bPropagateTransform = true)
+	{
+		SetLocalTransform(GetIndex(Bone), InTransform, bPropagateTransform);
+	}
+
+	void SetLocalTransform(int32 Index, const FTransform& InTransform, bool bPropagateTransform = true)
+	{
+		if (Bones.IsValidIndex(Index))
+		{
+			FRigBone& Bone = Bones[Index];
+			Bone.LocalTransform = InTransform;
+			RecalculateGlobalTransform(Bone);
+
+			if (bPropagateTransform)
+			{
+				PropagateTransform(Index);
+			}
+		}
+	}
+
+	FTransform GetLocalTransform(const FName& Bone) const
+	{
+		return GetLocalTransform(GetIndex(Bone));
+	}
+
+	FTransform GetLocalTransform(int32 Index) const
+	{
+		if (Bones.IsValidIndex(Index))
+		{
+			return Bones[Index].LocalTransform;
+		}
+
+		return FTransform::Identity;
+	}
+
+	void SetInitialTransform(const FName& Bone, const FTransform& InTransform)
+	{
+		SetInitialTransform(GetIndex(Bone), InTransform);
+	}
+
+	void SetInitialTransform(int32 Index, const FTransform& InTransform)
+	{
+		if (Bones.IsValidIndex(Index))
+		{
+			FRigBone& Bone = Bones[Index];
+			Bone.InitialTransform = InTransform;
+			Bone.InitialTransform.NormalizeRotation();
+			RecalculateLocalTransform(Bone);
+		}
+	}
+
+	FTransform GetInitialTransform(const FName& Bone) const
+	{
+		return GetInitialTransform(GetIndex(Bone));
+	}
+
+	FTransform GetInitialTransform(int32 Index) const
+	{
+		if (Bones.IsValidIndex(Index))
+		{
+			return Bones[Index].InitialTransform;
+		}
+
+		return FTransform::Identity;
+	}
+	// @todo: move to private
+	void RecalculateLocalTransform(FRigBone& InOutBone)
+	{
+		bool bHasParent = InOutBone.ParentIndex != INDEX_NONE;
+		InOutBone.LocalTransform = (bHasParent) ? InOutBone.GlobalTransform.GetRelativeTransform(Bones[InOutBone.ParentIndex].GlobalTransform) : InOutBone.GlobalTransform;
+	}
+
+	// @todo: move to private
+	void RecalculateGlobalTransform(FRigBone& InOutBone)
+	{
+		bool bHasParent = InOutBone.ParentIndex != INDEX_NONE;
+		InOutBone.GlobalTransform = (bHasParent) ? InOutBone.LocalTransform * Bones[InOutBone.ParentIndex].GlobalTransform : InOutBone.LocalTransform;
+	}
+
+	void Rename(const FName& OldName, const FName& NewName)
+	{
+		if (OldName != NewName)
+		{
+			const int32 Found = GetIndex(OldName);
+			if (Found != INDEX_NONE)
+			{
+				Bones[Found].Name = NewName;
+
+				// go through find all children and rename them
+				for (int32 Index = 0; Index < Bones.Num(); ++Index)
+				{
+					if (Bones[Index].ParentName == OldName)
+					{
+						Bones[Index].ParentName = NewName;
+					}
+				}
+
+				RefreshMapping();
+			}
+		}
+	}
+
+	// updates all of the internal caches
+	void Initialize();
+
+	// clears the hierarchy and removes all content
+	void Reset();
+
+	// resets all of the transforms back to the initial transform
+	void ResetTransforms();
+
+	int32 GetNum() const
+	{
+		return Bones.Num();
+	}
+private:
+	void RefreshMapping();
+	void Sort();
+
+	// list of names of children - this is not cheap, and is supposed to be used only for one time set up
+	int32 GetChildrenRecursive(const int32 InBoneIndex, TArray<int32>& OutChildren, bool bRecursively) const
+	{
+		const int32 StartChildIndex = OutChildren.Num();
+
+		// all children should be later than parent
+		for (int32 ChildIndex = InBoneIndex + 1; ChildIndex < Bones.Num(); ++ChildIndex)
+		{
+			if (Bones[ChildIndex].ParentIndex == InBoneIndex)
+			{
+				OutChildren.AddUnique(ChildIndex);
+			}
+		}
+
+		if (bRecursively)
+		{
+			// since we keep appending inside of functions, we make sure not to go over original list
+			const int32 EndChildIndex = OutChildren.Num() - 1;
+			for (int32 ChildIndex = StartChildIndex; ChildIndex <= EndChildIndex; ++ChildIndex)
+			{
+				GetChildrenRecursive(OutChildren[ChildIndex], OutChildren, bRecursively);
+			}
+		}
+
+		return OutChildren.Num();
+	}
+
+	void PropagateTransform(int32 BoneIndex);
+};
+
+USTRUCT()
+struct FRigHierarchyContainer
+{
+	GENERATED_BODY()
+
+	// index to hierarchy
+	TMap<FName, uint32> MapContainer;
+
+	// list of hierarchies
+	TArray<FRigHierarchy> Hierarchies;
+
+	// base hierarchy
+	// this should serialize
+	UPROPERTY()
+	FRigHierarchy BaseHierarchy;
+
+	FRigHierarchy* Find(const FName& InName)
+	{
+		uint32* IndexPtr = MapContainer.Find(InName);
+		if (IndexPtr && Hierarchies.IsValidIndex(*IndexPtr))
+		{
+			return &Hierarchies[*IndexPtr];
+		}
+
+		return nullptr;
+	}
+
+	void Reset()
+	{
+		BaseHierarchy.Reset();
+
+		// @todo reset whole hierarhcy
+	}
+
+	void ResetTransforms()
+	{
+		BaseHierarchy.ResetTransforms();
+
+		// @todo reset whole hierarhcy
+	}
+};
+
+USTRUCT()
+struct FRigHierarchyRef
+{
+	GENERATED_BODY()
+
+	FRigHierarchyRef()
+		: Container(nullptr)
+		, bUseBaseHierarchy(true)
+	{
+
+	}
+
+	FRigHierarchy* Get() 
+	{
+		return GetInternal();
+	}
+
+	const FRigHierarchy* Get() const
+	{
+		return (const FRigHierarchy*)GetInternal();
+	}
+
+	FRigHierarchy* Find(const FName& InName)
+	{
+		if(Container)
+		{
+			return Container->Find(InName);
+		}
+		return nullptr;
+	}
+
+	// create name if the name isn't set, for now it only accepts root
+	bool CreateHierarchy(const FName& RootName, const FRigHierarchyRef& SourceHierarchyRef);
+	bool MergeHierarchy(const FRigHierarchyRef& SourceHierarchyRef);
+
+private:
+	struct FRigHierarchyContainer* Container;
+
+	// @todo: this only works with merge right now. Should fix for all cases
+	UPROPERTY(EditAnywhere, Category = FRigHierarchyRef)
+	bool bUseBaseHierarchy; 
+
+	/** Name of Hierarchy */
+	UPROPERTY(EditAnywhere, Category = FRigHierarchyRef, meta=(EditCondition = "!bUseBaseHierarchy"))
+	FName Name;
+
+	// utility functions
+	bool CreateHierarchy(const FName& RootName, const FRigHierarchy* SourceHierarchy);
+	bool MergeHierarchy(const FRigHierarchy* InSourceHierarchy);
+
+	FRigHierarchy* GetInternal() const
+	{
+		if (Container)
+		{
+			if (bUseBaseHierarchy)
+			{
+				return &Container->BaseHierarchy;
+			}
+
+			return Container->Find(Name);
+		}
+		return nullptr;
+	}
+	friend class UControlRig;
+	friend class FControlRigUnitTestBase;
+};
\ No newline at end of file
diff --git a/Engine/Plugins/Experimental/ControlRig/Source/ControlRig/Public/Units/Deprecated/RigUnit_TwoBoneIKFK.h b/Engine/Plugins/Experimental/ControlRig/Source/ControlRig/Public/Units/Deprecated/RigUnit_TwoBoneIKFK.h
--- a/Engine/Plugins/Experimental/ControlRig/Source/ControlRig/Public/Units/Deprecated/RigUnit_TwoBoneIKFK.h
+++ b/Engine/Plugins/Experimental/ControlRig/Source/ControlRig/Public/Units/Deprecated/RigUnit_TwoBoneIKFK.h
@@ -3,7 +3,7 @@
 #pragma once
 
 #include "Units/RigUnit.h"
-#include "Hierarchy.h"
+#include "Rigs/Hierarchy.h"
 #include "ControlRigDefines.h"
 #include "RigUnit_TwoBoneIKFK.generated.h"
 
diff --git a/Engine/Plugins/Experimental/ControlRig/Source/ControlRig/Public/Units/RigUnitContext.h b/Engine/Plugins/Experimental/ControlRig/Source/ControlRig/Public/Units/RigUnitContext.h
--- a/Engine/Plugins/Experimental/ControlRig/Source/ControlRig/Public/Units/RigUnitContext.h
+++ b/Engine/Plugins/Experimental/ControlRig/Source/ControlRig/Public/Units/RigUnitContext.h
@@ -3,7 +3,7 @@
 #pragma once
 
 #include "CoreMinimal.h"
-#include "Hierarchy.h"
+#include "Rigs/Hierarchy.h"
 #include "ControlRigLog.h"
 #include "AnimationDataSource.h"
 #include "Drawing/ControlRigDrawInterface.h"
diff --git a/Engine/Plugins/Experimental/ControlRig/Source/ControlRig/Public/Units/RigUnitTest.h b/Engine/Plugins/Experimental/ControlRig/Source/ControlRig/Public/Units/RigUnitTest.h
--- a/Engine/Plugins/Experimental/ControlRig/Source/ControlRig/Public/Units/RigUnitTest.h
+++ b/Engine/Plugins/Experimental/ControlRig/Source/ControlRig/Public/Units/RigUnitTest.h
@@ -5,7 +5,7 @@
 #include "Misc/AutomationTest.h"
 #include "RigUnit.h"
 #include "RigUnitContext.h"
-#include "Hierarchy.h"
+#include "Rigs/Hierarchy.h"
 
 class FControlRigUnitTestBase : public FAutomationTestBase
 {
diff --git a/Engine/Plugins/Experimental/ControlRig/Source/ControlRigDeveloper/Public/ControlRigBlueprint.h b/Engine/Plugins/Experimental/ControlRig/Source/ControlRigDeveloper/Public/ControlRigBlueprint.h
--- a/Engine/Plugins/Experimental/ControlRig/Source/ControlRigDeveloper/Public/ControlRigBlueprint.h
+++ b/Engine/Plugins/Experimental/ControlRig/Source/ControlRigDeveloper/Public/ControlRigBlueprint.h
@@ -8,7 +8,7 @@
 #include "Engine/Blueprint.h"
 #include "Misc/Crc.h"
 #include "ControlRigDefines.h"
-#include "Hierarchy.h"
+#include "Rigs/Hierarchy.h"
 #include "Interfaces/Interface_PreviewMeshProvider.h"
 #include "ControlRigController.h"
 #include "ControlRigBlueprint.generated.h"
diff --git a/Engine/Plugins/Experimental/ControlRig/Source/ControlRigDeveloper/Public/Graph/ControlRigGraph.h b/Engine/Plugins/Experimental/ControlRig/Source/ControlRigDeveloper/Public/Graph/ControlRigGraph.h
--- a/Engine/Plugins/Experimental/ControlRig/Source/ControlRigDeveloper/Public/Graph/ControlRigGraph.h
+++ b/Engine/Plugins/Experimental/ControlRig/Source/ControlRigDeveloper/Public/Graph/ControlRigGraph.h
@@ -4,7 +4,7 @@
 
 #include "EdGraph/EdGraph.h"
 #include "Graph/ControlRigGraphNode.h"
-#include "Hierarchy.h"
+#include "Rigs/Hierarchy.h"
 #include "ControlRigModel.h"
 #include "ControlRigGraph.generated.h"
 
diff --git a/Engine/Plugins/Experimental/ControlRig/Source/ControlRigEditor/Private/Editor/SRigHierarchy.h b/Engine/Plugins/Experimental/ControlRig/Source/ControlRigEditor/Private/Editor/SRigHierarchy.h
--- a/Engine/Plugins/Experimental/ControlRig/Source/ControlRigEditor/Private/Editor/SRigHierarchy.h
+++ b/Engine/Plugins/Experimental/ControlRig/Source/ControlRigEditor/Private/Editor/SRigHierarchy.h
@@ -4,7 +4,7 @@
 
 #include "CoreMinimal.h"
 #include "Widgets/Views/STreeView.h"
-#include "Hierarchy.h"
+#include "Rigs/Hierarchy.h"
 #include "EditorUndoClient.h"
 #include "DragAndDrop/GraphNodeDragDropOp.h"
 
