Copying //UE4/Dev-Anim to Dev-Main (//UE4/Dev-Main) @ 8782600
#rnx

diff --git a/Engine/Plugins/Experimental/ControlRig/ControlRig.uplugin b/Engine/Plugins/Experimental/ControlRig/ControlRig.uplugin
--- a/Engine/Plugins/Experimental/ControlRig/ControlRig.uplugin
+++ b/Engine/Plugins/Experimental/ControlRig/ControlRig.uplugin
@@ -30,6 +30,11 @@
 			"Name" : "ControlRigEditor",
 			"Type" : "Editor",
 			"LoadingPhase" : "PostEngineInit"
+		},
+		{
+			"Name" : "ControlRigManipulation",
+			"Type" : "Editor",
+			"LoadingPhase" : "PostEngineInit"
 		}
 	]
 }
\ No newline at end of file
diff --git a/Engine/Plugins/Experimental/ControlRig/Source/ControlRig/Private/ControlRig.cpp b/Engine/Plugins/Experimental/ControlRig/Source/ControlRig/Private/ControlRig.cpp
--- a/Engine/Plugins/Experimental/ControlRig/Source/ControlRig/Private/ControlRig.cpp
+++ b/Engine/Plugins/Experimental/ControlRig/Source/ControlRig/Private/ControlRig.cpp
@@ -16,7 +16,7 @@
 
 #define LOCTEXT_NAMESPACE "ControlRig"
 
-DEFINE_LOG_CATEGORY_STATIC(LogControlRig, Log, All);
+DEFINE_LOG_CATEGORY(LogControlRig);
 
 DECLARE_STATS_GROUP(TEXT("ControlRig"), STATGROUP_ControlRig, STATCAT_Advanced);
 DECLARE_CYCLE_STAT_EXTERN(TEXT("Control Rig Execution"), STAT_RigExecution, STATGROUP_ControlRig, );
@@ -31,6 +31,8 @@ const FName UControlRig::DisplayNameMetaName("DisplayName");
 const FName UControlRig::MenuDescSuffixMetaName("MenuDescSuffix");
 const FName UControlRig::ShowVariableNameInTitleMetaName("ShowVariableNameInTitle");
 const FName UControlRig::BoneNameMetaName("BoneName");
+const FName UControlRig::ControlNameMetaName("ControlName");
+const FName UControlRig::SpaceNameMetaName("SpaceName");
 const FName UControlRig::CurveNameMetaName("CurveName");
 const FName UControlRig::ConstantMetaName("Constant");
 const FName UControlRig::TitleColorMetaName("TitleColor");
@@ -46,14 +48,12 @@ const FName UControlRig::DefaultArraySizeMetaName("DefaultArraySize");
 
 UControlRig::UControlRig()
 	: DeltaTime(0.0f)
-#if WITH_EDITORONLY_DATA
-	, bExecutionOn (true)
-#endif // WITH_EDITORONLY_DATA
 	, ExecutionType(ERigExecutionType::Runtime)
 #if WITH_EDITOR
 	, ControlRigLog(nullptr)
 	, bEnableControlRigLogging(true)
 	, DrawInterface(nullptr)
+	, DataSourceRegistry(nullptr)
 #endif
 {
 #if DEBUG_CONTROLRIG_PROPERTYCHANGE
@@ -198,8 +198,8 @@ void UControlRig::Initialize(bool bInitRigUnits)
 #endif // WITH_EDITORONLY_DATA
 
 	// should refresh mapping 
-	Hierarchy.BaseHierarchy.Initialize();
-	CurveContainer.Initialize();
+	Hierarchy.Initialize();
+
 	// resolve IO properties
 	ResolveInputOutputProperties();
 
@@ -221,15 +221,6 @@ void UControlRig::InitializeFromCDO()
 	UControlRig* CDO = GetClass()->GetDefaultObject<UControlRig>();
 	// copy operation
 	Hierarchy = CDO->Hierarchy;
-	CurveContainer = CDO->CurveContainer;
-}
-
-void UControlRig::PreEvaluate_GameThread()
-{
-	DECLARE_SCOPE_HIERARCHICAL_COUNTER_FUNC()
-
-	// input delegates
-	OnPreEvaluateGatherInput.ExecuteIfBound(this);
 }
 
 void UControlRig::Evaluate_AnyThread()
@@ -239,14 +230,6 @@ void UControlRig::Evaluate_AnyThread()
 	Execute(EControlRigState::Update);
 }
 
-void UControlRig::PostEvaluate_GameThread()
-{
-	DECLARE_SCOPE_HIERARCHICAL_COUNTER_FUNC()
-
-	// output delgates
-	OnPostEvaluateQueryOutput.ExecuteIfBound(this);
-}
-
 #if WITH_EDITOR
 FText UControlRig::GetCategory() const
 {
@@ -307,11 +290,15 @@ void UControlRig::Execute(const EControlRigState InState)
 
 	FRigUnitContext Context;
 	Context.DrawInterface = DrawInterface;
+
+	if (InState == EControlRigState::Init)
+	{
+		Context.DataSourceRegistry = DataSourceRegistry;
+	}
+
 	Context.DeltaTime = DeltaTime;
 	Context.State = InState;
-	Context.HierarchyReference.Container = &Hierarchy;
-	Context.HierarchyReference.bUseBaseHierarchy = true;
-	Context.CurveReference = FRigCurveContainerRef(&CurveContainer);
+	Context.Hierarchy = &Hierarchy;
 
 #if WITH_EDITOR
 	Context.Log = ControlRigLog;
@@ -321,13 +308,6 @@ void UControlRig::Execute(const EControlRigState InState)
 	}
 #endif
 
-#if WITH_EDITORONLY_DATA
-	if (!bExecutionOn)
-	{
-		return;
-	}
-#endif // WITH_EDITORONLY_DATA
-
 	ControlRigVM::Execute(this, Context, Operators, ExecutionType);
 
 #if WITH_EDITOR
@@ -388,62 +368,62 @@ FTransform UControlRig::GetGlobalTransform(const FName& BoneName) const
 {
 	DECLARE_SCOPE_HIERARCHICAL_COUNTER_FUNC()
 
-	int32 Index = Hierarchy.BaseHierarchy.GetIndex(BoneName);
+	int32 Index = Hierarchy.BoneHierarchy.GetIndex(BoneName);
 	if (Index != INDEX_NONE)
 	{
-		return Hierarchy.BaseHierarchy.GetGlobalTransform(Index);
+		return Hierarchy.BoneHierarchy.GetGlobalTransform(Index);
 	}
 
 	return FTransform::Identity;
 
 }
 
-void UControlRig::SetGlobalTransform(const FName& BoneName, const FTransform& InTransform, bool bPropagateTransform) 
+void UControlRig::SetGlobalTransform(const FName& BoneName, const FTransform& InTransform, bool bPropagateTransform)
 {
 	DECLARE_SCOPE_HIERARCHICAL_COUNTER_FUNC()
 
-	int32 Index = Hierarchy.BaseHierarchy.GetIndex(BoneName);
+	int32 Index = Hierarchy.BoneHierarchy.GetIndex(BoneName);
 	if (Index != INDEX_NONE)
 	{
-		Hierarchy.BaseHierarchy.SetGlobalTransform(Index, InTransform, bPropagateTransform);
+		return Hierarchy.BoneHierarchy.SetGlobalTransform(Index, InTransform, bPropagateTransform);
 	}
 }
 
 FTransform UControlRig::GetGlobalTransform(const int32 BoneIndex) const
 {
 	DECLARE_SCOPE_HIERARCHICAL_COUNTER_FUNC()
-	return Hierarchy.BaseHierarchy.GetGlobalTransform(BoneIndex);
+	return Hierarchy.BoneHierarchy.GetGlobalTransform(BoneIndex);
 }
 
 void UControlRig::SetGlobalTransform(const int32 BoneIndex, const FTransform& InTransform, bool bPropagateTransform)
 {
 	DECLARE_SCOPE_HIERARCHICAL_COUNTER_FUNC()
-	Hierarchy.BaseHierarchy.SetGlobalTransform(BoneIndex, InTransform, bPropagateTransform);
+	Hierarchy.BoneHierarchy.SetGlobalTransform(BoneIndex, InTransform, bPropagateTransform);
 }
 
 float UControlRig::GetCurveValue(const FName& CurveName) const
 {
 	DECLARE_SCOPE_HIERARCHICAL_COUNTER_FUNC()
-	return CurveContainer.GetValue(CurveName);
+	return Hierarchy.CurveContainer.GetValue(CurveName);
 }
 
 void UControlRig::SetCurveValue(const FName& CurveName, const float CurveValue)
 {
 	DECLARE_SCOPE_HIERARCHICAL_COUNTER_FUNC()
 
-	CurveContainer.SetValue(CurveName, CurveValue);
+	Hierarchy.CurveContainer.SetValue(CurveName, CurveValue);
 }
 
 float UControlRig::GetCurveValue(const int32 CurveIndex) const
 {
 	DECLARE_SCOPE_HIERARCHICAL_COUNTER_FUNC()
-	return CurveContainer.GetValue(CurveIndex);
+	return Hierarchy.CurveContainer.GetValue(CurveIndex);
 }
 
 void UControlRig::SetCurveValue(const int32 CurveIndex, const float CurveValue)
 {
 	DECLARE_SCOPE_HIERARCHICAL_COUNTER_FUNC()
-	CurveContainer.SetValue(CurveIndex, CurveValue);
+	Hierarchy.CurveContainer.SetValue(CurveIndex, CurveValue);
 }
 
 void UControlRig::GetMappableNodeData(TArray<FName>& OutNames, TArray<FNodeItem>& OutNodeItems) const
@@ -454,24 +434,13 @@ void UControlRig::GetMappableNodeData(TArray<FName>& OutNames, TArray<FNodeItem>
 	OutNodeItems.Reset();
 
 	// now add all nodes
-	const FRigHierarchy& BaseHierarchy = Hierarchy.BaseHierarchy;
+	const FRigBoneHierarchy& BoneHierarchy = Hierarchy.BoneHierarchy;
 
-	const TArray<FRigBone>& Bones = BaseHierarchy.GetBones();
-	for (int32 Index = 0; Index < Bones.Num(); ++Index)
+	for (const FRigBone& Bone : BoneHierarchy)
 	{
-		OutNames.Add(Bones[Index].Name);
-		OutNodeItems.Add(FNodeItem(Bones[Index].ParentName, Bones[Index].InitialTransform));
+		OutNames.Add(Bone.Name);
+		OutNodeItems.Add(FNodeItem(Bone.ParentName, Bone.InitialTransform));
 	}
-
-	// have to handle curve separate because otherwise it will confuse with hierarchy
-	// we also supply curves 
-// 	const FRigCurveContainer& RigCurveContainer = CurveContainer;
-// 
-// 	for (int32 Index = 0; Index < RigCurveContainer.Curves.Num(); ++Index)
-// 	{
-// 		OutNames.Add(RigCurveContainer.Curves[Index].Name);
-// 		OutNodeItems.Add(FNodeItem(NAME_None, FTransform::Identity));
-// 	}
 }
 
 void UControlRig::ResolveInputOutputProperties()
@@ -654,4 +623,137 @@ void UControlRig::QueryIOVariables(bool bInput, TArray<FControlRigIOVariable>& O
 		OutVars.Add(NewVar);
 	}
 }
+
+// BEGIN IControlRigManipulatable interface
+const TArray<FRigSpace>& UControlRig::AvailableSpaces() const
+{
+	return Hierarchy.SpaceHierarchy.GetSpaces();
+}
+
+// do we need to return pointer? They can't save this
+FRigSpace* UControlRig::FindSpace(const FName& InSpaceName)
+{
+	const int32 SpaceIndex = Hierarchy.SpaceHierarchy.GetIndex(InSpaceName);
+	if (SpaceIndex != INDEX_NONE)
+	{
+		return &Hierarchy.SpaceHierarchy[SpaceIndex];
+	}
+
+	return nullptr;
+}
+
+FTransform UControlRig::GetSpaceGlobalTransform(const FName& InSpaceName)
+{
+	return Hierarchy.SpaceHierarchy.GetGlobalTransform(InSpaceName);
+}
+
+bool UControlRig::SetSpaceGlobalTransform(const FName& InSpaceName, const FTransform& InTransform)
+{
+	Hierarchy.SpaceHierarchy.SetGlobalTransform(InSpaceName, InTransform);
+	return true;
+}
+
+const TArray<FRigControl>& UControlRig::AvailableControls() const
+{
+	return Hierarchy.ControlHierarchy.GetControls();
+}
+
+FRigControl* UControlRig::FindControl(const FName& InControlName)
+{
+	const int32 ControlIndex = Hierarchy.ControlHierarchy.GetIndex(InControlName);
+	if (ControlIndex != INDEX_NONE)
+	{
+		return &Hierarchy.ControlHierarchy[ControlIndex];
+	}
+
+	return nullptr;
+}
+
+FTransform UControlRig::GetControlGlobalTransform(const FName& InControlName) const
+{
+	return Hierarchy.ControlHierarchy.GetGlobalTransform(InControlName);
+}
+
+FRigControlValue UControlRig::GetControlValueFromGlobalTransform(const FName& InControlName, const FTransform& InGlobalTransform)
+{
+	FRigControlValue RetVal;
+
+	const int32 ControlIndex = Hierarchy.ControlHierarchy.GetIndex(InControlName);
+	if (ControlIndex != INDEX_NONE)
+	{
+		FRigControl& Control = Hierarchy.ControlHierarchy[ControlIndex];
+
+		RetVal = Control.Value;
+
+		FTransform ParentTransform = Hierarchy.ControlHierarchy.GetParentTransform(ControlIndex);
+		switch (Control.ControlType)
+		{
+			case ERigControlType::Bool:
+			case ERigControlType::Float:
+			case ERigControlType::Vector2D:
+			{
+				// not sure how to extract this from global transform
+				break;
+			}
+			case ERigControlType::Position:
+			{
+				FTransform Transform = InGlobalTransform.GetRelativeTransform(ParentTransform);
+				RetVal.Set<FVector>(Transform.GetLocation());
+				break;
+			}
+			case ERigControlType::Scale:
+			{
+				FTransform Transform = InGlobalTransform.GetRelativeTransform(ParentTransform);
+				RetVal.Set<FVector>(Transform.GetScale3D());
+				break;
+			}
+			case ERigControlType::Quat:
+			{
+				FTransform Transform = InGlobalTransform.GetRelativeTransform(ParentTransform);
+				RetVal.Set<FQuat>(Transform.GetRotation());
+				break;
+			}
+			case ERigControlType::Rotator:
+			{
+				FTransform Transform = InGlobalTransform.GetRelativeTransform(ParentTransform);
+				RetVal.Set<FRotator>(Transform.GetRotation().Rotator());
+				break;
+			}
+			case ERigControlType::Transform:
+			{
+				FTransform Transform = InGlobalTransform.GetRelativeTransform(ParentTransform);
+				RetVal.Set<FTransform>(Transform);
+				break;
+			}
+		}
+	}
+
+	return RetVal;
+}
+
+bool UControlRig::SetControlSpace(const FName& InControlName, const FName& InSpaceName)
+{
+	Hierarchy.ControlHierarchy.SetSpace(InControlName, InSpaceName);
+	return true;
+}
+
+UControlRigGizmoLibrary* UControlRig::GetGizmoLibrary() const
+{
+#if WITH_EDITOR
+	if (UControlRig* CDO = Cast<UControlRig>(GetClass()->GetDefaultObject()))
+	{
+		if (!CDO->GizmoLibrary.IsValid())
+		{
+			CDO->GizmoLibrary.LoadSynchronous();
+		}
+		if (CDO->GizmoLibrary.IsValid())
+		{
+			return CDO->GizmoLibrary.Get();
+		}
+	}
+#endif
+	return nullptr;
+}
+
+// END IControlRigManipulatable interface
 #undef LOCTEXT_NAMESPACE
diff --git a/Engine/Plugins/Experimental/ControlRig/Source/ControlRig/Private/ControlRigComponent.cpp b/Engine/Plugins/Experimental/ControlRig/Source/ControlRig/Private/ControlRigComponent.cpp
--- a/Engine/Plugins/Experimental/ControlRig/Source/ControlRig/Private/ControlRigComponent.cpp
+++ b/Engine/Plugins/Experimental/ControlRig/Source/ControlRig/Private/ControlRigComponent.cpp
@@ -88,13 +88,11 @@ void UControlRigComponent::TickComponent(float DeltaTime, enum ELevelTick TickTy
 
 		// Call pre-evaluation callbacks (e.g. for copying input data into the rig)
 		OnPreEvaluate();
-		ControlRig->PreEvaluate_GameThread();
 
 		// If we multi-thread rig evaluation, then this call should be made on worker threads, but pre/post evaluate should be called on the game thread
 		ControlRig->Evaluate_AnyThread();
 
 		// Call post-evaluation callbacks (e.g. for copying output data out of the rig)
-		ControlRig->PostEvaluate_GameThread();
 		OnPostEvaluate();
 	}
 }
diff --git a/Engine/Plugins/Experimental/ControlRig/Source/ControlRig/Private/ControlRigGizmoLibrary.cpp b/Engine/Plugins/Experimental/ControlRig/Source/ControlRig/Private/ControlRigGizmoLibrary.cpp
new file mode 100644
--- /dev/null
+++ b/Engine/Plugins/Experimental/ControlRig/Source/ControlRig/Private/ControlRigGizmoLibrary.cpp
@@ -0,0 +1,124 @@
+// Copyright 1998-2019 Epic Games, Inc. All Rights Reserved.
+
+#include "ControlRigGizmoLibrary.h"
+
+#define LOCTEXT_NAMESPACE "ControlRigGizmoLibrary"
+
+UControlRigGizmoLibrary::UControlRigGizmoLibrary()
+{
+}
+
+#if WITH_EDITOR
+
+// UObject interface
+void UControlRigGizmoLibrary::PostEditChangeChainProperty(struct FPropertyChangedChainEvent& PropertyChangedEvent)
+{
+	if (PropertyChangedEvent.Property->GetName() == TEXT("GizmoName"))
+	{
+		UProperty* MemberProperty = PropertyChangedEvent.PropertyChain.GetHead()->GetValue();
+		if (MemberProperty->GetName() == TEXT("DefaultGizmo"))
+		{
+			DefaultGizmo.GizmoName = TEXT("Gizmo");
+			GetUpdatedNameList(true);
+		}
+		else if(MemberProperty->GetName() == TEXT("Gizmos"))
+		{
+			if (Gizmos.Num() == 0)
+			{
+				return;
+			}
+
+			int32 GizmoIndexEdited = PropertyChangedEvent.GetArrayIndex(TEXT("Gizmos"));
+			if (Gizmos.IsValidIndex(GizmoIndexEdited))
+			{
+				TArray<FName> Names;
+				Names.Add(DefaultGizmo.GizmoName);
+				for (int32 GizmoIndex = 0; GizmoIndex < Gizmos.Num(); GizmoIndex++)
+				{
+					if (GizmoIndex != GizmoIndexEdited)
+					{
+						Names.Add(Gizmos[GizmoIndex].GizmoName);
+					}
+				}
+
+				FName DesiredName = Gizmos[GizmoIndexEdited].GizmoName;
+				FString Name = DesiredName.ToString();
+				int32 Suffix = 0;
+				while (Names.Contains(*Name))
+				{
+					Suffix++;
+					Name = FString::Printf(TEXT("%s_%d"), *DesiredName.ToString(), Suffix);
+				}
+				Gizmos[GizmoIndexEdited].GizmoName = *Name;
+			}
+			GetUpdatedNameList(true);
+		}
+	}
+	else if (PropertyChangedEvent.Property->GetName() == TEXT("Gizmos"))
+	{
+		TArray<FName> Names;
+		Names.Add(DefaultGizmo.GizmoName);
+		for (int32 GizmoIndex = 0; GizmoIndex < Gizmos.Num(); GizmoIndex++)
+		{
+			FName DesiredName = Gizmos[GizmoIndex].GizmoName;
+			FString Name = DesiredName.ToString();
+			int32 Suffix = 0;
+			while (Names.Contains(*Name))
+			{
+				Suffix++;
+				Name = FString::Printf(TEXT("%s_%d"), *DesiredName.ToString(), Suffix);
+			}
+			Gizmos[GizmoIndex].GizmoName = *Name;
+
+			Names.Add(Gizmos[GizmoIndex].GizmoName);
+		}
+		GetUpdatedNameList(true);
+	}
+}
+
+#endif
+
+const FControlRigGizmoDefinition* UControlRigGizmoLibrary::GetGizmoByName(const FName& InName, bool bUseDefaultIfNotFound) const
+{
+	if (InName == DefaultGizmo.GizmoName)
+	{
+		return &DefaultGizmo;
+	}
+
+	for (int32 GizmoIndex = 0; GizmoIndex < Gizmos.Num(); GizmoIndex++)
+	{
+		if (Gizmos[GizmoIndex].GizmoName == InName)
+		{
+			return &Gizmos[GizmoIndex];
+		}
+	}
+
+	if (bUseDefaultIfNotFound)
+	{
+		return &DefaultGizmo;
+	}
+
+	return nullptr;
+}
+
+const TArray<FName> UControlRigGizmoLibrary::GetUpdatedNameList(bool bReset)
+{
+	if (bReset)
+	{
+		NameList.Reset();
+	}
+
+	if (NameList.Num() != Gizmos.Num())
+	{
+		NameList.Reset();
+		for (const FControlRigGizmoDefinition& Gizmo : Gizmos)
+		{
+			NameList.Add(Gizmo.GizmoName);
+		}
+		NameList.Sort(FNameLexicalLess());
+	}
+
+	return NameList;
+}
+
+#undef LOCTEXT_NAMESPACE
diff --git a/Engine/Plugins/Experimental/ControlRig/Source/ControlRig/Private/Manipulatable/IControlRigManipulatable.cpp b/Engine/Plugins/Experimental/ControlRig/Source/ControlRig/Private/Manipulatable/IControlRigManipulatable.cpp
new file mode 100644
--- /dev/null
+++ b/Engine/Plugins/Experimental/ControlRig/Source/ControlRig/Private/Manipulatable/IControlRigManipulatable.cpp
@@ -0,0 +1,29 @@
+// Copyright 1998-2019 Epic Games, Inc. All Rights Reserved.
+
+#include "Manipulatable/IControlRigManipulatable.h"
+
+IControlRigManipulatable::IControlRigManipulatable()
+{
+	bManipulationEnabled = false;
+}
+
+IControlRigManipulatable::~IControlRigManipulatable()
+{
+	OnFilterControl.Clear();
+	OnControlModified.Clear();
+}
+
+bool IControlRigManipulatable::SetControlGlobalTransform(const FName& InControlName, const FTransform& InGlobalTransform)
+{
+	FRigControlValue Value = GetControlValueFromGlobalTransform(InControlName, InGlobalTransform);
+	if (OnFilterControl.IsBound())
+	{
+		FRigControl* Control = FindControl(InControlName);
+		if (Control)
+		{
+			OnFilterControl.Broadcast(this, *Control, Value);
+		}
+	}
+	SetControlValue(InControlName, Value, true /* notify */);
+	return true;
+}
\ No newline at end of file
diff --git a/Engine/Plugins/Experimental/ControlRig/Source/ControlRig/Private/Rigs/RigControlHierarchy.cpp b/Engine/Plugins/Experimental/ControlRig/Source/ControlRig/Private/Rigs/RigControlHierarchy.cpp
new file mode 100644
--- /dev/null
+++ b/Engine/Plugins/Experimental/ControlRig/Source/ControlRig/Private/Rigs/RigControlHierarchy.cpp
@@ -0,0 +1,828 @@
+// Copyright 1998-2019 Epic Games, Inc. All Rights Reserved.
+
+#include "Rigs/RigControlHierarchy.h"
+#include "ControlRig.h"
+#include "HelperUtil.h"
+#include "Settings/ControlRigSettings.h"
+
+////////////////////////////////////////////////////////////////////////////////
+// FRigControlHierarchy
+////////////////////////////////////////////////////////////////////////////////
+
+FRigControlHierarchy::FRigControlHierarchy()
+	:Container(nullptr)
+{
+}
+
+FRigControlHierarchy& FRigControlHierarchy::operator= (const FRigControlHierarchy &InOther)
+{
+#if WITH_EDITOR
+	for (int32 Index = Num() - 1; Index >= 0; Index--)
+	{
+		FRigControl ControlToRemove = Controls[Index];
+		OnControlRemoved.Broadcast(Container, FRigElementKey(ControlToRemove.Name, ERigElementType::Control));
+	}
+#endif
+
+	Controls.Reset();
+	Controls.Append(InOther.Controls);
+	NameToIndexMapping.Reset();
+	RefreshMapping();
+
+#if WITH_EDITOR
+	for (const FRigControl& ControlAdded : Controls)
+	{
+		OnControlAdded.Broadcast(Container, FRigElementKey(ControlAdded.Name, ERigElementType::Control));
+	}
+#endif
+
+	return *this;
+}
+
+FName FRigControlHierarchy::GetSafeNewName(const FName& InPotentialNewName) const
+{
+	FName Name = InPotentialNewName;
+	int32 Suffix = 1;
+	while(!IsNameAvailable(Name))
+	{
+		Name = *FString::Printf(TEXT("%s_%d"), *InPotentialNewName.ToString(), ++Suffix);
+	}
+	return Name;
+}
+
+FRigControl& FRigControlHierarchy::Add(
+	const FName& InNewName,
+	ERigControlType InControlType,
+	const FName& InParentName,
+	const FName& InSpaceName,
+	const FRigControlValue& InValue
+#if WITH_EDITORONLY_DATA
+	,const FName& InGizmoName,
+	const FTransform& InGizmoTransform,
+	const FLinearColor& InGizmoColor
+#endif
+)
+{
+	FRigControl NewControl;
+	NewControl.Name = GetSafeNewName(InNewName);
+	NewControl.ControlType = InControlType;
+	NewControl.ParentIndex = GetIndex(InParentName);
+	NewControl.ParentName = NewControl.ParentIndex == INDEX_NONE ? NAME_None : InParentName;
+	NewControl.SpaceIndex = INDEX_NONE;
+	NewControl.SpaceName = NAME_None;
+	NewControl.InitialValue = InValue;
+	NewControl.Value = InValue;
+#if WITH_EDITORONLY_DATA
+	NewControl.GizmoName = InGizmoName;
+	NewControl.GizmoTransform = InGizmoTransform;
+	NewControl.GizmoColor = InGizmoColor;
+#endif
+
+	FName NewControlName = NewControl.Name;
+	Controls.Add(NewControl);
+	RefreshMapping();
+
+#if WITH_EDITOR
+	OnControlAdded.Broadcast(Container, NewControl.GetElementKey());
+#endif
+
+	SetSpace(NewControlName, InSpaceName);
+
+	int32 Index = GetIndex(NewControlName);
+	return Controls[Index];
+}
+
+bool FRigControlHierarchy::Reparent(const FName& InName, const FName& InNewParentName)
+{
+	int32 Index = GetIndex(InName);
+	// can't parent to itself
+	if (Index != INDEX_NONE && InName != InNewParentName)
+	{
+		FRigControl& Control = Controls[Index];
+
+#if WITH_EDITOR
+		FName OldParentName = Control.ParentName;
+#endif
+
+		struct Local
+		{
+			static bool IsParentedTo(int32 Child, int32 Parent, const TArray<FRigControl>& Controls)
+			{
+				if (Parent == INDEX_NONE || Child == INDEX_NONE)
+				{
+					return false;
+				}
+
+				if (Child == Parent)
+				{
+					return true;
+				}
+
+				if (Controls[Child].ParentIndex == Parent)
+				{
+					return true;
+				}
+
+				return IsParentedTo(Controls[Child].ParentIndex, Parent, Controls);
+			}
+		};
+
+		int32 ParentIndex = GetIndex(InNewParentName);
+		if (Local::IsParentedTo(ParentIndex, Index, Controls))
+		{
+			ParentIndex = INDEX_NONE;
+		}
+
+		Control.ParentIndex = ParentIndex;
+		Control.ParentName = Control.ParentIndex == INDEX_NONE ? NAME_None : InNewParentName;
+
+#if WITH_EDITOR
+		FName NewParentName = Control.ParentName;
+#endif
+
+		RefreshMapping();
+
+#if WITH_EDITOR
+		if (OldParentName != NewParentName)
+		{
+			OnControlReparented.Broadcast(Container, FRigElementKey(InName, RigElementType()), OldParentName, NewParentName);
+		}
+#endif
+		return Controls[GetIndex(InName)].ParentName == InNewParentName;
+	}
+	return false;
+}
+
+void FRigControlHierarchy::SetSpace(const FName& InName, const FName& InNewSpaceName)
+{
+	int32 Index = GetIndex(InName);
+	if (Index != INDEX_NONE)
+	{
+		int32 SpaceIndex = GetSpaceIndex(InNewSpaceName);
+
+		if (SpaceIndex != INDEX_NONE)
+		{
+			if (Container != nullptr)
+			{
+				if (Container->IsParentedTo(ERigElementType::Space, SpaceIndex, ERigElementType::Control, Index))
+				{
+					SpaceIndex = INDEX_NONE;
+				}
+			}
+		}
+
+		Controls[Index].SpaceIndex = SpaceIndex;
+		Controls[Index].SpaceName = Controls[Index].SpaceIndex == INDEX_NONE ? NAME_None : InNewSpaceName;
+
+#if WITH_EDITOR
+		OnControlReparented.Broadcast(Container, FRigElementKey(InName, RigElementType()), Controls[Index].ParentName, Controls[Index].ParentName);
+#endif
+	}
+}
+
+FRigControl FRigControlHierarchy::Remove(const FName& InNameToRemove)
+{
+	TArray<int32> Children;
+#if WITH_EDITOR
+	TArray<FName> RemovedChildControls;
+#endif
+	if (GetChildren(InNameToRemove, Children, true) > 0)
+	{
+		// sort by child index
+		Children.Sort([](const int32& A, const int32& B) { return A < B; });
+
+		// want to delete from end to the first 
+		for (int32 ChildIndex = Children.Num() - 1; ChildIndex >= 0; --ChildIndex)
+		{
+#if WITH_EDITOR
+			RemovedChildControls.Add(Controls[Children[ChildIndex]].Name);
+#endif
+			Controls.RemoveAt(Children[ChildIndex]);
+		}
+	}
+
+	int32 IndexToDelete = GetIndex(InNameToRemove);
+#if WITH_EDITOR
+	Select(InNameToRemove, false);
+#endif
+	FRigControl RemovedControl = Controls[IndexToDelete];
+	Controls.RemoveAt(IndexToDelete);
+
+	RefreshMapping();
+
+#if WITH_EDITOR
+	for (const FName& RemovedChildControl : RemovedChildControls)
+	{
+		OnControlRemoved.Broadcast(Container, FRigElementKey(RemovedChildControl, RigElementType()));
+	}
+	OnControlRemoved.Broadcast(Container, RemovedControl.GetElementKey());
+#endif
+
+	return RemovedControl;
+}
+
+// list of names of children - this is not cheap, and is supposed to be used only for one time set up
+int32 FRigControlHierarchy::GetChildren(const FName& InName, TArray<int32>& OutChildren, bool bRecursively) const
+{
+	return GetChildren(GetIndex(InName), OutChildren, bRecursively);
+}
+
+int32 FRigControlHierarchy::GetChildren(const int32 InIndex, TArray<int32>& OutChildren, bool bRecursively) const
+{
+	OutChildren.Reset();
+
+	if (InIndex != INDEX_NONE)
+	{
+		GetChildrenRecursive(InIndex, OutChildren, bRecursively);
+	}
+
+	return OutChildren.Num();
+}
+
+FName FRigControlHierarchy::GetName(int32 InIndex) const
+{
+	if (Controls.IsValidIndex(InIndex))
+	{
+		return Controls[InIndex].Name;
+	}
+
+	return NAME_None;
+}
+
+int32 FRigControlHierarchy::GetIndexSlow(const FName& InName) const
+{
+	for (int32 Index = 0; Index < Controls.Num(); ++Index)
+	{
+		if (Controls[Index].Name == InName)
+		{
+			return Index;
+		}
+	}
+
+	return INDEX_NONE;
+}
+
+void FRigControlHierarchy::SetGlobalTransform(const FName& InName, const FTransform& InTransform)
+{
+	SetGlobalTransform(GetIndex(InName), InTransform);
+}
+
+void FRigControlHierarchy::SetGlobalTransform(int32 InIndex, const FTransform& InTransform)
+{
+	if (Container == nullptr)
+	{
+		SetLocalTransform(InIndex, InTransform);
+		return;
+	}
+
+	if (Controls.IsValidIndex(InIndex))
+	{
+		const FRigControl& Control = Controls[InIndex];
+		FTransform ParentTransform = GetParentTransform(InIndex);
+		SetLocalTransform(InIndex, InTransform.GetRelativeTransform(ParentTransform));
+	}
+}
+
+FTransform FRigControlHierarchy::GetGlobalTransform(const FName& InName) const
+{
+	return GetGlobalTransform(GetIndex(InName));
+}
+
+FTransform FRigControlHierarchy::GetGlobalTransform(int32 InIndex) const
+{
+	if (Container == nullptr)
+	{
+		return GetLocalTransform(InIndex);
+	}
+
+	if (Controls.IsValidIndex(InIndex))
+	{
+		const FRigControl& Control = Controls[InIndex];
+
+		FTransform ParentTransform = GetParentTransform(InIndex);
+		switch (Control.ControlType)
+		{
+			case ERigControlType::Bool:
+			case ERigControlType::Float:
+			case ERigControlType::Vector2D:
+			{
+				return ParentTransform;
+			}
+			case ERigControlType::Position:
+			{
+				FTransform Transform;
+				Transform.SetLocation(Control.Value.Get<FVector>());
+				Transform = Transform * ParentTransform;
+				Transform.NormalizeRotation();
+				return Transform;
+			}
+			case ERigControlType::Scale:
+			{
+				FTransform Transform;
+				Transform.SetScale3D(Control.Value.Get<FVector>());
+				Transform = Transform * ParentTransform;
+				Transform.NormalizeRotation();
+				return Transform;
+			}
+			case ERigControlType::Quat:
+			{
+				FTransform Transform;
+				Transform.SetRotation(Control.Value.Get<FQuat>());
+				Transform = Transform * ParentTransform;
+				Transform.NormalizeRotation();
+				return Transform;
+			}
+			case ERigControlType::Rotator:
+			{
+				FTransform Transform;
+				Transform.SetRotation(FQuat(Control.Value.Get<FRotator>()));
+				Transform = Transform * ParentTransform;
+				Transform.NormalizeRotation();
+				return Transform;
+			}
+			case ERigControlType::Transform:
+			{
+				FTransform Transform = Control.Value.Get<FTransform>() * ParentTransform;
+				Transform.NormalizeRotation();
+				return Transform;
+			}
+			default:
+			{
+				ensure(false);
+				break;
+			}
+		}
+	}
+
+	return FTransform::Identity;
+}
+
+void FRigControlHierarchy::SetLocalTransform(const FName& InName, const FTransform& InTransform)
+{
+	SetValue<FTransform>(InName, InTransform);
+}
+
+void FRigControlHierarchy::SetLocalTransform(int32 InIndex, const FTransform& InTransform)
+{
+	SetValue<FTransform>(InIndex, InTransform);
+}
+
+FTransform FRigControlHierarchy::GetLocalTransform(const FName& InName) const
+{
+	return GetValue<FTransform>(InName);
+}
+
+FTransform FRigControlHierarchy::GetLocalTransform(int32 InIndex) const
+{
+	return GetValue<FTransform>(InIndex);
+}
+
+FTransform FRigControlHierarchy::GetParentTransform(int32 InIndex) const
+{
+	const FRigControl& Control = Controls[InIndex];
+
+	FTransform ParentTransform = FTransform::Identity;
+	if (Control.SpaceIndex != INDEX_NONE && Container != nullptr)
+	{
+		ParentTransform = Container->GetGlobalTransform(ERigElementType::Space, Control.SpaceIndex);
+	}
+	else if (Control.ParentIndex != INDEX_NONE && Container != nullptr)
+	{
+		ParentTransform = GetGlobalTransform(Control.ParentIndex);
+	}
+
+	return ParentTransform;
+}
+
+void FRigControlHierarchy::SetValue(const FName& InName, const FRigControlValue& InValue)
+{
+	SetValue(GetIndex(InName), InValue);
+}
+
+void FRigControlHierarchy::SetValue(int32 InIndex, const FRigControlValue& InValue)
+{
+	if (Controls.IsValidIndex(InIndex))
+	{
+		FRigControl& Control = Controls[InIndex];
+		Control.Value = InValue;
+
+		if (Control.ControlType == ERigControlType::Transform)
+		{
+			Control.Value.GetRef<FTransform>().NormalizeRotation();
+		}
+		else if (Control.ControlType == ERigControlType::Quat)
+		{
+			Control.Value.GetRef<FQuat>().Normalize();
+		}
+	}
+}
+
+FRigControlValue FRigControlHierarchy::GetValue(const FName& InName) const
+{
+	return GetValue(GetIndex(InName));
+}
+
+FRigControlValue FRigControlHierarchy::GetValue(int32 InIndex) const
+{
+	if (Controls.IsValidIndex(InIndex))
+	{
+		return Controls[InIndex].Value;
+	}
+
+	return FRigControlValue();
+}
+
+void FRigControlHierarchy::SetInitialValue(const FName& InName, const FRigControlValue& InValue)
+{
+	SetInitialValue(GetIndex(InName), InValue);
+}
+
+void FRigControlHierarchy::SetInitialValue(int32 InIndex, const FRigControlValue& InValue)
+{
+	if (Controls.IsValidIndex(InIndex))
+	{
+		FRigControl& Control = Controls[InIndex];
+		Control.InitialValue = InValue;
+
+		if (Control.ControlType == ERigControlType::Transform)
+		{
+			Control.InitialValue.GetRef<FTransform>().NormalizeRotation();
+		}
+		else if (Control.ControlType == ERigControlType::Quat)
+		{
+			Control.InitialValue.GetRef<FQuat>().Normalize();
+		}
+	}
+}
+
+FRigControlValue FRigControlHierarchy::GetInitialValue(const FName& InName) const
+{
+	return GetInitialValue(GetIndex(InName));
+}
+
+FRigControlValue FRigControlHierarchy::GetInitialValue(int32 InIndex) const
+{
+	if (Controls.IsValidIndex(InIndex))
+	{
+		return Controls[InIndex].InitialValue;
+	}
+	return FRigControlValue();
+}
+
+void FRigControlHierarchy::SetInitialGlobalTransform(const FName& InName, const FTransform& GlobalTransform)
+{
+	SetInitialGlobalTransform(GetIndex(InName), GlobalTransform);
+}
+
+// wip - should support all types that can provide space transform data (pos/rotation)
+void FRigControlHierarchy::SetInitialGlobalTransform(int32 InIndex, const FTransform& GlobalTransform)
+{
+	if (Controls.IsValidIndex(InIndex))
+	{
+		FRigControl& Control = Controls[InIndex];
+		FTransform ParentTransform = FTransform::Identity;
+		if (Container)
+		{
+			if (Control.SpaceName != NAME_None)
+			{
+				ParentTransform = Container->GetInitialGlobalTransform(Control.GetSpaceElementKey());
+			}
+			else
+			{
+				ParentTransform = Container->GetInitialGlobalTransform(Control.GetParentElementKey());
+			}
+		}
+
+		if (Control.ControlType == ERigControlType::Transform)
+		{
+			Control.InitialValue.Set<FTransform>(GlobalTransform.GetRelativeTransform(ParentTransform));
+		}
+	}
+}
+
+FTransform FRigControlHierarchy::GetInitialGlobalTransform(const FName& InName) const
+{
+	return GetInitialGlobalTransform(GetIndex(InName));
+}
+
+FTransform FRigControlHierarchy::GetInitialGlobalTransform(int32 InIndex) const
+{
+	// @todo: Templatize
+	if (Controls.IsValidIndex(InIndex))
+	{
+		const FRigControl& Control = Controls[InIndex];
+		FTransform ParentTransform = FTransform::Identity;
+		if (Container)
+		{
+			if (Control.SpaceName != NAME_None)
+			{
+				ParentTransform = Container->GetInitialGlobalTransform(Control.GetSpaceElementKey());
+			}
+			else
+			{
+				ParentTransform = Container->GetInitialGlobalTransform(Control.GetParentElementKey());
+			}
+		}
+
+		switch (Control.ControlType)
+		{
+			case ERigControlType::Bool:
+			case ERigControlType::Float:
+			case ERigControlType::Vector2D:
+			{
+				return ParentTransform;
+			}
+			case ERigControlType::Position:
+			{
+				FTransform Transform;
+				Transform.SetLocation(Control.InitialValue.Get<FVector>());
+				return Transform * ParentTransform;
+			}
+			case ERigControlType::Scale:
+			{
+				FTransform Transform;
+				Transform.SetScale3D(Control.InitialValue.Get<FVector>());
+				return Transform * ParentTransform;
+			}
+			case ERigControlType::Quat:
+			{
+				FTransform Transform;
+				Transform.SetRotation(Control.InitialValue.Get<FQuat>());
+				return Transform * ParentTransform;
+			}
+			case ERigControlType::Rotator:
+			{
+				FTransform Transform;
+				Transform.SetRotation(FQuat(Control.InitialValue.Get<FRotator>()));
+				return Transform * ParentTransform;
+			}
+			case ERigControlType::Transform:
+			{
+				return Control.InitialValue.Get<FTransform>() * ParentTransform;
+			}
+		}
+	}
+	return FTransform::Identity;
+}
+
+FName FRigControlHierarchy::Rename(const FName& InOldName, const FName& InNewName)
+{
+	if (InOldName != InNewName)
+	{
+		const int32 Found = GetIndex(InOldName);
+		if (Found != INDEX_NONE)
+		{
+			FName NewName = GetSafeNewName(InNewName);
+
+#if WITH_EDITOR
+			bool bWasSelected = IsSelected(InOldName);
+			if(bWasSelected)
+			{
+				Select(InOldName, false);
+			}
+#endif
+
+			Controls[Found].Name = NewName;
+
+			// go through find all children and rename them
+#if WITH_EDITOR
+			TArray<FName> ReparentedControls;
+#endif
+			for (int32 Index = 0; Index < Controls.Num(); ++Index)
+			{
+				if (Controls[Index].ParentName == InOldName)
+				{
+					Controls[Index].ParentName = NewName;
+#if WITH_EDITOR
+					ReparentedControls.Add(Controls[Index].Name);
+#endif
+				}
+			}
+
+			RefreshMapping();
+
+#if WITH_EDITOR
+			OnControlRenamed.Broadcast(Container, RigElementType(), InOldName, NewName);
+			for (const FName& ReparentedControl : ReparentedControls)
+			{
+				OnControlReparented.Broadcast(Container, FRigElementKey(ReparentedControl, RigElementType()), InOldName, NewName);
+			}
+			if(bWasSelected)
+			{
+				Select(NewName, true);
+			}
+#endif
+			return NewName;
+		}
+	}
+
+	return NAME_None;
+}
+
+void FRigControlHierarchy::RefreshMapping()
+{
+	NameToIndexMapping.Empty();
+	for (int32 Index = 0; Index < Controls.Num(); ++Index)
+	{
+		Controls[Index].Index = Index;
+		NameToIndexMapping.Add(Controls[Index].Name, Index);
+	}
+}
+
+void FRigControlHierarchy::Initialize()
+{
+	RefreshMapping();
+
+	// update parent index
+	for (int32 Index = 0; Index < Controls.Num(); ++Index)
+	{
+		Controls[Index].ParentIndex = GetIndex(Controls[Index].ParentName);
+		if (Container)
+		{
+			Controls[Index].SpaceIndex = Container->SpaceHierarchy.GetIndex(Controls[Index].SpaceName);
+		}
+	}
+
+	// initialize transform
+	for (int32 Index = 0; Index < Controls.Num(); ++Index)
+	{
+		Controls[Index].Value = Controls[Index].InitialValue;
+
+		// update children
+		GetChildren(Index, Controls[Index].Dependents, false);
+	}
+}
+
+void FRigControlHierarchy::Reset()
+{
+	Controls.Reset();
+}
+
+void FRigControlHierarchy::ResetValues()
+{
+	// initialize transform
+	for (int32 Index = 0; Index < Controls.Num(); ++Index)
+	{
+		Controls[Index].Value = Controls[Index].InitialValue;
+	}
+}
+
+int32 FRigControlHierarchy::GetChildrenRecursive(const int32 InIndex, TArray<int32>& OutChildren, bool bRecursively) const
+{
+	const int32 StartChildIndex = OutChildren.Num();
+
+	// all children should be later than parent
+	for (int32 ChildIndex = InIndex + 1; ChildIndex < Controls.Num(); ++ChildIndex)
+	{
+		if (Controls[ChildIndex].ParentIndex == InIndex)
+		{
+			OutChildren.AddUnique(ChildIndex);
+		}
+	}
+
+	if (bRecursively)
+	{
+		// since we keep appending inside of functions, we make sure not to go over original list
+		const int32 EndChildIndex = OutChildren.Num() - 1;
+		for (int32 ChildIndex = StartChildIndex; ChildIndex <= EndChildIndex; ++ChildIndex)
+		{
+			GetChildrenRecursive(OutChildren[ChildIndex], OutChildren, bRecursively);
+		}
+	}
+
+	return OutChildren.Num();
+}
+
+int32 FRigControlHierarchy::GetSpaceIndex(const FName& InName) const
+{
+	if (Container == nullptr || InName == NAME_None)
+	{
+		return INDEX_NONE;
+	}
+	return Container->GetIndex(FRigElementKey(InName, ERigElementType::Space));
+}
+
+#if WITH_EDITOR
+
+bool FRigControlHierarchy::Select(const FName& InName, bool bSelect)
+{
+	if(GetIndex(InName) == INDEX_NONE)
+	{
+		return false;
+	}
+
+	if(bSelect == IsSelected(InName))
+	{
+		return false;
+	}
+
+	if(bSelect)
+	{
+		if (Container)
+		{
+			Container->BoneHierarchy.ClearSelection();
+			Container->SpaceHierarchy.ClearSelection();
+			Container->CurveContainer.ClearSelection();
+		}
+
+		Selection.Add(InName);
+	}
+	else
+	{
+		Selection.Remove(InName);
+	}
+
+	OnControlSelected.Broadcast(Container, FRigElementKey(InName, RigElementType()), bSelect);
+
+	return true;
+}
+
+bool FRigControlHierarchy::ClearSelection()
+{
+	TArray<FName> TempSelection;
+	TempSelection.Append(Selection);
+	for(const FName& SelectedName : TempSelection)
+	{
+		Select(SelectedName, false);
+	}
+	return TempSelection.Num() > 0;
+}
+
+TArray<FName> FRigControlHierarchy::CurrentSelection() const
+{
+	TArray<FName> TempSelection;
+	TempSelection.Append(Selection);
+	return TempSelection;
+}
+
+bool FRigControlHierarchy::IsSelected(const FName& InName) const
+{
+	return Selection.Contains(InName);
+}
+
+void FRigControlHierarchy::HandleOnElementRemoved(FRigHierarchyContainer* InContainer, const FRigElementKey& InKey)
+{
+	if (Container == nullptr)
+	{
+		return;
+	}
+
+	switch (InKey.Type)
+	{
+		case ERigElementType::Space:
+		{
+			for (FRigControl& Control : Controls)
+			{
+				if (Control.SpaceName == InKey.Name)
+				{
+					Control.SpaceIndex = INDEX_NONE;
+					Control.SpaceName = NAME_None;
+#if WITH_EDITOR
+					OnControlReparented.Broadcast(Container, Control.GetElementKey(), Control.ParentName, Control.ParentName);
+#endif
+				}
+			}
+			break;
+		}
+		case ERigElementType::Bone:
+		case ERigElementType::Control:
+		case ERigElementType::Curve:
+		{
+			break;
+		}
+	}
+}
+
+void FRigControlHierarchy::HandleOnElementRenamed(FRigHierarchyContainer* InContainer, ERigElementType InElementType, const FName& InOldName, const FName& InNewName)
+{
+	if (Container == nullptr)
+	{
+		return;
+	}
+
+	switch (InElementType)
+	{
+		case ERigElementType::Space:
+		{
+			for (FRigControl& Control : Controls)
+			{
+				if (Control.SpaceName == InOldName)
+				{
+					Control.SpaceIndex = Container->SpaceHierarchy.GetIndex(InNewName);
+					Control.SpaceName = Control.SpaceIndex == INDEX_NONE ? NAME_None : InNewName;
+#if WITH_EDITOR
+					OnControlReparented.Broadcast(Container, Control.GetElementKey(), Control.ParentName, Control.ParentName);
+#endif
+				}
+			}
+			break;
+		}
+		case ERigElementType::Bone:
+		case ERigElementType::Control:
+		case ERigElementType::Curve:
+		{
+			break;
+		}
+	}
+}
+#endif
diff --git a/Engine/Plugins/Experimental/ControlRig/Source/ControlRig/Private/Rigs/RigHierarchyContainer.cpp b/Engine/Plugins/Experimental/ControlRig/Source/ControlRig/Private/Rigs/RigHierarchyContainer.cpp
new file mode 100644
--- /dev/null
+++ b/Engine/Plugins/Experimental/ControlRig/Source/ControlRig/Private/Rigs/RigHierarchyContainer.cpp
@@ -0,0 +1,965 @@
+// Copyright 1998-2019 Epic Games, Inc. All Rights Reserved.
+
+#include "Rigs/RigHierarchyContainer.h"
+#include "ControlRig.h"
+#include "HelperUtil.h"
+#include "UObject/PropertyPortFlags.h"
+
+////////////////////////////////////////////////////////////////////////////////
+// FRigHierarchyContainer
+////////////////////////////////////////////////////////////////////////////////
+
+FRigHierarchyContainer::FRigHierarchyContainer()
+{
+	Initialize();
+}
+
+FRigHierarchyContainer& FRigHierarchyContainer::operator= (const FRigHierarchyContainer &InOther)
+{
+	BoneHierarchy = InOther.BoneHierarchy;
+	SpaceHierarchy = InOther.SpaceHierarchy;
+	ControlHierarchy = InOther.ControlHierarchy;
+	CurveContainer = InOther.CurveContainer;
+	return *this;
+}
+
+void FRigHierarchyContainer::Initialize()
+{
+	BoneHierarchy.Container = this;
+	SpaceHierarchy.Container = this;
+	ControlHierarchy.Container = this;
+	CurveContainer.Container = this;
+
+#if WITH_EDITOR
+	BoneHierarchy.OnBoneAdded.RemoveAll(this);
+	BoneHierarchy.OnBoneRemoved.RemoveAll(this);
+	BoneHierarchy.OnBoneRenamed.RemoveAll(this);
+	BoneHierarchy.OnBoneReparented.RemoveAll(this);
+	BoneHierarchy.OnBoneSelected.RemoveAll(this);
+
+	BoneHierarchy.OnBoneAdded.AddRaw(this, &FRigHierarchyContainer::HandleOnElementAdded);
+	BoneHierarchy.OnBoneRemoved.AddRaw(this, &FRigHierarchyContainer::HandleOnElementRemoved);
+	BoneHierarchy.OnBoneRenamed.AddRaw(this, &FRigHierarchyContainer::HandleOnElementRenamed);
+	BoneHierarchy.OnBoneReparented.AddRaw(this, &FRigHierarchyContainer::HandleOnElementReparented);
+	BoneHierarchy.OnBoneSelected.AddRaw(this, &FRigHierarchyContainer::HandleOnElementSelected);
+
+	SpaceHierarchy.OnSpaceAdded.RemoveAll(this);
+	SpaceHierarchy.OnSpaceRemoved.RemoveAll(this);
+	SpaceHierarchy.OnSpaceRenamed.RemoveAll(this);
+	SpaceHierarchy.OnSpaceReparented.RemoveAll(this);
+	SpaceHierarchy.OnSpaceSelected.RemoveAll(this);
+
+	SpaceHierarchy.OnSpaceAdded.AddRaw(this, &FRigHierarchyContainer::HandleOnElementAdded);
+	SpaceHierarchy.OnSpaceRemoved.AddRaw(this, &FRigHierarchyContainer::HandleOnElementRemoved);
+	SpaceHierarchy.OnSpaceRenamed.AddRaw(this, &FRigHierarchyContainer::HandleOnElementRenamed);
+	SpaceHierarchy.OnSpaceReparented.AddRaw(this, &FRigHierarchyContainer::HandleOnElementReparented);
+	SpaceHierarchy.OnSpaceSelected.AddRaw(this, &FRigHierarchyContainer::HandleOnElementSelected);
+
+	ControlHierarchy.OnControlAdded.RemoveAll(this);
+	ControlHierarchy.OnControlRemoved.RemoveAll(this);
+	ControlHierarchy.OnControlRenamed.RemoveAll(this);
+	ControlHierarchy.OnControlReparented.RemoveAll(this);
+	ControlHierarchy.OnControlSelected.RemoveAll(this);
+
+	ControlHierarchy.OnControlAdded.AddRaw(this, &FRigHierarchyContainer::HandleOnElementAdded);
+	ControlHierarchy.OnControlRemoved.AddRaw(this, &FRigHierarchyContainer::HandleOnElementRemoved);
+	ControlHierarchy.OnControlRenamed.AddRaw(this, &FRigHierarchyContainer::HandleOnElementRenamed);
+	ControlHierarchy.OnControlReparented.AddRaw(this, &FRigHierarchyContainer::HandleOnElementReparented);
+	ControlHierarchy.OnControlSelected.AddRaw(this, &FRigHierarchyContainer::HandleOnElementSelected);
+
+	CurveContainer.OnCurveAdded.RemoveAll(this);
+	CurveContainer.OnCurveRemoved.RemoveAll(this);
+	CurveContainer.OnCurveRenamed.RemoveAll(this);
+	CurveContainer.OnCurveSelected.RemoveAll(this);
+
+	CurveContainer.OnCurveAdded.AddRaw(this, &FRigHierarchyContainer::HandleOnElementAdded);
+	CurveContainer.OnCurveRemoved.AddRaw(this, &FRigHierarchyContainer::HandleOnElementRemoved);
+	CurveContainer.OnCurveRenamed.AddRaw(this, &FRigHierarchyContainer::HandleOnElementRenamed);
+	CurveContainer.OnCurveSelected.AddRaw(this, &FRigHierarchyContainer::HandleOnElementSelected);
+
+	// wire them between each other
+	BoneHierarchy.OnBoneRenamed.RemoveAll(&SpaceHierarchy);
+	SpaceHierarchy.OnSpaceRenamed.RemoveAll(&ControlHierarchy);
+	ControlHierarchy.OnControlRenamed.RemoveAll(&SpaceHierarchy);
+	BoneHierarchy.OnBoneRenamed.AddRaw(&SpaceHierarchy, &FRigSpaceHierarchy::HandleOnElementRenamed);
+	SpaceHierarchy.OnSpaceRenamed.AddRaw(&ControlHierarchy, &FRigControlHierarchy::HandleOnElementRenamed);
+	ControlHierarchy.OnControlRenamed.AddRaw(&SpaceHierarchy, &FRigSpaceHierarchy::HandleOnElementRenamed);
+	BoneHierarchy.OnBoneRemoved.AddRaw(&SpaceHierarchy, &FRigSpaceHierarchy::HandleOnElementRemoved);
+	ControlHierarchy.OnControlRemoved.AddRaw(&SpaceHierarchy, &FRigSpaceHierarchy::HandleOnElementRemoved);
+	SpaceHierarchy.OnSpaceRemoved.AddRaw(&ControlHierarchy, &FRigControlHierarchy::HandleOnElementRemoved);
+
+#endif
+
+	BoneHierarchy.Initialize();
+	SpaceHierarchy.Initialize();
+	ControlHierarchy.Initialize();
+	CurveContainer.Initialize();
+
+	ResetTransforms();
+}
+
+void FRigHierarchyContainer::Reset()
+{
+	BoneHierarchy.Reset();
+	SpaceHierarchy.Reset();
+	ControlHierarchy.Reset();
+	CurveContainer.Reset();
+
+	Initialize();
+}
+
+void FRigHierarchyContainer::ResetTransforms()
+{
+	BoneHierarchy.ResetTransforms();
+	SpaceHierarchy.ResetTransforms();
+	ControlHierarchy.ResetValues();
+	CurveContainer.ResetValues();
+}
+
+FTransform FRigHierarchyContainer::GetInitialTransform(ERigElementType InElementType, int32 InIndex) const
+{
+	if(InIndex == INDEX_NONE)
+	{
+		return FTransform::Identity;
+	}
+
+	switch(InElementType)
+	{
+		case ERigElementType::Bone:
+		{
+			return BoneHierarchy.GetInitialTransform(InIndex);
+		}
+		case ERigElementType::Space:
+		{
+			return SpaceHierarchy.GetInitialTransform(InIndex);
+		}
+		case ERigElementType::Control:
+		{
+			return ControlHierarchy.GetInitialValue<FTransform>(InIndex);
+		}
+		case ERigElementType::Curve:
+		{
+			break;
+		}
+	}
+
+	return FTransform::Identity;
+}
+
+#if WITH_EDITOR
+
+void FRigHierarchyContainer::SetInitialTransform(ERigElementType InElementType, int32 InIndex, const FTransform& InTransform)
+{
+	if (InIndex == INDEX_NONE)
+	{
+		return;
+	}
+
+	switch(InElementType)
+	{
+		case ERigElementType::Bone:
+		{
+			BoneHierarchy.SetInitialTransform(InIndex, InTransform);
+			break;
+		}
+		case ERigElementType::Space:
+		{
+			SpaceHierarchy.SetInitialTransform(InIndex, InTransform);
+			break;
+		}
+		case ERigElementType::Control:
+		{
+			ControlHierarchy.SetInitialValue<FTransform>(InIndex, InTransform);
+			break;
+		}
+		case ERigElementType::Curve:
+		{
+			break;
+		}
+	}
+}
+
+#endif
+
+FTransform FRigHierarchyContainer::GetInitialGlobalTransform(ERigElementType InElementType, int32 InIndex) const
+{
+	if (InIndex == INDEX_NONE)
+	{
+		return FTransform::Identity;
+	}
+
+	switch (InElementType)
+	{
+		case ERigElementType::Bone:
+		{
+			return BoneHierarchy.GetInitialTransform(InIndex);
+		}
+		case ERigElementType::Space:
+		{
+			return SpaceHierarchy.GetInitialGlobalTransform(InIndex);
+		}
+		case ERigElementType::Control:
+		{
+			return ControlHierarchy.GetInitialGlobalTransform(InIndex);
+		}
+		case ERigElementType::Curve:
+		{
+			break;
+		}
+	}
+
+	return FTransform::Identity;
+}
+
+void FRigHierarchyContainer::SetInitialGlobalTransform(ERigElementType InElementType, int32 InIndex, const FTransform& InTransform)
+{
+	if (InIndex == INDEX_NONE)
+	{
+		return;
+	}
+
+	switch (InElementType)
+	{
+		case ERigElementType::Bone:
+		{
+			BoneHierarchy.SetInitialTransform(InIndex, InTransform);
+			break;
+		}
+		case ERigElementType::Space:
+		{
+			SpaceHierarchy.SetInitialGlobalTransform(InIndex, InTransform);
+			break;
+		}
+		case ERigElementType::Control:
+		{
+			ControlHierarchy.SetInitialGlobalTransform(InIndex, InTransform);
+			break;
+		}
+		case ERigElementType::Curve:
+		{
+			break;
+		}
+	}
+}
+
+FTransform FRigHierarchyContainer::GetLocalTransform(ERigElementType InElementType, int32 InIndex) const
+{
+	if(InIndex == INDEX_NONE)
+	{
+		return FTransform::Identity;
+	}
+
+	switch(InElementType)
+	{
+		case ERigElementType::Bone:
+		{
+			return BoneHierarchy.GetLocalTransform(InIndex);
+		}
+		case ERigElementType::Space:
+		{
+			return SpaceHierarchy.GetLocalTransform(InIndex);
+		}
+		case ERigElementType::Control:
+		{
+			return ControlHierarchy.GetLocalTransform(InIndex);
+		}
+		case ERigElementType::Curve:
+		{
+			break;
+		}
+	}
+
+	return FTransform::Identity;
+}
+
+void FRigHierarchyContainer::SetLocalTransform(ERigElementType InElementType, int32 InIndex, const FTransform& InTransform)
+{
+	if (InIndex == INDEX_NONE)
+	{
+		return;
+	}
+
+	switch(InElementType)
+	{
+		case ERigElementType::Bone:
+		{
+			BoneHierarchy.SetLocalTransform(InIndex, InTransform);
+			break;
+		}
+		case ERigElementType::Space:
+		{
+			SpaceHierarchy.SetLocalTransform(InIndex, InTransform);
+			break;
+		}
+		case ERigElementType::Control:
+		{
+			ControlHierarchy.SetLocalTransform(InIndex, InTransform);
+			break;
+		}
+		case ERigElementType::Curve:
+		{
+			break;
+		}
+	}
+}
+
+FTransform FRigHierarchyContainer::GetGlobalTransform(ERigElementType InElementType, int32 InIndex) const
+{
+	if(InIndex == INDEX_NONE)
+	{
+		return FTransform::Identity;
+	}
+
+	switch(InElementType)
+	{
+		case ERigElementType::Bone:
+		{
+			return BoneHierarchy.GetGlobalTransform(InIndex);
+		}
+		case ERigElementType::Space:
+		{
+			return SpaceHierarchy.GetGlobalTransform(InIndex);
+		}
+		case ERigElementType::Control:
+		{
+			return ControlHierarchy.GetGlobalTransform(InIndex);
+		}
+		case ERigElementType::Curve:
+		{
+			break;
+		}
+	}
+
+	return FTransform::Identity;
+}
+
+void FRigHierarchyContainer::SetGlobalTransform(ERigElementType InElementType, int32 InIndex, const FTransform& InTransform)
+{
+	if (InIndex == INDEX_NONE)
+	{
+		return;
+	}
+
+	switch(InElementType)
+	{
+		case ERigElementType::Bone:
+		{
+			BoneHierarchy.SetGlobalTransform(InIndex, InTransform);
+			break;
+		}
+		case ERigElementType::Space:
+		{
+			SpaceHierarchy.SetGlobalTransform(InIndex, InTransform);
+			break;
+		}
+		case ERigElementType::Control:
+		{
+			ControlHierarchy.SetGlobalTransform(InIndex, InTransform);
+			break;
+		}
+		case ERigElementType::Curve:
+		{
+			break;
+		}
+	}
+}
+
+#if WITH_EDITOR
+
+void FRigHierarchyContainer::HandleOnElementAdded(FRigHierarchyContainer* InContainer, const FRigElementKey& InKey)
+{
+	// todo
+	OnElementAdded.Broadcast(InContainer, InKey);
+	OnElementChanged.Broadcast(InContainer, InKey);
+}
+
+void FRigHierarchyContainer::HandleOnElementRemoved(FRigHierarchyContainer* InContainer, const FRigElementKey& InKey)
+{
+	// todo
+	OnElementRemoved.Broadcast(InContainer, InKey);
+	OnElementChanged.Broadcast(InContainer, InKey);
+}
+
+void FRigHierarchyContainer::HandleOnElementRenamed(FRigHierarchyContainer* InContainer, ERigElementType InElementType, const FName& InOldName, const FName& InNewName)
+{
+	// todo
+	OnElementRenamed.Broadcast(InContainer, InElementType, InOldName, InNewName);
+	OnElementChanged.Broadcast(InContainer, FRigElementKey(InNewName, InElementType));
+}
+
+void FRigHierarchyContainer::HandleOnElementReparented(FRigHierarchyContainer* InContainer, const FRigElementKey& InKey, const FName& InOldParentName, const FName& InNewParentName)
+{
+	// todo
+	OnElementReparented.Broadcast(InContainer, InKey, InOldParentName, InNewParentName);
+	OnElementChanged.Broadcast(InContainer, InKey);
+}
+
+void FRigHierarchyContainer::HandleOnElementSelected(FRigHierarchyContainer* InContainer, const FRigElementKey& InKey, bool bSelected)
+{
+	OnElementSelected.Broadcast(InContainer, InKey, bSelected);
+	OnElementChanged.Broadcast(InContainer, InKey);
+}
+
+#endif
+
+bool FRigHierarchyContainer::IsParentedTo(ERigElementType InChildType, int32 InChildIndex, ERigElementType InParentType, int32 InParentIndex) const
+{
+	if (InParentIndex == INDEX_NONE || InChildIndex == INDEX_NONE)
+	{
+		return false;
+	}
+
+	switch (InChildType)
+	{
+		case ERigElementType::Curve:
+		{
+			return false;
+		}
+		case ERigElementType::Bone:
+		{
+			switch (InParentType)
+			{
+				case ERigElementType::Bone:
+				{
+					if (BoneHierarchy[InChildIndex].ParentIndex != INDEX_NONE)
+					{
+						if (BoneHierarchy[InChildIndex].ParentIndex == InParentIndex)
+						{
+							return true;
+						}
+						return IsParentedTo(ERigElementType::Bone, BoneHierarchy[InChildIndex].ParentIndex, InParentType, InParentIndex);
+					}
+					// no break - fall through to next case
+				}
+				case ERigElementType::Space:
+				case ERigElementType::Control:
+				case ERigElementType::Curve:
+				{
+					return false;
+				}
+			}
+		}
+		case ERigElementType::Space:
+		{
+			const FRigSpace& ChildSpace = SpaceHierarchy[InChildIndex];
+			switch (ChildSpace.SpaceType)
+			{
+				case ERigSpaceType::Global:
+				{
+					return false;
+				}
+				case ERigSpaceType::Bone:
+				{
+					if (ChildSpace.ParentIndex == InParentIndex && InParentType == ERigElementType::Bone)
+					{
+						return true;
+					}
+					return IsParentedTo(ERigElementType::Bone, ChildSpace.ParentIndex, InParentType, InParentIndex);
+				}
+				case ERigSpaceType::Space:
+				{
+					if (ChildSpace.ParentIndex == InParentIndex && InParentType == ERigElementType::Space)
+					{
+						return true;
+					}
+					return IsParentedTo(ERigElementType::Space, ChildSpace.ParentIndex, InParentType, InParentIndex);
+				}
+				case ERigSpaceType::Control:
+				{
+					if (ChildSpace.ParentIndex == InParentIndex && InParentType == ERigElementType::Control)
+					{
+						return true;
+					}
+					return IsParentedTo(ERigElementType::Control, ChildSpace.ParentIndex, InParentType, InParentIndex);
+				}
+			}
+		}
+		case ERigElementType::Control:
+		{
+			const FRigControl& ChildControl = ControlHierarchy[InChildIndex];
+			switch (InParentType)
+			{
+				case ERigElementType::Space:
+				{
+					if (ChildControl.SpaceIndex == InParentIndex)
+					{
+						return true;
+					}
+					// no break - fall through to next cases
+				}
+				case ERigElementType::Control:
+				case ERigElementType::Bone:
+				{
+					if (ChildControl.SpaceIndex != INDEX_NONE)
+					{
+						return IsParentedTo(ERigElementType::Space, ChildControl.SpaceIndex, InParentType, InParentIndex);
+					}
+					else if (ChildControl.ParentIndex != INDEX_NONE)
+					{
+						if (ChildControl.ParentIndex == InParentIndex)
+						{
+							return true;
+						}
+						return IsParentedTo(ERigElementType::Control, ChildControl.ParentIndex, InParentType, InParentIndex);
+					}
+				}
+				case ERigElementType::Curve:
+				{
+					return false;
+				}
+			}
+		}
+	}
+
+	return false;
+}
+
+#if WITH_EDITOR
+
+FString FRigHierarchyContainer::ExportSelectionToText() const
+{
+	TArray<FRigElementKey> Selection = CurrentSelection();
+	return ExportToText(Selection);
+}
+
+FString FRigHierarchyContainer::ExportToText(const TArray<FRigElementKey>& InSelection) const
+{
+	FRigHierarchyCopyPasteContent Data;
+	for (const FRigElementKey& Key : InSelection)
+	{
+		Data.Types.Add(Key.Type);
+		FString Content;
+		switch (Key.Type)
+		{
+		case ERigElementType::Bone:
+		{
+			FRigBone::StaticStruct()->ExportText(Content, &BoneHierarchy[Key.Name], nullptr, nullptr, PPF_None, nullptr);
+			break;
+		}
+		case ERigElementType::Control:
+		{
+			FRigControl::StaticStruct()->ExportText(Content, &ControlHierarchy[Key.Name], nullptr, nullptr, PPF_None, nullptr);
+			break;
+		}
+		case ERigElementType::Space:
+		{
+			FRigSpace::StaticStruct()->ExportText(Content, &SpaceHierarchy[Key.Name], nullptr, nullptr, PPF_None, nullptr);
+			break;
+		}
+		case ERigElementType::Curve:
+		{
+			FRigCurve::StaticStruct()->ExportText(Content, &CurveContainer[Key.Name], nullptr, nullptr, PPF_None, nullptr);
+			break;
+		}
+		default:
+		{
+			ensure(false);
+			break;
+		}
+		}
+		Data.Contents.Add(Content);
+		Data.LocalTransforms.Add(GetLocalTransform(Key));
+		Data.GlobalTransforms.Add(GetGlobalTransform(Key));
+	}
+
+	FString ExportedText;
+	FRigHierarchyCopyPasteContent::StaticStruct()->ExportText(ExportedText, &Data, nullptr, nullptr, PPF_None, nullptr);
+	return ExportedText;
+}
+
+class FRigHierarchyContainerImportErrorContext : public FOutputDevice
+{
+public:
+
+	int32 NumErrors;
+
+	FRigHierarchyContainerImportErrorContext()
+		: FOutputDevice()
+		, NumErrors(0)
+	{
+	}
+
+	virtual void Serialize(const TCHAR* V, ELogVerbosity::Type Verbosity, const class FName& Category) override
+	{
+		UE_LOG(LogControlRig, Error, TEXT("Error Importing To Hierarchy: %s"), V);
+		NumErrors++;
+	}
+};
+
+TArray<FRigElementKey> FRigHierarchyContainer::ImportFromText(const FString& InContent, ERigHierarchyImportMode InImportMode, bool bSelectNewElements)
+{
+	TArray<FRigElementKey> PastedKeys;
+
+	FRigHierarchyCopyPasteContent Data;
+	FRigHierarchyContainerImportErrorContext ErrorPipe;
+	FRigHierarchyCopyPasteContent::StaticStruct()->ImportText(*InContent, &Data, nullptr, EPropertyPortFlags::PPF_None, &ErrorPipe, FRigHierarchyCopyPasteContent::StaticStruct()->GetName(), true);
+	if (ErrorPipe.NumErrors > 0)
+	{
+		return PastedKeys;
+	}
+
+	if (Data.Contents.Num() == 0 ||
+		(Data.Types.Num() != Data.Contents.Num()) ||
+		(Data.LocalTransforms.Num() != Data.Contents.Num()) ||
+		(Data.GlobalTransforms.Num() != Data.Contents.Num()))
+	{
+		return PastedKeys;
+	}
+
+	TMap<FRigElementKey, FRigElementKey> ElementMap;
+	for (const FRigBone& Element : BoneHierarchy)
+	{
+		ElementMap.Add(Element.GetElementKey(), Element.GetElementKey());
+	}
+	for (const FRigControl& Element : ControlHierarchy)
+	{
+		ElementMap.Add(Element.GetElementKey(), Element.GetElementKey());
+	}
+	for (const FRigSpace& Element : SpaceHierarchy)
+	{
+		ElementMap.Add(Element.GetElementKey(), Element.GetElementKey());
+	}
+	for (const FRigCurve& Element : CurveContainer)
+	{
+		ElementMap.Add(Element.GetElementKey(), Element.GetElementKey());
+	}
+
+	TArray<TSharedPtr<FRigElement>> Elements;
+	for (int32 Index = 0; Index < Data.Types.Num(); Index++)
+	{
+		ErrorPipe.NumErrors = 0;
+		TSharedPtr<FRigElement> NewElement;
+		switch (Data.Types[Index])
+		{
+			case ERigElementType::Bone:
+			{
+				NewElement = MakeShared<FRigBone>();
+				FRigBone::StaticStruct()->ImportText(*Data.Contents[Index], NewElement.Get(), nullptr, EPropertyPortFlags::PPF_None, &ErrorPipe, FRigBone::StaticStruct()->GetName(), true);
+				break;
+			}
+			case ERigElementType::Control:
+			{
+				NewElement = MakeShared<FRigControl>();
+				FRigControl::StaticStruct()->ImportText(*Data.Contents[Index], NewElement.Get(), nullptr, EPropertyPortFlags::PPF_None, &ErrorPipe, FRigControl::StaticStruct()->GetName(), true);
+				break;
+			}
+			case ERigElementType::Space:
+			{
+				NewElement = MakeShared<FRigSpace>();
+				FRigSpace::StaticStruct()->ImportText(*Data.Contents[Index], NewElement.Get(), nullptr, EPropertyPortFlags::PPF_None, &ErrorPipe, FRigSpace::StaticStruct()->GetName(), true);
+				break;
+			}
+			case ERigElementType::Curve:
+			{
+				NewElement = MakeShared<FRigCurve>();
+				FRigCurve::StaticStruct()->ImportText(*Data.Contents[Index], NewElement.Get(), nullptr, EPropertyPortFlags::PPF_None, &ErrorPipe, FRigCurve::StaticStruct()->GetName(), true);
+				break;
+			}
+			default:
+			{
+				ensure(false);
+				break;
+			}
+		}
+
+		if (ErrorPipe.NumErrors > 0)
+		{
+			return PastedKeys;
+		}
+
+		Elements.Add(NewElement);
+	}
+
+	TArray<FRigElementKey> Selection = CurrentSelection();
+
+	switch (InImportMode)
+	{
+		case ERigHierarchyImportMode::Append:
+		{
+			for (int32 Index = 0; Index < Data.Types.Num(); Index++)
+			{
+				switch (Data.Types[Index])
+				{
+					case ERigElementType::Bone:
+					{
+						const FRigBone& Element = *static_cast<FRigBone*>(Elements[Index].Get());
+
+						FName ParentName = NAME_None;
+						if (const FRigElementKey* ParentKey = ElementMap.Find(Element.GetParentElementKey()))
+						{
+							ParentName = ParentKey->Name;
+						}
+
+						FRigBone& NewElement = BoneHierarchy.Add(Element.Name, ParentName, Element.InitialTransform, Element.LocalTransform, Data.GlobalTransforms[Index]);
+						ElementMap.FindOrAdd(Element.GetElementKey()) = NewElement.GetElementKey();
+						PastedKeys.Add(NewElement.GetElementKey());
+						break;
+					}
+					case ERigElementType::Control:
+					{
+						const FRigControl& Element = *static_cast<FRigControl*>(Elements[Index].Get());
+
+						FName ParentName = NAME_None;
+						if (const FRigElementKey* ParentKey = ElementMap.Find(Element.GetParentElementKey()))
+						{
+							ParentName = ParentKey->Name;
+						}
+
+						FName SpaceName = NAME_None;
+						if (const FRigElementKey* SpaceKey = ElementMap.Find(Element.GetSpaceElementKey()))
+						{
+							SpaceName = SpaceKey->Name;
+						}
+
+						FRigControl& NewElement = ControlHierarchy.Add(Element.Name, Element.ControlType, ParentName, SpaceName, Element.InitialValue, Element.GizmoName, Element.GizmoTransform, Element.GizmoColor);
+						ElementMap.FindOrAdd(Element.GetElementKey()) = NewElement.GetElementKey();
+						PastedKeys.Add(NewElement.GetElementKey());
+
+						break;
+					}
+					case ERigElementType::Space:
+					{
+						const FRigSpace& Element = *static_cast<FRigSpace*>(Elements[Index].Get());
+
+						FName ParentName = NAME_None;
+						if (const FRigElementKey* ParentKey = ElementMap.Find(Element.GetParentElementKey()))
+						{
+							ParentName = ParentKey->Name;
+						}
+
+						FRigSpace& NewElement = SpaceHierarchy.Add(Element.Name, Element.SpaceType, ParentName, Element.InitialTransform);
+						ElementMap.FindOrAdd(Element.GetElementKey()) = NewElement.GetElementKey();
+						PastedKeys.Add(NewElement.GetElementKey());
+						break;
+					}
+					case ERigElementType::Curve:
+					{
+						const FRigCurve& Element = *static_cast<FRigCurve*>(Elements[Index].Get());
+						FRigCurve& NewElement = CurveContainer.Add(Element.Name, Element.Value);
+						ElementMap.FindOrAdd(Element.GetElementKey()) = NewElement.GetElementKey();
+						PastedKeys.Add(NewElement.GetElementKey());
+						break;
+					}
+					default:
+					{
+						ensure(false);
+						break;
+					}
+				}
+			}
+			break;
+		}
+		case ERigHierarchyImportMode::ReplaceLocalTransform:
+		case ERigHierarchyImportMode::ReplaceGlobalTransform:
+		{
+			if (Selection.Num() == 0)
+			{
+				for (const TSharedPtr<FRigElement>& Element : Elements)
+				{
+					Selection.Add(Element->GetElementKey());
+				}
+			}
+			else if (Selection.Num() != Data.Types.Num())
+			{
+				// todo: error message
+				return PastedKeys;
+			}
+
+			for (int32 Index = 0; Index < Data.Types.Num(); Index++)
+			{
+				if(InImportMode == ERigHierarchyImportMode::ReplaceLocalTransform)
+				{
+					Data.LocalTransforms[Index].NormalizeRotation();
+					SetLocalTransform(Selection[Index], Data.LocalTransforms[Index]);
+				}
+				else
+				{
+					Data.GlobalTransforms[Index].NormalizeRotation();
+					SetGlobalTransform(Selection[Index], Data.GlobalTransforms[Index]);
+				}
+				PastedKeys.Add(Selection[Index]);
+			}
+			break;
+		}
+		default:
+		{
+			ensure(false);
+			break;
+		}
+	}
+
+	if (PastedKeys.Num() > 0 && InImportMode == ERigHierarchyImportMode::Append)
+	{
+		BoneHierarchy.Initialize();
+		SpaceHierarchy.Initialize();
+		ControlHierarchy.Initialize();
+		CurveContainer.Initialize();
+	}
+
+	if (bSelectNewElements && PastedKeys.Num() > 0)
+	{
+		ClearSelection();
+		for (const FRigElementKey& Key : PastedKeys)
+		{
+			Select(Key, true);
+		}
+	}
+
+	return PastedKeys;
+}
+
+TArray<FRigElementKey> FRigHierarchyContainer::CurrentSelection() const
+{
+	TArray<FRigElementKey> Selection;
+	for (const FRigBone& Element : BoneHierarchy)
+	{
+		if (BoneHierarchy.IsSelected(Element.Name))
+		{
+			Selection.Add(Element.GetElementKey());
+		}
+	}
+	for (const FRigControl& Element : ControlHierarchy)
+	{
+		if (ControlHierarchy.IsSelected(Element.Name))
+		{
+			Selection.Add(Element.GetElementKey());
+		}
+	}
+	for (const FRigSpace& Element : SpaceHierarchy)
+	{
+		if (SpaceHierarchy.IsSelected(Element.Name))
+		{
+			Selection.Add(Element.GetElementKey());
+		}
+	}
+	for (const FRigCurve& Element : CurveContainer)
+	{
+		if (CurveContainer.IsSelected(Element.Name))
+		{
+			Selection.Add(Element.GetElementKey());
+		}
+	}
+
+	Selection.Sort([&](const FRigElementKey& A, const FRigElementKey& B) { 
+		int32 IndexA = GetIndex(A);
+		int32 IndexB = GetIndex(B);
+		return IsParentedTo(B.Type, IndexB, A.Type, IndexA);
+	});
+
+	return Selection;
+}
+
+TArray<FRigElementKey> FRigHierarchyContainer::GetAllItems(bool bSort) const
+{
+	TArray<FRigElementKey> Items;
+	for (const FRigBone& Element : BoneHierarchy)
+	{
+		Items.Add(Element.GetElementKey());
+	}
+	for (const FRigControl& Element : ControlHierarchy)
+	{
+		Items.Add(Element.GetElementKey());
+	}
+	for (const FRigSpace& Element : SpaceHierarchy)
+	{
+		Items.Add(Element.GetElementKey());
+	}
+	for (const FRigCurve& Element : CurveContainer)
+	{
+		Items.Add(Element.GetElementKey());
+	}
+	
+	if (bSort)
+	{
+		Items.Sort([&](const FRigElementKey& A, const FRigElementKey& B) {
+			int32 IndexA = GetIndex(A);
+			int32 IndexB = GetIndex(B);
+			return IsParentedTo(B.Type, IndexB, A.Type, IndexA);
+		});
+	}
+
+	return Items;
+}
+
+bool FRigHierarchyContainer::Select(const FRigElementKey& InKey, bool bSelect)
+{
+	switch(InKey.Type)
+	{
+		case ERigElementType::Bone:
+		{
+			return BoneHierarchy.Select(InKey.Name, bSelect);
+		}
+		case ERigElementType::Space:
+		{
+			return SpaceHierarchy.Select(InKey.Name, bSelect);
+		}
+		case ERigElementType::Control:
+		{
+			return ControlHierarchy.Select(InKey.Name, bSelect);
+		}
+		case ERigElementType::Curve:
+		{
+			return CurveContainer.Select(InKey.Name, bSelect);
+		}
+	}
+	return false;
+}
+
+bool FRigHierarchyContainer::ClearSelection()
+{
+	ClearSelection(ERigElementType::Bone);
+	ClearSelection(ERigElementType::Control);
+	ClearSelection(ERigElementType::Space);
+	ClearSelection(ERigElementType::Curve);
+	return true;
+}
+
+bool FRigHierarchyContainer::ClearSelection(ERigElementType InElementType)
+{
+	switch(InElementType)
+	{
+		case ERigElementType::Bone:
+		{
+			return BoneHierarchy.ClearSelection();
+		}
+		case ERigElementType::Space:
+		{
+			return SpaceHierarchy.ClearSelection();
+		}
+		case ERigElementType::Control:
+		{
+			return ControlHierarchy.ClearSelection();
+		}
+		case ERigElementType::Curve:
+		{
+			return CurveContainer.ClearSelection();
+		}
+	}
+	return false;
+}
+
+bool FRigHierarchyContainer::IsSelected(const FRigElementKey& InKey) const
+{
+	switch(InKey.Type)
+	{
+		case ERigElementType::Bone:
+		{
+			return BoneHierarchy.IsSelected(InKey.Name);
+		}
+		case ERigElementType::Space:
+		{
+			return SpaceHierarchy.IsSelected(InKey.Name);
+		}
+		case ERigElementType::Control:
+		{
+			return ControlHierarchy.IsSelected(InKey.Name);
+		}
+		case ERigElementType::Curve:
+		{
+			return CurveContainer.IsSelected(InKey.Name);
+		}
+	}
+	return false;
+}
+
+#endif
\ No newline at end of file
diff --git a/Engine/Plugins/Experimental/ControlRig/Source/ControlRig/Private/Rigs/RigSpaceHierarchy.cpp b/Engine/Plugins/Experimental/ControlRig/Source/ControlRig/Private/Rigs/RigSpaceHierarchy.cpp
new file mode 100644
--- /dev/null
+++ b/Engine/Plugins/Experimental/ControlRig/Source/ControlRig/Private/Rigs/RigSpaceHierarchy.cpp
@@ -0,0 +1,644 @@
+// Copyright 1998-2019 Epic Games, Inc. All Rights Reserved.
+
+#include "Rigs/RigSpaceHierarchy.h"
+#include "ControlRig.h"
+#include "HelperUtil.h"
+
+////////////////////////////////////////////////////////////////////////////////
+// FRigSpaceHierarchy
+////////////////////////////////////////////////////////////////////////////////
+
+FRigSpaceHierarchy::FRigSpaceHierarchy()
+	:Container(nullptr)
+{
+}
+
+FRigSpaceHierarchy& FRigSpaceHierarchy::operator= (const FRigSpaceHierarchy &InOther)
+{
+#if WITH_EDITOR
+	for (int32 Index = Num() - 1; Index >= 0; Index--)
+	{
+		FRigSpace SpaceToRemove = Spaces[Index];
+		OnSpaceRemoved.Broadcast(Container, FRigElementKey(SpaceToRemove.Name, ERigElementType::Space));
+	}
+#endif
+
+	Spaces.Reset();
+	Spaces.Append(InOther.Spaces);
+	NameToIndexMapping.Reset();
+	RefreshMapping();
+
+#if WITH_EDITOR
+	for (const FRigSpace& SpaceAdded : Spaces)
+	{
+		OnSpaceAdded.Broadcast(Container, FRigElementKey(SpaceAdded.Name, ERigElementType::Space));
+	}
+#endif
+
+	return *this;
+}
+
+FName FRigSpaceHierarchy::GetSafeNewName(const FName& InPotentialNewName) const
+{
+	FName Name = InPotentialNewName;
+	int32 Suffix = 1;
+	while(!IsNameAvailable(Name))
+	{
+		Name = *FString::Printf(TEXT("%s_%d"), *InPotentialNewName.ToString(), ++Suffix);
+	}
+	return Name;
+}
+
+FRigSpace& FRigSpaceHierarchy::Add(const FName& InNewName, ERigSpaceType InSpaceType, const FName& InParentName, const FTransform& InTransform)
+{
+	FRigSpace NewSpace;
+	NewSpace.Name = GetSafeNewName(InNewName);
+	NewSpace.ParentIndex = GetParentIndex(InSpaceType, InParentName);
+	NewSpace.SpaceType = NewSpace.ParentIndex == INDEX_NONE ? ERigSpaceType::Global : InSpaceType;
+	NewSpace.ParentName = NewSpace.ParentIndex == INDEX_NONE ? NAME_None : InParentName;
+	NewSpace.InitialTransform = InTransform;
+	NewSpace.LocalTransform = InTransform;
+	FName NewSpaceName = NewSpace.Name;
+	Spaces.Add(NewSpace);
+	RefreshMapping();
+
+#if WITH_EDITOR
+	OnSpaceAdded.Broadcast(Container, NewSpace.GetElementKey());
+#endif
+
+	int32 Index = GetIndex(NewSpaceName);
+	return Spaces[Index];
+}
+
+bool FRigSpaceHierarchy::Reparent(const FName& InName, ERigSpaceType InSpaceType, const FName& InNewParentName)
+{
+	int32 Index = GetIndex(InName);
+	if (Index != INDEX_NONE)
+	{
+		FRigSpace& Space = Spaces[Index];
+
+#if WITH_EDITOR
+		FName OldParentName = Space.ParentName;
+#endif
+
+		int32 ParentIndex = GetParentIndex(InSpaceType, InNewParentName);
+		if (ParentIndex != INDEX_NONE)
+		{
+			if (Container != nullptr)
+			{
+				switch (InSpaceType)
+				{
+					case ERigSpaceType::Global:
+					{
+						ParentIndex = INDEX_NONE;
+						break;
+					}
+					case ERigSpaceType::Bone:
+					{
+						break;
+					}
+					case ERigSpaceType::Space:
+					{
+						if (Container->IsParentedTo(ERigElementType::Space, ParentIndex, ERigElementType::Space, Index))
+						{
+							ParentIndex = INDEX_NONE;
+						}
+						break;
+					}
+					case ERigSpaceType::Control:
+					{
+						if (Container->IsParentedTo(ERigElementType::Control, ParentIndex, ERigElementType::Space, Index))
+						{
+							ParentIndex = INDEX_NONE;
+						}
+						break;
+					}
+				}
+			}
+		}
+
+		Space.ParentIndex = ParentIndex;
+		Space.SpaceType = Space.ParentIndex == INDEX_NONE ? ERigSpaceType::Global : InSpaceType;
+		Space.ParentName = Space.ParentIndex == INDEX_NONE ? NAME_None : InNewParentName;
+
+#if WITH_EDITOR
+		FName NewParentName = Space.ParentName;
+#endif
+
+		RefreshMapping();
+
+#if WITH_EDITOR
+		if (OldParentName != NewParentName)
+		{
+			OnSpaceReparented.Broadcast(Container, Space.GetElementKey(), OldParentName, NewParentName);
+		}
+#endif		
+		return Spaces[GetIndex(InName)].ParentName == InNewParentName;
+	}
+	return false;
+}
+
+FRigSpace FRigSpaceHierarchy::Remove(const FName& InNameToRemove)
+{
+	int32 IndexToDelete = GetIndex(InNameToRemove);
+#if WITH_EDITOR
+	Select(InNameToRemove, false);
+#endif
+	FRigSpace RemovedSpace = Spaces[IndexToDelete];
+	Spaces.RemoveAt(IndexToDelete);
+
+#if WITH_EDITOR
+	TArray<FName> SpacesReparented;
+#endif
+	for (FRigSpace& ChildSpace : Spaces)
+	{
+		if (ChildSpace.SpaceType == ERigSpaceType::Space && ChildSpace.ParentName == InNameToRemove)
+		{
+			ChildSpace.SpaceType = ERigSpaceType::Global;
+			ChildSpace.ParentIndex = INDEX_NONE;
+			ChildSpace.ParentName = NAME_None;
+#if WITH_EDITOR
+			SpacesReparented.Add(ChildSpace.Name);
+#endif
+		}
+	}
+
+	RefreshMapping();
+
+#if WITH_EDITOR
+	for (const FName& ReparentedSpace : SpacesReparented)
+	{
+		OnSpaceReparented.Broadcast(Container, FRigElementKey(ReparentedSpace, RigElementType()), InNameToRemove, NAME_None);
+	}
+	OnSpaceRemoved.Broadcast(Container, RemovedSpace.GetElementKey());
+#endif
+
+	return RemovedSpace;
+}
+
+FName FRigSpaceHierarchy::GetName(int32 InIndex) const
+{
+	if (Spaces.IsValidIndex(InIndex))
+	{
+		return Spaces[InIndex].Name;
+	}
+
+	return NAME_None;
+}
+
+int32 FRigSpaceHierarchy::GetIndexSlow(const FName& InName) const
+{
+	for (int32 Index = 0; Index < Spaces.Num(); ++Index)
+	{
+		if (Spaces[Index].Name == InName)
+		{
+			return Index;
+		}
+	}
+
+	return INDEX_NONE;
+}
+
+void FRigSpaceHierarchy::SetGlobalTransform(const FName& InName, const FTransform& InTransform)
+{
+	SetGlobalTransform(GetIndex(InName), InTransform);
+}
+
+void FRigSpaceHierarchy::SetGlobalTransform(int32 InIndex, const FTransform& InTransform)
+{
+	if (Container == nullptr)
+	{
+		SetLocalTransform(InIndex, InTransform);
+		return;
+	}
+
+	if (Spaces.IsValidIndex(InIndex))
+	{
+		FRigSpace& Space = Spaces[InIndex];
+		if (Space.SpaceType == ERigSpaceType::Global)
+		{
+			SetLocalTransform(InIndex, InTransform);
+			return;
+		}
+
+		FTransform ParentTransform = Container->GetGlobalTransform(Space.GetParentElementKey());
+		SetLocalTransform(InIndex, InTransform.GetRelativeTransform(ParentTransform));
+	}
+}
+
+FTransform FRigSpaceHierarchy::GetGlobalTransform(const FName& InName) const
+{
+	return GetGlobalTransform(GetIndex(InName));
+}
+
+FTransform FRigSpaceHierarchy::GetGlobalTransform(int32 InIndex) const
+{
+	if(Container == nullptr)
+	{
+		return GetLocalTransform(InIndex);
+	}
+
+	if (Spaces.IsValidIndex(InIndex))
+	{
+		const FRigSpace& Space = Spaces[InIndex];
+		switch (Space.SpaceType)
+		{
+			case ERigSpaceType::Global:
+			{
+				return Space.LocalTransform;
+			}
+			case ERigSpaceType::Bone:
+			{
+				FTransform Transform = Space.LocalTransform * Container->GetGlobalTransform(ERigElementType::Bone, Space.ParentIndex);
+				Transform.NormalizeRotation();
+				return Transform;
+			}
+			case ERigSpaceType::Space:
+			{
+				FTransform Transform = Space.LocalTransform * Container->GetGlobalTransform(ERigElementType::Space, Space.ParentIndex);
+				Transform.NormalizeRotation();
+				return Transform;
+			}
+			case ERigSpaceType::Control:
+			{
+				FTransform Transform = Space.LocalTransform * Container->GetGlobalTransform(ERigElementType::Control, Space.ParentIndex);
+				Transform.NormalizeRotation();
+				return Transform;
+			}
+		}
+	}
+
+	return FTransform::Identity;
+}
+
+void FRigSpaceHierarchy::SetLocalTransform(const FName& InName, const FTransform& InTransform)
+{
+	SetLocalTransform(GetIndex(InName), InTransform);
+}
+
+void FRigSpaceHierarchy::SetLocalTransform(int32 InIndex, const FTransform& InTransform)
+{
+	if (Spaces.IsValidIndex(InIndex))
+	{
+		FRigSpace& Space = Spaces[InIndex];
+		Space.LocalTransform = InTransform;
+		Space.LocalTransform.NormalizeRotation();
+	}
+}
+
+FTransform FRigSpaceHierarchy::GetLocalTransform(const FName& InName) const
+{
+	return GetLocalTransform(GetIndex(InName));
+}
+
+FTransform FRigSpaceHierarchy::GetLocalTransform(int32 InIndex) const
+{
+	if (Spaces.IsValidIndex(InIndex))
+	{
+		return Spaces[InIndex].LocalTransform;
+	}
+
+	return FTransform::Identity;
+}
+
+void FRigSpaceHierarchy::SetInitialTransform(const FName& InName, const FTransform& InTransform)
+{
+	SetInitialTransform(GetIndex(InName), InTransform);
+}
+
+void FRigSpaceHierarchy::SetInitialTransform(int32 InIndex, const FTransform& InTransform)
+{
+	if (Spaces.IsValidIndex(InIndex))
+	{
+		FRigSpace& Space = Spaces[InIndex];
+		Space.InitialTransform = InTransform;
+		Space.InitialTransform.NormalizeRotation();
+	}
+}
+
+FTransform FRigSpaceHierarchy::GetInitialTransform(const FName& InName) const
+{
+	return GetInitialTransform(GetIndex(InName));
+}
+
+FTransform FRigSpaceHierarchy::GetInitialTransform(int32 InIndex) const
+{
+	if (Spaces.IsValidIndex(InIndex))
+	{
+		return Spaces[InIndex].InitialTransform;
+	}
+
+	return FTransform::Identity;
+}
+
+void FRigSpaceHierarchy::SetInitialGlobalTransform(const FName& InName, const FTransform& GlobalTransform)
+{
+	SetInitialGlobalTransform(GetIndex(InName), GlobalTransform);
+}
+
+// wip - should support all types that can provide space transform data (pos/rotation)
+void FRigSpaceHierarchy::SetInitialGlobalTransform(int32 InIndex, const FTransform& GlobalTransform)
+{
+	if (Spaces.IsValidIndex(InIndex))
+	{
+		FRigSpace& Space = Spaces[InIndex];
+		FTransform ParentTransform = FTransform::Identity;
+		if (Container)
+		{
+			ParentTransform = Container->GetInitialGlobalTransform(Space.GetParentElementKey());
+		}
+		Space.InitialTransform = GlobalTransform.GetRelativeTransform(ParentTransform);
+	}
+}
+
+FTransform FRigSpaceHierarchy::GetInitialGlobalTransform(const FName& InName) const
+{
+	return GetInitialGlobalTransform(GetIndex(InName));
+}
+
+FTransform FRigSpaceHierarchy::GetInitialGlobalTransform(int32 InIndex) const
+{
+	// @todo: Templatize
+	if (Spaces.IsValidIndex(InIndex))
+	{
+		const FRigSpace& Space = Spaces[InIndex];
+		FTransform ParentTransform = FTransform::Identity;
+		if (Container)
+		{
+			ParentTransform = Container->GetInitialGlobalTransform(Space.GetParentElementKey());
+		}
+		FTransform Transform = Space.InitialTransform * ParentTransform;
+		Transform.NormalizeRotation();
+		return Transform;
+	}
+	return FTransform::Identity;
+}
+
+FName FRigSpaceHierarchy::Rename(const FName& InOldName, const FName& InNewName)
+{
+	if (InOldName != InNewName)
+	{
+		const int32 Found = GetIndex(InOldName);
+		if (Found != INDEX_NONE)
+		{
+			FName NewName = GetSafeNewName(InNewName);
+
+#if WITH_EDITOR
+			bool bWasSelected = IsSelected(InOldName);
+			if(bWasSelected)
+			{
+				Select(InOldName, false);
+			}
+#endif
+
+			Spaces[Found].Name = NewName;
+
+			// go through find all children and rename them
+			for (int32 Index = 0; Index < Spaces.Num(); ++Index)
+			{
+				if (Spaces[Index].SpaceType == ERigSpaceType::Space && Spaces[Index].ParentName == InOldName)
+				{
+					Spaces[Index].ParentName = NewName;
+				}
+			}
+
+			RefreshMapping();
+
+#if WITH_EDITOR
+			OnSpaceRenamed.Broadcast(Container, RigElementType(), InOldName, NewName);
+			if(bWasSelected)
+			{
+				Select(NewName, true);
+			}
+#endif
+			return NewName;
+		}
+	}
+
+	return NAME_None;
+}
+
+void FRigSpaceHierarchy::RefreshMapping()
+{
+	NameToIndexMapping.Empty();
+	for (int32 Index = 0; Index < Spaces.Num(); ++Index)
+	{
+		Spaces[Index].Index = Index;
+		NameToIndexMapping.Add(Spaces[Index].Name, Index);
+	}
+}
+
+void FRigSpaceHierarchy::Initialize()
+{
+	RefreshMapping();
+
+	// initialize transform
+	for (int32 Index = 0; Index < Spaces.Num(); ++Index)
+	{
+		if (Container)
+		{
+			Spaces[Index].ParentIndex = Container->GetIndex(Spaces[Index].GetParentElementKey());
+		}
+		Spaces[Index].LocalTransform = Spaces[Index].InitialTransform;
+	}
+}
+
+void FRigSpaceHierarchy::Reset()
+{
+	Spaces.Reset();
+}
+
+void FRigSpaceHierarchy::ResetTransforms()
+{
+	// initialize transform
+	for (int32 Index = 0; Index < Spaces.Num(); ++Index)
+	{
+		Spaces[Index].LocalTransform = Spaces[Index].InitialTransform;
+	}
+}
+
+int32 FRigSpaceHierarchy::GetParentIndex(ERigSpaceType InSpaceType, const FName& InName) const
+{
+	if (Container != nullptr)
+	{
+		switch (InSpaceType)
+		{
+			case ERigSpaceType::Global:
+			{
+				return INDEX_NONE;
+			}
+			case ERigSpaceType::Bone:
+			{
+				return Container->GetIndex(FRigElementKey(InName, ERigElementType::Bone));
+			}
+			case ERigSpaceType::Space:
+			{
+				return Container->GetIndex(FRigElementKey(InName, ERigElementType::Space));
+			}
+			case ERigSpaceType::Control:
+			{
+				return Container->GetIndex(FRigElementKey(InName, ERigElementType::Control));
+			}
+		}
+	}
+
+	return INDEX_NONE;
+}
+
+#if WITH_EDITOR
+
+bool FRigSpaceHierarchy::Select(const FName& InName, bool bSelect)
+{
+	if(GetIndex(InName) == INDEX_NONE)
+	{
+		return false;
+	}
+
+	if(bSelect == IsSelected(InName))
+	{
+		return false;
+	}
+
+	if(bSelect)
+	{
+		if (Container)
+		{
+			Container->BoneHierarchy.ClearSelection();
+			Container->ControlHierarchy.ClearSelection();
+			Container->CurveContainer.ClearSelection();
+		}
+
+		Selection.Add(InName);
+	}
+	else
+	{
+		Selection.Remove(InName);
+	}
+
+	OnSpaceSelected.Broadcast(Container, FRigElementKey(InName, RigElementType()), bSelect);
+
+	return true;
+}
+
+bool FRigSpaceHierarchy::ClearSelection()
+{
+	TArray<FName> TempSelection;
+	TempSelection.Append(Selection);
+	for(const FName& SelectedName : TempSelection)
+	{
+		Select(SelectedName, false);
+	}
+	return TempSelection.Num() > 0;
+}
+
+TArray<FName> FRigSpaceHierarchy::CurrentSelection() const
+{
+	TArray<FName> TempSelection;
+	TempSelection.Append(Selection);
+	return TempSelection;
+}
+
+bool FRigSpaceHierarchy::IsSelected(const FName& InName) const
+{
+	return Selection.Contains(InName);
+}
+
+void FRigSpaceHierarchy::HandleOnElementRemoved(FRigHierarchyContainer* InContainer, const FRigElementKey& InKey)
+{
+	if (Container == nullptr)
+	{
+		return;
+	}
+
+	switch (InKey.Type)
+	{
+		case ERigElementType::Bone:
+		{
+			for (FRigSpace& Space : Spaces)
+			{
+				if (Space.SpaceType == ERigSpaceType::Bone && Space.ParentName == InKey.Name)
+				{
+					Space.ParentIndex = Container->BoneHierarchy.GetIndex(InKey.Name);
+					Space.ParentName = Space.ParentIndex == INDEX_NONE ? NAME_None : InKey.Name;
+					Space.SpaceType = Space.ParentIndex == INDEX_NONE ? ERigSpaceType::Global : Space.SpaceType;
+#if WITH_EDITOR
+					OnSpaceReparented.Broadcast(Container, Space.GetElementKey(), InKey.Name, Space.ParentName);
+#endif
+				}
+			}
+			break;
+		}
+		case ERigElementType::Control:
+		{
+			for (FRigSpace& Space : Spaces)
+			{
+				if (Space.SpaceType == ERigSpaceType::Control && Space.ParentName == InKey.Name)
+				{
+					Space.ParentIndex = Container->ControlHierarchy.GetIndex(InKey.Name);
+					Space.ParentName = Space.ParentIndex == INDEX_NONE ? NAME_None : InKey.Name;
+					Space.SpaceType = Space.ParentIndex == INDEX_NONE ? ERigSpaceType::Global : Space.SpaceType;
+#if WITH_EDITOR
+					OnSpaceReparented.Broadcast(Container, Space.GetElementKey(), InKey.Name, Space.ParentName);
+#endif
+				}
+			}
+			break;
+		}
+		case ERigElementType::Space:
+		case ERigElementType::Curve:
+		{
+			break;
+		}
+	}
+}
+
+void FRigSpaceHierarchy::HandleOnElementRenamed(FRigHierarchyContainer* InContainer, ERigElementType InElementType, const FName& InOldName, const FName& InNewName)
+{
+	if (Container == nullptr)
+	{
+		return;
+	}
+
+	switch (InElementType)
+	{
+		case ERigElementType::Bone:
+		{
+			for (FRigSpace& Space : Spaces)
+			{
+				if (Space.SpaceType == ERigSpaceType::Bone && Space.ParentName == InOldName)
+				{
+					Space.ParentIndex = Container->BoneHierarchy.GetIndex(InNewName);
+					Space.ParentName = Space.ParentIndex == INDEX_NONE ? NAME_None : InNewName;
+					Space.SpaceType = Space.ParentIndex == INDEX_NONE ? ERigSpaceType::Global : Space.SpaceType;
+#if WITH_EDITOR
+					OnSpaceReparented.Broadcast(Container, Space.GetElementKey(), InOldName, Space.ParentName);
+#endif
+				}
+			}
+			break;
+		}
+		case ERigElementType::Control:
+		{
+			for (FRigSpace& Space : Spaces)
+			{
+				if (Space.SpaceType == ERigSpaceType::Control && Space.ParentName == InOldName)
+				{
+					Space.ParentIndex = Container->ControlHierarchy.GetIndex(InNewName);
+					Space.ParentName = Space.ParentIndex == INDEX_NONE ? NAME_None : InNewName;
+					Space.SpaceType = Space.ParentIndex == INDEX_NONE ? ERigSpaceType::Global : Space.SpaceType;
+#if WITH_EDITOR
+					OnSpaceReparented.Broadcast(Container, Space.GetElementKey(), InOldName, Space.ParentName);
+#endif
+				}
+			}
+			break;
+		}
+		case ERigElementType::Space:
+		case ERigElementType::Curve:
+		{
+			break;
+		}
+	}
+}
+
+#endif
\ No newline at end of file
diff --git a/Engine/Plugins/Experimental/ControlRig/Source/ControlRig/Private/Sequencer/ControlRigBindingTemplate.cpp b/Engine/Plugins/Experimental/ControlRig/Source/ControlRig/Private/Sequencer/ControlRigBindingTemplate.cpp
--- a/Engine/Plugins/Experimental/ControlRig/Source/ControlRig/Private/Sequencer/ControlRigBindingTemplate.cpp
+++ b/Engine/Plugins/Experimental/ControlRig/Source/ControlRig/Private/Sequencer/ControlRigBindingTemplate.cpp
@@ -180,7 +180,8 @@ struct FBindControlRigObjectToken : IMovieSceneExecutionToken
 				{
 					if (UControlRigSequencerAnimInstance* AnimInstance = Cast<UControlRigSequencerAnimInstance>(SkeletalMeshComponent->GetAnimInstance()))
 					{
-						bool bStructureChanged = AnimInstance->UpdateControlRig(ControlRig, Operand.SequenceID.GetInternalValue(), bAdditive, bApplyBoneFilter, *BoneFilter, Weight);
+						// @todo: let's think about the input here, I'm disabling handling of inputs but allowing it to execute when ticks
+						bool bStructureChanged = AnimInstance->UpdateControlRig(ControlRig, Operand.SequenceID.GetInternalValue(), bAdditive, bApplyBoneFilter, *BoneFilter, Weight, false, true);
 						if (bStructureChanged)
 						{
 							AnimInstance->RecalcRequiredBones();
@@ -209,7 +210,7 @@ struct FBindControlRigObjectToken : IMovieSceneExecutionToken
 						if (UControlRigSequencerAnimInstance* AnimInstance = Cast<UControlRigSequencerAnimInstance>(SkeletalMeshComponent->GetAnimInstance()))
 						{
 							// Force us to zero weight before we despawn, as the graph could persist
-							AnimInstance->UpdateControlRig(ControlRig, Operand.SequenceID.GetInternalValue(), bAdditive, bApplyBoneFilter, *BoneFilter, 0.0f);
+							AnimInstance->UpdateControlRig(ControlRig, Operand.SequenceID.GetInternalValue(), bAdditive, bApplyBoneFilter, *BoneFilter, 0.0f, true, true);
 							AnimInstance->RecalcRequiredBones();
 						}
 						UAnimSequencerInstance::UnbindFromSkeletalMeshComponent(SkeletalMeshComponent);
diff --git a/Engine/Plugins/Experimental/ControlRig/Source/ControlRig/Private/Sequencer/ControlRigSequencerAnimInstance.cpp b/Engine/Plugins/Experimental/ControlRig/Source/ControlRig/Private/Sequencer/ControlRigSequencerAnimInstance.cpp
--- a/Engine/Plugins/Experimental/ControlRig/Source/ControlRig/Private/Sequencer/ControlRigSequencerAnimInstance.cpp
+++ b/Engine/Plugins/Experimental/ControlRig/Source/ControlRig/Private/Sequencer/ControlRigSequencerAnimInstance.cpp
@@ -15,10 +15,10 @@ FAnimInstanceProxy* UControlRigSequencerAnimInstance::CreateAnimInstanceProxy()
 	return new FControlRigSequencerAnimInstanceProxy(this);
 }
 
-bool UControlRigSequencerAnimInstance::UpdateControlRig(UControlRig* InControlRig, uint32 SequenceId, bool bAdditive, bool bApplyBoneFilter, const FInputBlendPose& BoneFilter, float Weight, bool bExternalSource)
+bool UControlRigSequencerAnimInstance::UpdateControlRig(UControlRig* InControlRig, uint32 SequenceId, bool bAdditive, bool bApplyBoneFilter, const FInputBlendPose& BoneFilter, float Weight, bool bUpdateInput, bool bExecute)
 {
 	CachedControlRig = InControlRig;
-	return GetProxyOnGameThread<FControlRigSequencerAnimInstanceProxy>().UpdateControlRig(InControlRig, SequenceId, bAdditive, bApplyBoneFilter, BoneFilter, Weight, bExternalSource);
+	return GetProxyOnGameThread<FControlRigSequencerAnimInstanceProxy>().UpdateControlRig(InControlRig, SequenceId, bAdditive, bApplyBoneFilter, BoneFilter, Weight, bUpdateInput, bExecute);
 }
 
 bool UControlRigSequencerAnimInstance::SetAnimationAsset(class UAnimationAsset* NewAsset)
diff --git a/Engine/Plugins/Experimental/ControlRig/Source/ControlRig/Private/Sequencer/ControlRigSequencerAnimInstanceProxy.cpp b/Engine/Plugins/Experimental/ControlRig/Source/ControlRig/Private/Sequencer/ControlRigSequencerAnimInstanceProxy.cpp
--- a/Engine/Plugins/Experimental/ControlRig/Source/ControlRig/Private/Sequencer/ControlRigSequencerAnimInstanceProxy.cpp
+++ b/Engine/Plugins/Experimental/ControlRig/Source/ControlRig/Private/Sequencer/ControlRigSequencerAnimInstanceProxy.cpp
@@ -137,7 +137,7 @@ void FControlRigSequencerAnimInstanceProxy::InitControlRigTrack(UControlRig* InC
 	}
 }
 
-bool FControlRigSequencerAnimInstanceProxy::UpdateControlRig(UControlRig* InControlRig, uint32 SequenceId, bool bAdditive, bool bApplyBoneFilter, const FInputBlendPose& BoneFilter, float Weight, bool bExternalSource)
+bool FControlRigSequencerAnimInstanceProxy::UpdateControlRig(UControlRig* InControlRig, uint32 SequenceId, bool bAdditive, bool bApplyBoneFilter, const FInputBlendPose& BoneFilter, float Weight, bool bUpdateInput, bool bExecute)
 {
 	bool bCreated = EnsureControlRigTrack(InControlRig, bAdditive, bApplyBoneFilter, BoneFilter, SequenceId);
 
@@ -153,8 +153,8 @@ bool FControlRigSequencerAnimInstanceProxy::UpdateControlRig(UControlRig* InCont
 		BlendNode.DesiredAlphas[PlayerState->PoseIndex] = Weight;
 	}
 
-	PlayerState->ControlRigNode.bUpdateInput = !bExternalSource;
-	PlayerState->ControlRigNode.bExecute = !bExternalSource;
+	PlayerState->ControlRigNode.bUpdateInput = bUpdateInput;
+	PlayerState->ControlRigNode.bExecute = bExecute;
 
 	return bCreated;
 }
diff --git a/Engine/Plugins/Experimental/ControlRig/Source/ControlRig/Private/Sequencer/ControlRigSequencerAnimInstanceProxy.h b/Engine/Plugins/Experimental/ControlRig/Source/ControlRig/Private/Sequencer/ControlRigSequencerAnimInstanceProxy.h
--- a/Engine/Plugins/Experimental/ControlRig/Source/ControlRig/Private/Sequencer/ControlRigSequencerAnimInstanceProxy.h
+++ b/Engine/Plugins/Experimental/ControlRig/Source/ControlRig/Private/Sequencer/ControlRigSequencerAnimInstanceProxy.h
@@ -48,7 +48,7 @@ public:
 	// FAnimSequencerInstanceProxy interface
 	virtual void ResetNodes() override;
 
-	bool UpdateControlRig(UControlRig* InControlRig, uint32 SequenceId, bool bAdditive, bool bApplyBoneFilter, const FInputBlendPose& BoneFilter, float Weight, bool bExternalSource);
+	bool UpdateControlRig(UControlRig* InControlRig, uint32 SequenceId, bool bAdditive, bool bApplyBoneFilter, const FInputBlendPose& BoneFilter, float Weight, bool bUpdateInput, bool bExecute);
 	bool SetAnimationAsset(class UAnimationAsset* NewAsset);
 
 private:
diff --git a/Engine/Plugins/Experimental/ControlRig/Source/ControlRig/Private/Settings/ControlRigSettings.cpp b/Engine/Plugins/Experimental/ControlRig/Source/ControlRig/Private/Settings/ControlRigSettings.cpp
new file mode 100644
--- /dev/null
+++ b/Engine/Plugins/Experimental/ControlRig/Source/ControlRig/Private/Settings/ControlRigSettings.cpp
@@ -0,0 +1,11 @@
+// Copyright 1998-2019 Epic Games, Inc. All Rights Reserved.
+
+#include "Settings/ControlRigSettings.h"
+
+UControlRigSettings::UControlRigSettings(const FObjectInitializer& ObjectInitializer)
+	: Super(ObjectInitializer)
+{
+	DefaultGizmoLibrary = LoadObject<UControlRigGizmoLibrary>(nullptr, TEXT("/ControlRig/Controls/DefaultGizmoLibrary.DefaultGizmoLibrary"));
+}
+
+
diff --git a/Engine/Plugins/Experimental/ControlRig/Source/ControlRig/Private/Units/Control/RigUnit_Control_StaticMesh.cpp b/Engine/Plugins/Experimental/ControlRig/Source/ControlRig/Private/Units/Control/RigUnit_Control_StaticMesh.cpp
--- a/Engine/Plugins/Experimental/ControlRig/Source/ControlRig/Private/Units/Control/RigUnit_Control_StaticMesh.cpp
+++ b/Engine/Plugins/Experimental/ControlRig/Source/ControlRig/Private/Units/Control/RigUnit_Control_StaticMesh.cpp
@@ -1,7 +1,6 @@
 // Copyright 1998-2019 Epic Games, Inc. All Rights Reserved.
 
 #include "Units/Control/RigUnit_Control_StaticMesh.h"
-#include "ControlRigStaticMeshControl.h"
 
 FRigUnit_Control_StaticMesh::FRigUnit_Control_StaticMesh()
 #if WITH_EDITORONLY_DATA
@@ -9,7 +8,4 @@ FRigUnit_Control_StaticMesh::FRigUnit_Control_StaticMesh()
 	, MeshTransform(FTransform::Identity)
 #endif
 {
-#if WITH_EDITORONLY_DATA
-	ControlClass = AControlRigStaticMeshControl::StaticClass();
-#endif 
 }
diff --git a/Engine/Plugins/Experimental/ControlRig/Source/ControlRig/Private/Units/Hierarchy/RigUnit_BoneName.h b/Engine/Plugins/Experimental/ControlRig/Source/ControlRig/Private/Units/Hierarchy/RigUnit_BoneName.h
--- a/Engine/Plugins/Experimental/ControlRig/Source/ControlRig/Private/Units/Hierarchy/RigUnit_BoneName.h
+++ b/Engine/Plugins/Experimental/ControlRig/Source/ControlRig/Private/Units/Hierarchy/RigUnit_BoneName.h
@@ -8,7 +8,7 @@
 /**
  * BoneName is used to represent a bone name in the graph
  */
-USTRUCT(meta=(DisplayName="Bone Name", Category="Hierarchy", DocumentationPolicy = "Strict"))
+USTRUCT(meta=(DisplayName="Bone Name", Category="Hierarchy", DocumentationPolicy = "Strict", Deprecated = "4.24"))
 struct FRigUnit_BoneName : public FRigUnit
 {
 	GENERATED_BODY()
diff --git a/Engine/Plugins/Experimental/ControlRig/Source/ControlRig/Private/Units/Hierarchy/RigUnit_GetControlInitialTransform.cpp b/Engine/Plugins/Experimental/ControlRig/Source/ControlRig/Private/Units/Hierarchy/RigUnit_GetControlInitialTransform.cpp
new file mode 100644
--- /dev/null
+++ b/Engine/Plugins/Experimental/ControlRig/Source/ControlRig/Private/Units/Hierarchy/RigUnit_GetControlInitialTransform.cpp
@@ -0,0 +1,53 @@
+// Copyright 1998-2019 Epic Games, Inc. All Rights Reserved.
+
+#include "RigUnit_GetControlInitialTransform.h"
+#include "Units/RigUnitContext.h"
+
+FString FRigUnit_GetControlInitialTransform::GetUnitLabel() const
+{
+	return FString::Printf(TEXT("Get Initial Control %s"), *Control.ToString());
+}
+
+FRigUnit_GetControlInitialTransform_Execute()
+{
+    DECLARE_SCOPE_HIERARCHICAL_COUNTER_RIGUNIT()
+	const FRigControlHierarchy* Hierarchy = Context.GetControls();
+	if (Hierarchy)
+	{
+		switch (Context.State)
+		{
+			case EControlRigState::Init:
+			{
+				CachedControlIndex = Hierarchy->GetIndex(Control);
+				break;				
+			}
+			case EControlRigState::Update:
+			{
+				if (CachedControlIndex != INDEX_NONE)
+				{
+					switch (Space)
+					{
+						case EBoneGetterSetterMode::GlobalSpace:
+						{
+							Transform = Hierarchy->GetInitialGlobalTransform(CachedControlIndex);
+							break;
+						}
+						case EBoneGetterSetterMode::LocalSpace:
+						{
+							Transform = Hierarchy->GetInitialValue<FTransform>(CachedControlIndex);
+							break;
+						}
+						default:
+						{
+							break;
+						}
+					}
+				}
+			}
+			default:
+			{
+				break;
+			}
+		}
+	}
+}
diff --git a/Engine/Plugins/Experimental/ControlRig/Source/ControlRig/Private/Units/Hierarchy/RigUnit_GetControlInitialTransform.h b/Engine/Plugins/Experimental/ControlRig/Source/ControlRig/Private/Units/Hierarchy/RigUnit_GetControlInitialTransform.h
new file mode 100644
--- /dev/null
+++ b/Engine/Plugins/Experimental/ControlRig/Source/ControlRig/Private/Units/Hierarchy/RigUnit_GetControlInitialTransform.h
@@ -0,0 +1,46 @@
+// Copyright 1998-2019 Epic Games, Inc. All Rights Reserved.
+
+#pragma once
+
+#include "Units/RigUnit.h"
+#include "RigUnit_GetControlInitialTransform.generated.h"
+
+/**
+ * GetControlTransform is used to retrieve a single transform from a hierarchy.
+ */
+USTRUCT(meta=(DisplayName="Get Control Initial Transform", Category="Controls", DocumentationPolicy = "Strict", Keywords="GetControlInitialTransform"))
+struct FRigUnit_GetControlInitialTransform : public FRigUnit
+{
+	GENERATED_BODY()
+
+	FRigUnit_GetControlInitialTransform()
+		: Space(EBoneGetterSetterMode::LocalSpace)
+		, CachedControlIndex(INDEX_NONE)
+	{}
+
+	virtual FString GetUnitLabel() const override;
+
+	RIGVM_METHOD()
+	virtual void Execute(const FRigUnitContext& Context) override;
+
+	/**
+	 * The name of the Control to retrieve the transform for.
+	 */
+	UPROPERTY(meta = (Input, ControlName, Constant))
+	FName Control;
+
+	/**
+	 * Defines if the Control's transform should be retrieved
+	 * in local or global space.
+	 */ 
+	UPROPERTY(meta = (Input))
+	EBoneGetterSetterMode Space;
+
+	// The current transform of the given bone - or identity in case it wasn't found.
+	UPROPERTY(meta=(Output))
+	FTransform Transform;
+
+	// Used to cache the internally used bone index
+	UPROPERTY()
+	int32 CachedControlIndex;
+};
diff --git a/Engine/Plugins/Experimental/ControlRig/Source/ControlRig/Private/Units/Hierarchy/RigUnit_GetControlTransform.cpp b/Engine/Plugins/Experimental/ControlRig/Source/ControlRig/Private/Units/Hierarchy/RigUnit_GetControlTransform.cpp
new file mode 100644
--- /dev/null
+++ b/Engine/Plugins/Experimental/ControlRig/Source/ControlRig/Private/Units/Hierarchy/RigUnit_GetControlTransform.cpp
@@ -0,0 +1,90 @@
+// Copyright 1998-2019 Epic Games, Inc. All Rights Reserved.
+
+#include "RigUnit_GetControlTransform.h"
+#include "Units/RigUnitContext.h"
+
+FString FRigUnit_GetControlTransform::GetUnitLabel() const
+{
+	return FString::Printf(TEXT("Get Transform %s"), *Control.ToString());
+}
+
+FRigUnit_GetControlTransform_Execute()
+{
+    DECLARE_SCOPE_HIERARCHICAL_COUNTER_RIGUNIT()
+	const FRigControlHierarchy* Hierarchy = Context.GetControls();
+	if (Hierarchy)
+	{
+		switch (Context.State)
+		{
+			case EControlRigState::Init:
+			{
+				CachedControlIndex = Hierarchy->GetIndex(Control);
+				break;				
+			}
+			case EControlRigState::Update:
+			{
+				if (CachedControlIndex != INDEX_NONE)
+				{
+					switch (Space)
+					{
+						case EBoneGetterSetterMode::GlobalSpace:
+						{
+							Transform = Hierarchy->GetGlobalTransform(CachedControlIndex);
+							break;
+						}
+						case EBoneGetterSetterMode::LocalSpace:
+						{
+							Transform = Hierarchy->GetLocalTransform(CachedControlIndex);
+							break;
+						}
+						default:
+						{
+							break;
+						}
+					}
+				}
+			}
+			default:
+			{
+				break;
+			}
+		}
+	}
+}
+
+#if WITH_DEV_AUTOMATION_TESTS
+#include "Units/RigUnitTest.h"
+
+IMPLEMENT_RIGUNIT_AUTOMATION_TEST(FRigUnit_GetControlTransform)
+{
+	ControlHierarchy.Add(TEXT("Root"), ERigControlType::Transform, NAME_None, NAME_None, FRigControlValue::Make(FTransform(FVector(1.f, 0.f, 0.f))));
+	ControlHierarchy.Add(TEXT("ControlA"), ERigControlType::Transform, TEXT("Root"), NAME_None, FRigControlValue::Make(FTransform(FVector(1.f, 2.f, 3.f))));
+	ControlHierarchy.Initialize();
+
+	Unit.Control = TEXT("Unknown");
+	Unit.Space = EBoneGetterSetterMode::GlobalSpace;
+	InitAndExecute();
+	AddErrorIfFalse(Unit.Transform.GetTranslation().Equals(FVector(0.f, 0.f, 0.f)), TEXT("unexpected global transform"));
+	Unit.Space = EBoneGetterSetterMode::LocalSpace;
+	InitAndExecute();
+	AddErrorIfFalse(Unit.Transform.GetTranslation().Equals(FVector(0.f, 0.f, 0.f)), TEXT("unexpected local transform"));
+
+	Unit.Control = TEXT("Root");
+	Unit.Space = EBoneGetterSetterMode::GlobalSpace;
+	InitAndExecute();
+	AddErrorIfFalse(Unit.Transform.GetTranslation().Equals(FVector(1.f, 0.f, 0.f)), TEXT("unexpected global transform"));
+	Unit.Space = EBoneGetterSetterMode::LocalSpace;
+	InitAndExecute();
+	AddErrorIfFalse(Unit.Transform.GetTranslation().Equals(FVector(1.f, 0.f, 0.f)), TEXT("unexpected local transform"));
+
+	Unit.Control = TEXT("ControlA");
+	Unit.Space = EBoneGetterSetterMode::GlobalSpace;
+	InitAndExecute();
+	AddErrorIfFalse(Unit.Transform.GetTranslation().Equals(FVector(2.f, 2.f, 3.f)), TEXT("unexpected global transform"));
+	Unit.Space = EBoneGetterSetterMode::LocalSpace;
+	InitAndExecute();
+	AddErrorIfFalse(Unit.Transform.GetTranslation().Equals(FVector(1.f, 2.f, 3.f)), TEXT("unexpected local transform"));
+
+	return true;
+}
+#endif
\ No newline at end of file
diff --git a/Engine/Plugins/Experimental/ControlRig/Source/ControlRig/Private/Units/Hierarchy/RigUnit_GetControlTransform.h b/Engine/Plugins/Experimental/ControlRig/Source/ControlRig/Private/Units/Hierarchy/RigUnit_GetControlTransform.h
new file mode 100644
--- /dev/null
+++ b/Engine/Plugins/Experimental/ControlRig/Source/ControlRig/Private/Units/Hierarchy/RigUnit_GetControlTransform.h
@@ -0,0 +1,45 @@
+// Copyright 1998-2019 Epic Games, Inc. All Rights Reserved.
+
+#pragma once
+
+#include "Units/RigUnit.h"
+#include "RigUnit_GetControlTransform.generated.h"
+
+/**
+ * GetControlTransform is used to retrieve a single transform from a hierarchy.
+ */
+USTRUCT(meta=(DisplayName="Get Control Transform", Category="Controls", DocumentationPolicy = "Strict", Keywords="GetControlTransform"))
+struct FRigUnit_GetControlTransform : public FRigUnit
+{
+	GENERATED_BODY()
+
+	FRigUnit_GetControlTransform()
+		: Space(EBoneGetterSetterMode::LocalSpace)
+		, CachedControlIndex(INDEX_NONE)
+	{}
+
+	virtual FString GetUnitLabel() const override;
+	RIGVM_METHOD()
+	virtual void Execute(const FRigUnitContext& Context) override;
+
+	/**
+	 * The name of the Control to retrieve the transform for.
+	 */
+	UPROPERTY(meta = (Input, ControlName, Constant))
+	FName Control;
+
+	/**
+	 * Defines if the Control's transform should be retrieved
+	 * in local or global space.
+	 */ 
+	UPROPERTY(meta = (Input))
+	EBoneGetterSetterMode Space;
+
+	// The current transform of the given bone - or identity in case it wasn't found.
+	UPROPERTY(meta=(Output))
+	FTransform Transform;
+
+	// Used to cache the internally used bone index
+	UPROPERTY()
+	int32 CachedControlIndex;
+};
