Copying //UE4/Dev-Anim to Dev-Main (//UE4/Dev-Main) @ 8782600

diff --git a/Engine/Plugins/Experimental/ControlRig/Source/ControlRig/Private/Units/Hierarchy/RigUnit_GetSpaceTransform.cpp b/Engine/Plugins/Experimental/ControlRig/Source/ControlRig/Private/Units/Hierarchy/RigUnit_GetSpaceTransform.cpp
new file mode 100644
--- /dev/null
+++ b/Engine/Plugins/Experimental/ControlRig/Source/ControlRig/Private/Units/Hierarchy/RigUnit_GetSpaceTransform.cpp
@@ -0,0 +1,52 @@
+// Copyright 1998-2019 Epic Games, Inc. All Rights Reserved.
+
+#include "RigUnit_GetSpaceTransform.h"
+#include "Units/RigUnitContext.h"
+
+FString FRigUnit_GetSpaceTransform::GetUnitLabel() const
+{
+	return FString::Printf(TEXT("Get Transform %s"), *Space.ToString());
+}
+
+FRigUnit_GetSpaceTransform_Execute()
+{
+    DECLARE_SCOPE_HIERARCHICAL_COUNTER_RIGUNIT()
+	const FRigSpaceHierarchy* Hierarchy = Context.GetSpaces();
+	if (Hierarchy)
+	{
+		switch (Context.State)
+		{
+			case EControlRigState::Init:
+			{
+				CachedSpaceIndex = Hierarchy->GetIndex(Space);
+			}
+			case EControlRigState::Update:
+			{
+				if (CachedSpaceIndex != INDEX_NONE)
+				{
+					switch (SpaceType)
+					{
+						case EBoneGetterSetterMode::GlobalSpace:
+						{
+							Transform = Hierarchy->GetGlobalTransform(CachedSpaceIndex);
+							break;
+						}
+						case EBoneGetterSetterMode::LocalSpace:
+						{
+							Transform = Hierarchy->GetLocalTransform(CachedSpaceIndex);
+							break;
+						}
+						default:
+						{
+							break;
+						}
+					}
+				}
+			}
+			default:
+			{
+				break;
+			}
+		}
+	}
+}
diff --git a/Engine/Plugins/Experimental/ControlRig/Source/ControlRig/Private/Units/Hierarchy/RigUnit_GetSpaceTransform.h b/Engine/Plugins/Experimental/ControlRig/Source/ControlRig/Private/Units/Hierarchy/RigUnit_GetSpaceTransform.h
new file mode 100644
--- /dev/null
+++ b/Engine/Plugins/Experimental/ControlRig/Source/ControlRig/Private/Units/Hierarchy/RigUnit_GetSpaceTransform.h
@@ -0,0 +1,45 @@
+// Copyright 1998-2019 Epic Games, Inc. All Rights Reserved.
+
+#pragma once
+
+#include "Units/RigUnit.h"
+#include "RigUnit_GetSpaceTransform.generated.h"
+
+/**
+ * GetSpaceTransform is used to retrieve a single transform from a hierarchy.
+ */
+USTRUCT(meta=(DisplayName="Get Space Transform", Category="Spaces", DocumentationPolicy = "Strict", Keywords="GetSpaceTransform"))
+struct FRigUnit_GetSpaceTransform : public FRigUnit
+{
+	GENERATED_BODY()
+
+	FRigUnit_GetSpaceTransform()
+		: SpaceType(EBoneGetterSetterMode::LocalSpace)
+		, CachedSpaceIndex(INDEX_NONE)
+	{}
+
+	virtual FString GetUnitLabel() const override;
+	RIGVM_METHOD()
+	virtual void Execute(const FRigUnitContext& Context) override;
+
+	/**
+	 * The name of the Space to retrieve the transform for.
+	 */
+	UPROPERTY(meta = (Input, SpaceName, Constant))
+	FName Space;
+
+	/**
+	 * Defines if the Space's transform should be retrieved
+	 * in local or global space.
+	 */ 
+	UPROPERTY(meta = (Input))
+	EBoneGetterSetterMode SpaceType;
+
+	// The current transform of the given bone - or identity in case it wasn't found.
+	UPROPERTY(meta=(Output))
+	FTransform Transform;
+
+	// Used to cache the internally used bone index
+	UPROPERTY()
+	int32 CachedSpaceIndex;
+};
diff --git a/Engine/Plugins/Experimental/ControlRig/Source/ControlRig/Private/Units/Hierarchy/RigUnit_SetBoneTransform.h b/Engine/Plugins/Experimental/ControlRig/Source/ControlRig/Private/Units/Hierarchy/RigUnit_SetBoneTransform.h
--- a/Engine/Plugins/Experimental/ControlRig/Source/ControlRig/Private/Units/Hierarchy/RigUnit_SetBoneTransform.h
+++ b/Engine/Plugins/Experimental/ControlRig/Source/ControlRig/Private/Units/Hierarchy/RigUnit_SetBoneTransform.h
@@ -21,6 +21,8 @@ struct FRigUnit_SetBoneTransform : public FRigUnitMutable
 	{}
 
 	virtual FString GetUnitLabel() const override;
+
+	RIGVM_METHOD()
 	virtual void Execute(const FRigUnitContext& Context) override;
 
 	/**
@@ -32,7 +34,7 @@ struct FRigUnit_SetBoneTransform : public FRigUnitMutable
 	/**
 	 * The transform value to set for the given Bone.
 	 */
-	UPROPERTY(meta = (Input))
+	UPROPERTY(meta = (Input, Output))
 	FTransform Transform;
 
 	/**
diff --git a/Engine/Plugins/Experimental/ControlRig/Source/ControlRig/Private/Units/Hierarchy/RigUnit_SetControlTransform.cpp b/Engine/Plugins/Experimental/ControlRig/Source/ControlRig/Private/Units/Hierarchy/RigUnit_SetControlTransform.cpp
new file mode 100644
--- /dev/null
+++ b/Engine/Plugins/Experimental/ControlRig/Source/ControlRig/Private/Units/Hierarchy/RigUnit_SetControlTransform.cpp
@@ -0,0 +1,53 @@
+// Copyright 1998-2019 Epic Games, Inc. All Rights Reserved.
+
+#include "RigUnit_SetControlTransform.h"
+#include "Units/RigUnitContext.h"
+
+FString FRigUnit_SetControlTransform::GetUnitLabel() const
+{
+	return FString::Printf(TEXT("Set Control %s"), *Control.ToString());
+}
+
+FRigUnit_SetControlTransform_Execute()
+{
+    DECLARE_SCOPE_HIERARCHICAL_COUNTER_RIGUNIT()
+	FRigControlHierarchy* Hierarchy = ExecuteContext.GetControls();
+	if (Hierarchy)
+	{
+		switch (Context.State)
+		{
+			case EControlRigState::Init:
+			{
+				CachedControlIndex = Hierarchy->GetIndex(Control);
+				// fall through to update
+			}
+			case EControlRigState::Update:
+			{
+				if (CachedControlIndex != INDEX_NONE)
+				{
+					switch (Space)
+					{
+						case EBoneGetterSetterMode::GlobalSpace:
+						{
+							Hierarchy->SetGlobalTransform(CachedControlIndex, Transform);
+							break;
+						}
+						case EBoneGetterSetterMode::LocalSpace:
+						{
+							Hierarchy->SetLocalTransform(CachedControlIndex, Transform);
+							break;
+						}
+						default:
+						{
+							break;
+						}
+					}
+				}
+			}
+			default:
+			{
+				break;
+			}
+		}
+	}
+}
diff --git a/Engine/Plugins/Experimental/ControlRig/Source/ControlRig/Private/Units/Hierarchy/RigUnit_SetControlTransform.h b/Engine/Plugins/Experimental/ControlRig/Source/ControlRig/Private/Units/Hierarchy/RigUnit_SetControlTransform.h
new file mode 100644
--- /dev/null
+++ b/Engine/Plugins/Experimental/ControlRig/Source/ControlRig/Private/Units/Hierarchy/RigUnit_SetControlTransform.h
@@ -0,0 +1,49 @@
+// Copyright 1998-2019 Epic Games, Inc. All Rights Reserved.
+
+#pragma once
+
+#include "Units/RigUnit.h"
+#include "RigUnit_SetControlTransform.generated.h"
+
+
+/**
+ * SetControlTransform is used to perform a change in the hierarchy by setting a single bone's transform.
+ */
+USTRUCT(meta=(DisplayName="Set Control", Category="Hierarchy", DocumentationPolicy="Strict", Keywords = "SetControlTransform"))
+struct FRigUnit_SetControlTransform : public FRigUnitMutable
+{
+	GENERATED_BODY()
+
+	FRigUnit_SetControlTransform()
+		: Space(EBoneGetterSetterMode::LocalSpace)
+		, CachedControlIndex(INDEX_NONE)
+	{}
+
+	virtual FString GetUnitLabel() const override;
+
+	RIGVM_METHOD()
+	virtual void Execute(const FRigUnitContext& Context) override;
+
+	/**
+	 * The name of the Control to set the transform for.
+	 */
+	UPROPERTY(meta = (Input, ControlName, Constant))
+	FName Control;
+
+	/**
+	 * The transform value to set for the given Control.
+	 */
+	UPROPERTY(meta = (Input, Output))
+	FTransform Transform;
+
+	/**
+	 * Defines if the bone's transform should be set
+	 * in local or global space.
+	 */
+	UPROPERTY(meta = (Input))
+	EBoneGetterSetterMode Space;
+
+	// Used to cache the internally used bone index
+	UPROPERTY()
+	int32 CachedControlIndex;
+};
diff --git a/Engine/Plugins/Experimental/ControlRig/Source/ControlRig/Private/Units/Hierarchy/RigUnit_SetCurveValue.h b/Engine/Plugins/Experimental/ControlRig/Source/ControlRig/Private/Units/Hierarchy/RigUnit_SetCurveValue.h
--- a/Engine/Plugins/Experimental/ControlRig/Source/ControlRig/Private/Units/Hierarchy/RigUnit_SetCurveValue.h
+++ b/Engine/Plugins/Experimental/ControlRig/Source/ControlRig/Private/Units/Hierarchy/RigUnit_SetCurveValue.h
@@ -9,7 +9,7 @@
 /**
  * SetCurveValue is used to perform a change in the curve container by setting a single Curve value.
  */
-USTRUCT(meta=(DisplayName="Set Curve Value", Category="Hierarchy", Keywords = "SetCurveValue"))
+USTRUCT(meta=(DisplayName="Set Curve Value", Category="Curve", Keywords = "SetCurveValue"))
 struct FRigUnit_SetCurveValue : public FRigUnitMutable
 {
 	GENERATED_BODY()
@@ -20,6 +20,7 @@ struct FRigUnit_SetCurveValue : public FRigUnitMutable
 	{}
 
 	virtual FString GetUnitLabel() const override;
+	RIGVM_METHOD()
 	virtual void Execute(const FRigUnitContext& Context) override;
 
 	/**
diff --git a/Engine/Plugins/Experimental/ControlRig/Source/ControlRig/Private/Units/Hierarchy/RigUnit_SetSpaceTransform.cpp b/Engine/Plugins/Experimental/ControlRig/Source/ControlRig/Private/Units/Hierarchy/RigUnit_SetSpaceTransform.cpp
new file mode 100644
--- /dev/null
+++ b/Engine/Plugins/Experimental/ControlRig/Source/ControlRig/Private/Units/Hierarchy/RigUnit_SetSpaceTransform.cpp
@@ -0,0 +1,53 @@
+// Copyright 1998-2019 Epic Games, Inc. All Rights Reserved.
+
+#include "RigUnit_SetSpaceTransform.h"
+#include "Units/RigUnitContext.h"
+
+FString FRigUnit_SetSpaceTransform::GetUnitLabel() const
+{
+	return FString::Printf(TEXT("Set Space %s"), *Space.ToString());
+}
+
+FRigUnit_SetSpaceTransform_Execute()
+{
+    DECLARE_SCOPE_HIERARCHICAL_COUNTER_RIGUNIT()
+	FRigSpaceHierarchy* Hierarchy = ExecuteContext.GetSpaces();
+	if (Hierarchy)
+	{
+		switch (Context.State)
+		{
+			case EControlRigState::Init:
+			{
+				CachedSpaceIndex = Hierarchy->GetIndex(Space);
+				// fall through to update
+			}
+			case EControlRigState::Update:
+			{
+				if (CachedSpaceIndex != INDEX_NONE)
+				{
+					switch (SpaceType)
+					{
+						case EBoneGetterSetterMode::GlobalSpace:
+						{
+							Hierarchy->SetGlobalTransform(CachedSpaceIndex, Transform);
+							break;
+						}
+						case EBoneGetterSetterMode::LocalSpace:
+						{
+							Hierarchy->SetLocalTransform(CachedSpaceIndex, Transform);
+							break;
+						}
+						default:
+						{
+							break;
+						}
+					}
+				}
+			}
+			default:
+			{
+				break;
+			}
+		}
+	}
+}
diff --git a/Engine/Plugins/Experimental/ControlRig/Source/ControlRig/Private/Units/Hierarchy/RigUnit_SetSpaceTransform.h b/Engine/Plugins/Experimental/ControlRig/Source/ControlRig/Private/Units/Hierarchy/RigUnit_SetSpaceTransform.h
new file mode 100644
--- /dev/null
+++ b/Engine/Plugins/Experimental/ControlRig/Source/ControlRig/Private/Units/Hierarchy/RigUnit_SetSpaceTransform.h
@@ -0,0 +1,49 @@
+// Copyright 1998-2019 Epic Games, Inc. All Rights Reserved.
+
+#pragma once
+
+#include "Units/RigUnit.h"
+#include "RigUnit_SetSpaceTransform.generated.h"
+
+
+/**
+ * SetSpaceTransform is used to perform a change in the hierarchy by setting a single bone's transform.
+ */
+USTRUCT(meta=(DisplayName="Set Space", Category="Hierarchy", DocumentationPolicy="Strict", Keywords = "SetSpaceTransform"))
+struct FRigUnit_SetSpaceTransform : public FRigUnitMutable
+{
+	GENERATED_BODY()
+
+	FRigUnit_SetSpaceTransform()
+		: SpaceType(EBoneGetterSetterMode::LocalSpace)
+		, CachedSpaceIndex(INDEX_NONE)
+	{}
+
+	virtual FString GetUnitLabel() const override;
+
+	RIGVM_METHOD()
+	virtual void Execute(const FRigUnitContext& Context) override;
+
+	/**
+	 * The name of the Space to set the transform for.
+	 */
+	UPROPERTY(meta = (Input, SpaceName, Constant))
+	FName Space;
+
+	/**
+	 * The transform value to set for the given Space.
+	 */
+	UPROPERTY(meta = (Input, Output))
+	FTransform Transform;
+
+	/**
+	 * Defines if the bone's transform should be set
+	 * in local or global space.
+	 */
+	UPROPERTY(meta = (Input))
+	EBoneGetterSetterMode SpaceType;
+
+	// Used to cache the internally used bone index
+	UPROPERTY()
+	int32 CachedSpaceIndex;
+};
diff --git a/Engine/Plugins/Experimental/ControlRig/Source/ControlRig/Public/ControlRig.h b/Engine/Plugins/Experimental/ControlRig/Source/ControlRig/Public/ControlRig.h
--- a/Engine/Plugins/Experimental/ControlRig/Source/ControlRig/Public/ControlRig.h
+++ b/Engine/Plugins/Experimental/ControlRig/Source/ControlRig/Public/ControlRig.h
@@ -5,16 +5,15 @@
 #include "CoreMinimal.h"
 #include "UObject/ObjectMacros.h"
 #include "UObject/ScriptMacros.h"
-#include "Animation/ControlRigInterface.h"
 #include "Engine/EngineBaseTypes.h"
 #include "Templates/SubclassOf.h"
 #include "ControlRigDefines.h"
-#include "Hierarchy.h"
-#include "CurveContainer.h"
+#include "Rigs/RigHierarchyContainer.h"
 #include "Units/RigUnitContext.h"
 #include "Animation/NodeMappingProviderInterface.h"
 #include "Units/RigUnit.h"
 #include "Units/Control/RigUnit_Control.h"
+#include "Manipulatable/IControlRigManipulatable.h"
 #include "ControlRig.generated.h"
 
 class IControlRigObjectBinding;
@@ -22,6 +21,8 @@ class UScriptStruct;
 struct FRigUnit;
 struct FControlRigIOVariable;
 
+DECLARE_LOG_CATEGORY_EXTERN(LogControlRig, Log, All);
+
 /** Delegate used to optionally gather inputs before evaluating a ControlRig */
 DECLARE_DELEGATE_OneParam(FPreEvaluateGatherInput, UControlRig*);
 DECLARE_DELEGATE_OneParam(FPostEvaluateQueryOutput, UControlRig*);
@@ -30,7 +31,7 @@ DECLARE_DELEGATE_OneParam(FPostEvaluateQueryOutput, UControlRig*);
 
 /** Runs logic for mapping input data to transforms (the "Rig") */
 UCLASS(Blueprintable, Abstract, editinlinenew)
-class CONTROLRIG_API UControlRig : public UObject, public IControlRigInterface, public INodeMappingProviderInterface
+class CONTROLRIG_API UControlRig : public UObject, public INodeMappingProviderInterface, public IControlRigManipulatable
 {
 	GENERATED_BODY()
 
@@ -47,6 +48,8 @@ public:
 	static const FName MenuDescSuffixMetaName;
 	static const FName ShowVariableNameInTitleMetaName;
 	static const FName BoneNameMetaName;
+	static const FName ControlNameMetaName;
+	static const FName SpaceNameMetaName;
 	static const FName CurveNameMetaName;
 	static const FName ConstantMetaName;
 	static const FName TitleColorMetaName;
@@ -88,10 +91,8 @@ public:
 	/** Initialize things for the ControlRig */
 	virtual void Initialize(bool bInitRigUnits = true);
 
-	/** IControlRigInterface implementation */
-	virtual void PreEvaluate_GameThread() override;
-	virtual void Evaluate_AnyThread() override;
-	virtual void PostEvaluate_GameThread() override;
+	/** Evaluate at Any Thread */
+	virtual void Evaluate_AnyThread();
 
 	/** Setup bindings to a runtime object (or clear by passing in nullptr). */
 	void SetObjectBinding(TSharedPtr<IControlRigObjectBinding> InObjectBinding)
@@ -105,6 +106,31 @@ public:
 		return ObjectBinding;
 	}
 
+	FRigHierarchyContainer* GetHierarchy()
+	{
+		return &Hierarchy;
+	}
+
+	FRigBoneHierarchy& GetBoneHierarchy()
+	{
+		return Hierarchy.BoneHierarchy;
+	}
+
+	FRigSpaceHierarchy& GetSpaceHierarchy()
+	{
+		return Hierarchy.SpaceHierarchy;
+	}
+
+	FRigControlHierarchy& GetControlHierarchy()
+	{
+		return Hierarchy.ControlHierarchy;
+	}
+
+	FRigCurveContainer& GetCurveContainer()
+	{
+		return Hierarchy.CurveContainer;
+	}
+
 	/** Evaluate another animation ControlRig */
 	FTransform GetGlobalTransform(const FName& BoneName) const;
 
@@ -117,12 +143,6 @@ public:
 	/** Evaluate another animation ControlRig */
 	void SetGlobalTransform(const int32 BoneIndex, const FTransform& InTransform, bool bPropagateTransform = true) ;
 
-	/** Returns base hierarchy */
-	const FRigHierarchy& GetBaseHierarchy() const
-	{
-		return Hierarchy.BaseHierarchy;
-	}
-
 	/** Evaluate another animation ControlRig */
 	float GetCurveValue(const FName& CurveName) const;
 
@@ -135,32 +155,6 @@ public:
 	/** Evaluate another animation ControlRig */
 	void SetCurveValue(const int32 CurveIndex, const float CurveValue);
 
-	/** Returns base hierarchy */
-	const FRigCurveContainer& GetCurveContainer() const
-	{
-		return CurveContainer;
-	}
-
-	void SetPreEvaluateGatherInputDelegate(const FPreEvaluateGatherInput& Delegate)
-	{
-		OnPreEvaluateGatherInput = Delegate;
-	}
-
-	void ClearPreEvaluateGatherInputDelegate()
-	{
-		OnPreEvaluateGatherInput.Unbind();
-	}
-
-	void SetPostEvaluateQueryOutputDelegate(const FPostEvaluateQueryOutput& Delegate)
-	{
-		OnPostEvaluateQueryOutput = Delegate;
-	}
-
-	void ClearPostEvaluateQueryOutputDelegate()
-	{
-		OnPostEvaluateQueryOutput.Unbind();
-	}
-
 	/* 
 	 * Query input output variables
 	 *
@@ -202,12 +196,6 @@ public:
 	virtual void BeginDestroy() override;
 	// END UObject interface
 
-#if WITH_EDITORONLY_DATA
-	// only editor feature that stops execution
-	// whether we're executing the graph or not
-	bool bExecutionOn;
-#endif // #if WITH_EDITORONLY_DATA
-
 	UPROPERTY(transient)
 	ERigExecutionType ExecutionType;
 
@@ -217,13 +205,27 @@ public:
 	/** INodeMappingInterface implementation */
 	virtual void GetMappableNodeData(TArray<FName>& OutNames, TArray<FNodeItem>& OutNodeItems) const override;
 
+	// BEGIN IControlRigManipulatable interface
+	virtual const TArray<FRigSpace>& AvailableSpaces() const override;
+	virtual FRigSpace* FindSpace(const FName& InSpaceName) override;
+	virtual FTransform GetSpaceGlobalTransform(const FName& InSpaceName) override;
+	virtual bool SetSpaceGlobalTransform(const FName& InSpaceName, const FTransform& InTransform) override;
+	virtual const TArray<FRigControl>& AvailableControls() const override;
+	virtual FRigControl* FindControl(const FName& InControlName) override;
+	virtual FTransform GetControlGlobalTransform(const FName& InControlName) const override;
+	virtual FRigControlValue GetControlValueFromGlobalTransform(const FName& InControlName, const FTransform& InGlobalTransform) override;
+	virtual bool SetControlSpace(const FName& InControlName, const FName& InSpaceName) override;
+	virtual UControlRigGizmoLibrary* GetGizmoLibrary() const override;
+	// END IControlRigManipulatable interface
+
+	DECLARE_EVENT_TwoParams(UControlRig, FControlRigExecuteEvent, class UControlRig*, const EControlRigState);
+	FControlRigExecuteEvent& OnInitialized() { return InitializedEvent; }
+	FControlRigExecuteEvent& OnExecuted() { return ExecutedEvent; }
+
 private:
 	UPROPERTY(VisibleDefaultsOnly, Category = "Hierarchy")
 	FRigHierarchyContainer Hierarchy;
 
-	UPROPERTY(VisibleDefaultsOnly, Category = "Curve")
-	FRigCurveContainer CurveContainer;
-
 #if WITH_EDITORONLY_DATA
 	/** The properties of source accessible <target, source local path> when source -> target
 	 * For example, if you have property RigUnitA.B->RigUnitB.C, this will save as <RigUnitB.C, RigUnitA.B> */
@@ -238,16 +240,12 @@ private:
 	UPROPERTY(Transient)
 	TArray<FControlRigOperator> Operators;
 
+	UPROPERTY()
+	TAssetPtr<UControlRigGizmoLibrary> GizmoLibrary;
+
 	/** Runtime object binding */
 	TSharedPtr<IControlRigObjectBinding> ObjectBinding;
 
-	FPreEvaluateGatherInput OnPreEvaluateGatherInput;
-	FPostEvaluateQueryOutput OnPostEvaluateQueryOutput;
-
-	DECLARE_EVENT_TwoParams(UControlRig, FControlRigExecuteEvent, class UControlRig*, const EControlRigState);
-	FControlRigExecuteEvent& OnInitialized() { return InitializedEvent; }
-	FControlRigExecuteEvent& OnExecuted() { return ExecutedEvent; }
-
 #if WITH_EDITOR
 	FControlRigLog* ControlRigLog;
 	bool bEnableControlRigLogging;
@@ -264,6 +262,10 @@ private:
 	/** The draw interface for the units to use */
 	FControlRigDrawInterface* DrawInterface;
 
+	/** The registry to access data source */
+	UPROPERTY(transient)
+	UAnimationDataSourceRegistry* DataSourceRegistry;
+
 #if DEBUG_CONTROLRIG_PROPERTYCHANGE
 	// This is to debug class size when constructed and destroyed to verify match
 	// if this size changes, that implies more problem, where properties have been changed and layout has been modified
@@ -303,8 +305,8 @@ private:
 	friend class UControlRigEditorLibrary;
 	friend class URigUnitEditor_Base;
 	friend class FControlRigEditor;
-	friend class SRigHierarchy;
 	friend class SRigCurveContainer;
+	friend class SRigHierarchy;
 	friend class UEngineTestControlRig;
  	friend class FControlRigEditMode;
 	friend class FControlRigIOHelper;
diff --git a/Engine/Plugins/Experimental/ControlRig/Source/ControlRig/Public/ControlRigGizmoLibrary.h b/Engine/Plugins/Experimental/ControlRig/Source/ControlRig/Public/ControlRigGizmoLibrary.h
new file mode 100644
--- /dev/null
+++ b/Engine/Plugins/Experimental/ControlRig/Source/ControlRig/Public/ControlRigGizmoLibrary.h
@@ -0,0 +1,70 @@
+// Copyright 1998-2019 Epic Games, Inc. All Rights Reserved.
+
+#pragma once
+
+#include "CoreMinimal.h"
+#include "UObject/ObjectMacros.h"
+#include "UObject/ScriptMacros.h"
+#include "Engine/StaticMesh.h"
+#include "Materials/Material.h"
+
+#include "ControlRigGizmoLibrary.generated.h"
+
+USTRUCT(BlueprintType, meta = (DisplayName = "Gizmo"))
+struct CONTROLRIG_API FControlRigGizmoDefinition
+{
+	GENERATED_USTRUCT_BODY()
+
+	FControlRigGizmoDefinition()
+	{
+		GizmoName = TEXT("Gizmo");
+		StaticMesh = nullptr;
+		Transform = FTransform::Identity;
+	}
+
+	UPROPERTY(EditAnywhere, Category = "Gizmo")
+	FName GizmoName;
+
+	UPROPERTY(EditAnywhere, Category = "Gizmo")
+	TAssetPtr<UStaticMesh> StaticMesh;
+
+	UPROPERTY(EditAnywhere, Category = "Gizmo")
+	FTransform Transform;
+};
+
+UCLASS(BlueprintType, meta = (DisplayName = "GizmoLibrary"))
+class CONTROLRIG_API UControlRigGizmoLibrary : public UObject
+{
+	GENERATED_BODY()
+
+public:
+
+	UControlRigGizmoLibrary();
+
+	UPROPERTY(EditAnywhere, Category = "GizmoLibrary")
+	FControlRigGizmoDefinition DefaultGizmo;
+
+	UPROPERTY(EditAnywhere, Category = "GizmoLibrary")
+	TAssetPtr<UMaterial> DefaultMaterial;
+
+	UPROPERTY(EditAnywhere, Category = "GizmoLibrary")
+	FName MaterialColorParameter;
+
+	UPROPERTY(EditAnywhere, Category = "GizmoLibrary")
+	TArray<FControlRigGizmoDefinition> Gizmos;
+
+	const FControlRigGizmoDefinition* GetGizmoByName(const FName& InName, bool bUseDefaultIfNotFound = false) const;
+
+#if WITH_EDITOR
+
+	// UObject interface
+	virtual void PostEditChangeChainProperty(struct FPropertyChangedChainEvent& PropertyChangedEvent) override;
+
+#endif
+
+private:
+
+	const TArray<FName> GetUpdatedNameList(bool bReset = false);
+
+	TArray<FName> NameList;
+};
\ No newline at end of file
diff --git a/Engine/Plugins/Experimental/ControlRig/Source/ControlRig/Public/Manipulatable/IControlRigManipulatable.h b/Engine/Plugins/Experimental/ControlRig/Source/ControlRig/Public/Manipulatable/IControlRigManipulatable.h
new file mode 100644
--- /dev/null
+++ b/Engine/Plugins/Experimental/ControlRig/Source/ControlRig/Public/Manipulatable/IControlRigManipulatable.h
@@ -0,0 +1,129 @@
+// Copyright 1998-2019 Epic Games, Inc. All Rights Reserved.
+
+#pragma once
+
+#include "CoreMinimal.h"
+#include "Delegates/DelegateCombinations.h"
+#include "Rigs/RigControlHierarchy.h"
+#include "Rigs/RigSpaceHierarchy.h"
+#include "ControlRigGizmoLibrary.h"
+
+/**
+ *
+ * IControlRigManipulatable provides an interface for subjects who desire to be 
+ * manipulatable by the Control Rig manipulation framework.
+ * The manipulatable provides opaque access to available controls / spaces
+ * and implements a series of setter functions to perform changes.
+ *
+ */
+class CONTROLRIG_API IControlRigManipulatable
+{
+public:
+	IControlRigManipulatable();
+	virtual ~IControlRigManipulatable();
+
+	/** Bindable event for external objects to contribute to / filter a control value */
+	DECLARE_EVENT_ThreeParams(IControlRigManipulatable, FFilterControlEvent, IControlRigManipulatable*, const FRigControl&, FRigControlValue&);
+
+	/** Bindable event for external objects to be notified of Control changes */
+	DECLARE_EVENT_TwoParams(IControlRigManipulatable, FControlModifiedEvent, IControlRigManipulatable*, const FRigControl&);
+
+	// Returns true if this manipulatable subject is currently
+	// available for manipulation / is enabled.
+	virtual bool ManipulationEnabled() const
+	{
+		return bManipulationEnabled;
+	}
+
+	// Sets the manipulatable subject to enabled or disabled
+	virtual bool SetManipulationEnabled(bool Enabled = true)
+	{
+		if (bManipulationEnabled == Enabled)
+		{
+			return false;
+		}
+		bManipulationEnabled = Enabled;
+		return true;
+	}
+
+	// Returns a list of available spaces on the subject
+	virtual const TArray<FRigSpace>& AvailableSpaces() const = 0;
+
+	// Returns a space given its name
+	virtual FRigSpace* FindSpace(const FName& InSpaceName) = 0;
+
+	// Gets a space's transform given a global / world transform.
+	virtual FTransform GetSpaceGlobalTransform(const FName& InSpaceName) = 0;
+
+	// Sets a space's transform given a global / world transform.
+	// Returns true when successful.
+	virtual bool SetSpaceGlobalTransform(const FName& InSpaceName, const FTransform& InTransform) = 0;
+
+	// Returns a list of available controls on the subject.
+	// Each control provides additional information such as metadata,
+	// it's value type and so on.
+	virtual const TArray<FRigControl>& AvailableControls() const = 0;
+
+	// Returns a control given its name
+	virtual FRigControl* FindControl(const FName& InControlName) = 0;
+
+	// Returns the value of a Control
+	FORCEINLINE const FRigControlValue& GetControlValue(const FName& InControlName)
+	{
+		const FRigControl* Control = FindControl(InControlName);
+		check(Control);
+		return Control->Value;
+	}
+
+	// Sets the relative value of a Control
+	FORCEINLINE void SetControlValue(const FName& InControlName, const FRigControlValue& InValue, bool bNotify = true)
+	{
+		FRigControl* Control = FindControl(InControlName);
+		check(Control);
+		Control->Value = InValue;
+
+		if (bNotify && OnControlModified.IsBound())
+		{
+			OnControlModified.Broadcast(this, *Control);
+		}
+	}
+
+	// Sets the relative value of a Control
+	template<class T>
+	FORCEINLINE bool SetControlValue(const FName& InControlName, T InValue, bool bNotify = true)
+	{
+		SetControlValue(InControlName, FRigControlValue::Make<T>(InValue), bNotify);
+	}
+
+	// Returns the global / world transform of a Control
+	virtual FTransform GetControlGlobalTransform(const FName& InControlName) const = 0;
+
+	// Sets the global / world transform of a Control. This should be called from the interaction
+	// layer / edit mode. Returns true when successful.
+	bool SetControlGlobalTransform(const FName& InControlName, const FTransform& InGlobalTransform);
+
+	// Returns the value given a global transform
+	virtual FRigControlValue GetControlValueFromGlobalTransform(const FName& InControlName, const FTransform& InGlobalTransform) = 0;
+
+	// Sets a Control's Space (for space switching), returns true when successful.
+	virtual bool SetControlSpace(const FName& InControlName, const FName& InSpaceName) = 0;
+
+	// Returns a event than can be used to subscribe to
+	// filtering control data when needed
+	FFilterControlEvent& ControlFilter() { return OnFilterControl;  }
+	
+	// Returns a event than can be used to subscribe to
+	// notifications coming from the manipulated subject.
+	FControlModifiedEvent& ControlModified() { return OnControlModified;  }
+
+	// Returns the gizmo library used for generating gizmos
+	virtual UControlRigGizmoLibrary* GetGizmoLibrary() const { return nullptr; }
+
+private:
+
+	FFilterControlEvent OnFilterControl;
+	FControlModifiedEvent OnControlModified;
+
+	/** True if currently manipulation is enabled */
+	bool bManipulationEnabled;
+};
\ No newline at end of file
diff --git a/Engine/Plugins/Experimental/ControlRig/Source/ControlRig/Public/Rigs/RigControlHierarchy.h b/Engine/Plugins/Experimental/ControlRig/Source/ControlRig/Public/Rigs/RigControlHierarchy.h
new file mode 100644
--- /dev/null
+++ b/Engine/Plugins/Experimental/ControlRig/Source/ControlRig/Public/Rigs/RigControlHierarchy.h
@@ -0,0 +1,376 @@
+// Copyright 1998-2019 Epic Games, Inc. All Rights Reserved.
+
+#pragma once
+
+#include "CoreMinimal.h"
+#include "RigHierarchyDefines.h"
+#include "RigControlHierarchy.generated.h"
+
+class UControlRig;
+
+UENUM()
+enum class ERigControlType : uint8
+{
+	Bool,
+	Float,
+	Vector2D,
+	Position,
+	Scale,
+	Quat,
+	Rotator,
+	Transform
+};
+
+USTRUCT()
+struct CONTROLRIG_API FRigControlValue
+{
+	GENERATED_BODY()
+
+public:
+
+	FRigControlValue()
+		:Storage(FTransform::Identity)
+	{
+	}
+
+	template<class T>
+	FORCEINLINE T Get() const
+	{
+		return GetRef<T>();
+	}
+
+	template<class T>
+	FORCEINLINE T& GetRef()
+	{
+		ensure(sizeof(T) <= sizeof(FTransform));
+		return *(T*)&Storage;
+	}
+
+	template<class T>
+	FORCEINLINE const T& GetRef() const
+	{
+		ensure(sizeof(T) <= sizeof(FTransform));
+		return *(T*)&Storage;
+	}
+
+	template<class T>
+	FORCEINLINE void Set(T InValue)
+	{
+		GetRef<T>() = InValue;
+	}
+
+	template<class T>
+	FORCEINLINE static FRigControlValue Make(T InValue)
+	{
+		FRigControlValue Value;
+		Value.Set<T>(InValue);
+		return Value;
+	}
+
+private:
+	
+	UPROPERTY()
+	FTransform Storage;
+};
+
+
+class UStaticMesh;
+
+USTRUCT()
+struct CONTROLRIG_API FRigControl : public FRigElement
+{
+	GENERATED_BODY()
+
+	FRigControl()
+		: FRigElement()
+		, ControlType(ERigControlType::Transform)
+		, ParentName(NAME_None)
+		, ParentIndex(INDEX_NONE)
+		, SpaceName(NAME_None)
+		, SpaceIndex(INDEX_NONE)
+		, InitialValue()
+		, Value()
+#if WITH_EDITORONLY_DATA
+		, GizmoName(TEXT("Gizmo"))
+		, GizmoTransform(FTransform::Identity)
+		, GizmoColor(FLinearColor::Red)
+#endif
+		, Dependents()
+	{
+	}
+	virtual ~FRigControl() {}
+
+	UPROPERTY(VisibleAnywhere, Category = FRigElement)
+	ERigControlType ControlType;
+
+	UPROPERTY(VisibleAnywhere, Category = FRigElement)
+	FName ParentName;
+
+	UPROPERTY(transient)
+	int32 ParentIndex;
+
+	UPROPERTY(VisibleAnywhere, Category = FRigElement)
+	FName SpaceName;
+
+	UPROPERTY(transient)
+	int32 SpaceIndex;
+
+	UPROPERTY(VisibleAnywhere, Category = FRigElement)
+	FRigControlValue InitialValue;
+
+	UPROPERTY(transient, VisibleAnywhere, Category = FRigElement)
+	FRigControlValue Value;
+
+#if WITH_EDITORONLY_DATA
+	/* This is optional UI setting - this doesn't mean this is always used, but it is optional for manipulation layer to use this*/
+	UPROPERTY(EditAnywhere, Category = FRigElement)
+	FName GizmoName;
+
+	UPROPERTY(EditAnywhere, Category = FRigElement)
+	FTransform GizmoTransform;
+
+	UPROPERTY(EditAnywhere, Category = FRigElement)
+	FLinearColor GizmoColor;
+#endif // WITH_EDITORONLY_DATA
+
+	/** dependent list - direct dependent for child or anything that needs to update due to this */
+	UPROPERTY(transient)
+	TArray<int32> Dependents;
+
+	FORCEINLINE virtual ERigElementType GetElementType() const override
+	{
+		return ERigElementType::Control;
+	}
+
+	FORCEINLINE virtual FRigElementKey GetParentElementKey() const
+	{
+		return FRigElementKey(ParentName, GetElementType());
+	}
+
+	FORCEINLINE virtual FRigElementKey GetSpaceElementKey() const
+	{
+		return FRigElementKey(SpaceName, ERigElementType::Space);
+	}
+};
+
+USTRUCT()
+struct CONTROLRIG_API FRigControlHierarchy
+{
+	GENERATED_BODY()
+
+	FRigControlHierarchy();
+	FRigControlHierarchy& operator= (const FRigControlHierarchy &InOther);
+
+	FORCEINLINE ERigElementType RigElementType() const { return ERigElementType::Control; }
+
+	FORCEINLINE int32 Num() const { return Controls.Num(); }
+	FORCEINLINE const FRigControl& operator[](int32 InIndex) const { return Controls[InIndex]; }
+	FORCEINLINE FRigControl& operator[](int32 InIndex) { return Controls[InIndex]; }
+	FORCEINLINE const FRigControl& operator[](const FName& InName) const { return Controls[GetIndex(InName)]; }
+	FORCEINLINE FRigControl& operator[](const FName& InName) { return Controls[GetIndex(InName)]; }
+
+	FORCEINLINE const TArray<FRigControl>& GetControls() const { return Controls; }
+
+	FORCEINLINE TArray<FRigControl>::RangedForIteratorType      begin()       { return Controls.begin(); }
+	FORCEINLINE TArray<FRigControl>::RangedForConstIteratorType begin() const { return Controls.begin(); }
+	FORCEINLINE TArray<FRigControl>::RangedForIteratorType      end()         { return Controls.end();   }
+	FORCEINLINE TArray<FRigControl>::RangedForConstIteratorType end() const   { return Controls.end();   }
+
+	FORCEINLINE bool IsNameAvailable(const FName& InPotentialNewName) const { return GetIndex(InPotentialNewName) == INDEX_NONE; }
+
+	FName GetSafeNewName(const FName& InPotentialNewName) const;
+
+	FRigControl& Add(
+		const FName& InNewName,
+		ERigControlType InControlType = ERigControlType::Transform,
+		const FName& InParentName = NAME_None,
+		const FName& InSpaceName = NAME_None,
+		const FRigControlValue& InValue = FRigControlValue()
+#if WITH_EDITORONLY_DATA
+		, const FName& InGizmoName = TEXT("Gizmo"),
+		const FTransform& InGizmoTransform = FTransform::Identity,
+		const FLinearColor& InGizmoColor = FLinearColor::Red
+#endif
+	);
+
+	FRigControl Remove(const FName& InNameToRemove);
+
+	FName Rename(const FName& InOldName, const FName& InNewName);
+
+	bool Reparent(const FName& InName, const FName& InNewParentName);
+
+	void SetSpace(const FName& InName, const FName& InNewSpaceName);
+
+	FName GetName(int32 InIndex) const;
+
+	FORCEINLINE int32 GetIndex(const FName& InName) const
+	{
+		if(NameToIndexMapping.Num() != Controls.Num())
+		{
+			return GetIndexSlow(InName);
+		}
+
+		const int32* Index = NameToIndexMapping.Find(InName);
+		if (Index)
+		{
+			return *Index;
+		}
+
+		return INDEX_NONE;
+	}
+
+	// list of names of children - this is not cheap, and is supposed to be used only for one time set up
+	int32 GetChildren(const FName& InName, TArray<int32>& OutChildren, bool bRecursively) const;
+
+	int32 GetChildren(const int32 InIndex, TArray<int32>& OutChildren, bool bRecursively) const;
+
+	void SetGlobalTransform(const FName& InName, const FTransform& InTransform);
+
+	void SetGlobalTransform(int32 InIndex, const FTransform& InTransform);
+
+	FTransform GetGlobalTransform(const FName& InName) const;
+
+	FTransform GetGlobalTransform(int32 InIndex) const;
+
+	void SetLocalTransform(const FName& InName, const FTransform& InTransform);
+
+	void SetLocalTransform(int32 InIndex, const FTransform& InTransform);
+
+	FTransform GetLocalTransform(const FName& InName) const;
+
+	FTransform GetLocalTransform(int32 InIndex) const;
+
+	FTransform GetParentTransform(int32 InIndex) const;
+
+	// this is only valid if you have trasnform as type
+	// setting initial transform from global transform (control transform is based on parent)
+	void SetInitialGlobalTransform(const FName& InName, const FTransform& GlobalTransform);
+	
+	void SetInitialGlobalTransform(int32 InIndex, const FTransform& GlobalTransform);
+	
+	FTransform GetInitialGlobalTransform(const FName& InName) const;
+	
+	FTransform GetInitialGlobalTransform(int32 InIndex) const;
+	
+	void SetValue(const FName& InName, const FRigControlValue& InValue);
+
+	void SetValue(int32 InIndex, const FRigControlValue& InValue);
+
+	template<class T>
+	FORCEINLINE void SetValue(const FName& InName, T InValue)
+	{
+		SetValue(InName, FRigControlValue::Make<T>(InValue));
+	}
+
+	template<class T>
+	FORCEINLINE void SetValue(int32 InIndex, T InValue)
+	{
+		SetValue(InIndex, FRigControlValue::Make<T>(InValue));
+	}
+
+	FRigControlValue GetValue(const FName& InName) const;
+
+	FRigControlValue GetValue(int32 InIndex) const;
+
+	template<class T>
+	FORCEINLINE T GetValue(const FName& InName) const
+	{
+		return GetValue(InName).Get<T>();
+	}
+
+	template<class T>
+	FORCEINLINE T GetValue(int32 InIndex) const
+	{
+		return GetValue(InIndex).Get<T>();
+	}
+
+	void SetInitialValue(const FName& InName, const FRigControlValue& InValue);
+
+	void SetInitialValue(int32 InIndex, const FRigControlValue& InValue);
+
+	template<class T>
+	FORCEINLINE void SetInitialValue(const FName& InName, T InValue)
+	{
+		SetInitialValue(InName, FRigControlValue::Make<T>(InValue));
+	}
+
+	template<class T>
+	FORCEINLINE void SetInitialValue(int32 InIndex, T InValue)
+	{
+		SetInitialValue(InIndex, FRigControlValue::Make<T>(InValue));
+	}
+
+	FRigControlValue GetInitialValue(const FName& InName) const;
+
+	FRigControlValue GetInitialValue(int32 InIndex) const;
+
+	template<class T>
+	FORCEINLINE T GetInitialValue(const FName& InName) const
+	{
+		return GetInitialValue(InName).Get<T>();
+	}
+
+	template<class T>
+	FORCEINLINE T GetInitialValue(int32 InIndex) const
+	{
+		return GetInitialValue(InIndex).Get<T>();
+	}
+
+	// updates all of the internal caches
+	void Initialize();
+
+	// clears the hierarchy and removes all content
+	void Reset();
+
+	// resets all of the values back to the initial values
+	void ResetValues();
+
+#if WITH_EDITOR
+
+	bool Select(const FName& InName, bool bSelect = true);
+	bool ClearSelection();
+	TArray<FName> CurrentSelection() const;
+	bool IsSelected(const FName& InName) const;
+
+	FRigElementAdded OnControlAdded;
+	FRigElementRemoved OnControlRemoved;
+	FRigElementRenamed OnControlRenamed;
+	FRigElementReparented OnControlReparented;
+	FRigElementSelected OnControlSelected;
+	FRigElementChanged OnControlUISettingsChanged;
+
+	void HandleOnElementRemoved(FRigHierarchyContainer* InContainer, const FRigElementKey& InKey);
+	void HandleOnElementRenamed(FRigHierarchyContainer* InContainer, ERigElementType InElementType, const FName& InOldName, const FName& InNewName);
+
+#endif
+
+private:
+
+	// disable copy constructor
+	FRigControlHierarchy(const FRigControlHierarchy& InOther) {}
+
+	FRigHierarchyContainer* Container;
+
+	UPROPERTY(EditAnywhere, Category = FRigControlHierarchy)
+	TArray<FRigControl> Controls;
+
+	// can serialize fine? 
+	UPROPERTY()
+	TMap<FName, int32> NameToIndexMapping;
+
+#if WITH_EDITORONLY_DATA
+	UPROPERTY(transient)
+	TArray<FName> Selection;
+#endif
+
+	int32 GetSpaceIndex(const FName& InName) const;
+
+	int32 GetIndexSlow(const FName& InName) const;
+
+	void RefreshMapping();
+
+	// list of names of children - this is not cheap, and is supposed to be used only for one time set up
+	int32 GetChildrenRecursive(const int32 InIndex, TArray<int32>& OutChildren, bool bRecursively) const;
+
+	friend struct FRigHierarchyContainer;
+};
diff --git a/Engine/Plugins/Experimental/ControlRig/Source/ControlRig/Public/Rigs/RigHierarchyContainer.h b/Engine/Plugins/Experimental/ControlRig/Source/ControlRig/Public/Rigs/RigHierarchyContainer.h
new file mode 100644
--- /dev/null
+++ b/Engine/Plugins/Experimental/ControlRig/Source/ControlRig/Public/Rigs/RigHierarchyContainer.h
@@ -0,0 +1,191 @@
+// Copyright 1998-2019 Epic Games, Inc. All Rights Reserved.
+
+#pragma once
+
+#include "CoreMinimal.h"
+#include "RigHierarchyDefines.h"
+#include "RigBoneHierarchy.h"
+#include "RigSpaceHierarchy.h"
+#include "RigControlHierarchy.h"
+#include "RigCurveContainer.h"
+#include "RigHierarchyContainer.generated.h"
+
+class UControlRig;
+struct FRigHierarchyContainer;
+
+UENUM()
+enum class ERigHierarchyImportMode : uint8
+{
+	Append,
+
+	Replace,
+
+	ReplaceLocalTransform,
+
+	ReplaceGlobalTransform,
+
+	/** MAX - invalid */
+	Max UMETA(Hidden),
+};
+
+USTRUCT()
+struct CONTROLRIG_API FRigHierarchyCopyPasteContent
+{
+	GENERATED_USTRUCT_BODY()
+
+	UPROPERTY()
+	TArray<ERigElementType> Types;
+
+	UPROPERTY()
+	TArray<FString> Contents;
+
+	UPROPERTY()
+	TArray<FTransform> LocalTransforms;
+
+	UPROPERTY()
+	TArray<FTransform> GlobalTransforms;
+};
+
+USTRUCT()
+struct CONTROLRIG_API FRigHierarchyContainer
+{
+public:
+	GENERATED_BODY()
+
+	FRigHierarchyContainer();
+	FRigHierarchyContainer& operator= (const FRigHierarchyContainer &InOther);
+
+	UPROPERTY()
+	FRigBoneHierarchy BoneHierarchy;
+
+	UPROPERTY()
+	FRigSpaceHierarchy SpaceHierarchy;
+
+	UPROPERTY()
+	FRigControlHierarchy ControlHierarchy;
+
+	UPROPERTY()
+	FRigCurveContainer CurveContainer;
+
+	void Initialize();
+	void Reset();
+	void ResetTransforms();
+
+	FORCEINLINE int32 GetIndex(const FRigElementKey& InKey) const
+	{
+		switch(InKey.Type)
+		{
+			case ERigElementType::Bone:
+			{
+				return BoneHierarchy.GetIndex(InKey.Name);
+			}
+			case ERigElementType::Space:
+			{
+				return SpaceHierarchy.GetIndex(InKey.Name);
+			}
+			case ERigElementType::Control:
+			{
+				return ControlHierarchy.GetIndex(InKey.Name);
+			}
+			case ERigElementType::Curve:
+			{
+				return CurveContainer.GetIndex(InKey.Name);
+			}
+		}
+		return INDEX_NONE;
+	}
+
+	FORCEINLINE FTransform GetInitialTransform(const FRigElementKey& InKey) const
+	{
+		return GetInitialTransform(InKey.Type, GetIndex(InKey));
+	}
+	FTransform GetInitialTransform(ERigElementType InElementType, int32 InIndex) const;
+
+#if WITH_EDITOR
+	FORCEINLINE void SetInitialTransform(const FRigElementKey& InKey, const FTransform& InTransform)
+	{
+		return SetInitialTransform(InKey.Type, GetIndex(InKey), InTransform);
+	}
+	void SetInitialTransform(ERigElementType InElementType, int32 InIndex, const FTransform& InTransform);
+#endif
+
+	FORCEINLINE FTransform GetInitialGlobalTransform(const FRigElementKey& InKey) const
+	{
+		return GetInitialGlobalTransform(InKey.Type, GetIndex(InKey));
+	}
+	FTransform GetInitialGlobalTransform(ERigElementType InElementType, int32 InIndex) const;
+
+	FORCEINLINE void SetInitialGlobalTransform(const FRigElementKey& InKey, const FTransform& InTransform)
+	{
+		SetInitialGlobalTransform(InKey.Type, GetIndex(InKey), InTransform);
+	}
+	void SetInitialGlobalTransform(ERigElementType InElementType, int32 InIndex, const FTransform& InTransform);
+
+	FORCEINLINE FTransform GetLocalTransform(const FRigElementKey& InKey) const
+	{
+		return GetLocalTransform(InKey.Type, GetIndex(InKey));
+	}
+	FTransform GetLocalTransform(ERigElementType InElementType, int32 InIndex) const;
+
+	FORCEINLINE void SetLocalTransform(const FRigElementKey& InKey, const FTransform& InTransform)
+	{
+		return SetLocalTransform(InKey.Type, GetIndex(InKey), InTransform);
+	}
+	void SetLocalTransform(ERigElementType InElementType, int32 InIndex, const FTransform& InTransform);
+
+	FORCEINLINE FTransform GetGlobalTransform(const FRigElementKey& InKey) const
+	{
+		return GetGlobalTransform(InKey.Type, GetIndex(InKey));
+	}
+	FTransform GetGlobalTransform(ERigElementType InElementType, int32 InIndex) const;
+
+	FORCEINLINE void SetGlobalTransform(const FRigElementKey& InKey, const FTransform& InTransform)
+	{
+		return SetGlobalTransform(InKey.Type, GetIndex(InKey), InTransform);
+	}
+	void SetGlobalTransform(ERigElementType InElementType, int32 InIndex, const FTransform& InTransform);
+
+	bool IsParentedTo(ERigElementType InChildType, int32 InChildIndex, ERigElementType InParentType, int32 InParentIndex) const;
+
+#if WITH_EDITOR
+
+	FString ExportSelectionToText() const;
+	FString ExportToText(const TArray<FRigElementKey>& InSelection) const;
+	TArray<FRigElementKey> ImportFromText(const FString& InContent, ERigHierarchyImportMode InImportMode = ERigHierarchyImportMode::Append, bool bSelectNewElements = true);
+
+	TArray<FRigElementKey> CurrentSelection() const;
+	TArray<FRigElementKey> GetAllItems(bool bSort = false) const;
+	bool Select(const FRigElementKey& InKey, bool bSelect = true);
+	bool ClearSelection();
+	bool ClearSelection(ERigElementType InElementType);
+	bool IsSelected(const FRigElementKey& InKey) const;
+
+	FRigElementChanged OnElementChanged;
+	FRigElementAdded OnElementAdded;
+	FRigElementRemoved OnElementRemoved;
+	FRigElementRenamed OnElementRenamed;
+	FRigElementReparented OnElementReparented;
+	FRigElementSelected OnElementSelected;
+
+#endif
+
+protected:
+
+	// disable copy constructor
+	FRigHierarchyContainer(const FRigHierarchyContainer& InContainer) {}
+
+#if WITH_EDITOR
+	void HandleOnElementAdded(FRigHierarchyContainer* InContainer, const FRigElementKey& InKey);
+	void HandleOnElementRemoved(FRigHierarchyContainer* InContainer, const FRigElementKey& InKey);
+	void HandleOnElementRenamed(FRigHierarchyContainer* InContainer, ERigElementType InElementType, const FName& InOldName, const FName& InNewName);
+	void HandleOnElementReparented(FRigHierarchyContainer* InContainer, const FRigElementKey& InKey, const FName& InOldParentName, const FName& InNewParentName);
+	void HandleOnElementSelected(FRigHierarchyContainer* InContainer, const FRigElementKey& InKey, bool bSelected);
+#endif
+};
+
+// this struct is still here for backwards compatibility - but not used anywhere
+USTRUCT()
+struct CONTROLRIG_API FRigHierarchyRef
+{
+	GENERATED_BODY()
+};
\ No newline at end of file
diff --git a/Engine/Plugins/Experimental/ControlRig/Source/ControlRig/Public/Rigs/RigHierarchyDefines.h b/Engine/Plugins/Experimental/ControlRig/Source/ControlRig/Public/Rigs/RigHierarchyDefines.h
new file mode 100644
--- /dev/null
+++ b/Engine/Plugins/Experimental/ControlRig/Source/ControlRig/Public/Rigs/RigHierarchyDefines.h
@@ -0,0 +1,143 @@
+// Copyright 1998-2019 Epic Games, Inc. All Rights Reserved.
+
+#pragma once
+
+#include "CoreMinimal.h"
+#include "RigHierarchyDefines.generated.h"
+
+struct FRigHierarchyContainer;
+
+/* 
+ * This is rig element types that we support
+ * This can be used as a mask so supported as a bitfield
+ */
+UENUM()
+enum class ERigElementType : uint8
+{
+	None UMETA(Hidden),
+	Bone = 0x001,
+	Space = 0x002,
+	Control = 0x004,
+	Curve = 0x008,
+	All = Bone | Space | Control | Curve UMETA(Hidden),
+};
+
+/*
+ * Because it's bitfield, we support some basic functionality
+ */
+namespace FRigElementTypeHelper
+{
+	static uint32 Add(uint32 InMasks, ERigElementType InType)
+	{
+		return InMasks & (uint32)InType;
+	}
+
+	static uint32 Remove(uint32 InMasks, ERigElementType InType)
+	{
+		return InMasks & ~((uint32)InType);
+	}
+
+	static uint32 ToMask(ERigElementType InType)
+	{
+		return (uint32)InType;
+	}
+
+	static bool DoesHave(uint32 InMasks, ERigElementType InType)
+	{
+		return (InMasks & (uint32)InType) != 0;
+	}
+}
+
+struct FRigElementKey
+{
+	FName Name;
+	ERigElementType Type;
+
+	FRigElementKey()
+		: Name(NAME_None)
+		, Type(ERigElementType::None)
+	{}
+
+	FRigElementKey(const FName& InName, ERigElementType InType)
+		: Name(InName)
+		, Type(InType)
+	{}
+
+	FORCEINLINE bool IsValid() const
+	{
+		return Name != NAME_None && Type != ERigElementType::None;
+	}
+
+	FORCEINLINE operator bool() const
+	{
+		return IsValid();
+	}
+
+	friend FORCEINLINE uint32 GetTypeHash(const FRigElementKey& Key)
+	{
+		return GetTypeHash(Key.Name) * 10 + (uint32)Key.Type;
+	}
+
+	FORCEINLINE bool operator ==(const FRigElementKey& Other) const
+	{
+		return Name == Other.Name && Type == Other.Type;
+	}
+
+	FORCEINLINE bool operator !=(const FRigElementKey& Other) const
+	{
+		return Name != Other.Name || Type != Other.Type;
+	}
+
+	FORCEINLINE bool operator <(const FRigElementKey& Other) const
+	{
+		if (Type < Other.Type)
+		{
+			return true;
+		}
+		return Name.LexicalLess(Other.Name);
+	}
+
+	FORCEINLINE bool operator >(const FRigElementKey& Other) const
+	{
+		if (Type > Other.Type)
+		{
+			return true;
+		}
+		return Other.Name.LexicalLess(Name);
+	}
+};
+
+USTRUCT(BlueprintType)
+struct CONTROLRIG_API FRigElement
+{
+	GENERATED_BODY()
+
+	FRigElement()
+		: Name(NAME_None)
+		, Index(INDEX_NONE)
+	{}
+	virtual ~FRigElement() {}
+	
+	UPROPERTY(EditAnywhere, Category = FRigElement)
+	FName Name;
+
+	UPROPERTY(VisibleAnywhere, Category = FRigElement)
+	int32 Index;
+
+	virtual ERigElementType GetElementType() const
+	{
+		return ERigElementType::None; 
+	}
+
+	FRigElementKey GetElementKey() const
+	{
+		return FRigElementKey(Name, GetElementType());
+	}
+};
+
+DECLARE_MULTICAST_DELEGATE_TwoParams(FRigElementChanged, FRigHierarchyContainer*, const FRigElementKey&);
+DECLARE_MULTICAST_DELEGATE_TwoParams(FRigElementAdded, FRigHierarchyContainer*, const FRigElementKey&);
+DECLARE_MULTICAST_DELEGATE_TwoParams(FRigElementRemoved, FRigHierarchyContainer*, const FRigElementKey&);
+DECLARE_MULTICAST_DELEGATE_FourParams(FRigElementRenamed, FRigHierarchyContainer*, ERigElementType, const FName&, const FName&);
+DECLARE_MULTICAST_DELEGATE_FourParams(FRigElementReparented, FRigHierarchyContainer*, const FRigElementKey&, const FName&, const FName&);
+DECLARE_MULTICAST_DELEGATE_ThreeParams(FRigElementSelected, FRigHierarchyContainer*, const FRigElementKey&, bool);
diff --git a/Engine/Plugins/Experimental/ControlRig/Source/ControlRig/Public/Sequencer/ControlRigSequencerAnimInstance.h b/Engine/Plugins/Experimental/ControlRig/Source/ControlRig/Public/Sequencer/ControlRigSequencerAnimInstance.h
--- a/Engine/Plugins/Experimental/ControlRig/Source/ControlRig/Public/Sequencer/ControlRigSequencerAnimInstance.h
+++ b/Engine/Plugins/Experimental/ControlRig/Source/ControlRig/Public/Sequencer/ControlRigSequencerAnimInstance.h
@@ -15,7 +15,7 @@ class CONTROLRIG_API UControlRigSequencerAnimInstance : public UAnimSequencerIns
 
 public:
 	/** Update an animation ControlRig in this sequence */
-	bool UpdateControlRig(UControlRig* InControlRig, uint32 SequenceId, bool bAdditive, bool bApplyBoneFilter, const FInputBlendPose& BoneFilter, float Weight, bool bExternalSource = false);
+	bool UpdateControlRig(UControlRig* InControlRig, uint32 SequenceId, bool bAdditive, bool bApplyBoneFilter, const FInputBlendPose& BoneFilter, float Weight, bool bUpdateInput, bool bExecute);
 
 	virtual bool SetAnimationAsset(class UAnimationAsset* NewAsset);
 
diff --git a/Engine/Plugins/Experimental/ControlRig/Source/ControlRig/Public/Settings/ControlRigSettings.h b/Engine/Plugins/Experimental/ControlRig/Source/ControlRig/Public/Settings/ControlRigSettings.h
new file mode 100644
--- /dev/null
+++ b/Engine/Plugins/Experimental/ControlRig/Source/ControlRig/Public/Settings/ControlRigSettings.h
@@ -0,0 +1,30 @@
+// Copyright 1998-2019 Epic Games, Inc. All Rights Reserved.
+
+/*=============================================================================
+	ControlRigSettings.h: Declares the ControlRigSettings class.
+=============================================================================*/
+
+#pragma once
+
+#include "CoreMinimal.h"
+#include "Engine/DeveloperSettings.h"
+#include "ControlRigGizmoLibrary.h"
+#include "ControlRigSettings.generated.h"
+
+class UStaticMesh;
+
+/**
+ * Default ControlRig settings.
+ */
+UCLASS(config=Engine, defaultconfig, meta=(DisplayName="Control Rig"))
+class CONTROLRIG_API UControlRigSettings : public UDeveloperSettings
+{
+	GENERATED_UCLASS_BODY()
+
+public:
+
+	UPROPERTY(EditAnywhere, config, Category = DefaultGizmo)
+	TAssetPtr<UControlRigGizmoLibrary> DefaultGizmoLibrary;
+
+	static UControlRigSettings * Get() { return CastChecked<UControlRigSettings>(UControlRigSettings::StaticClass()->GetDefaultObject()); }
+};
diff --git a/Engine/Plugins/Experimental/ControlRig/Source/ControlRig/Public/Units/Control/RigUnit_Control.h b/Engine/Plugins/Experimental/ControlRig/Source/ControlRig/Public/Units/Control/RigUnit_Control.h
--- a/Engine/Plugins/Experimental/ControlRig/Source/ControlRig/Public/Units/Control/RigUnit_Control.h
+++ b/Engine/Plugins/Experimental/ControlRig/Source/ControlRig/Public/Units/Control/RigUnit_Control.h
@@ -4,12 +4,11 @@
 
 #include "Units/RigUnit.h"
 #include "Constraint.h"
-#include "ControlRigControl.h"
 #include "EulerTransform.h"
 #include "RigUnit_Control.generated.h"
 
 /** A control unit used to drive a transform from an external source */
-USTRUCT(meta=(DisplayName="Control", Category="Controls", ShowVariableNameInTitle))
+USTRUCT(meta=(DisplayName="Control", Category="Controls", ShowVariableNameInTitle, Deprecated = "4.24.0"))
 struct CONTROLRIG_API FRigUnit_Control : public FRigUnit
 {
 	GENERATED_BODY()
@@ -21,28 +20,28 @@ struct CONTROLRIG_API FRigUnit_Control : public FRigUnit
 	{
 	}
 
+	RIGVM_METHOD()
 	virtual void Execute(const FRigUnitContext& Context) override;
 
 	/** Combine Transform and Base to make the resultant transform */
 	FTransform GetResultantTransform() const;
+	static FTransform StaticGetResultantTransform(const FEulerTransform& InTransform, const FTransform& InBase, const FTransformFilter& InFilter);
 
 	/** Combine Transform and Base to make the resultant transform (as a matrix) */
 	FMatrix GetResultantMatrix() const;
+	static FMatrix StaticGetResultantMatrix(const FEulerTransform& InTransform, const FTransform& InBase, const FTransformFilter& InFilter);
 
 	/** Set the transform using a resultant transform (already incorporating Base) */
 	void SetResultantTransform(const FTransform& InResultantTransform);
+	static void StaticSetResultantTransform(const FTransform& InResultantTransform, const FTransform& InBase, FEulerTransform& OutTransform);
 
 	/** Set the transform using a resultant matrix (already incorporating Base) */
 	void SetResultantMatrix(const FMatrix& InResultantMatrix);
+	static void StaticSetResultantMatrix(const FMatrix& InResultantMatrix, const FTransform& InBase, FEulerTransform& OutTransform);
 
 	/** Get the local transform (i.e. without base) with filter applied */
 	FEulerTransform GetFilteredTransform() const;
-
-#if WITH_EDITORONLY_DATA
-	/** Actor class to use to display this in the viewport */
-	UPROPERTY()
-	TSubclassOf<AControlRigControl> ControlClass;
-#endif
+	static FEulerTransform StaticGetFilteredTransform(const FEulerTransform& InTransform, const FTransformFilter& InFilter);
 
 	/** The transform of this control */
 	UPROPERTY(EditAnywhere, Category="Control", Interp, meta = (AnimationInput))
diff --git a/Engine/Plugins/Experimental/ControlRig/Source/ControlRig/Public/Units/Control/RigUnit_Control_StaticMesh.h b/Engine/Plugins/Experimental/ControlRig/Source/ControlRig/Public/Units/Control/RigUnit_Control_StaticMesh.h
--- a/Engine/Plugins/Experimental/ControlRig/Source/ControlRig/Public/Units/Control/RigUnit_Control_StaticMesh.h
+++ b/Engine/Plugins/Experimental/ControlRig/Source/ControlRig/Public/Units/Control/RigUnit_Control_StaticMesh.h
@@ -5,6 +5,9 @@
 #include "RigUnit_Control.h"
 #include "RigUnit_Control_StaticMesh.generated.h"
 
+class UStaticMesh;
+class UMaterialInterface;
+
 /** A control unit used to drive a transform from an external source */
 USTRUCT(meta=(DisplayName="Static Mesh Control", Category="Controls", ShowVariableNameInTitle))
 struct CONTROLRIG_API FRigUnit_Control_StaticMesh : public FRigUnit_Control
diff --git a/Engine/Plugins/Experimental/ControlRig/Source/ControlRigDeveloper/Private/ControlRigBlueprint.cpp b/Engine/Plugins/Experimental/ControlRig/Source/ControlRigDeveloper/Private/ControlRigBlueprint.cpp
--- a/Engine/Plugins/Experimental/ControlRig/Source/ControlRigDeveloper/Private/ControlRigBlueprint.cpp
+++ b/Engine/Plugins/Experimental/ControlRig/Source/ControlRigDeveloper/Private/ControlRigBlueprint.cpp
@@ -10,11 +10,13 @@
 #include "ControlRig.h"
 #include "Graph/ControlRigGraph.h"
 #include "Graph/ControlRigGraphNode.h"
+#include "UObject/UObjectGlobals.h"
 
 #if WITH_EDITOR
 #include "IControlRigEditorModule.h"
 #include "Kismet2/BlueprintEditorUtils.h"
 #include "ControlRigBlueprintUtils.h"
+#include "Settings/ControlRigSettings.h"
 #endif//WITH_EDITOR
 
 #define LOCTEXT_NAMESPACE "ControlRigBlueprint"
@@ -25,6 +27,10 @@ UControlRigBlueprint::UControlRigBlueprint()
 	bSuspendModelNotificationsForOthers = false;
 	Model = nullptr;
 	ModelController = nullptr;
+
+#if WITH_EDITORONLY_DATA
+	GizmoLibrary = UControlRigSettings::Get()->DefaultGizmoLibrary;
+#endif
 }
 
 void UControlRigBlueprint::InitializeModel()
@@ -45,6 +51,12 @@ void UControlRigBlueprint::InitializeModel()
 				Graph->Initialize(this);
 			}
 		}
+
+		HierarchyContainer.OnElementAdded.AddUObject(this, &UControlRigBlueprint::HandleOnElementAdded);
+		HierarchyContainer.OnElementRemoved.AddUObject(this, &UControlRigBlueprint::HandleOnElementRemoved);
+		HierarchyContainer.OnElementRenamed.AddUObject(this, &UControlRigBlueprint::HandleOnElementRenamed);
+		HierarchyContainer.OnElementReparented.AddUObject(this, &UControlRigBlueprint::HandleOnElementReparented);
+		HierarchyContainer.OnElementSelected.AddUObject(this, &UControlRigBlueprint::HandleOnElementSelected);
 	}
 }
 
@@ -59,6 +71,7 @@ UControlRigBlueprintGeneratedClass* UControlRigBlueprint::GetControlRigBlueprint
 	UControlRigBlueprintGeneratedClass* Result = Cast<UControlRigBlueprintGeneratedClass>(*SkeletonGeneratedClass);
 	return Result;
 }
+
 UClass* UControlRigBlueprint::GetBlueprintClass() const
 {
 	return UControlRigBlueprintGeneratedClass::StaticClass();
@@ -123,6 +136,57 @@ void UControlRigBlueprint::SetObjectBeingDebugged(UObject* NewObject)
 	Super::SetObjectBeingDebugged(NewObject);
 }
 
+void UControlRigBlueprint::PostLoad()
+{
+	DECLARE_SCOPE_HIERARCHICAL_COUNTER_FUNC()
+
+	Super::PostLoad();
+
+	// remove all non-controlrig-graphs
+	TArray<UEdGraph*> NewUberGraphPages;
+	for (UEdGraph* Graph : UbergraphPages)
+	{
+		UControlRigGraph* RigGraph = Cast<UControlRigGraph>(Graph);
+		if (RigGraph)
+		{
+			NewUberGraphPages.Add(RigGraph);
+		}
+		else
+		{
+			Graph->MarkPendingKill();
+			Graph->Rename(nullptr, GetTransientPackage(), REN_ForceNoResetLoaders);
+		}
+	}
+	UbergraphPages = NewUberGraphPages;
+}
+
+void UControlRigBlueprint::PostTransacted(const FTransactionObjectEvent& TransactionEvent)
+{
+	DECLARE_SCOPE_HIERARCHICAL_COUNTER_FUNC()
+	Super::PostTransacted(TransactionEvent);
+
+	if (TransactionEvent.GetEventType() == ETransactionObjectEventType::UndoRedo)
+	{
+		if (TransactionEvent.GetChangedProperties().Contains(TEXT("HierarchyContainer")))
+		{
+			PropagateHierarchyFromBPToInstances(true);
+			HierarchyContainer.OnElementChanged.Broadcast(&HierarchyContainer, FRigElementKey());
+
+			// make sure the bone name list is up 2 date for the editor graph
+			for (UEdGraph* Graph : UbergraphPages)
+			{
+				UControlRigGraph* RigGraph = Cast<UControlRigGraph>(Graph);
+				if (RigGraph == nullptr)
+				{
+					continue;
+				}
+				RigGraph->CacheNameLists(&HierarchyContainer);
+			}
+		}
+	}
+}
+
+
 UControlRigModel::FModifiedEvent& UControlRigBlueprint::OnModified()
 {
 	return _ModifiedEvent;
@@ -252,7 +316,6 @@ void UControlRigBlueprint::RebuildGraphFromModel()
 	ModelController->ResendAllNotifications();
 }
 
-
 void UControlRigBlueprint::HandleModelModified(const UControlRigModel* InModel, EControlRigModelNotifType InType, const void* InPayload)
 {
 	DECLARE_SCOPE_HIERARCHICAL_COUNTER_FUNC()
@@ -598,5 +661,87 @@ bool UControlRigBlueprint::PerformArrayOperation(const FString& InPropertyPath,
 	return false;
 }
 
+void UControlRigBlueprint::CleanupBoneHierarchyDeprecated()
+{
+	if (Hierarchy_DEPRECATED.Num() > 0)
+	{
+		HierarchyContainer.BoneHierarchy = Hierarchy_DEPRECATED;
+		Hierarchy_DEPRECATED.Reset();
+	}
+
+	if (CurveContainer_DEPRECATED.Num() > 0)
+	{
+		HierarchyContainer.CurveContainer = CurveContainer_DEPRECATED;
+		CurveContainer_DEPRECATED.Reset();
+	}
+
+}
+
+void UControlRigBlueprint::PropagatePoseFromInstanceToBP()
+{
+	if (UClass* MyControlRigClass = GeneratedClass)
+	{
+		if (UControlRig* DebuggedRig = Cast<UControlRig>(GetObjectBeingDebugged()))
+		{
+			/* todo */
+		}
+	}
+}
+
+void UControlRigBlueprint::PropagateHierarchyFromBPToInstances(bool bInitialize)
+{
+	if (UClass* MyControlRigClass = GeneratedClass)
+	{
+		if (UControlRig* DefaultObject = Cast<UControlRig>(MyControlRigClass->GetDefaultObject(false)))
+		{
+			if (bInitialize)
+			{
+				HierarchyContainer.Initialize();
+				HierarchyContainer.ResetTransforms();
+			}
+
+			DefaultObject->Hierarchy = HierarchyContainer;
+
+			TArray<UObject*> ArchetypeInstances;
+			DefaultObject->GetArchetypeInstances(ArchetypeInstances);
+			for (UObject* ArchetypeInstance : ArchetypeInstances)
+			{
+				if (UControlRig* InstanceRig = Cast<UControlRig>(ArchetypeInstance))
+				{
+					InstanceRig->Hierarchy = HierarchyContainer;
+				}
+			}
+		}
+	}
+}
+
+#if WITH_EDITOR
+
+void UControlRigBlueprint::HandleOnElementAdded(FRigHierarchyContainer* InContainer, const FRigElementKey& InKey)
+{
+	PropagateHierarchyFromBPToInstances();
+}
+
+void UControlRigBlueprint::HandleOnElementRemoved(FRigHierarchyContainer* InContainer, const FRigElementKey& InKey)
+{
+	PropagateHierarchyFromBPToInstances();
+}
+
+void UControlRigBlueprint::HandleOnElementRenamed(FRigHierarchyContainer* InContainer, ERigElementType InElementType, const FName& InOldName, const FName& InNewName)
+{
+	PropagateHierarchyFromBPToInstances();
+}
+
+void UControlRigBlueprint::HandleOnElementReparented(FRigHierarchyContainer* InContainer, const FRigElementKey& InKey, const FName& InOldParentName, const FName& InNewParentName)
+{
+	PropagateHierarchyFromBPToInstances();
+}
+
+void UControlRigBlueprint::HandleOnElementSelected(FRigHierarchyContainer* InContainer, const FRigElementKey& InKey, bool bSelected)
+{
+}
+
+#endif
+
 #undef LOCTEXT_NAMESPACE
 
diff --git a/Engine/Plugins/Experimental/ControlRig/Source/ControlRigDeveloper/Private/ControlRigController.cpp b/Engine/Plugins/Experimental/ControlRig/Source/ControlRigDeveloper/Private/ControlRigController.cpp
--- a/Engine/Plugins/Experimental/ControlRig/Source/ControlRigDeveloper/Private/ControlRigController.cpp
+++ b/Engine/Plugins/Experimental/ControlRig/Source/ControlRigDeveloper/Private/ControlRigController.cpp
@@ -87,6 +87,25 @@ bool UControlRigController::CloseUndoBracket()
 	return true;
 }
 
+bool UControlRigController::CancelUndoBracket()
+{
+	DECLARE_SCOPE_HIERARCHICAL_COUNTER_FUNC()
+
+	if (!EnsureModel())
+	{
+		return false;
+	}
+
+	if (_UndoBrackets.Num() == 0)
+	{
+		return false;
+	}
+
+	Model->CurrentActions.Pop();
+	_UndoBrackets.Pop();
+	return true;
+}
+
 #endif
 
 
diff --git a/Engine/Plugins/Experimental/ControlRig/Source/ControlRigDeveloper/Private/ControlRigModel.cpp b/Engine/Plugins/Experimental/ControlRig/Source/ControlRigDeveloper/Private/ControlRigModel.cpp
--- a/Engine/Plugins/Experimental/ControlRig/Source/ControlRigDeveloper/Private/ControlRigModel.cpp
+++ b/Engine/Plugins/Experimental/ControlRig/Source/ControlRigDeveloper/Private/ControlRigModel.cpp
@@ -12,6 +12,7 @@
 
 #if CONTROLRIG_UNDO
 #include "ScopedTransaction.h"
+#include "Misc/ITransaction.h"
 #endif
 
 #define LOCTEXT_NAMESPACE "ControlRigModel"
diff --git a/Engine/Plugins/Experimental/ControlRig/Source/ControlRigDeveloper/Private/Graph/ControlRigGraph.cpp b/Engine/Plugins/Experimental/ControlRig/Source/ControlRigDeveloper/Private/Graph/ControlRigGraph.cpp
--- a/Engine/Plugins/Experimental/ControlRig/Source/ControlRigDeveloper/Private/Graph/ControlRigGraph.cpp
+++ b/Engine/Plugins/Experimental/ControlRig/Source/ControlRigDeveloper/Private/Graph/ControlRigGraph.cpp
@@ -272,24 +272,13 @@ void UControlRigGraph::OnBlueprintCompiledPostLoad(UBlueprint* InCompiledBluepri
 	}
 }
 
-void UControlRigGraph::CacheBoneNameList(const FRigHierarchy& Hierarchy)
+void UControlRigGraph::CacheNameLists(const FRigHierarchyContainer* Container)
 {
-	DECLARE_SCOPE_HIERARCHICAL_COUNTER_FUNC()
-
-	TArray<FString> Names;
-	const TArray<FRigBone>& Bones = Hierarchy.GetBones();
-	for (const FRigBone& Bone : Bones )
-	{
-		Names.Add(Bone.Name.ToString());
-	}
-	Names.Sort();
-
-	BoneNameList.Reset();
-	BoneNameList.Add(MakeShared<FString>(FName(NAME_None).ToString()));
-	for (const FString& Name : Names)
-	{
-		BoneNameList.Add(MakeShared<FString>(Name));
-	}
+	check(Container);
+	CacheNameList<FRigBoneHierarchy>(Container->BoneHierarchy, BoneNameList);
+	CacheNameList<FRigControlHierarchy>(Container->ControlHierarchy, ControlNameList);
+	CacheNameList<FRigSpaceHierarchy>(Container->SpaceHierarchy, SpaceNameList);
+	CacheNameList<FRigCurveContainer>(Container->CurveContainer, CurveNameList);
 }
 
 const TArray<TSharedPtr<FString>>& UControlRigGraph::GetBoneNameList() const
@@ -297,6 +286,20 @@ const TArray<TSharedPtr<FString>>& UControlRigGraph::GetBoneNameList() const
 	return BoneNameList;
 }
 
+const TArray<TSharedPtr<FString>>& UControlRigGraph::GetControlNameList() const
+{
+	return ControlNameList;
+}
+
+const TArray<TSharedPtr<FString>>& UControlRigGraph::GetSpaceNameList() const
+{
+	return SpaceNameList;
+}
+
+const TArray<TSharedPtr<FString>>& UControlRigGraph::GetCurveNameList() const
+{
+	return CurveNameList;
+}
 
 void UControlRigGraph::HandleModelModified(const UControlRigModel* InModel, EControlRigModelNotifType InType, const void* InPayload)
 {
@@ -574,30 +577,6 @@ UEdGraphNode* UControlRigGraph::FindNodeFromPropertyName(const FName& InProperty
 	return nullptr;
 }
 
-void UControlRigGraph::CacheCurveNameList(const FRigCurveContainer& Container)
-{
-	DECLARE_SCOPE_HIERARCHICAL_COUNTER_FUNC()
-
-	TArray<FString> Names;
-	const TArray<FRigCurve>& Curves = Container.GetCurves();
-	for (const FRigCurve& Curve : Curves)
-	{
-		Names.Add(Curve.Name.ToString());
-	}
-	Names.Sort();
-
-	CurveNameList.Reset();
-	CurveNameList.Add(MakeShared<FString>(FName(NAME_None).ToString()));
-	for (const FString& Name : Names)
-	{
-		CurveNameList.Add(MakeShared<FString>(Name));
-	}
-}
-
-const TArray<TSharedPtr<FString>>& UControlRigGraph::GetCurveNameList() const
-{
-	return CurveNameList;
-}
 #endif
 
 #undef LOCTEXT_NAMESPACE
diff --git a/Engine/Plugins/Experimental/ControlRig/Source/ControlRigDeveloper/Public/ControlRigBlueprint.h b/Engine/Plugins/Experimental/ControlRig/Source/ControlRigDeveloper/Public/ControlRigBlueprint.h
--- a/Engine/Plugins/Experimental/ControlRig/Source/ControlRigDeveloper/Public/ControlRigBlueprint.h
+++ b/Engine/Plugins/Experimental/ControlRig/Source/ControlRigDeveloper/Public/ControlRigBlueprint.h
@@ -8,10 +8,11 @@
 #include "Engine/Blueprint.h"
 #include "Misc/Crc.h"
 #include "ControlRigDefines.h"
-#include "Hierarchy.h"
-#include "CurveContainer.h"
+#include "Rigs/RigHierarchyContainer.h"
+#include "Rigs/RigCurveContainer.h"
 #include "Interfaces/Interface_PreviewMeshProvider.h"
 #include "ControlRigController.h"
+#include "ControlRigGizmoLibrary.h"
 #include "ControlRigBlueprint.generated.h"
 
 class UControlRigBlueprintGeneratedClass;
@@ -116,6 +117,7 @@ public:
 	UControlRigBlueprintGeneratedClass* GetControlRigBlueprintSkeletonClass() const;
 
 #if WITH_EDITOR
+
 	// UBlueprint interface
 	virtual UClass* GetBlueprintClass() const override;
 	virtual bool SupportedByDefaultBlueprintFactory() const override { return false; }
@@ -124,6 +126,19 @@ public:
 	virtual void GetTypeActions(FBlueprintActionDatabaseRegistrar& ActionRegistrar) const override;
 	virtual void GetInstanceActions(FBlueprintActionDatabaseRegistrar& ActionRegistrar) const override;	
 	virtual void SetObjectBeingDebugged(UObject* NewObject) override;
+	virtual bool RequiresMarkAsStructurallyModifiedOnUndo() const override { return false; }
+	virtual void PostLoad() override;
+	virtual void PostTransacted(const FTransactionObjectEvent& TransactionEvent) override;
+
+	virtual bool SupportsGlobalVariables() const override { return false; }
+	virtual bool SupportsLocalVariables() const override { return false; }
+	virtual bool SupportsFunctions() const override { return false; }
+	virtual bool SupportsMacros() const override { return false; }
+	virtual bool SupportsDelegates() const override { return false; }
+	virtual bool SupportsEventGraphs() const override { return false; }
+	virtual bool SupportsAnimLayers() const override { return false; }
+
+
 #endif	// #if WITH_EDITOR
 
 	/** Make a property link between the specified properties - used by the compiler */
@@ -148,6 +163,11 @@ public:
 
 	UControlRigModel::FModifiedEvent& OnModified();
 
+#if WITH_EDITORONLY_DATA
+	UPROPERTY(EditAnywhere, config, Category = DefaultGizmo)
+	TAssetPtr<UControlRigGizmoLibrary> GizmoLibrary;
+#endif
+
 private:
 	/** Links between the various properties we have */
 	UPROPERTY()
@@ -162,11 +182,17 @@ private:
 	UPROPERTY()
 	TMap<FName, FString> AllowSourceAccessProperties;
 
+public:
+	UPROPERTY()
+	FRigHierarchyContainer HierarchyContainer;
+
+private:
+
 	UPROPERTY()
-	FRigHierarchy Hierarchy;
+	FRigBoneHierarchy Hierarchy_DEPRECATED;
 
 	UPROPERTY()
-	FRigCurveContainer CurveContainer;
+	FRigCurveContainer CurveContainer_DEPRECATED;
 
 	/** The default skeletal mesh to use when previewing this asset */
 	UPROPERTY(DuplicateTransient, AssetRegistrySearchable)
@@ -183,10 +209,22 @@ private:
 	void HandleModelModified(const UControlRigModel* InModel, EControlRigModelNotifType InType, const void* InPayload);
 	bool UpdateParametersOnControlRig(UControlRig* InRig = nullptr);
 	bool PerformArrayOperation(const FString& InPropertyPath, TFunctionRef<bool(FScriptArrayHelper&, int32)> InOperation, bool bCallModify, bool bPropagateToInstances);
+	void CleanupBoneHierarchyDeprecated();
+
+	void PropagatePoseFromInstanceToBP();
+	void PropagateHierarchyFromBPToInstances(bool bInitialize = true);
+#if WITH_EDITOR
+	void HandleOnElementAdded(FRigHierarchyContainer* InContainer, const FRigElementKey& InKey);
+	void HandleOnElementRemoved(FRigHierarchyContainer* InContainer, const FRigElementKey& InKey);
+	void HandleOnElementRenamed(FRigHierarchyContainer* InContainer, ERigElementType InElementType, const FName& InOldName, const FName& InNewName);
+	void HandleOnElementReparented(FRigHierarchyContainer* InContainer, const FRigElementKey& InKey, const FName& InOldParentName, const FName& InNewParentName);
+	void HandleOnElementSelected(FRigHierarchyContainer* InContainer, const FRigElementKey& InKey, bool bSelected);
+#endif
 
 	friend class FControlRigBlueprintCompilerContext;
 	friend class SRigHierarchy;
 	friend class SRigCurveContainer;
 	friend class FControlRigEditor;
 	friend class UEngineTestControlRig;
+	friend class FControlRigEditMode;
 };
diff --git a/Engine/Plugins/Experimental/ControlRig/Source/ControlRigDeveloper/Public/ControlRigController.h b/Engine/Plugins/Experimental/ControlRig/Source/ControlRigDeveloper/Public/ControlRigController.h
--- a/Engine/Plugins/Experimental/ControlRig/Source/ControlRigDeveloper/Public/ControlRigController.h
+++ b/Engine/Plugins/Experimental/ControlRig/Source/ControlRigDeveloper/Public/ControlRigController.h
@@ -86,6 +86,7 @@ public:
 #if CONTROLRIG_UNDO
 	virtual bool OpenUndoBracket(const FString& Title);
 	virtual bool CloseUndoBracket();
+	virtual bool CancelUndoBracket();
 #endif
 
 	UControlRigModel::FModifiedEvent& OnModified();
diff --git a/Engine/Plugins/Experimental/ControlRig/Source/ControlRigDeveloper/Public/Graph/ControlRigGraph.h b/Engine/Plugins/Experimental/ControlRig/Source/ControlRigDeveloper/Public/Graph/ControlRigGraph.h
--- a/Engine/Plugins/Experimental/ControlRig/Source/ControlRigDeveloper/Public/Graph/ControlRigGraph.h
+++ b/Engine/Plugins/Experimental/ControlRig/Source/ControlRigDeveloper/Public/Graph/ControlRigGraph.h
@@ -4,7 +4,7 @@
 
 #include "EdGraph/EdGraph.h"
 #include "Graph/ControlRigGraphNode.h"
-#include "Hierarchy.h"
+#include "Rigs/RigHierarchyContainer.h"
 #include "ControlRigModel.h"
 #include "ControlRigGraph.generated.h"
 
@@ -36,9 +36,12 @@ public:
 	virtual void PostLoad() override;
 	void OnBlueprintCompiledPostLoad(UBlueprint*);
 	FDelegateHandle BlueprintOnCompiledHandle;
-	void CacheBoneNameList(const FRigHierarchy& Hierarchy);
+
+	void CacheNameLists(const FRigHierarchyContainer* Container);
+
 	const TArray<TSharedPtr<FString>>& GetBoneNameList() const;
-	void CacheCurveNameList(const FRigCurveContainer& Container);
+	const TArray<TSharedPtr<FString>>& GetControlNameList() const;
+	const TArray<TSharedPtr<FString>>& GetSpaceNameList() const;
 	const TArray<TSharedPtr<FString>>& GetCurveNameList() const;
 
 	bool bSuspendModelNotifications;
@@ -50,10 +53,33 @@ private:
 
 	void HandleModelModified(const UControlRigModel* InModel, EControlRigModelNotifType InType, const void* InPayload);
 
+	template<class T>
+	void CacheNameList(const T& Hierarchy, TArray<TSharedPtr<FString>>& OutNameList)
+	{
+		DECLARE_SCOPE_HIERARCHICAL_COUNTER_FUNC()
+
+		TArray<FString> Names;
+		for (auto Element : Hierarchy)
+		{
+			Names.Add(Element.Name.ToString());
+		}
+		Names.Sort();
+
+		OutNameList.Reset();
+		OutNameList.Add(MakeShared<FString>(FName(NAME_None).ToString()));
+		for (const FString& Name : Names)
+		{
+			OutNameList.Add(MakeShared<FString>(Name));
+		}
+	}
+
 	TArray<UControlRigGraphNode*> FoundHierarchyRefVariableNodes;
 	TArray<UControlRigGraphNode*> FoundHierarchyRefMutableNodes;
 	TMap<UControlRigGraphNode*, TArray<UControlRigGraphNode*>> FoundHierarchyRefConnections;
+
 	TArray<TSharedPtr<FString>> BoneNameList;
+	TArray<TSharedPtr<FString>> ControlNameList;
+	TArray<TSharedPtr<FString>> SpaceNameList;
 	TArray<TSharedPtr<FString>> CurveNameList;
 #endif
 };
diff --git a/Engine/Plugins/Experimental/ControlRig/Source/ControlRigEditor/ControlRigEditor.Build.cs b/Engine/Plugins/Experimental/ControlRig/Source/ControlRigEditor/ControlRigEditor.Build.cs
--- a/Engine/Plugins/Experimental/ControlRig/Source/ControlRigEditor/ControlRigEditor.Build.cs
+++ b/Engine/Plugins/Experimental/ControlRig/Source/ControlRigEditor/ControlRigEditor.Build.cs
@@ -53,6 +53,7 @@ namespace UnrealBuildTool.Rules
 					"Json",
 					"DesktopPlatform",
 					"ToolMenus",
+					"ControlRigManipulation"
                 }
             );
         }
diff --git a/Engine/Plugins/Experimental/ControlRig/Source/ControlRigEditor/Private/ControlRigEditorModule.cpp b/Engine/Plugins/Experimental/ControlRig/Source/ControlRigEditor/Private/ControlRigEditorModule.cpp
--- a/Engine/Plugins/Experimental/ControlRig/Source/ControlRigEditor/Private/ControlRigEditorModule.cpp
+++ b/Engine/Plugins/Experimental/ControlRig/Source/ControlRigEditor/Private/ControlRigEditorModule.cpp
@@ -50,6 +50,7 @@
 #include "AssetRegistryModule.h"
 #include "Editor/ControlRigEditor.h"
 #include "ControlRigBlueprintActions.h"
+#include "ControlRigGizmoLibraryActions.h"
 #include "Graph/ControlRigGraphSchema.h"
 #include "Graph/ControlRigGraph.h"
 #include "Graph/NodeSpawners/ControlRigPropertyNodeSpawner.h"
@@ -69,6 +70,7 @@
 #include "Animation/AnimSequence.h"
 #include "ControlRigEditorEditMode.h"
 #include "ControlRigDetails.h"
+#include "ControlRigElementDetails.h"
 #include "Units/Deprecated/RigUnitEditor_TwoBoneIKFK.h"
 #include "Animation/AnimSequence.h"
 #include "Editor/SControlRigProfilingView.h"
@@ -123,12 +125,22 @@ void FControlRigEditorModule::StartupModule()
 	ClassesToUnregisterOnShutdown.Add(UControlRigSequenceExporterSettings::StaticClass()->GetFName());
 	PropertyEditorModule.RegisterCustomClassLayout(ClassesToUnregisterOnShutdown.Last(), FOnGetDetailCustomizationInstance::CreateStatic(&FControlRigSequenceExporterSettingsDetailsCustomization::MakeInstance));
 
+	ClassesToUnregisterOnShutdown.Add(FRigBone::StaticStruct()->GetFName());
+	PropertyEditorModule.RegisterCustomClassLayout(ClassesToUnregisterOnShutdown.Last(), FOnGetDetailCustomizationInstance::CreateStatic(&FRigBoneDetails::MakeInstance));
+
+	ClassesToUnregisterOnShutdown.Add(FRigControl::StaticStruct()->GetFName());
+	PropertyEditorModule.RegisterCustomClassLayout(ClassesToUnregisterOnShutdown.Last(), FOnGetDetailCustomizationInstance::CreateStatic(&FRigControlDetails::MakeInstance));
+
+	ClassesToUnregisterOnShutdown.Add(FRigSpace::StaticStruct()->GetFName());
+	PropertyEditorModule.RegisterCustomClassLayout(ClassesToUnregisterOnShutdown.Last(), FOnGetDetailCustomizationInstance::CreateStatic(&FRigSpaceDetails::MakeInstance));
+
 	ClassesToUnregisterOnShutdown.Add(UControlRig::StaticClass()->GetFName());
 	PropertyEditorModule.RegisterCustomClassLayout(ClassesToUnregisterOnShutdown.Last(), FOnGetDetailCustomizationInstance::CreateStatic(&FControlRigDetails::MakeInstance));
 
 	// same as ClassesToUnregisterOnShutdown but for properties, there is none right now
 	PropertiesToUnregisterOnShutdown.Reset();
 
+
 	// Register asset tools
 	auto RegisterAssetTypeAction = [this](const TSharedRef<IAssetTypeActions>& InAssetTypeAction)
 	{
@@ -139,6 +151,7 @@ void FControlRigEditorModule::StartupModule()
 
 	RegisterAssetTypeAction(MakeShareable(new FControlRigSequenceActions()));
 	RegisterAssetTypeAction(MakeShareable(new FControlRigBlueprintActions()));
+	RegisterAssetTypeAction(MakeShareable(new FControlRigGizmoLibraryActions()));
 	
 	// Register sequencer track editor
 	ISequencerModule& SequencerModule = FModuleManager::Get().LoadModuleChecked<ISequencerModule>("Sequencer");
@@ -465,7 +478,7 @@ void FControlRigEditorModule::HandleSequencerCreated(TSharedRef<ISequencer> InSe
 				if (FControlRigEditMode* ControlRigEditMode = static_cast<FControlRigEditMode*>(GLevelEditorModeTools().GetActiveMode(FControlRigEditMode::ModeName)))
 				{
 					ControlRigEditMode->SetSequencer(nullptr);
-					ControlRigEditMode->SetObjects(TWeakObjectPtr<>(), FGuid());
+					ControlRigEditMode->SetObjects(TWeakObjectPtr<>(), FGuid(), nullptr);
 				}
 			}
 		}
@@ -502,7 +515,7 @@ void FControlRigEditorModule::HandleSequencerCreated(TSharedRef<ISequencer> InSe
 					GLevelEditorModeTools().ActivateMode(FControlRigEditMode::ModeName);
 					if (FControlRigEditMode* ControlRigEditMode = static_cast<FControlRigEditMode*>(GLevelEditorModeTools().GetActiveMode(FControlRigEditMode::ModeName)))
 					{
-						ControlRigEditMode->SetObjects(SelectedObject, ObjectBinding);
+						ControlRigEditMode->SetObjects(SelectedObject, ObjectBinding, nullptr);
 					}
 				}
 			}
@@ -520,7 +533,6 @@ void FControlRigEditorModule::HandleSequencerCreated(TSharedRef<ISequencer> InSe
 				if (FControlRigEditMode* ControlRigEditMode = static_cast<FControlRigEditMode*>(GLevelEditorModeTools().GetActiveMode(FControlRigEditMode::ModeName)))
 				{
 					ControlRigEditMode->RefreshObjects();
-					ControlRigEditMode->RefreshTrajectoryCache();
 				}
 			}
 		}
@@ -547,8 +559,12 @@ void FControlRigEditorModule::HandleSequencerCreated(TSharedRef<ISequencer> InSe
 
 				if (FControlRigEditMode* ControlRigEditMode = static_cast<FControlRigEditMode*>(GLevelEditorModeTools().GetActiveMode(FControlRigEditMode::ModeName)))
 				{
-					ControlRigEditMode->ClearControlSelection();
-					ControlRigEditMode->SetControlSelection(PropertyPaths, true);
+					ControlRigEditMode->ClearRigElementSelection(FRigElementTypeHelper::ToMask(ERigElementType::All));
+					for (int32 Index = 0; Index < PropertyPaths.Num(); ++Index)
+					{
+						ControlRigEditMode->SetRigElementSelection(ERigElementType::Control, FName(*PropertyPaths[Index], FNAME_Find), true);
+					}
+					
 				}
 			}
 		}
@@ -722,13 +738,14 @@ bool FControlRigEditorModule::IsTrackVisible(const UMovieSceneTrack* InTrack)
 	if (FControlRigEditMode* ControlRigEditMode = static_cast<FControlRigEditMode*>(GLevelEditorModeTools().GetActiveMode(FControlRigEditMode::ModeName)))
 	{		
 		// If nothing selected, show all nodes
-		if (ControlRigEditMode->GetNumSelectedControls() == 0)
+		if (ControlRigEditMode->AreRigElementsSelected(FRigElementTypeHelper::ToMask(ERigElementType::Control)) == 0)
 		{
 			return true;
 		}
 
-		return ControlRigEditMode->IsControlSelected(ControlRigEditMode->GetControlFromPropertyPath(InTrack->GetTrackName().ToString()));
+		return ControlRigEditMode->SelectedRigElements.Contains(FRigElementKey(InTrack->GetTrackName(), ERigElementType::Control));
 	}
+
 	return true;
 }
 
diff --git a/Engine/Plugins/Experimental/ControlRig/Source/ControlRigEditor/Private/ControlRigEditorStyle.h b/Engine/Plugins/Experimental/ControlRig/Source/ControlRigEditor/Private/ControlRigEditorStyle.h
--- a/Engine/Plugins/Experimental/ControlRig/Source/ControlRigEditor/Private/ControlRigEditorStyle.h
+++ b/Engine/Plugins/Experimental/ControlRig/Source/ControlRigEditor/Private/ControlRigEditorStyle.h
@@ -30,54 +30,61 @@ public:
 
 		// Class Icons
 		{
-			Set("ClassIcon.ControlRigSequence", new IMAGE_BRUSH("ControlRigSequence_16x", Icon16x16));
-			Set("ClassIcon.ControlRigBlueprint", new IMAGE_BRUSH("ControlRigBlueprint_16x", Icon16x16));
+			Set("ClassIcon.ControlRigSequence", new IMAGE_BRUSH("Slate/ControlRigSequence_16x", Icon16x16));
+			Set("ClassIcon.ControlRigBlueprint", new IMAGE_BRUSH("Slate/ControlRigBlueprint_16x", Icon16x16));
 		}
 
 		// Edit mode styles
 		{
-			Set("ControlRigEditMode", new IMAGE_BRUSH("ControlRigMode_40x", Icon40x40));
-			Set("ControlRigEditMode.Small", new IMAGE_BRUSH("ControlRigMode_40x", Icon20x20));
+			Set("ControlRigEditMode", new IMAGE_BRUSH("Slate/ControlRigMode_40x", Icon40x40));
+			Set("ControlRigEditMode.Small", new IMAGE_BRUSH("Slate/ControlRigMode_40x", Icon20x20));
 		}
 
 		// Sequencer styles
 		{
-			Set("ControlRig.ExportAnimSequence", new IMAGE_BRUSH("ExportAnimSequence_24x", Icon24x24));
-			Set("ControlRig.ExportAnimSequence.Small", new IMAGE_BRUSH("ExportAnimSequence_24x", Icon24x24));
-			Set("ControlRig.ReExportAnimSequence", new IMAGE_BRUSH("ExportAnimSequence_24x", Icon24x24));
-			Set("ControlRig.ReExportAnimSequence.Small", new IMAGE_BRUSH("ExportAnimSequence_24x", Icon24x24));
-			Set("ControlRig.ImportFromRigSequence", new IMAGE_BRUSH("ReImportRigSequence_16x", Icon16x16));
-			Set("ControlRig.ImportFromRigSequence.Small", new IMAGE_BRUSH("ReImportRigSequence_16x", Icon16x16));
-			Set("ControlRig.ReImportFromRigSequence", new IMAGE_BRUSH("ReImportRigSequence_16x", Icon16x16));
-			Set("ControlRig.ReImportFromRigSequence.Small", new IMAGE_BRUSH("ReImportRigSequence_16x", Icon16x16));
+			Set("ControlRig.ExportAnimSequence", new IMAGE_BRUSH("Slate/ExportAnimSequence_24x", Icon24x24));
+			Set("ControlRig.ExportAnimSequence.Small", new IMAGE_BRUSH("Slate/ExportAnimSequence_24x", Icon24x24));
+			Set("ControlRig.ReExportAnimSequence", new IMAGE_BRUSH("Slate/ExportAnimSequence_24x", Icon24x24));
+			Set("ControlRig.ReExportAnimSequence.Small", new IMAGE_BRUSH("Slate/ExportAnimSequence_24x", Icon24x24));
+			Set("ControlRig.ImportFromRigSequence", new IMAGE_BRUSH("Slate/ReImportRigSequence_16x", Icon16x16));
+			Set("ControlRig.ImportFromRigSequence.Small", new IMAGE_BRUSH("Slate/ReImportRigSequence_16x", Icon16x16));
+			Set("ControlRig.ReImportFromRigSequence", new IMAGE_BRUSH("Slate/ReImportRigSequence_16x", Icon16x16));
+			Set("ControlRig.ReImportFromRigSequence.Small", new IMAGE_BRUSH("Slate/ReImportRigSequence_16x", Icon16x16));
 		}
 
 		// Control Rig Editor styles
 		{
-			Set("ControlRig.TabIcon", new IMAGE_BRUSH("ControlRigTab_16x", Icon16x16));
-			Set("ControlRig.RigUnit", new IMAGE_BRUSH("ControlRigUnit_16x", Icon16x16));
+			Set("ControlRig.TabIcon", new IMAGE_BRUSH("Slate/ControlRigTab_16x", Icon16x16));
+			Set("ControlRig.RigUnit", new IMAGE_BRUSH("Slate/ControlRigUnit_16x", Icon16x16));
 
 			// icons for control units
-			Set("ControlRig.ControlUnitOn", new IMAGE_BRUSH(TEXT("ControlUnit_On"), Icon32x32));
-			Set("ControlRig.ControlUnitOff", new IMAGE_BRUSH(TEXT("ControlUnit_Off"), Icon32x32));
+			Set("ControlRig.ControlUnitOn", new IMAGE_BRUSH("Slate/ControlUnit_On", Icon32x32));
+			Set("ControlRig.ControlUnitOff", new IMAGE_BRUSH("Slate/ControlUnit_Off", Icon32x32));
 
-			Set("ControlRig.ExecuteGraph", new IMAGE_BRUSH("ExecuteGraph", Icon40x40));
-			Set("ControlRig.ExecuteGraph.Small", new IMAGE_BRUSH("ExecuteGraph", Icon20x20));
+			Set("ControlRig.ExecuteGraph", new IMAGE_BRUSH("Slate/ExecuteGraph", Icon40x40));
+			Set("ControlRig.ExecuteGraph.Small", new IMAGE_BRUSH("Slate/ExecuteGraph", Icon20x20));
 		}
 
 		// Graph styles
 		{
-			Set("ControlRig.Node.PinTree.Arrow_Collapsed_Left", new IMAGE_BRUSH("TreeArrow_Collapsed_Left", Icon10x10, DefaultForeground));
-			Set("ControlRig.Node.PinTree.Arrow_Collapsed_Hovered_Left", new IMAGE_BRUSH("TreeArrow_Collapsed_Hovered_Left", Icon10x10, DefaultForeground));
+			Set("ControlRig.Node.PinTree.Arrow_Collapsed_Left", new IMAGE_BRUSH("Slate/TreeArrow_Collapsed_Left", Icon10x10, DefaultForeground));
+			Set("ControlRig.Node.PinTree.Arrow_Collapsed_Hovered_Left", new IMAGE_BRUSH("Slate/TreeArrow_Collapsed_Hovered_Left", Icon10x10, DefaultForeground));
 
-			Set("ControlRig.Node.PinTree.Arrow_Expanded_Left", new IMAGE_BRUSH("TreeArrow_Expanded_Left", Icon10x10, DefaultForeground));
-			Set("ControlRig.Node.PinTree.Arrow_Expanded_Hovered_Left", new IMAGE_BRUSH("TreeArrow_Expanded_Hovered_Left", Icon10x10, DefaultForeground));
+			Set("ControlRig.Node.PinTree.Arrow_Expanded_Left", new IMAGE_BRUSH("Slate/TreeArrow_Expanded_Left", Icon10x10, DefaultForeground));
+			Set("ControlRig.Node.PinTree.Arrow_Expanded_Hovered_Left", new IMAGE_BRUSH("Slate/TreeArrow_Expanded_Hovered_Left", Icon10x10, DefaultForeground));
 
-			Set("ControlRig.Node.PinTree.Arrow_Collapsed_Right", new IMAGE_BRUSH("TreeArrow_Collapsed_Right", Icon10x10, DefaultForeground));
-			Set("ControlRig.Node.PinTree.Arrow_Collapsed_Hovered_Right", new IMAGE_BRUSH("TreeArrow_Collapsed_Hovered_Right", Icon10x10, DefaultForeground));
+			Set("ControlRig.Node.PinTree.Arrow_Collapsed_Right", new IMAGE_BRUSH("Slate/TreeArrow_Collapsed_Right", Icon10x10, DefaultForeground));
+			Set("ControlRig.Node.PinTree.Arrow_Collapsed_Hovered_Right", new IMAGE_BRUSH("Slate/TreeArrow_Collapsed_Hovered_Right", Icon10x10, DefaultForeground));
 
-			Set("ControlRig.Node.PinTree.Arrow_Expanded_Right", new IMAGE_BRUSH("TreeArrow_Expanded_Right", Icon10x10, DefaultForeground));
-			Set("ControlRig.Node.PinTree.Arrow_Expanded_Hovered_Right", new IMAGE_BRUSH("TreeArrow_Expanded_Hovered_Right", Icon10x10, DefaultForeground));
+			Set("ControlRig.Node.PinTree.Arrow_Expanded_Right", new IMAGE_BRUSH("Slate/TreeArrow_Expanded_Right", Icon10x10, DefaultForeground));
+			Set("ControlRig.Node.PinTree.Arrow_Expanded_Hovered_Right", new IMAGE_BRUSH("Slate/TreeArrow_Expanded_Hovered_Right", Icon10x10, DefaultForeground));
+		}
+
+		// Tree styles
+		{
+			Set("ControlRig.Tree.Bone", new IMAGE_BRUSH("Slate/Bone_16x", Icon16x16));
+			Set("ControlRig.Tree.Control", new IMAGE_BRUSH("Slate/Control_16x", Icon16x16));
+			Set("ControlRig.Tree.Space", new IMAGE_BRUSH("Slate/Space_16x", Icon16x16));
 		}
 
 		// Font?
