Copying //UE4/Dev-Anim to Dev-Main (//UE4/Dev-Main) @ 8782600

diff --git a/Engine/Plugins/Experimental/ControlRig/Source/ControlRigEditor/Private/ControlRigElementDetails.cpp b/Engine/Plugins/Experimental/ControlRig/Source/ControlRigEditor/Private/ControlRigElementDetails.cpp
new file mode 100644
--- /dev/null
+++ b/Engine/Plugins/Experimental/ControlRig/Source/ControlRigEditor/Private/ControlRigElementDetails.cpp
@@ -0,0 +1,473 @@
+// Copyright 1998-2019 Epic Games, Inc. All Rights Reserved.
+
+#include "ControlRigElementDetails.h"
+#include "Widgets/SWidget.h"
+#include "DetailLayoutBuilder.h"
+#include "DetailCategoryBuilder.h"
+#include "DetailWidgetRow.h"
+#include "Widgets/Text/STextBlock.h"
+#include "Widgets/Input/SEditableTextBox.h"
+#include "Widgets/Input/SVectorInputBox.h"
+#include "Widgets/Input/SRotatorInputBox.h"
+#include "Editor/SControlRigGizmoNameList.h"
+
+#define LOCTEXT_NAMESPACE "ControlRigElementDetails"
+
+void FRigElementDetails::CustomizeDetails(IDetailLayoutBuilder& DetailBuilder)
+{
+	DetailBuilder.HideProperty(TEXT("Index"), FRigElement::StaticStruct());
+	DetailBuilder.HideProperty(TEXT("Name"), FRigElement::StaticStruct());
+
+	TArray<TSharedPtr<FStructOnScope>> StructsBeingCustomized;
+	DetailBuilder.GetStructsBeingCustomized(StructsBeingCustomized);
+	for (TSharedPtr<FStructOnScope> StructBeingCustomized : StructsBeingCustomized)
+	{
+		if (UPackage* Package = StructBeingCustomized->GetPackage())
+		{
+			TArray<UObject*> SubObjects;
+			Package->GetDefaultSubobjects(SubObjects);
+
+			for (UObject* SubObject : SubObjects)
+			{
+				if (UControlRig* Rig = Cast<UControlRig>(SubObject))
+				{
+					BlueprintBeingCustomized = Cast<UControlRigBlueprint>(Rig->GetClass()->ClassGeneratedBy);
+					if(BlueprintBeingCustomized)
+					{
+						ContainerBeingCustomized = &BlueprintBeingCustomized->HierarchyContainer;
+						break;
+					}
+				}
+			}
+
+			if (ContainerBeingCustomized)
+			{
+				ElementKeyBeingCustomized = ((const FRigElement*)StructBeingCustomized->GetStructMemory())->GetElementKey();
+				break;
+			}
+		}
+	}
+}
+
+void FRigElementDetails::SetName(const FText& InNewText, ETextCommit::Type InCommitType)
+{
+	if (FRigHierarchyContainer* Hierarchy = GetHierarchy())
+	{
+		switch (ElementKeyBeingCustomized.Type)
+		{
+			case ERigElementType::Bone:
+			{
+				ElementKeyBeingCustomized.Name = Hierarchy->BoneHierarchy.Rename(ElementKeyBeingCustomized.Name, *InNewText.ToString());
+				break;
+			}
+			case ERigElementType::Control:
+			{
+				ElementKeyBeingCustomized.Name = Hierarchy->ControlHierarchy.Rename(ElementKeyBeingCustomized.Name, *InNewText.ToString());
+				break;
+			}
+			case ERigElementType::Space:
+			{
+				ElementKeyBeingCustomized.Name = Hierarchy->SpaceHierarchy.Rename(ElementKeyBeingCustomized.Name, *InNewText.ToString());
+				break;
+			}
+			case ERigElementType::Curve:
+			{
+				ElementKeyBeingCustomized.Name = Hierarchy->CurveContainer.Rename(ElementKeyBeingCustomized.Name, *InNewText.ToString());
+				break;
+			}
+			default:
+			{
+				ensure(false);
+				break;
+			}
+		}
+	}
+}
+
+float FRigElementDetails::GetTransformComponent(const FTransform& InTransform, ERigElementDetailsTransformComponent InComponent)
+{
+	switch (InComponent)
+	{
+		case ERigElementDetailsTransformComponent::TranslationX:
+		{
+			return InTransform.GetTranslation().X;
+		}
+		case ERigElementDetailsTransformComponent::TranslationY:
+		{
+			return InTransform.GetTranslation().Y;
+		}
+		case ERigElementDetailsTransformComponent::TranslationZ:
+		{
+			return InTransform.GetTranslation().Z;
+		}
+		case ERigElementDetailsTransformComponent::RotationRoll:
+		{
+			return InTransform.GetRotation().Rotator().Roll;
+		}
+		case ERigElementDetailsTransformComponent::RotationPitch:
+		{
+			return InTransform.GetRotation().Rotator().Pitch;
+		}
+		case ERigElementDetailsTransformComponent::RotationYaw:
+		{
+			return InTransform.GetRotation().Rotator().Yaw;
+		}
+		case ERigElementDetailsTransformComponent::ScaleX:
+		{
+			return InTransform.GetScale3D().X;
+		}
+		case ERigElementDetailsTransformComponent::ScaleY:
+		{
+			return InTransform.GetScale3D().Y;
+		}
+		case ERigElementDetailsTransformComponent::ScaleZ:
+		{
+			return InTransform.GetScale3D().Z;
+		}
+	}
+	return 0.f;
+}
+
+void FRigElementDetails::SetTransformComponent(FTransform& OutTransform, ERigElementDetailsTransformComponent InComponent, float InNewValue)
+{
+	switch (InComponent)
+	{
+		case ERigElementDetailsTransformComponent::TranslationX:
+		{
+			FVector Translation = OutTransform.GetTranslation();
+			Translation.X = InNewValue;
+			OutTransform.SetTranslation(Translation);
+			break;
+		}
+		case ERigElementDetailsTransformComponent::TranslationY:
+		{
+			FVector Translation = OutTransform.GetTranslation();
+			Translation.Y = InNewValue;
+			OutTransform.SetTranslation(Translation);
+			break;
+		}
+		case ERigElementDetailsTransformComponent::TranslationZ:
+		{
+			FVector Translation = OutTransform.GetTranslation();
+			Translation.Z = InNewValue;
+			OutTransform.SetTranslation(Translation);
+			break;
+		}
+		case ERigElementDetailsTransformComponent::RotationRoll:
+		{
+			FRotator Rotator = OutTransform.Rotator();
+			Rotator.Roll = InNewValue;
+			OutTransform.SetRotation(FQuat(Rotator));
+			break;
+		}
+		case ERigElementDetailsTransformComponent::RotationPitch:
+		{
+			FRotator Rotator = OutTransform.Rotator();
+			Rotator.Pitch = InNewValue;
+			OutTransform.SetRotation(FQuat(Rotator));
+			break;
+		}
+		case ERigElementDetailsTransformComponent::RotationYaw:
+		{
+			FRotator Rotator = OutTransform.Rotator();
+			Rotator.Yaw = InNewValue;
+			OutTransform.SetRotation(FQuat(Rotator));
+			break;
+		}
+		case ERigElementDetailsTransformComponent::ScaleX:
+		{
+			FVector Scale = OutTransform.GetScale3D();
+			Scale.X = InNewValue;
+			OutTransform.SetScale3D(Scale);
+			break;
+		}
+		case ERigElementDetailsTransformComponent::ScaleY:
+		{
+			FVector Scale = OutTransform.GetScale3D();
+			Scale.Y = InNewValue;
+			OutTransform.SetScale3D(Scale);
+			break;
+		}
+		case ERigElementDetailsTransformComponent::ScaleZ:
+		{
+			FVector Scale = OutTransform.GetScale3D();
+			Scale.Z = InNewValue;
+			OutTransform.SetScale3D(Scale);
+			break;
+		}
+	}
+}
+
+void FRigBoneDetails::CustomizeDetails( IDetailLayoutBuilder& DetailBuilder )
+{
+	FRigElementDetails::CustomizeDetails(DetailBuilder);
+	DetailBuilder.HideProperty(TEXT("ParentName"));
+
+	IDetailCategoryBuilder& Category = DetailBuilder.EditCategory(TEXT("FRigElement"), LOCTEXT("BoneCategory", "Bone"));
+	Category.InitiallyCollapsed(false);
+
+	Category.AddCustomRow(FText::FromString(TEXT("Name")))
+	.NameContent()
+	[
+		SNew(STextBlock)
+		.Text(FText::FromString(TEXT("Name")))
+		.Font(IDetailLayoutBuilder::GetDetailFont())
+	]
+	.ValueContent()
+	[
+		SNew(SEditableTextBox)
+		.Font(IDetailLayoutBuilder::GetDetailFont())
+		.Text(this, &FRigElementDetails::GetName)
+		.OnTextCommitted(this, &FRigElementDetails::SetName)
+	];
+}
+
+void FRigControlDetails::CustomizeDetails(IDetailLayoutBuilder& DetailBuilder)
+{
+	FRigElementDetails::CustomizeDetails(DetailBuilder);
+
+	DetailBuilder.HideProperty(TEXT("ControlType"));
+	DetailBuilder.HideProperty(TEXT("ParentName"));
+	DetailBuilder.HideProperty(TEXT("SpaceName"));
+	DetailBuilder.HideProperty(TEXT("InitialValue"));
+	DetailBuilder.HideProperty(TEXT("Value"));
+	DetailBuilder.HideProperty(TEXT("GizmoName"));
+
+	GizmoNameList.Reset();
+	if (BlueprintBeingCustomized)
+	{
+		if (!BlueprintBeingCustomized->GizmoLibrary.IsValid())
+		{
+			BlueprintBeingCustomized->GizmoLibrary.LoadSynchronous();
+		}
+		if (BlueprintBeingCustomized->GizmoLibrary.IsValid())
+		{
+			GizmoNameList.Add(MakeShared<FString>(BlueprintBeingCustomized->GizmoLibrary->DefaultGizmo.GizmoName.ToString()));
+			for (const FControlRigGizmoDefinition& Gizmo : BlueprintBeingCustomized->GizmoLibrary->Gizmos)
+			{
+				GizmoNameList.Add(MakeShared<FString>(Gizmo.GizmoName.ToString()));
+			}
+		}
+	}
+
+	IDetailCategoryBuilder& Category = DetailBuilder.EditCategory(TEXT("FRigElement"), LOCTEXT("ControlCategory", "Control"));
+	Category.InitiallyCollapsed(false);
+
+	Category.AddCustomRow(FText::FromString(TEXT("Name")))
+	.NameContent()
+	[
+		SNew(STextBlock)
+		.Text(FText::FromString(TEXT("Name")))
+		.Font(IDetailLayoutBuilder::GetDetailFont())
+	]
+	.ValueContent()
+	[
+		SNew(SEditableTextBox)
+		.Font(IDetailLayoutBuilder::GetDetailFont())
+		.Text(this, &FRigElementDetails::GetName)
+		.OnTextCommitted(this, &FRigElementDetails::SetName)
+	];
+
+	if (ContainerBeingCustomized == nullptr || !ElementKeyBeingCustomized)
+	{
+		return;
+	}
+
+	FRigControl& Control = ContainerBeingCustomized->ControlHierarchy[ElementKeyBeingCustomized.Name];
+	switch (Control.ControlType)
+	{
+		case ERigControlType::Transform:
+		{
+			TArray<FString> Labels = { TEXT("Initial"), TEXT("Current") };
+			for(FString Label : Labels)
+			{
+				bool bIsInitial = Label == TEXT("Initial");
+				Category.AddCustomRow(FText::FromString(Label))
+				.NameContent()
+				.VAlign(VAlign_Top)
+				[
+					SNew(STextBlock)
+					.Text(FText::FromString(Label))
+					.Font(IDetailLayoutBuilder::GetDetailFont())
+				]
+				.ValueContent()
+				.MinDesiredWidth(125.0f * 3.0f) // copied from FComponentTransformDetails
+				.MaxDesiredWidth(125.0f * 3.0f)
+				[
+					SNew(SVerticalBox)
+					+ SVerticalBox::Slot()
+					[
+						SNew(SVectorInputBox)
+						.AllowSpin(true)
+						.bColorAxisLabels(true)
+						.Font(IDetailLayoutBuilder::GetDetailFont())
+						.X(this, &FRigControlDetails::GetComponentValue, bIsInitial, ERigElementDetailsTransformComponent::TranslationX)
+						.Y(this, &FRigControlDetails::GetComponentValue, bIsInitial, ERigElementDetailsTransformComponent::TranslationY)
+						.Z(this, &FRigControlDetails::GetComponentValue, bIsInitial, ERigElementDetailsTransformComponent::TranslationZ)
+						.OnXCommitted(this, &FRigControlDetails::SetComponentValue, bIsInitial, ERigElementDetailsTransformComponent::TranslationX)
+						.OnYCommitted(this, &FRigControlDetails::SetComponentValue, bIsInitial, ERigElementDetailsTransformComponent::TranslationY)
+						.OnZCommitted(this, &FRigControlDetails::SetComponentValue, bIsInitial, ERigElementDetailsTransformComponent::TranslationZ)
+					]
+					+ SVerticalBox::Slot()
+					[
+						SNew(SRotatorInputBox)
+						.AllowSpin(true)
+						.bColorAxisLabels(true)
+						.Font(IDetailLayoutBuilder::GetDetailFont())
+						.Roll(this, &FRigControlDetails::GetComponentValue, bIsInitial, ERigElementDetailsTransformComponent::RotationRoll)
+						.Pitch(this, &FRigControlDetails::GetComponentValue, bIsInitial, ERigElementDetailsTransformComponent::RotationPitch)
+						.Yaw(this, &FRigControlDetails::GetComponentValue, bIsInitial, ERigElementDetailsTransformComponent::RotationYaw)
+						.OnRollCommitted(this, &FRigControlDetails::SetComponentValue, bIsInitial, ERigElementDetailsTransformComponent::RotationRoll)
+						.OnPitchCommitted(this, &FRigControlDetails::SetComponentValue, bIsInitial, ERigElementDetailsTransformComponent::RotationPitch)
+						.OnYawCommitted(this, &FRigControlDetails::SetComponentValue, bIsInitial, ERigElementDetailsTransformComponent::RotationYaw)
+					]
+					+ SVerticalBox::Slot()
+					[
+						SNew(SVectorInputBox)
+						.AllowSpin(true)
+						.bColorAxisLabels(true)
+						.Font(IDetailLayoutBuilder::GetDetailFont())
+						.X(this, &FRigControlDetails::GetComponentValue, bIsInitial, ERigElementDetailsTransformComponent::ScaleX)
+						.Y(this, &FRigControlDetails::GetComponentValue, bIsInitial, ERigElementDetailsTransformComponent::ScaleY)
+						.Z(this, &FRigControlDetails::GetComponentValue, bIsInitial, ERigElementDetailsTransformComponent::ScaleZ)
+						.OnXCommitted(this, &FRigControlDetails::SetComponentValue, bIsInitial, ERigElementDetailsTransformComponent::ScaleX)
+						.OnYCommitted(this, &FRigControlDetails::SetComponentValue, bIsInitial, ERigElementDetailsTransformComponent::ScaleY)
+						.OnZCommitted(this, &FRigControlDetails::SetComponentValue, bIsInitial, ERigElementDetailsTransformComponent::ScaleZ)
+					]
+				];
+			}
+			break;
+		}
+		default:
+		{
+			ensure(false);
+			break;
+		}
+	}
+
+	Category.AddCustomRow(FText::FromString(TEXT("GizmoName")))
+	.NameContent()
+	[
+		SNew(STextBlock)
+		.Text(FText::FromString(TEXT("Gizmo Name")))
+		.Font(IDetailLayoutBuilder::GetDetailFont())
+	]
+	.ValueContent()
+	[
+		SNew(SControlRigGizmoNameList, &Control, BlueprintBeingCustomized)
+		.OnGetNameListContent(this, &FRigControlDetails::GetGizmoNameList)
+	];
+}
+
+TOptional<float> FRigControlDetails::GetComponentValue(bool bInitial, ERigElementDetailsTransformComponent Component) const
+{
+	if (ContainerBeingCustomized != nullptr && ElementKeyBeingCustomized)
+	{
+		int32 Index = ContainerBeingCustomized->ControlHierarchy.GetIndex(ElementKeyBeingCustomized.Name);
+		if (Index != INDEX_NONE)
+		{
+			FRigControl Control = ContainerBeingCustomized->ControlHierarchy[ElementKeyBeingCustomized.Name];
+			if (!bInitial)
+			{
+				if (UControlRig* DebuggedRig = Cast<UControlRig>(BlueprintBeingCustomized->GetObjectBeingDebugged()))
+				{
+					Control = DebuggedRig->GetControlHierarchy()[Index];
+				}
+			}
+
+			switch (Control.ControlType)
+			{
+				case ERigControlType::Transform:
+				{
+					FTransform Transform = bInitial ? Control.InitialValue.Get<FTransform>() : Control.Value.Get<FTransform>();
+					return GetTransformComponent(Transform, Component);
+				}
+				default:
+				{
+					ensure(false);
+					break;
+				}
+			}
+		}
+	}
+	return 0.f;
+}
+
+void FRigControlDetails::SetComponentValue(float InNewValue, ETextCommit::Type InCommitType, bool bInitial, ERigElementDetailsTransformComponent Component)
+{
+	if (ContainerBeingCustomized != nullptr && ElementKeyBeingCustomized)
+	{
+		int32 Index = ContainerBeingCustomized->ControlHierarchy.GetIndex(ElementKeyBeingCustomized.Name);
+		if (Index != INDEX_NONE)
+		{
+			const FRigControl& Control = ContainerBeingCustomized->ControlHierarchy[ElementKeyBeingCustomized.Name];
+			FRigControlValue Value = bInitial ? Control.InitialValue : Control.Value;
+
+			switch (Control.ControlType)
+			{
+				case ERigControlType::Transform:
+				{
+					FTransform Transform = Value.Get<FTransform>();
+					SetTransformComponent(Transform, Component, InNewValue);
+					Value.Set<FTransform>(Transform);
+					break;
+				}
+				default:
+				{
+					ensure(false);
+					break;
+				}
+			}
+
+			if (bInitial)
+			{
+				ContainerBeingCustomized->ControlHierarchy.SetInitialValue(ElementKeyBeingCustomized.Name, Value);
+				if (UControlRig* DebuggedRig = Cast<UControlRig>(BlueprintBeingCustomized->GetObjectBeingDebugged()))
+				{
+					((FRigControlHierarchy*)&DebuggedRig->GetControlHierarchy())->SetInitialValue(ElementKeyBeingCustomized.Name, Value);
+				}
+			}
+			else
+			{
+				ContainerBeingCustomized->ControlHierarchy.SetValue(ElementKeyBeingCustomized.Name, Value);
+				if (UControlRig* DebuggedRig = Cast<UControlRig>(BlueprintBeingCustomized->GetObjectBeingDebugged()))
+				{
+					((FRigControlHierarchy*)&DebuggedRig->GetControlHierarchy())->SetValue(ElementKeyBeingCustomized.Name, Value);
+				}
+			}
+		}
+	}
+}
+
+const TArray<TSharedPtr<FString>>& FRigControlDetails::GetGizmoNameList() const
+{
+	return GizmoNameList;
+}
+
+void FRigSpaceDetails::CustomizeDetails(IDetailLayoutBuilder& DetailBuilder)
+{
+	FRigElementDetails::CustomizeDetails(DetailBuilder);
+
+	DetailBuilder.HideProperty(TEXT("SpaceType"));
+	DetailBuilder.HideProperty(TEXT("ParentName"));
+
+	IDetailCategoryBuilder& Category = DetailBuilder.EditCategory(TEXT("FRigElement"), LOCTEXT("SpaceCategory", "Space"));
+	Category.InitiallyCollapsed(false);
+
+	Category.AddCustomRow(FText::FromString(TEXT("Name")))
+	.NameContent()
+	[
+		SNew(STextBlock)
+		.Text(FText::FromString(TEXT("Name")))
+		.Font(IDetailLayoutBuilder::GetDetailFont())
+	]
+	.ValueContent()
+	[
+		SNew(SEditableTextBox)
+		.Font(IDetailLayoutBuilder::GetDetailFont())
+		.Text(this, &FRigElementDetails::GetName)
+		.OnTextCommitted(this, &FRigElementDetails::SetName)
+	];
+}
+
+#undef LOCTEXT_NAMESPACE
diff --git a/Engine/Plugins/Experimental/ControlRig/Source/ControlRigEditor/Private/ControlRigElementDetails.h b/Engine/Plugins/Experimental/ControlRig/Source/ControlRigEditor/Private/ControlRigElementDetails.h
new file mode 100644
--- /dev/null
+++ b/Engine/Plugins/Experimental/ControlRig/Source/ControlRigEditor/Private/ControlRigElementDetails.h
@@ -0,0 +1,98 @@
+// Copyright 1998-2019 Epic Games, Inc. All Rights Reserved.
+
+#pragma once
+
+#include "CoreMinimal.h"
+#include "IDetailCustomization.h"
+#include "Rigs/RigHierarchyDefines.h"
+#include "Rigs/RigHierarchyContainer.h"
+#include "ControlRig.h"
+#include "ControlRigBlueprint.h"
+
+class IPropertyHandle;
+
+UENUM()
+enum class ERigElementDetailsTransformComponent : uint8
+{
+	TranslationX,
+	TranslationY,
+	TranslationZ,
+	RotationRoll,
+	RotationPitch,
+	RotationYaw,
+	ScaleX,
+	ScaleY,
+	ScaleZ
+};
+
+class FRigElementDetails : public IDetailCustomization
+{
+public:
+
+	/** IDetailCustomization interface */
+	virtual void CustomizeDetails(IDetailLayoutBuilder& DetailBuilder) override;
+
+	FRigElementKey GetElementKey() const { return ElementKeyBeingCustomized; }
+	FRigHierarchyContainer* GetHierarchy() const { return ContainerBeingCustomized; }
+	FText GetName() const { return FText::FromName(GetElementKey().Name); }
+	void SetName(const FText& InNewText, ETextCommit::Type InCommitType);
+
+	static float GetTransformComponent(const FTransform& InTransform, ERigElementDetailsTransformComponent InComponent);
+	static void SetTransformComponent(FTransform& OutTransform, ERigElementDetailsTransformComponent InComponent, float InNewValue);
+
+protected:
+
+	FRigElementKey ElementKeyBeingCustomized;
+	UControlRigBlueprint* BlueprintBeingCustomized;
+	FRigHierarchyContainer* ContainerBeingCustomized;
+};
+
+class FRigBoneDetails : public FRigElementDetails
+{
+public:
+
+	/** Makes a new instance of this detail layout class for a specific detail view requesting it */
+	static TSharedRef<IDetailCustomization> MakeInstance()
+	{
+		return MakeShareable(new FRigBoneDetails);
+	}
+
+	/** IDetailCustomization interface */
+	virtual void CustomizeDetails( IDetailLayoutBuilder& DetailBuilder ) override;
+};
+
+class FRigControlDetails : public FRigElementDetails
+{
+public:
+
+	/** Makes a new instance of this detail layout class for a specific detail view requesting it */
+	static TSharedRef<IDetailCustomization> MakeInstance()
+	{
+		return MakeShareable(new FRigControlDetails);
+	}
+
+	/** IDetailCustomization interface */
+	virtual void CustomizeDetails(IDetailLayoutBuilder& DetailBuilder) override;
+
+	TOptional<float> GetComponentValue(bool bInitial, ERigElementDetailsTransformComponent Component) const;
+	void SetComponentValue(float InNewValue, ETextCommit::Type InCommitType, bool bInitial, ERigElementDetailsTransformComponent Component);
+
+	const TArray<TSharedPtr<FString>>& GetGizmoNameList() const;
+
+private:
+	TArray<TSharedPtr<FString>> GizmoNameList;
+};
+
+class FRigSpaceDetails : public FRigElementDetails
+{
+public:
+
+	/** Makes a new instance of this detail layout class for a specific detail view requesting it */
+	static TSharedRef<IDetailCustomization> MakeInstance()
+	{
+		return MakeShareable(new FRigSpaceDetails);
+	}
+
+	/** IDetailCustomization interface */
+	virtual void CustomizeDetails(IDetailLayoutBuilder& DetailBuilder) override;
+};
diff --git a/Engine/Plugins/Experimental/ControlRig/Source/ControlRigEditor/Private/ControlRigGizmoLibraryActions.cpp b/Engine/Plugins/Experimental/ControlRig/Source/ControlRigEditor/Private/ControlRigGizmoLibraryActions.cpp
new file mode 100644
--- /dev/null
+++ b/Engine/Plugins/Experimental/ControlRig/Source/ControlRigEditor/Private/ControlRigGizmoLibraryActions.cpp
@@ -0,0 +1,12 @@
+// Copyright 1998-2019 Epic Games, Inc. All Rights Reserved.
+
+#include "ControlRigGizmoLibraryActions.h"
+
+#define LOCTEXT_NAMESPACE "AssetTypeActions"
+
+void FControlRigGizmoLibraryActions::GetActions( const TArray<UObject*>& InObjects, FMenuBuilder& MenuBuilder )
+{
+	//FAssetTypeActions_AnimationAsset::GetActions(InObjects, MenuBuilder);
+}
+
+#undef LOCTEXT_NAMESPACE
diff --git a/Engine/Plugins/Experimental/ControlRig/Source/ControlRigEditor/Private/ControlRigGizmoLibraryActions.h b/Engine/Plugins/Experimental/ControlRig/Source/ControlRigEditor/Private/ControlRigGizmoLibraryActions.h
new file mode 100644
--- /dev/null
+++ b/Engine/Plugins/Experimental/ControlRig/Source/ControlRigEditor/Private/ControlRigGizmoLibraryActions.h
@@ -0,0 +1,25 @@
+// Copyright 1998-2019 Epic Games, Inc. All Rights Reserved.
+
+#pragma once
+
+#include "CoreMinimal.h"
+#include "AssetTypeActions_Base.h"
+#include "ControlRigGizmoLibrary.h"
+
+class FMenuBuilder;
+
+class FControlRigGizmoLibraryActions : public FAssetTypeActions_Base
+{
+public:
+	// IAssetTypeActions Implementation
+	virtual bool HasActions(const TArray<UObject*>& InObjects) const override
+	{
+		return true;
+	}
+	virtual void GetActions(const TArray<UObject*>& InObjects, FMenuBuilder& MenuBuilder) override;
+	virtual FText GetName() const override { return NSLOCTEXT("AssetTypeActions", "AssetTypeActions_ControlRigGizmoLibrary", "Gizmo Library"); }
+	virtual FColor GetTypeColor() const override { return FColor(100,100,255); }
+	virtual UClass* GetSupportedClass() const override { return UControlRigGizmoLibrary::StaticClass(); }
+	virtual bool CanFilter() override { return true; }
+	virtual uint32 GetCategories() override { return EAssetTypeCategories::Animation; }
+};
diff --git a/Engine/Plugins/Experimental/ControlRig/Source/ControlRigEditor/Private/ControlRigGizmoLibraryFactory.cpp b/Engine/Plugins/Experimental/ControlRig/Source/ControlRigEditor/Private/ControlRigGizmoLibraryFactory.cpp
new file mode 100644
--- /dev/null
+++ b/Engine/Plugins/Experimental/ControlRig/Source/ControlRigEditor/Private/ControlRigGizmoLibraryFactory.cpp
@@ -0,0 +1,37 @@
+// Copyright 1998-2019 Epic Games, Inc. All Rights Reserved.
+
+#include "ControlRigGizmoLibraryFactory.h"
+#include "AssetTypeCategories.h"
+#include "Engine/StaticMesh.h"
+
+#define LOCTEXT_NAMESPACE "ControlRigGizmoLibraryFactory"
+
+UControlRigGizmoLibraryFactory::UControlRigGizmoLibraryFactory()
+{
+	bCreateNew = true;
+	bEditAfterNew = true;
+	SupportedClass = UControlRigGizmoLibrary::StaticClass();
+}
+
+UObject* UControlRigGizmoLibraryFactory::FactoryCreateNew(UClass* Class, UObject* InParent, FName Name, EObjectFlags Flags, UObject* Context, FFeedbackContext* Warn)
+{
+	UControlRigGizmoLibrary* GizmoLibrary = NewObject<UControlRigGizmoLibrary>(InParent, Name, Flags);
+
+	GizmoLibrary->DefaultGizmo.StaticMesh = LoadObject<UStaticMesh>(nullptr, TEXT("/ControlRig/Controls/ControlRig_Sphere_solid.ControlRig_Sphere_solid"));
+	GizmoLibrary->DefaultMaterial = LoadObject<UMaterial>(nullptr, TEXT("/ControlRig/Controls/ControlRigGizmoMaterial.ControlRigGizmoMaterial"));
+	GizmoLibrary->MaterialColorParameter = TEXT("Color");
+
+	return GizmoLibrary;
+}
+
+FText UControlRigGizmoLibraryFactory::GetDisplayName() const
+{
+	return LOCTEXT("ControlRigGizmoLibraryFactoryName", "Control Rig Gizmo Library");
+}
+
+uint32 UControlRigGizmoLibraryFactory::GetMenuCategories() const
+{
+	return EAssetTypeCategories::Animation;
+}
+
+#undef LOCTEXT_NAMESPACE
diff --git a/Engine/Plugins/Experimental/ControlRig/Source/ControlRigEditor/Private/ControlRigHierarchyCommands.cpp b/Engine/Plugins/Experimental/ControlRig/Source/ControlRigEditor/Private/ControlRigHierarchyCommands.cpp
--- a/Engine/Plugins/Experimental/ControlRig/Source/ControlRigEditor/Private/ControlRigHierarchyCommands.cpp
+++ b/Engine/Plugins/Experimental/ControlRig/Source/ControlRigEditor/Private/ControlRigHierarchyCommands.cpp
@@ -6,10 +6,22 @@
 
 void FControlRigHierarchyCommands::RegisterCommands()
 {
-	UI_COMMAND(AddItem, "New", "Add new item at the origin (0, 0, 0) to the hierarchy.", EUserInterfaceActionType::Button, FInputChord());
-	UI_COMMAND(DuplicateItem, "Duplicate", "Duplicate the selected items in the hierarchy.", EUserInterfaceActionType::Button, FInputChord());
+	UI_COMMAND(AddBoneItem, "New Bone", "Add new bone at the origin (0, 0, 0) to the hierarchy.", EUserInterfaceActionType::Button, FInputChord());
+	UI_COMMAND(AddControlItem, "New Control", "Add new control at the origin (0, 0, 0) to the hierarchy.", EUserInterfaceActionType::Button, FInputChord(EKeys::N, EModifierKey::Control));
+	UI_COMMAND(AddSpaceItem, "New Space", "Add new space at the origin (0, 0, 0) to the hierarchy.", EUserInterfaceActionType::Button, FInputChord());
+	UI_COMMAND(DuplicateItem, "Duplicate", "Duplicate the selected items in the hierarchy.", EUserInterfaceActionType::Button, FInputChord(EKeys::W, EModifierKey::Control));
 	UI_COMMAND(DeleteItem, "Delete", "Deletes the selected items from the hierarchy.", EUserInterfaceActionType::Button, FInputChord(EKeys::Delete));
 	UI_COMMAND(RenameItem, "Rename", "Rename the selected item.", EUserInterfaceActionType::Button, FInputChord(EKeys::F2));
+	UI_COMMAND(CopyItems, "Copy", "Copy the selected items.", EUserInterfaceActionType::Button, FInputChord(EKeys::C, EModifierKey::Control));
+	UI_COMMAND(PasteItems, "Paste", "Paste the selected items.", EUserInterfaceActionType::Button, FInputChord(EKeys::V, EModifierKey::Control));
+	UI_COMMAND(PasteLocalTransforms, "Paste Local Xfo", "Paste the local transforms.", EUserInterfaceActionType::Button, FInputChord());
+	UI_COMMAND(PasteGlobalTransforms, "Paste Global Xfo", "Paste the global transforms.", EUserInterfaceActionType::Button, FInputChord(EKeys::V, EModifierKey::Control | EModifierKey::Shift));
+	UI_COMMAND(SetInitialTransformFromClosestBone, "Set Initial Transform from Closest Bone", "Find the Closest Bone to Initial Transform of the Selected Bones", EUserInterfaceActionType::Button, FInputChord());
+	UI_COMMAND(SetInitialTransformFromCurrentTransform, "Set Initial Transform from Current", "Save the Current Transform To Initial Transform of the Selected Bones", EUserInterfaceActionType::Button, FInputChord());
+	UI_COMMAND(ResetTransform, "Reset Transform", "Reset the Transform", EUserInterfaceActionType::Button, FInputChord(EKeys::G, EModifierKey::Control));
+	UI_COMMAND(ResetInitialTransform, "Reset Initial Transform", "Reset the Initial Transform to the Identity", EUserInterfaceActionType::Button, FInputChord());
+	UI_COMMAND(ResetSpace, "Reset Space", "Resets or injects a Space below the Control", EUserInterfaceActionType::Button, FInputChord());
+	UI_COMMAND(FrameSelection, "Frame Selection", "Expands and frames the selection in the tree", EUserInterfaceActionType::Button, FInputChord(EKeys::F));
 }
 
 #undef LOCTEXT_NAMESPACE
diff --git a/Engine/Plugins/Experimental/ControlRig/Source/ControlRigEditor/Private/ControlRigHierarchyCommands.h b/Engine/Plugins/Experimental/ControlRig/Source/ControlRigEditor/Private/ControlRigHierarchyCommands.h
--- a/Engine/Plugins/Experimental/ControlRig/Source/ControlRigEditor/Private/ControlRigHierarchyCommands.h
+++ b/Engine/Plugins/Experimental/ControlRig/Source/ControlRigEditor/Private/ControlRigHierarchyCommands.h
@@ -19,7 +19,13 @@ public:
 	{}
 	
 	/** Add Item at origin */
-	TSharedPtr< FUICommandInfo > AddItem;
+	TSharedPtr< FUICommandInfo > AddBoneItem;
+
+	/** Add Item at origin */
+	TSharedPtr< FUICommandInfo > AddControlItem;
+
+	/** Add Item at origin */
+	TSharedPtr< FUICommandInfo > AddSpaceItem;
 
 	/** Duplicate currently selected items */
 	TSharedPtr< FUICommandInfo > DuplicateItem;
@@ -30,6 +36,36 @@ public:
 	/** Rename selected item */
 	TSharedPtr< FUICommandInfo > RenameItem;
 
+	/** Copy the selected items. */
+	TSharedPtr< FUICommandInfo > CopyItems;
+
+	/** Paste the selected items. */
+	TSharedPtr< FUICommandInfo > PasteItems;
+
+	/** Paste Local Xfo", "Paste the local transforms. */
+	TSharedPtr< FUICommandInfo > PasteLocalTransforms;
+
+	/** Paste Global Xfo", "Paste the global transforms. */
+	TSharedPtr< FUICommandInfo > PasteGlobalTransforms;
+
+	/* Reset transform */
+	TSharedPtr<FUICommandInfo> ResetTransform;
+
+	/* Reset initial transform */
+	TSharedPtr<FUICommandInfo> ResetInitialTransform;
+
+	/* Reset space */
+	TSharedPtr<FUICommandInfo> ResetSpace;
+
+	/* Set initial transform from current */
+	TSharedPtr<FUICommandInfo> SetInitialTransformFromCurrentTransform;
+
+	/* set initial transform from closest bone */
+	TSharedPtr<FUICommandInfo> SetInitialTransformFromClosestBone;
+
+	/* frames the selection in the tree */
+	TSharedPtr<FUICommandInfo> FrameSelection;
+
 	/**
 	 * Initialize commands
 	 */
diff --git a/Engine/Plugins/Experimental/ControlRig/Source/ControlRigEditor/Private/EditMode/ControlRigEditMode.cpp b/Engine/Plugins/Experimental/ControlRig/Source/ControlRigEditor/Private/EditMode/ControlRigEditMode.cpp
--- a/Engine/Plugins/Experimental/ControlRig/Source/ControlRigEditor/Private/EditMode/ControlRigEditMode.cpp
+++ b/Engine/Plugins/Experimental/ControlRig/Source/ControlRigEditor/Private/EditMode/ControlRigEditMode.cpp
@@ -24,10 +24,7 @@
 #include "ISequencerModule.h"
 #include "Toolkits/AssetEditorToolkit.h"
 #include "ControlRigEditorModule.h"
-#include "ControlUnitProxy.h"
 #include "Constraint.h"
-#include "Units/Control/RigUnit_Control.h"
-#include "ControlRigControl.h"
 #include "EngineUtils.h"
 #include "ControlRigBlueprintGeneratedClass.h"
 #include "IControlRigObjectBinding.h"
@@ -35,6 +32,13 @@
 #include "Drawing/ControlRigDrawInterface.h"
 #include "ControlRigBlueprint.h"
 #include "ControlRigController.h"
+#include "ControlRigGizmoActor.h"
+#include "DefaultControlRigManipulationLayer.h"
+#include "Framework/MultiBox/MultiBoxBuilder.h"
+#include "SEditorViewport.h"
+
+#include "Framework/Notifications/NotificationManager.h"
+#include "Widgets/Notifications/SNotificationList.h"
 
 FName FControlRigEditMode::ModeName("EditMode.ControlRig");
 
@@ -55,12 +59,17 @@ enum class ETransformComponent
 FControlRigEditMode::FControlRigEditMode()
 	: bIsTransacting(false)
 	, bManipulatorMadeChange(false)
-	, bSelectedBone(false)
 	, bSelecting(false)
-	, bSelectingByPath(false)
 	, PivotTransform(FTransform::Identity)
+	, bRecreateManipulationLayerRequired(false)
+	, ManipulationLayer(nullptr)
+	, CurrentViewportClient(nullptr)
 {
 	Settings = NewObject<UControlRigEditModeSettings>(GetTransientPackage(), TEXT("Settings"));
+	
+	// @todo: thinking of removing this, but for now leaving it. 
+	// it indicates, execution is stopped, and you can modify default bone transform or controls if set to be true
+	bEnableRigElementDefaultPoseEditing = false;
 
 	CommandBindings = MakeShareable(new FUICommandList);
 	BindCommands();
@@ -105,11 +114,27 @@ void FControlRigEditMode::SetSequencer(TSharedPtr<ISequencer> InSequencer)
 	}
 }
 
-void FControlRigEditMode::SetObjects(const TWeakObjectPtr<>& InSelectedObject, const FGuid& InObjectBinding)
+void FControlRigEditMode::SetObjects(const TWeakObjectPtr<>& InSelectedObject, const FGuid& InObjectBinding, UObject* BindingObject)
 {
-	WeakControlRig = Cast<UControlRig>(InSelectedObject.Get());
+	WeakControlRigEditing = Cast<UControlRig>(InSelectedObject.Get());
 	ControlRigGuid = InObjectBinding;
 
+	// if we get binding object, set it to control rig binding object
+	if (BindingObject)
+	{
+		if (UControlRig* ControlRig = WeakControlRigEditing.Get())
+		{
+			if (TSharedPtr<IControlRigObjectBinding> ObjectBinding = ControlRig->GetObjectBinding())
+			{
+				if (ObjectBinding->GetBoundObject() == nullptr)
+				{
+					ObjectBinding->BindToObject(BindingObject);
+				}
+			}
+		}
+
+	}
+
 	SetObjects_Internal();
 }
 
@@ -120,9 +145,11 @@ void FControlRigEditMode::SetObjects_Internal()
 	{
 		SelectedObjects.Add(Settings);
 	}
-	if(WeakControlRig.IsValid())
+
+	if(WeakControlRigEditing.IsValid())
 	{
-		SelectedObjects.Add(WeakControlRig);
+		SelectedObjects.Add(WeakControlRigEditing);
+		WeakControlRigEditing.Get()->DrawInterface = &DrawInterface;
 	}
 
 	if(UsesToolkits())
@@ -130,7 +157,9 @@ void FControlRigEditMode::SetObjects_Internal()
 		StaticCastSharedPtr<SControlRigEditModeTools>(Toolkit->GetInlineContent())->SetDetailsObjects(SelectedObjects);
 	}
 
-	RefreshControlProxies();
+	// create default manipulation layer
+	RecreateManipulationLayer();
+	HandleSelectionChanged();
 }
 
 void FControlRigEditMode::HandleBindToActor(AActor* InActor, bool bFocus)
@@ -167,20 +196,15 @@ void FControlRigEditMode::HandleBindToActor(AActor* InActor, bool bFocus)
 				// now notify the sequence (will rebind when it re-evaluates
 				Sequencer->NotifyMovieSceneDataChanged(EMovieSceneDataChangeType::TrackValueChangedRefreshImmediately);
 
-				// Force a rig evaluation here to make sure our manipulators are up to date
- 				if (UControlRig* ControlRig = WeakControlRig.Get())
- 				{
- 					ControlRig->PreEvaluate_GameThread();
- 					ControlRig->Evaluate_AnyThread();
- 					ControlRig->PostEvaluate_GameThread();
- 				}
+				// we tick first
+				ManipulationLayer->TickManipulatableObjects(0.f);
 
 				// Now re-display our objects in the details panel (they may have changed)
 				if (MovieScene->GetSpawnableCount() > 0)
 				{
 					FGuid SpawnableGuid = MovieScene->GetSpawnable(0).GetGuid();
 					TWeakObjectPtr<> BoundObject = Sequencer->FindSpawnedObjectOrTemplate(SpawnableGuid);
-					SetObjects(BoundObject, SpawnableGuid);
+					SetObjects(BoundObject, SpawnableGuid, nullptr);
 				}
 			}
 
@@ -239,7 +263,16 @@ void FControlRigEditMode::Exit()
 {
 	if (bIsTransacting)
 	{
-		GEditor->EndTransaction();
+		if (ManipulationLayer)
+		{
+			ManipulationLayer->EndTransaction();
+		}
+
+		if (GEditor)
+		{
+			GEditor->EndTransaction();
+		}
+
 		bIsTransacting = false;
 		bManipulatorMadeChange = false;
 	}
@@ -249,212 +282,159 @@ void FControlRigEditMode::Exit()
 		FToolkitManager::Get().CloseToolkit(Toolkit.ToSharedRef());
 	}
 
-	// Call parent implementation
-	FEdMode::Exit();
-}
-
-static bool ModeSupportedByTransformFilter(const FTransformFilter& InFilter, FWidget::EWidgetMode InMode)
-{
-	if(InMode == FWidget::WM_Translate && InFilter.TranslationFilter.IsValid())
+	if (ManipulationLayer)
 	{
-		return true;
+		ManipulationLayer->DestroyLayer();
+		ManipulationLayer = nullptr;
 	}
 
-	if(InMode == FWidget::WM_Rotate && InFilter.RotationFilter.IsValid())
-	{
-		return true;
-	}
-
-	if(InMode == FWidget::WM_Scale && InFilter.ScaleFilter.IsValid())
-	{
-		return true;
-	}
-
-	return false;
+	// Call parent implementation
+	FEdMode::Exit();
 }
 
 void FControlRigEditMode::Tick(FEditorViewportClient* ViewportClient, float DeltaTime)
 {
 	FEdMode::Tick(ViewportClient, DeltaTime);
 
-	if(UControlRig* ControlRig = WeakControlRig.Get())
+	ViewportClient->Invalidate();
+
+	RecalcPivotTransform();
+
+	if (bRecreateManipulationLayerRequired)
 	{
-		if (bSelectedBone)
+		RecreateManipulationLayer();
+
+		for (const FRigElementKey& SelectedKey : SelectedRigElements)
 		{
-			if(AreControlsSelected())
+			if (SelectedKey.Type == ERigElementType::Control)
 			{
-				// cycle the widget mode if it is not supported on this selection
-				FWidget::EWidgetMode CurrentMode = GetModeManager()->GetWidgetMode();
-				bool bModeSupported = false;
-				for (FControlUnitProxy& UnitProxy : ControlUnits)
-				{
-					if(UnitProxy.IsSelected())
-					{
-						if(FRigUnit_Control* ControlUnit = GetRigUnit(UnitProxy, ControlRig))
-						{
-							if(ModeSupportedByTransformFilter(ControlUnit->Filter, CurrentMode))
-							{
-								bModeSupported = true;
-							}
-						}
-					}
-				}
-
-				if (!bModeSupported)
+				AControlRigGizmoActor* GizmoActor = GetGizmoFromControlName(SelectedKey.Name);
+				if (GizmoActor)
 				{
-					GetModeManager()->CycleWidgetMode();
+					GizmoActor->SetSelected(true);
 				}
 			}
 		}
+		HandleSelectionChanged();
+		bRecreateManipulationLayerRequired = false;
+	}
 
-		ViewportClient->Invalidate();
-		bSelectedBone = false;
+	// If we have detached from sequencer, unbind the settings UI
+	if (!WeakSequencer.IsValid() && Settings->Sequence != nullptr)
+	{
+		Settings->Sequence = nullptr;
+		RefreshObjects();
+	}
+}
 
-		// If we have detached from sequencer, unbind the settings UI
-		if (!WeakSequencer.IsValid() && Settings->Sequence != nullptr)
-		{
-			Settings->Sequence = nullptr;
-			RefreshObjects();
-		}
+void FControlRigEditMode::Render(const FSceneView* View, FViewport* Viewport, FPrimitiveDrawInterface* PDI)
+{
+	bool bRender = true;
+	if (WeakSequencer.IsValid())
+	{
+		bRender = WeakSequencer.Pin()->GetPlaybackStatus() != EMovieScenePlayerStatus::Playing || Settings->bShowManipulatorsDuringPlayback;
+	}
 
-		FTransform ComponentTransform = FTransform::Identity;
-		if (TSharedPtr<IControlRigObjectBinding> ObjectBinding = ControlRig->GetObjectBinding())
-		{
-			USceneComponent* Component = Cast<USceneComponent>(ObjectBinding->GetBoundObject());
-			ComponentTransform = Component ? Component->GetComponentTransform() : FTransform::Identity;
-		}
+	if (!WeakControlRigEditing.IsValid())
+	{
+		DrawInterface.DrawInstructions.Reset();
+		return;
+	}
+
+	// Force off manipulators if hide flag is set
+	if (Settings->bHideManipulators)
+	{
+		bRender = false;
+	}
 
-		// Update controls from rig
-		for(const FControlUnitProxy& UnitProxy : ControlUnits)
+	FTransform ComponentTransform = (ManipulationLayer)? ManipulationLayer->GetSkeletalMeshComponentTransform() : FTransform::Identity;
+	if (bRender)
+	{
+		if (Settings->bDisplayHierarchy)
 		{
-			if(UnitProxy.Control)
+			UControlRig* ControlRig = WeakControlRigEditing.Get();
+			// each base hierarchy Bone
+			const FRigBoneHierarchy& BaseHierarchy = ControlRig->GetBoneHierarchy();
+			for (int32 BoneIndex = 0; BoneIndex < BaseHierarchy.Num(); ++BoneIndex)
 			{
-				UScriptStruct* Struct = nullptr;
-				if(FRigUnit_Control* ControlUnit = GetRigUnit(UnitProxy, ControlRig, &Struct))
+				const FRigBone& CurrentBone = BaseHierarchy[BoneIndex];
+				const FTransform Transform = BaseHierarchy.GetGlobalTransform(BoneIndex);
+
+				if (CurrentBone.ParentIndex != INDEX_NONE)
 				{
-					UnitProxy.Control->SetTransform(ControlUnit->GetResultantTransform() * ComponentTransform);
-					UnitProxy.Control->TickControl(DeltaTime, *ControlUnit, Struct);
+					const FTransform ParentTransform = BaseHierarchy.GetGlobalTransform(CurrentBone.ParentIndex);
+
+					PDI->DrawLine(ComponentTransform.TransformPosition(Transform.GetLocation()),ComponentTransform.TransformPosition(ParentTransform.GetLocation()), FLinearColor::White, SDPG_Foreground);
 				}
+
+				PDI->DrawPoint(ComponentTransform.TransformPosition(Transform.GetLocation()), FLinearColor::White, 5.0f, SDPG_Foreground);
 			}
 		}
 
-		// update the pivot transform of our selected objects (they could be animating)
-		RecalcPivotTransform();
-
-		// Tick controls
-		for(FControlUnitProxy& UnitProxy : ControlUnits)
+		if (Settings->bDisplayAxesOnSelection && Settings->AxisScale > SMALL_NUMBER)
 		{
-			if(UnitProxy.Control)
+			UControlRig* ControlRig = WeakControlRigEditing.Get();
+			const FRigHierarchyContainer* Hierarchy = ControlRig->GetHierarchy();
+			const float Scale = Settings->AxisScale;
+
+			for (const FRigElementKey& SelectedElement : SelectedRigElements)
 			{
-				UnitProxy.Control->Tick(DeltaTime);
+				FTransform ElementTransform = Hierarchy->GetGlobalTransform(SelectedElement);
+				ElementTransform = ElementTransform * ComponentTransform;
+
+				PDI->DrawLine(ElementTransform.GetTranslation(), ElementTransform.TransformPosition(FVector(Scale, 0.f, 0.f)), FLinearColor::Red, SDPG_Foreground);
+				PDI->DrawLine(ElementTransform.GetTranslation(), ElementTransform.TransformPosition(FVector(0.f, Scale, 0.f)), FLinearColor::Green, SDPG_Foreground);
+				PDI->DrawLine(ElementTransform.GetTranslation(), ElementTransform.TransformPosition(FVector(0.f, 0.f, Scale)), FLinearColor::Blue, SDPG_Foreground);
 			}
 		}
 	}
-}
 
-void FControlRigEditMode::Render(const FSceneView* View, FViewport* Viewport, FPrimitiveDrawInterface* PDI)
-{
-	if(UControlRig* ControlRig = WeakControlRig.Get())
-	{
-		bool bRender = true;
-		if (WeakSequencer.IsValid())
-		{
-			bRender = WeakSequencer.Pin()->GetPlaybackStatus() != EMovieScenePlayerStatus::Playing || Settings->bShowManipulatorsDuringPlayback;
-		}
 
-		// Force off manipulators if hide flag is set
-		if (Settings->bHideManipulators)
+	for (const FControlRigDrawInterface::FDrawIntruction& Instruction : DrawInterface.DrawInstructions)
+	{
+		if (Instruction.Positions.Num() == 0)
 		{
-			bRender = false;
+			continue;
 		}
-
-		if (bRender)
+		switch (Instruction.DrawType)
 		{
-			if (Settings->bDisplayHierarchy)
+			case FControlRigDrawInterface::EDrawType_Point:
 			{
-				FTransform ComponentTransform = FTransform::Identity;
-				if (TSharedPtr<IControlRigObjectBinding> ObjectBinding = ControlRig->GetObjectBinding())
-				{
-					USceneComponent* Component = Cast<USceneComponent>(ObjectBinding->GetBoundObject());
-					ComponentTransform = Component ? Component->GetComponentTransform() : FTransform::Identity;
-				}
-
-				// each base hierarchy Bone
-				const FRigHierarchy& BaseHierarchy = ControlRig->GetBaseHierarchy();
-				const TArray<FRigBone>& Bones = BaseHierarchy.GetBones();
-				for (int32 BoneIndex = 0; BoneIndex < Bones.Num(); ++BoneIndex)
+				for (const FVector& Point : Instruction.Positions)
 				{
-					const FRigBone& CurrentBone = Bones[BoneIndex];
-					const FTransform Transform = BaseHierarchy.GetGlobalTransform(BoneIndex);
-
-					if (CurrentBone.ParentIndex != INDEX_NONE)
-					{
-						const FTransform ParentTransform = BaseHierarchy.GetGlobalTransform(CurrentBone.ParentIndex);
-
-						PDI->DrawLine(Transform.GetLocation(), ParentTransform.GetLocation(), FLinearColor::White, SDPG_Foreground);
-					}
-
-					PDI->DrawPoint(Transform.GetLocation(), FLinearColor::White, 5.0f, SDPG_Foreground);
+					PDI->DrawPoint(ComponentTransform.TransformPosition(Point), Instruction.Color, Instruction.Thickness, SDPG_Foreground);
 				}
+				break;
 			}
-
-			// @TODO: debug drawing per rig Bone (like details customizations) for this
-
-// 			if (Settings->bDisplayTrajectories)
-// 			{
-// 				TrajectoryCache.RenderTrajectories(ComponentTransform, PDI);
-// 			}
-		}
-
-		FControlRigDrawInterface* DrawInterface = ControlRig->DrawInterface;
-		if (DrawInterface)
-		{
-			for (const FControlRigDrawInterface::FDrawIntruction& Instruction : DrawInterface->DrawInstructions)
+			case FControlRigDrawInterface::EDrawType_Lines:
 			{
-				if (Instruction.Positions.Num() == 0)
+				const TArray<FVector>& Points = Instruction.Positions;
+				for (int32 PointIndex = 0; PointIndex < Points.Num() - 1; PointIndex += 2)
 				{
-					continue;
+					PDI->DrawLine(ComponentTransform.TransformPosition(Points[PointIndex]), ComponentTransform.TransformPosition(Points[PointIndex+1]), Instruction.Color, SDPG_Foreground, Instruction.Thickness);
 				}
-				switch (Instruction.DrawType)
+				break;
+			}
+			case FControlRigDrawInterface::EDrawType_LineStrip:
+			{
+				const TArray<FVector>& Points = Instruction.Positions;
+				for (int32 PointIndex = 0; PointIndex < Points.Num() - 1; PointIndex++)
 				{
-					case FControlRigDrawInterface::EDrawType_Point:
-					{
-						for (const FVector& Point : Instruction.Positions)
-						{
-							PDI->DrawPoint(Point, Instruction.Color, Instruction.Thickness, SDPG_Foreground);
-						}
-						break;
-					}
-					case FControlRigDrawInterface::EDrawType_Lines:
-					{
-						const TArray<FVector>& Points = Instruction.Positions;
-						for (int32 PointIndex = 0; PointIndex < Points.Num() - 1; PointIndex += 2)
-						{
-							PDI->DrawLine(Points[PointIndex], Points[PointIndex+1], Instruction.Color, SDPG_Foreground, Instruction.Thickness);
-						}
-						break;
-					}
-					case FControlRigDrawInterface::EDrawType_LineStrip:
-					{
-						const TArray<FVector>& Points = Instruction.Positions;
-						for (int32 PointIndex = 0; PointIndex < Points.Num() - 1; PointIndex++)
-						{
-							PDI->DrawLine(Points[PointIndex], Points[PointIndex + 1], Instruction.Color, SDPG_Foreground, Instruction.Thickness);
-						}
-						break;
-					}
+					PDI->DrawLine(ComponentTransform.TransformPosition(Points[PointIndex]), ComponentTransform.TransformPosition(Points[PointIndex + 1]), Instruction.Color, SDPG_Foreground, Instruction.Thickness);
 				}
+				break;
 			}
-
-			DrawInterface->DrawInstructions.Reset();
 		}
 	}
+	DrawInterface.DrawInstructions.Reset();
 }
 
 bool FControlRigEditMode::InputKey(FEditorViewportClient* InViewportClient, FViewport* InViewport, FKey InKey, EInputEvent InEvent)
 {
 	if (InEvent != IE_Released)
 	{
+		TGuardValue<FEditorViewportClient*> ViewportGuard(CurrentViewportClient, InViewportClient);
+
 		FModifierKeysState KeyState = FSlateApplication::Get().GetModifierKeys();
 		if (CommandBindings->ProcessCommandBindings(InKey, KeyState, (InEvent == IE_Repeat)))
 		{
@@ -472,24 +452,16 @@ bool FControlRigEditMode::EndTracking(FEditorViewportClient* InViewportClient, F
 		if (bManipulatorMadeChange)
 		{
 			// One final notify of our manipulators to make sure the property is keyed
-			if(UControlRig* ControlRig = WeakControlRig.Get())
+			for (AControlRigGizmoActor* GizmoActor : GizmoActors)
 			{
-				for(FControlUnitProxy& UnitProxy : ControlUnits)
+				if (GizmoActor->IsManipulating())
 				{
-					if(UnitProxy.IsManipulating())
-					{
-						UnitProxy.SetManipulating(false);
-						UnitProxy.NotifyPostEditChangeProperty(ControlRig);
-					}
+					GizmoActor->SetManipulating(false);
 				}
 			}
-
-			if (Settings->bDisplayTrajectories)
-			{
-				TrajectoryCache.ForceRecalc();
-			}
 		}
 
+		ManipulationLayer->EndTransaction();
 		GEditor->EndTransaction();
 		bIsTransacting = false;
 		bManipulatorMadeChange = false;
@@ -507,15 +479,10 @@ bool FControlRigEditMode::StartTracking(FEditorViewportClient* InViewportClient,
 	{
 		GEditor->BeginTransaction(LOCTEXT("MoveControlTransaction", "Move Control"));
 
-		if(UControlRig* ControlRig = WeakControlRig.Get())
+		ManipulationLayer->BeginTransaction();
+		for (AControlRigGizmoActor* GizmoActor : GizmoActors)
 		{
-			ControlRig->SetFlags(RF_Transactional);
-			ControlRig->Modify();
-
-			for(FControlUnitProxy& UnitProxy : ControlUnits)
-			{
-				UnitProxy.SetManipulating(true);
-			}
+			GizmoActor->SetManipulating(true);
 		}
 
 		bIsTransacting = true;
@@ -529,18 +496,15 @@ bool FControlRigEditMode::StartTracking(FEditorViewportClient* InViewportClient,
 
 bool FControlRigEditMode::UsesTransformWidget() const
 {
-	if(UControlRig* ControlRig = WeakControlRig.Get())
+	for (const AControlRigGizmoActor* GizmoActor : GizmoActors)
 	{
-		for (const FControlUnitProxy& UnitProxy : ControlUnits)
+		if (GizmoActor->IsSelected())
 		{
-			if(UnitProxy.IsSelected())
-			{
-				return true;
-			}
+			return true;
 		}
 	}
 
-	if (AreBoneSelectedAndMovable())
+	if (AreRigElementSelectedAndMovable())
 	{
 		return true;
 	}
@@ -550,23 +514,17 @@ bool FControlRigEditMode::UsesTransformWidget() const
 
 bool FControlRigEditMode::UsesTransformWidget(FWidget::EWidgetMode CheckMode) const
 {
-	if(UControlRig* ControlRig = WeakControlRig.Get())
+	for (const AControlRigGizmoActor* GizmoActor : GizmoActors)
 	{
-		for (const FControlUnitProxy& UnitProxy : ControlUnits)
+		if (GizmoActor->IsSelected())
 		{
-			if(UnitProxy.IsSelected())
-			{
-				if(FRigUnit_Control* ControlUnit = GetRigUnit(UnitProxy, ControlRig))
-				{
-					return ModeSupportedByTransformFilter(ControlUnit->Filter, CheckMode);
-				}
-			}
+			return ManipulationLayer->ModeSupportedByGizmoActor(GizmoActor, CheckMode);
 		}
+	}
 
-		if (AreBoneSelectedAndMovable())
-		{
-			return true;
-		}
+	if (AreRigElementSelectedAndMovable())
+	{
+		return true;
 	}
 
 	return FEdMode::UsesTransformWidget(CheckMode);
@@ -574,29 +532,10 @@ bool FControlRigEditMode::UsesTransformWidget(FWidget::EWidgetMode CheckMode) co
 
 FVector FControlRigEditMode::GetWidgetLocation() const
 {
-	if(UControlRig* ControlRig = WeakControlRig.Get())
+	if (AreRigElementSelectedAndMovable())
 	{
-		FTransform ComponentTransform = FTransform::Identity;
-		if (TSharedPtr<IControlRigObjectBinding> ObjectBinding = ControlRig->GetObjectBinding())
-		{
-			USceneComponent* Component = Cast<USkeletalMeshComponent>(ObjectBinding->GetBoundObject());
-			ComponentTransform = Component ? Component->GetComponentTransform() : FTransform::Identity;
-		}
-
-		for (const FControlUnitProxy& UnitProxy : ControlUnits)
-		{
-			if(UnitProxy.IsSelected())
-			{
-				return ComponentTransform.TransformPosition(PivotTransform.GetLocation());
-			}
-		}
-
-		// @todo: we only supports the first ast one for now
-		// later we support multi select
-		if (AreBoneSelectedAndMovable())
-		{
-			return ComponentTransform.TransformPosition(OnGetBoneTransformDelegate.Execute(SelectedBones[0], false).GetLocation());
-		}
+		FTransform ComponentTransform = ManipulationLayer->GetSkeletalMeshComponentTransform();
+		return ComponentTransform.TransformPosition(PivotTransform.GetLocation());
 	}
 
 	return FEdMode::GetWidgetLocation();
@@ -604,29 +543,10 @@ FVector FControlRigEditMode::GetWidgetLocation() const
 
 bool FControlRigEditMode::GetCustomDrawingCoordinateSystem(FMatrix& OutMatrix, void* InData)
 {
-	if(UControlRig* ControlRig = WeakControlRig.Get())
+	if (AreRigElementSelectedAndMovable())
 	{
-		for (const FControlUnitProxy& UnitProxy : ControlUnits)
-		{
-			if(UnitProxy.IsSelected())
-			{
-				OutMatrix = PivotTransform.ToMatrixNoScale().RemoveTranslation();
-				return true;
-			}
-		}
-
-		if (AreBoneSelectedAndMovable())
-		{
-			FTransform ComponentTransform = FTransform::Identity;
-			if (TSharedPtr<IControlRigObjectBinding> ObjectBinding = ControlRig->GetObjectBinding())
-			{
-				USceneComponent* Component = Cast<USkeletalMeshComponent>(ObjectBinding->GetBoundObject());
-				ComponentTransform = Component ? Component->GetComponentTransform() : FTransform::Identity;
-			}
-			FTransform BoneTransform = OnGetBoneTransformDelegate.Execute(SelectedBones[0], false)*ComponentTransform;
-			OutMatrix = BoneTransform.ToMatrixWithScale().RemoveTranslation();
-			return true;
-		}
+		OutMatrix = PivotTransform.ToMatrixNoScale().RemoveTranslation();
+		return true;
 	}
 
 	return false;
@@ -641,25 +561,49 @@ bool FControlRigEditMode::HandleClick(FEditorViewportClient* InViewportClient, H
 {
 	if(HActor* ActorHitProxy = HitProxyCast<HActor>(HitProxy))
 	{
-		if(ActorHitProxy->Actor && ActorHitProxy->Actor->IsA<AControlRigControl>())
+		if(ActorHitProxy->Actor)
 		{
-			AControlRigControl* ControlRigControl = Cast<AControlRigControl>(ActorHitProxy->Actor);
-			if (Click.IsShiftDown() || Click.IsControlDown())
+			if (ActorHitProxy->Actor->IsA<AControlRigGizmoActor>())
 			{
-				SetControlSelection(ControlRigControl->GetPropertyPath(), !IsControlSelected(ControlRigControl->GetPropertyPath()));
-			}
-			else
-			{
-				ClearControlSelection();
-				SetControlSelection(ControlRigControl->GetPropertyPath(), true);
+				AControlRigGizmoActor* GizmoActor = CastChecked<AControlRigGizmoActor>(ActorHitProxy->Actor);
+				const FControlData* ControlData = ManipulationLayer->GetControlDataFromGizmo(GizmoActor);
+
+				if (ControlData)
+				{
+					const FName& ControlName = ControlData->ControlName;
+					if (Click.IsShiftDown() || Click.IsControlDown())
+					{
+						SetRigElementSelection(ERigElementType::Control, ControlName, true);
+					}
+					else
+					{
+						ClearRigElementSelection(FRigElementTypeHelper::ToMask(ERigElementType::Control));
+						SetRigElementSelection(ERigElementType::Control, ControlName, true);
+					}
+				}
 			}
+		}
 
-			return true;
+		// for now we show this menu all the time if body is selected
+		// if we want some global menu, we'll have to move this
+		if (Click.GetKey() == EKeys::RightMouseButton)
+		{
+			OpenContextMenu(InViewportClient);
 		}
+
+		return true;
+	}
+
+	// for now we show this menu all the time if body is selected
+	// if we want some global menu, we'll have to move this
+	if (Click.GetKey() == EKeys::RightMouseButton)
+	{
+		OpenContextMenu(InViewportClient);
+		return true;
 	}
 
 	// clear selected controls
-	ClearControlSelection();
+	ClearRigElementSelection(FRigElementTypeHelper::ToMask(ERigElementType::All));
 
 	// If we are animating then swallow clicks so we dont select things other than controls
 	if(WeakSequencer.IsValid() && WeakSequencer.Pin()->GetFocusedMovieSceneSequence()->IsA<UControlRigSequence>())
@@ -670,43 +614,66 @@ bool FControlRigEditMode::HandleClick(FEditorViewportClient* InViewportClient, H
 	return FEdMode::HandleClick(InViewportClient, HitProxy, Click);
 }
 
-bool FControlRigEditMode::IntersectSelect(bool InSelect, const TFunctionRef<bool(const FControlUnitProxy&, const FTransform&)>& Intersects)
+void FControlRigEditMode::OpenContextMenu(FEditorViewportClient* InViewportClient)
 {
-	if(UControlRig* ControlRig = WeakControlRig.Get())
+	TSharedPtr<FUICommandList> Commands = CommandBindings;
+	if (OnContextMenuCommandsDelegate.IsBound())
+	{
+		Commands = OnContextMenuCommandsDelegate.Execute();
+	}
+
+	if (OnContextMenuDelegate.IsBound())
 	{
-		FTransform ComponentTransform = FTransform::Identity;
-		if (TSharedPtr<IControlRigObjectBinding> ObjectBinding = ControlRig->GetObjectBinding())
+		FMenuBuilder MenuBuilder(true, Commands);
+		OnContextMenuDelegate.Execute(MenuBuilder);
+
+		TSharedPtr<SWidget> MenuWidget = MenuBuilder.MakeWidget();
+		TSharedPtr<SWidget> ParentWidget = InViewportClient->GetEditorViewportWidget();
+
+		if (MenuWidget.IsValid() && ParentWidget.IsValid())
 		{
-			USceneComponent* Component = Cast<USkeletalMeshComponent>(ObjectBinding->GetBoundObject());
-			ComponentTransform = Component ? Component->GetComponentTransform() : FTransform::Identity;
+			const FVector2D MouseCursorLocation = FSlateApplication::Get().GetCursorPos();
+
+			FSlateApplication::Get().PushMenu(
+				ParentWidget.ToSharedRef(),
+				FWidgetPath(),
+				MenuWidget.ToSharedRef(),
+				MouseCursorLocation,
+				FPopupTransitionEffect(FPopupTransitionEffect::ContextMenu)
+			);
 		}
+	}
+}
+
+bool FControlRigEditMode::IntersectSelect(bool InSelect, const TFunctionRef<bool(const AControlRigGizmoActor*, const FTransform&)>& Intersects)
+{
+	FTransform ComponentTransform = ManipulationLayer->GetSkeletalMeshComponentTransform();
 
-		bool bSelected = false;
-		for (const FControlUnitProxy& UnitProxy : ControlUnits)
+	bool bSelected = false;
+	for (AControlRigGizmoActor* GizmoActor: GizmoActors)
+	{
+		const FTransform ControlTransform = GizmoActor->GetGlobalTransform() * ComponentTransform;
+		if (Intersects(GizmoActor, ControlTransform))
 		{
-			if(FRigUnit_Control* ControlUnit = GetRigUnit(UnitProxy, ControlRig))
+			const FControlData* ControlData = ManipulationLayer->GetControlDataFromGizmo(GizmoActor);
+			if (ControlData)
 			{
-				const FTransform ControlTransform = ControlUnit->GetResultantTransform() * ComponentTransform;
-				if (Intersects(UnitProxy, ControlTransform))
-				{
-					SetControlSelection(UnitProxy.PropertyPathString, InSelect);
-					bSelected = true;
-				}
+				SetRigElementSelection(ERigElementType::Control, ControlData->ControlName, InSelect);
+				bSelected = true;
 			}
 		}
-
-		return bSelected;
 	}
-	return false;
+
+	return bSelected;
 }
 
 bool FControlRigEditMode::BoxSelect(FBox& InBox, bool InSelect)
 {
-	bool bIntersects = IntersectSelect(InSelect, [&](const FControlUnitProxy& ControlProxy, const FTransform& Transform)
+	bool bIntersects = IntersectSelect(InSelect, [&](const AControlRigGizmoActor* GizmoActor, const FTransform& Transform)
 	{ 
-		if(ControlProxy.Control != nullptr)
+		if(GizmoActor != nullptr)
 		{
-			FBox Bounds = ControlProxy.Control->GetComponentsBoundingBox(true);
+			FBox Bounds = GizmoActor->GetComponentsBoundingBox(true);
 			Bounds = Bounds.TransformBy(Transform);
 			return InBox.Intersect(Bounds);
 		}
@@ -723,11 +690,11 @@ bool FControlRigEditMode::BoxSelect(FBox& InBox, bool InSelect)
 
 bool FControlRigEditMode::FrustumSelect(const FConvexVolume& InFrustum, FEditorViewportClient* InViewportClient, bool InSelect)
 {
-	bool bIntersects = IntersectSelect(InSelect, [&](const FControlUnitProxy& ControlProxy, const FTransform& Transform) 
+	bool bIntersects = IntersectSelect(InSelect, [&](const AControlRigGizmoActor* GizmoActor, const FTransform& Transform)
 	{
-		if(ControlProxy.Control != nullptr)
+		if(GizmoActor != nullptr)
 		{
-			FBox Bounds = ControlProxy.Control->GetComponentsBoundingBox(true);
+			FBox Bounds = GizmoActor->GetComponentsBoundingBox(true);
 			Bounds = Bounds.TransformBy(Transform);
 			return InFrustum.IntersectBox(Bounds.GetCenter(), Bounds.GetExtent());
 		}
@@ -744,171 +711,108 @@ bool FControlRigEditMode::FrustumSelect(const FConvexVolume& InFrustum, FEditorV
 
 void FControlRigEditMode::SelectNone()
 {
-	ClearControlSelection();
-
-	SelectedBones.Reset();
+	ClearRigElementSelection(FRigElementTypeHelper::ToMask(ERigElementType::All));
 
 	FEdMode::SelectNone();
 }
 
 bool FControlRigEditMode::InputDelta(FEditorViewportClient* InViewportClient, FViewport* InViewport, FVector& InDrag, FRotator& InRot, FVector& InScale)
 {
-	if(UControlRig* ControlRig = WeakControlRig.Get())
-	{
-		FVector Drag = InDrag;
-		FRotator Rot = InRot;
-		FVector Scale = InScale;
+	FVector Drag = InDrag;
+	FRotator Rot = InRot;
+	FVector Scale = InScale;
 
-		const bool bCtrlDown = InViewport->KeyState(EKeys::LeftControl) || InViewport->KeyState(EKeys::RightControl);
-		const bool bShiftDown = InViewport->KeyState(EKeys::LeftShift) || InViewport->KeyState(EKeys::RightShift);
-		const bool bAltDown = InViewport->KeyState(EKeys::LeftAlt) || InViewport->KeyState(EKeys::RightAlt);
-		const bool bMouseButtonDown = InViewport->KeyState(EKeys::LeftMouseButton);
+	const bool bCtrlDown = InViewport->KeyState(EKeys::LeftControl) || InViewport->KeyState(EKeys::RightControl);
+	const bool bShiftDown = InViewport->KeyState(EKeys::LeftShift) || InViewport->KeyState(EKeys::RightShift);
+	const bool bAltDown = InViewport->KeyState(EKeys::LeftAlt) || InViewport->KeyState(EKeys::RightAlt);
+	const bool bMouseButtonDown = InViewport->KeyState(EKeys::LeftMouseButton);
 
-		const FWidget::EWidgetMode WidgetMode = InViewportClient->GetWidgetMode();
-		const EAxisList::Type CurrentAxis = InViewportClient->GetCurrentWidgetAxis();
-		const ECoordSystem CoordSystem = InViewportClient->GetWidgetCoordSystemSpace();
+	const FWidget::EWidgetMode WidgetMode = InViewportClient->GetWidgetMode();
+	const EAxisList::Type CurrentAxis = InViewportClient->GetCurrentWidgetAxis();
+	const ECoordSystem CoordSystem = InViewportClient->GetWidgetCoordSystemSpace();
 
-		if (bIsTransacting && bMouseButtonDown && !bCtrlDown && !bShiftDown && !bAltDown && CurrentAxis != EAxisList::None)
-		{
-			const bool bDoRotation = !Rot.IsZero() && (WidgetMode == FWidget::WM_Rotate || WidgetMode == FWidget::WM_TranslateRotateZ);
-			const bool bDoTranslation = !Drag.IsZero() && (WidgetMode == FWidget::WM_Translate || WidgetMode == FWidget::WM_TranslateRotateZ);
-			const bool bDoScale = !Scale.IsZero() && WidgetMode == FWidget::WM_Scale;
-
-			FTransform ComponentTransform = FTransform::Identity;
-			if (TSharedPtr<IControlRigObjectBinding> ObjectBinding = ControlRig->GetObjectBinding())
-			{
-				USceneComponent* Component = Cast<USkeletalMeshComponent>(ObjectBinding->GetBoundObject());
-				ComponentTransform = Component ? Component->GetComponentTransform() : FTransform::Identity;
-			}
+	if (bIsTransacting && bMouseButtonDown && !bCtrlDown && !bShiftDown && !bAltDown && CurrentAxis != EAxisList::None)
+	{
+		const bool bDoRotation = !Rot.IsZero() && (WidgetMode == FWidget::WM_Rotate || WidgetMode == FWidget::WM_TranslateRotateZ);
+		const bool bDoTranslation = !Drag.IsZero() && (WidgetMode == FWidget::WM_Translate || WidgetMode == FWidget::WM_TranslateRotateZ);
+		const bool bDoScale = !Scale.IsZero() && WidgetMode == FWidget::WM_Scale;
 
-			if (AreControlsSelected())
+		FTransform ComponentTransform = ManipulationLayer->GetSkeletalMeshComponentTransform();
+		if (AreRigElementsSelected(FRigElementTypeHelper::ToMask(ERigElementType::Control)))
+		{
+			for (AControlRigGizmoActor* GizmoActor : GizmoActors)
 			{
-				// manipulator transform is always on actor base - (actor origin being 0)
-				for (FControlUnitProxy& UnitProxy : ControlUnits)
+				if (GizmoActor->IsSelected())
 				{
-					if (UnitProxy.IsSelected())
-					{
-						if (FRigUnit_Control* ControlUnit = GetRigUnit(UnitProxy, ControlRig))
-						{
-							FTransform NewWorldTransform = ControlUnit->GetResultantTransform() * ComponentTransform;
-
-							bool bTransformChanged = false;
-							if (bDoRotation && ControlUnit->Filter.RotationFilter.IsValid())
-							{
-								FQuat CurrentRotation = NewWorldTransform.GetRotation();
-								CurrentRotation = (Rot.Quaternion() * CurrentRotation);
-								NewWorldTransform.SetRotation(CurrentRotation);
-								bTransformChanged = true;
-							}
-
-							if (bDoTranslation && ControlUnit->Filter.TranslationFilter.IsValid())
-							{
-								FVector CurrentLocation = NewWorldTransform.GetLocation();
-								CurrentLocation = CurrentLocation + Drag;
-								NewWorldTransform.SetLocation(CurrentLocation);
-								bTransformChanged = true;
-							}
-
-							if (bDoScale && ControlUnit->Filter.ScaleFilter.IsValid())
-							{
-								FVector CurrentScale = NewWorldTransform.GetScale3D();
-								CurrentScale = CurrentScale + Scale;
-								NewWorldTransform.SetScale3D(CurrentScale);
-								bTransformChanged = true;
-							}
-
-							if (bTransformChanged)
-							{
-								FTransform ResultantTransform = NewWorldTransform.GetRelativeTransform(ComponentTransform);
-
-								UnitProxy.NotifyPreEditChangeProperty(ControlRig);
-
-								ControlUnit->SetResultantTransform(ResultantTransform);
-
-								if (UnitProxy.Control)
-								{
-									UnitProxy.Control->SetTransform(NewWorldTransform);
-								}
-								UnitProxy.NotifyPostEditChangeProperty(ControlRig);
-
-								// Push to CDO if we are not in the level editor
-								if (!IsInLevelEditor())
-								{
-									UClass* Class = ControlRig->GetClass();
-									UControlRig* CDO = Class->GetDefaultObject<UControlRig>();
-									if (FRigUnit_Control* DefaultControlUnit = GetRigUnit(UnitProxy, CDO))
-									{
-										CDO->Modify();
-
-										DefaultControlUnit->SetResultantTransform(ResultantTransform);
-
-										UBlueprint* Blueprint = Cast<UBlueprint>(Class->ClassGeneratedBy);
-										if (Blueprint)
-										{
-											Blueprint->MarkPackageDirty();
-											//FBlueprintEditorUtils::MarkBlueprintAsModified(Blueprint);
-										}
-									}
-								}
-
-								bManipulatorMadeChange = true;
-							}
-						}
-					}
+					// test local vs global
+					ManipulationLayer->MoveGizmo(GizmoActor, bDoTranslation, InDrag, bDoRotation, InRot, bDoScale, InScale, ComponentTransform);
+					bManipulatorMadeChange = true;
 				}
+			}
 
-				RecalcPivotTransform();
+			RecalcPivotTransform();
 
-				return true;
+			if (bManipulatorMadeChange)
+			{
+				ManipulationLayer->TickManipulatableObjects(0.f);
 			}
-			else if (AreBoneSelectedAndMovable())
+			return true;
+		}
+		else if (AreRigElementSelectedAndMovable())
+		{
+			// set Bone transform
+			// that will set initial Bone transform
+			for (int32 Index = 0; Index < SelectedRigElements.Num(); ++Index)
 			{
-				// set Bone transform
-				// that will set initial Bone transform
-				const FName CurrentBone = SelectedBones[0];
-				FTransform NewWorldTransform = OnGetBoneTransformDelegate.Execute(CurrentBone, false) * ComponentTransform;
-				bool bTransformChanged = false;
-				if (bDoRotation)
-				{
-					FQuat CurrentRotation = NewWorldTransform.GetRotation();
-					CurrentRotation = (Rot.Quaternion() * CurrentRotation);
-					NewWorldTransform.SetRotation(CurrentRotation);
-					bTransformChanged = true;
-				}
+				const ERigElementType SelectedRigElementType = SelectedRigElements[Index].Type;
 
-				if (bDoTranslation)
+				if (SelectedRigElementType == ERigElementType::Bone || SelectedRigElementType == ERigElementType::Control)
 				{
-					FVector CurrentLocation = NewWorldTransform.GetLocation();
-					CurrentLocation = CurrentLocation + Drag;
-					NewWorldTransform.SetLocation(CurrentLocation);
-					bTransformChanged = true;
-				}
+					FTransform NewWorldTransform = OnGetRigElementTransformDelegate.Execute(SelectedRigElements[Index], false) * ComponentTransform;
+					bool bTransformChanged = false;
+					if (bDoRotation)
+					{
+						FQuat CurrentRotation = NewWorldTransform.GetRotation();
+						CurrentRotation = (Rot.Quaternion() * CurrentRotation);
+						NewWorldTransform.SetRotation(CurrentRotation);
+						bTransformChanged = true;
+					}
 
-				if (bDoScale)
-				{
-					FVector CurrentScale = NewWorldTransform.GetScale3D();
-					CurrentScale = CurrentScale + Scale;
-					NewWorldTransform.SetScale3D(CurrentScale);
-					bTransformChanged = true;
-				}
+					if (bDoTranslation)
+					{
+						FVector CurrentLocation = NewWorldTransform.GetLocation();
+						CurrentLocation = CurrentLocation + Drag;
+						NewWorldTransform.SetLocation(CurrentLocation);
+						bTransformChanged = true;
+					}
 
-				if (bTransformChanged)
-				{
-					FTransform NewComponentTransform = NewWorldTransform.GetRelativeTransform(ComponentTransform);
-					OnSetBoneTransformDelegate.Execute(CurrentBone, NewComponentTransform);
-				}
+					if (bDoScale)
+					{
+						FVector CurrentScale = NewWorldTransform.GetScale3D();
+						CurrentScale = CurrentScale + Scale;
+						NewWorldTransform.SetScale3D(CurrentScale);
+						bTransformChanged = true;
+					}
 
-				return true;
+					if (bTransformChanged)
+					{
+						FTransform NewComponentTransform = NewWorldTransform.GetRelativeTransform(ComponentTransform);
+						OnSetRigElementTransformDelegate.Execute(SelectedRigElements[Index], NewComponentTransform, false);
+						bManipulatorMadeChange = true;
+					}
+				}
 			}
+
+			// not sure this makes sense @rethink
+			return bManipulatorMadeChange;
 		}
 	}
-
 	return false;
 }
 
 bool FControlRigEditMode::ShouldDrawWidget() const
 {
-	if (AreControlsSelected() || AreBoneSelectedAndMovable())
+	if (AreRigElementSelectedAndMovable())
 	{
 		return true;
 	}
@@ -928,231 +832,105 @@ bool FControlRigEditMode::IsCompatibleWith(FEditorModeID OtherModeID) const
 void FControlRigEditMode::AddReferencedObjects( FReferenceCollector& Collector )
 {
 	Collector.AddReferencedObject(Settings);
+	Collector.AddReferencedObject(ManipulationLayer);
 
-	for(FControlUnitProxy& UnitProxy : ControlUnits)
+	for (AControlRigGizmoActor* GizmoActor : GizmoActors)
 	{
-		Collector.AddReferencedObject(UnitProxy.Control);
+		Collector.AddReferencedObject(GizmoActor);
 	}
 }
 
-void FControlRigEditMode::ClearControlSelection()
+
+void FControlRigEditMode::ClearRigElementSelection(uint32 InTypes)
 {
-	if (!bSelecting)
+	if (!WeakControlRigEditing.IsValid())
 	{
-		TGuardValue<bool> ReentrantGuard(bSelecting, true);
-
-		for(FControlUnitProxy& UnitProxy : ControlUnits)
-		{
-			UnitProxy.SetSelected(false);
-
-			if (UControlRig* ControlRig = WeakControlRig.Get())
-			{
-				if (UControlRigBlueprintGeneratedClass* Class = Cast<UControlRigBlueprintGeneratedClass>(ControlRig->GetClass()))
-				{
-					if (UControlRigBlueprint* RigBlueprint = Cast<UControlRigBlueprint>(Class->ClassGeneratedBy))
-					{
-						if (RigBlueprint)
-						{
-							if (RigBlueprint->ModelController)
-							{
-								RigBlueprint->ModelController->DeselectNode(*UnitProxy.PropertyPath.ToString());
-							}
-						}
-					}
-				}
-			}
-		}
+		return;
+	}
 
-		bSelectedBone = true;
-		HandleSelectionChanged(TArray<FString>());
+	UControlRigBlueprint* Blueprint = Cast<UControlRigBlueprint>(WeakControlRigEditing->GetClass()->ClassGeneratedBy);
+	if (Blueprint)
+	{
+		Blueprint->HierarchyContainer.ClearSelection();
 	}
 }
 
-void FControlRigEditMode::SetControlSelection(const FString& InControlPropertyPath, bool bSelected)
+// internal private function that doesn't use guarding.
+void FControlRigEditMode::SetRigElementSelectionInternal(ERigElementType Type, const FName& InRigElementName, bool bSelected)
 {
-	if (!bSelecting)
+	if (!WeakControlRigEditing.IsValid())
 	{
-		TGuardValue<bool> ReentrantGuard(bSelecting, true);
-
-		TArray<FString> SelectedPropertyPaths;
-		for(FControlUnitProxy& UnitProxy : ControlUnits)
-		{
-			if(InControlPropertyPath == UnitProxy.PropertyPathString || InControlPropertyPath == UnitProxy.TransformPropertyPathString)
-			{
-				UnitProxy.SetSelected(bSelected);
-				SelectedPropertyPaths.Add(UnitProxy.TransformPropertyPathString);
-
-				if (UControlRig* ControlRig = WeakControlRig.Get())
-				{
-					if (UControlRigBlueprintGeneratedClass* Class = Cast<UControlRigBlueprintGeneratedClass>(ControlRig->GetClass()))
-					{
-						if (UControlRigBlueprint* RigBlueprint = Cast<UControlRigBlueprint>(Class->ClassGeneratedBy))
-						{
-							if (RigBlueprint)
-							{
-								if (RigBlueprint->ModelController)
-								{
-									TArray<FName> NodeNames;
-									NodeNames.Add(*UnitProxy.PropertyPath.ToString());
-									RigBlueprint->ModelController->SetSelection(NodeNames);
-								}
-							}
-						}
-					}
-				}
-			}
-		}
-
-		bSelectedBone = true;
+		return;
+	}
 
-		HandleSelectionChanged(SelectedPropertyPaths);
+	UControlRigBlueprint* Blueprint = Cast<UControlRigBlueprint>(WeakControlRigEditing->GetClass()->ClassGeneratedBy);
+	if (Blueprint)
+	{
+		Blueprint->HierarchyContainer.Select(FRigElementKey(InRigElementName, Type), bSelected);
 	}
 }
 
-void FControlRigEditMode::SetControlSelection(const TArray<FString>& InControlPropertyPaths, bool bSelected)
+void FControlRigEditMode::SetRigElementSelection(ERigElementType Type, const FName& InRigElementName, bool bSelected)
 {
 	if (!bSelecting)
 	{
 		TGuardValue<bool> ReentrantGuard(bSelecting, true);
 
-		TArray<FString> SelectedPropertyPaths;
-		for(FControlUnitProxy& UnitProxy : ControlUnits)
-		{
-			for (const FString& ControlPropertyPath : InControlPropertyPaths)
-			{
-				if(ControlPropertyPath == UnitProxy.PropertyPathString || ControlPropertyPath == UnitProxy.TransformPropertyPathString)
-				{
-					UnitProxy.SetSelected(bSelected);
-					SelectedPropertyPaths.Add(UnitProxy.TransformPropertyPathString);
-
-					if (UControlRig* ControlRig = WeakControlRig.Get())
-					{
-						if (UControlRigBlueprintGeneratedClass* Class = Cast<UControlRigBlueprintGeneratedClass>(ControlRig->GetClass()))
-						{
-							if (UControlRigBlueprint* RigBlueprint = Cast<UControlRigBlueprint>(Class->ClassGeneratedBy))
-							{
-								if (RigBlueprint)
-								{
-									if (RigBlueprint->ModelController)
-									{
-										TArray<FName> NodeNames;
-										NodeNames.Add(*UnitProxy.PropertyPath.ToString());
-										RigBlueprint->ModelController->SetSelection(NodeNames);
-									}
-								}
-							}
-						}
-					}
-
-					break;
-				}
-			}
-		}
+		SetRigElementSelectionInternal(Type, InRigElementName, bSelected);
 
-		bSelectedBone = true;
-		HandleSelectionChanged(SelectedPropertyPaths);
+		HandleSelectionChanged();
 	}
 }
 
-bool FControlRigEditMode::IsControlSelected(const FString& InControlPropertyPath) const
+void FControlRigEditMode::SetRigElementSelection(ERigElementType Type, const TArray<FName>& InRigElementNames, bool bSelected)
 {
-	for(const FControlUnitProxy& UnitProxy : ControlUnits)
+	if (!bSelecting)
 	{
-		if(UnitProxy.PropertyPathString == InControlPropertyPath)
+		TGuardValue<bool> ReentrantGuard(bSelecting, true);
+
+		for (const FName& ElementName : InRigElementNames)
 		{
-			return UnitProxy.IsSelected();
+			SetRigElementSelectionInternal(Type, ElementName, bSelected);
 		}
-	}
 
-	return false;
+		HandleSelectionChanged();
+	}
 }
 
-bool FControlRigEditMode::AreControlsSelected() const
+bool FControlRigEditMode::AreRigElementsSelected(uint32 InTypes) const
 {
-	if(UControlRig* ControlRig = WeakControlRig.Get())
+	for (const FRigElementKey& Ele : SelectedRigElements)
 	{
-		for(const FControlUnitProxy& UnitProxy : ControlUnits)
+		if (FRigElementTypeHelper::DoesHave(InTypes, Ele.Type))
 		{
-			if(UnitProxy.IsSelected())
-			{
-				return true;
-			}
+			return true;
 		}
 	}
 
 	return false;
 }
 
-int32 FControlRigEditMode::GetNumSelectedControls() const
+int32 FControlRigEditMode::GetNumSelectedRigElements(uint32 InTypes) const
 {
-	int32 NumSelected = 0;
-	if(UControlRig* ControlRig = WeakControlRig.Get())
+	if (FRigElementTypeHelper::DoesHave(InTypes, ERigElementType::All))
 	{
-		for(const FControlUnitProxy& UnitProxy : ControlUnits)
-		{
-			if(UnitProxy.IsSelected())
-			{
-				NumSelected++;
-			}
-		}
+		return SelectedRigElements.Num();
 	}
-
-	return NumSelected;
-}
-
-TArray<FString> FControlRigEditMode::GetSelectedControls() const
-{
-	TArray<FString> Paths;
-	if (UControlRig* ControlRig = WeakControlRig.Get())
+	else
 	{
-		for (const FControlUnitProxy& UnitProxy : ControlUnits)
+		int32 NumSelected = 0;
+		for (const FRigElementKey& Ele : SelectedRigElements)
 		{
-			if (UnitProxy.IsSelected())
+			if (FRigElementTypeHelper::DoesHave(InTypes, Ele.Type))
 			{
-				Paths.Add(UnitProxy.PropertyPath.ToString());
+				++NumSelected;
 			}
 		}
-	}
-
-	return Paths;
-}
-
-void FControlRigEditMode::SetControlEnabled(const FString& InControlPropertyPath, bool bEnabled)
-{
-	for(FControlUnitProxy& UnitProxy : ControlUnits)
-	{
-		if(UnitProxy.PropertyPathString == InControlPropertyPath)
-		{
-			UnitProxy.SetEnabled(bEnabled);
-		}
-	}
-}
 
-bool FControlRigEditMode::IsControlEnabled(const FString& InControlPropertyPath) const
-{
-	for(const FControlUnitProxy& UnitProxy : ControlUnits)
-	{
-		if(UnitProxy.PropertyPathString == InControlPropertyPath)
-		{
-			return UnitProxy.IsEnabled();
-		}
+		return NumSelected;
 	}
 
-	return false;
-}
-
-FString FControlRigEditMode::GetControlFromPropertyPath(const FString& InPropertyPath) const
-{
-	for (const FControlUnitProxy& UnitProxy : ControlUnits)
-	{
-		if (UnitProxy.PropertyPathString == InPropertyPath)
-		{
-			// the output
-			return UnitProxy.PropertyPath.ToString();
-		}
-	}
-	
-	return TEXT("");
+	return 0;
 }
 
 void FControlRigEditMode::HandleObjectSpawned(FGuid InObjectBinding, UObject* SpawnedObject, IMovieScenePlayer& Player)
@@ -1168,9 +946,9 @@ void FControlRigEditMode::HandleObjectSpawned(FGuid InObjectBinding, UObject* Sp
 			// check if the object is being displayed currently
 			if (ControlRigGuid == InObjectBinding)
 			{
-				if (WeakControlRig != Cast<UControlRig>(SpawnedObject))
+				if (WeakControlRigEditing != Cast<UControlRig>(SpawnedObject))
 				{
-					WeakControlRig = Cast<UControlRig>(SpawnedObject);
+					WeakControlRigEditing = Cast<UControlRig>(SpawnedObject);
 					SetObjects_Internal();
 				}
 				return;
@@ -1179,7 +957,7 @@ void FControlRigEditMode::HandleObjectSpawned(FGuid InObjectBinding, UObject* Sp
 			// We didnt find an existing Guid, so set up our internal cache
 			if (!ControlRigGuid.IsValid())
 			{
-				SetObjects(SpawnedObject, InObjectBinding);
+				SetObjects(SpawnedObject, InObjectBinding, nullptr);
 				if (UControlRig* ControlRig = Cast<UControlRig>(SpawnedObject))
 				{
 					if (TSharedPtr<IControlRigObjectBinding> ObjectBinding = ControlRig->GetObjectBinding())
@@ -1208,7 +986,7 @@ void FControlRigEditMode::RefreshObjects()
 			if (ControlRigGuid.IsValid() && MovieScene->FindSpawnable(ControlRigGuid) == nullptr)
 			{
 				ControlRigGuid.Invalidate();
-				WeakControlRig = nullptr;
+				WeakControlRigEditing = nullptr;
 			}
 
 			SetObjects_Internal();
@@ -1216,87 +994,123 @@ void FControlRigEditMode::RefreshObjects()
 	}
 	else
 	{
-		WeakControlRig = nullptr;
+		WeakControlRigEditing = nullptr;
 		ControlRigGuid.Invalidate();
 
 		SetObjects_Internal();
 	}
 }
 
-void FControlRigEditMode::RecalcPivotTransform()
+void FControlRigEditMode::EnableRigElementEditing(bool bEnabled)
 {
-	int32 NumSelectedControls = GetNumSelectedControls();
+	bEnableRigElementDefaultPoseEditing = bEnabled;
+	RecalcPivotTransform();
+}
 
+void FControlRigEditMode::RecalcPivotTransform()
+{
 	PivotTransform = FTransform::Identity;
 
-	if(UControlRig* ControlRig = WeakControlRig.Get())
+	// @todo: support bones also
+	if(AreRigElementsSelected(FRigElementTypeHelper::ToMask(ERigElementType::Control)))
 	{
-		if(NumSelectedControls > 0)
-		{
-			FTransform LastTransform = FTransform::Identity;
+		FTransform LastTransform = FTransform::Identity;
 
-			// Use average location as pivot location
-			FVector PivotLocation = FVector::ZeroVector;
+		// recalc coord system too
+		FTransform ComponentTransform = ManipulationLayer->GetSkeletalMeshComponentTransform();
 
-			for(const FControlUnitProxy& UnitProxy : ControlUnits)
-			{
-				if(UnitProxy.IsSelected())
-				{
-					if(FRigUnit_Control* ControlUnit = GetRigUnit(UnitProxy, ControlRig))
-					{
-						FTransform ResultantTransform = ControlUnit->GetResultantTransform();
-						PivotLocation += ResultantTransform.GetLocation();
-						LastTransform = ResultantTransform;
-					}
-				}
-			}
-
-			PivotLocation /= (float)NumSelectedControls;
-			PivotTransform.SetLocation(PivotLocation);
+		// Use average location as pivot location
+		FVector PivotLocation = FVector::ZeroVector;
 
-			// recalc coord system too
-			FTransform ComponentTransform = FTransform::Identity;
-			if (TSharedPtr<IControlRigObjectBinding> ObjectBinding = ControlRig->GetObjectBinding())
+		int32 NumSelectedControls = 0;
+		for (const AControlRigGizmoActor* GizmoActor : GizmoActors)
+		{
+			if (GizmoActor->IsSelected())
 			{
-				USceneComponent* Component = Cast<USceneComponent>(ObjectBinding->GetBoundObject());
-				ComponentTransform = Component ? Component->GetComponentTransform() : FTransform::Identity;
+				LastTransform = GizmoActor->GetActorTransform().GetRelativeTransform(ComponentTransform);
+				PivotLocation += LastTransform.GetLocation();
+				++NumSelectedControls;
 			}
+		}
 
-			if (NumSelectedControls == 1)
-			{
-				// A single Bone just uses its own transform
-				FTransform WorldTransform = LastTransform * ComponentTransform;
-				PivotTransform.SetRotation(WorldTransform.GetRotation());
-			}
-			else if (NumSelectedControls > 1)
+		PivotLocation /= (float)NumSelectedControls;
+		PivotTransform.SetLocation(PivotLocation);
+		
+		if (NumSelectedControls == 1)
+		{
+			// A single Bone just uses its own transform
+			FTransform WorldTransform = LastTransform * ComponentTransform;
+			PivotTransform.SetRotation(WorldTransform.GetRotation());
+		}
+		else if (NumSelectedControls > 1)
+		{
+			// If we have more than one Bone selected, use the coordinate space of the component
+			PivotTransform.SetRotation(ComponentTransform.GetRotation());
+		}
+	}
+	else if (AreRigElementSelectedAndMovable())
+	{
+		// recalc coord system too
+		FTransform ComponentTransform = ManipulationLayer->GetSkeletalMeshComponentTransform();
+
+		// Use average location as pivot location
+		FVector PivotLocation = FVector::ZeroVector;
+		int32 NumSelection = 0;
+		FTransform LastTransform = FTransform::Identity;
+		for (int32 Index = 0; Index < SelectedRigElements.Num(); ++Index)
+		{
+			if (SelectedRigElements[Index].Type == ERigElementType::Control || SelectedRigElements[Index].Type == ERigElementType::Bone)
 			{
-				// If we have more than one Bone selected, use the coordinate space of the component
-				PivotTransform.SetRotation(ComponentTransform.GetRotation());
+				LastTransform = OnGetRigElementTransformDelegate.Execute(SelectedRigElements[Index], false);
+				PivotLocation += LastTransform.GetLocation();
+				++NumSelection;
 			}
 		}
+
+		PivotLocation /= (float)NumSelection;
+		PivotTransform.SetLocation(PivotLocation);
+
+		if (NumSelection == 1)
+		{
+			// A single Bone just uses its own transform
+			FTransform WorldTransform = LastTransform * ComponentTransform;
+			PivotTransform.SetRotation(WorldTransform.GetRotation());
+		}
+		else if (NumSelection > 1)
+		{
+			// If we have more than one Bone selected, use the coordinate space of the component
+			PivotTransform.SetRotation(ComponentTransform.GetRotation());
+		}
 	}
 }
 
-void FControlRigEditMode::HandleSelectionChanged(const TArray<FString>& InSelectedPropertyPaths)
+void FControlRigEditMode::HandleSelectionChanged()
 {
 	if (WeakSequencer.IsValid())
 	{
-		if (InSelectedPropertyPaths.Num() > 0)
+		TArray<FString> SelectedNames;
+		for (int32 Index=0; Index<SelectedRigElements.Num(); ++Index)
+		{
+			// if it's control, make sure it's selected
+			if (SelectedRigElements[Index].Type == ERigElementType::Control)
+			{
+				SelectedNames.Add(SelectedRigElements[Index].Name.ToString());
+			}
+		}
+
+		if (SelectedNames.Num() > 0)
 		{
-			WeakSequencer.Pin()->SelectByPropertyPaths(InSelectedPropertyPaths);
+			WeakSequencer.Pin()->SelectByPropertyPaths(SelectedNames);
 		}
 	}
 
-	for(const FControlUnitProxy& UnitProxy : ControlUnits)
+	for (AControlRigGizmoActor* GizmoActor : GizmoActors)
 	{
 		TInlineComponentArray<UPrimitiveComponent*> PrimitiveComponents;
-		if (UnitProxy.Control)
+		GizmoActor->GetComponents(PrimitiveComponents, true);
+		for (UPrimitiveComponent* PrimitiveComponent : PrimitiveComponents)
 		{
-			UnitProxy.Control->GetComponents(PrimitiveComponents, true);
-			for (UPrimitiveComponent* PrimitiveComponent : PrimitiveComponents)
-			{
-				PrimitiveComponent->PushSelectionToProxy();
-			}
+			PrimitiveComponent->PushSelectionToProxy();
 		}
 	}
 
@@ -1307,6 +1121,9 @@ void FControlRigEditMode::HandleSelectionChanged(const TArray<FString>& InSelect
 			WeakSequencer.Pin()->RefreshTree();
 		}
 	}
+
+	// update the pivot transform of our selected objects (they could be animating)
+	RecalcPivotTransform();
 }
 
 void FControlRigEditMode::BindCommands()
@@ -1320,37 +1137,65 @@ void FControlRigEditMode::BindCommands()
 	CommandBindings->MapAction(
 		Commands.ToggleManipulators,
 		FExecuteAction::CreateRaw(this, &FControlRigEditMode::ToggleManipulators));
+}
 
-	CommandBindings->MapAction(
-		Commands.ToggleTrajectories,
-		FExecuteAction::CreateRaw(this, &FControlRigEditMode::ToggleTrajectories));
+bool FControlRigEditMode::IsControlSelected() const
+{
+	static uint32 TypeFlag = (uint32)ERigElementType::Control;
+	return (AreRigElementsSelected(TypeFlag));
 }
 
-void FControlRigEditMode::SetKeysForSelectedManipulators()
+bool FControlRigEditMode::IsControlOrSpaceOrBoneSelected() const
 {
-	if(UControlRig* ControlRig = WeakControlRig.Get())
+	static uint32 TypeFlag = (uint32)ERigElementType::Bone | (uint32)ERigElementType::Control | (uint32)ERigElementType::Space;
+	return (AreRigElementsSelected(TypeFlag));
+}
+
+
+bool FControlRigEditMode::GetRigElementGlobalTransform(const FRigElementKey& InElement, FTransform& OutGlobalTransform) const
+{
+	// if control, go through manipulation layer
+	if (InElement.Type == ERigElementType::Control)
 	{
-		for(const FControlUnitProxy& UnitProxy : ControlUnits)
+		// this code is weird. Need to set this info in manipulation layer
+		AControlRigGizmoActor* GizmoActor = GetGizmoFromControlName(InElement.Name);
+		if (GizmoActor && ensure(GizmoActor->IsSelected()))
 		{
-			if(UnitProxy.IsSelected())
+			if (ManipulationLayer->GetGlobalTransform(GizmoActor, InElement.Name, OutGlobalTransform))
 			{
-				SetKeyForControl(UnitProxy);
+				return true;
 			}
+
+			ensure(false);
+			return false;
 		}
 	}
+	else if (AreRigElementSelectedAndMovable())
+	{
+		// @tood: we often just cross ControlRig here without manipulation layer
+		// should we clean this up?
+		if (UControlRig* ControlRig = WeakControlRigEditing.Get())
+		{
+			OutGlobalTransform = ControlRig->GetHierarchy()->GetGlobalTransform(InElement);
+			return true;
+		}
+	}
+
+	return false;
 }
 
-void FControlRigEditMode::SetKeyForControl(const FControlUnitProxy& UnitProxy)
+void FControlRigEditMode::SetKeysForSelectedManipulators()
 {
-	if(UControlRig* ControlRig = WeakControlRig.Get())
+	TSharedPtr<ISequencer> Sequencer = WeakSequencer.Pin();
+
+	if (Sequencer)
 	{
-		TSharedPtr<ISequencer> Sequencer = WeakSequencer.Pin();
-		if (Sequencer.IsValid())
+		for (const AControlRigGizmoActor* GizmoActor : GizmoActors)
 		{
-			// @TODO: need sequencer support for the new property path lib
-			TArray<UObject*> Objects({ ControlRig });
-		//	FKeyPropertyParams KeyPropertyParams(Objects, UnitProxy.PropertyPathString, ESequencerKeyMode::ManualKeyForced);
-		//	Sequencer->KeyProperty(KeyPropertyParams);
+			if (GizmoActor->IsSelected())
+			{
+				ManipulationLayer->AddKey(GizmoActor, Sequencer);
+			}
 		}
 	}
 }
@@ -1361,27 +1206,18 @@ void FControlRigEditMode::ToggleManipulators()
 	Settings->bHideManipulators = !Settings->bHideManipulators;
 }
 
-void FControlRigEditMode::ToggleTrajectories()
-{
-	Settings->bDisplayTrajectories = !Settings->bDisplayTrajectories;
-//	TrajectoryCache.RebuildMesh(SelectedIndices);
-}
-
-
-void FControlRigEditMode::RefreshTrajectoryCache()
-{
-//	TrajectoryCache.ForceRecalc();
-}
-
 bool FControlRigEditMode::MouseMove(FEditorViewportClient* ViewportClient, FViewport* Viewport, int32 x, int32 y)
 {
 	// Inform units of hover state
 	HActor* ActorHitProxy = HitProxyCast<HActor>(Viewport->GetHitProxy(x, y));
-	if(ActorHitProxy && ActorHitProxy->Actor && ActorHitProxy->Actor->IsA<AControlRigControl>())
+	if(ActorHitProxy && ActorHitProxy->Actor)
 	{
-		for(FControlUnitProxy& UnitProxy : ControlUnits)
+		if (ActorHitProxy->Actor->IsA<AControlRigGizmoActor>())
 		{
-			UnitProxy.SetHovered(ActorHitProxy->Actor == UnitProxy.Control);
+			for (AControlRigGizmoActor* GizmoActor : GizmoActors)
+			{
+				GizmoActor->SetHovered(GizmoActor == ActorHitProxy->Actor);
+			}
 		}
 	}
 
@@ -1390,105 +1226,75 @@ bool FControlRigEditMode::MouseMove(FEditorViewportClient* ViewportClient, FView
 
 bool FControlRigEditMode::MouseLeave(FEditorViewportClient* ViewportClient, FViewport* Viewport)
 {
-	// Remove hover state from all units
-	for(FControlUnitProxy& UnitProxy : ControlUnits)
+	for (AControlRigGizmoActor* GizmoActor : GizmoActors)
 	{
-		UnitProxy.SetHovered(false);
+		GizmoActor->SetHovered(false);
 	}
 
 	return false;
 }
 
-void FControlRigEditMode::RefreshControlProxies()
+void FControlRigEditMode::RecreateManipulationLayer()
 {
-	TArray<FString> SelectedPropertyPaths;
+	if (ManipulationLayer)
+	{
+		ManipulationLayer->DestroyLayer();
+	}
 
-	for(FControlUnitProxy& UnitProxy : ControlUnits)
+	if (UControlRig* ControlRig = WeakControlRigEditing.Get())
 	{
-		if(UnitProxy.IsSelected())
-		{
-			SelectedPropertyPaths.Add(UnitProxy.PropertyPathString);
-		}
+		ManipulationLayer = NewObject<UDefaultControlRigManipulationLayer>();
+		
+		// create layer
+		ManipulationLayer->CreateLayer();
 
-		if(UnitProxy.Control)
-		{
-			GetWorld()->DestroyActor(UnitProxy.Control, false, false);
-			UnitProxy.Control = nullptr;
-		}
-	}
+		// default manipulation layer can support any control rig
+		ManipulationLayer->AddManipulatableObject(ControlRig);
 
-	ControlUnits.Reset();
+		// create gizmo actors
+		GizmoActors.Reset();
+		ManipulationLayer->CreateGizmoActors(GetWorld(), GizmoActors);
 
-	if(UControlRig* ControlRig = WeakControlRig.Get())
-	{
-		UControlRigBlueprintGeneratedClass* Class = Cast<UControlRigBlueprintGeneratedClass>(ControlRig->GetClass());
-		for(UStructProperty* ControlUnitProperty : Class->ControlUnitProperties)
+		USceneComponent* Component = ManipulationLayer->GetSkeletalMeshComponent();
+		if (Component)
 		{
-			FRigUnit_Control* Control = ControlUnitProperty->ContainerPtrToValuePtr<FRigUnit_Control>(ControlRig);
-			FControlUnitProxy& UnitProxy = ControlUnits[ControlUnits.AddDefaulted()];
-			UnitProxy.PropertyPath = FCachedPropertyPath(ControlUnitProperty->GetName());
-			UnitProxy.PropertyPathString = UnitProxy.PropertyPath.ToString();
-			UnitProxy.TransformPropertyPath = FCachedPropertyPath(ControlUnitProperty->GetName() + TEXT(".Transform"));
-			UnitProxy.TransformPropertyPathString = UnitProxy.TransformPropertyPath.ToString();
-			UnitProxy.SetSelected(SelectedPropertyPaths.Contains(UnitProxy.PropertyPathString));
-
-			if(Control->ControlClass)
+			AActor* PreviewActor = Component->GetOwner();
+
+			for (AControlRigGizmoActor* GizmoActor : GizmoActors)
 			{
-				FActorSpawnParameters ActorSpawnParameters;
-				ActorSpawnParameters.bTemporaryEditorActor = true;
-				UnitProxy.Control = GetWorld()->SpawnActor<AControlRigControl>(Control->ControlClass, ActorSpawnParameters);
-				UnitProxy.Control->SetPropertyPath(UnitProxy.PropertyPathString);
+				// attach to preview actor, so that we can communicate via relative transfrom from the previewactor
+				GizmoActor->AttachToActor(PreviewActor, FAttachmentTransformRules::KeepWorldTransform);
 
 				TInlineComponentArray<UPrimitiveComponent*> PrimitiveComponents;
-				UnitProxy.Control->GetComponents(PrimitiveComponents, true);
-				for(UPrimitiveComponent* PrimitiveComponent : PrimitiveComponents)
+				GizmoActor->GetComponents(PrimitiveComponents, true);
+				for (UPrimitiveComponent* PrimitiveComponent : PrimitiveComponents)
 				{
-					PrimitiveComponent->SelectionOverrideDelegate = UPrimitiveComponent::FSelectionOverride::CreateRaw(this, &FControlRigEditMode::PreviewComponentSelectionOverride);
+					PrimitiveComponent->SelectionOverrideDelegate = UPrimitiveComponent::FSelectionOverride::CreateRaw(this, &FControlRigEditMode::GizmoSelectionOverride);
 					PrimitiveComponent->PushSelectionToProxy();
 				}
 			}
 		}
 	}
-
-	HandleSelectionChanged(SelectedPropertyPaths);
 }
 
-FControlRigEditMode* FControlRigEditMode::GetEditModeFromWorldContext(UWorld* InWorldContext)
+void RequestToRecreateManipulationLayer()
 {
-	return nullptr;
+
 }
 
-FRigUnit_Control* FControlRigEditMode::GetRigUnit(const FControlUnitProxy& InProxy, UControlRig* InControlRig, UScriptStruct** OutControlStructPtr /*= nullptr*/)
-{
-	UControlRigBlueprintGeneratedClass* Class = CastChecked<UControlRigBlueprintGeneratedClass>(InControlRig->GetClass());
-	for(UStructProperty* Property : Class->ControlUnitProperties)
-	{
-		if(Property->GetFName() == InProxy.PropertyPath.GetLastSegment().GetName())
-		{
-			if(OutControlStructPtr)
-			{
-				*OutControlStructPtr = Property->Struct;
-			}
-			return Property->ContainerPtrToValuePtr<FRigUnit_Control>(InControlRig);
-		}
-	}
 
+FControlRigEditMode* FControlRigEditMode::GetEditModeFromWorldContext(UWorld* InWorldContext)
+{
 	return nullptr;
 }
 
-bool FControlRigEditMode::PreviewComponentSelectionOverride(const UPrimitiveComponent* InComponent) const
+bool FControlRigEditMode::GizmoSelectionOverride(const UPrimitiveComponent* InComponent) const
 {
-	AActor* OwnerActor = InComponent->GetOwner();
-	if(OwnerActor)
+	AControlRigGizmoActor* OwnerActor = Cast<AControlRigGizmoActor>(InComponent->GetOwner());
+	if (OwnerActor)
 	{
 		// See if the actor is in a selected unit proxy
-		for(const FControlUnitProxy& UnitProxy : ControlUnits)
-		{
-			if(UnitProxy.Control == OwnerActor)
-			{
-				return UnitProxy.IsSelected();
-			}
-		}
+		return OwnerActor->IsSelected();
 	}
 
 	return false;
@@ -1496,79 +1302,124 @@ bool FControlRigEditMode::PreviewComponentSelectionOverride(const UPrimitiveComp
 
 void FControlRigEditMode::OnObjectsReplaced(const TMap<UObject*, UObject*>& OldToNewInstanceMap)
 {
-	if (WeakControlRig.IsValid())
+	if (WeakControlRigEditing.IsValid())
 	{
-		UObject* OldObject = WeakControlRig.Get();
+		UObject* OldObject = WeakControlRigEditing.Get();
 		UObject* NewObject = OldToNewInstanceMap.FindRef(OldObject);
 		if (NewObject)
 		{
-			WeakControlRig = Cast<UControlRig>(NewObject);
-			WeakControlRig->PostReinstanceCallback(CastChecked<UControlRig>(OldObject));
+			WeakControlRigEditing = Cast<UControlRig>(NewObject);
+			WeakControlRigEditing->PostReinstanceCallback(CastChecked<UControlRig>(OldObject));
 			SetObjects_Internal();
 		}
 	}
 }
 
-bool FControlRigEditMode::AreBoneSelectedAndMovable() const
+bool FControlRigEditMode::IsTransformDelegateAvailable() const
+{
+	return (OnGetRigElementTransformDelegate.IsBound() && OnSetRigElementTransformDelegate.IsBound());
+}
+
+bool FControlRigEditMode::AreRigElementSelectedAndMovable() const
 {
-	if (UControlRig* ControlRig = WeakControlRig.Get())
+	if (IsTransformDelegateAvailable())
 	{
-		return (!ControlRig->bExecutionOn && OnGetBoneTransformDelegate.IsBound() && OnSetBoneTransformDelegate.IsBound() && SelectedBones.Num() > 0);
+		// you can move control anytime
+		if (AreRigElementsSelected(FRigElementTypeHelper::ToMask(ERigElementType::Control)))
+		{
+			return true;
+		}
+
+		// if bone, we want to make sure execution is stopped
+		if (bEnableRigElementDefaultPoseEditing && AreRigElementsSelected(FRigElementTypeHelper::ToMask(ERigElementType::Bone)))
+		{
+			return true;
+		}
 	}
 
 	return false;
+
 }
 
-bool FControlRigEditMode::AreBoneSelected() const
+void FControlRigEditMode::OnRigElementAdded(FRigHierarchyContainer* Container, const FRigElementKey& InKey)
 {
-	return (SelectedBones.Num() > 0);
+	RequestToRecreateManipulationLayer();
 }
 
-void FControlRigEditMode::SelectBone(const FName& InBone)
+void FControlRigEditMode::OnRigElementRemoved(FRigHierarchyContainer* Container, const FRigElementKey& InKey)
 {
-	ClearControlSelection();
+	RequestToRecreateManipulationLayer();
+}
 
-	SelectedBones.Reset();
-	if (InBone != NAME_None)
-	{
-		SelectedBones.Add(InBone);
-	}
+void FControlRigEditMode::OnRigElementRenamed(FRigHierarchyContainer* Container, ERigElementType ElementType, const FName& InOldName, const FName& InNewName)
+{
+	RequestToRecreateManipulationLayer();
 }
 
-void FControlRigEditMode::HandleModelModified(const UControlRigModel* InModel, EControlRigModelNotifType InType, const void* InPayload)
+void FControlRigEditMode::OnRigElementReparented(FRigHierarchyContainer* Container, const FRigElementKey& InKey, const FName& InOldParentName, const FName& InNewParentName)
 {
-	switch (InType)
+	RequestToRecreateManipulationLayer();
+}
+
+void FControlRigEditMode::OnRigElementSelected(FRigHierarchyContainer* Container, const FRigElementKey& InKey, bool bSelected)
+{
+	switch (InKey.Type)
 	{
-		case EControlRigModelNotifType::NodeSelected:
+		case ERigElementType::Bone:
+		case ERigElementType::Control:
+		case ERigElementType::Space:
+		case ERigElementType::Curve:
 		{
-			const FControlRigModelNode* Node = (const FControlRigModelNode*)InPayload;
-			if (Node)
+			if (bSelected)
 			{
-				SetControlSelection(*Node->Name.ToString(), true);
+				SelectedRigElements.AddUnique(InKey);
 			}
-			break;
-		}
-		case EControlRigModelNotifType::NodeDeselected:
-		{
-			const FControlRigModelNode* Node = (const FControlRigModelNode*)InPayload;
-			if (Node)
+			else
 			{
-				SetControlSelection(*Node->Name.ToString(), false);
+				SelectedRigElements.Remove(InKey);
 			}
+
+			// if it's control
+			if (InKey.Type == ERigElementType::Control)
+			{
+				// users may select gizmo and control rig units, so we have to let them go through both of them if they do
+				// first go through gizmo actor
+				AControlRigGizmoActor* GizmoActor = GetGizmoFromControlName(InKey.Name);
+				if (GizmoActor)
+				{
+					GizmoActor->SetSelected(bSelected);
+				}
+			}
+
+			HandleSelectionChanged();
 			break;
 		}
 		default:
 		{
+			// todo
+			ensureMsgf(false, TEXT("Unsupported Type of RigElement: %d"), InKey.Type);
 			break;
 		}
 	}
+}
+
+void FControlRigEditMode::OnRigElementChanged(FRigHierarchyContainer* Container, const FRigElementKey& InKey)
+{
+	if(!InKey) // all of them changed
+	{
+		RequestToRecreateManipulationLayer();
+	}
+}
+
 
-	_ModifiedEvent.Broadcast(InModel, InType, InPayload);
+void FControlRigEditMode::OnControlUISettingChanged(FRigHierarchyContainer* Container, const FRigElementKey& InKey)
+{
+	RequestToRecreateManipulationLayer();
 }
 
-UControlRigModel::FModifiedEvent& FControlRigEditMode::OnModified()
+AControlRigGizmoActor* FControlRigEditMode::GetGizmoFromControlName(const FName& InControlName) const
 {
-	return _ModifiedEvent;
+	return ManipulationLayer->GetGizmoFromControlName(InControlName);
 }
 
 #undef LOCTEXT_NAMESPACE
