Copying //UE4/Dev-Anim to Dev-Main (//UE4/Dev-Main) @ 8782600

diff --git a/Engine/Plugins/Experimental/ControlRig/Source/ControlRigEditor/Private/EditMode/ControlRigEditMode.h b/Engine/Plugins/Experimental/ControlRig/Source/ControlRigEditor/Private/EditMode/ControlRigEditMode.h
--- a/Engine/Plugins/Experimental/ControlRig/Source/ControlRigEditor/Private/EditMode/ControlRigEditMode.h
+++ b/Engine/Plugins/Experimental/ControlRig/Source/ControlRigEditor/Private/EditMode/ControlRigEditMode.h
@@ -5,9 +5,11 @@
 #include "CoreMinimal.h"
 #include "InputCoreTypes.h"
 #include "IPersonaEditMode.h"
-#include "ControlRigTrajectoryCache.h"
-#include "ControlUnitProxy.h"
 #include "ControlRigModel.h"
+#include "Rigs/RigHierarchyContainer.h"
+#include "Drawing/ControlRigDrawInterface.h"
+#include "Units/RigUnitContext.h"
+#include "Framework/MultiBox/MultiBoxBuilder.h"
 
 class FEditorViewportClient;
 class FViewport;
@@ -22,11 +24,12 @@ class FPrimitiveDrawInterface;
 class FToolBarBuilder;
 class FExtender;
 class IMovieScenePlayer;
-struct FRigUnit_Control;
+class AControlRigGizmoActor;
+class UDefaultControlRigManipulationLayer;
 
-/** Delegate fired when controls are selected */
-DECLARE_DELEGATE_RetVal_TwoParams(FTransform, FOnGetBoneTransform, const FName& /*BoneName*/, bool /*bLocal*/);
-DECLARE_DELEGATE_TwoParams(FOnSetBoneTransform, const FName& /*BoneName*/, const FTransform& /*Transform*/);
+DECLARE_DELEGATE_RetVal_TwoParams(FTransform, FOnGetRigElementTransform, const FRigElementKey& /*RigElementKey*/, bool /*bLocal*/);
+DECLARE_DELEGATE_ThreeParams(FOnSetRigElementTransform, const FRigElementKey& /*RigElementKey*/, const FTransform& /*Transform*/, bool /*bLocal*/);
+DECLARE_DELEGATE_RetVal(TSharedPtr<FUICommandList>, FNewMenuCommandsDelegate);
 
 class FControlRigEditMode : public IPersonaEditMode
 {
@@ -37,7 +40,7 @@ public:
 	~FControlRigEditMode();
 
 	/** Set the objects to be displayed in the details panel */
-	void SetObjects(const TWeakObjectPtr<>& InSelectedObject, const FGuid& InObjectBinding);
+	void SetObjects(const TWeakObjectPtr<>& InSelectedObject, const FGuid& InObjectBinding, UObject* BindingObject);
 
 	/** Set the sequencer we are bound to */
 	void SetSequencer(TSharedPtr<ISequencer> InSequencer);
@@ -77,35 +80,6 @@ public:
 	/** FGCObject interface */
 	virtual void AddReferencedObjects( FReferenceCollector& Collector ) override;
 
-	/** Clear all selected controls */
-	void ClearControlSelection();
-
-	/** Set a control's selection state */
-	void SetControlSelection(const FString& InControlPropertyPath, bool bSelected);
-
-	/** Set multiple control's selection states */
-	void SetControlSelection(const TArray<FString>& InControlPropertyPaths, bool bSelected);
-
-	/** Check if the specified control is selected */
-	bool IsControlSelected(const FString& InControlPropertyPath) const;
-
-	/** Check if any controls are selected */
-	bool AreControlsSelected() const;
-
-	/** Get the number of selected controls */
-	int32 GetNumSelectedControls() const;
-
-	/** returns all property paths for all selected controls */
-	TArray<FString> GetSelectedControls() const;
-
-	/** Set a control's enabled state */
-	void SetControlEnabled(const FString& InControlPropertyPath, bool bEnabled);
-
-	/** Get the node name from the property path */
-	FString GetControlFromPropertyPath(const FString& PropertyPath) const;
-	/** Check if the specified control is enabled */
-	bool IsControlEnabled(const FString& InControlPropertyPath) const;
-
 	/** 
 	 * Lets the edit mode know that an object has just been spawned. 
 	 * Allows us to redisplay different underlying objects in the details panel.
@@ -121,28 +95,37 @@ public:
 	/** Refresh our internal object list (they may have changed) */
 	void RefreshObjects();
 
-	/** Refresh our trajectory cache */
-	void RefreshTrajectoryCache();
-
-	/** Set a key for a specific control */
-	void SetKeyForControl(const FControlUnitProxy& UnitProxy);
-
 	/** Get the settings we are using */
 	const UControlRigEditModeSettings* GetSettings() { return Settings; }
 
 	/** Find the edit mode corresponding to the specified world context */
 	static FControlRigEditMode* GetEditModeFromWorldContext(UWorld* InWorldContext);
 
-	/** Helper function - get a rig unit from a proxy and a rig */
-	static FRigUnit_Control* GetRigUnit(const FControlUnitProxy& InProxy, UControlRig* InControlRig, UScriptStruct** OutControlStructPtr = nullptr);
+	/** Bone Manipulation Delegates */
+	FOnGetRigElementTransform& OnGetRigElementTransform() { return OnGetRigElementTransformDelegate; }
+	FOnSetRigElementTransform& OnSetRigElementTransform() { return OnSetRigElementTransformDelegate; }
+
+	/** Context Menu Delegates */
+	FNewMenuDelegate& OnContextMenu() { return OnContextMenuDelegate; }
+	FNewMenuCommandsDelegate& OnContextMenuCommands() { return OnContextMenuCommandsDelegate; }
 
-	/** Select Bone */
-	void SelectBone(const FName& InBone);
-	FOnGetBoneTransform& OnGetBoneTransform() { return OnGetBoneTransformDelegate; }
-	FOnSetBoneTransform& OnSetBoneTransform() { return OnSetBoneTransformDelegate; }
+	// callback that gets called when rig element is selected in other view
+	void OnRigElementAdded(FRigHierarchyContainer* Container, const FRigElementKey& InKey);
+	void OnRigElementRemoved(FRigHierarchyContainer* Container, const FRigElementKey& InKey);
+	void OnRigElementRenamed(FRigHierarchyContainer* Container, ERigElementType ElementType, const FName& InOldName, const FName& InNewName);
+	void OnRigElementReparented(FRigHierarchyContainer* Container, const FRigElementKey& InKey, const FName& InOldParentName, const FName& InNewParentName);
+	void OnRigElementSelected(FRigHierarchyContainer* Container, const FRigElementKey& InKey, bool bSelected);
+	void OnRigElementChanged(FRigHierarchyContainer* Container, const FRigElementKey& InKey);
+	void OnControlUISettingChanged(FRigHierarchyContainer* Container, const FRigElementKey& InKey);
 
-	UControlRigModel::FModifiedEvent& OnModified();
-	void HandleModelModified(const UControlRigModel* InModel, EControlRigModelNotifType InType, const void* InPayload);
+	/** Enable RigElement Editing */
+	void EnableRigElementEditing(bool bEnabled);
+
+protected:
+
+	// Gizmo related functions wrt enable/selection
+	/** Get the node name from the property path */
+	AControlRigGizmoActor* GetGizmoFromControlName(const FName& InControlName) const;
 
 protected:
 	/** Helper function: set ControlRigs array to the details panel */
@@ -152,10 +135,10 @@ protected:
 	void RecalcPivotTransform();
 
 	/** Helper function for box/frustum intersection */
-	bool IntersectSelect(bool InSelect, const TFunctionRef<bool(const FControlUnitProxy&, const FTransform&)>& Intersects);
+	bool IntersectSelect(bool InSelect, const TFunctionRef<bool(const AControlRigGizmoActor*, const FTransform&)>& Intersects);
 
 	/** Handle selection internally */
-	void HandleSelectionChanged(const TArray<FString>& InSelectedBones);
+	void HandleSelectionChanged();
 
 	/** Set keys on all selected manipulators */
 	void SetKeysForSelectedManipulators();
@@ -163,52 +146,42 @@ protected:
 	/** Toggles visibility of manipulators in the viewport */
 	void ToggleManipulators();
 
-	/** Toggles visibility of trajectories in the viewport */
-	void ToggleTrajectories();
-
 	/** Bind our keyboard commands */
 	void BindCommands();
 
-	/** Refresh control proxies when the control rig changes */
-	void RefreshControlProxies();
+	/** It creates if it doesn't have it */
+	void RecreateManipulationLayer();
+
+	/** Requests to recreate the manipulation layer in the next tick */
+	void RequestToRecreateManipulationLayer() { bRecreateManipulationLayerRequired = true; }
 
 	/** Let the preview scene know how we want to select components */
-	bool PreviewComponentSelectionOverride(const UPrimitiveComponent* InComponent) const;
+	bool GizmoSelectionOverride(const UPrimitiveComponent* InComponent) const;
 
 protected:
-	/** Cache for rendering trajectories */
-	FControlRigTrajectoryCache TrajectoryCache;
-
 	/** Settings object used to insert controls into the details panel */
 	UControlRigEditModeSettings* Settings;
 
-	/** The units we use to represent the rig */
-	TArray<FControlUnitProxy> ControlUnits;
-
 	/** Whether we are in the middle of a transaction */
 	bool bIsTransacting;
 
 	/** Whether a manipulator actually made a change when transacting */
 	bool bManipulatorMadeChange;
 
-	/** The ControlRig we are animating */
-	TWeakObjectPtr<UControlRig> WeakControlRig;
-
+	/** The ControlRig we are animating as a main. This is the main editing object we're working on */
+	TWeakObjectPtr<UControlRig> WeakControlRigEditing;
 	/** The sequencer GUID of the object we are animating */
 	FGuid ControlRigGuid;
 
+	/** The draw interface to use for the control rig */
+	FControlRigDrawInterface DrawInterface;
+
 	/** Sequencer we are currently bound to */
 	TWeakPtr<ISequencer> WeakSequencer;
 
-	/** As we cannot cycle widget mode during tracking, we defer cycling until after a click with this flag */
-	bool bSelectedBone;
-
 	/** Guard value for selection */
 	bool bSelecting;
 
-	/** Guard value for selection by property path */
-	bool bSelectingByPath;
-
 	/** Cached transform of pivot point for selected Bones */
 	FTransform PivotTransform;
 
@@ -218,15 +191,60 @@ protected:
 	/** Called from the editor when a blueprint object replacement has occurred */
 	void OnObjectsReplaced(const TMap<UObject*, UObject*>& OldToNewInstanceMap);
 
-	/** Selected Bones */
-	TArray<FName> SelectedBones;
+	/** Return true if transform setter/getter delegates are available */
+	bool IsTransformDelegateAvailable() const;
+
+	FOnGetRigElementTransform OnGetRigElementTransformDelegate;
+	FOnSetRigElementTransform OnSetRigElementTransformDelegate;
+	FNewMenuDelegate OnContextMenuDelegate;
+	FNewMenuCommandsDelegate OnContextMenuCommandsDelegate;
+	
+	TArray<FRigElementKey> SelectedRigElements;
+
+	/* Flag to enable element init pose editing @Todo: */
+	bool bEnableRigElementDefaultPoseEditing;
+
+	/* Flag to recreate manipulation layer during tick */
+	bool bRecreateManipulationLayerRequired;
+
+	/** Default Manipulation Layer */
+	UDefaultControlRigManipulationLayer* ManipulationLayer;
+	TArray<AControlRigGizmoActor*> GizmoActors;
+
+	/** Utility functions for UI/Some other viewport manipulation*/
+	bool IsControlSelected() const;
+	bool IsControlOrSpaceOrBoneSelected() const;
+	bool AreRigElementSelectedAndMovable() const;
+	
+	/** Set initial transform handlers */
+	void OpenContextMenu(FEditorViewportClient* InViewportClient);
+
+	/* Set initial transform helpers*/
+	bool GetRigElementGlobalTransform(const FRigElementKey& InElement, FTransform& OutGlobalTransform) const;
+	
+public: 
+	/** Clear all selected RigElements */
+	void ClearRigElementSelection(uint32 InTypes);
+
+	/** Set a RigElement's selection state */
+	void SetRigElementSelection(ERigElementType Type, const FName& InRigElementName, bool bSelected);
+
+	/** Set multiple RigElement's selection states */
+	void SetRigElementSelection(ERigElementType Type, const TArray<FName>& InRigElementNames, bool bSelected);
+
+	/** Check if any RigElements are selected */
+	bool AreRigElementsSelected(uint32 InTypes) const;
 
-	FOnGetBoneTransform OnGetBoneTransformDelegate;
-	FOnSetBoneTransform OnSetBoneTransformDelegate;
+	/** Get the number of selected RigElements */
+	int32 GetNumSelectedRigElements(uint32 InTypes) const;
 
-	bool AreBoneSelected() const;
+private:
+	/** Set a RigElement's selection state */
+	void SetRigElementSelectionInternal(ERigElementType Type, const FName& InRigElementName, bool bSelected);
 
-	bool AreBoneSelectedAndMovable() const;
+	FEditorViewportClient* CurrentViewportClient;
 
-	UControlRigModel::FModifiedEvent _ModifiedEvent;
+	friend class FControlRigEditorModule;
+	friend class UControlRigPickerWidget;
+	friend class FControlRigEditor;
 };
diff --git a/Engine/Plugins/Experimental/ControlRig/Source/ControlRigEditor/Private/EditMode/ControlRigEditModeCommands.cpp b/Engine/Plugins/Experimental/ControlRig/Source/ControlRigEditor/Private/EditMode/ControlRigEditModeCommands.cpp
--- a/Engine/Plugins/Experimental/ControlRig/Source/ControlRigEditor/Private/EditMode/ControlRigEditModeCommands.cpp
+++ b/Engine/Plugins/Experimental/ControlRig/Source/ControlRigEditor/Private/EditMode/ControlRigEditModeCommands.cpp
@@ -8,7 +8,6 @@ void FControlRigEditModeCommands::RegisterCommands()
 {
 	UI_COMMAND(SetKey, "Set Key", "Sets a key on the selected nodes", EUserInterfaceActionType::Button, FInputChord(EKeys::Enter));
 	UI_COMMAND(ToggleManipulators, "Toggle Manipulators", "Toggles visibility of manipulators in the viewport", EUserInterfaceActionType::Button, FInputChord(EKeys::R));
-	UI_COMMAND(ToggleTrajectories, "Toggle Trajectories", "Toggles visibility of trajectories in the viewport", EUserInterfaceActionType::Button, FInputChord(EKeys::T));
 	UI_COMMAND(ExportAnimSequence, "Export To Animation", "Export this sequence to an optimized skeleton-specific animation sequence", EUserInterfaceActionType::Button, FInputChord());
 	UI_COMMAND(ReExportAnimSequence, "Re-export To Animation", "Re-export this sequence to an animation sequence using the previous export settings", EUserInterfaceActionType::Button, FInputChord());
 	UI_COMMAND(ImportFromRigSequence, "Import From Rig Sequence", "Import this animation sequence from a source rig sequence", EUserInterfaceActionType::Button, FInputChord());
diff --git a/Engine/Plugins/Experimental/ControlRig/Source/ControlRigEditor/Private/EditMode/ControlRigEditModeCommands.h b/Engine/Plugins/Experimental/ControlRig/Source/ControlRigEditor/Private/EditMode/ControlRigEditModeCommands.h
--- a/Engine/Plugins/Experimental/ControlRig/Source/ControlRigEditor/Private/EditMode/ControlRigEditModeCommands.h
+++ b/Engine/Plugins/Experimental/ControlRig/Source/ControlRigEditor/Private/EditMode/ControlRigEditModeCommands.h
@@ -25,8 +25,6 @@ public:
 	/** Toggles hiding all manipulators in the viewport */
 	TSharedPtr< FUICommandInfo > ToggleManipulators;
 
-	/** Toggles hiding all trajectories in the viewport */
-	TSharedPtr< FUICommandInfo > ToggleTrajectories;
 
 	/** Export this sequence to an optimized skeleton-specific animation sequence */
 	TSharedPtr< FUICommandInfo > ExportAnimSequence;
diff --git a/Engine/Plugins/Experimental/ControlRig/Source/ControlRigEditor/Private/EditMode/ControlRigEditModeSettings.cpp b/Engine/Plugins/Experimental/ControlRig/Source/ControlRigEditor/Private/EditMode/ControlRigEditModeSettings.cpp
--- a/Engine/Plugins/Experimental/ControlRig/Source/ControlRigEditor/Private/EditMode/ControlRigEditModeSettings.cpp
+++ b/Engine/Plugins/Experimental/ControlRig/Source/ControlRigEditor/Private/EditMode/ControlRigEditModeSettings.cpp
@@ -56,15 +56,5 @@ void UControlRigEditModeSettings::PostEditChangeProperty(struct FPropertyChanged
 				GEditor->GetEditorSubsystem<UAssetEditorSubsystem>()->OpenEditorForAsset(Sequence);
 			}
 		}
-		else if (Property->GetFName() == GET_MEMBER_NAME_CHECKED(UControlRigEditModeSettings, bDisplayTrajectories))
-		{
-			if (Sequence)
-			{
-				if (FControlRigEditMode* ControlRigEditMode = static_cast<FControlRigEditMode*>(GLevelEditorModeTools().GetActiveMode(FControlRigEditMode::ModeName)))
-				{
-					ControlRigEditMode->RefreshTrajectoryCache();
-				}
-			}
-		}
 	}
 }
diff --git a/Engine/Plugins/Experimental/ControlRig/Source/ControlRigEditor/Private/EditMode/ControlRigEditModeSettings.h b/Engine/Plugins/Experimental/ControlRig/Source/ControlRigEditor/Private/EditMode/ControlRigEditModeSettings.h
--- a/Engine/Plugins/Experimental/ControlRig/Source/ControlRigEditor/Private/EditMode/ControlRigEditModeSettings.h
+++ b/Engine/Plugins/Experimental/ControlRig/Source/ControlRigEditor/Private/EditMode/ControlRigEditModeSettings.h
@@ -18,7 +18,11 @@ class UControlRigEditModeSettings : public UObject
 	GENERATED_BODY()
 
 	UControlRigEditModeSettings()
-	: bShowManipulatorsDuringPlayback(true)
+		: bDisplayHierarchy(false)
+		, bShowManipulatorsDuringPlayback(true)
+		, bHideManipulators(false)
+		, bDisplayAxesOnSelection(true)
+		, AxisScale(10.f)
 	{}
 
 	// UObject interface
@@ -46,9 +50,13 @@ public:
 	UPROPERTY(EditAnywhere, AdvancedDisplay, Category = "Animation")
 	bool bHideManipulators;
 
-	/** Whether to the trajectories when animations are being played back */
+	/** Should we show axes for the selected elements */
 	UPROPERTY(EditAnywhere, AdvancedDisplay, Category = "Animation")
-	bool bDisplayTrajectories;
+	bool bDisplayAxesOnSelection;
+
+	/** The scale for axes to draw on the selection */
+	UPROPERTY(EditAnywhere, AdvancedDisplay, Category = "Animation")
+	float AxisScale;
 
 private:
 	/** Cache the previous actor for pre/post edit change handling */
diff --git a/Engine/Plugins/Experimental/ControlRig/Source/ControlRigEditor/Private/EditMode/ControlRigPickerWidget.cpp b/Engine/Plugins/Experimental/ControlRig/Source/ControlRigEditor/Private/EditMode/ControlRigPickerWidget.cpp
--- a/Engine/Plugins/Experimental/ControlRig/Source/ControlRigEditor/Private/EditMode/ControlRigPickerWidget.cpp
+++ b/Engine/Plugins/Experimental/ControlRig/Source/ControlRigEditor/Private/EditMode/ControlRigPickerWidget.cpp
@@ -13,7 +13,7 @@ void UControlRigPickerWidget::SelectControl(const FString& ControlPropertyPath,
 {
 	if(EditMode)
 	{
-		EditMode->SetControlSelection(ControlPropertyPath, bSelected);
+		EditMode->SetRigElementSelection(ERigElementType::Control, FName(*ControlPropertyPath, FNAME_Find), bSelected);
 	}
 }
 
@@ -21,7 +21,7 @@ bool UControlRigPickerWidget::IsControlSelected(const FString& ControlPropertyPa
 {
 	if(EditMode)
 	{
-		return EditMode->IsControlSelected(ControlPropertyPath);
+		return EditMode->SelectedRigElements.Contains(FRigElementKey(FName(*ControlPropertyPath, FNAME_Find), ERigElementType::Control));
 	}
 
 	return false;
@@ -31,7 +31,8 @@ void UControlRigPickerWidget::EnableControl(const FString& ControlPropertyPath,
 {
 	if(EditMode)
 	{
-		EditMode->SetControlEnabled(ControlPropertyPath, bEnabled);
+		ensure(false);
+		//EditMode->SetControlEnabled(ControlPropertyPath, bEnabled);
 	}
 }
 
@@ -39,7 +40,8 @@ bool UControlRigPickerWidget::IsControlEnabled(const FString& ControlPropertyPat
 {
 	if(EditMode)
 	{
-		return EditMode->IsControlEnabled(ControlPropertyPath);
+		ensure(false);
+		return true; //EditMode->IsControlEnabled(ControlPropertyPath);
 	}
 
 	return false;
diff --git a/Engine/Plugins/Experimental/ControlRig/Source/ControlRigEditor/Private/EditMode/SControlRigEditModeTools.cpp b/Engine/Plugins/Experimental/ControlRig/Source/ControlRigEditor/Private/EditMode/SControlRigEditModeTools.cpp
--- a/Engine/Plugins/Experimental/ControlRig/Source/ControlRigEditor/Private/EditMode/SControlRigEditModeTools.cpp
+++ b/Engine/Plugins/Experimental/ControlRig/Source/ControlRigEditor/Private/EditMode/SControlRigEditModeTools.cpp
@@ -95,13 +95,6 @@ void SControlRigEditModeTools::Construct(const FArguments& InArgs, UWorld* InWor
 			]
 		]
 	];
-
-	// Bind notification when edit mode selection changes, so we can update picker
-	FControlRigEditMode* ControlRigEditMode = static_cast<FControlRigEditMode*>(GLevelEditorModeTools().GetActiveMode(FControlRigEditMode::ModeName));
-	if (ControlRigEditMode)
-	{
-		ControlRigEditMode->OnModified().AddSP(this, &SControlRigEditModeTools::HandleModelModified);
-	}	
 }
 
 void SControlRigEditModeTools::SetDetailsObjects(const TArray<TWeakObjectPtr<>>& InObjects)
@@ -252,7 +245,7 @@ bool SControlRigEditModeTools::IsReadOnlyPropertyOnDetailCustomization(const FPr
 
 static bool bPickerChangingSelection = false;
 
-void SControlRigEditModeTools::OnManipulatorsPicked(const TArray<FString>& Manipulators)
+void SControlRigEditModeTools::OnManipulatorsPicked(const TArray<FName>& Manipulators)
 {
 	FControlRigEditMode* ControlRigEditMode = static_cast<FControlRigEditMode*>(GLevelEditorModeTools().GetActiveMode(FControlRigEditMode::ModeName));
 	if (ControlRigEditMode)
@@ -260,36 +253,8 @@ void SControlRigEditModeTools::OnManipulatorsPicked(const TArray<FString>& Manip
 		if (!bPickerChangingSelection)
 		{
 			TGuardValue<bool> SelectGuard(bPickerChangingSelection, true);
-			ControlRigEditMode->ClearControlSelection();
-			ControlRigEditMode->SetControlSelection(Manipulators, true);
-		}
-	}
-}
-
-void SControlRigEditModeTools::HandleModelModified(const UControlRigModel* InModel, EControlRigModelNotifType InType, const void* InPayload)
-{
-	if (bPickerChangingSelection)
-	{
-		return;
-	}
-
-	TGuardValue<bool> SelectGuard(bPickerChangingSelection, true);
-	switch (InType)
-	{
-		case EControlRigModelNotifType::NodeSelected:
-		case EControlRigModelNotifType::NodeDeselected:
-		{
-			const FControlRigModelNode* Node = (const FControlRigModelNode*)InPayload;
-			if (Node)
-			{
-				// those are not yet implemented yet
-				// ControlPicker->SelectManipulator(Node->Name, InType == EControlRigModelNotifType::NodeSelected);
-			}
-			break;
-		}
-		default:
-		{
-			break;
+			ControlRigEditMode->ClearRigElementSelection((uint32)ERigElementType::Control);
+			ControlRigEditMode->SetRigElementSelection(ERigElementType::Control, Manipulators, true);
 		}
 	}
 }
diff --git a/Engine/Plugins/Experimental/ControlRig/Source/ControlRigEditor/Private/EditMode/SControlRigEditModeTools.h b/Engine/Plugins/Experimental/ControlRig/Source/ControlRigEditor/Private/EditMode/SControlRigEditModeTools.h
--- a/Engine/Plugins/Experimental/ControlRig/Source/ControlRigEditor/Private/EditMode/SControlRigEditModeTools.h
+++ b/Engine/Plugins/Experimental/ControlRig/Source/ControlRigEditor/Private/EditMode/SControlRigEditModeTools.h
@@ -49,7 +49,5 @@ private:
 	bool IsReadOnlyPropertyOnDetailCustomization(const struct FPropertyAndParent& InPropertyAndParent) const;
 
 	/** Called when a manipulator is selected in the picker */
-	void OnManipulatorsPicked(const TArray<FString>& Manipulators);
-
-	void HandleModelModified(const UControlRigModel* InModel, EControlRigModelNotifType InType, const void* InPayload);
+	void OnManipulatorsPicked(const TArray<FName>& Manipulators);
 };
diff --git a/Engine/Plugins/Experimental/ControlRig/Source/ControlRigEditor/Private/Editor/ControlRigEditor.cpp b/Engine/Plugins/Experimental/ControlRig/Source/ControlRigEditor/Private/Editor/ControlRigEditor.cpp
--- a/Engine/Plugins/Experimental/ControlRig/Source/ControlRigEditor/Private/Editor/ControlRigEditor.cpp
+++ b/Engine/Plugins/Experimental/ControlRig/Source/ControlRigEditor/Private/Editor/ControlRigEditor.cpp
@@ -18,6 +18,7 @@
 #include "IPersonaToolkit.h"
 #include "PersonaModule.h"
 #include "ControlRigEditorEditMode.h"
+#include "ControlRigEditModeSettings.h"
 #include "AssetEditorModeManager.h"
 #include "ControlRigBlueprintGeneratedClass.h"
 #include "Sequencer/ControlRigSequencerAnimInstance.h"
@@ -28,7 +29,6 @@
 #include "ControlRigSkeletalMeshComponent.h"
 #include "ControlRigSkeletalMeshBinding.h"
 #include "ControlRigBlueprintUtils.h"
-#include "IPersonaViewport.h"
 #include "EditorViewportClient.h"
 #include "AnimationEditorPreviewActor.h"
 #include "Misc/MessageDialog.h"
@@ -36,6 +36,7 @@
 #include "ControlRigEditorStyle.h"
 #include "EditorFontGlyphs.h"
 #include "Widgets/Input/SButton.h"
+#include "Widgets/Input/SNumericEntryBox.h"
 #include "Widgets/Notifications/SNotificationList.h"
 #include "Framework/Notifications/NotificationManager.h"
 #include "SRigHierarchy.h"
@@ -53,6 +54,8 @@
 #include "ControlRig/Private/Units/Hierarchy/RigUnit_GetInitialBoneTransform.h"
 #include "ControlRig/Private/Units/Hierarchy/RigUnit_AddBoneTransform.h"
 #include "ControlRig/Private/Units/Execution/RigUnit_BeginExecution.h"
+#include "ControlRig/Private/Units/Hierarchy/RigUnit_GetControlTransform.h"
+#include "ControlRig/Private/Units/Hierarchy/RigUnit_GetSpaceTransform.h"
 #include "Graph/NodeSpawners/ControlRigUnitNodeSpawner.h"
 #include "Graph/ControlRigGraphSchema.h"
 #include "ControlRigObjectVersion.h"
@@ -81,11 +84,24 @@ FControlRigEditor::FControlRigEditor()
 	, bControlRigEditorInitialized(false)
 	, bIsSelecting(false)
 	, bIsSettingObjectBeingDebugged(false)
+	, bExecutionControlRig(true)
 {
 }
 
 FControlRigEditor::~FControlRigEditor()
 {
+	UControlRigBlueprint* RigBlueprint = GetControlRigBlueprint();
+	if (RigBlueprint)
+	{
+		RigBlueprint->HierarchyContainer.OnElementChanged.Clear();
+		RigBlueprint->HierarchyContainer.OnElementAdded.Clear();
+		RigBlueprint->HierarchyContainer.OnElementRemoved.Clear();
+		RigBlueprint->HierarchyContainer.OnElementRenamed.Clear();
+		RigBlueprint->HierarchyContainer.OnElementReparented.Clear();
+		RigBlueprint->HierarchyContainer.OnElementSelected.Clear();
+		RigBlueprint->HierarchyContainer.OnElementChanged.Clear();
+		RigBlueprint->HierarchyContainer.ControlHierarchy.OnControlUISettingsChanged.Clear();
+ 	}
 }
 
 UControlRigBlueprint* FControlRigEditor::GetControlRigBlueprint() const
@@ -194,9 +210,13 @@ void FControlRigEditor::InitControlRigEditor(const EToolkitMode::Type Mode, cons
 //	GetAssetEditorModeManager()->SetToolkitHost(GetToolkitHost());
 	GetAssetEditorModeManager()->SetDefaultMode(FControlRigEditorEditMode::ModeName);
 	GetAssetEditorModeManager()->ActivateMode(FControlRigEditorEditMode::ModeName);
-	GetEditMode().OnGetBoneTransform() = FOnGetBoneTransform::CreateSP(this, &FControlRigEditor::GetBoneTransform);
-	GetEditMode().OnSetBoneTransform() = FOnSetBoneTransform::CreateSP(this, &FControlRigEditor::SetBoneTransform);
-	InControlRigBlueprint->OnModified().AddSP(&GetEditMode(), &FControlRigEditMode::HandleModelModified);
+	if (FControlRigEditMode* EditMode = GetEditMode())
+	{
+		EditMode->OnGetRigElementTransform() = FOnGetRigElementTransform::CreateSP(this, &FControlRigEditor::GetRigElementTransform);
+		EditMode->OnSetRigElementTransform() = FOnSetRigElementTransform::CreateSP(this, &FControlRigEditor::SetRigElementTransform);
+		EditMode->OnContextMenu() = FNewMenuDelegate::CreateSP(this, &FControlRigEditor::HandleOnViewportContextMenuDelegate);
+		EditMode->OnContextMenuCommands() = FNewMenuCommandsDelegate::CreateSP(this, &FControlRigEditor::HandleOnViewportContextMenuCommandsDelegate);
+	}
 
 	UpdateControlRig();
 
@@ -229,9 +249,40 @@ void FControlRigEditor::InitControlRigEditor(const EToolkitMode::Type Mode, cons
 			}
 			else
 			{
+				bool bWasDirty = false;
+				UPackage* Package = InControlRigBlueprint->GetOutermost();
+				if (Package)
+				{
+					bWasDirty = Package->IsDirty();
+				}
 				InControlRigBlueprint->RebuildGraphFromModel();
+				if (Package)
+				{
+					if (Package->IsDirty() != bWasDirty)
+					{
+						Package->SetDirtyFlag(bWasDirty);
+					}
+				}
 			}
 		}
+
+		InControlRigBlueprint->HierarchyContainer.OnElementAdded.AddSP(this, &FControlRigEditor::OnRigElementAdded);
+		InControlRigBlueprint->HierarchyContainer.OnElementRemoved.AddSP(this, &FControlRigEditor::OnRigElementRemoved);
+		InControlRigBlueprint->HierarchyContainer.OnElementRenamed.AddSP(this, &FControlRigEditor::OnRigElementRenamed);
+		InControlRigBlueprint->HierarchyContainer.OnElementReparented.AddSP(this, &FControlRigEditor::OnRigElementReparented);
+		InControlRigBlueprint->HierarchyContainer.OnElementSelected.AddSP(this, &FControlRigEditor::OnRigElementSelected);
+		InControlRigBlueprint->HierarchyContainer.ControlHierarchy.OnControlUISettingsChanged.AddSP(this, &FControlRigEditor::OnControlUISettingChanged);
+
+		if (FControlRigEditMode* EditMode = GetEditMode())
+		{
+			InControlRigBlueprint->HierarchyContainer.OnElementAdded.AddSP(EditMode, &FControlRigEditMode::OnRigElementAdded);
+			InControlRigBlueprint->HierarchyContainer.OnElementRemoved.AddSP(EditMode, &FControlRigEditMode::OnRigElementRemoved);
+			InControlRigBlueprint->HierarchyContainer.OnElementRenamed.AddSP(EditMode, &FControlRigEditMode::OnRigElementRenamed);
+			InControlRigBlueprint->HierarchyContainer.OnElementReparented.AddSP(EditMode, &FControlRigEditMode::OnRigElementReparented);
+			InControlRigBlueprint->HierarchyContainer.OnElementSelected.AddSP(EditMode, &FControlRigEditMode::OnRigElementSelected);
+			InControlRigBlueprint->HierarchyContainer.OnElementChanged.AddSP(EditMode, &FControlRigEditMode::OnRigElementChanged);
+			InControlRigBlueprint->HierarchyContainer.ControlHierarchy.OnControlUISettingsChanged.AddSP(EditMode, &FControlRigEditMode::OnControlUISettingChanged);
+		}
 	}
 
 	bControlRigEditorInitialized = true;
@@ -252,13 +303,24 @@ void FControlRigEditor::ToggleExecuteGraph()
 {
 	if (ControlRig)
 	{
-		ControlRig->bExecutionOn = !ControlRig->bExecutionOn;
+		bExecutionControlRig = !bExecutionControlRig;
+
+		// enable rig element init pose change, if not executing, we allow rig element change
+		if (FControlRigEditMode* EditMode = GetEditMode())
+		{
+			EditMode->EnableRigElementEditing(!bExecutionControlRig);
+		}
+
+		// This is required now since we update execution/input flag on update controlrig
+		// @fixme: change this to just change flag instead of updating whole control rig
+		// I'll do this before first check-in
+		UpdateControlRig();
 	}
 }
 
 bool FControlRigEditor::IsExecuteGraphOn() const
 {
-	return (ControlRig)? ControlRig->bExecutionOn : false;
+	return bExecutionControlRig;
 }
 
 void FControlRigEditor::ExtendToolbar()
@@ -438,7 +500,8 @@ void FControlRigEditor::HandleSetObjectBeingDebugged(UObject* InObject)
 	if (DebuggedControlRig)
 	{
 		bool bIsExternalControlRig = DebuggedControlRig != ControlRig;
-		DebuggedControlRig->DrawInterface = &DrawInterface;
+		// @fixme: I'll also change the "execute" button to be not toggleable during debugging external control rig
+		bool bShouldExecute = bIsExternalControlRig && bExecutionControlRig;
 		DebuggedControlRig->ControlRigLog = &ControlRigLog;
 
 		UControlRigSkeletalMeshComponent* EditorSkelComp = Cast<UControlRigSkeletalMeshComponent>(GetPersonaToolkit()->GetPreviewScene()->GetPreviewMeshComponent());
@@ -449,18 +512,24 @@ void FControlRigEditor::HandleSetObjectBeingDebugged(UObject* InObject)
 			{
 				// we might want to move this into another method
 				FInputBlendPose Filter;
-				AnimInstance->UpdateControlRig(DebuggedControlRig, 0, false, false, Filter, 1.0f, bIsExternalControlRig);
+				AnimInstance->UpdateControlRig(DebuggedControlRig, 0, false, false, Filter, 1.0f, bShouldExecute, bShouldExecute);
 				AnimInstance->RecalcRequiredBones();
 
 				// since rig has changed, rebuild draw skeleton
 				EditorSkelComp->RebuildDebugDrawSkeleton();
-				GetEditMode().SetObjects(DebuggedControlRig, FGuid());
+				if (FControlRigEditMode* EditMode = GetEditMode())
+				{
+					EditMode->SetObjects(DebuggedControlRig, FGuid(), EditorSkelComp);
+				}
 			}
 		}
 	}
 	else
 	{
-		GetEditMode().SetObjects(nullptr, FGuid());
+		if (FControlRigEditMode* EditMode = GetEditMode())
+		{
+			EditMode->SetObjects(nullptr, FGuid(), nullptr);
+		}
 	}
 }
 
@@ -499,11 +568,13 @@ UBlueprint* FControlRigEditor::GetBlueprintObj() const
 
 void FControlRigEditor::SetDetailObjects(const TArray<UObject*>& InObjects)
 {
+	RigElementInDetailPanel = FRigElementKey();
 	Inspector->ShowDetailsForObjects(InObjects);
 }
 
 void FControlRigEditor::SetDetailObject(UObject* Obj)
 {
+	RigElementInDetailPanel = FRigElementKey();
 	TArray<UObject*> Objects;
 	if (Obj)
 	{
@@ -512,13 +583,23 @@ void FControlRigEditor::SetDetailObject(UObject* Obj)
 	SetDetailObjects(Objects);
 }
 
-void FControlRigEditor::SetDetailStruct(TSharedPtr<FStructOnScope> StructToDisplay)
+void FControlRigEditor::SetDetailStruct(const FRigElementKey& InElement, TSharedPtr<FStructOnScope> StructToDisplay)
 {
+	if (RigElementInDetailPanel == InElement)
+	{
+		return;
+	}
+	RigElementInDetailPanel = InElement;
+	if (StructToDisplay.IsValid())
+	{
+		StructToDisplay->SetPackage(GetControlRigBlueprint()->GetOutermost());
+	}
 	Inspector->ShowSingleStruct(StructToDisplay);
 }
 
 void FControlRigEditor::ClearDetailObject()
 {
+	RigElementInDetailPanel = FRigElementKey();
 	Inspector->ShowDetailsForObjects(TArray<UObject*>());
 	Inspector->ShowSingleStruct(TSharedPtr<FStructOnScope>());
 }
@@ -566,7 +647,6 @@ void FControlRigEditor::Compile()
 		if (ControlRig)
 		{
 			ControlRig->ControlRigLog = &ControlRigLog;
-			ControlRig->DrawInterface = &DrawInterface;
 
 			UControlRigBlueprintGeneratedClass* GeneratedClass = Cast<UControlRigBlueprintGeneratedClass>(ControlRig->GetClass());
 			if (GeneratedClass)
@@ -1257,6 +1337,62 @@ void FControlRigEditor::HandleViewportCreated(const TSharedRef<class IPersonaVie
 			]
 		]
 	);
+
+	InViewport->AddToolbarExtender(TEXT("AnimViewportDefaultCamera"), FMenuExtensionDelegate::CreateLambda(
+		[&](FMenuBuilder& InMenuBuilder)
+		{
+			InMenuBuilder.AddMenuSeparator(TEXT("Control Rig"));
+			InMenuBuilder.BeginSection("ControlRig", LOCTEXT("ControlRig_Label", "Control Rig"));
+			{
+				InMenuBuilder.AddWidget(
+					SNew(SBox)
+					.HAlign(HAlign_Right)
+					[
+						SNew(SBox)
+						.Padding(FMargin(4.0f, 0.0f, 0.0f, 0.0f))
+						.WidthOverride(100.0f)
+						[
+							SNew(SNumericEntryBox<float>)
+							.Font(FEditorStyle::GetFontStyle(TEXT("MenuItem.Font")))
+							.AllowSpin(true)
+							.MinSliderValue(0.0f)
+							.MaxSliderValue(100.0f)
+							.Value(this, &FControlRigEditor::GetToolbarAxesScale)
+							.OnValueChanged(this, &FControlRigEditor::OnToolbarAxesScaleChanged)
+							.ToolTipText(LOCTEXT("ControlRigAxesScaleToolTip", "Scale of axes drawn for selected rig elements"))
+						]
+					], 
+					LOCTEXT("ControlRigAxesScale", "Axes Scale")
+				);
+			}
+			InMenuBuilder.EndSection();
+		}
+	));
+
+	InViewport->GetKeyDownDelegate().BindLambda([&](const FGeometry& MyGeometry, const FKeyEvent& InKeyEvent) -> FReply {
+		if (OnKeyDownDelegate.IsBound())
+		{
+			return OnKeyDownDelegate.Execute(MyGeometry, InKeyEvent);
+		}
+		return FReply::Unhandled();
+	});
+}
+
+TOptional<float> FControlRigEditor::GetToolbarAxesScale() const
+{
+	if (FControlRigEditMode* EditMode = GetEditMode())
+	{
+		return EditMode->Settings->AxisScale;
+	}
+	return 0.f;
+}
+
+void FControlRigEditor::OnToolbarAxesScaleChanged(float InValue)
+{
+	if (FControlRigEditMode* EditMode = GetEditMode())
+	{
+		EditMode->Settings->AxisScale = InValue;
+	}
 }
 
 void FControlRigEditor::HandlePreviewSceneCreated(const TSharedRef<IPersonaPreviewScene>& InPersonaPreviewScene)
@@ -1297,13 +1433,18 @@ void FControlRigEditor::UpdateControlRig()
 				ControlRig = NewObject<UControlRig>(EditorSkelComp, Class);
 				// this is editing time rig
 				ControlRig->ExecutionType = ERigExecutionType::Editing;
-
 				ControlRig->ControlRigLog = &ControlRigLog;
-				ControlRig->DrawInterface = &DrawInterface;
+
+				ControlRig->InitializeFromCDO();
+				ControlRig->InstantiateOperatorsFromGeneratedClass();
  			}
 
-			CacheBoneNameList();
-			CacheCurveNameList();
+			if (UControlRig* CDO = Cast<UControlRig>(Class->GetDefaultObject()))
+			{
+				CDO->GizmoLibrary = GetControlRigBlueprint()->GizmoLibrary;
+			}
+
+			CacheNameLists();
 
 			// When the control rig is re-instanced on compile, it loses its binding, so we refresh it here if needed
 			if (!ControlRig->GetObjectBinding().IsValid())
@@ -1316,12 +1457,15 @@ void FControlRigEditor::UpdateControlRig()
 
 			// initialize is moved post reinstance
 			FInputBlendPose Filter;
-			AnimInstance->UpdateControlRig(ControlRig, 0, false, false, Filter, 1.0f);
+			AnimInstance->UpdateControlRig(ControlRig, 0, false, false, Filter, 1.0f, bExecutionControlRig, bExecutionControlRig);
 			AnimInstance->RecalcRequiredBones();
 
 			// since rig has changed, rebuild draw skeleton
 			EditorSkelComp->RebuildDebugDrawSkeleton();
-			GetEditMode().SetObjects(ControlRig, FGuid());
+			if (FControlRigEditMode* EditMode = GetEditMode())
+			{
+				EditMode->SetObjects(ControlRig, FGuid(), EditorSkelComp);
+			}
 
 			// update the graph with the compiler errors
 			UpdateGraphCompilerErrors();
@@ -1329,14 +1473,14 @@ void FControlRigEditor::UpdateControlRig()
 	}
 }
 
-void FControlRigEditor::CacheBoneNameList()
+void FControlRigEditor::CacheNameLists()
 {
 	DECLARE_SCOPE_HIERARCHICAL_COUNTER_FUNC()
 
-	if (ControlRig)
+	if (UControlRigBlueprint* ControlRigBP = GetControlRigBlueprint())
 	{
 		// make sure the bone name list is up 2 date for the editor graph
-		for (UEdGraph* Graph : GetBlueprintObj()->UbergraphPages)
+		for (UEdGraph* Graph : ControlRigBP->UbergraphPages)
 		{
 			UControlRigGraph* RigGraph = Cast<UControlRigGraph>(Graph);
 			if (RigGraph == nullptr)
@@ -1344,27 +1488,7 @@ void FControlRigEditor::CacheBoneNameList()
 				continue;
 			}
 
-			RigGraph->CacheBoneNameList(ControlRig->GetBaseHierarchy());
-		}
-	}
-}
-
-void FControlRigEditor::CacheCurveNameList()
-{
-	DECLARE_SCOPE_HIERARCHICAL_COUNTER_FUNC()
-
-	if (ControlRig)
-	{
-		// make sure the Curve name list is up 2 date for the editor graph
-		for (UEdGraph* Graph : GetBlueprintObj()->UbergraphPages)
-		{
-			UControlRigGraph* RigGraph = Cast<UControlRigGraph>(Graph);
-			if (RigGraph == nullptr)
-			{
-				continue;
-			}
-
-			RigGraph->CacheCurveNameList(ControlRig->GetCurveContainer());
+			RigGraph->CacheNameLists(&ControlRigBP->HierarchyContainer);
 		}
 	}
 }
@@ -1424,107 +1548,156 @@ void FControlRigEditor::OnNodeTitleCommitted(const FText& NewText, ETextCommit::
 	}
 }
 
-
-void FControlRigEditor::SelectBone(const FName& InBone)
+FTransform FControlRigEditor::GetRigElementTransform(const FRigElementKey& InElement, bool bLocal) const
 {
 	DECLARE_SCOPE_HIERARCHICAL_COUNTER_FUNC()
-
-	// edit mode has to know
-	GetEditMode().SelectBone(InBone);
-	// copy locally, we use this for copying back to template when modified
-
-	SelectedBone = InBone;
-	UControlRigSkeletalMeshComponent* EditorSkelComp = Cast<UControlRigSkeletalMeshComponent>(GetPersonaToolkit()->GetPreviewScene()->GetPreviewMeshComponent());
-	if (EditorSkelComp)
-	{
-		EditorSkelComp->BonesOfInterest.Reset();
-
-		int32 Index = ControlRig->Hierarchy.BaseHierarchy.GetIndex(InBone);
-		if (Index != INDEX_NONE)
-		{
-			EditorSkelComp->BonesOfInterest.Add(Index);
-		}
-	}
-}
-
-FTransform FControlRigEditor::GetBoneTransform(const FName& InBone, bool bLocal) const
-{
-	DECLARE_SCOPE_HIERARCHICAL_COUNTER_FUNC()
-
-	// @todo: think about transform mode
 	if (bLocal)
 	{
-		return ControlRig->Hierarchy.BaseHierarchy.GetLocalTransform(InBone);
+		return GetControlRigBlueprint()->HierarchyContainer.GetLocalTransform(InElement);
 	}
-
-	return ControlRig->Hierarchy.BaseHierarchy.GetGlobalTransform(InBone);
+	return GetControlRigBlueprint()->HierarchyContainer.GetGlobalTransform(InElement);
 }
 
-void FControlRigEditor::SetBoneTransform(const FName& InBone, const FTransform& InTransform)
+void FControlRigEditor::SetRigElementTransform(const FRigElementKey& InElement, const FTransform& InTransform, bool bLocal)
 {
 	DECLARE_SCOPE_HIERARCHICAL_COUNTER_FUNC()
 
-	// execution should be off
-	ensure(!ControlRig->bExecutionOn);
-
 	FScopedTransaction Transaction(LOCTEXT("Move Bone", "Move Bone transform"));
 	UControlRigBlueprint* ControlRigBP = GetControlRigBlueprint();
 	ControlRigBP->Modify();
 
-	// moving ref pose warning
-	// update init/global transform
-	// @todo: this needs revision once we decide how we allow users to modify init/global transform
-	// for now, updating init/global of the Bone from instances, but only modify init transform for archetype
-	// get local transform of current
-	// apply init based on parent init * current local 
+	switch (InElement.Type)
+	{
+		case ERigElementType::Bone:
+		{
+			FTransform Transform = InTransform;
+			if (bLocal)
+			{
+				FTransform ParentTransform = FTransform::Identity;
+				FRigElementKey ParentKey = ControlRigBP->HierarchyContainer.BoneHierarchy[InElement.Name].GetParentElementKey();
+				if (ParentKey.IsValid())
+				{
+					ParentTransform = GetRigElementTransform(ParentKey, false);
+				}
+				Transform = Transform * ParentTransform;
+				Transform.NormalizeRotation();
+			}
+
+			ControlRigBP->HierarchyContainer.BoneHierarchy.SetInitialTransform(InElement.Name, Transform);
+			ControlRigBP->HierarchyContainer.BoneHierarchy.SetGlobalTransform(InElement.Name, Transform);
+			OnHierarchyChanged();
+			break;
+		}
+		case ERigElementType::Control:
+		{
+			FTransform LocalTransform = InTransform;
+			FTransform GlobalTransform = InTransform;
+			if (!bLocal)
+			{
+				ControlRigBP->HierarchyContainer.ControlHierarchy.SetGlobalTransform(InElement.Name, InTransform);
+				LocalTransform = ControlRigBP->HierarchyContainer.ControlHierarchy.GetLocalTransform(InElement.Name);
+			}
+			else
+			{
+				ControlRigBP->HierarchyContainer.ControlHierarchy.SetLocalTransform(InElement.Name, InTransform);
+				GlobalTransform = ControlRigBP->HierarchyContainer.ControlHierarchy.GetGlobalTransform(InElement.Name);
+			}
 
-	ControlRig->Hierarchy.BaseHierarchy.SetInitialTransform(InBone, InTransform);
-	ControlRig->Hierarchy.BaseHierarchy.SetGlobalTransform(InBone, InTransform);
+			ControlRigBP->HierarchyContainer.ControlHierarchy.SetInitialValue<FTransform>(InElement.Name, LocalTransform);
+			ControlRigBP->HierarchyContainer.ControlHierarchy.SetGlobalTransform(InElement.Name, GlobalTransform);
+			OnHierarchyChanged();
+			break;
+		}
+		case ERigElementType::Space:
+		{
+			FTransform LocalTransform = InTransform;
+			FTransform GlobalTransform = InTransform;
+			if (!bLocal)
+			{
+				ControlRigBP->HierarchyContainer.SpaceHierarchy.SetGlobalTransform(InElement.Name, InTransform);
+				LocalTransform = ControlRigBP->HierarchyContainer.SpaceHierarchy.GetLocalTransform(InElement.Name);
+			}
+			else
+			{
+				ControlRigBP->HierarchyContainer.SpaceHierarchy.SetLocalTransform(InElement.Name, InTransform);
+				GlobalTransform = ControlRigBP->HierarchyContainer.SpaceHierarchy.GetGlobalTransform(InElement.Name);
+			}
 
-	ControlRigBP->Hierarchy.SetInitialTransform(InBone, InTransform);
+			ControlRigBP->HierarchyContainer.SpaceHierarchy.SetInitialTransform(InElement.Name, LocalTransform);
+			ControlRigBP->HierarchyContainer.SpaceHierarchy.SetGlobalTransform(InElement.Name, GlobalTransform);
+			OnHierarchyChanged();
+			break;
+		}
+		default:
+		{
+			ensureMsgf(false, TEXT("Unsupported RigElement Type : %d"), InElement.Type);
+			break;
+		}
+	}
 	
 	UControlRigSkeletalMeshComponent* EditorSkelComp = Cast<UControlRigSkeletalMeshComponent>(GetPersonaToolkit()->GetPreviewScene()->GetPreviewMeshComponent());
 	if (EditorSkelComp)
 	{
 		EditorSkelComp->RebuildDebugDrawSkeleton();
 	}
-
-	// I don't think I have to mark dirty here. 
-	// FBlueprintEditorUtils::MarkBlueprintAsModified(GetControlRigBlueprint());
-
-	// I don't think I have to mark dirty here. 
-	// FBlueprintEditorUtils::MarkBlueprintAsModified(GetControlRigBlueprint());
-	{
-		EditorSkelComp->RebuildDebugDrawSkeleton();
-	}
-
-	// I don't think I have to mark dirty here. 
-	// FBlueprintEditorUtils::MarkBlueprintAsModified(GetControlRigBlueprint());
 }
 
 void FControlRigEditor::OnFinishedChangingProperties(const FPropertyChangedEvent& PropertyChangedEvent)
 {
 	DECLARE_SCOPE_HIERARCHICAL_COUNTER_FUNC()
-	//	UE_LOG(LogControlRigEditor, Warning, TEXT("Current Property being modified : %s"), *GetNameSafe(PropertyChangedEvent.Property));
 
-	FName PropertyName = (PropertyChangedEvent.Property != NULL) ? PropertyChangedEvent.Property->GetFName() : NAME_None;
-	if (PropertyName == GET_MEMBER_NAME_CHECKED(FRigBone, InitialTransform))
+	UControlRigBlueprint* ControlRigBP = GetControlRigBlueprint();
+	if (ControlRig && ControlRigBP && RigElementInDetailPanel)
 	{
-		// if init transform changes, it updates to the base
-		UControlRigBlueprint* ControlRigBP = GetControlRigBlueprint();
-		if (ControlRig && ControlRigBP)
+		UControlRig* DebuggedControlRig = Cast<UControlRig>(GetBlueprintObj()->GetObjectBeingDebugged());
+		UScriptStruct* ScriptStruct = Cast< UScriptStruct>(PropertyChangedEvent.Property->GetOuter());
+		if (ScriptStruct)
 		{
-			if (SelectedBone != NAME_None)
+			if (ScriptStruct == FRigBone::StaticStruct() && RigElementInDetailPanel.Type == ERigElementType::Bone)
+			{
+				ControlRig->Hierarchy.BoneHierarchy[RigElementInDetailPanel.Name] = ControlRigBP->HierarchyContainer.BoneHierarchy[RigElementInDetailPanel.Name];
+				if (DebuggedControlRig && DebuggedControlRig != ControlRig)
+				{
+					DebuggedControlRig->Hierarchy.BoneHierarchy[RigElementInDetailPanel.Name] = ControlRigBP->HierarchyContainer.BoneHierarchy[RigElementInDetailPanel.Name];
+				}
+			}
+			else if (ScriptStruct == FRigSpace::StaticStruct() && RigElementInDetailPanel.Type == ERigElementType::Space)
+			{
+				ControlRig->Hierarchy.SpaceHierarchy[RigElementInDetailPanel.Name] = ControlRigBP->HierarchyContainer.SpaceHierarchy[RigElementInDetailPanel.Name];
+				if (DebuggedControlRig && DebuggedControlRig != ControlRig)
+				{
+					DebuggedControlRig->Hierarchy.SpaceHierarchy[RigElementInDetailPanel.Name] = ControlRigBP->HierarchyContainer.SpaceHierarchy[RigElementInDetailPanel.Name];
+				}
+			}
+			else if (ScriptStruct == FRigControl::StaticStruct() && RigElementInDetailPanel.Type == ERigElementType::Control)
+			{
+				ControlRig->Hierarchy.ControlHierarchy[RigElementInDetailPanel.Name] = ControlRigBP->HierarchyContainer.ControlHierarchy[RigElementInDetailPanel.Name];
+				if (DebuggedControlRig && DebuggedControlRig != ControlRig)
+				{
+					DebuggedControlRig->Hierarchy.ControlHierarchy[RigElementInDetailPanel.Name] = ControlRigBP->HierarchyContainer.ControlHierarchy[RigElementInDetailPanel.Name];
+				}
+			}
+			else if (ScriptStruct == FRigCurve::StaticStruct() && RigElementInDetailPanel.Type == ERigElementType::Curve)
 			{
-				const int32 BoneIndex = ControlRig->Hierarchy.BaseHierarchy.GetIndex(SelectedBone);
-				if (BoneIndex != INDEX_NONE)
+				ControlRig->Hierarchy.CurveContainer[RigElementInDetailPanel.Name] = ControlRigBP->HierarchyContainer.CurveContainer[RigElementInDetailPanel.Name];
+				if (DebuggedControlRig && DebuggedControlRig != ControlRig)
 				{
-					FTransform InitialTransform = ControlRig->Hierarchy.BaseHierarchy.GetInitialTransform(BoneIndex);
-					// update CDO  @todo - re-think about how we wrap around this nicer
-					// copy currently selected Bone to base hierarchy			
-					ControlRigBP->Hierarchy.SetInitialTransform(BoneIndex, InitialTransform);
+					DebuggedControlRig->Hierarchy.CurveContainer[RigElementInDetailPanel.Name] = ControlRigBP->HierarchyContainer.CurveContainer[RigElementInDetailPanel.Name];
 				}
 			}
+
+			ControlRigBP->MarkPackageDirty();
+		}
+
+		if (PropertyChangedEvent.Property->GetName().StartsWith(TEXT("Gizmo")) && RigElementInDetailPanel.Type == ERigElementType::Control)
+		{
+			ControlRig->Hierarchy.ControlHierarchy[RigElementInDetailPanel.Name] = ControlRigBP->HierarchyContainer.ControlHierarchy[RigElementInDetailPanel.Name];
+			if (DebuggedControlRig && DebuggedControlRig != ControlRig)
+			{
+				DebuggedControlRig->Hierarchy.ControlHierarchy[RigElementInDetailPanel.Name] = ControlRigBP->HierarchyContainer.ControlHierarchy[RigElementInDetailPanel.Name];
+			}
+			ControlRigBP->HierarchyContainer.ControlHierarchy.OnControlUISettingsChanged.Broadcast(&ControlRigBP->HierarchyContainer, FRigElementKey(RigElementInDetailPanel.Name, ERigElementType::Control));
+			ControlRigBP->MarkPackageDirty();
 		}
 	}
 }
@@ -1533,30 +1706,39 @@ void FControlRigEditor::OnHierarchyChanged()
 {
 	DECLARE_SCOPE_HIERARCHICAL_COUNTER_FUNC()
 
-	ClearDetailObject();
+	CacheNameLists();
+	if (UControlRigBlueprint* ControlRigBP = GetControlRigBlueprint())
+	{
+		ControlRigBP->PropagateHierarchyFromBPToInstances();
 
-	FBlueprintEditorUtils::MarkBlueprintAsModified(GetControlRigBlueprint());
+		TArray<FRigElementKey> SelectedElements = ControlRigBP->HierarchyContainer.CurrentSelection();
+		for(const FRigElementKey& SelectedElement : SelectedElements)
+		{
+			ControlRigBP->HierarchyContainer.OnElementSelected.Broadcast(&ControlRigBP->HierarchyContainer, SelectedElement, true);
+		}
 
-	UControlRigSkeletalMeshComponent* EditorSkelComp = Cast<UControlRigSkeletalMeshComponent>(GetPersonaToolkit()->GetPreviewScene()->GetPreviewMeshComponent());
-	if (EditorSkelComp)
+		if (ControlRigBP->HierarchyContainer.GetIndex(RigElementInDetailPanel) == INDEX_NONE)
+		{
+			ClearDetailObject();
+		}
+	}
+	else
 	{
-		// restart animation 
-		EditorSkelComp->InitAnim(true);
-		UpdateControlRig();
+		ClearDetailObject();
 	}
-	CacheBoneNameList();
+}
 
-	// notification
-	FNotificationInfo Info(LOCTEXT("HierarchyChangeHelpMessage", "Hierarchy has been successfully modified. If you want to move the Bone, compile and turn off execution mode."));
-	Info.bFireAndForget = true;
-	Info.FadeOutDuration = 10.0f;
-	Info.ExpireDuration = 0.0f;
+void FControlRigEditor::OnRigElementAdded(FRigHierarchyContainer* Container, const FRigElementKey& InKey)
+{
+	OnHierarchyChanged();
+}
 
-	TSharedPtr<SNotificationItem> NotificationPtr = FSlateNotificationManager::Get().AddNotification(Info);
-	NotificationPtr->SetCompletionState(SNotificationItem::CS_Success);
+void FControlRigEditor::OnRigElementRemoved(FRigHierarchyContainer* Container, const FRigElementKey& InKey)
+{
+	OnHierarchyChanged();
 }
 
-void FControlRigEditor::OnBoneRenamed(const FName& OldName, const FName& NewName)
+void FControlRigEditor::OnRigElementRenamed(FRigHierarchyContainer* Container, ERigElementType ElementType, const FName& InOldName, const FName& InNewName)
 {
 	DECLARE_SCOPE_HIERARCHICAL_COUNTER_FUNC()
 
@@ -1569,6 +1751,8 @@ void FControlRigEditor::OnBoneRenamed(const FName& OldName, const FName& NewName
 			continue;
 		}
 
+		CacheNameLists();
+
 		for (UEdGraphNode* Node : RigGraph->Nodes)
 		{
 			UControlRigGraphNode* RigNode = Cast<UControlRigGraphNode>(Node);
@@ -1583,18 +1767,55 @@ void FControlRigEditor::OnBoneRenamed(const FName& OldName, const FName& NewName
 			{
 				for (TFieldIterator<UNameProperty> It(UnitStruct); It; ++It)
 				{
-					if (It->HasMetaData(UControlRig::BoneNameMetaName))
+					if (It->HasMetaData(UControlRig::BoneNameMetaName) && ElementType == ERigElementType::Bone)
 					{
 						FString PinName = FString::Printf(TEXT("%s.%s"), *UnitProperty->GetName(), *It->GetName());
 						UEdGraphPin* Pin = Node->FindPin(PinName, EEdGraphPinDirection::EGPD_Input);
 						if (Pin)
 						{
 							FName CurrentBone = FName(*Pin->GetDefaultAsString());
-							if (CurrentBone == OldName)
+							if (CurrentBone == InOldName)
 							{
-								const FScopedTransaction Transaction(NSLOCTEXT("ControlRigEditor", "ChangeBoneNamePinValue", "Change Bone Name Pin Value"));
-								Pin->Modify();
-								Pin->GetSchema()->TrySetDefaultValue(*Pin, NewName.ToString());
+								Pin->GetSchema()->TrySetDefaultValue(*Pin, InNewName.ToString());
+							}
+						}
+					}
+					if (It->HasMetaData(UControlRig::ControlNameMetaName) && ElementType == ERigElementType::Control)
+					{
+						FString PinName = FString::Printf(TEXT("%s.%s"), *UnitProperty->GetName(), *It->GetName());
+						UEdGraphPin* Pin = Node->FindPin(PinName, EEdGraphPinDirection::EGPD_Input);
+						if (Pin)
+						{
+							FName CurrentControl = FName(*Pin->GetDefaultAsString());
+							if (CurrentControl == InOldName)
+							{
+								Pin->GetSchema()->TrySetDefaultValue(*Pin, InNewName.ToString());
+							}
+						}
+					}
+					if (It->HasMetaData(UControlRig::SpaceNameMetaName) && ElementType == ERigElementType::Space)
+					{
+						FString PinName = FString::Printf(TEXT("%s.%s"), *UnitProperty->GetName(), *It->GetName());
+						UEdGraphPin* Pin = Node->FindPin(PinName, EEdGraphPinDirection::EGPD_Input);
+						if (Pin)
+						{
+							FName CurrentSpace = FName(*Pin->GetDefaultAsString());
+							if (CurrentSpace == InOldName)
+							{
+								Pin->GetSchema()->TrySetDefaultValue(*Pin, InNewName.ToString());
+							}
+						}
+					}
+					if (It->HasMetaData(UControlRig::CurveNameMetaName) && ElementType == ERigElementType::Curve)
+					{
+						FString PinName = FString::Printf(TEXT("%s.%s"), *UnitProperty->GetName(), *It->GetName());
+						UEdGraphPin* Pin = Node->FindPin(PinName, EEdGraphPinDirection::EGPD_Input);
+						if (Pin)
+						{
+							FName CurrentCurve = FName(*Pin->GetDefaultAsString());
+							if (CurrentCurve == InOldName)
+							{
+								Pin->GetSchema()->TrySetDefaultValue(*Pin, InNewName.ToString());
 							}
 						}
 					}
@@ -1602,10 +1823,103 @@ void FControlRigEditor::OnBoneRenamed(const FName& OldName, const FName& NewName
 			}
 		}
 
-		CacheBoneNameList();
+		CacheNameLists();
+	}
+}
+
+void FControlRigEditor::OnRigElementReparented(FRigHierarchyContainer* Container, const FRigElementKey& InKey, const FName& InOldParentName, const FName& InNewParentName)
+{
+	OnHierarchyChanged();
+}
+
+void FControlRigEditor::OnRigElementSelected(FRigHierarchyContainer* Container, const FRigElementKey& InKey, bool bSelected)
+{
+	UControlRigBlueprint* RigBlueprint = GetControlRigBlueprint();
+	if (RigBlueprint == nullptr)
+	{
+		return;
+	}
+
+	switch (InKey.Type)
+	{
+		case ERigElementType::Bone:
+		{
+			FRigBoneHierarchy& BoneHierarchy = RigBlueprint->HierarchyContainer.BoneHierarchy;
+			if (bSelected)
+			{
+				SetDetailStruct(InKey, MakeShareable(new FStructOnScope(FRigBone::StaticStruct(), (uint8*)&(BoneHierarchy[InKey.Name]))));
+			}
+			else
+			{
+				ClearDetailObject();
+			}
+
+			UControlRigSkeletalMeshComponent* EditorSkelComp = Cast<UControlRigSkeletalMeshComponent>(GetPersonaToolkit()->GetPreviewScene()->GetPreviewMeshComponent());
+			if (EditorSkelComp)
+			{
+				int32 Index = ControlRig->Hierarchy.BoneHierarchy.GetIndex(InKey.Name);
+				if (Index != INDEX_NONE)
+				{
+					if (bSelected)
+					{
+						EditorSkelComp->BonesOfInterest.Add(Index);
+					}
+					else
+					{
+						EditorSkelComp->BonesOfInterest.Remove(Index);
+					}
+				}
+			}
+
+			break;
+		}
+		case ERigElementType::Space:
+		{
+			if (bSelected)
+			{
+				FRigSpaceHierarchy& SpaceHierarchy = RigBlueprint->HierarchyContainer.SpaceHierarchy;
+				SetDetailStruct(InKey, MakeShareable(new FStructOnScope(FRigSpace::StaticStruct(), (uint8*)&(SpaceHierarchy[InKey.Name]))));
+			}
+			else
+			{
+				ClearDetailObject();
+			}
+			break;
+		}
+		case ERigElementType::Control:
+		{
+			if (bSelected)
+			{
+				FRigControlHierarchy& ControlHierarchy = RigBlueprint->HierarchyContainer.ControlHierarchy;
+				SetDetailStruct(InKey, MakeShareable(new FStructOnScope(FRigControl::StaticStruct(), (uint8*)&(ControlHierarchy[InKey.Name]))));
+			}
+			else
+			{
+				ClearDetailObject();
+			}
+			break;
+		}
+		case ERigElementType::Curve:
+		{
+			if (bSelected)
+			{
+				FRigCurveContainer& CurveContainer = RigBlueprint->HierarchyContainer.CurveContainer;
+				SetDetailStruct(InKey, MakeShareable(new FStructOnScope(FRigCurve::StaticStruct(), (uint8*)&(CurveContainer[InKey.Name]))));
+			}
+			else
+			{
+				ClearDetailObject();
+			}
+			break;
+		}
 	}
 }
 
+void FControlRigEditor::OnControlUISettingChanged(FRigHierarchyContainer* Container, const FRigElementKey& InKey)
+{
+	OnHierarchyChanged();
+}
+
 void FControlRigEditor::OnCurveContainerChanged()
 {
 	DECLARE_SCOPE_HIERARCHICAL_COUNTER_FUNC()
@@ -1621,7 +1935,7 @@ void FControlRigEditor::OnCurveContainerChanged()
 		EditorSkelComp->InitAnim(true);
 		UpdateControlRig();
 	}
-	CacheCurveNameList();
+	CacheNameLists();
 
 	// notification
 	FNotificationInfo Info(LOCTEXT("CurveContainerChangeHelpMessage", "CurveContainer has been successfully modified."));
@@ -1633,168 +1947,142 @@ void FControlRigEditor::OnCurveContainerChanged()
 	NotificationPtr->SetCompletionState(SNotificationItem::CS_Success);
 }
 
-void FControlRigEditor::OnCurveRenamed(const FName& OldName, const FName& NewName)
+void FControlRigEditor::OnGraphNodeDropToPerform(TSharedPtr<FGraphNodeDragDropOp> DragDropOp, UEdGraph* Graph, const FVector2D& NodePosition, const FVector2D& ScreenPosition)
 {
 	DECLARE_SCOPE_HIERARCHICAL_COUNTER_FUNC()
 
-	UControlRigBlueprint* Blueprint = GetControlRigBlueprint();
-	for (UEdGraph* Graph : Blueprint->UbergraphPages)
+	if (DragDropOp->IsOfType<FRigElementHierarchyDragDropOp>())
 	{
-		UControlRigGraph* RigGraph = Cast<UControlRigGraph>(Graph);
-		if (RigGraph == nullptr)
-		{
-			continue;
-		}
+		TSharedPtr<FRigElementHierarchyDragDropOp> RigHierarchyOp = StaticCastSharedPtr<FRigElementHierarchyDragDropOp>(DragDropOp);
 
-		for (UEdGraphNode* Node : RigGraph->Nodes)
+		TArray<FRigElementKey> DraggedKeys = RigHierarchyOp->GetElements();
+		uint8 DraggedTypes = 0;
+		for (const FRigElementKey& DraggedKey : DraggedKeys)
 		{
-			UControlRigGraphNode* RigNode = Cast<UControlRigGraphNode>(Node);
-			if (RigNode == nullptr)
-			{
-				continue;
-			}
-
-			UStructProperty* UnitProperty = RigNode->GetUnitProperty();
-			UStruct* UnitStruct = RigNode->GetUnitScriptStruct();
-			if (UnitProperty && UnitStruct)
-			{
-				for (TFieldIterator<UNameProperty> It(UnitStruct); It; ++It)
-				{
-					if (It->HasMetaData(UControlRig::CurveNameMetaName))
-					{
-						FString PinName = FString::Printf(TEXT("%s.%s"), *UnitProperty->GetName(), *It->GetName());
-						UEdGraphPin* Pin = Node->FindPin(PinName, EEdGraphPinDirection::EGPD_Input);
-						if (Pin)
-						{
-							FName CurrentCurve = FName(*Pin->GetDefaultAsString());
-							if (CurrentCurve == OldName)
-							{
-								const FScopedTransaction Transaction(NSLOCTEXT("ControlRigEditor", "ChangeCurveNamePinValue", "Change Curve Name Pin Value"));
-								Pin->Modify();
-								Pin->GetSchema()->TrySetDefaultValue(*Pin, NewName.ToString());
-							}
-						}
-					}
-				}
-			}
+			DraggedTypes = DraggedTypes | (uint8)DraggedKey.Type;
 		}
 
-		CacheCurveNameList();
-	}
-}
-
-void FControlRigEditor::OnGraphNodeDropToPerform(TSharedPtr<FGraphNodeDragDropOp> DragDropOp, UEdGraph* Graph, const FVector2D& NodePosition, const FVector2D& ScreenPosition)
-{
-	DECLARE_SCOPE_HIERARCHICAL_COUNTER_FUNC()
-
-	if (DragDropOp->IsOfType<FRigHierarchyDragDropOp>())
-	{
-		TSharedPtr<FRigHierarchyDragDropOp> RigHierarchyOp = StaticCastSharedPtr<FRigHierarchyDragDropOp>(DragDropOp);
-		TArray<FName> BoneNames = RigHierarchyOp->GetBoneNames();
-		if (BoneNames.Num() > 0 && FocusedGraphEdPtr.IsValid())
+		if (DraggedTypes != 0 && FocusedGraphEdPtr.IsValid())
 		{
 			FMenuBuilder MenuBuilder(true, NULL);
-			const FText BoneNameText = FText::FromString(RigHierarchyOp->GetJoinedBoneNames());
-
-			MenuBuilder.BeginSection("RigHierarchyDroppedOn", BoneNameText);
-
-			MenuBuilder.AddMenuEntry(
-				LOCTEXT("CreateGetBoneTransform", "Get Transform"),
-				LOCTEXT("CreateGetBoneTransformTooltip", "Getter for bone transform\n"),
-				FSlateIcon(),
-				FUIAction(
-					FExecuteAction::CreateSP(this, &FControlRigEditor::HandleMakeBoneGetterSetter, EBoneGetterSetterType_Transform, true, BoneNames, Graph, NodePosition),
-					FCanExecuteAction()
-				)
-			);
-			MenuBuilder.AddMenuEntry(
-				LOCTEXT("CreateSetBoneTransform", "Set Transform"),
-				LOCTEXT("CreateSetBoneTransformTooltip", "Setter for bone transform\n"),
-				FSlateIcon(),
-				FUIAction(
-					FExecuteAction::CreateSP(this, &FControlRigEditor::HandleMakeBoneGetterSetter, EBoneGetterSetterType_Transform, false, BoneNames, Graph, NodePosition),
-					FCanExecuteAction()
-				)
-			);
-
-			MenuBuilder.AddMenuSeparator();
-
-			MenuBuilder.AddMenuEntry(
-				LOCTEXT("CreateSetBoneRotation", "Set Rotation"),
-				LOCTEXT("CreateSetBoneRotationTooltip", "Setter for bone Rotation\n"),
-				FSlateIcon(),
-				FUIAction(
-					FExecuteAction::CreateSP(this, &FControlRigEditor::HandleMakeBoneGetterSetter, EBoneGetterSetterType_Rotation, false, BoneNames, Graph, NodePosition),
-					FCanExecuteAction()
-				)
-			);
+			const FText SectionText = FText::FromString(RigHierarchyOp->GetJoinedElementNames());
 
-			MenuBuilder.AddMenuEntry(
-				LOCTEXT("CreateSetBoneTranslation", "Set Translation"),
-				LOCTEXT("CreateSetBoneTranslationTooltip", "Setter for bone translation\n"),
-				FSlateIcon(),
-				FUIAction(
-					FExecuteAction::CreateSP(this, &FControlRigEditor::HandleMakeBoneGetterSetter, EBoneGetterSetterType_Translation, false, BoneNames, Graph, NodePosition),
-					FCanExecuteAction()
-				)
-			);
-
-			MenuBuilder.AddMenuEntry(
-				LOCTEXT("CreateSetBoneOffset", "Set Offset"),
-				LOCTEXT("CreateSetBoneOffsetTooltip", "Setter for bone offset\n"),
-				FSlateIcon(),
-				FUIAction(
-					FExecuteAction::CreateSP(this, &FControlRigEditor::HandleMakeBoneGetterSetter, EBoneGetterSetterType_Offset, false, BoneNames, Graph, NodePosition),
-					FCanExecuteAction()
-				)
-			);
+			MenuBuilder.BeginSection("RigHierarchyDroppedOn", SectionText);
 
-			MenuBuilder.AddMenuEntry(
-				LOCTEXT("CreateGetInitialBoneTransform", "Get Initial Bone Transform"),
-				LOCTEXT("CreateGetInitialBoneTransformTooltip", "Getter for initial bone transform\n"),
-				FSlateIcon(),
-				FUIAction(
-					FExecuteAction::CreateSP(this, &FControlRigEditor::HandleMakeBoneGetterSetter, EBoneGetterSetterType_Initial, true, BoneNames, Graph, NodePosition),
-					FCanExecuteAction()
-				)
-			);
-
-			MenuBuilder.AddMenuSeparator();
-
-			MenuBuilder.AddMenuEntry(
-				LOCTEXT("CreateGetBoneRelativeTransform", "Get Relative Transform"),
-				LOCTEXT("CreateGetBoneRelativeTransformTooltip", "Getter for bone relative transform\n"),
-				FSlateIcon(),
-				FUIAction(
-					FExecuteAction::CreateSP(this, &FControlRigEditor::HandleMakeBoneGetterSetter, EBoneGetterSetterType_Relative, true, BoneNames, Graph, NodePosition),
-					FCanExecuteAction()
-				)
-			);
-			MenuBuilder.AddMenuEntry(
-				LOCTEXT("CreateSetBoneRelativeTransform", "Set Relative Transform"),
-				LOCTEXT("CreateSetBoneRelativeTransformTooltip", "Setter for bone relative transform\n"),
-				FSlateIcon(),
-				FUIAction(
-					FExecuteAction::CreateSP(this, &FControlRigEditor::HandleMakeBoneGetterSetter, EBoneGetterSetterType_Relative, false, BoneNames, Graph, NodePosition),
-					FCanExecuteAction()
-				)
-			);
+			if ((DraggedTypes & (uint8)ERigElementType::Bone) != 0)
+			{
+				MenuBuilder.AddMenuEntry(
+					LOCTEXT("CreateGetBoneTransform", "Get Bone Transform"),
+					LOCTEXT("CreateGetBoneTransformTooltip", "Getter for bone transform\n"),
+					FSlateIcon(),
+					FUIAction(
+						FExecuteAction::CreateSP(this, &FControlRigEditor::HandleMakeElementGetterSetter, ERigElementGetterSetterType_Transform, true, DraggedKeys, Graph, NodePosition),
+						FCanExecuteAction()
+					)
+				);
+				MenuBuilder.AddMenuEntry(
+					LOCTEXT("CreateSetBoneTransform", "Set Bone Transform"),
+					LOCTEXT("CreateSetBoneTransformTooltip", "Setter for bone transform\n"),
+					FSlateIcon(),
+					FUIAction(
+						FExecuteAction::CreateSP(this, &FControlRigEditor::HandleMakeElementGetterSetter, ERigElementGetterSetterType_Transform, false, DraggedKeys, Graph, NodePosition),
+						FCanExecuteAction()
+					)
+				);
+
+				MenuBuilder.AddMenuSeparator();
+
+				MenuBuilder.AddMenuEntry(
+					LOCTEXT("CreateSetBoneRotation", "Set Bone Rotation"),
+					LOCTEXT("CreateSetBoneRotationTooltip", "Setter for bone Rotation\n"),
+					FSlateIcon(),
+					FUIAction(
+						FExecuteAction::CreateSP(this, &FControlRigEditor::HandleMakeElementGetterSetter, ERigElementGetterSetterType_Rotation, false, DraggedKeys, Graph, NodePosition),
+						FCanExecuteAction()
+					)
+				);
+
+				MenuBuilder.AddMenuEntry(
+					LOCTEXT("CreateSetBoneTranslation", "Set Bone Translation"),
+					LOCTEXT("CreateSetBoneTranslationTooltip", "Setter for bone translation\n"),
+					FSlateIcon(),
+					FUIAction(
+						FExecuteAction::CreateSP(this, &FControlRigEditor::HandleMakeElementGetterSetter, ERigElementGetterSetterType_Translation, false, DraggedKeys, Graph, NodePosition),
+						FCanExecuteAction()
+					)
+				);
+
+				MenuBuilder.AddMenuEntry(
+					LOCTEXT("CreateSetBoneOffset", "Set Bone Offset"),
+					LOCTEXT("CreateSetBoneOffsetTooltip", "Setter for bone offset\n"),
+					FSlateIcon(),
+					FUIAction(
+						FExecuteAction::CreateSP(this, &FControlRigEditor::HandleMakeElementGetterSetter, ERigElementGetterSetterType_Offset, false, DraggedKeys, Graph, NodePosition),
+						FCanExecuteAction()
+					)
+				);
+
+				MenuBuilder.AddMenuEntry(
+					LOCTEXT("CreateGetInitialBoneTransform", "Get Initial Bone Transform"),
+					LOCTEXT("CreateGetInitialBoneTransformTooltip", "Getter for initial bone transform\n"),
+					FSlateIcon(),
+					FUIAction(
+						FExecuteAction::CreateSP(this, &FControlRigEditor::HandleMakeElementGetterSetter, ERigElementGetterSetterType_Initial, true, DraggedKeys, Graph, NodePosition),
+						FCanExecuteAction()
+					)
+				);
+
+				MenuBuilder.AddMenuSeparator();
+
+				MenuBuilder.AddMenuEntry(
+					LOCTEXT("CreateGetBoneRelativeTransform", "Get Bone Relative Transform"),
+					LOCTEXT("CreateGetBoneRelativeTransformTooltip", "Getter for bone relative transform\n"),
+					FSlateIcon(),
+					FUIAction(
+						FExecuteAction::CreateSP(this, &FControlRigEditor::HandleMakeElementGetterSetter, ERigElementGetterSetterType_Relative, true, DraggedKeys, Graph, NodePosition),
+						FCanExecuteAction()
+					)
+				);
+				MenuBuilder.AddMenuEntry(
+					LOCTEXT("CreateSetBoneRelativeTransform", "Set Bone Relative Transform"),
+					LOCTEXT("CreateSetBoneRelativeTransformTooltip", "Setter for bone relative transform\n"),
+					FSlateIcon(),
+					FUIAction(
+						FExecuteAction::CreateSP(this, &FControlRigEditor::HandleMakeElementGetterSetter, ERigElementGetterSetterType_Relative, false, DraggedKeys, Graph, NodePosition),
+						FCanExecuteAction()
+					)
+				);
 
-			MenuBuilder.AddMenuSeparator();
+			}
 
-			MenuBuilder.AddMenuEntry(
-				LOCTEXT("CreateGetBoneName", "Get Bone Name"),
-				LOCTEXT("CreateGetBoneNameTooltip", "Getter for bone name\n"),
-				FSlateIcon(),
-				FUIAction(
-					FExecuteAction::CreateSP(this, &FControlRigEditor::HandleMakeBoneGetterSetter, EBoneGetterSetterType_Name, true, BoneNames, Graph, NodePosition),
-					FCanExecuteAction()
-				)
-			);
+			if ((DraggedTypes & (uint8)ERigElementType::Control) != 0)
+			{
+				MenuBuilder.AddMenuEntry(
+					LOCTEXT("CreateGetControlTransform", "Get Control Transform"),
+					LOCTEXT("CreateGetControlTransformTooltip", "Getter for control transform\n"),
+					FSlateIcon(),
+					FUIAction(
+						FExecuteAction::CreateSP(this, &FControlRigEditor::HandleMakeElementGetterSetter, ERigElementGetterSetterType_Transform, true, DraggedKeys, Graph, NodePosition),
+						FCanExecuteAction()
+					)
+				);
+			}
 
+			if ((DraggedTypes & (uint8)ERigElementType::Space) != 0)
+			{
+				MenuBuilder.AddMenuEntry(
+					LOCTEXT("CreateGetSpaceTransform", "Get Space Transform"),
+					LOCTEXT("CreateGetSpaceTransformTooltip", "Getter for space transform\n"),
+					FSlateIcon(),
+					FUIAction(
+						FExecuteAction::CreateSP(this, &FControlRigEditor::HandleMakeElementGetterSetter, ERigElementGetterSetterType_Transform, true, DraggedKeys, Graph, NodePosition),
+						FCanExecuteAction()
+					)
+				);
+			}
 
 			TSharedRef<SWidget> GraphEditorPanel = FocusedGraphEdPtr.Pin().ToSharedRef();
-			
+
 			// Show dialog to choose getter vs setter
 			FSlateApplication::Get().PushMenu(
 				GraphEditorPanel,
@@ -1809,69 +2097,123 @@ void FControlRigEditor::OnGraphNodeDropToPerform(TSharedPtr<FGraphNodeDragDropOp
 	}
 }
 
-void FControlRigEditor::HandleMakeBoneGetterSetter(EBoneGetterSetterType Type, bool bIsGetter, TArray<FName> BoneNames, UEdGraph* Graph, FVector2D NodePosition)
+void FControlRigEditor::HandleMakeElementGetterSetter(ERigElementGetterSetterType Type, bool bIsGetter, TArray<FRigElementKey> Keys, UEdGraph* Graph, FVector2D NodePosition)
 {
 	DECLARE_SCOPE_HIERARCHICAL_COUNTER_FUNC()
 
-	UStruct* StructTemplate = nullptr;
+	if (Keys.Num() == 0)
+	{
+		return;
+	}
 
-	if (bIsGetter)
+	UControlRigBlueprint* Blueprint = Cast<UControlRigBlueprint>(GetBlueprintObj());
+	if (Blueprint == nullptr)
 	{
-		switch (Type)
-		{
-			case EBoneGetterSetterType_Transform:
-			{
-				StructTemplate = FRigUnit_GetBoneTransform::StaticStruct();
-				break;
-			}
-			case EBoneGetterSetterType_Initial:
-			{
-				StructTemplate = FRigUnit_GetInitialBoneTransform::StaticStruct();
-				break;
-			}
-			case EBoneGetterSetterType_Relative:
-			{
-				StructTemplate = FRigUnit_GetRelativeBoneTransform::StaticStruct();
-				break;
-			}
-			case EBoneGetterSetterType_Name:
-			{
-				StructTemplate = FRigUnit_BoneName::StaticStruct();
-				break;
-			}
-			default:
-			{
-				break;
-			}
-		}
+		return;
 	}
-	else
+	if (Blueprint->ModelController == nullptr)
+	{
+		return;
+	}
+
+	Blueprint->ModelController->OpenUndoBracket(TEXT("Adding Nodes from Hierarchy"));
+
+	TArray<FName> NewNodeNames;
+
+	for (const FRigElementKey& Key : Keys)
 	{
-		switch (Type)
+		UStruct* StructTemplate = nullptr;
+
+		FName PinName;
+		switch (Key.Type)
 		{
-			case EBoneGetterSetterType_Transform:
-			{
-				StructTemplate = FRigUnit_SetBoneTransform::StaticStruct();
-				break;
-			}
-			case EBoneGetterSetterType_Relative:
-			{
-				StructTemplate = FRigUnit_SetRelativeBoneTransform::StaticStruct();
-				break;
-			}
-			case EBoneGetterSetterType_Rotation:
+			case ERigElementType::Bone:
 			{
-				StructTemplate = FRigUnit_SetBoneRotation::StaticStruct();
+				PinName = TEXT("Bone");
+
+				if (bIsGetter)
+				{
+					switch (Type)
+					{
+						case ERigElementGetterSetterType_Transform:
+						{
+							StructTemplate = FRigUnit_GetBoneTransform::StaticStruct();
+							break;
+						}
+						case ERigElementGetterSetterType_Initial:
+						{
+							StructTemplate = FRigUnit_GetInitialBoneTransform::StaticStruct();
+							break;
+						}
+						case ERigElementGetterSetterType_Relative:
+						{
+							StructTemplate = FRigUnit_GetRelativeBoneTransform::StaticStruct();
+							break;
+						}
+						case ERigElementGetterSetterType_Name:
+						{
+							StructTemplate = FRigUnit_BoneName::StaticStruct();
+							break;
+						}
+						default:
+						{
+							break;
+						}
+					}
+				}
+				else
+				{
+					switch (Type)
+					{
+						case ERigElementGetterSetterType_Transform:
+						{
+							StructTemplate = FRigUnit_SetBoneTransform::StaticStruct();
+							break;
+						}
+						case ERigElementGetterSetterType_Relative:
+						{
+							StructTemplate = FRigUnit_SetRelativeBoneTransform::StaticStruct();
+							break;
+						}
+						case ERigElementGetterSetterType_Rotation:
+						{
+							StructTemplate = FRigUnit_SetBoneRotation::StaticStruct();
+							break;
+						}
+						case ERigElementGetterSetterType_Translation:
+						{
+							StructTemplate = FRigUnit_SetBoneTranslation::StaticStruct();
+							break;
+						}
+						case ERigElementGetterSetterType_Offset:
+						{
+							StructTemplate = FRigUnit_AddBoneTransform::StaticStruct();
+							break;
+						}
+						default:
+						{
+							break;
+						}
+					}
+				}
 				break;
 			}
-			case EBoneGetterSetterType_Translation:
+			case ERigElementType::Control:
 			{
-				StructTemplate = FRigUnit_SetBoneTranslation::StaticStruct();
+				PinName = TEXT("Control");
+				if (bIsGetter)
+				{
+					StructTemplate = FRigUnit_GetControlTransform::StaticStruct();
+				}
 				break;
 			}
-			case EBoneGetterSetterType_Offset:
+			case ERigElementType::Space:
 			{
-				StructTemplate = FRigUnit_AddBoneTransform::StaticStruct();
+				PinName = TEXT("Space");
+				if (bIsGetter)
+				{
+					StructTemplate = FRigUnit_GetSpaceTransform::StaticStruct();
+				}
 				break;
 			}
 			default:
@@ -1879,30 +2221,14 @@ void FControlRigEditor::HandleMakeBoneGetterSetter(EBoneGetterSetterType Type, b
 				break;
 			}
 		}
-	}
 
-	if (StructTemplate == nullptr)
-	{
-		return;
-	}
-
-	UControlRigBlueprint* Blueprint = Cast<UControlRigBlueprint>(GetBlueprintObj());
-	if (Blueprint == nullptr)
-	{
-		return;
-	}
-	if (Blueprint->ModelController == nullptr)
-	{
-		return;
-	}
-
-	Blueprint->ModelController->OpenUndoBracket(TEXT("Adding Nodes from Hierarchy"));
+		if (StructTemplate == nullptr)
+		{
+			continue;
+		}
 
-	TArray<FName> NewNodeNames;
-	for (const FName& BoneName : BoneNames)
-	{
 		FVector2D NodePositionIncrement(0.f, 120.f);
-		if(!bIsGetter)
+		if (!bIsGetter)
 		{
 			NodePositionIncrement = FVector2D(380.f, 0.f);
 		}
@@ -1914,19 +2240,22 @@ void FControlRigEditor::HandleMakeBoneGetterSetter(EBoneGetterSetterType Type, b
 			if (Node)
 			{
 				NewNodeNames.Add(Node->Name);
-				Blueprint->ModelController->SetPinDefaultValueName(Node->Name, TEXT("Bone"), BoneName, true);
+				Blueprint->ModelController->SetPinDefaultValueName(Node->Name, PinName, Key.Name, true);
 			}
 		}
 
 		NodePosition += NodePositionIncrement;
 	}
 
-	Blueprint->ModelController->CloseUndoBracket();
-
 	if (NewNodeNames.Num() > 0)
 	{
+		Blueprint->ModelController->CloseUndoBracket();
 		Blueprint->ModelController->SetSelection(NewNodeNames);
 	}
+	else
+	{
+		Blueprint->ModelController->CancelUndoBracket();
+	}
 }
 
 void FControlRigEditor::UpdateGraphCompilerErrors()
@@ -2060,10 +2389,10 @@ void FControlRigEditor::DumpUnitTestCode()
 		// dump the hierarchy
 		if (ControlRig)
 		{
-			const FRigHierarchy& Hierarchy = ControlRig->GetBaseHierarchy();
+			const FRigBoneHierarchy& Hierarchy = ControlRig->GetBoneHierarchy();
 			if (Hierarchy.Bones.Num() > 0)
 			{
-				Code.Add(TEXT("FRigHierarchy& Hierarchy = Rig->GetBaseHierarchy();"));
+				Code.Add(TEXT("FRigBoneHierarchy& Hierarchy = Rig->GetBoneHierarchy();"));
 			}
 			for (const FRigBone& Bone : Hierarchy.Bones)
 			{
@@ -2160,4 +2489,21 @@ void FControlRigEditor::DumpUnitTestCode()
 	*/
 }
 
+void FControlRigEditor::HandleOnViewportContextMenuDelegate(class FMenuBuilder& MenuBuilder)
+{
+	if (OnViewportContextMenuDelegate.IsBound())
+	{
+		OnViewportContextMenuDelegate.Execute(MenuBuilder);
+	}
+}
+
+TSharedPtr<FUICommandList> FControlRigEditor::HandleOnViewportContextMenuCommandsDelegate()
+{
+	if (OnViewportContextMenuCommandsDelegate.IsBound())
+	{
+		return OnViewportContextMenuCommandsDelegate.Execute();
+	}
+	return TSharedPtr<FUICommandList>();
+}
+
 #undef LOCTEXT_NAMESPACE
diff --git a/Engine/Plugins/Experimental/ControlRig/Source/ControlRigEditor/Private/Editor/ControlRigEditor.h b/Engine/Plugins/Experimental/ControlRig/Source/ControlRigEditor/Private/Editor/ControlRigEditor.h
--- a/Engine/Plugins/Experimental/ControlRig/Source/ControlRigEditor/Private/Editor/ControlRigEditor.h
+++ b/Engine/Plugins/Experimental/ControlRig/Source/ControlRigEditor/Private/Editor/ControlRigEditor.h
@@ -8,8 +8,9 @@
 #include "DragAndDrop/GraphNodeDragDropOp.h"
 #include "ControlRigDefines.h"
 #include "ControlRigLog.h"
-#include "Drawing/ControlRigDrawInterface.h"
 #include "ControlRigModel.h"
+#include "IPersonaViewport.h"
+#include "Framework/MultiBox/MultiBoxBuilder.h"
 
 class UControlRigBlueprint;
 class IPersonaToolkit;
@@ -91,7 +92,7 @@ public:
 
 	void SetDetailObject(UObject* Obj);
 
-	void SetDetailStruct(TSharedPtr<FStructOnScope> StructToDisplay);
+	void SetDetailStruct(const FRigElementKey& InElement, TSharedPtr<FStructOnScope> StructToDisplay);
 
 	void ClearDetailObject();
 
@@ -102,20 +103,31 @@ public:
 	TSharedRef<SBorder> GetToolbox() { return Toolbox.ToSharedRef(); }
 
 	/** Get the edit mode */
-	FControlRigEditorEditMode& GetEditMode() { return *static_cast<FControlRigEditorEditMode*>(GetAssetEditorModeManager()->GetActiveMode(FControlRigEditorEditMode::ModeName)); }
+	FControlRigEditorEditMode* GetEditMode() const { return static_cast<FControlRigEditorEditMode*>(GetAssetEditorModeManager()->GetActiveMode(FControlRigEditorEditMode::ModeName)); }
 
-	void SelectBone(const FName& InBone);
 	// this changes everytime you compile, so don't cache it expecting it will last. 
 	UControlRig* GetInstanceRig() const { return ControlRig;  }
-	// restart animation 
-	void OnHierarchyChanged();
-	void OnBoneRenamed(const FName& OldName, const FName& NewName);
+
 	void OnCurveContainerChanged();
-	void OnCurveRenamed(const FName& OldName, const FName& NewName);
+
+	void OnRigElementAdded(FRigHierarchyContainer* Container, const FRigElementKey& InKey);
+	void OnRigElementRemoved(FRigHierarchyContainer* Container, const FRigElementKey& InKey);
+	void OnRigElementRenamed(FRigHierarchyContainer* Container, ERigElementType ElementType, const FName& InOldName, const FName& InNewName);
+	void OnRigElementReparented(FRigHierarchyContainer* Container, const FRigElementKey& InKey, const FName& InOldParentName, const FName& InNewParentName);
+	void OnRigElementSelected(FRigHierarchyContainer* Container, const FRigElementKey& InKey, bool bSelected);
+	void OnControlUISettingChanged(FRigHierarchyContainer* Container, const FRigElementKey& InKey);
 
 	void OnGraphNodeDropToPerform(TSharedPtr<FGraphNodeDragDropOp> DragDropOp, UEdGraph* Graph, const FVector2D& NodePosition, const FVector2D& ScreenPosition);
 
+	FPersonaViewportKeyDownDelegate& GetKeyDownDelegate() { return OnKeyDownDelegate; }
+	FNewMenuDelegate& OnViewportContextMenu() { return OnViewportContextMenuDelegate; }
+	FNewMenuCommandsDelegate& OnViewportContextMenuCommands() { return OnViewportContextMenuCommandsDelegate; }
+
 protected:
+
+	void OnHierarchyChanged();
+	void OnControlsSettingsChanged();
+
 	// FBlueprintEditor Interface
 	virtual void CreateDefaultCommands() override;
 	virtual void OnCreateGraphEditorCommands(TSharedPtr<FUICommandList> GraphEditorCommandsList);
@@ -175,18 +187,17 @@ private:
 	/** Handle preview scene setup */
 	void HandlePreviewSceneCreated(const TSharedRef<IPersonaPreviewScene>& InPersonaPreviewScene);
 	void HandleViewportCreated(const TSharedRef<class IPersonaViewport>& InViewport);
+	TOptional<float> GetToolbarAxesScale() const;
+	void OnToolbarAxesScaleChanged(float InValue);
 
-	/** Handle switching skeletal meshes */
+		/** Handle switching skeletal meshes */
 	void HandlePreviewMeshChanged(USkeletalMesh* InOldSkeletalMesh, USkeletalMesh* InNewSkeletalMesh);
 
 	/** Push a newly compiled/opened control rig to the edit mode */
 	void UpdateControlRig();
 
-	/** Update the bone name list for use in bone name combo boxes */
-	void CacheBoneNameList();
-
-	/** Update the curve name list for use in curve name combo boxes */
-	void CacheCurveNameList();
+	/** Update the name lists for use in name combo boxes */
+	void CacheNameLists();
 
 	/** Rebind our anim instance to the preview's skeletal mesh component */
 	void RebindToSkeletalMeshComponent();
@@ -198,18 +209,18 @@ private:
 	void ToggleExecuteGraph();
 	bool IsExecuteGraphOn() const;
 
-	enum EBoneGetterSetterType
+	enum ERigElementGetterSetterType
 	{
-		EBoneGetterSetterType_Transform,
-		EBoneGetterSetterType_Rotation,
-		EBoneGetterSetterType_Translation,
-		EBoneGetterSetterType_Initial,
-		EBoneGetterSetterType_Relative,
-		EBoneGetterSetterType_Offset,
-		EBoneGetterSetterType_Name
+		ERigElementGetterSetterType_Transform,
+		ERigElementGetterSetterType_Rotation,
+		ERigElementGetterSetterType_Translation,
+		ERigElementGetterSetterType_Initial,
+		ERigElementGetterSetterType_Relative,
+		ERigElementGetterSetterType_Offset,
+		ERigElementGetterSetterType_Name
 	};
 
-	 void HandleMakeBoneGetterSetter(EBoneGetterSetterType Type, bool bIsGetter, TArray<FName> BoneNames, UEdGraph* Graph, FVector2D NodePosition);
+	 void HandleMakeElementGetterSetter(ERigElementGetterSetterType Type, bool bIsGetter, TArray<FRigElementKey> Keys, UEdGraph* Graph, FVector2D NodePosition);
 
 protected:
 
@@ -234,28 +245,35 @@ protected:
 	/** preview scene */
 	TSharedPtr<IPersonaPreviewScene> PreviewScene;
 
-	/** Bone Selection related */
-	FTransform GetBoneTransform(const FName& InBone, bool bLocal) const;
-	void SetBoneTransform(const FName& InBone, const FTransform& InTransform);
+	/** Delegate to deal with key down evens in the viewport / editor */
+	FPersonaViewportKeyDownDelegate OnKeyDownDelegate;
 
-	/** delegate for changing property */
-	void OnFinishedChangingProperties(const FPropertyChangedEvent& PropertyChangedEvent);
+	/** Delgate to build the context menu for the viewport */
+	FNewMenuDelegate OnViewportContextMenuDelegate;
+	void HandleOnViewportContextMenuDelegate(class FMenuBuilder& MenuBuilder);
+	FNewMenuCommandsDelegate OnViewportContextMenuCommandsDelegate;
+	TSharedPtr<FUICommandList> HandleOnViewportContextMenuCommandsDelegate();
 
-	/** Selected Bone from hierarchy tree */
-	FName SelectedBone;
+	/** Bone Selection related */
+	FTransform GetRigElementTransform(const FRigElementKey& InElement, bool bLocal) const;
+	void SetRigElementTransform(const FRigElementKey& InElement, const FTransform& InTransform, bool bLocal);
+	
+	/** delegate for changing property */
+	virtual void OnFinishedChangingProperties(const FPropertyChangedEvent& PropertyChangedEvent) override;
 
 	bool bControlRigEditorInitialized;
 	bool bIsSelecting;
 	bool bIsSettingObjectBeingDebugged;
+	FRigElementKey RigElementInDetailPanel;
+
+	/** Currently executing ControlRig or not - later maybe this will change to enum for whatever different mode*/
+	bool bExecutionControlRig;
 
 	/** The log to use for errors resulting from the init phase of the units */
 	FControlRigLog ControlRigLog;
 	/** Once the log is collected update the graph */
 	void UpdateGraphCompilerErrors();
 
-	/** The draw interface to use for the control rig */
-	FControlRigDrawInterface DrawInterface;
-
 	/** This can be used to enable dumping of a unit test */
 	void DumpUnitTestCode();
 
diff --git a/Engine/Plugins/Experimental/ControlRig/Source/ControlRigEditor/Private/Editor/ControlRigEditorEditMode.cpp b/Engine/Plugins/Experimental/ControlRig/Source/ControlRigEditor/Private/Editor/ControlRigEditorEditMode.cpp
--- a/Engine/Plugins/Experimental/ControlRig/Source/ControlRigEditor/Private/Editor/ControlRigEditorEditMode.cpp
+++ b/Engine/Plugins/Experimental/ControlRig/Source/ControlRigEditor/Private/Editor/ControlRigEditorEditMode.cpp
@@ -3,6 +3,7 @@
 #include "ControlRigEditorEditMode.h"
 #include "IPersonaPreviewScene.h"
 #include "AssetEditorModeManager.h"
+#include "ControlRigGizmoActor.h"
 
 FName FControlRigEditorEditMode::ModeName("EditMode.ControlRigEditor");
 
@@ -10,20 +11,11 @@ bool FControlRigEditorEditMode::GetCameraTarget(FSphere& OutTarget) const
 {
 	FBox Box(ForceInit);
 
-	for(const FControlUnitProxy& UnitProxy : ControlUnits)
+	for (int32 Index = 0; Index < SelectedRigElements.Num(); ++Index)
 	{
-		if(UnitProxy.IsSelected() && UnitProxy.Control)
+		if (SelectedRigElements[Index].Type == ERigElementType::Bone || SelectedRigElements[Index].Type == ERigElementType::Control)
 		{
-			FBox ActorBox = UnitProxy.Control->GetComponentsBoundingBox(true);
-			Box += ActorBox;
-		}
-	}
-
-	if (AreBoneSelected())
-	{
-		for (int32 Index = 0; Index < SelectedBones.Num(); ++Index)
-		{
-			FTransform Transform = OnGetBoneTransformDelegate.Execute(SelectedBones[Index], false);
+			FTransform Transform = OnGetRigElementTransformDelegate.Execute(SelectedRigElements[Index], false);
 			Box += Transform.GetLocation();
 		}
 	}
@@ -31,7 +23,7 @@ bool FControlRigEditorEditMode::GetCameraTarget(FSphere& OutTarget) const
 	if(Box.IsValid)
 	{
 		OutTarget.Center = Box.GetCenter();
-		OutTarget.W = Box.GetExtent().GetAbsMax();
+		OutTarget.W = Box.GetExtent().GetAbsMax() * 1.25f;
 		return true;
 	}
 
diff --git a/Engine/Plugins/Experimental/ControlRig/Source/ControlRigEditor/Private/Editor/ControlRigSkeletalMeshComponent.cpp b/Engine/Plugins/Experimental/ControlRig/Source/ControlRigEditor/Private/Editor/ControlRigSkeletalMeshComponent.cpp
--- a/Engine/Plugins/Experimental/ControlRig/Source/ControlRigEditor/Private/Editor/ControlRigSkeletalMeshComponent.cpp
+++ b/Engine/Plugins/Experimental/ControlRig/Source/ControlRigEditor/Private/Editor/ControlRigSkeletalMeshComponent.cpp
@@ -20,30 +20,6 @@ void UControlRigSkeletalMeshComponent::InitAnim(bool bForceReinit)
 	RebuildDebugDrawSkeleton();
 }
 
-void UControlRigSkeletalMeshComponent::ShowReferencePose(bool bRefPose)
-{
-	UControlRigSequencerAnimInstance* ControlRigInstance = Cast<UControlRigSequencerAnimInstance>(GetAnimInstance());
-
-	if (ControlRigInstance && ControlRigInstance->CachedControlRig.IsValid())
-	{
-		UControlRig* ControlRig = ControlRigInstance->CachedControlRig.Get();
-		ControlRig->bExecutionOn = !bRefPose;
-		ControlRigInstance->SetAnimationAsset(nullptr);
-	}
-}
-
-bool UControlRigSkeletalMeshComponent::IsReferencePoseShown() const
-{
-	UControlRigSequencerAnimInstance* ControlRigInstance = Cast<UControlRigSequencerAnimInstance>(GetAnimInstance());
-
-	if (ControlRigInstance && ControlRigInstance->CachedControlRig.IsValid())
-	{
-		UControlRig* ControlRig = ControlRigInstance->CachedControlRig.Get();
-		return !ControlRig->bExecutionOn;
-	}
-
-	return false;
-}
 
 void UControlRigSkeletalMeshComponent::SetCustomDefaultPose()
 {
@@ -58,7 +34,7 @@ void UControlRigSkeletalMeshComponent::RebuildDebugDrawSkeleton()
 	{
 		UControlRig* ControlRig = ControlRigInstance->CachedControlRig.Get();
 		// just copy it because this is not thread safe
-		const FRigHierarchy BaseHiearchy = ControlRig->GetBaseHierarchy();
+		const FRigBoneHierarchy& BaseHiearchy = ControlRig->GetBoneHierarchy();
 
 		DebugDrawSkeleton.Empty();
 		DebugDrawBones.Reset();
@@ -66,11 +42,11 @@ void UControlRigSkeletalMeshComponent::RebuildDebugDrawSkeleton()
 		// create ref modifier
  		FReferenceSkeletonModifier RefSkelModifier(DebugDrawSkeleton, nullptr);
  
- 		for (int32 Index = 0; Index < BaseHiearchy.GetNum(); Index++)
+ 		for (int32 Index = 0; Index < BaseHiearchy.Num(); Index++)
  		{
 			FMeshBoneInfo NewMeshBoneInfo;
 			NewMeshBoneInfo.Name = BaseHiearchy.GetName(Index);
-			NewMeshBoneInfo.ParentIndex = BaseHiearchy.GetParentIndex(Index);
+			NewMeshBoneInfo.ParentIndex = BaseHiearchy[Index].ParentIndex;
 			// give ref pose here
 			RefSkelModifier.Add(NewMeshBoneInfo, BaseHiearchy.GetInitialTransform(Index));
 
@@ -87,7 +63,7 @@ FTransform UControlRigSkeletalMeshComponent::GetDrawTransform(int32 BoneIndex) c
 	{
 		UControlRig* ControlRig = ControlRigInstance->CachedControlRig.Get();
 		// just copy it because this is not thread safe
-		const FRigHierarchy& BaseHiearchy = ControlRig->GetBaseHierarchy();
+		const FRigBoneHierarchy& BaseHiearchy = ControlRig->GetBoneHierarchy();
 		return BaseHiearchy.GetGlobalTransform(BoneIndex);
 	}
 
diff --git a/Engine/Plugins/Experimental/ControlRig/Source/ControlRigEditor/Private/Editor/ControlRigSkeletalMeshComponent.h b/Engine/Plugins/Experimental/ControlRig/Source/ControlRigEditor/Private/Editor/ControlRigSkeletalMeshComponent.h
--- a/Engine/Plugins/Experimental/ControlRig/Source/ControlRigEditor/Private/Editor/ControlRigSkeletalMeshComponent.h
+++ b/Engine/Plugins/Experimental/ControlRig/Source/ControlRigEditor/Private/Editor/ControlRigSkeletalMeshComponent.h
@@ -14,8 +14,6 @@ class UControlRigSkeletalMeshComponent : public UDebugSkelMeshComponent
 	virtual void InitAnim(bool bForceReinit) override;
 
 	// BEGIN UDebugSkeletalMeshComponent interface
-	virtual void ShowReferencePose(bool bRefPose) override;
-	virtual bool IsReferencePoseShown() const override;
 	virtual void SetCustomDefaultPose() override;
 	virtual const FReferenceSkeleton& GetReferenceSkeleton() const override
 	{
diff --git a/Engine/Plugins/Experimental/ControlRig/Source/ControlRigEditor/Private/Editor/RigHierarchyTabSummoner.cpp b/Engine/Plugins/Experimental/ControlRig/Source/ControlRigEditor/Private/Editor/RigHierarchyTabSummoner.cpp
--- a/Engine/Plugins/Experimental/ControlRig/Source/ControlRigEditor/Private/Editor/RigHierarchyTabSummoner.cpp
+++ b/Engine/Plugins/Experimental/ControlRig/Source/ControlRigEditor/Private/Editor/RigHierarchyTabSummoner.cpp
@@ -13,12 +13,12 @@ FRigHierarchyTabSummoner::FRigHierarchyTabSummoner(const TSharedRef<FControlRigE
 	: FWorkflowTabFactory(TabID, InControlRigEditor)
 	, ControlRigEditor(InControlRigEditor)
 {
-	TabLabel = LOCTEXT("RigHierarchyTabLabel", "Hierarchy");
+	TabLabel = LOCTEXT("RigHierarchyTabLabel", "Rig Hierarchy");
 	TabIcon = FSlateIcon(FControlRigEditorStyle::Get().GetStyleSetName(), "ControlRig.TabIcon");
 
 	bIsSingleton = true;
 
-	ViewMenuDescription = LOCTEXT("RigHierarchy_ViewMenu_Desc", "Hierarchy");
+	ViewMenuDescription = LOCTEXT("RigHierarchy_ViewMenu_Desc", "Rig Hierarchy");
 	ViewMenuTooltip = LOCTEXT("RigHierarchy_ViewMenu_ToolTip", "Show the Rig Hierarchy tab");
 }
 
diff --git a/Engine/Plugins/Experimental/ControlRig/Source/ControlRigEditor/Private/Editor/SControlRigGizmoNameList.cpp b/Engine/Plugins/Experimental/ControlRig/Source/ControlRigEditor/Private/Editor/SControlRigGizmoNameList.cpp
new file mode 100644
--- /dev/null
+++ b/Engine/Plugins/Experimental/ControlRig/Source/ControlRigEditor/Private/Editor/SControlRigGizmoNameList.cpp
@@ -0,0 +1,110 @@
+// Copyright 1998-2019 Epic Games, Inc. All Rights Reserved.
+
+
+#include "Editor/SControlRigGizmoNameList.h"
+#include "Graph/SGraphPinNameListValueWidget.h"
+#include "Widgets/Input/SEditableTextBox.h"
+#include "EdGraphSchema_K2.h"
+#include "ScopedTransaction.h"
+#include "DetailLayoutBuilder.h"
+#include "ControlRigBlueprint.h"
+
+void SControlRigGizmoNameList::Construct(const FArguments& InArgs, FRigControl* InControl, UControlRigBlueprint* InBlueprint)
+{
+	this->OnGetNameListContent = InArgs._OnGetNameListContent;
+	this->ControlKey = InControl->GetElementKey();
+	this->Blueprint = InBlueprint;
+
+	SBox::Construct(SBox::FArguments());
+
+	TSharedPtr<FString> InitialSelected;
+	for (TSharedPtr<FString> Item : GetNameList())
+	{
+		if (Item->Equals(GetNameListText().ToString()))
+		{
+			InitialSelected = Item;
+		}
+	}
+
+	SetContent(
+		SNew(SBox)
+		.MinDesiredWidth(150)
+		.MaxDesiredWidth(400)
+		[
+			SAssignNew(NameListComboBox, SGraphPinNameListValueWidget)
+			.OptionsSource(&GetNameList())
+			.OnGenerateWidget(this, &SControlRigGizmoNameList::MakeNameListItemWidget)
+			.OnSelectionChanged(this, &SControlRigGizmoNameList::OnNameListChanged)
+			.OnComboBoxOpening(this, &SControlRigGizmoNameList::OnNameListComboBox)
+			.InitiallySelectedItem(InitialSelected)
+			.Content()
+			[
+				SNew(STextBlock)
+				.Text(this, &SControlRigGizmoNameList::GetNameListText)
+			]
+		]
+	);
+}
+
+const TArray<TSharedPtr<FString>>& SControlRigGizmoNameList::GetNameList() const
+{
+	if (OnGetNameListContent.IsBound())
+	{
+		return OnGetNameListContent.Execute();
+	}
+	return EmptyList;
+}
+
+FText SControlRigGizmoNameList::GetNameListText() const
+{
+	int32 ControlIndex = Blueprint->HierarchyContainer.GetIndex(ControlKey);
+	if (ControlIndex != INDEX_NONE)
+	{
+		return FText::FromName(Blueprint->HierarchyContainer.ControlHierarchy[ControlIndex].GizmoName);
+	}
+	return FText();
+}
+
+void SControlRigGizmoNameList::SetNameListText(const FText& NewTypeInValue, ETextCommit::Type /*CommitInfo*/)
+{
+	int32 ControlIndex = Blueprint->HierarchyContainer.GetIndex(ControlKey);
+	if (ControlIndex != INDEX_NONE)
+	{
+		FName NewName = *NewTypeInValue.ToString();
+		FRigControl& Control = Blueprint->HierarchyContainer.ControlHierarchy[ControlIndex];
+		if (Control.GizmoName != NewName)
+		{
+			const FScopedTransaction Transaction(NSLOCTEXT("ControlRigEditor", "ChangeGizmoName", "Change Gizmo Name"));
+			Blueprint->Modify();
+			Control.GizmoName = NewName;
+			Blueprint->HierarchyContainer.ControlHierarchy.OnControlUISettingsChanged.Broadcast(&Blueprint->HierarchyContainer, ControlKey);
+		}
+	}
+}
+
+TSharedRef<SWidget> SControlRigGizmoNameList::MakeNameListItemWidget(TSharedPtr<FString> InItem)
+{
+	return 	SNew(STextBlock).Text(FText::FromString(*InItem));// .Font(FEditorStyle::GetFontStyle(TEXT("PropertyWindow.NormalFont")));
+}
+
+void SControlRigGizmoNameList::OnNameListChanged(TSharedPtr<FString> NewSelection, ESelectInfo::Type SelectInfo)
+{
+	if (SelectInfo != ESelectInfo::Direct)
+	{
+		FString NewValue = *NewSelection.Get();
+		SetNameListText(FText::FromString(NewValue), ETextCommit::OnEnter);
+	}
+}
+
+void SControlRigGizmoNameList::OnNameListComboBox()
+{
+	TSharedPtr<FString> CurrentlySelected;
+	for (TSharedPtr<FString> Item : GetNameList())
+	{
+		if (Item->Equals(GetNameListText().ToString()))
+		{
+			CurrentlySelected = Item;
+		}
+	}
+	NameListComboBox->SetSelectedItem(CurrentlySelected);
+}
diff --git a/Engine/Plugins/Experimental/ControlRig/Source/ControlRigEditor/Private/Editor/SControlRigGizmoNameList.h b/Engine/Plugins/Experimental/ControlRig/Source/ControlRigEditor/Private/Editor/SControlRigGizmoNameList.h
new file mode 100644
--- /dev/null
+++ b/Engine/Plugins/Experimental/ControlRig/Source/ControlRigEditor/Private/Editor/SControlRigGizmoNameList.h
@@ -0,0 +1,43 @@
+// Copyright 1998-2019 Epic Games, Inc. All Rights Reserved.
+
+#pragma once
+
+#include "CoreMinimal.h"
+#include "Widgets/DeclarativeSyntaxSupport.h"
+#include "Widgets/SWidget.h"
+#include "Widgets/Layout/SBox.h"
+#include "Graph/SGraphPinNameListValueWidget.h"
+#include "Rigs/RigHierarchyContainer.h"
+#include "ControlRigBlueprint.h"
+
+class SControlRigGizmoNameList : public SBox
+{
+public:
+
+	DECLARE_DELEGATE_RetVal( const TArray<TSharedPtr<FString>>&, FOnGetNameListContent );
+
+	SLATE_BEGIN_ARGS(SControlRigGizmoNameList){}
+
+		SLATE_EVENT(FOnGetNameListContent, OnGetNameListContent)
+
+	SLATE_END_ARGS()
+
+	void Construct(const FArguments& InArgs, FRigControl* InControl, UControlRigBlueprint* InBlueprint);
+
+protected:
+
+	const TArray<TSharedPtr<FString>>& GetNameList() const;
+	FText GetNameListText() const;
+	virtual void SetNameListText(const FText& NewTypeInValue, ETextCommit::Type CommitInfo);
+
+	TSharedRef<SWidget> MakeNameListItemWidget(TSharedPtr<FString> InItem);
+	void OnNameListChanged(TSharedPtr<FString> NewSelection, ESelectInfo::Type SelectInfo);
+	void OnNameListComboBox();
+
+	FOnGetNameListContent OnGetNameListContent;
+	TSharedPtr<SGraphPinNameListValueWidget> NameListComboBox;
+	TArray<TSharedPtr<FString>> EmptyList;
+
+	FRigElementKey ControlKey;
+	UControlRigBlueprint* Blueprint;
+};
