Copying //UE4/Dev-Anim to Dev-Main (//UE4/Dev-Main) @ 8782600

diff --git a/Engine/Plugins/Experimental/ControlRig/Source/ControlRigEditor/Private/Editor/SRigHierarchy.cpp b/Engine/Plugins/Experimental/ControlRig/Source/ControlRigEditor/Private/Editor/SRigHierarchy.cpp
--- a/Engine/Plugins/Experimental/ControlRig/Source/ControlRigEditor/Private/Editor/SRigHierarchy.cpp
+++ b/Engine/Plugins/Experimental/ControlRig/Source/ControlRigEditor/Private/Editor/SRigHierarchy.cpp
@@ -29,110 +29,141 @@
 #include "HelperUtil.h"
 #include "Widgets/Text/SInlineEditableTextBlock.h"
 #include "ControlRig.h"
+#include "ControlRigEditorStyle.h"
+#include "HAL/PlatformApplicationMisc.h"
 
 #define LOCTEXT_NAMESPACE "SRigHierarchy"
 
 //////////////////////////////////////////////////////////////
-/// FRigTreeBone
+/// FRigTreeElement
 ///////////////////////////////////////////////////////////
-FRigTreeBone::FRigTreeBone(const FName& InBone, TWeakPtr<SRigHierarchy> InHierarchyHandler)
+FRigTreeElement::FRigTreeElement(const FRigElementKey& InKey, TWeakPtr<SRigHierarchy> InHierarchyHandler)
 {
-	CachedBone = InBone;
+	Key = InKey;
 }
 
-TSharedRef<ITableRow> FRigTreeBone::MakeTreeRowWidget(TSharedPtr<FControlRigEditor> InControlRigEditor, const TSharedRef<STableViewBase>& InOwnerTable, TSharedRef<FRigTreeBone> InRigTreeBone, TSharedRef<FUICommandList> InCommandList, TSharedPtr<SRigHierarchy> InHierarchy)
+TSharedRef<ITableRow> FRigTreeElement::MakeTreeRowWidget(TSharedPtr<FControlRigEditor> InControlRigEditor, const TSharedRef<STableViewBase>& InOwnerTable, TSharedRef<FRigTreeElement> InRigTreeElement, TSharedRef<FUICommandList> InCommandList, TSharedPtr<SRigHierarchy> InHierarchy)
 {
-	return SNew(SRigHierarchyItem, InControlRigEditor, InOwnerTable, InRigTreeBone, InCommandList, InHierarchy)
-		.OnRenameBone(InHierarchy.Get(), &SRigHierarchy::RenameBone)
-		.OnVerifyBoneNameChanged(InHierarchy.Get(), &SRigHierarchy::OnVerifyNameChanged);
+	return SNew(SRigHierarchyItem, InControlRigEditor, InOwnerTable, InRigTreeElement, InCommandList, InHierarchy)
+		.OnRenameElement(InHierarchy.Get(), &SRigHierarchy::RenameElement)
+		.OnVerifyElementNameChanged(InHierarchy.Get(), &SRigHierarchy::OnVerifyNameChanged);
 }
 
-void FRigTreeBone::RequestRename()
+void FRigTreeElement::RequestRename()
 {
 	OnRenameRequested.ExecuteIfBound();
 }
 
 //////////////////////////////////////////////////////////////
-/// FRigHierarchyDragDropOp
+/// FRigElementHierarchyDragDropOp
 ///////////////////////////////////////////////////////////
-TSharedRef<FRigHierarchyDragDropOp> FRigHierarchyDragDropOp::New(TArray<FName> InBoneNames)
+TSharedRef<FRigElementHierarchyDragDropOp> FRigElementHierarchyDragDropOp::New(const TArray<FRigElementKey>& InElements)
 {
-	TSharedRef<FRigHierarchyDragDropOp> Operation = MakeShared<FRigHierarchyDragDropOp>();
-	Operation->BoneNames = InBoneNames;
+	TSharedRef<FRigElementHierarchyDragDropOp> Operation = MakeShared<FRigElementHierarchyDragDropOp>();
+	Operation->Elements = InElements;
 	Operation->Construct();
 	return Operation;
 }
 
-TSharedPtr<SWidget> FRigHierarchyDragDropOp::GetDefaultDecorator() const
+TSharedPtr<SWidget> FRigElementHierarchyDragDropOp::GetDefaultDecorator() const
 {
 	return SNew(SBorder)
 		.Visibility(EVisibility::Visible)
 		.BorderImage(FEditorStyle::GetBrush("Menu.Background"))
 		[
 			SNew(STextBlock)
-			.Text(FText::FromString(GetJoinedBoneNames()))
-			.Font(FEditorStyle::Get().GetFontStyle("FontAwesome.10"))
+			.Text(FText::FromString(GetJoinedElementNames()))
+			//.Font(FEditorStyle::Get().GetFontStyle("FontAwesome.10"))
 		];
 }
 
-FString FRigHierarchyDragDropOp::GetJoinedBoneNames() const
+FString FRigElementHierarchyDragDropOp::GetJoinedElementNames() const
 {
-	TArray<FString> BoneNameStrings;
-	for (const FName& BoneName : BoneNames)
+	TArray<FString> ElementNameStrings;
+	for (const FRigElementKey& Element: Elements)
 	{
-		BoneNameStrings.Add(BoneName.ToString());
+		ElementNameStrings.Add(Element.Name.ToString());
 	}
-	return FString::Join(BoneNameStrings, TEXT(","));
+	return FString::Join(ElementNameStrings, TEXT(","));
 }
 
-
 //////////////////////////////////////////////////////////////
 /// SRigHierarchyItem
 ///////////////////////////////////////////////////////////
-void SRigHierarchyItem::Construct(const FArguments& InArgs, TSharedPtr<FControlRigEditor> InControlRigEditor, const TSharedRef<STableViewBase>& OwnerTable, TSharedRef<FRigTreeBone> InRigTreeBone, TSharedRef<FUICommandList> InCommandList, TSharedPtr<SRigHierarchy> InHierarchy)
+void SRigHierarchyItem::Construct(const FArguments& InArgs, TSharedPtr<FControlRigEditor> InControlRigEditor, const TSharedRef<STableViewBase>& OwnerTable, TSharedRef<FRigTreeElement> InRigTreeElement, TSharedRef<FUICommandList> InCommandList, TSharedPtr<SRigHierarchy> InHierarchy)
 {
-	WeakRigTreeBone = InRigTreeBone;
+	WeakRigTreeElement = InRigTreeElement;
 	WeakCommandList = InCommandList;
 	ControlRigEditor = InControlRigEditor;
 
-	OnVerifyBoneNameChanged = InArgs._OnVerifyBoneNameChanged;
-	OnRenameBone = InArgs._OnRenameBone;
+	OnVerifyElementNameChanged = InArgs._OnVerifyElementNameChanged;
+	OnRenameElement = InArgs._OnRenameElement;
 
 	TSharedPtr< SInlineEditableTextBlock > InlineWidget;
 
-	STableRow<TSharedPtr<FRigTreeBone>>::Construct(
-		STableRow<TSharedPtr<FRigTreeBone>>::FArguments()
+	const FSlateBrush* Brush = nullptr;
+	switch (InRigTreeElement->Key.Type)
+	{
+		case ERigElementType::Control:
+		{
+			Brush = FControlRigEditorStyle::Get().GetBrush("ControlRig.Tree.Control");
+			break;
+		}
+		case ERigElementType::Space:
+		{
+			Brush = FControlRigEditorStyle::Get().GetBrush("ControlRig.Tree.Space");
+			break;
+		}
+		default:
+		{
+			Brush = FControlRigEditorStyle::Get().GetBrush("ControlRig.Tree.Bone");
+			break;
+		}
+	}
+
+	STableRow<TSharedPtr<FRigTreeElement>>::Construct(
+		STableRow<TSharedPtr<FRigTreeElement>>::FArguments()
 		.OnDragDetected(InHierarchy.Get(), &SRigHierarchy::OnDragDetected)
+		.OnCanAcceptDrop(InHierarchy.Get(), &SRigHierarchy::OnCanAcceptDrop)
+		.OnAcceptDrop(InHierarchy.Get(), &SRigHierarchy::OnAcceptDrop)
 		.Content()
 		[
 			SNew(SHorizontalBox)
 			+ SHorizontalBox::Slot()
+			.MaxWidth(18)
+			.FillWidth(1.0)
+			.HAlign(HAlign_Left)
+			.VAlign(VAlign_Center)
+			[
+				SNew(SImage)
+				.Image(Brush)
+			]
+			+ SHorizontalBox::Slot()
 			.AutoWidth()
 			.VAlign(VAlign_Center)
 			[
 				SAssignNew(InlineWidget, SInlineEditableTextBlock)
 				.Text(this, &SRigHierarchyItem::GetName)
-				//.HighlightText(FilterText)
 				.OnVerifyTextChanged(this, &SRigHierarchyItem::OnVerifyNameChanged)
 				.OnTextCommitted(this, &SRigHierarchyItem::OnNameCommitted)
 				.MultiLine(false)
 			]
 		], OwnerTable);
 
-	InRigTreeBone->OnRenameRequested.BindSP(InlineWidget.Get(), &SInlineEditableTextBlock::EnterEditingMode);
+	InRigTreeElement->OnRenameRequested.BindSP(InlineWidget.Get(), &SInlineEditableTextBlock::EnterEditingMode);
 }
 
 FText SRigHierarchyItem::GetName() const
 {
-	return (FText::FromName(WeakRigTreeBone.Pin()->CachedBone));
+	return (FText::FromName(WeakRigTreeElement.Pin()->Key.Name));
 }
 
 bool SRigHierarchyItem::OnVerifyNameChanged(const FText& InText, FText& OutErrorMessage)
 {
 	const FName NewName = FName(*InText.ToString());
-	if (OnVerifyBoneNameChanged.IsBound())
+	if (OnVerifyElementNameChanged.IsBound())
 	{
-		return OnVerifyBoneNameChanged.Execute(WeakRigTreeBone.Pin()->CachedBone, NewName, OutErrorMessage);
+		return OnVerifyElementNameChanged.Execute(WeakRigTreeElement.Pin()->Key, NewName, OutErrorMessage);
 	}
 
 	// if not bound, just allow
@@ -146,13 +177,13 @@ void SRigHierarchyItem::OnNameCommitted(const FText& InText, ETextCommit::Type I
 	if (InCommitType == ETextCommit::OnEnter)
 	{
 		FName NewName = FName(*InText.ToString());
-		FName OldName = WeakRigTreeBone.Pin()->CachedBone;
+		FRigElementKey OldKey = WeakRigTreeElement.Pin()->Key;
 
-		if (!OnRenameBone.IsBound() || OnRenameBone.Execute(OldName, NewName))
+		if (!OnRenameElement.IsBound() || OnRenameElement.Execute(OldKey, NewName))
 		{
-			if (WeakRigTreeBone.IsValid())
+			if (WeakRigTreeElement.IsValid())
 			{
-				WeakRigTreeBone.Pin()->CachedBone = NewName;
+				WeakRigTreeElement.Pin()->Key.Name = NewName;
 			}
 		}
 	}
@@ -162,6 +193,22 @@ void SRigHierarchyItem::OnNameCommitted(const FText& InText, ETextCommit::Type I
 
 SRigHierarchy::~SRigHierarchy()
 {
+	if (ControlRigEditor.IsValid())
+	{
+		ControlRigEditor.Pin()->GetKeyDownDelegate().Unbind();
+		ControlRigEditor.Pin()->OnViewportContextMenu().Unbind();
+		ControlRigEditor.Pin()->OnViewportContextMenuCommands().Unbind();
+
+		ControlRigBlueprint = ControlRigEditor.Pin()->GetControlRigBlueprint();
+		if (ControlRigBlueprint.IsValid())
+		{
+			ControlRigBlueprint->HierarchyContainer.OnElementAdded.RemoveAll(this);
+			ControlRigBlueprint->HierarchyContainer.OnElementRemoved.RemoveAll(this);
+			ControlRigBlueprint->HierarchyContainer.OnElementRenamed.RemoveAll(this);
+			ControlRigBlueprint->HierarchyContainer.OnElementReparented.RemoveAll(this);
+			ControlRigBlueprint->HierarchyContainer.OnElementSelected.RemoveAll(this);
+		}
+	}
 }
 
 void SRigHierarchy::Construct(const FArguments& InArgs, TSharedRef<FControlRigEditor> InControlRigEditor)
@@ -169,8 +216,14 @@ void SRigHierarchy::Construct(const FArguments& InArgs, TSharedRef<FControlRigEd
 	ControlRigEditor = InControlRigEditor;
 
 	ControlRigBlueprint = ControlRigEditor.Pin()->GetControlRigBlueprint();
-	// @todo: find a better place to do it
-	ControlRigBlueprint->Hierarchy.Initialize();
+	ControlRigBlueprint->HierarchyContainer.Initialize();
+
+	ControlRigBlueprint->HierarchyContainer.OnElementAdded.AddRaw(this, &SRigHierarchy::OnRigElementAdded);
+	ControlRigBlueprint->HierarchyContainer.OnElementRemoved.AddRaw(this, &SRigHierarchy::OnRigElementRemoved);
+	ControlRigBlueprint->HierarchyContainer.OnElementRenamed.AddRaw(this, &SRigHierarchy::OnRigElementRenamed);
+	ControlRigBlueprint->HierarchyContainer.OnElementReparented.AddRaw(this, &SRigHierarchy::OnRigElementReparented);
+	ControlRigBlueprint->HierarchyContainer.OnElementSelected.AddRaw(this, &SRigHierarchy::OnRigElementSelected);
+
 	// for deleting, renaming, dragging
 	CommandList = MakeShared<FUICommandList>();
 
@@ -218,8 +271,8 @@ void SRigHierarchy::Construct(const FArguments& InArgs, TSharedRef<FControlRigEd
 			.Padding(2.0f)
 			.BorderImage(FEditorStyle::GetBrush("SCSEditor.TreePanel"))
 			[
-				SAssignNew(TreeView, STreeView<TSharedPtr<FRigTreeBone>>)
-				.TreeItemsSource(&RootBones)
+				SAssignNew(TreeView, STreeView<TSharedPtr<FRigTreeElement>>)
+				.TreeItemsSource(&RootElements)
 				.SelectionMode(ESelectionMode::Multi)
 				.OnGenerateRow(this, &SRigHierarchy::MakeTableRowWidget)
 				.OnGetChildren(this, &SRigHierarchy::HandleGetChildrenForTree)
@@ -231,15 +284,32 @@ void SRigHierarchy::Construct(const FArguments& InArgs, TSharedRef<FControlRigEd
 		]
 	];
 
+	bIsChangingRigHierarchy = false;
 	RefreshTreeView();
+
+	if (ControlRigEditor.IsValid())
+	{
+		ControlRigEditor.Pin()->GetKeyDownDelegate().BindLambda([&](const FGeometry& MyGeometry, const FKeyEvent& InKeyEvent)->FReply {
+			return OnKeyDown(MyGeometry, InKeyEvent);
+		});
+		ControlRigEditor.Pin()->OnViewportContextMenu().BindSP(this, &SRigHierarchy::FillContextMenu);
+		ControlRigEditor.Pin()->OnViewportContextMenuCommands().BindSP(this, &SRigHierarchy::GetContextMenuCommands);
+	}
 }
 
 void SRigHierarchy::BindCommands()
 {
 	// create new command
 	const FControlRigHierarchyCommands& Commands = FControlRigHierarchyCommands::Get();
-	CommandList->MapAction(Commands.AddItem,
-		FExecuteAction::CreateSP(this, &SRigHierarchy::HandleNewItem));
+
+	CommandList->MapAction(Commands.AddBoneItem,
+		FExecuteAction::CreateSP(this, &SRigHierarchy::HandleNewItem, ERigElementType::Bone));
+
+	CommandList->MapAction(Commands.AddControlItem,
+		FExecuteAction::CreateSP(this, &SRigHierarchy::HandleNewItem, ERigElementType::Control));
+
+	CommandList->MapAction(Commands.AddSpaceItem,
+		FExecuteAction::CreateSP(this, &SRigHierarchy::HandleNewItem, ERigElementType::Space));
 
 	CommandList->MapAction(Commands.DuplicateItem,
 		FExecuteAction::CreateSP(this, &SRigHierarchy::HandleDuplicateItem),
@@ -252,6 +322,61 @@ void SRigHierarchy::BindCommands()
 	CommandList->MapAction(Commands.RenameItem,
 		FExecuteAction::CreateSP(this, &SRigHierarchy::HandleRenameItem),
 		FCanExecuteAction::CreateSP(this, &SRigHierarchy::CanRenameItem));
+
+	CommandList->MapAction(Commands.CopyItems,
+		FExecuteAction::CreateSP(this, &SRigHierarchy::HandleCopyItems),
+		FCanExecuteAction::CreateSP(this, &SRigHierarchy::CanCopyOrPasteItems));
+
+	CommandList->MapAction(Commands.PasteItems,
+		FExecuteAction::CreateSP(this, &SRigHierarchy::HandlePasteItems),
+		FCanExecuteAction::CreateSP(this, &SRigHierarchy::CanPasteItems));
+
+	CommandList->MapAction(Commands.PasteLocalTransforms,
+		FExecuteAction::CreateSP(this, &SRigHierarchy::HandlePasteLocalTransforms),
+		FCanExecuteAction::CreateSP(this, &SRigHierarchy::CanCopyOrPasteItems));
+
+	CommandList->MapAction(Commands.PasteGlobalTransforms,
+		FExecuteAction::CreateSP(this, &SRigHierarchy::HandlePasteGlobalTransforms),
+		FCanExecuteAction::CreateSP(this, &SRigHierarchy::CanCopyOrPasteItems));
+
+	CommandList->MapAction(
+		Commands.ResetTransform,
+		FExecuteAction::CreateSP(this, &SRigHierarchy::HandleResetTransform),
+		FCanExecuteAction::CreateSP(this, &SRigHierarchy::IsMultiSelected));
+
+	CommandList->MapAction(
+		Commands.ResetInitialTransform,
+		FExecuteAction::CreateSP(this, &SRigHierarchy::HandleResetInitialTransform),
+		FCanExecuteAction::CreateSP(this, &SRigHierarchy::IsMultiSelected));
+
+	CommandList->MapAction(
+		Commands.ResetSpace,
+		FExecuteAction::CreateSP(this, &SRigHierarchy::HandleResetSpace),
+		FCanExecuteAction::CreateSP(this, &SRigHierarchy::IsControlSelected));
+
+	CommandList->MapAction(
+		Commands.SetInitialTransformFromClosestBone,
+		FExecuteAction::CreateSP(this, &SRigHierarchy::HandleSetInitialTransformFromClosestBone),
+		FCanExecuteAction::CreateSP(this, &SRigHierarchy::IsControlOrSpaceSelected));
+
+	CommandList->MapAction(
+		Commands.SetInitialTransformFromCurrentTransform,
+		FExecuteAction::CreateSP(this, &SRigHierarchy::HandleSetInitialTransformFromCurrentTransform),
+		FCanExecuteAction::CreateSP(this, &SRigHierarchy::IsControlOrSpaceSelected));
+
+	CommandList->MapAction(
+		Commands.FrameSelection,
+		FExecuteAction::CreateSP(this, &SRigHierarchy::HandleFrameSelection),
+		FCanExecuteAction::CreateSP(this, &SRigHierarchy::IsMultiSelected));
+}
+
+FReply SRigHierarchy::OnKeyDown(const FGeometry& MyGeometry, const FKeyEvent& InKeyEvent)
+{
+	if (CommandList.IsValid() && CommandList->ProcessCommandBindings(InKeyEvent))
+	{
+		return FReply::Handled();
+	}
+	return FReply::Unhandled();
 }
 
 void SRigHierarchy::OnFilterTextChanged(const FText& SearchText)
@@ -263,136 +388,318 @@ void SRigHierarchy::OnFilterTextChanged(const FText& SearchText)
 
 void SRigHierarchy::RefreshTreeView()
 {
-	RootBones.Reset();
-	FilteredRootBones.Reset();
+	TMap<FRigElementKey, bool> ExpansionState;
+	for (TPair<FRigElementKey, TSharedPtr<FRigTreeElement>> Pair : ElementMap)
+	{
+		ExpansionState.FindOrAdd(Pair.Key) = TreeView->IsItemExpanded(Pair.Value);
+	}
+
+	RootElements.Reset();
+	ElementMap.Reset();
+	ParentMap.Reset();
 
 	if (ControlRigBlueprint.IsValid())
 	{
-		FRigHierarchy& Hierarchy = ControlRigBlueprint->Hierarchy;
-
-		TMap<FName, TSharedPtr<FRigTreeBone>> SearchTable;
+		FRigHierarchyContainer* Container = GetHierarchyContainer();
+		FRigBoneHierarchy& BoneHierarchy = Container->BoneHierarchy;
+		FRigControlHierarchy& ControlHierarchy = Container->ControlHierarchy;
+		FRigSpaceHierarchy& SpaceHierarchy = Container->SpaceHierarchy;
 
-		FString FilteredString = FilterText.ToString();
-		const bool bSearchOff = FilteredString.IsEmpty();
-		const TArray<FRigBone>& Bones = Hierarchy.GetBones();
-		for (int32 BoneIndex = 0; BoneIndex < Bones.Num(); ++BoneIndex)
+		for (const FRigBone& Element : BoneHierarchy)
 		{
-			const FRigBone& Bone = Bones[BoneIndex];
+			AddBoneElement(Element);
+		}
+		for (const FRigControl& Element : ControlHierarchy)
+		{
+			AddControlElement(Element);
+		}
+		for (const FRigSpace& Element : SpaceHierarchy)
+		{
+			AddSpaceElement(Element);
+		}
 
-			// create new item
-			if (bSearchOff)
+		if (ExpansionState.Num() == 0)
+		{
+			for (TSharedPtr<FRigTreeElement> RootElement : RootElements)
 			{
-				TSharedPtr<FRigTreeBone> NewItem = MakeShared<FRigTreeBone>(Bones[BoneIndex].Name, SharedThis(this));
-				SearchTable.Add(Bone.Name, NewItem);
-
-				if (Bone.ParentName == NAME_None)
+				SetExpansionRecursive(RootElement, false);
+			}
+		}
+		else
+		{
+			for (TPair<FRigElementKey, bool> Pair : ExpansionState)
+			{
+				if (!Pair.Value)
 				{
-					RootBones.Add(NewItem);
+					continue;
 				}
-				else
+
+				if (TSharedPtr<FRigTreeElement>* ItemPtr = ElementMap.Find(Pair.Key))
 				{
-					// you have to find one
-					TSharedPtr<FRigTreeBone>* FoundItem = SearchTable.Find(Bone.ParentName);
-					check(FoundItem);
-					// add to children list
-					FoundItem->Get()->Children.Add(NewItem);
+					TreeView->SetItemExpansion(*ItemPtr, true);
 				}
 			}
-			else if (Bone.Name.ToString().Contains(FilteredString))
-			{
-				// if contains, just list out everything to root
-				TSharedPtr<FRigTreeBone> NewItem = MakeShared<FRigTreeBone>(Bones[BoneIndex].Name, SharedThis(this));
-				// during search, everything is on root
-				RootBones.Add(NewItem);
-			}
 		}
 
-		if (bSearchOff)
+		TreeView->RequestTreeRefresh();
+
+		TArray<FRigElementKey> Selection = Container->CurrentSelection();
+		for (const FRigElementKey& Key : Selection)
 		{
-			for (int32 RootIndex = 0; RootIndex < RootBones.Num(); ++RootIndex)
-			{
-				SetExpansionRecursive(RootBones[RootIndex]);
-			}
+			OnRigElementSelected(Container, Key, true);
 		}
 	}
-
-	TreeView->RequestTreeRefresh();
 }
 
-void SRigHierarchy::SetExpansionRecursive(TSharedPtr<FRigTreeBone> InBone)
+void SRigHierarchy::SetExpansionRecursive(TSharedPtr<FRigTreeElement> InElement, bool bTowardsParent)
 {
-	TreeView->SetItemExpansion(InBone, true);
+	TreeView->SetItemExpansion(InElement, true);
 
-	for (int32 ChildIndex = 0; ChildIndex < InBone->Children.Num(); ++ChildIndex)
+	if (bTowardsParent)
+	{
+		if (const FRigElementKey* ParentKey = ParentMap.Find(InElement->Key))
+		{
+			if (TSharedPtr<FRigTreeElement>* ParentItem = ElementMap.Find(*ParentKey))
+			{
+				SetExpansionRecursive(*ParentItem, bTowardsParent);
+			}
+		}
+	}
+	else
 	{
-		SetExpansionRecursive(InBone->Children[ChildIndex]);
+		for (int32 ChildIndex = 0; ChildIndex < InElement->Children.Num(); ++ChildIndex)
+		{
+			SetExpansionRecursive(InElement->Children[ChildIndex], bTowardsParent);
+		}
 	}
 }
-TSharedRef<ITableRow> SRigHierarchy::MakeTableRowWidget(TSharedPtr<FRigTreeBone> InItem, const TSharedRef<STableViewBase>& OwnerTable)
+TSharedRef<ITableRow> SRigHierarchy::MakeTableRowWidget(TSharedPtr<FRigTreeElement> InItem, const TSharedRef<STableViewBase>& OwnerTable)
 {
 	return InItem->MakeTreeRowWidget(ControlRigEditor.Pin(), OwnerTable, InItem.ToSharedRef(), CommandList.ToSharedRef(), SharedThis(this));
 }
 
-void SRigHierarchy::HandleGetChildrenForTree(TSharedPtr<FRigTreeBone> InItem, TArray<TSharedPtr<FRigTreeBone>>& OutChildren)
+void SRigHierarchy::HandleGetChildrenForTree(TSharedPtr<FRigTreeElement> InItem, TArray<TSharedPtr<FRigTreeElement>>& OutChildren)
 {
 	OutChildren = InItem.Get()->Children;
 }
 
-void SRigHierarchy::OnSelectionChanged(TSharedPtr<FRigTreeBone> Selection, ESelectInfo::Type SelectInfo)
+void SRigHierarchy::OnSelectionChanged(TSharedPtr<FRigTreeElement> Selection, ESelectInfo::Type SelectInfo)
 {
-	// need dummy object
-	if (Selection.IsValid())
+	if (bIsChangingRigHierarchy)
+	{
+		return;
+	}
+
+	FRigHierarchyContainer* Hierarchy = GetHierarchyContainer();
+
+	if (Hierarchy)
 	{
-		FRigHierarchy* RigHierarchy = GetInstanceHierarchy();
+		TGuardValue<bool> GuardRigHierarchyChanges(bIsChangingRigHierarchy, true);
+
+		TArray<FRigElementKey> OldSelection = Hierarchy->CurrentSelection();
+		TArray<FRigElementKey> NewSelection;
 
-		if (RigHierarchy)
+		TArray<TSharedPtr<FRigTreeElement>> SelectedItems = TreeView->GetSelectedItems();
+		for (const TSharedPtr<FRigTreeElement>& SelectedItem : SelectedItems)
 		{
-			const int32 BoneIndex = RigHierarchy->GetIndex(Selection->CachedBone);
-			if (BoneIndex != INDEX_NONE)
-			{
-				ControlRigEditor.Pin()->SetDetailStruct(MakeShareable(new FStructOnScope(FRigBone::StaticStruct(), (uint8*)&RigHierarchy->GetBones()[BoneIndex])));
-				ControlRigEditor.Pin()->SelectBone(Selection->CachedBone);
-				return;
-			}
-			else
+			NewSelection.Add(SelectedItem->Key);
+		}
+
+		for (const FRigElementKey& PreviouslySelected : OldSelection)
+		{
+			if (NewSelection.Contains(PreviouslySelected))
 			{
-				// clear the current selection
-				ControlRigEditor.Pin()->ClearDetailObject();
-				ControlRigEditor.Pin()->SelectBone(NAME_None);
+				continue;
 			}
+			Hierarchy->Select(PreviouslySelected, false);
 		}
 
-		// if failed, try BP hierarhcy? Todo:
+		for (const FRigElementKey& NewlySelected : NewSelection)
+		{
+			Hierarchy->Select(NewlySelected, true);
+		}
 	}
 }
 
-TSharedPtr<FRigTreeBone> FindBone(const FName& InBoneName, TSharedPtr<FRigTreeBone> CurrentItem)
+TSharedPtr<FRigTreeElement> SRigHierarchy::FindElement(const FRigElementKey& InElementKey, TSharedPtr<FRigTreeElement> CurrentItem)
 {
-	if (CurrentItem->CachedBone == InBoneName)
+	if (CurrentItem->Key == InElementKey)
 	{
 		return CurrentItem;
 	}
 
 	for (int32 ChildIndex = 0; ChildIndex < CurrentItem->Children.Num(); ++ChildIndex)
 	{
-		TSharedPtr<FRigTreeBone> Found = FindBone(InBoneName, CurrentItem->Children[ChildIndex]);
+		TSharedPtr<FRigTreeElement> Found = FindElement(InElementKey, CurrentItem->Children[ChildIndex]);
 		if (Found.IsValid())
 		{
 			return Found;
 		}
 	}
 
-	return TSharedPtr<FRigTreeBone>();
+	return TSharedPtr<FRigTreeElement>();
+}
+
+void SRigHierarchy::AddElement(FRigElementKey InKey, FRigElementKey InParentKey)
+{
+	if(ElementMap.Contains(InKey))
+	{
+		return;
+	}
+
+	FString FilteredString = FilterText.ToString();
+	if (FilteredString.IsEmpty())
+	{
+		TSharedPtr<FRigTreeElement> NewItem = MakeShared<FRigTreeElement>(InKey, SharedThis(this));
+		ElementMap.Add(InKey, NewItem);
+		if (InParentKey)
+		{
+			ParentMap.Add(InKey, InParentKey);
+		}
+
+		if (InParentKey)
+		{
+			TSharedPtr<FRigTreeElement>* FoundItem = ElementMap.Find(InParentKey);
+			check(FoundItem);
+			FoundItem->Get()->Children.Add(NewItem);
+		}
+		else
+		{
+			RootElements.Add(NewItem);
+		}
+	}
+	else if (InKey.Name.ToString().Contains(FilteredString))
+	{
+		TSharedPtr<FRigTreeElement> NewItem = MakeShared<FRigTreeElement>(InKey, SharedThis(this));
+		RootElements.Add(NewItem);
+	}
+}
+
+void SRigHierarchy::AddBoneElement(FRigBone InBone)
+{
+	const FRigHierarchyContainer* Container = GetHierarchyContainer();
+	const FRigBoneHierarchy& BoneHierarchy = Container->BoneHierarchy;
+	const FRigControlHierarchy& ControlHierarchy = Container->ControlHierarchy;
+	const FRigSpaceHierarchy& SpaceHierarchy = Container->SpaceHierarchy;
+
+	FRigElementKey ParentKey;
+	if(InBone.ParentIndex != INDEX_NONE)
+	{
+		AddBoneElement(BoneHierarchy[InBone.ParentIndex]);
+		ParentKey = BoneHierarchy[InBone.ParentIndex].GetElementKey();
+	}
+	AddElement(InBone.GetElementKey(), ParentKey);
+}
+
+void SRigHierarchy::AddControlElement(FRigControl InControl)
+{
+	const FRigHierarchyContainer* Container = GetHierarchyContainer();
+	const FRigControlHierarchy& ControlHierarchy = Container->ControlHierarchy;
+	const FRigSpaceHierarchy& SpaceHierarchy = Container->SpaceHierarchy;
+
+	FRigElementKey ParentKey;
+	if(InControl.SpaceIndex != INDEX_NONE)
+	{
+		AddSpaceElement(SpaceHierarchy[InControl.SpaceIndex]);
+		ParentKey = SpaceHierarchy[InControl.SpaceIndex].GetElementKey();
+	}
+	else if(InControl.ParentIndex != INDEX_NONE)
+	{
+		AddControlElement(ControlHierarchy[InControl.ParentIndex]);
+		ParentKey = ControlHierarchy[InControl.ParentIndex].GetElementKey();
+	}
+	AddElement(InControl.GetElementKey(), ParentKey);
+}
+
+void SRigHierarchy::AddSpaceElement(FRigSpace InSpace)
+{
+	const FRigHierarchyContainer* Container = GetHierarchyContainer();
+	const FRigBoneHierarchy& BoneHierarchy = Container->BoneHierarchy;
+	const FRigControlHierarchy& ControlHierarchy = Container->ControlHierarchy;
+	const FRigSpaceHierarchy& SpaceHierarchy = Container->SpaceHierarchy;
+
+	FRigElementKey ParentKey;
+	if(InSpace.ParentIndex != INDEX_NONE)
+	{
+		switch(InSpace.SpaceType)
+		{
+			case ERigSpaceType::Bone:
+			{
+				AddBoneElement(BoneHierarchy[InSpace.ParentIndex]);
+				ParentKey = BoneHierarchy[InSpace.ParentIndex].GetElementKey();
+				break;
+			}
+			case ERigSpaceType::Control:
+			{
+				AddControlElement(ControlHierarchy[InSpace.ParentIndex]);
+				ParentKey = ControlHierarchy[InSpace.ParentIndex].GetElementKey();
+				break;
+			}
+			case ERigSpaceType::Space:
+			{
+				AddSpaceElement(SpaceHierarchy[InSpace.ParentIndex]);
+				ParentKey = SpaceHierarchy[InSpace.ParentIndex].GetElementKey();
+				break;
+			}
+			default:
+			{
+				break;
+			}
+		}
+	}
+	AddElement(InSpace.GetElementKey(), ParentKey);
+}
+
+void SRigHierarchy::OnRigElementAdded(FRigHierarchyContainer* Container, const FRigElementKey& InKey)
+{
+	if (bIsChangingRigHierarchy || InKey.Type == ERigElementType::Curve)
+	{
+		return;
+	}
+	RefreshTreeView();
+}
+
+void SRigHierarchy::OnRigElementRemoved(FRigHierarchyContainer* Container, const FRigElementKey& InKey)
+{
+	if (bIsChangingRigHierarchy || InKey.Type == ERigElementType::Curve)
+	{
+		return;
+	}
+	RefreshTreeView();
+}
+
+void SRigHierarchy::OnRigElementRenamed(FRigHierarchyContainer* Container, ERigElementType ElementType, const FName& InOldName, const FName& InNewName)
+{
+	if (bIsChangingRigHierarchy || ElementType == ERigElementType::Curve)
+	{
+		return;
+	}
+	RefreshTreeView();
+}
+
+void SRigHierarchy::OnRigElementReparented(FRigHierarchyContainer* Container, const FRigElementKey& InKey, const FName& InOldParentName, const FName& InNewParentName)
+{
+	if (bIsChangingRigHierarchy || InKey.Type == ERigElementType::Curve)
+	{
+		return;
+	}
+	RefreshTreeView();
 }
 
-void SRigHierarchy::SelectBone(const FName& BoneName) const
+void SRigHierarchy::OnRigElementSelected(FRigHierarchyContainer* Container, const FRigElementKey& InKey, bool bSelected)
 {
-	for (int32 RootIndex = 0; RootIndex < RootBones.Num(); ++RootIndex)
+	if (bIsChangingRigHierarchy || InKey.Type == ERigElementType::Curve)
+	{
+		return;
+	}
+
+	for (int32 RootIndex = 0; RootIndex < RootElements.Num(); ++RootIndex)
 	{
-		TSharedPtr<FRigTreeBone> Found = FindBone(BoneName, RootBones[RootIndex]);
+		TSharedPtr<FRigTreeElement> Found = FindElement(InKey, RootElements[RootIndex]);
 		if (Found.IsValid())
 		{
-			TreeView->SetSelection(Found);
+			TreeView->SetItemSelection(Found, bSelected, ESelectInfo::OnNavigation);
+			HandleFrameSelection();
 		}
 	}
 }
@@ -404,30 +711,137 @@ void SRigHierarchy::ClearDetailPanel() const
 
 TSharedPtr< SWidget > SRigHierarchy::CreateContextMenu()
 {
-	const FControlRigHierarchyCommands& Actions = FControlRigHierarchyCommands::Get();
-
 	const bool CloseAfterSelection = true;
 	FMenuBuilder MenuBuilder(CloseAfterSelection, CommandList);
+	FillContextMenu(MenuBuilder);
+	return MenuBuilder.MakeWidget();
+}
+
+void SRigHierarchy::FillContextMenu(class FMenuBuilder& MenuBuilder)
+{
+	const FControlRigHierarchyCommands& Actions = FControlRigHierarchyCommands::Get();
 	{
-		MenuBuilder.BeginSection("HierarchyEditAction", LOCTEXT("EditAction", "Edit"));
-		MenuBuilder.AddMenuEntry(Actions.AddItem);
+		struct FLocalMenuBuilder
+		{
+			static void FillNewMenu(FMenuBuilder& InSubMenuBuilder, TSharedPtr<STreeView<TSharedPtr<FRigTreeElement>>> InTreeView)
+			{
+				const FControlRigHierarchyCommands& Actions = FControlRigHierarchyCommands::Get();
+
+				FRigElementKey SelectedKey;
+				TArray<TSharedPtr<FRigTreeElement>> SelectedItems = InTreeView->GetSelectedItems();
+				if (SelectedItems.Num() > 0)
+				{
+					SelectedKey = SelectedItems[0]->Key;
+				}
+
+				if (!SelectedKey || SelectedKey.Type == ERigElementType::Bone)
+				{
+					InSubMenuBuilder.AddMenuEntry(Actions.AddBoneItem);
+				}
+
+				InSubMenuBuilder.AddMenuEntry(Actions.AddControlItem);
+				InSubMenuBuilder.AddMenuEntry(Actions.AddSpaceItem);
+			}
+		};
+
+		MenuBuilder.BeginSection("Elements", LOCTEXT("ElementsHeader", "Elements"));
+		MenuBuilder.AddSubMenu(LOCTEXT("New", "New"), LOCTEXT("New_ToolTip", "Create New Elements"),
+			FNewMenuDelegate::CreateStatic(&FLocalMenuBuilder::FillNewMenu, TreeView));
+
 		MenuBuilder.AddMenuEntry(Actions.DeleteItem);
 		MenuBuilder.AddMenuEntry(Actions.DuplicateItem);
 		MenuBuilder.AddMenuEntry(Actions.RenameItem);
+		MenuBuilder.EndSection();
+
+		MenuBuilder.BeginSection("Copy&Paste", LOCTEXT("Copy&PasteHeader", "Copy & Paste"));
+		MenuBuilder.AddMenuEntry(Actions.CopyItems);
+		MenuBuilder.AddMenuEntry(Actions.PasteItems);
+		MenuBuilder.AddMenuEntry(Actions.PasteLocalTransforms);
+		MenuBuilder.AddMenuEntry(Actions.PasteGlobalTransforms);
+		MenuBuilder.EndSection();
+
+		MenuBuilder.BeginSection("Transforms", LOCTEXT("TransformsHeader", "Transforms"));
+		MenuBuilder.AddMenuEntry(Actions.ResetTransform);
+		MenuBuilder.AddMenuEntry(Actions.ResetInitialTransform);
+		MenuBuilder.AddMenuEntry(Actions.SetInitialTransformFromCurrentTransform);
+		MenuBuilder.AddMenuEntry(Actions.SetInitialTransformFromClosestBone);
+		MenuBuilder.AddMenuEntry(Actions.ResetSpace);
+		MenuBuilder.EndSection();
 
-		MenuBuilder.AddMenuSeparator();
+		MenuBuilder.BeginSection("Assets", LOCTEXT("AssetsHeader", "Assets"));
 		MenuBuilder.AddSubMenu(
 			LOCTEXT("ImportSubMenu", "Import"),
-			LOCTEXT("ImportSubMenu_ToolTip", "Import hierarchy to the current rig. This overrides the data if it contains the existing node."),
+			LOCTEXT("ImportSubMenu_ToolTip", "Import hierarchy to the current rig. This only imports non-existing node. For example, if there is hand_r, it won't import hand_r. \
+				If you want to reimport whole new hiearchy, delete all nodes, and use import hierarchy."),
 			FNewMenuDelegate::CreateSP(this, &SRigHierarchy::CreateImportMenu)
 		);
-
+		MenuBuilder.AddSubMenu(
+			LOCTEXT("RefreshSubMenu", "Refresh"),
+			LOCTEXT("RefreshSubMenu_ToolTip", "Refresh the existing initial transform from the selected mesh. This only updates if the node is found."),
+			FNewMenuDelegate::CreateSP(this, &SRigHierarchy::CreateRefreshMenu)
+		);
 		MenuBuilder.EndSection();
 	}
+}
 
-	return MenuBuilder.MakeWidget();
+TSharedPtr<FUICommandList> SRigHierarchy::GetContextMenuCommands()
+{
+	return CommandList;
+}
+
+void SRigHierarchy::CreateRefreshMenu(FMenuBuilder& MenuBuilder)
+{
+	MenuBuilder.AddWidget(
+		SNew(SVerticalBox)
+
+		+ SVerticalBox::Slot()
+		.AutoHeight()
+		.Padding(3)
+		[
+			SNew(STextBlock)
+			.Font(FEditorStyle::GetFontStyle("ControlRig.Hierarchy.Menu"))
+			.Text(LOCTEXT("RefreshMesh_Title", "Select Mesh"))
+			.ToolTipText(LOCTEXT("RefreshMesh_Tooltip", "Select Mesh to refresh transform from... It will refresh init transform from selected mesh. This doesn't change hierarchy. \
+				If you want to reimport hierarchy, please delete all nodes, and use import hierarchy."))
+		]
+		+ SVerticalBox::Slot()
+		.AutoHeight()
+		.Padding(3)
+		[
+			SNew(SObjectPropertyEntryBox)
+			.AllowedClass(USkeletalMesh::StaticClass())
+			.OnObjectChanged(this, &SRigHierarchy::RefreshHierarchy)
+		]
+		,
+		FText()
+	);
 }
 
+void SRigHierarchy::RefreshHierarchy(const FAssetData& InAssetData)
+{
+	FRigHierarchyContainer* Hierarchy = GetHierarchyContainer();
+	USkeletalMesh* Mesh = Cast<USkeletalMesh>(InAssetData.GetAsset());
+	if (Mesh && Hierarchy)
+	{
+		FScopedTransaction Transaction(LOCTEXT("HierarchyRefresh", "Refresh Transform"));
+		ControlRigBlueprint->Modify();
+
+		const FReferenceSkeleton& RefSkeleton = Mesh->RefSkeleton;
+		const TArray<FMeshBoneInfo>& BoneInfos = RefSkeleton.GetRefBoneInfo();
+		const TArray<FTransform>& BonePoses = RefSkeleton.GetRefBonePose();
+
+		for (int32 Index = 0; Index < RefSkeleton.GetNum(); ++Index)
+		{
+			// only add if you don't have it. This may change in the future
+			int32 RigIndex = Hierarchy->BoneHierarchy.GetIndex(BoneInfos[Index].Name);
+			if (RigIndex != INDEX_NONE)
+			{
+				// @todo: add optimized version without sorting, but if no sort, we should make sure not to use find index function
+				Hierarchy->BoneHierarchy.SetInitialTransform(RigIndex, FAnimationRuntime::GetComponentSpaceTransform(RefSkeleton, BonePoses, Index));
+			}
+		}
+	}
+}
 void SRigHierarchy::CreateImportMenu(FMenuBuilder& MenuBuilder)
 {
 	MenuBuilder.AddWidget(
@@ -447,7 +861,7 @@ void SRigHierarchy::CreateImportMenu(FMenuBuilder& MenuBuilder)
 		.Padding(3)
 		[
 			SNew(SObjectPropertyEntryBox)
-			.OnShouldFilterAsset(this, &SRigHierarchy::ShouldFilterOnImport)
+			.AllowedClass(USkeletalMesh::StaticClass())
 			.OnObjectChanged(this, &SRigHierarchy::ImportHierarchy)
 		]
 		,
@@ -455,65 +869,41 @@ void SRigHierarchy::CreateImportMenu(FMenuBuilder& MenuBuilder)
 	);
 }
 
-bool SRigHierarchy::ShouldFilterOnImport(const FAssetData& AssetData) const
-{
-	return (AssetData.AssetClass != USkeletalMesh::StaticClass()->GetFName() &&
-		AssetData.AssetClass != USkeleton::StaticClass()->GetFName());
-}
-
 void SRigHierarchy::ImportHierarchy(const FAssetData& InAssetData)
 {
-	FRigHierarchy* Hier = GetHierarchy();
-	if (Hier)
+	FRigHierarchyContainer* Hierarchy = GetHierarchyContainer();
+	USkeletalMesh* Mesh = Cast<USkeletalMesh> (InAssetData.GetAsset());
+	if (Mesh && Hierarchy)
 	{
-		const FReferenceSkeleton* RefSkeleton = nullptr;
-		if (USkeletalMesh* Mesh = Cast<USkeletalMesh>(InAssetData.GetAsset()))
-		{
-			RefSkeleton = &Mesh->RefSkeleton;
-			ControlRigBlueprint->SourceHierarchyImport = Mesh;
-		}
-		else if (USkeleton* Skeleton = Cast<USkeleton>(InAssetData.GetAsset()))
-		{
-			RefSkeleton = &Skeleton->GetReferenceSkeleton();
-			ControlRigBlueprint->SourceHierarchyImport = Skeleton;
-		}
+		FScopedTransaction Transaction(LOCTEXT("HierarchyImport", "Import Hierarchy"));
+		ControlRigBlueprint->Modify();
 
-		if (RefSkeleton)
-		{
-			FScopedTransaction Transaction(LOCTEXT("HierarchyImport", "Import Hierarchy"));
-			ControlRigBlueprint->Modify();
+		const FReferenceSkeleton& RefSkeleton = Mesh->RefSkeleton;
+		const TArray<FMeshBoneInfo>& BoneInfos = RefSkeleton.GetRefBoneInfo();
+		const TArray<FTransform>& BonePoses = RefSkeleton.GetRefBonePose();
 
-			const TArray<FMeshBoneInfo>& BoneInfos = RefSkeleton->GetRefBoneInfo();
-			const TArray<FTransform>& BonePoses = RefSkeleton->GetRefBonePose();
+		Hierarchy->ClearSelection();
 
-			for (int32 BoneIndex = 0; BoneIndex < RefSkeleton->GetNum(); ++BoneIndex)
+		TArray<FName> AddedBones;
+		for (int32 Index = 0; Index < RefSkeleton.GetNum(); ++Index)
+		{
+			// only add if you don't have it. This may change in the future
+			if (Hierarchy->BoneHierarchy.GetIndex(BoneInfos[Index].Name) == INDEX_NONE)
 			{
-				// only add if you don't have it. This may change in the future
-				int32 HierIndex = Hier->GetIndex(BoneInfos[BoneIndex].Name);
-				FTransform InitialTransform = FAnimationRuntime::GetComponentSpaceTransform(*RefSkeleton, BonePoses, BoneIndex);
 				// @todo: add optimized version without sorting, but if no sort, we should make sure not to use find index function
-				FName ParentName = (BoneInfos[BoneIndex].ParentIndex != INDEX_NONE) ? BoneInfos[BoneInfos[BoneIndex].ParentIndex].Name : NAME_None;
-				// if exists, see if we should change parents
-				if (HierIndex != INDEX_NONE)
-				{
-					if (ParentName != Hier->GetParentName(BoneInfos[BoneIndex].Name))
-					{
-						// reparent
-						Hier->Reparent(BoneInfos[BoneIndex].Name, ParentName);
-					}
-
-					Hier->SetInitialTransform(BoneInfos[BoneIndex].Name, InitialTransform);
-				}
-				else
-				{
-					Hier->AddBone(BoneInfos[BoneIndex].Name, ParentName, InitialTransform);
-				}
+				FName ParentName = (BoneInfos[Index].ParentIndex != INDEX_NONE) ? BoneInfos[BoneInfos[Index].ParentIndex].Name : NAME_None;
+				Hierarchy->BoneHierarchy.Add(BoneInfos[Index].Name, ParentName, FAnimationRuntime::GetComponentSpaceTransform(RefSkeleton, BonePoses, Index));
+				AddedBones.Add(BoneInfos[Index].Name);
 			}
+		}
 
-			ControlRigEditor.Pin()->OnHierarchyChanged();
-			RefreshTreeView();
-			FSlateApplication::Get().DismissAllMenus();
+		for (const FName& AddedBone : AddedBones)
+		{
+			Hierarchy->BoneHierarchy.Select(AddedBone);
 		}
+
+		FSlateApplication::Get().DismissAllMenus();
+		RefreshTreeView();
 	}
 }
 
@@ -527,31 +917,89 @@ bool SRigHierarchy::IsSingleSelected() const
 	return TreeView->GetNumItemsSelected() == 1;
 }
 
-void SRigHierarchy::HandleDeleteItem()
+bool SRigHierarchy::IsControlSelected() const
 {
- 	FRigHierarchy* Hierarchy = GetHierarchy();
- 	if (Hierarchy)
- 	{
-		ClearDetailPanel();
-		FScopedTransaction Transaction(LOCTEXT("HierarchyTreeDeleteSelected", "Delete selected items from hierarchy"));
-		ControlRigBlueprint->Modify();
-
-		// clear detail view display
-		ControlRigEditor.Pin()->ClearDetailObject();
-
-		TArray<TSharedPtr<FRigTreeBone>> SelectedItems = TreeView->GetSelectedItems();
-
-		for (int32 ItemIndex = 0; ItemIndex < SelectedItems.Num(); ++ItemIndex)
+	TArray<TSharedPtr<FRigTreeElement>> SelectedItems = TreeView->GetSelectedItems();
+	for (TSharedPtr<FRigTreeElement> SelectedItem : SelectedItems)
+	{
+		if (SelectedItem->Key.Type == ERigElementType::Control)
 		{
-			// when you select whole Bones, you might not have them anymore
-			if (Hierarchy->GetIndex(SelectedItems[ItemIndex]->CachedBone) != INDEX_NONE)
-			{
-				Hierarchy->DeleteBone(SelectedItems[ItemIndex]->CachedBone, true);
-			}
+			return true;
 		}
-
-		ControlRigEditor.Pin()->OnHierarchyChanged();
-		RefreshTreeView();
+	}
+	return false;
+}
+
+bool SRigHierarchy::IsControlOrSpaceSelected() const
+{
+	TArray<TSharedPtr<FRigTreeElement>> SelectedItems = TreeView->GetSelectedItems();
+	for (TSharedPtr<FRigTreeElement> SelectedItem : SelectedItems)
+	{
+		if (SelectedItem->Key.Type == ERigElementType::Control)
+		{
+			return true;
+		}
+		if (SelectedItem->Key.Type == ERigElementType::Space)
+		{
+			return true;
+		}
+	}
+	return false;
+}
+
+void SRigHierarchy::HandleDeleteItem()
+{
+ 	FRigHierarchyContainer* Hierarchy = GetHierarchyContainer();
+ 	if (Hierarchy)
+ 	{
+		ClearDetailPanel();
+		FScopedTransaction Transaction(LOCTEXT("HierarchyTreeDeleteSelected", "Delete selected items from hierarchy"));
+		ControlRigBlueprint->Modify();
+
+		// clear detail view display
+		ControlRigEditor.Pin()->ClearDetailObject();
+
+		TArray<TSharedPtr<FRigTreeElement>> SelectedItems = TreeView->GetSelectedItems();
+		for (int32 ItemIndex = 0; ItemIndex < SelectedItems.Num(); ++ItemIndex)
+		{
+			TGuardValue<bool> GuardRigHierarchyChanges(bIsChangingRigHierarchy, true);
+
+			FRigElementKey SelectedKey = SelectedItems[ItemIndex]->Key;
+			switch (SelectedKey.Type)
+			{
+				case ERigElementType::Bone:
+				{
+					if (Hierarchy->BoneHierarchy.GetIndex(SelectedKey.Name) != INDEX_NONE)
+					{
+						Hierarchy->BoneHierarchy.Remove(SelectedKey.Name);
+					}
+					break;
+				}
+				case ERigElementType::Control:
+				{
+					if (Hierarchy->ControlHierarchy.GetIndex(SelectedKey.Name) != INDEX_NONE)
+					{
+						Hierarchy->ControlHierarchy.Remove(SelectedKey.Name);
+					}
+					break;
+				}
+				case ERigElementType::Space:
+				{
+					if (Hierarchy->SpaceHierarchy.GetIndex(SelectedKey.Name) != INDEX_NONE)
+					{
+						Hierarchy->SpaceHierarchy.Remove(SelectedKey.Name);
+					}
+					break;
+				}
+				default:
+				{
+					return;
+				}
+			}
+		}
+
+		RefreshTreeView();
+		FSlateApplication::Get().DismissAllMenus();
  	}
 }
 
@@ -561,9 +1009,9 @@ bool SRigHierarchy::CanDeleteItem() const
 }
 
 /** Delete Item */
-void SRigHierarchy::HandleNewItem()
+void SRigHierarchy::HandleNewItem(ERigElementType InElementType)
 {
-	FRigHierarchy* Hierarchy = GetHierarchy();
+	FRigHierarchyContainer* Hierarchy = GetHierarchyContainer();
 	if (Hierarchy)
 	{
 		// unselect current selected item
@@ -572,23 +1020,97 @@ void SRigHierarchy::HandleNewItem()
 		FScopedTransaction Transaction(LOCTEXT("HierarchyTreeAdded", "Add new item to hierarchy"));
 		ControlRigBlueprint->Modify();
 
-		FName ParentName = NAME_None;
+		FRigElementKey ParentKey;
 		FTransform ParentTransform = FTransform::Identity;
 
-		TArray<TSharedPtr<FRigTreeBone>> SelectedItems = TreeView->GetSelectedItems();
+		TArray<TSharedPtr<FRigTreeElement>> SelectedItems = TreeView->GetSelectedItems();
 		if (SelectedItems.Num() > 0)
 		{
-			ParentName = SelectedItems[0]->CachedBone;
-			ParentTransform = Hierarchy->GetGlobalTransform(ParentName);
+			ParentKey = SelectedItems[0]->Key;
+			ParentTransform = Hierarchy->GetGlobalTransform(ParentKey);
 		}
 
-		const FName NewBoneName = CreateUniqueName(TEXT("NewBone"));
-		Hierarchy->AddBone(NewBoneName, ParentName, ParentTransform);
+		FString NewNameTemplate = FString::Printf(TEXT("New%s"), *StaticEnum<ERigElementType>()->GetNameStringByValue((int64)InElementType));
+		const FName NewElementName = CreateUniqueName(*NewNameTemplate, InElementType);
+		{
+			TGuardValue<bool> GuardRigHierarchyChanges(bIsChangingRigHierarchy, true);
+			switch (InElementType)
+			{
+				case ERigElementType::Bone:
+				{
+					FName ParentName = ParentKey.Type == ERigElementType::Bone ? ParentKey.Name : NAME_None;
+					ParentTransform = ParentName == NAME_None ? FTransform::Identity : ParentTransform;
+					Hierarchy->BoneHierarchy.Add(NewElementName, ParentName, ParentTransform);
+					break;
+				}
+				case ERigElementType::Control:
+				{
+					if (ParentKey.Type == ERigElementType::Bone)
+					{
+						for (const FRigSpace& ExistingSpace : Hierarchy->SpaceHierarchy)
+						{
+							if (ExistingSpace.SpaceType == ERigSpaceType::Bone && ExistingSpace.ParentName == ParentKey.Name)
+							{
+								ParentKey = ExistingSpace.GetElementKey();
+								break;
+							}
+						}
+
+						if (ParentKey.Type != ERigElementType::Space)
+						{
+							FString SpaceName = FString::Printf(TEXT("%sSpace"), *NewElementName.ToString());
+							FRigSpace& NewSpace = Hierarchy->SpaceHierarchy.Add(*SpaceName, ERigSpaceType::Bone, ParentKey.Name);
+							ParentKey = NewSpace.GetElementKey();
+						}
+					}
+
+					if(!ParentKey)
+					{
+						Hierarchy->ControlHierarchy.Add(NewElementName, ERigControlType::Transform);
+					}
+					else if (ParentKey.Type == ERigElementType::Space)
+					{
+						Hierarchy->ControlHierarchy.Add(NewElementName, ERigControlType::Transform, NAME_None, ParentKey.Name);
+					}
+					else if (ParentKey.Type == ERigElementType::Control)
+					{
+						Hierarchy->ControlHierarchy.Add(NewElementName, ERigControlType::Transform, ParentKey.Name);
+					}
+					break;
+				}
+				case ERigElementType::Space:
+				{
+					if (!ParentKey)
+					{
+						Hierarchy->SpaceHierarchy.Add(NewElementName, ERigSpaceType::Global);
+					}
+					else if (ParentKey.Type == ERigElementType::Bone)
+					{
+						FString SpaceName = FString::Printf(TEXT("%sSpace"), *NewElementName.ToString());
+						Hierarchy->SpaceHierarchy.Add(*SpaceName, ERigSpaceType::Bone, ParentKey.Name);
+					}
+					else if (ParentKey.Type == ERigElementType::Control)
+					{
+						FString SpaceName = FString::Printf(TEXT("%sSpace"), *NewElementName.ToString());
+						Hierarchy->SpaceHierarchy.Add(*SpaceName, ERigSpaceType::Control, ParentKey.Name);
+					}
+					else if (ParentKey.Type == ERigElementType::Space)
+					{
+						Hierarchy->SpaceHierarchy.Add(NewElementName, ERigSpaceType::Space, ParentKey.Name);
+					}
+					break;
+				}
+				default:
+				{
+					return;
+				}
+			}
+		}
+		Hierarchy->ClearSelection();
+		Hierarchy->Select(FRigElementKey(NewElementName, InElementType));
 
+		FSlateApplication::Get().DismissAllMenus();
 		RefreshTreeView();
-		ControlRigEditor.Pin()->OnHierarchyChanged();
-		// reselect current selected item
-		SelectBone(NewBoneName);
 	}
 }
 
@@ -601,7 +1123,7 @@ bool SRigHierarchy::CanDuplicateItem() const
 /** Duplicate Item */
 void SRigHierarchy::HandleDuplicateItem()
 {
-	FRigHierarchy* Hierarchy = GetHierarchy();
+	FRigHierarchyContainer* Hierarchy = GetHierarchyContainer();
 	if (Hierarchy)
 	{
 		ClearDetailPanel();
@@ -609,27 +1131,57 @@ void SRigHierarchy::HandleDuplicateItem()
 		FScopedTransaction Transaction(LOCTEXT("HierarchyTreeDuplicateSelected", "Duplicate selected items from hierarchy"));
 		ControlRigBlueprint->Modify();
 
-		TArray<TSharedPtr<FRigTreeBone>> SelectedItems = TreeView->GetSelectedItems();
-		TArray<FName> NewNames;
+		TArray<TSharedPtr<FRigTreeElement>> SelectedItems = TreeView->GetSelectedItems();
+		TArray<FRigElementKey> NewKeys;
 		for (int32 Index = 0; Index < SelectedItems.Num(); ++Index)
 		{
-			FName Name = SelectedItems[Index]->CachedBone;
-			FTransform Transform = Hierarchy->GetGlobalTransform(Name);
+			TGuardValue<bool> GuardRigHierarchyChanges(bIsChangingRigHierarchy, true);
+
+			FRigElementKey Key = SelectedItems[Index]->Key;
 
-			FName ParentName = Hierarchy->GetParentName(Name);
+			switch (Key.Type)
+			{
+				case ERigElementType::Bone:
+				{
+					FTransform Transform = Hierarchy->GetGlobalTransform(Key);
+					FName ParentName = Hierarchy->BoneHierarchy[Key.Name].ParentName;
 
-			const FName NewName = CreateUniqueName(Name);
-			Hierarchy->AddBone(NewName, ParentName, Transform);
-			NewNames.Add(NewName);
+					const FName NewName = CreateUniqueName(Key.Name, Key.Type);
+					FRigBone& NewBone = Hierarchy->BoneHierarchy.Add(NewName, ParentName, Transform);
+					NewKeys.Add(NewBone.GetElementKey());
+					break;
+				}
+				case ERigElementType::Control:
+				{
+					FRigControl Control = Hierarchy->ControlHierarchy[Key.Name];
+					const FName NewName = CreateUniqueName(Key.Name, Key.Type);
+					FRigControl& NewControl = Hierarchy->ControlHierarchy.Add(NewName, Control.ControlType, Control.ParentName, Control.SpaceName, Control.InitialValue, Control.GizmoName, Control.GizmoTransform, Control.GizmoColor);
+					NewKeys.Add(NewControl.GetElementKey());
+					break;
+				}
+				case ERigElementType::Space:
+				{
+					FRigSpace Space = Hierarchy->SpaceHierarchy[Key.Name];
+					const FName NewName = CreateUniqueName(Key.Name, Key.Type);
+					FRigSpace& NewSpace = Hierarchy->SpaceHierarchy.Add(NewName, Space.SpaceType, Space.ParentName, Space.InitialTransform);
+					NewKeys.Add(NewSpace.GetElementKey());
+					break;
+				}
+				default:
+				{
+					return;
+				}
+			}
 		}
 
-		RefreshTreeView();
-		ControlRigEditor.Pin()->OnHierarchyChanged();
-		
-		for (int32 Index = 0; Index < NewNames.Num(); ++Index)
+		Hierarchy->ClearSelection();
+		for (int32 Index = 0; Index < NewKeys.Num(); ++Index)
 		{
-			SelectBone(NewNames[Index]);
+			Hierarchy->Select(NewKeys[Index]);
 		}
+
+		FSlateApplication::Get().DismissAllMenus();
+		RefreshTreeView();
 	}
 }
 
@@ -642,7 +1194,7 @@ bool SRigHierarchy::CanRenameItem() const
 /** Delete Item */
 void SRigHierarchy::HandleRenameItem()
 {
-	FRigHierarchy* Hierarchy = GetHierarchy();
+	FRigHierarchyContainer* Hierarchy = GetHierarchyContainer();
 	if (Hierarchy)
 	{
 		ClearDetailPanel();
@@ -650,46 +1202,136 @@ void SRigHierarchy::HandleRenameItem()
 		FScopedTransaction Transaction(LOCTEXT("HierarchyTreeRenameSelected", "Rename selected item from hierarchy"));
 		ControlRigBlueprint->Modify();
 
-		TArray<TSharedPtr<FRigTreeBone>> SelectedItems = TreeView->GetSelectedItems();
-		if (SelectedItems.Num() > 0)
+		TArray<TSharedPtr<FRigTreeElement>> SelectedItems = TreeView->GetSelectedItems();
+		if (SelectedItems.Num() == 1)
 		{
 			SelectedItems[0]->RequestRename();
 		}
 	}
 }
 
-FRigHierarchy* SRigHierarchy::GetHierarchy() const
+bool SRigHierarchy::CanPasteItems() const
+{
+	return true;
+}
+
+bool SRigHierarchy::CanCopyOrPasteItems() const
+{
+	return TreeView->GetNumItemsSelected() > 0;
+}
+
+void SRigHierarchy::HandleCopyItems()
+{
+	if (FRigHierarchyContainer* Hierarchy = GetHierarchyContainer())
+	{
+		TArray<FRigElementKey> Selection = GetHierarchyContainer()->CurrentSelection();
+		FString Content = Hierarchy->ExportToText(Selection);
+		FPlatformApplicationMisc::ClipboardCopy(*Content);
+	}
+}
+
+void SRigHierarchy::HandlePasteItems()
+{
+	if (FRigHierarchyContainer* Hierarchy = GetHierarchyContainer())
+	{
+		FString Content;
+		FPlatformApplicationMisc::ClipboardPaste(Content);
+
+		FScopedTransaction Transaction(LOCTEXT("HierarchyTreePaste", "Pasted rig elements."));
+		ControlRigBlueprint->Modify();
+
+		Hierarchy->ImportFromText(Content, ERigHierarchyImportMode::Append, true);
+		ControlRigBlueprint->PropagateHierarchyFromBPToInstances(false);
+	}
+}
+
+void SRigHierarchy::HandlePasteLocalTransforms()
+{
+	if (FRigHierarchyContainer* Hierarchy = GetHierarchyContainer())
+	{
+		FString Content;
+		FPlatformApplicationMisc::ClipboardPaste(Content);
+
+		FScopedTransaction Transaction(LOCTEXT("HierarchyTreePasteLocal", "Pasted local transforms."));
+		ControlRigBlueprint->Modify();
+
+		Hierarchy->ImportFromText(Content, ERigHierarchyImportMode::ReplaceLocalTransform, true);
+		ControlRigBlueprint->PropagateHierarchyFromBPToInstances(false);
+	}
+}
+
+void SRigHierarchy::HandlePasteGlobalTransforms()
+{
+	if (FRigHierarchyContainer* Hierarchy = GetHierarchyContainer())
+	{
+		FString Content;
+		FPlatformApplicationMisc::ClipboardPaste(Content);
+
+		FScopedTransaction Transaction(LOCTEXT("HierarchyTreePasteGlobal", "Pasted global transforms."));
+		ControlRigBlueprint->Modify();
+
+		Hierarchy->ImportFromText(Content, ERigHierarchyImportMode::ReplaceGlobalTransform, true);
+		ControlRigBlueprint->PropagateHierarchyFromBPToInstances(false);
+	}
+}
+
+FRigHierarchyContainer* SRigHierarchy::GetHierarchyContainer() const
 {
 	if (ControlRigBlueprint.IsValid())
 	{
-		return &ControlRigBlueprint->Hierarchy;
+		return &ControlRigBlueprint->HierarchyContainer;
 	}
 
 	return nullptr;
 }
 
-FRigHierarchy* SRigHierarchy::GetInstanceHierarchy() const
+FRigHierarchyContainer* SRigHierarchy::GetDebuggedHierarchyContainer() const
 {
-	if (ControlRigEditor.IsValid())
+	if (UControlRig* CurrentRig = ControlRigEditor.Pin()->ControlRig)
 	{
-		UControlRig* ControlRig = ControlRigEditor.Pin()->GetInstanceRig();
-		if (ControlRig)
-		{
-			return &ControlRig->Hierarchy.BaseHierarchy;
-		}
+		return (FRigHierarchyContainer*)CurrentRig->GetHierarchy();
+	}
+	else if (UControlRig* DebuggedRig = Cast<UControlRig>(ControlRigBlueprint->GetObjectBeingDebugged()))
+	{
+		return (FRigHierarchyContainer*)DebuggedRig->GetHierarchy();
 	}
 
-	return nullptr;
+	return GetHierarchyContainer();
 }
-FName SRigHierarchy::CreateUniqueName(const FName& InBaseName) const
+
+FName SRigHierarchy::CreateUniqueName(const FName& InBaseName, ERigElementType InElementType) const
 {
-	return UtilityHelpers::CreateUniqueName(InBaseName, [this](const FName& CurName) { return GetHierarchy()->GetIndex(CurName) == INDEX_NONE; });
+	switch (InElementType)
+	{
+		case ERigElementType::Bone:
+		{
+			FRigBoneHierarchy& Hierarchy = GetHierarchyContainer()->BoneHierarchy;
+			return UtilityHelpers::CreateUniqueName(InBaseName, [&Hierarchy](const FName& CurName) { return Hierarchy.GetIndex(CurName) == INDEX_NONE; });
+		}
+		case ERigElementType::Control:
+		{
+			FRigControlHierarchy& Hierarchy = GetHierarchyContainer()->ControlHierarchy;
+			return UtilityHelpers::CreateUniqueName(InBaseName, [&Hierarchy](const FName& CurName) { return Hierarchy.GetIndex(CurName) == INDEX_NONE; });
+		}
+		case ERigElementType::Space:
+		{
+			FRigSpaceHierarchy& Hierarchy = GetHierarchyContainer()->SpaceHierarchy;
+			return UtilityHelpers::CreateUniqueName(InBaseName, [&Hierarchy](const FName& CurName) { return Hierarchy.GetIndex(CurName) == INDEX_NONE; });
+		}
+		default:
+		{
+			break;
+		}
+	}
+	ensure(false);
+	return NAME_None;
 }
 
 void SRigHierarchy::PostRedo(bool bSuccess) 
 {
 	if (bSuccess)
 	{
+		ControlRigBlueprint->PropagateHierarchyFromBPToInstances(false);
 		RefreshTreeView();
 	}
 }
@@ -698,24 +1340,25 @@ void SRigHierarchy::PostUndo(bool bSuccess)
 {
 	if (bSuccess)
 	{
+		ControlRigBlueprint->PropagateHierarchyFromBPToInstances(false);
 		RefreshTreeView();
 	}
 }
 
 FReply SRigHierarchy::OnDragDetected(const FGeometry& MyGeometry, const FPointerEvent& MouseEvent)
 {
-	TArray<FName> DraggedBoneNames;
-	TArray<TSharedPtr<FRigTreeBone>> SelectedItems =TreeView->GetSelectedItems();
-	for (const TSharedPtr<FRigTreeBone>& SelectedItem : SelectedItems)
+	TArray<FRigElementKey> DraggedElements;
+	TArray<TSharedPtr<FRigTreeElement>> SelectedItems = TreeView->GetSelectedItems();
+	for (const TSharedPtr<FRigTreeElement>& SelectedItem : SelectedItems)
 	{
-		DraggedBoneNames.Add(SelectedItem->CachedBone);
+		DraggedElements.Add(SelectedItem->Key);
 	}
 
 	if (MouseEvent.IsMouseButtonDown(EKeys::LeftMouseButton))
 	{
 		if (ControlRigEditor.IsValid())
 		{
-			TSharedRef<FRigHierarchyDragDropOp> DragDropOp = FRigHierarchyDragDropOp::New(MoveTemp(DraggedBoneNames));
+			TSharedRef<FRigElementHierarchyDragDropOp> DragDropOp = FRigElementHierarchyDragDropOp::New(MoveTemp(DraggedElements));
 			DragDropOp->OnPerformDropToGraph.BindSP(ControlRigEditor.Pin().Get(), &FControlRigEditor::OnGraphNodeDropToPerform);
 			return FReply::Handled().BeginDragDrop(DragDropOp);
 		}
@@ -724,55 +1367,633 @@ FReply SRigHierarchy::OnDragDetected(const FGeometry& MyGeometry, const FPointer
 	return FReply::Unhandled();
 }
 
+TOptional<EItemDropZone> SRigHierarchy::OnCanAcceptDrop(const FDragDropEvent& DragDropEvent, EItemDropZone DropZone, TSharedPtr<FRigTreeElement> TargetItem)
+{
+	TOptional<EItemDropZone> ReturnDropZone;
+
+	TSharedPtr<FRigElementHierarchyDragDropOp> RigDragDropOp = DragDropEvent.GetOperationAs<FRigElementHierarchyDragDropOp>();
+	if (RigDragDropOp.IsValid())
+	{
+		FRigHierarchyContainer* Container = GetHierarchyContainer();
+		if (Container)
+		{
+			for (const FRigElementKey& DraggedKey : RigDragDropOp->GetElements())
+			{
+				if (Container->IsParentedTo(
+					TargetItem->Key.Type,
+					Container->GetIndex(TargetItem->Key),
+					DraggedKey.Type,
+					Container->GetIndex(DraggedKey)))
+				{
+					return ReturnDropZone;
+				}
+			}
+		}
+
+		switch (TargetItem->Key.Type)
+		{
+			case ERigElementType::Bone:
+			{
+				// bones can parent anything
+				ReturnDropZone = EItemDropZone::OntoItem;
+				break;
+			}
+			case ERigElementType::Control:
+			case ERigElementType::Space:
+			{
+				for (const FRigElementKey& DraggedKey : RigDragDropOp->GetElements())
+				{
+					switch (DraggedKey.Type)
+					{
+						case ERigElementType::Control:
+						case ERigElementType::Space:
+						{
+							break;
+						}
+						default:
+						{
+							return ReturnDropZone;
+						}
+					}
+				}
+				ReturnDropZone = EItemDropZone::OntoItem;
+				break;
+			}
+			default:
+			{
+				break;
+			}
+		}
+	}
+
+	return ReturnDropZone;
+}
+
+FReply SRigHierarchy::OnAcceptDrop(const FDragDropEvent& DragDropEvent, EItemDropZone DropZone, TSharedPtr<FRigTreeElement> TargetItem)
+{
+	TSharedPtr<FRigElementHierarchyDragDropOp> RigDragDropOp = DragDropEvent.GetOperationAs<FRigElementHierarchyDragDropOp>();
+	if (RigDragDropOp.IsValid())
+	{
+		FRigHierarchyContainer* Container = GetHierarchyContainer();
+		FRigHierarchyContainer* DebuggedContainer = GetDebuggedHierarchyContainer();
+
+		if (Container && ControlRigBlueprint.IsValid())
+		{
+			TGuardValue<bool> GuardRigHierarchyChanges(bIsChangingRigHierarchy, true);
+			FScopedTransaction Transaction(LOCTEXT("HierarchyDragAndDrop", "Drag & Drop"));
+			ControlRigBlueprint->Modify();
+
+			for (const FRigElementKey& DraggedKey : RigDragDropOp->GetElements())
+			{
+				if (Container->IsParentedTo(
+					TargetItem->Key.Type,
+					Container->GetIndex(TargetItem->Key),
+					DraggedKey.Type,
+					Container->GetIndex(DraggedKey)))
+				{
+					return FReply::Unhandled();;
+				}
+			}
+
+			for (const FRigElementKey& DraggedKey : RigDragDropOp->GetElements())
+			{
+				FRigElementKey ParentKey = TargetItem->Key;
+
+				FTransform Transform = DebuggedContainer->GetGlobalTransform(DraggedKey);
+				FTransform ParentTransform = DebuggedContainer->GetGlobalTransform(ParentKey);
+
+				switch (DraggedKey.Type)
+				{
+					case ERigElementType::Bone:
+					{
+						if (TargetItem->Key.Type == ERigElementType::Bone)
+						{
+							Container->BoneHierarchy.Reparent(DraggedKey.Name, TargetItem->Key.Name);
+						}
+						break;
+					}
+					case ERigElementType::Control:
+					{
+						if (ParentKey.Type == ERigElementType::Bone)
+						{
+							for (const FRigSpace& ExistingSpace : Container->SpaceHierarchy)
+							{
+								if (ExistingSpace.SpaceType == ERigSpaceType::Bone && ExistingSpace.ParentName == ParentKey.Name)
+								{
+									ParentKey = ExistingSpace.GetElementKey();
+									break;
+								}
+							}
+
+							if (ParentKey.Type != ERigElementType::Space)
+							{
+								FString SpaceName = FString::Printf(TEXT("%sSpace"), *DraggedKey.Name.ToString());
+								FRigSpace& NewSpace = Container->SpaceHierarchy.Add(*SpaceName, ERigSpaceType::Bone, ParentKey.Name);
+								if (DebuggedContainer != Container)
+								{
+									DebuggedContainer->SpaceHierarchy.Add(*SpaceName, ERigSpaceType::Bone, ParentKey.Name);
+								}
+								ParentKey = NewSpace.GetElementKey();
+							}
+						}
+
+						if (ParentKey.Type == ERigElementType::Control)
+						{
+							Container->ControlHierarchy.SetSpace(DraggedKey.Name, NAME_None);
+							Container->ControlHierarchy.Reparent(DraggedKey.Name, ParentKey.Name);
+							if (DebuggedContainer != Container)
+							{
+								DebuggedContainer->ControlHierarchy.SetSpace(DraggedKey.Name, NAME_None);
+								DebuggedContainer->ControlHierarchy.Reparent(DraggedKey.Name, ParentKey.Name);
+							}
+						}
+						else if (ParentKey.Type == ERigElementType::Space)
+						{
+							Container->ControlHierarchy.Reparent(DraggedKey.Name, NAME_None);
+							Container->ControlHierarchy.SetSpace(DraggedKey.Name, ParentKey.Name);
+							if (DebuggedContainer != Container)
+							{
+								DebuggedContainer->ControlHierarchy.Reparent(DraggedKey.Name, NAME_None);
+								DebuggedContainer->ControlHierarchy.SetSpace(DraggedKey.Name, ParentKey.Name);
+							}
+						}
+						break;
+					}
+					case ERigElementType::Space:
+					{
+						if (TargetItem->Key.Type == ERigElementType::Bone)
+						{
+							Container->SpaceHierarchy.Reparent(DraggedKey.Name, ERigSpaceType::Bone, TargetItem->Key.Name);
+							if (DebuggedContainer != Container)
+							{
+								DebuggedContainer->SpaceHierarchy.Reparent(DraggedKey.Name, ERigSpaceType::Bone, TargetItem->Key.Name);
+							}
+						}
+						else if (TargetItem->Key.Type == ERigElementType::Control)
+						{
+							Container->SpaceHierarchy.Reparent(DraggedKey.Name, ERigSpaceType::Control, TargetItem->Key.Name);
+							if (DebuggedContainer != Container)
+							{
+								DebuggedContainer->SpaceHierarchy.Reparent(DraggedKey.Name, ERigSpaceType::Control, TargetItem->Key.Name);
+							}
+						}
+						else if (TargetItem->Key.Type == ERigElementType::Space)
+						{
+							Container->SpaceHierarchy.Reparent(DraggedKey.Name, ERigSpaceType::Space, TargetItem->Key.Name);
+							if (DebuggedContainer != Container)
+							{
+								DebuggedContainer->SpaceHierarchy.Reparent(DraggedKey.Name, ERigSpaceType::Space, TargetItem->Key.Name);
+							}
+						}
+						break;
+					}
+					case ERigElementType::Curve:
+					{
+						break;
+					}
+					default:
+					{
+						ensure(false);
+						break;
+					}
+				}
+
+				FTransform LocalTransform = Transform.GetRelativeTransform(ParentTransform);
+				DebuggedContainer->SetInitialTransform(DraggedKey, LocalTransform);
+				DebuggedContainer->SetLocalTransform(DraggedKey, LocalTransform);
+				Container->SetInitialTransform(DraggedKey, LocalTransform);
+				Container->SetLocalTransform(DraggedKey, LocalTransform);
+			}
+		}
+
+		ControlRigBlueprint->PropagateHierarchyFromBPToInstances(true);
+		RefreshTreeView();
+
+		return FReply::Handled();
+	}
 
-bool SRigHierarchy::RenameBone(const FName& OldName, const FName& NewName)
+	return FReply::Unhandled();
+}
+
+bool SRigHierarchy::RenameElement(const FRigElementKey& OldKey, const FName& NewName)
 {
 	ClearDetailPanel();
 
-	if (OldName == NewName)
+	if (OldKey.Name == NewName)
 	{
 		return true;
 	}
 
 	// make sure there is no duplicate
-	FRigHierarchy* Hierarchy = GetHierarchy();
-	if (Hierarchy)
+	if (ControlRigBlueprint.IsValid())
 	{
-		Hierarchy->Rename(OldName, NewName);
-		SelectBone(NewName);
+		FScopedTransaction Transaction(LOCTEXT("HierarchyRename", "Rename Hierarchy Element"));
+		ControlRigBlueprint->Modify();
 
-		ControlRigEditor.Pin()->OnHierarchyChanged();
-		ControlRigEditor.Pin()->OnBoneRenamed(OldName, NewName);
-		return true;
+		FRigHierarchyContainer* Container = GetHierarchyContainer();
+		switch (OldKey.Type)
+		{
+			case ERigElementType::Bone:
+			{
+				Container->BoneHierarchy.Rename(OldKey.Name, NewName);
+				return true;
+			}
+			case ERigElementType::Control:
+			{
+				Container->ControlHierarchy.Rename(OldKey.Name, NewName);
+				return true;
+			}
+			case ERigElementType::Space:
+			{
+				Container->SpaceHierarchy.Rename(OldKey.Name, NewName);
+				return true;
+			}
+			default:
+			{
+				break;
+			}
+		}
 	}
 
 	return false;
 }
 
-bool SRigHierarchy::OnVerifyNameChanged(const FName& OldName, const FName& NewName, FText& OutErrorMessage)
+bool SRigHierarchy::OnVerifyNameChanged(const FRigElementKey& OldKey, const FName& NewName, FText& OutErrorMessage)
 {
-	if (OldName == NewName)
+	if (OldKey.Name == NewName)
 	{
 		return true;
 	}
 
 	// make sure there is no duplicate
-	FRigHierarchy* Hierarchy = GetHierarchy();
-	if (Hierarchy)
+	if (ControlRigBlueprint.IsValid())
+	{
+		FRigHierarchyContainer* Container = GetHierarchyContainer();
+		switch (OldKey.Type)
+		{
+		case ERigElementType::Bone:
+		{
+			return Container->BoneHierarchy.IsNameAvailable(NewName);
+		}
+		case ERigElementType::Control:
+		{
+			return Container->ControlHierarchy.IsNameAvailable(NewName);
+		}
+		case ERigElementType::Space:
+		{
+			return Container->SpaceHierarchy.IsNameAvailable(NewName);
+		}
+		default:
+		{
+			break;
+		}
+		}
+	}
+	return false;
+}
+
+void SRigHierarchy::HandleResetTransform()
+{
+	if (IsMultiSelected() && ControlRigEditor.IsValid())
 	{
-		const int32 Found = Hierarchy->GetIndex(OldName);
-		if (Found != INDEX_NONE)
+		if (UControlRigBlueprint* Blueprint = ControlRigEditor.Pin()->GetControlRigBlueprint())
 		{
-			const int32 Duplicate = Hierarchy->GetIndex(NewName);
-			if (Duplicate != INDEX_NONE)
+			if (FRigHierarchyContainer* DebuggedContainer = GetDebuggedHierarchyContainer())
 			{
-				OutErrorMessage = FText::FromString(TEXT("Duplicate name exists"));
+				FScopedTransaction Transaction(LOCTEXT("HierarchyResetInitialTransforms", "Reset Initial Transforms"));
+				Blueprint->Modify();
 
-				return false;
+				TArray<TSharedPtr<FRigTreeElement>> SelectedItems = TreeView->GetSelectedItems();
+				for (TSharedPtr<FRigTreeElement> SelectedItem : SelectedItems)
+				{
+					FRigElementKey Key = SelectedItem->Key;
+					if (Key.Type == ERigElementType::Control || Key.Type == ERigElementType::Space)
+					{
+						FTransform InitialTransform = GetHierarchyContainer()->GetInitialTransform(Key);
+						GetHierarchyContainer()->SetLocalTransform(Key, InitialTransform);
+						DebuggedContainer->SetLocalTransform(Key, InitialTransform);
+					}
+				}
 			}
 		}
 	}
+}
 
-	return true;
+void SRigHierarchy::HandleResetInitialTransform()
+{
+	if (IsMultiSelected() && ControlRigEditor.IsValid())
+	{
+		if (UControlRigBlueprint* Blueprint = ControlRigEditor.Pin()->GetControlRigBlueprint())
+		{
+			if (FRigHierarchyContainer* DebuggedContainer = GetDebuggedHierarchyContainer())
+			{
+				FScopedTransaction Transaction(LOCTEXT("HierarchyResetInitialTransforms", "Reset Initial Transforms"));
+				Blueprint->Modify();
+
+				TArray<TSharedPtr<FRigTreeElement>> SelectedItems = TreeView->GetSelectedItems();
+				for (TSharedPtr<FRigTreeElement> SelectedItem : SelectedItems)
+				{
+					FRigElementKey Key = SelectedItem->Key;
+					if (Key.Type == ERigElementType::Control || Key.Type == ERigElementType::Space)
+					{
+						GetHierarchyContainer()->SetInitialTransform(Key, FTransform::Identity);
+						GetHierarchyContainer()->SetLocalTransform(Key, FTransform::Identity);
+						DebuggedContainer->SetInitialTransform(Key, FTransform::Identity);
+						DebuggedContainer->SetLocalTransform(Key, FTransform::Identity);
+					}
+				}
+			}
+		}
+	}
+}
+
+void SRigHierarchy::HandleResetSpace()
+{
+	if (IsControlSelected())
+	{
+		if (UControlRigBlueprint* Blueprint = ControlRigEditor.Pin()->GetControlRigBlueprint())
+		{
+			if (FRigHierarchyContainer* DebuggedContainer = GetDebuggedHierarchyContainer())
+			{
+				FScopedTransaction Transaction(LOCTEXT("HierarchyResetSpace", "Reset Space"));
+				Blueprint->Modify();
+
+				TArray<TSharedPtr<FRigTreeElement>> SelectedItems = TreeView->GetSelectedItems();
+				for (TSharedPtr<FRigTreeElement> SelectedItem : SelectedItems)
+				{
+					FRigElementKey Key = SelectedItem->Key;
+					if (Key.Type == ERigElementType::Control)
+					{
+						FRigControl& Control = DebuggedContainer->ControlHierarchy[Key.Name];
+						if (FRigElementKey SpaceKey = Control.GetSpaceElementKey())
+						{
+							FRigSpace& Space = DebuggedContainer->SpaceHierarchy[SpaceKey.Name];
+
+							FTransform Transform = DebuggedContainer->GetGlobalTransform(Key);
+							FTransform SpaceParentTransform = DebuggedContainer->GetGlobalTransform(Space.GetParentElementKey());
+							FTransform LocalTransform = Transform.GetRelativeTransform(SpaceParentTransform);
+
+							GetHierarchyContainer()->SetInitialTransform(SpaceKey, LocalTransform);
+							GetHierarchyContainer()->SetLocalTransform(SpaceKey, LocalTransform);
+							GetHierarchyContainer()->SetInitialTransform(Key, FTransform::Identity);
+							GetHierarchyContainer()->SetLocalTransform(Key, FTransform::Identity);
+
+							DebuggedContainer->SetInitialTransform(SpaceKey, LocalTransform);
+							DebuggedContainer->SetLocalTransform(SpaceKey, LocalTransform);
+							DebuggedContainer->SetInitialTransform(Key, FTransform::Identity);
+							DebuggedContainer->SetLocalTransform(Key, FTransform::Identity);
+						}
+						else
+						{
+							FString SpaceName = FString::Printf(TEXT("%sSpace"), *Control.Name.ToString());
+							if (FRigElementKey CurrentParentKey = Control.GetParentElementKey())
+							{
+								FRigSpace& NewSpace = GetHierarchyContainer()->SpaceHierarchy.Add(*SpaceName, ERigSpaceType::Control, CurrentParentKey.Name, Control.Value.Get<FTransform>());
+								SpaceKey = NewSpace.GetElementKey();
+							}
+							else
+							{
+								FRigSpace& NewSpace = GetHierarchyContainer()->SpaceHierarchy.Add(*SpaceName, ERigSpaceType::Global, NAME_None, Control.Value.Get<FTransform>());
+								SpaceKey = NewSpace.GetElementKey();
+							}
+
+							GetHierarchyContainer()->ControlHierarchy.Reparent(Control.Name, NAME_None);
+							GetHierarchyContainer()->ControlHierarchy.SetSpace(Control.Name, SpaceKey.Name);
+							GetHierarchyContainer()->ControlHierarchy.SetInitialValue<FTransform>(Control.Name, FTransform::Identity);
+							GetHierarchyContainer()->ControlHierarchy.SetLocalTransform(Control.Name, FTransform::Identity);
+							Blueprint->PropagateHierarchyFromBPToInstances();
+						}
+					}
+				}
+			}
+		}
+	}
+}
+
+void SRigHierarchy::HandleSetInitialTransformFromCurrentTransform()
+{
+	if (IsMultiSelected())
+	{
+		if (UControlRigBlueprint* Blueprint = ControlRigEditor.Pin()->GetControlRigBlueprint())
+		{
+			if (FRigHierarchyContainer* DebuggedContainer = GetDebuggedHierarchyContainer())
+			{
+				FScopedTransaction Transaction(LOCTEXT("HierarchySetInitialTransforms", "Set Initial Transforms"));
+				Blueprint->Modify();
+
+				TArray<TSharedPtr<FRigTreeElement>> SelectedItems = TreeView->GetSelectedItems();
+				TMap<FRigElementKey, FTransform> GlobalTransforms;
+				TMap<FRigElementKey, FTransform> ParentGlobalTransforms;
+
+				for (TSharedPtr<FRigTreeElement> SelectedItem : SelectedItems)
+				{
+					FRigElementKey Key = SelectedItem->Key;
+					if (Key.Type == ERigElementType::Control || Key.Type == ERigElementType::Space || Key.Type == ERigElementType::Bone)
+					{
+						GlobalTransforms.FindOrAdd(Key) = DebuggedContainer->GetGlobalTransform(Key);
+						FRigElementKey ParentKey;
+						switch (Key.Type)
+						{
+							case ERigElementType::Bone:
+							{
+								const FRigBone& Element = DebuggedContainer->BoneHierarchy[Key.Name];
+								ParentKey = Element.GetParentElementKey();
+								break;
+							}
+							case ERigElementType::Control:
+							{
+								const FRigControl& Element = DebuggedContainer->ControlHierarchy[Key.Name];
+								if (Element.SpaceName != NAME_None)
+								{
+									ParentKey = Element.GetSpaceElementKey();
+								}
+								else
+								{
+									ParentKey = Element.GetParentElementKey();
+								}
+								break;
+							}
+							case ERigElementType::Space:
+							{
+								const FRigSpace& Element = DebuggedContainer->SpaceHierarchy[Key.Name];
+								ParentKey = Element.GetParentElementKey();
+								break;
+							}
+							case ERigElementType::Curve:
+							{
+								break;
+							}
+							default:
+							{
+								ensure(false);
+								break;
+							}
+						}
+
+						if (ParentKey)
+						{
+							ParentGlobalTransforms.FindOrAdd(Key) = DebuggedContainer->GetGlobalTransform(ParentKey);
+						}
+						else
+						{
+							ParentGlobalTransforms.FindOrAdd(Key) = FTransform::Identity;
+						}
+					}
+
+				}
+
+				for (TSharedPtr<FRigTreeElement> SelectedItem : SelectedItems)
+				{
+					FRigElementKey Key = SelectedItem->Key;
+					if (Key.Type == ERigElementType::Control || Key.Type == ERigElementType::Space || Key.Type == ERigElementType::Bone)
+					{
+						FTransform LocalTransform = GlobalTransforms[Key].GetRelativeTransform(ParentGlobalTransforms[Key]);
+						GetHierarchyContainer()->SetInitialTransform(Key, LocalTransform);
+						GetHierarchyContainer()->SetLocalTransform(Key, LocalTransform);
+						DebuggedContainer->SetInitialTransform(Key, LocalTransform);
+						DebuggedContainer->SetLocalTransform(Key, LocalTransform);
+					}
+				}
+			}
+		}
+	}
 }
+
+void SRigHierarchy::HandleFrameSelection()
+{
+	TArray<TSharedPtr<FRigTreeElement>> SelectedItems = TreeView->GetSelectedItems();
+	for (TSharedPtr<FRigTreeElement> SelectedItem : SelectedItems)
+	{
+		SetExpansionRecursive(SelectedItem, true);
+	}
+
+	if (SelectedItems.Num() > 0)
+	{
+		TreeView->RequestScrollIntoView(SelectedItems.Last());
+	}
+}
+
+bool SRigHierarchy::FindClosestBone(const FVector& Point, FName& OutRigElementName, FTransform& OutGlobalTransform) const
+{
+	if (FRigHierarchyContainer* DebuggedContainer = GetDebuggedHierarchyContainer())
+	{
+		const FRigBoneHierarchy& RigHierarchy = DebuggedContainer->BoneHierarchy;
+		float NearestDistance = BIG_NUMBER;
+
+		for (int32 Index = 0; Index < RigHierarchy.Num(); ++Index)
+		{
+			FTransform CurTransform = RigHierarchy.GetGlobalTransform(Index);
+			float CurDistance = FVector::Distance(CurTransform.GetLocation(), Point);
+			if (CurDistance < NearestDistance)
+			{
+				NearestDistance = CurDistance;
+				OutGlobalTransform = CurTransform;
+				OutRigElementName = RigHierarchy[Index].Name;
+			}
+		}
+
+		return (OutRigElementName != NAME_None);
+	}
+	return false;
+}
+
+void SRigHierarchy::HandleSetInitialTransformFromClosestBone()
+{
+	if (IsControlOrSpaceSelected())
+	{
+		if (UControlRigBlueprint* Blueprint = ControlRigEditor.Pin()->GetControlRigBlueprint())
+		{
+			if (FRigHierarchyContainer* DebuggedContainer = GetDebuggedHierarchyContainer())
+			{
+				FScopedTransaction Transaction(LOCTEXT("HierarchySetInitialTransforms", "Set Initial Transforms"));
+				Blueprint->Modify();
+
+				TArray<TSharedPtr<FRigTreeElement>> SelectedItems = TreeView->GetSelectedItems();
+				TMap<FRigElementKey, FTransform> ClosestTransforms;
+				TMap<FRigElementKey, FTransform> ParentGlobalTransforms;
+
+				for (TSharedPtr<FRigTreeElement> SelectedItem : SelectedItems)
+				{
+					FRigElementKey Key = SelectedItem->Key;
+					if (Key.Type == ERigElementType::Control || Key.Type == ERigElementType::Space)
+					{
+						FTransform GlobalTransform = DebuggedContainer->GetGlobalTransform(Key);
+						FTransform ClosestTransform;
+						FName ClosestRigElement;
+
+						if (!FindClosestBone(GlobalTransform.GetLocation(), ClosestRigElement, ClosestTransform))
+						{
+							continue;
+						}
+
+						ClosestTransforms.FindOrAdd(Key) = ClosestTransform;
+
+						FRigElementKey ParentKey;
+						switch (Key.Type)
+						{
+							case ERigElementType::Control:
+							{
+								const FRigControl& Element = DebuggedContainer->ControlHierarchy[Key.Name];
+								if (Element.SpaceName != NAME_None)
+								{
+									ParentKey = Element.GetSpaceElementKey();
+								}
+								else
+								{
+									ParentKey = Element.GetParentElementKey();
+								}
+								break;
+							}
+							case ERigElementType::Space:
+							{
+								const FRigSpace& Element = DebuggedContainer->SpaceHierarchy[Key.Name];
+								ParentKey = Element.GetParentElementKey();
+								break;
+							}
+							case ERigElementType::Bone:
+							case ERigElementType::Curve:
+							{
+								break;
+							}
+							default:
+							{
+								ensure(false);
+								break;
+							}
+						}
+
+						if (ParentKey)
+						{
+							ParentGlobalTransforms.FindOrAdd(Key) = DebuggedContainer->GetGlobalTransform(ParentKey);
+						}
+						else
+						{
+							ParentGlobalTransforms.FindOrAdd(Key) = FTransform::Identity;
+						}
+					}
+
+				}
+
+				for (TSharedPtr<FRigTreeElement> SelectedItem : SelectedItems)
+				{
+					FRigElementKey Key = SelectedItem->Key;
+					if (Key.Type == ERigElementType::Control || Key.Type == ERigElementType::Space || Key.Type == ERigElementType::Bone)
+					{
+						if (!ClosestTransforms.Contains(Key))
+						{
+							continue;
+						}
+
+						FTransform LocalTransform = ClosestTransforms[Key].GetRelativeTransform(ParentGlobalTransforms[Key]);
+						GetHierarchyContainer()->SetInitialTransform(Key, LocalTransform);
+						GetHierarchyContainer()->SetLocalTransform(Key, LocalTransform);
+						DebuggedContainer->SetInitialTransform(Key, LocalTransform);
+						DebuggedContainer->SetLocalTransform(Key, LocalTransform);
+					}
+				}
+			}
+		}
+	}
+}
+
 #undef LOCTEXT_NAMESPACE
\ No newline at end of file
diff --git a/Engine/Plugins/Experimental/ControlRig/Source/ControlRigEditor/Private/Editor/SRigHierarchy.h b/Engine/Plugins/Experimental/ControlRig/Source/ControlRigEditor/Private/Editor/SRigHierarchy.h
--- a/Engine/Plugins/Experimental/ControlRig/Source/ControlRigEditor/Private/Editor/SRigHierarchy.h
+++ b/Engine/Plugins/Experimental/ControlRig/Source/ControlRigEditor/Private/Editor/SRigHierarchy.h
@@ -4,7 +4,7 @@
 
 #include "CoreMinimal.h"
 #include "Widgets/Views/STreeView.h"
-#include "Hierarchy.h"
+#include "Rigs/RigHierarchyContainer.h"
 #include "EditorUndoClient.h"
 #include "DragAndDrop/GraphNodeDragDropOp.h"
 
@@ -17,20 +17,20 @@ struct FAssetData;
 class FMenuBuilder;
 class SRigHierarchyItem;
 
-DECLARE_DELEGATE_RetVal_TwoParams(bool, FOnRenameBone, const FName& /*OldName*/, const FName& /*NewName*/);
-DECLARE_DELEGATE_RetVal_ThreeParams(bool, FOnVerifyBoneNameChanged, const FName& /*OldName*/, const FName& /*NewName*/, FText& /*OutErrorMessage*/);
+DECLARE_DELEGATE_RetVal_TwoParams(bool, FOnRenameElement, const FRigElementKey& /*OldKey*/, const FName& /*NewName*/);
+DECLARE_DELEGATE_RetVal_ThreeParams(bool, FOnVerifyElementNameChanged, const FRigElementKey& /*OldKey*/, const FName& /*NewName*/, FText& /*OutErrorMessage*/);
 
 /** An item in the tree */
-class FRigTreeBone : public TSharedFromThis<FRigTreeBone>
+class FRigTreeElement : public TSharedFromThis<FRigTreeElement>
 {
 public:
-	FRigTreeBone(const FName& InBone, TWeakPtr<SRigHierarchy> InHierarchyHandler);
+	FRigTreeElement(const FRigElementKey& InKey, TWeakPtr<SRigHierarchy> InHierarchyHandler);
 public:
-	/** Bone Data to display */
-	FName CachedBone;
-	TArray<TSharedPtr<FRigTreeBone>> Children;
+	/** Element Data to display */
+	FRigElementKey Key;
+	TArray<TSharedPtr<FRigTreeElement>> Children;
 
-	TSharedRef<ITableRow> MakeTreeRowWidget(TSharedPtr<FControlRigEditor> InControlRigEditor, const TSharedRef<STableViewBase>& InOwnerTable, TSharedRef<FRigTreeBone> InRigTreeBone, TSharedRef<FUICommandList> InCommandList, TSharedPtr<SRigHierarchy> InHierarchy);
+	TSharedRef<ITableRow> MakeTreeRowWidget(TSharedPtr<FControlRigEditor> InControlRigEditor, const TSharedRef<STableViewBase>& InOwnerTable, TSharedRef<FRigTreeElement> InRigTreeElement, TSharedRef<FUICommandList> InCommandList, TSharedPtr<SRigHierarchy> InHierarchy);
 
 	void RequestRename();
 
@@ -39,55 +39,55 @@ public:
 	FOnRenameRequested OnRenameRequested;
 };
 
-class FRigHierarchyDragDropOp : public FGraphNodeDragDropOp
+class FRigElementHierarchyDragDropOp : public FGraphNodeDragDropOp
 {
 public:
-	DRAG_DROP_OPERATOR_TYPE(FRigHierarchyDragDropOp, FGraphNodeDragDropOp)
+	DRAG_DROP_OPERATOR_TYPE(FRigElementHierarchyDragDropOp, FGraphNodeDragDropOp)
 
-	static TSharedRef<FRigHierarchyDragDropOp> New(TArray<FName> InBoneNames);
+	static TSharedRef<FRigElementHierarchyDragDropOp> New(const TArray<FRigElementKey>& InElements);
 
 	virtual TSharedPtr<SWidget> GetDefaultDecorator() const override;
 
 	/** @return true if this drag operation contains property paths */
-	bool HasBoneNames() const
+	bool HasElements() const
 	{
-		return BoneNames.Num() > 0;
+		return Elements.Num() > 0;
 	}
 
 	/** @return The property paths from this drag operation */
-	const TArray<FName>& GetBoneNames() const
+	const TArray<FRigElementKey>& GetElements() const
 	{
-		return BoneNames;
+		return Elements;
 	}
 
-	FString GetJoinedBoneNames() const;
+	FString GetJoinedElementNames() const;
 
 private:
 
 	/** Data for the property paths this item represents */
-	TArray<FName> BoneNames;
+	TArray<FRigElementKey> Elements;
 };
 
- class SRigHierarchyItem : public STableRow<TSharedPtr<FRigTreeBone>>
+ class SRigHierarchyItem : public STableRow<TSharedPtr<FRigTreeElement>>
 {
 	SLATE_BEGIN_ARGS(SRigHierarchyItem) {}
 		/** Callback when the text is committed. */
-		SLATE_EVENT(FOnRenameBone, OnRenameBone)
+		SLATE_EVENT(FOnRenameElement, OnRenameElement)
 		/** Called whenever the text is changed interactively by the user */
-		SLATE_EVENT(FOnVerifyBoneNameChanged, OnVerifyBoneNameChanged)
+		SLATE_EVENT(FOnVerifyElementNameChanged, OnVerifyElementNameChanged)
 	SLATE_END_ARGS()
 
-	void Construct(const FArguments& InArgs, TSharedPtr<FControlRigEditor> InControlRigEditor, const TSharedRef<STableViewBase>& OwnerTable, TSharedRef<FRigTreeBone> InControlRigTreeBone, TSharedRef<FUICommandList> InCommandList, TSharedPtr<SRigHierarchy> InHierarchy);
+	void Construct(const FArguments& InArgs, TSharedPtr<FControlRigEditor> InControlRigEditor, const TSharedRef<STableViewBase>& OwnerTable, TSharedRef<FRigTreeElement> InRigTreeElement, TSharedRef<FUICommandList> InCommandList, TSharedPtr<SRigHierarchy> InHierarchy);
 	void OnNameCommitted(const FText& InText, ETextCommit::Type InCommitType) const;
 	bool OnVerifyNameChanged(const FText& InText, FText& OutErrorMessage);
 
 private:
-	TWeakPtr<FRigTreeBone> WeakRigTreeBone;
+	TWeakPtr<FRigTreeElement> WeakRigTreeElement;
 	TWeakPtr<FUICommandList> WeakCommandList;
 	TWeakPtr<FControlRigEditor> ControlRigEditor;
 
-	FOnRenameBone OnRenameBone;
-	FOnVerifyBoneNameChanged OnVerifyBoneNameChanged;
+	FOnRenameElement OnRenameElement;
+	FOnVerifyElementNameChanged OnVerifyElementNameChanged;
 
 	FText GetName() const;
 };
@@ -107,14 +107,17 @@ private:
 	/** Bind commands that this widget handles */
 	void BindCommands();
 
+	/** SWidget interface */
+	virtual FReply OnKeyDown(const FGeometry& MyGeometry, const FKeyEvent& InKeyEvent);
+
 	/** Rebuild the tree view */
 	void RefreshTreeView();
 
 	/** Make a row widget for the table */
-	TSharedRef<ITableRow> MakeTableRowWidget(TSharedPtr<FRigTreeBone> InItem, const TSharedRef<STableViewBase>& OwnerTable);
+	TSharedRef<ITableRow> MakeTableRowWidget(TSharedPtr<FRigTreeElement> InItem, const TSharedRef<STableViewBase>& OwnerTable);
 
 	/** Get children for the tree */
-	void HandleGetChildrenForTree(TSharedPtr<FRigTreeBone> InItem, TArray<TSharedPtr<FRigTreeBone>>& OutChildren);
+	void HandleGetChildrenForTree(TSharedPtr<FRigTreeElement> InItem, TArray<TSharedPtr<FRigTreeElement>>& OutChildren);
 
 	/** Check whether we can deleting the selected item(s) */
 	bool CanDeleteItem() const;
@@ -122,8 +125,8 @@ private:
 	/** Delete Item */
 	void HandleDeleteItem();
 
-	/** Delete Item */
-	void HandleNewItem();
+	/** Create a new item */
+	void HandleNewItem(ERigElementType InElementType);
 
 	/** Check whether we can deleting the selected item(s) */
 	bool CanDuplicateItem() const;
@@ -137,8 +140,15 @@ private:
 	/** Delete Item */
 	void HandleRenameItem();
 
+	bool CanPasteItems() const;
+	bool CanCopyOrPasteItems() const;
+	void HandleCopyItems();
+	void HandlePasteItems();
+	void HandlePasteLocalTransforms();
+	void HandlePasteGlobalTransforms();
+
 	/** Set Selection Changed */
-	void OnSelectionChanged(TSharedPtr<FRigTreeBone> Selection, ESelectInfo::Type SelectInfo);
+	void OnSelectionChanged(TSharedPtr<FRigTreeElement> Selection, ESelectInfo::Type SelectInfo);
 
 	TSharedPtr< SWidget > CreateContextMenu();
 
@@ -148,8 +158,14 @@ private:
 
 	// reply to a drag operation
 	FReply OnDragDetected(const FGeometry& MyGeometry, const FPointerEvent& MouseEvent);
+	TOptional<EItemDropZone> OnCanAcceptDrop(const FDragDropEvent& DragDropEvent, EItemDropZone DropZone, TSharedPtr<FRigTreeElement> TargetItem);
+	FReply OnAcceptDrop(const FDragDropEvent& DragDropEvent, EItemDropZone DropZone, TSharedPtr<FRigTreeElement> TargetItem);
 
 private:
+
+	void FillContextMenu(class FMenuBuilder& MenuBuilder);
+	TSharedPtr<FUICommandList> GetContextMenuCommands();
+
 	/** Our owning control rig editor */
 	TWeakPtr<FControlRigEditor> ControlRigEditor;
 
@@ -160,13 +176,16 @@ private:
 	void OnFilterTextChanged(const FText& SearchText);
 
 	/** Tree view widget */
-	TSharedPtr<STreeView<TSharedPtr<FRigTreeBone>>> TreeView;
+	TSharedPtr<STreeView<TSharedPtr<FRigTreeElement>>> TreeView;
 
 	/** Backing array for tree view */
-	TArray<TSharedPtr<FRigTreeBone>> RootBones;
+	TArray<TSharedPtr<FRigTreeElement>> RootElements;
 
-	/** Backing array for tree view (filtered, displayed) */
-	TArray<TSharedPtr<FRigTreeBone>> FilteredRootBones;
+	/** A map for looking up items based on their key */
+	TMap<FRigElementKey, TSharedPtr<FRigTreeElement>> ElementMap;
+
+	/** A map for looking up a parent based on their key */
+	TMap<FRigElementKey, FRigElementKey> ParentMap;
 
 	TWeakObjectPtr<UControlRigBlueprint> ControlRigBlueprint;
 	
@@ -175,23 +194,48 @@ private:
 
 	bool IsMultiSelected() const;
 	bool IsSingleSelected() const;
+	bool IsControlSelected() const;
+	bool IsControlOrSpaceSelected() const;
 
-	FRigHierarchy* GetHierarchy() const;
-	FRigHierarchy* GetInstanceHierarchy() const;
+	FRigHierarchyContainer* GetHierarchyContainer() const;
+	FRigHierarchyContainer* GetDebuggedHierarchyContainer() const;
 
 	void ImportHierarchy(const FAssetData& InAssetData);
 	void CreateImportMenu(FMenuBuilder& MenuBuilder);
+	void CreateRefreshMenu(FMenuBuilder& MenuBuilder);
 	bool ShouldFilterOnImport(const FAssetData& AssetData) const;
+	void RefreshHierarchy(const FAssetData& InAssetData);
+
+	void HandleResetTransform();
+	void HandleResetInitialTransform();
+	void HandleResetSpace();
+	void HandleSetInitialTransformFromCurrentTransform();
+	void HandleSetInitialTransformFromClosestBone();
+	void HandleFrameSelection();
+	bool FindClosestBone(const FVector& Point, FName& OutRigElementName, FTransform& OutGlobalTransform) const;
 
-	FName CreateUniqueName(const FName& InBaseName) const;
+	FName CreateUniqueName(const FName& InBaseName, ERigElementType InElementType) const;
 
-	void SetExpansionRecursive(TSharedPtr<FRigTreeBone> InBones);
+	void SetExpansionRecursive(TSharedPtr<FRigTreeElement> InElements, bool bTowardsParent);
 
 	void ClearDetailPanel() const;
-	void SelectBone(const FName& BoneName) const;
+
+	bool bIsChangingRigHierarchy;
+	void OnRigElementAdded(FRigHierarchyContainer* Container, const FRigElementKey& InKey);
+	void OnRigElementRemoved(FRigHierarchyContainer* Container, const FRigElementKey& InKey);
+	void OnRigElementRenamed(FRigHierarchyContainer* Container, ERigElementType ElementType, const FName& InOldName, const FName& InNewName);
+	void OnRigElementReparented(FRigHierarchyContainer* Container, const FRigElementKey& InKey, const FName& InOldParentName, const FName& InNewParentName);
+	void OnRigElementSelected(FRigHierarchyContainer* Container, const FRigElementKey& InKey, bool bSelected);
+
+	static TSharedPtr<FRigTreeElement> FindElement(const FRigElementKey& InElementKey, TSharedPtr<FRigTreeElement> CurrentItem);
+	void AddElement(FRigElementKey InKey, FRigElementKey InParentKey = FRigElementKey());
+	void AddBoneElement(FRigBone InBone);
+	void AddControlElement(FRigControl InControl);
+	void AddSpaceElement(FRigSpace InSpace);
+
 public:
-	bool RenameBone(const FName& OldName, const FName& NewName);
-	bool OnVerifyNameChanged(const FName& OldName, const FName& NewName, FText& OutErrorMessage);
+	bool RenameElement(const FRigElementKey& OldKey, const FName& NewName);
+	bool OnVerifyNameChanged(const FRigElementKey& OldKey, const FName& NewName, FText& OutErrorMessage);
 
 	friend class SRigHierarchyItem;
 };
\ No newline at end of file
diff --git a/Engine/Plugins/Experimental/ControlRig/Source/ControlRigEditor/Private/Graph/ControlRigGraphPanelPinFactory.cpp b/Engine/Plugins/Experimental/ControlRig/Source/ControlRigEditor/Private/Graph/ControlRigGraphPanelPinFactory.cpp
--- a/Engine/Plugins/Experimental/ControlRig/Source/ControlRigEditor/Private/Graph/ControlRigGraphPanelPinFactory.cpp
+++ b/Engine/Plugins/Experimental/ControlRig/Source/ControlRigEditor/Private/Graph/ControlRigGraphPanelPinFactory.cpp
@@ -3,8 +3,8 @@
 #include "ControlRigGraphPanelPinFactory.h"
 #include "Graph/ControlRigGraphSchema.h"
 #include "Graph/ControlRigGraphNode.h"
-#include "Graph/SGraphPinBoneName.h"
-#include "Graph/SGraphPinCurveName.h"
+#include "Graph/ControlRigGraph.h"
+#include "Graph/SGraphPinNameList.h"
 #include "Graph/SGraphPinCurveFloat.h"
 #include "KismetPins/SGraphPinExec.h"
 #include "ControlRig.h"
@@ -19,6 +19,8 @@ TSharedPtr<SGraphPin> FControlRigGraphPanelPinFactory::CreatePin(UEdGraphPin* In
 		UControlRigGraphNode* RigNode = Cast<UControlRigGraphNode>(InPin->GetOwningNode());
 		if (RigNode)
 		{
+			UControlRigGraph* RigGraph = Cast<UControlRigGraph>(RigNode->GetGraph());
+
 			// use a bone name widget in case we are looking at a name with appropriate metadata
 			if (InPin->PinType.PinCategory == UEdGraphSchema_K2::PC_Name)
 			{
@@ -41,11 +43,23 @@ TSharedPtr<SGraphPin> FControlRigGraphPanelPinFactory::CreatePin(UEdGraphPin* In
 						{
 							if (Property->HasMetaData(UControlRig::BoneNameMetaName))
 							{
-								return SNew(SGraphPinBoneName, InPin);
+								return SNew(SGraphPinNameList, InPin)
+									.OnGetNameListContent_UObject(RigGraph, &UControlRigGraph::GetBoneNameList);
+							}
+							if (Property->HasMetaData(UControlRig::ControlNameMetaName))
+							{
+								return SNew(SGraphPinNameList, InPin)
+									.OnGetNameListContent_UObject(RigGraph, &UControlRigGraph::GetControlNameList);
+							}
+							if (Property->HasMetaData(UControlRig::SpaceNameMetaName))
+							{
+								return SNew(SGraphPinNameList, InPin)
+									.OnGetNameListContent_UObject(RigGraph, &UControlRigGraph::GetSpaceNameList);
 							}
 							else if (Property->HasMetaData(UControlRig::CurveNameMetaName))
 							{
-								return SNew(SGraphPinCurveName, InPin);
+								return SNew(SGraphPinNameList, InPin)
+									.OnGetNameListContent_UObject(RigGraph, &UControlRigGraph::GetCurveNameList);
 							}
 						}
 					}
diff --git a/Engine/Plugins/Experimental/ControlRig/Source/ControlRigEditor/Private/Graph/SGraphPinNameList.cpp b/Engine/Plugins/Experimental/ControlRig/Source/ControlRigEditor/Private/Graph/SGraphPinNameList.cpp
new file mode 100644
--- /dev/null
+++ b/Engine/Plugins/Experimental/ControlRig/Source/ControlRigEditor/Private/Graph/SGraphPinNameList.cpp
@@ -0,0 +1,98 @@
+// Copyright 1998-2019 Epic Games, Inc. All Rights Reserved.
+
+
+#include "Graph/SGraphPinNameList.h"
+#include "Graph/SGraphPinNameListValueWidget.h"
+#include "Widgets/Layout/SBox.h"
+#include "Widgets/Input/SEditableTextBox.h"
+#include "EdGraphSchema_K2.h"
+#include "ScopedTransaction.h"
+#include "DetailLayoutBuilder.h"
+
+void SGraphPinNameList::Construct(const FArguments& InArgs, UEdGraphPin* InGraphPinObj)
+{
+	this->OnGetNameListContent = InArgs._OnGetNameListContent;
+
+	SGraphPin::Construct(SGraphPin::FArguments(), InGraphPinObj);
+}
+
+TSharedRef<SWidget>	SGraphPinNameList::GetDefaultValueWidget()
+{
+	TSharedPtr<FString> InitialSelected;
+	for (TSharedPtr<FString> Item : GetNameList())
+	{
+		if (Item->Equals(GetNameListText().ToString()))
+		{
+			InitialSelected = Item;
+		}
+	}
+
+	return SNew(SBox)
+		.MinDesiredWidth(150)
+		.MaxDesiredWidth(400)
+		[
+			SAssignNew(NameListComboBox, SGraphPinNameListValueWidget)
+				.Visibility(this, &SGraphPin::GetDefaultValueVisibility)
+				.OptionsSource(&GetNameList())
+				.OnGenerateWidget(this, &SGraphPinNameList::MakeNameListItemWidget)
+				.OnSelectionChanged(this, &SGraphPinNameList::OnNameListChanged)
+				.OnComboBoxOpening(this, &SGraphPinNameList::OnNameListComboBox)
+				.InitiallySelectedItem(InitialSelected)
+				.Content()
+				[
+					SNew(STextBlock)
+					.Text(this, &SGraphPinNameList::GetNameListText)
+				]
+		];
+}
+
+const TArray<TSharedPtr<FString>>& SGraphPinNameList::GetNameList() const
+{
+	if (OnGetNameListContent.IsBound())
+	{
+		return OnGetNameListContent.Execute();
+	}
+	return EmptyList;
+}
+
+FText SGraphPinNameList::GetNameListText() const
+{
+	return FText::FromString( GraphPinObj->GetDefaultAsString() );
+}
+
+void SGraphPinNameList::SetNameListText(const FText& NewTypeInValue, ETextCommit::Type /*CommitInfo*/)
+{
+	if(!GraphPinObj->GetDefaultAsString().Equals(NewTypeInValue.ToString()))
+	{
+		const FScopedTransaction Transaction( NSLOCTEXT("GraphEditor", "ChangeNameListPinValue", "Change Bone Name Pin Value" ) );
+		GraphPinObj->Modify();
+		GraphPinObj->GetSchema()->TrySetDefaultValue(*GraphPinObj, NewTypeInValue.ToString());
+	}
+}
+
+TSharedRef<SWidget> SGraphPinNameList::MakeNameListItemWidget(TSharedPtr<FString> InItem)
+{
+	return 	SNew(STextBlock).Text(FText::FromString(*InItem));// .Font(FEditorStyle::GetFontStyle(TEXT("PropertyWindow.NormalFont")));
+}
+
+void SGraphPinNameList::OnNameListChanged(TSharedPtr<FString> NewSelection, ESelectInfo::Type SelectInfo)
+{
+	if (SelectInfo != ESelectInfo::Direct)
+	{
+		FString NewValue = *NewSelection.Get();
+		SetNameListText(FText::FromString(NewValue), ETextCommit::OnEnter);
+	}
+}
+
+void SGraphPinNameList::OnNameListComboBox()
+{
+	TSharedPtr<FString> CurrentlySelected;
+	for (TSharedPtr<FString> Item : GetNameList())
+	{
+		if (Item->Equals(GetNameListText().ToString()))
+		{
+			CurrentlySelected = Item;
+		}
+	}
+	NameListComboBox->SetSelectedItem(CurrentlySelected);
+}
diff --git a/Engine/Plugins/Experimental/ControlRig/Source/ControlRigEditor/Private/Graph/SGraphPinNameList.h b/Engine/Plugins/Experimental/ControlRig/Source/ControlRigEditor/Private/Graph/SGraphPinNameList.h
new file mode 100644
--- /dev/null
+++ b/Engine/Plugins/Experimental/ControlRig/Source/ControlRigEditor/Private/Graph/SGraphPinNameList.h
@@ -0,0 +1,42 @@
+// Copyright 1998-2019 Epic Games, Inc. All Rights Reserved.
+
+#pragma once
+
+#include "CoreMinimal.h"
+#include "Widgets/DeclarativeSyntaxSupport.h"
+#include "Widgets/SWidget.h"
+#include "SGraphPin.h"
+#include "SGraphPinNameListValueWidget.h"
+
+class SGraphPinNameList : public SGraphPin
+{
+public:
+
+	DECLARE_DELEGATE_RetVal( const TArray<TSharedPtr<FString>>&, FOnGetNameListContent );
+
+	SLATE_BEGIN_ARGS(SGraphPinNameList){}
+
+		SLATE_EVENT(FOnGetNameListContent, OnGetNameListContent)
+
+	SLATE_END_ARGS()
+
+	void Construct(const FArguments& InArgs, UEdGraphPin* InGraphPinObj);
+
+protected:
+	//~ Begin SGraphPin Interface
+	virtual TSharedRef<SWidget>	GetDefaultValueWidget() override;
+	//~ End SGraphPin Interface
+
+	const TArray<TSharedPtr<FString>>& GetNameList() const;
+	FText GetNameListText() const;
+	virtual void SetNameListText(const FText& NewTypeInValue, ETextCommit::Type CommitInfo);
+
+	TSharedRef<SWidget> MakeNameListItemWidget(TSharedPtr<FString> InItem);
+	void OnNameListChanged(TSharedPtr<FString> NewSelection, ESelectInfo::Type SelectInfo);
+	void OnNameListComboBox();
+
+	FOnGetNameListContent OnGetNameListContent;
+	TSharedPtr<SGraphPinNameListValueWidget> NameListComboBox;
+	TArray<TSharedPtr<FString>> EmptyList;
+
+};
