Copying //UE4/Dev-Anim to Dev-Main (//UE4/Dev-Main) @ 8782600

diff --git a/Engine/Plugins/Experimental/ControlRig/Source/ControlRigEditor/Private/Graph/SGraphPinNameListValueWidget.cpp b/Engine/Plugins/Experimental/ControlRig/Source/ControlRigEditor/Private/Graph/SGraphPinNameListValueWidget.cpp
new file mode 100644
--- /dev/null
+++ b/Engine/Plugins/Experimental/ControlRig/Source/ControlRigEditor/Private/Graph/SGraphPinNameListValueWidget.cpp
@@ -0,0 +1,249 @@
+// Copyright 1998-2019 Epic Games, Inc. All Rights Reservekd.
+
+#include "SGraphPinNameListValueWidget.h"
+#include "DetailLayoutBuilder.h"
+#include "Framework/Application/SlateUser.h"
+
+#define LOCTEXT_NAMESPACE "GraphPinNameListValueWidget"
+
+void SGraphPinNameListValueWidget::Construct(const FArguments& InArgs)
+{
+	this->OnComboBoxOpening = InArgs._OnComboBoxOpening;
+	this->OnSelectionChanged = InArgs._OnSelectionChanged;
+	this->OnGenerateWidget = InArgs._OnGenerateWidget;
+
+	OptionsSource = InArgs._OptionsSource;
+	CustomScrollbar = InArgs._CustomScrollbar;
+
+	TSharedRef<SWidget> ComboBoxMenuContent =
+		SNew(SBox)
+		.MaxDesiredHeight(InArgs._MaxListHeight)
+		[
+			SNew(SVerticalBox)
+
+			+ SVerticalBox::Slot()
+			.AutoHeight()
+			[
+				SAssignNew(this->SearchField, SEditableTextBox)
+				.HintText(LOCTEXT("Search", "Search"))
+				.OnTextChanged(this, &SGraphPinNameListValueWidget::OnSearchTextChanged)
+			]
+
+			+ SVerticalBox::Slot()
+			[
+				SAssignNew(this->ComboListView, SComboListType)
+				.ListItemsSource(OptionsSource)
+				.OnGenerateRow(this, &SGraphPinNameListValueWidget::GenerateMenuItemRow)
+				.OnSelectionChanged(this, &SGraphPinNameListValueWidget::OnSelectionChanged_Internal)
+				.SelectionMode(ESelectionMode::Single)
+				.ExternalScrollbar(InArgs._CustomScrollbar)
+			]
+		];
+
+	// Set up content
+	TSharedPtr<SWidget> ButtonContent = InArgs._Content.Widget;
+	if (InArgs._Content.Widget == SNullWidget::NullWidget)
+	{
+		SAssignNew(ButtonContent, STextBlock)
+			.Text(NSLOCTEXT("SGraphPinNameListValueWidget", "ContentWarning", "No Content Provided"))
+			.ColorAndOpacity(FLinearColor::Red);
+	}
+
+
+	SComboButton::Construct(SComboButton::FArguments()
+		.Method(InArgs._Method)
+		.ButtonContent()
+		[
+			ButtonContent.ToSharedRef()
+		]
+		.MenuContent()
+		[
+			ComboBoxMenuContent
+		]
+		.HasDownArrow(InArgs._HasDownArrow)
+		.ContentPadding(InArgs._ContentPadding)
+		.OnMenuOpenChanged(this, &SGraphPinNameListValueWidget::OnMenuOpenChanged)
+		.IsFocusable(true)
+		);
+	SetMenuContentWidgetToFocus(ComboListView);
+
+	// Need to establish the selected item at point of construction so its available for querying
+	// NB: If you need a selection to fire use SetItemSelection rather than setting an IntiallySelectedItem
+	SelectedItem = InArgs._InitiallySelectedItem;
+	if (TListTypeTraits<TSharedPtr<FString>>::IsPtrValid(SelectedItem))
+	{
+		ComboListView->Private_SetItemSelection(SelectedItem, true);
+	}
+
+}
+
+void SGraphPinNameListValueWidget::ClearSelection()
+{
+	ComboListView->ClearSelection();
+}
+
+void SGraphPinNameListValueWidget::SetSelectedItem(TSharedPtr<FString> InSelectedItem)
+{
+	if (TListTypeTraits<TSharedPtr<FString>>::IsPtrValid(InSelectedItem))
+	{
+		ComboListView->SetSelection(InSelectedItem);
+	}
+	else
+	{
+		ComboListView->ClearSelection();
+	}
+}
+
+TSharedPtr<FString> SGraphPinNameListValueWidget::GetSelectedItem()
+{
+	return SelectedItem;
+}
+
+void SGraphPinNameListValueWidget::RefreshOptions()
+{
+	if (!ComboListView->IsPendingRefresh())
+	{
+		ComboListView->RequestListRefresh();
+	}
+}
+
+TSharedRef<ITableRow> SGraphPinNameListValueWidget::GenerateMenuItemRow(TSharedPtr<FString> InItem, const TSharedRef<STableViewBase>& OwnerTable)
+{
+	if (OnGenerateWidget.IsBound())
+	{
+		FString SearchToken = SearchField->GetText().ToString().ToLower();
+		EVisibility WidgetVisibility = EVisibility::Visible;
+		if (!SearchToken.IsEmpty())
+		{
+			if (InItem->ToLower().Find(SearchToken) < 0)
+			{
+				WidgetVisibility = EVisibility::Collapsed;
+			}
+		}
+		return SNew(SComboRow<TSharedPtr<FString>>, OwnerTable)
+			.Visibility(WidgetVisibility)
+			[
+				OnGenerateWidget.Execute(InItem)
+			];
+	}
+	else
+	{
+		return SNew(SComboRow<TSharedPtr<FString>>, OwnerTable)
+			[
+				SNew(STextBlock).Text(NSLOCTEXT("SlateCore", "ComboBoxMissingOnGenerateWidgetMethod", "Please provide a .OnGenerateWidget() handler."))
+			];
+
+	}
+}
+
+void SGraphPinNameListValueWidget::OnMenuOpenChanged(bool bOpen)
+{
+	if (bOpen == false)
+	{
+		if (TListTypeTraits<TSharedPtr<FString>>::IsPtrValid(SelectedItem))
+		{
+			// Ensure the ListView selection is set back to the last committed selection
+			ComboListView->SetSelection(SelectedItem, ESelectInfo::OnNavigation);
+			ComboListView->RequestScrollIntoView(SelectedItem, 0);
+		}
+
+		// Set focus back to ComboBox for users focusing the ListView that just closed
+		TSharedRef<SWidget> ThisRef = AsShared();
+		FSlateApplication::Get().ForEachUser([&ThisRef](FSlateUser& User) {
+			if (User.HasFocusedDescendants(ThisRef))
+			{
+				User.SetFocus(ThisRef);
+			}
+		});
+
+	}
+	else
+	{
+		RegisterActiveTimer(0.f, FWidgetActiveTimerDelegate::CreateSP(this, &SGraphPinNameListValueWidget::SetFocusPostConstruct));
+	}
+}
+
+EActiveTimerReturnType SGraphPinNameListValueWidget::SetFocusPostConstruct(double InCurrentTime, float InDeltaTime)
+{
+	if (SearchField.IsValid())
+	{
+		bool bSucceeded = false;
+		FSlateApplication::Get().ForEachUser([&](FSlateUser& User) {
+			bSucceeded |= User.SetFocus(SearchField.ToSharedRef());
+		});
+
+		if (bSucceeded)
+		{
+			return EActiveTimerReturnType::Stop;
+		}
+	}
+
+	return EActiveTimerReturnType::Continue;
+}
+
+void SGraphPinNameListValueWidget::OnSelectionChanged_Internal(TSharedPtr<FString> ProposedSelection, ESelectInfo::Type SelectInfo)
+{
+	// Ensure that the proposed selection is different
+	if (SelectInfo != ESelectInfo::OnNavigation)
+	{
+		// Ensure that the proposed selection is different from selected
+		if (ProposedSelection != SelectedItem)
+		{
+			SelectedItem = ProposedSelection;
+			OnSelectionChanged.ExecuteIfBound(ProposedSelection, SelectInfo);
+		}
+		// close combo even if user reselected item
+		this->SetIsOpen(false);
+	}
+}
+
+void SGraphPinNameListValueWidget::OnSearchTextChanged(const FText& ChangedText)
+{
+	FString SearchToken = ChangedText.ToString().ToLower();
+	for (int32 i = 0; i < OptionsSource->Num(); i++)
+	{
+		TSharedPtr<ITableRow> Row = ComboListView->WidgetFromItem((*OptionsSource)[i]);
+		if (Row)
+		{
+			if (SearchToken.IsEmpty())
+			{
+				Row->AsWidget()->SetVisibility(EVisibility::Visible);
+			}
+			else if ((*OptionsSource)[i]->ToLower().Find(SearchToken) >= 0)
+			{
+				Row->AsWidget()->SetVisibility(EVisibility::Visible);
+			}
+			else
+			{
+				Row->AsWidget()->SetVisibility(EVisibility::Collapsed);
+			}
+		}
+	}
+
+	ComboListView->RequestListRefresh();
+
+	SelectedItem = TSharedPtr< FString >();
+}
+
+FReply SGraphPinNameListValueWidget::OnButtonClicked()
+{
+	// if user clicked to close the combo menu
+	if (this->IsOpen())
+	{
+		// Re-select first selected item, just in case it was selected by navigation previously
+		TArray<TSharedPtr<FString>> SelectedItems = ComboListView->GetSelectedItems();
+		if (SelectedItems.Num() > 0)
+		{
+			OnSelectionChanged_Internal(SelectedItems[0], ESelectInfo::Direct);
+		}
+	}
+	else
+	{
+		SearchField->SetText(FText::GetEmpty());
+		OnComboBoxOpening.ExecuteIfBound();
+	}
+
+	return SComboButton::OnButtonClicked();
+}
+
+#undef LOCTEXT_NAMESPACE
diff --git a/Engine/Plugins/Experimental/ControlRig/Source/ControlRigEditor/Private/Graph/SGraphPinNameListValueWidget.h b/Engine/Plugins/Experimental/ControlRig/Source/ControlRigEditor/Private/Graph/SGraphPinNameListValueWidget.h
new file mode 100644
--- /dev/null
+++ b/Engine/Plugins/Experimental/ControlRig/Source/ControlRigEditor/Private/Graph/SGraphPinNameListValueWidget.h
@@ -0,0 +1,131 @@
+// Copyright 1998-2019 Epic Games, Inc. All Rights Reservekd.
+
+#pragma once
+
+#include "CoreMinimal.h"
+#include "Input/Reply.h"
+#include "Widgets/SWidget.h"
+#include "Widgets/Input/SComboButton.h"
+#include "Widgets/Input/SEditableTextBox.h"
+#include "Widgets/Views/SListView.h"
+#include "Widgets/Input/SComboBox.h"
+
+/**
+ * A searchable text combo box
+ */
+class SGraphPinNameListValueWidget : public SComboButton
+{
+public:
+	/** Type of list used for showing menu options. */
+	typedef SListView< TSharedPtr<FString> > SComboListType;
+	/** Delegate type used to generate widgets that represent Options */
+	typedef typename TSlateDelegates< TSharedPtr<FString> >::FOnGenerateWidget FOnGenerateWidget;
+	typedef typename TSlateDelegates< TSharedPtr<FString> >::FOnSelectionChanged FOnSelectionChanged;
+
+	SLATE_BEGIN_ARGS(SGraphPinNameListValueWidget)
+		: _Content()
+		, _ContentPadding(FMargin(3.0, 3.0))
+		, _OptionsSource()
+		, _OnSelectionChanged()
+		, _OnGenerateWidget()
+		, _InitiallySelectedItem(nullptr)
+		, _Method()
+		, _MaxListHeight(450.0f)
+		, _HasDownArrow(true)
+	{}
+
+	/** Slot for this button's content (optional) */
+	SLATE_DEFAULT_SLOT(FArguments, Content)
+
+		SLATE_ATTRIBUTE(FMargin, ContentPadding)
+
+		SLATE_ARGUMENT(const TArray< TSharedPtr<FString> >*, OptionsSource)
+		SLATE_EVENT(FOnSelectionChanged, OnSelectionChanged)
+		SLATE_EVENT(FOnGenerateWidget, OnGenerateWidget)
+
+		/** Called when combo box is opened, before list is actually created */
+		SLATE_EVENT(FOnComboBoxOpening, OnComboBoxOpening)
+
+		/** The custom scrollbar to use in the ListView */
+		SLATE_ARGUMENT(TSharedPtr<SScrollBar>, CustomScrollbar)
+
+		/** The option that should be selected when the combo box is first created */
+		SLATE_ARGUMENT(TSharedPtr<FString>, InitiallySelectedItem)
+
+		SLATE_ARGUMENT(TOptional<EPopupMethod>, Method)
+
+		/** The max height of the combo box menu */
+		SLATE_ARGUMENT(float, MaxListHeight)
+
+		/**
+		 * When false, the down arrow is not generated and it is up to the API consumer
+		 * to make their own visual hint that this is a drop down.
+		 */
+		SLATE_ARGUMENT(bool, HasDownArrow)
+
+	SLATE_END_ARGS()
+
+	/**
+	 * Construct the widget from a declaration
+	 *
+	 * @param InArgs   Declaration from which to construct the combo box
+	 */
+	void Construct(const FArguments& InArgs);
+
+	void ClearSelection();
+
+	void SetSelectedItem(TSharedPtr<FString> InSelectedItem);
+
+	/** @return the item currently selected by the combo box. */
+	TSharedPtr<FString> GetSelectedItem();
+
+	/**
+	 * Requests a list refresh after updating options
+	 * Call SetSelectedItem to update the selected item if required
+	 * @see SetSelectedItem
+	 */
+	void RefreshOptions();
+
+private:
+
+	/** Generate a row for the InItem in the combo box's list (passed in as OwnerTable). Do this by calling the user-specified OnGenerateWidget */
+	TSharedRef<ITableRow> GenerateMenuItemRow(TSharedPtr<FString> InItem, const TSharedRef<STableViewBase>& OwnerTable);
+
+	//** Called if the menu is closed
+	void OnMenuOpenChanged(bool bOpen);
+
+	/** Invoked when the selection in the list changes */
+	void OnSelectionChanged_Internal(TSharedPtr<FString> ProposedSelection, ESelectInfo::Type SelectInfo);
+
+	/** Invoked when the search text changes */
+	void OnSearchTextChanged(const FText& ChangedText);
+
+	/** Handle clicking on the content menu */
+	virtual FReply OnButtonClicked() override;
+
+	/** The item style to use. */
+	const FTableRowStyle* ItemStyle;
+
+private:
+	/** Delegate that is invoked when the selected item in the combo box changes */
+	FOnSelectionChanged OnSelectionChanged;
+	/** The item currently selected in the combo box */
+	TSharedPtr<FString> SelectedItem;
+	/** The search field used for the combox box's contents */
+	TSharedPtr< SEditableTextBox > SearchField;
+	/** The ListView that we pop up; visualized the available options. */
+	TSharedPtr< SComboListType > ComboListView;
+	/** The Scrollbar used in the ListView. */
+	TSharedPtr< SScrollBar > CustomScrollbar;
+	/** Delegate to invoke before the combo box is opening. */
+	FOnComboBoxOpening OnComboBoxOpening;
+	/** Delegate to invoke when we need to visualize an option as a widget. */
+	FOnGenerateWidget OnGenerateWidget;
+
+	const TArray< TSharedPtr<FString> >* OptionsSource;
+
+	/** Used to focus the name box immediately following construction */
+	EActiveTimerReturnType SetFocusPostConstruct(double InCurrentTime, float InDeltaTime);
+};
+
+
diff --git a/Engine/Plugins/Experimental/ControlRig/Source/ControlRigManipulation/ControlRigManipulation.Build.cs b/Engine/Plugins/Experimental/ControlRig/Source/ControlRigManipulation/ControlRigManipulation.Build.cs
new file mode 100644
--- /dev/null
+++ b/Engine/Plugins/Experimental/ControlRig/Source/ControlRigManipulation/ControlRigManipulation.Build.cs
@@ -0,0 +1,22 @@
+// Copyright 1998-2019 Epic Games, Inc. All Rights Reserved.
+
+namespace UnrealBuildTool.Rules
+{
+    public class ControlRigManipulation : ModuleRules
+    {
+        public ControlRigManipulation(ReadOnlyTargetRules Target) : base(Target)
+        {
+            PrivateIncludePaths.Add("ControlRigManipulation/Private");
+
+            PrivateDependencyModuleNames.AddRange(
+                new string[]
+                {
+                    "Core",
+                    "CoreUObject",
+                    "Engine",
+                    "ControlRig",
+                }
+            );
+        }
+    }
+}
diff --git a/Engine/Plugins/Experimental/ControlRig/Source/ControlRigManipulation/Private/ControlRigGizmoActor.cpp b/Engine/Plugins/Experimental/ControlRig/Source/ControlRigManipulation/Private/ControlRigGizmoActor.cpp
new file mode 100644
--- /dev/null
+++ b/Engine/Plugins/Experimental/ControlRig/Source/ControlRigManipulation/Private/ControlRigGizmoActor.cpp
@@ -0,0 +1,202 @@
+// Copyright 1998-2019 Epic Games, Inc. All Rights Reserved.
+
+#include "ControlRigGizmoActor.h"
+#include "Engine/World.h"
+#include "Engine/StaticMesh.h"
+#include "Components/StaticMeshComponent.h"
+#include "Materials/MaterialInstanceDynamic.h"
+#include "Engine/CollisionProfile.h"
+
+AControlRigGizmoActor::AControlRigGizmoActor(const FObjectInitializer& ObjectInitializer)
+	: Super(ObjectInitializer)
+	, bEnabled(true)
+	, bSelected(false)
+	, bHovered(false)
+	, bManipulating(false)
+{
+
+	ActorRootComponent = CreateDefaultSubobject<USceneComponent>(TEXT("SceneComponent0"));
+	StaticMeshComponent = CreateDefaultSubobject<UStaticMeshComponent>(TEXT("StaticMeshComponent0"));
+	StaticMeshComponent->SetCollisionProfileName(UCollisionProfile::BlockAll_ProfileName);
+	StaticMeshComponent->Mobility = EComponentMobility::Movable;
+	StaticMeshComponent->SetGenerateOverlapEvents(false);
+	StaticMeshComponent->bUseDefaultCollision = true;
+
+	RootComponent = ActorRootComponent;
+	StaticMeshComponent->SetupAttachment(RootComponent);
+}
+
+void AControlRigGizmoActor::SetEnabled(bool bInEnabled)
+{
+	bool bOldEnabled = bEnabled;
+
+	bEnabled = bInEnabled;
+
+	if(bEnabled != bOldEnabled)
+	{
+		FEditorScriptExecutionGuard Guard;
+		OnEnabledChanged(bEnabled);
+	}
+}
+
+bool AControlRigGizmoActor::IsEnabled() const
+{
+	return bEnabled;
+}
+
+void AControlRigGizmoActor::SetSelected(bool bInSelected)
+{
+	bool bOldSelected = bSelected;
+
+	bSelected = bInSelected;
+
+	if(bSelected != bOldSelected)
+	{
+		FEditorScriptExecutionGuard Guard;
+		OnSelectionChanged(bSelected);
+	}
+}
+
+bool AControlRigGizmoActor::IsSelectedInEditor() const
+{
+	return bSelected;
+}
+
+void AControlRigGizmoActor::SetHovered(bool bInHovered)
+{
+	bool bOldHovered = bHovered;
+
+	bHovered = bInHovered;
+
+	if(bHovered != bOldHovered)
+	{
+		FEditorScriptExecutionGuard Guard;
+		OnHoveredChanged(bHovered);
+	}
+}
+
+bool AControlRigGizmoActor::IsHovered() const
+{
+	return bHovered;
+}
+
+void AControlRigGizmoActor::SetManipulating(bool bInManipulating)
+{
+	bool bOldManipulating = bManipulating;
+
+	bManipulating = bInManipulating;
+
+	if(bManipulating != bOldManipulating)
+	{
+		FEditorScriptExecutionGuard Guard;
+		OnManipulatingChanged(bManipulating);
+	}
+}
+
+bool AControlRigGizmoActor::IsManipulating() const
+{
+	return bManipulating;
+}
+
+// FControlRigGizmoHelper START
+
+namespace FControlRigGizmoHelper
+{
+	FActorSpawnParameters GetDefaultSpawnParameter()
+	{
+		FActorSpawnParameters ActorSpawnParameters;
+		ActorSpawnParameters.bTemporaryEditorActor = true;
+		ActorSpawnParameters.bHideFromSceneOutliner = true;
+		ActorSpawnParameters.SpawnCollisionHandlingOverride = ESpawnActorCollisionHandlingMethod::AlwaysSpawn;
+
+		return ActorSpawnParameters;
+	}
+
+	// create gizmo from custom staticmesh, may deprecate this unless we come up with better usage
+	AControlRigGizmoActor* CreateGizmoActor(UWorld* InWorld, UStaticMesh* InStaticMesh, const FGizmoActorCreationParam& CreationParam)
+	{
+		if (InWorld)
+		{
+			AControlRigGizmoActor* GizmoActor = CreateDefaultGizmoActor(InWorld, CreationParam);
+
+			if (GizmoActor)
+			{
+				if (InStaticMesh)
+				{
+					GizmoActor->StaticMeshComponent->SetStaticMesh(InStaticMesh);
+				}
+
+				return GizmoActor;
+			}
+		}
+
+		return nullptr;
+	}
+
+	AControlRigGizmoActor* CreateGizmoActor(UWorld* InWorld, TSubclassOf<AControlRigGizmoActor> InClass, const FGizmoActorCreationParam& CreationParam)
+	{
+		AControlRigGizmoActor* GizmoActor = InWorld->SpawnActor<AControlRigGizmoActor>(InClass, GetDefaultSpawnParameter());
+		if (GizmoActor)
+		{
+			// set transform
+			GizmoActor->SetActorTransform(CreationParam.SpawnTransform);
+			return GizmoActor;
+		}
+
+		return nullptr;
+	}
+
+	AControlRigGizmoActor* CreateDefaultGizmoActor(UWorld* InWorld, const FGizmoActorCreationParam& CreationParam)
+	{
+		AControlRigGizmoActor* GizmoActor = InWorld->SpawnActor<AControlRigGizmoActor>(AControlRigGizmoActor::StaticClass(), GetDefaultSpawnParameter());
+		if (GizmoActor)
+		{
+			GizmoActor->SetActorTransform(CreationParam.SpawnTransform);
+
+			UStaticMeshComponent* MeshComponent = GizmoActor->StaticMeshComponent;
+
+			if (!CreationParam.StaticMesh.IsValid())
+			{
+				CreationParam.StaticMesh.LoadSynchronous();
+			}
+			if (CreationParam.StaticMesh.IsValid())
+			{
+				MeshComponent->SetStaticMesh(CreationParam.StaticMesh.Get());
+				MeshComponent->SetRelativeTransform(CreationParam.MeshTransform * CreationParam.GizmoTransform);
+			}
+
+			if (!CreationParam.Material.IsValid())
+			{
+				CreationParam.Material.LoadSynchronous();
+			}
+			if (CreationParam.StaticMesh.IsValid())
+			{
+				UMaterialInstanceDynamic* MaterialInstance = UMaterialInstanceDynamic::Create(CreationParam.Material.Get(), GizmoActor);
+				MaterialInstance->SetVectorParameterValue(CreationParam.ColorParameterName, FVector(CreationParam.Color));
+				MeshComponent->SetMaterial(0, MaterialInstance);
+			}
+			return GizmoActor;
+		}
+
+		return nullptr;
+	}
+}
+
+void AControlRigGizmoActor::SetGlobalTransform(const FTransform& InTransform)
+{
+	if (RootComponent)
+	{
+		RootComponent->SetRelativeTransform(InTransform, false, nullptr, ETeleportType::TeleportPhysics);
+	}
+}
+
+FTransform AControlRigGizmoActor::GetGlobalTransform() const
+{
+	if (RootComponent)
+	{
+		return RootComponent->GetRelativeTransform();
+	}
+
+	return FTransform::Identity;
+}
+// FControlRigGizmoHelper END
\ No newline at end of file
diff --git a/Engine/Plugins/Experimental/ControlRig/Source/ControlRigManipulation/Private/ControlRigManipulationModule.cpp b/Engine/Plugins/Experimental/ControlRig/Source/ControlRigManipulation/Private/ControlRigManipulationModule.cpp
new file mode 100644
--- /dev/null
+++ b/Engine/Plugins/Experimental/ControlRig/Source/ControlRigManipulation/Private/ControlRigManipulationModule.cpp
@@ -0,0 +1,14 @@
+// Copyright 1998-2019 Epic Games, Inc. All Rights Reserved.
+
+#include "ControlRigManipulationModule.h"
+#include "Modules/ModuleManager.h"
+
+void FControlRigManipulationModule::StartupModule()
+{
+}
+
+void FControlRigManipulationModule::ShutdownModule()
+{
+}
+
+IMPLEMENT_MODULE(FControlRigManipulationModule, ControlRigManipulation)
diff --git a/Engine/Plugins/Experimental/ControlRig/Source/ControlRigManipulation/Private/ControlRigManipulationModule.h b/Engine/Plugins/Experimental/ControlRig/Source/ControlRigManipulation/Private/ControlRigManipulationModule.h
new file mode 100644
--- /dev/null
+++ b/Engine/Plugins/Experimental/ControlRig/Source/ControlRigManipulation/Private/ControlRigManipulationModule.h
@@ -0,0 +1,14 @@
+// Copyright 1998-2019 Epic Games, Inc. All Rights Reserved.
+
+#pragma once
+
+#include "Delegates/IDelegateInstance.h"
+#include "IControlRigManipulationModule.h"
+
+class FControlRigManipulationModule : public IControlRigManipulationModule
+{
+public:
+	/** IModuleInterface implementation */
+	virtual void StartupModule() override;
+	virtual void ShutdownModule() override;
+};
diff --git a/Engine/Plugins/Experimental/ControlRig/Source/ControlRigManipulation/Private/DefaultControlRigManipulationLayer.cpp b/Engine/Plugins/Experimental/ControlRig/Source/ControlRigManipulation/Private/DefaultControlRigManipulationLayer.cpp
new file mode 100644
--- /dev/null
+++ b/Engine/Plugins/Experimental/ControlRig/Source/ControlRigManipulation/Private/DefaultControlRigManipulationLayer.cpp
@@ -0,0 +1,427 @@
+// Copyright 1998-2019 Epic Games, Inc. All Rights Reserved.
+
+#include "DefaultControlRigManipulationLayer.h"
+#include "Manipulatable/IControlRigManipulatable.h"
+#include "ControlRig.h"
+#include "Engine/World.h"
+#include "Components/SkeletalMeshComponent.h"
+#include "ControlRigGizmoLibrary.h"
+#include "IControlRigObjectBinding.h"
+
+UDefaultControlRigManipulationLayer::UDefaultControlRigManipulationLayer(const FObjectInitializer& ObjectInitializer)
+	: Super(ObjectInitializer)
+{
+}
+
+void UDefaultControlRigManipulationLayer::CreateLayer()
+{
+}
+
+void UDefaultControlRigManipulationLayer::DestroyLayer()
+{
+	ExecuteDelegateHandles.Empty();
+	InitializeDelegateHandle.Empty();
+
+	IControlRigManipulationLayer::DestroyLayer();
+}
+
+void UDefaultControlRigManipulationLayer::AddManipulatableObject(IControlRigManipulatable* InObject)
+{
+	if (InObject)
+	{
+		IControlRigManipulationLayer::AddManipulatableObject(InObject);
+
+		OnControlRigAdded(static_cast<UControlRig*>(InObject));
+	}
+}
+
+void UDefaultControlRigManipulationLayer::RemoveManipulatableObject(IControlRigManipulatable* InObject)
+{
+	if (InObject)
+	{
+		OnControlRigRemoved(static_cast<UControlRig*>(InObject));
+
+		IControlRigManipulationLayer::RemoveManipulatableObject(InObject);
+	}
+}
+
+
+void UDefaultControlRigManipulationLayer::OnControlRigAdded(UControlRig* InControlRig)
+{
+	// bind execution delegate
+	ExecuteDelegateHandles.Add(InControlRig->OnExecuted().AddUObject(this, &UDefaultControlRigManipulationLayer::OnRigExecuted));
+	InitializeDelegateHandle.Add(InControlRig->OnInitialized().AddUObject(this, &UDefaultControlRigManipulationLayer::OnRigInitialized));
+	ControlModifiedDelegateHandles.Add(InControlRig->ControlModified().AddUObject(this, &UDefaultControlRigManipulationLayer::OnControlModified));
+
+	// ensure the size matches
+	ensure(ExecuteDelegateHandles.Num() == ManipulatableObjects.Num());
+	ensure(InitializeDelegateHandle.Num() == ManipulatableObjects.Num());
+	ensure(ControlModifiedDelegateHandles.Num() == ManipulatableObjects.Num());
+
+	// object binding. This overwrites if there were multi
+	SetObjectBinding(InControlRig->GetObjectBinding());
+}
+
+void UDefaultControlRigManipulationLayer::OnControlRigRemoved(UControlRig* InControlRig)
+{
+	// unbind execution delegate
+	int32 Found = ManipulatableObjects.Find(InControlRig);
+	if (Found != INDEX_NONE)
+	{
+		if (ExecuteDelegateHandles.IsValidIndex(Found) && ExecuteDelegateHandles[Found].IsValid())
+		{
+			InControlRig->OnExecuted().Remove(ExecuteDelegateHandles[Found]);
+		}
+
+		if (InitializeDelegateHandle.IsValidIndex(Found) && InitializeDelegateHandle[Found].IsValid())
+		{
+			InControlRig->OnInitialized().Remove(InitializeDelegateHandle[Found]);
+		}
+
+		if (ControlModifiedDelegateHandles.IsValidIndex(Found) && ControlModifiedDelegateHandles[Found].IsValid())
+		{
+			InControlRig->ControlModified().Remove(ControlModifiedDelegateHandles[Found]);
+		}
+	}
+	
+}
+
+// gizmo related 
+void UDefaultControlRigManipulationLayer::SetGizmoTransform(AControlRigGizmoActor* GizmoActor, const FTransform& InTransform)
+{
+	const FControlData* Manip = GetControlDataFromGizmo(GizmoActor);
+	if (Manip)
+	{
+		Manip->ManipObject->SetControlGlobalTransform(Manip->ControlName, InTransform);
+	}
+}
+
+void UDefaultControlRigManipulationLayer::GetGizmoTransform(AControlRigGizmoActor* GizmoActor, FTransform& OutTransform) const
+{
+	const FControlData* Manip = GetControlDataFromGizmo(GizmoActor);
+	if (Manip)
+	{
+		OutTransform = Manip->ManipObject->GetControlGlobalTransform(Manip->ControlName);
+	}
+}
+
+void UDefaultControlRigManipulationLayer::MoveGizmo(AControlRigGizmoActor* GizmoActor, const bool bTranslation, FVector& InDrag,
+	const bool bRotation, FRotator& InRot, const bool bScale, FVector& InScale, const FTransform& ToWorldTransform)
+{
+	const FControlData* Manip = GetControlDataFromGizmo(GizmoActor);
+	if (Manip)
+	{
+		FTransform CurrentTransform = Manip->ManipObject->GetControlGlobalTransform(Manip->ControlName) * ToWorldTransform;
+
+		bool bTransformChanged = false;
+		if (bRotation)
+		{
+			FQuat CurrentRotation = CurrentTransform.GetRotation();
+			CurrentRotation = (InRot.Quaternion() * CurrentRotation);
+			CurrentTransform.SetRotation(CurrentRotation);
+			bTransformChanged = true;
+		}
+
+		if (bTranslation)
+		{
+			FVector CurrentLocation = CurrentTransform.GetLocation();
+			CurrentLocation = CurrentLocation + InDrag;
+			CurrentTransform.SetLocation(CurrentLocation);
+			bTransformChanged = true;
+		}
+
+		if (bScale)
+		{
+			FVector CurrentScale = CurrentTransform.GetScale3D();
+			CurrentScale = CurrentScale + InScale;
+			CurrentTransform.SetScale3D(CurrentScale);
+			bTransformChanged = true;
+		}
+
+		if (bTransformChanged)
+		{
+			FTransform NewTransform = CurrentTransform.GetRelativeTransform(ToWorldTransform);
+			Manip->ManipObject->SetControlGlobalTransform(Manip->ControlName, NewTransform);
+			// assumes it's attached to actor
+			GizmoActor->SetGlobalTransform(NewTransform);
+		}
+	}
+}
+
+// temporarily we just support following types of gizmo
+bool IsSupportedControlType(const ERigControlType ControlType)
+{
+	switch (ControlType)
+	{
+	case ERigControlType::Position:
+	case ERigControlType::Scale:
+	case ERigControlType::Quat:
+	case ERigControlType::Rotator:
+	case ERigControlType::Transform:
+		return true;
+	}
+
+	return false;
+}
+
+void UDefaultControlRigManipulationLayer::GetGizmoCreationParams(TArray<FGizmoActorCreationParam>& OutCreationParams)
+{
+	OutCreationParams.Reset();
+
+	// for now we only support FTransform
+	for (IControlRigManipulatable* ManipObject : ManipulatableObjects)
+	{
+		const TArray<FRigControl>& Controls = ManipObject->AvailableControls();
+		UControlRigGizmoLibrary* GizmoLibrary = ManipObject->GetGizmoLibrary();
+
+		for (const FRigControl& Control : Controls)
+		{
+			if (IsSupportedControlType(Control.ControlType))
+			{
+				FGizmoActorCreationParam Param;
+				Param.ManipObj = ManipObject;
+				Param.ControlName = Control.Name;
+				Param.SpawnTransform = ManipObject->GetControlGlobalTransform(Control.Name);
+				Param.GizmoTransform = Control.GizmoTransform;
+
+				if (GizmoLibrary)
+				{
+					if (const FControlRigGizmoDefinition* Gizmo = GizmoLibrary->GetGizmoByName(Control.GizmoName, true /* use default */))
+					{
+						Param.MeshTransform = Gizmo->Transform;
+						Param.StaticMesh = Gizmo->StaticMesh;
+						Param.Material = GizmoLibrary->DefaultMaterial;
+						Param.ColorParameterName = GizmoLibrary->MaterialColorParameter;
+					}
+				}
+
+				Param.Color = Control.GizmoColor;
+				OutCreationParams.Add(Param);
+			}
+		}
+	}
+}
+
+bool UDefaultControlRigManipulationLayer::ModeSupportedByGizmoActor(const AControlRigGizmoActor* GizmoActor, FWidget::EWidgetMode InMode) const
+{
+	const FControlData* Control = GetControlDataFromGizmo(GizmoActor);
+	if (Control)
+	{
+		const FRigControl* RigControl = Control->ManipObject->FindControl(Control->ControlName);
+		if (RigControl)
+		{
+			return IsSupportedControlType(RigControl->ControlType);
+		}
+	}
+	return true;
+}
+
+void UDefaultControlRigManipulationLayer::TickGizmo(AControlRigGizmoActor* GizmoActor, const FTransform& ComponentTransform)
+{
+	if (GizmoActor)
+	{
+		const FControlData* DataPtr = GetControlDataFromGizmo(GizmoActor);
+		if (DataPtr)
+		{
+			FTransform Transform = DataPtr->ManipObject->GetControlGlobalTransform(DataPtr->ControlName);
+			GizmoActor->SetActorTransform(Transform * ComponentTransform);
+		}
+	}
+}
+
+const FControlData* UDefaultControlRigManipulationLayer::GetControlDataFromGizmo(const AControlRigGizmoActor* GizmoActor) const
+{
+	const FControlID* Found = GizmoToControlMap.Find(GizmoActor);
+
+	// for now we only support one to one
+	if (Found)
+	{
+		return &ControlData[*Found];
+	}
+
+	return nullptr;
+}
+
+AControlRigGizmoActor* UDefaultControlRigManipulationLayer::GetGizmoFromControlName(const FName& ControlName) const
+{
+	for (auto Iter=GizmoToControlMap.CreateConstIterator(); Iter; ++Iter)
+	{
+		const FControlID ID= Iter.Value();
+
+		if (ControlData[ID].ControlName == ControlName)
+		{
+			return Iter.Key();
+		}
+	}
+
+	return nullptr;
+}
+
+void UDefaultControlRigManipulationLayer::TickManipulatableObjects(float DeltaTime)
+{
+	// tick skeletalmeshcomponent, that's how they update their transform from rig change
+	USkeletalMeshComponent* SkeletalMeshComponent = GetSkeletalMeshComponent();
+	if (SkeletalMeshComponent)
+	{
+		SkeletalMeshComponent->RefreshBoneTransforms();
+		SkeletalMeshComponent->RefreshSlaveComponents();
+		SkeletalMeshComponent->UpdateComponentToWorld();
+		SkeletalMeshComponent->FinalizeBoneTransform();
+		SkeletalMeshComponent->MarkRenderTransformDirty();
+		SkeletalMeshComponent->MarkRenderDynamicDataDirty();
+	}
+}
+
+bool UDefaultControlRigManipulationLayer::GetGlobalTransform(AControlRigGizmoActor* GizmoActor, const FName& ControlName, FTransform& OutTransform) const
+{
+	const FControlData* Data = GetControlDataFromGizmo(GizmoActor);
+	if (Data)
+	{
+		OutTransform = Data->ManipObject->GetControlGlobalTransform(ControlName);
+		return true;
+	}
+
+	return false;
+}
+
+void UDefaultControlRigManipulationLayer::AddToControlData(AControlRigGizmoActor* GizmoActor, IControlRigManipulatable* InManipulatableObject, const FName& InControlName)
+{
+	FControlID NewID = ControlData.AddUninitialized();
+	FControlData& NewManipulatable = ControlData[NewID];
+	NewManipulatable.ManipObject = InManipulatableObject;
+	NewManipulatable.ControlName = InControlName;
+
+	GizmoToControlMap.Add(GizmoActor, NewID);
+}
+
+void UDefaultControlRigManipulationLayer::ResetControlData()
+{
+	GizmoToControlMap.Reset();
+	ControlData.Reset();
+}
+
+bool UDefaultControlRigManipulationLayer::CreateGizmoActors(UWorld* World, TArray<AControlRigGizmoActor*>& OutGizmoActors)
+{
+	DestroyGizmosActors();
+
+	FActorSpawnParameters ActorSpawnParameters;
+	ActorSpawnParameters.bTemporaryEditorActor = true;
+
+	TArray<FGizmoActorCreationParam> Params;
+	GetGizmoCreationParams(Params);
+
+	OutGizmoActors.Reset();
+
+	for (const FGizmoActorCreationParam& Param : Params)
+	{
+		AControlRigGizmoActor* GizmoActor = FControlRigGizmoHelper::CreateDefaultGizmoActor(World, Param);
+		if (GizmoActor)
+		{
+			AddToControlData(GizmoActor, Param.ManipObj, Param.ControlName);
+			OutGizmoActors.Add(GizmoActor);
+		}
+	}
+	
+	return (OutGizmoActors.Num() > 0);
+}
+
+void UDefaultControlRigManipulationLayer::DestroyGizmosActors()
+{
+	// clear previous gizmo actors
+	// since this is creating, we also have to destroy
+	for (auto Iter = GizmoToControlMap.CreateIterator(); Iter; ++Iter)
+	{
+		AControlRigGizmoActor* GizmoActor = Iter.Key();
+		UWorld* World = GizmoActor->GetWorld();
+		if (World)
+		{
+			World->DestroyActor(GizmoActor);
+		}
+	}
+
+	ResetControlData();
+}
+
+void UDefaultControlRigManipulationLayer::OnRigExecuted(class UControlRig* InRig, const EControlRigState State)
+{
+	FTransform ComponentTransform = GetSkeletalMeshComponentTransform();
+
+	// after executing rig, update gizmo
+	for (auto Iter=GizmoToControlMap.CreateIterator(); Iter; ++Iter)
+	{
+		TickGizmo(Iter.Key(), ComponentTransform);
+	}
+}
+
+void UDefaultControlRigManipulationLayer::OnRigInitialized(class UControlRig* InRig, const EControlRigState State)
+{
+}
+
+void UDefaultControlRigManipulationLayer::OnControlModified(IControlRigManipulatable* InManipulatable, const FRigControl& InControl)
+{
+	if (UControlRig* ControlRig = static_cast<UControlRig*>(InManipulatable))
+	{
+		if (AControlRigGizmoActor* const* Actor = GizmoToControlMap.FindKey(InControl.Index))
+		{
+			FTransform ComponentTransform = GetSkeletalMeshComponentTransform();
+			TickGizmo(*Actor, ComponentTransform);
+		}
+	}
+}
+
+TSharedPtr<IControlRigObjectBinding> UDefaultControlRigManipulationLayer::GetObjectBinding() const
+{
+	for (IControlRigManipulatable* Manip : ManipulatableObjects)
+	{
+		UControlRig* ControlRig = static_cast<UControlRig*>(Manip);
+		if (ControlRig)
+		{
+			return ControlRig->GetObjectBinding();
+		}
+	}
+
+	return TSharedPtr<IControlRigObjectBinding>();
+}
+
+void UDefaultControlRigManipulationLayer::SetObjectBinding(TSharedPtr<IControlRigObjectBinding> InObjectBinding)
+{
+	for (IControlRigManipulatable* Manip : ManipulatableObjects)
+	{
+		UControlRig* ControlRig = static_cast<UControlRig*>(Manip);
+		if (ControlRig)
+		{
+			ControlRig->SetObjectBinding(InObjectBinding);
+		}
+	}
+}
+
+USkeletalMeshComponent* UDefaultControlRigManipulationLayer::GetSkeletalMeshComponent() const
+{
+	TSharedPtr<IControlRigObjectBinding> ObjectBinding = GetObjectBinding();
+	if (ObjectBinding.IsValid())
+	{
+		return Cast<USkeletalMeshComponent>(ObjectBinding->GetBoundObject());
+	}
+
+	return nullptr;
+}
+
+FTransform	UDefaultControlRigManipulationLayer::GetSkeletalMeshComponentTransform() const
+{
+	USkeletalMeshComponent* MeshComponent = GetSkeletalMeshComponent();
+	return MeshComponent ? MeshComponent->GetComponentTransform() : FTransform::Identity;
+}
+
+void UDefaultControlRigManipulationLayer::BeginTransaction()
+{
+	for (IControlRigManipulatable* Manip : ManipulatableObjects)
+	{
+		UControlRig* ControlRig = static_cast<UControlRig*>(Manip);
+		ControlRig->SetFlags(RF_Transactional);
+		ControlRig->Modify();
+	}
+}
+
+void UDefaultControlRigManipulationLayer::EndTransaction()
+{
+}
\ No newline at end of file
diff --git a/Engine/Plugins/Experimental/ControlRig/Source/ControlRigManipulation/Private/IControlRigManipulationLayer.cpp b/Engine/Plugins/Experimental/ControlRig/Source/ControlRigManipulation/Private/IControlRigManipulationLayer.cpp
new file mode 100644
--- /dev/null
+++ b/Engine/Plugins/Experimental/ControlRig/Source/ControlRigManipulation/Private/IControlRigManipulationLayer.cpp
@@ -0,0 +1,50 @@
+// Copyright 1998-2019 Epic Games, Inc. All Rights Reserved.
+
+#include "IControlRigManipulationLayer.h"
+
+IControlRigManipulationLayer::IControlRigManipulationLayer()
+	:bLayerCreated(false)
+{}
+
+void IControlRigManipulationLayer::CreateLayer()
+{
+	// clear current layers
+	if (bLayerCreated)
+	{
+		DestroyLayer();
+	}
+
+	bLayerCreated = true;
+}
+
+void IControlRigManipulationLayer::DestroyLayer()
+{
+	bLayerCreated = false;
+
+	DestroyGizmosActors();
+
+	for (int32 Index = ManipulatableObjects.Num() - 1; Index >= 0 ; --Index)
+	{
+		IControlRigManipulatable* ManipObj = ManipulatableObjects[Index];
+		RemoveManipulatableObject(ManipObj);
+	}
+
+	ManipulatableObjects.Num();
+}
+
+void IControlRigManipulationLayer::AddManipulatableObject(IControlRigManipulatable* InObject)
+{
+	if (InObject && !ManipulatableObjects.Contains(InObject))
+	{
+		ManipulatableObjects.Add(InObject);
+	}
+}
+
+void IControlRigManipulationLayer::RemoveManipulatableObject(IControlRigManipulatable* InObject)
+{
+	if (InObject && ManipulatableObjects.Contains(InObject))
+	{
+		ManipulatableObjects.Remove(InObject);
+	}
+}
+
diff --git a/Engine/Plugins/Experimental/ControlRig/Source/ControlRigManipulation/Public/ControlRigGizmoActor.h b/Engine/Plugins/Experimental/ControlRig/Source/ControlRigManipulation/Public/ControlRigGizmoActor.h
new file mode 100644
--- /dev/null
+++ b/Engine/Plugins/Experimental/ControlRig/Source/ControlRigManipulation/Public/ControlRigGizmoActor.h
@@ -0,0 +1,151 @@
+// Copyright 1998-2019 Epic Games, Inc. All Rights Reserved.
+
+#pragma once
+
+#include "CoreMinimal.h"
+#include "GameFramework/Actor.h"
+#include "Engine/StaticMesh.h"
+#include "Materials/Material.h"
+#include "Rigs/RigControlHierarchy.h"
+#include "ControlRigGizmoActor.generated.h"
+
+struct FActorSpawnParameters;
+class IControlRigManipulatable;
+
+USTRUCT()
+struct FGizmoActorCreationParam
+{
+	GENERATED_BODY()
+
+	FGizmoActorCreationParam()
+		: ManipObj(nullptr)
+		, ControlName(NAME_None)
+		, SpawnTransform(FTransform::Identity)
+		, GizmoTransform(FTransform::Identity)
+		, MeshTransform(FTransform::Identity)
+		, StaticMesh(nullptr)
+		, Material(nullptr)
+		, ColorParameterName(NAME_None)
+		, Color(FLinearColor::Red)
+	{
+	}
+
+	IControlRigManipulatable*	ManipObj;
+	FName		ControlName;
+	FTransform	SpawnTransform;
+	FTransform  GizmoTransform;
+	FTransform  MeshTransform;
+	TAssetPtr<UStaticMesh> StaticMesh;
+	TAssetPtr<UMaterial> Material;
+	FName ColorParameterName;
+	FLinearColor Color;
+};
+
+/** An actor used to represent a rig control */
+UCLASS()
+class CONTROLRIGMANIPULATION_API AControlRigGizmoActor : public AActor
+{
+	GENERATED_BODY()
+
+public:
+	AControlRigGizmoActor(const FObjectInitializer& ObjectInitializer);
+
+	// this is the one holding transform for the controls
+	UPROPERTY()
+	class USceneComponent* ActorRootComponent;
+
+	// this is visual representation of the transform
+	UPROPERTY(Category = StaticMesh, VisibleAnywhere, BlueprintReadOnly, meta = (ExposeFunctionCategories = "Mesh,Rendering,Physics,Components|StaticMesh", AllowPrivateAccess = "true"))
+	class UStaticMeshComponent* StaticMeshComponent;
+
+	UFUNCTION(BlueprintSetter)
+	/** Set the control to be enabled/disabled */
+	virtual void SetEnabled(bool bInEnabled);
+
+	UFUNCTION(BlueprintGetter)
+	/** Get whether the control is enabled/disabled */
+	virtual bool IsEnabled() const;
+
+	UFUNCTION(BlueprintSetter)
+	/** Set the control to be selected/unselected */
+	virtual void SetSelected(bool bInSelected);
+
+	UFUNCTION(BlueprintGetter)
+	/** Get whether the control is selected/unselected */
+	virtual bool IsSelectedInEditor() const override;
+
+	UFUNCTION(BlueprintSetter)
+	/** Set the control to be hovered */
+	virtual void SetHovered(bool bInHovered);
+
+	UFUNCTION(BlueprintGetter)
+	/** Get whether the control is hovered */
+	virtual bool IsHovered() const;
+
+	UFUNCTION(BlueprintSetter)
+	/** Set whether the control is being manipulated */
+	virtual void SetManipulating(bool bInManipulating);
+
+	UFUNCTION(BlueprintGetter)
+	/** Get whether the control is being manipulated */
+	virtual bool IsManipulating() const;
+
+	/** Called from the edit mode each tick */
+	virtual void TickControl() {};
+
+	/** Event called when the transform of this control has changed */
+	UFUNCTION(BlueprintImplementableEvent)
+	void OnTransformChanged(const FTransform& NewTransform);
+
+	/** Event called when the enabled state of this control has changed */
+	UFUNCTION(BlueprintImplementableEvent)
+	void OnEnabledChanged(bool bIsEnabled);
+
+	/** Event called when the selection state of this control has changed */
+	UFUNCTION(BlueprintImplementableEvent)
+	void OnSelectionChanged(bool bIsSelected);
+
+	/** Event called when the hovered state of this control has changed */
+	UFUNCTION(BlueprintImplementableEvent)
+	void OnHoveredChanged(bool bIsSelected);
+
+	/** Event called when the manipulating state of this control has changed */
+	UFUNCTION(BlueprintImplementableEvent)
+	void OnManipulatingChanged(bool bIsManipulating);
+
+	// this returns root component transform based on attach
+	// when there is no attach, it is based on 0
+	UFUNCTION(BlueprintCallable, Category = "ControlRig|Gizmo")
+	void SetGlobalTransform(const FTransform& InTransform);
+
+	UFUNCTION(BlueprintPure, Category = "ControlRig|Gizmo")
+	FTransform GetGlobalTransform() const;
+private:
+	/** Whether this control is enabled */
+	UPROPERTY(BlueprintGetter = IsEnabled, BlueprintSetter= SetEnabled, Category = "ControlRig|Gizmo")
+	uint8 bEnabled : 1;
+
+	/** Whether this control is selected */
+	UPROPERTY(BlueprintGetter = IsSelectedInEditor, BlueprintSetter = SetSelected, Category = "ControlRig|Gizmo")
+	uint8 bSelected : 1;
+
+	/** Whether this control is hovered */
+	UPROPERTY(BlueprintGetter = IsHovered, BlueprintSetter = SetHovered, Category = "ControlRig|Gizmo")
+	uint8 bHovered : 1;
+
+	/** Whether this control is being manipulated */
+	UPROPERTY(BlueprintGetter = IsManipulating, BlueprintSetter = SetManipulating, Category = "ControlRig|Gizmo")
+	uint8 bManipulating : 1;
+};
+
+/**
+ * Creating Gizmo Param helper functions
+ */
+namespace FControlRigGizmoHelper
+{
+	extern CONTROLRIGMANIPULATION_API AControlRigGizmoActor* CreateGizmoActor(UWorld* InWorld, UStaticMesh* InStaticMesh, const FGizmoActorCreationParam& CreationParam);
+	AControlRigGizmoActor* CreateGizmoActor(UWorld* InWorld, TSubclassOf<AControlRigGizmoActor> InClass, const FGizmoActorCreationParam& CreationParam);
+	extern CONTROLRIGMANIPULATION_API AControlRigGizmoActor* CreateDefaultGizmoActor(UWorld* InWorld, const FGizmoActorCreationParam& CreationParam);
+
+	FActorSpawnParameters GetDefaultSpawnParameter();
+}
\ No newline at end of file
diff --git a/Engine/Plugins/Experimental/ControlRig/Source/ControlRigManipulation/Public/DefaultControlRigManipulationLayer.h b/Engine/Plugins/Experimental/ControlRig/Source/ControlRigManipulation/Public/DefaultControlRigManipulationLayer.h
new file mode 100644
--- /dev/null
+++ b/Engine/Plugins/Experimental/ControlRig/Source/ControlRigManipulation/Public/DefaultControlRigManipulationLayer.h
@@ -0,0 +1,103 @@
+// Copyright 1998-2019 Epic Games, Inc. All Rights Reserved.
+
+#pragma once
+
+#include "IControlRigManipulationLayer.h"
+#include "ControlRigGizmoActor.h"
+#include "Units/RigUnitContext.h"
+#include "DefaultControlRigManipulationLayer.generated.h"
+
+class IControlRigObjectBinding;
+typedef uint32 FControlID;
+
+// control ID, I should use ID for it
+struct FControlData
+{
+	IControlRigManipulatable* ManipObject;
+	FName ControlName;
+
+	friend bool operator==(const FControlData& Lhs, const FControlData& Rhs) { return &Lhs.ManipObject == &Rhs.ManipObject && Lhs.ControlName == Rhs.ControlName; }
+};
+
+/**
+ * Default control rig manipulation layer
+ *
+ * This is default manipulation layer that supports editor functionality
+ * This can support multiple control rigs and all control values by types
+ * For now, it only supports 3D spacial types - vector/rotation/transform and one controlrig
+ */
+UCLASS()
+class CONTROLRIGMANIPULATION_API UDefaultControlRigManipulationLayer : public UObject, public IControlRigManipulationLayer
+{
+	GENERATED_UCLASS_BODY()
+
+	virtual ~UDefaultControlRigManipulationLayer() {}
+
+public:
+	// IControlRigManipulationLayer START
+	virtual bool CreateGizmoActors(UWorld* World, TArray<AControlRigGizmoActor*>& OutGizmoActors);
+	virtual void DestroyGizmosActors();
+
+	virtual void CreateLayer() override;
+	virtual void DestroyLayer() override;
+
+	virtual void AddManipulatableObject(IControlRigManipulatable* InObject) override;
+	virtual void RemoveManipulatableObject(IControlRigManipulatable* InObject) override;
+	virtual void TickManipulatableObjects(float DeltaTime) override;
+
+	virtual void SetGizmoTransform(AControlRigGizmoActor* GizmoActor, const FTransform& InTransform) override;
+	virtual void GetGizmoTransform(AControlRigGizmoActor* GizmoActor, FTransform& OutTransform) const override;
+	virtual void MoveGizmo(AControlRigGizmoActor* GizmoActor, const bool bTranslation, FVector& InDrag, 
+		const bool bRotation, FRotator& InRot, const bool bScale, FVector& InScale, const FTransform& ToWorldTransform) override;
+	virtual void TickGizmo(AControlRigGizmoActor* GizmoActor, const FTransform& ComponentTransform) override;
+	// need templatizing input type @fixme: add implementation
+	virtual void AddKey(const AControlRigGizmoActor* GizmoActor, TSharedPtr<ISequencer> Sequencer) override {};
+	virtual bool ModeSupportedByGizmoActor(const AControlRigGizmoActor* GizmoActor, FWidget::EWidgetMode InMode) const override;
+	// IControlRigManipulationLayer END
+
+	// Object binding
+	/** Setup bindings to a runtime object (or clear by passing in nullptr). */
+	void SetObjectBinding(TSharedPtr<IControlRigObjectBinding> InObjectBinding);
+	/** Get bindings to a runtime object */
+	TSharedPtr<IControlRigObjectBinding> GetObjectBinding() const;
+
+	USkeletalMeshComponent* GetSkeletalMeshComponent() const;
+	FTransform	GetSkeletalMeshComponentTransform() const;
+
+	// Utility functions
+	void BeginTransaction();
+	void EndTransaction();
+
+	// in this layer, we only care one to one
+	const FControlData* GetControlDataFromGizmo(const AControlRigGizmoActor* GizmoActor) const;
+	// this is slow, and it only finds the first one, and there is no guarantee it will always find the same name if you run this multiple sessions.
+	// for example, if you have 2 control rigs (A and B for example) with same ControlName, it may find A or B
+	AControlRigGizmoActor* GetGizmoFromControlName(const FName& ControlName) const;
+	bool GetGlobalTransform(AControlRigGizmoActor* GizmoActor, const FName& ControlName, FTransform& OutTransform) const;
+
+private:
+	// GizmoActor* to ControlData index
+	TMap<AControlRigGizmoActor*, FControlID> GizmoToControlMap;
+	// the index is used for FControlID, so ensure the change of order is applied to GizmoToControlMap
+	TArray<FControlData> ControlData;
+
+	// link gizmo actor to (manipulatable object, control name). 
+	void AddToControlData(AControlRigGizmoActor* GizmoActor, IControlRigManipulatable* InManipulatableObject, const FName& InControlName);
+
+	// control data related
+	virtual void ResetControlData();
+
+	void OnRigExecuted(class UControlRig* InRig, const EControlRigState State);
+	TArray<FDelegateHandle> ExecuteDelegateHandles;
+
+	void OnRigInitialized(class UControlRig* InRig, const EControlRigState State);
+	TArray<FDelegateHandle> InitializeDelegateHandle;
+
+	void OnControlModified(IControlRigManipulatable* InManipulatable, const FRigControl& InControl);
+	TArray<FDelegateHandle> ControlModifiedDelegateHandles;
+
+	void OnControlRigAdded(UControlRig* InControlRig);
+	void OnControlRigRemoved(UControlRig* InControlRig);
+
+	void GetGizmoCreationParams(TArray<FGizmoActorCreationParam>& OutCreationParams);
+};
diff --git a/Engine/Plugins/Experimental/ControlRig/Source/ControlRigManipulation/Public/IControlRigManipulationLayer.h b/Engine/Plugins/Experimental/ControlRig/Source/ControlRigManipulation/Public/IControlRigManipulationLayer.h
new file mode 100644
--- /dev/null
+++ b/Engine/Plugins/Experimental/ControlRig/Source/ControlRigManipulation/Public/IControlRigManipulationLayer.h
@@ -0,0 +1,52 @@
+// Copyright 1998-2019 Epic Games, Inc. All Rights Reserved.
+
+#pragma once
+
+#include "CoreMinimal.h"
+#include "UnrealWidget.h"
+
+class ISequencer;
+class IControlRigManipulatable;
+class AControlRigGizmoActor;
+
+/**
+ * Manipulation Layer interface to create interactive interface
+ * This work with IManipulatable object that provides options
+ */
+class CONTROLRIGMANIPULATION_API IControlRigManipulationLayer
+{
+public:
+	IControlRigManipulationLayer();
+	virtual ~IControlRigManipulationLayer() {};
+
+	virtual void CreateLayer();
+	virtual void DestroyLayer();
+
+	// we only allow one type of class for now. This requires re-create layer. 
+	virtual void AddManipulatableObject(IControlRigManipulatable* InObject);
+	virtual void RemoveManipulatableObject(IControlRigManipulatable* InObject);
+	virtual void TickManipulatableObjects(float DeltaTime) = 0;
+
+	// virtual functions for child manipulation layer to write
+	virtual bool CreateGizmoActors(UWorld* World, TArray<AControlRigGizmoActor*>& OutGizmoActors) = 0;
+	virtual void DestroyGizmosActors() = 0;
+	virtual void SetGizmoTransform(AControlRigGizmoActor* GizmoActor, const FTransform& InTransform) = 0;
+	virtual void GetGizmoTransform(AControlRigGizmoActor* GizmoActor, FTransform& OutTransform) const = 0;
+	virtual void MoveGizmo(AControlRigGizmoActor* GizmoActor, const bool bTranslation, FVector& InDrag,
+		const bool bRotation, FRotator& InRot, const bool bScale, FVector& InScale, const FTransform& ToWorldTransform) = 0;
+	virtual void TickGizmo(AControlRigGizmoActor* GizmoActor, const FTransform& ComponentTransform) = 0;
+	virtual bool ModeSupportedByGizmoActor(const AControlRigGizmoActor* GizmoActor, FWidget::EWidgetMode InMode) const = 0;
+
+	// add key of current gizmo actor transform data
+	virtual void AddKey(const AControlRigGizmoActor* GizmoActor, TSharedPtr<ISequencer> Sequencer) = 0;
+
+protected:
+
+	// because manipulatable is created externally often. Manipulation Layer doesn't have any ownership on these object
+	// @fixme: We may want to make weak ptr? Can't be object ptr because IControlRiGManipulatable is not from UObject
+	TArray<IControlRigManipulatable*> ManipulatableObjects;
+
+private:
+	// tracks whether this layer is created or not
+	bool bLayerCreated;
+};
diff --git a/Engine/Plugins/Experimental/ControlRig/Source/ControlRigManipulation/Public/IControlRigManipulationModule.h b/Engine/Plugins/Experimental/ControlRig/Source/ControlRigManipulation/Public/IControlRigManipulationModule.h
new file mode 100644
--- /dev/null
+++ b/Engine/Plugins/Experimental/ControlRig/Source/ControlRigManipulation/Public/IControlRigManipulationModule.h
@@ -0,0 +1,12 @@
+// Copyright 1998-2019 Epic Games, Inc. All Rights Reserved.
+
+#pragma once
+
+#include "CoreMinimal.h"
+#include "Modules/ModuleInterface.h"
+#include "Templates/SubclassOf.h"
+
+class CONTROLRIGMANIPULATION_API IControlRigManipulationModule : public IModuleInterface
+{
+public:
+};
diff --git a/Engine/Plugins/Experimental/PythonScriptPlugin/Source/PythonScriptPlugin/Private/PyUtil.cpp b/Engine/Plugins/Experimental/PythonScriptPlugin/Source/PythonScriptPlugin/Private/PyUtil.cpp
--- a/Engine/Plugins/Experimental/PythonScriptPlugin/Source/PythonScriptPlugin/Private/PyUtil.cpp
+++ b/Engine/Plugins/Experimental/PythonScriptPlugin/Source/PythonScriptPlugin/Private/PyUtil.cpp
@@ -563,6 +563,11 @@ bool InvokeFunctionCall(UObject* InObj, const UFunction* InFunc, void* InBasePar
 
 bool InspectFunctionArgs(PyObject* InFunc, TArray<FString>& OutArgNames, TArray<FPyObjectPtr>* OutArgDefaults)
 {
+	if (!PyFunction_Check(InFunc) && !PyMethod_Check(InFunc))
+	{
+		return false;
+	}
+
 	FPyObjectPtr PyInspectModule = FPyObjectPtr::StealReference(PyImport_ImportModule("inspect"));
 	if (PyInspectModule)
 	{
diff --git a/Engine/Plugins/Tests/PythonAutomationTest/Source/PythonAutomationTest/Private/PythonAutomationTestModule.cpp b/Engine/Plugins/Tests/PythonAutomationTest/Source/PythonAutomationTest/Private/PythonAutomationTestModule.cpp
new file mode 100644
--- /dev/null
+++ b/Engine/Plugins/Tests/PythonAutomationTest/Source/PythonAutomationTest/Private/PythonAutomationTestModule.cpp
@@ -0,0 +1,8 @@
+// Copyright 1998-2019 Epic Games, Inc. All Rights Reserved.
+
+#include "CoreMinimal.h"
+#include "Modules/ModuleManager.h"
+
+class FPythonAutomationTestModule : public FDefaultModuleImpl {};
+
+IMPLEMENT_MODULE(FPythonAutomationTestModule, PythonAutomationTest);
diff --git a/Engine/Plugins/Tests/PythonAutomationTest/Source/PythonAutomationTest/Public/PythonAutomationTest.h b/Engine/Plugins/Tests/PythonAutomationTest/Source/PythonAutomationTest/Public/PythonAutomationTest.h
new file mode 100644
--- /dev/null
+++ b/Engine/Plugins/Tests/PythonAutomationTest/Source/PythonAutomationTest/Public/PythonAutomationTest.h
@@ -0,0 +1,35 @@
+// Copyright 1998-2019 Epic Games, Inc. All Rights Reserved.
+
+#pragma once
+
+#include "CoreMinimal.h"
+#include "UObject/ObjectMacros.h"
+#include "Kismet/BlueprintFunctionLibrary.h"
+
+#include "PythonAutomationTest.generated.h"
+
+UCLASS(meta = (ScriptName = "PyAutomationTest"))
+class PYTHONAUTOMATIONTEST_API UPyAutomationTestLibrary : public UBlueprintFunctionLibrary
+{
+	GENERATED_BODY()
+public:
+
+	UFUNCTION(BlueprintCallable, Category = "PyAutomationTest")
+	static void SetIsRunningPyLatentCommand(bool isRunning);
+
+	UFUNCTION(BlueprintCallable, Category = "PyAutomationTest")
+	static bool GetIsRunningPyLatentCommand();
+
+	UFUNCTION(BlueprintCallable, Category = "PyAutomationTest")
+	static void SetPyLatentCommandTimeout(float Seconds);
+
+	UFUNCTION(BlueprintCallable, Category = "PyAutomationTest")
+	static float GetPyLatentCommandTimeout();
+
+	UFUNCTION(BlueprintCallable, Category = "PyAutomationTest")
+	static void ResetPyLatentCommand();
+
+private:
+	static bool IsRunningPyLatentCommand;
+	static float PyLatentCommandTimeout;
+};
\ No newline at end of file
diff --git a/Engine/Source/Developer/FunctionalTesting/Private/AutomationBlueprintFunctionLibrary.cpp b/Engine/Source/Developer/FunctionalTesting/Private/AutomationBlueprintFunctionLibrary.cpp
--- a/Engine/Source/Developer/FunctionalTesting/Private/AutomationBlueprintFunctionLibrary.cpp
+++ b/Engine/Source/Developer/FunctionalTesting/Private/AutomationBlueprintFunctionLibrary.cpp
@@ -42,6 +42,8 @@
 #include "SceneView.h"
 #include "Engine/GameEngine.h"
 #include "Engine/LevelStreaming.h"
+#include "Templates/UnrealTemplate.h"
+#include "UObject/GCObjectScopeGuard.h"
 
 #define LOCTEXT_NAMESPACE "Automation"
 
@@ -60,6 +62,22 @@ static TAutoConsoleVariable<int32> CVarAutomationScreenshotResolutionHeight(
 	TEXT("The height of automation screenshots."),
 	ECVF_Default);
 
+
+bool UAutomationEditorTask::IsValidTask() const
+{
+	return Task.IsValid();
+}
+
+void UAutomationEditorTask::BindTask(TUniquePtr<FAutomationTaskStatusBase> inTask)
+{
+	Task = MoveTemp(inTask);
+}
+
+bool UAutomationEditorTask::IsTaskDone() const
+{
+	return IsValidTask() && Task->IsDone();
+}
+
 #if (WITH_DEV_AUTOMATION_TESTS || WITH_PERF_AUTOMATION_TESTS)
 
 template<typename T>
@@ -403,29 +421,29 @@ public:
 		{
 			FAutomationScreenshotData Data = AutomationCommon::BuildScreenshotData(World->GetName(), Name, InSizeX, InSizeY);
 
-		// Copy the relevant data into the metadata for the screenshot.
-		Data.bHasComparisonRules = true;
-		Data.ToleranceRed = Options.ToleranceAmount.Red;
-		Data.ToleranceGreen = Options.ToleranceAmount.Green;
-		Data.ToleranceBlue = Options.ToleranceAmount.Blue;
-		Data.ToleranceAlpha = Options.ToleranceAmount.Alpha;
-		Data.ToleranceMinBrightness = Options.ToleranceAmount.MinBrightness;
-		Data.ToleranceMaxBrightness = Options.ToleranceAmount.MaxBrightness;
-		Data.bIgnoreAntiAliasing = Options.bIgnoreAntiAliasing;
-		Data.bIgnoreColors = Options.bIgnoreColors;
-		Data.MaximumLocalError = Options.MaximumLocalError;
-		Data.MaximumGlobalError = Options.MaximumGlobalError;
+			// Copy the relevant data into the metadata for the screenshot.
+			Data.bHasComparisonRules = true;
+			Data.ToleranceRed = Options.ToleranceAmount.Red;
+			Data.ToleranceGreen = Options.ToleranceAmount.Green;
+			Data.ToleranceBlue = Options.ToleranceAmount.Blue;
+			Data.ToleranceAlpha = Options.ToleranceAmount.Alpha;
+			Data.ToleranceMinBrightness = Options.ToleranceAmount.MinBrightness;
+			Data.ToleranceMaxBrightness = Options.ToleranceAmount.MaxBrightness;
+			Data.bIgnoreAntiAliasing = Options.bIgnoreAntiAliasing;
+			Data.bIgnoreColors = Options.bIgnoreColors;
+			Data.MaximumLocalError = Options.MaximumLocalError;
+			Data.MaximumGlobalError = Options.MaximumGlobalError;
 
 			// Record any user notes that were made to accompany this shot.
 			Data.Notes = Notes;
 
 			bool bAttemptToCompareShot = FAutomationTestFramework::Get().OnScreenshotCaptured().ExecuteIfBound(InImageData, Data);
 
-		UE_LOG(AutomationFunctionLibrary, Log, TEXT("Screenshot captured as %s"), *Data.Path);
+			UE_LOG(AutomationFunctionLibrary, Log, TEXT("Screenshot captured as %s"), *Data.Path);
 
-		if ( GIsAutomationTesting )
-		{
-			FAutomationTestFramework::Get().OnScreenshotCompared.AddRaw(this, &FAutomationScreenshotTaker::OnComparisonComplete);
+			if ( GIsAutomationTesting )
+			{
+				FAutomationTestFramework::Get().OnScreenshotCompared.AddRaw(this, &FAutomationScreenshotTaker::OnComparisonComplete);
 				FScreenshotRequest::OnScreenshotRequestProcessed().RemoveAll(this);
 				return;
 			}
@@ -478,8 +496,129 @@ private:
 	bool bNeedsViewportSizeRestore;
 };
 
+class FAutomationHighResScreenshotGrabber
+{
+public:
+	FAutomationHighResScreenshotGrabber(const FString& InContext, const FString& InName, const FString& InNotes, FAutomationScreenshotOptions InOptions)
+		: Context(InContext)
+		, Name(InName)
+		, Notes(InNotes)
+		, Options(InOptions)
+	{
+		FScreenshotRequest::OnScreenshotCaptured().AddRaw(this, &FAutomationHighResScreenshotGrabber::GrabScreenShot);
+		FWorldDelegates::LevelRemovedFromWorld.AddRaw(this, &FAutomationHighResScreenshotGrabber::WorldDestroyed);
+	}
+
+	virtual ~FAutomationHighResScreenshotGrabber()
+	{
+		FAutomationTestFramework::Get().OnScreenshotCompared.RemoveAll(this);
+		FScreenshotRequest::OnScreenshotCaptured().RemoveAll(this);
+
+		FWorldDelegates::LevelRemovedFromWorld.RemoveAll(this);
+
+		FAutomationTestFramework::Get().NotifyScreenshotTakenAndCompared();
+	}
+
+	void GrabScreenShot(int32 InSizeX, int32 InSizeY, const TArray<FColor>& InImageData)
+	{
+		FScreenshotRequest::OnScreenshotCaptured().RemoveAll(this);
+
+		FAutomationScreenshotData Data = AutomationCommon::BuildScreenshotData(Context, Name, InSizeX, InSizeY);
+
+		// Copy the relevant data into the metadata for the screenshot.
+		Data.bHasComparisonRules = true;
+		Data.ToleranceRed = Options.ToleranceAmount.Red;
+		Data.ToleranceGreen = Options.ToleranceAmount.Green;
+		Data.ToleranceBlue = Options.ToleranceAmount.Blue;
+		Data.ToleranceAlpha = Options.ToleranceAmount.Alpha;
+		Data.ToleranceMinBrightness = Options.ToleranceAmount.MinBrightness;
+		Data.ToleranceMaxBrightness = Options.ToleranceAmount.MaxBrightness;
+		Data.bIgnoreAntiAliasing = Options.bIgnoreAntiAliasing;
+		Data.bIgnoreColors = Options.bIgnoreColors;
+		Data.MaximumLocalError = Options.MaximumLocalError;
+		Data.MaximumGlobalError = Options.MaximumGlobalError;
+
+		// Record any user notes that were made to accompany this shot.
+		Data.Notes = Notes;
+
+		bool bAttemptToCompareShot = FAutomationTestFramework::Get().OnScreenshotCaptured().ExecuteIfBound(InImageData, Data);
+
+		UE_LOG(AutomationFunctionLibrary, Log, TEXT("Screenshot captured as %s"), *Data.Path);
+
+		FAutomationTestFramework::Get().OnScreenshotCompared.AddRaw(this, &FAutomationHighResScreenshotGrabber::OnComparisonComplete);
+	}
+
+	void OnComparisonComplete(const FAutomationScreenshotCompareResults& CompareResults)
+	{
+		FAutomationTestFramework::Get().OnScreenshotCompared.RemoveAll(this);
+
+		if (FAutomationTestBase* CurrentTest = FAutomationTestFramework::Get().GetCurrentTest())
+		{
+			CurrentTest->AddEvent(CompareResults.ToAutomationEvent(Name));
+		}
+
+		delete this;
+	}
+
+	void WorldDestroyed(ULevel* InLevel, UWorld* InWorld)
+	{
+		// If the InLevel is null, it's a signal that the entire world is about to disappear, so
+		// go ahead and remove this widget from the viewport, it could be holding onto too many
+		// dangerous actor references that won't carry over into the next world.
+		if (InLevel == nullptr)
+		{
+			delete this;
+		}
+	}
+
+private:
+	FString	Context;
+	FString	Name;
+	FString Notes;
+	FAutomationScreenshotOptions Options;
+};
+
 #endif
 
+class FScreenshotTakenState : public FAutomationTaskStatusBase
+{
+public:
+	FScreenshotTakenState()
+	{
+		if (GIsAutomationTesting)
+		{
+			// When Automation test are running we hook to the FAutomationTestFramework::OnComparisonComplete instead of the
+			// FScreenshotRequest::OnScreenshotRequestProcessed, because with HighResScreenshot, FScreenshotRequest::OnScreenshotRequestProcessed
+			// is fired before comparison is completed.
+			FAutomationTestFramework::Get().OnScreenshotCompared.AddRaw(this, &FScreenshotTakenState::OnComparisonComplete);
+		}
+		else
+		{
+			FScreenshotRequest::OnScreenshotRequestProcessed().AddRaw(this, &FScreenshotTakenState::SetDone);
+		}
+	};
+
+	virtual ~FScreenshotTakenState()
+	{
+#if (WITH_DEV_AUTOMATION_TESTS || WITH_PERF_AUTOMATION_TESTS)
+		FAutomationTestFramework::Get().OnScreenshotCompared.RemoveAll(this);
+#endif
+		FScreenshotRequest::OnScreenshotRequestProcessed().RemoveAll(this);
+	};
+
+	virtual void SetDone() override
+	{
+		FScreenshotRequest::OnScreenshotRequestProcessed().RemoveAll(this);
+		Done = true;
+	};
+
+	void OnComparisonComplete(const FAutomationScreenshotCompareResults& CompareResults)
+	{
+		FAutomationTestFramework::Get().OnScreenshotCompared.RemoveAll(this);
+		SetDone();
+	};
+};
+
 UAutomationBlueprintFunctionLibrary::UAutomationBlueprintFunctionLibrary(const class FObjectInitializer& Initializer)
 	: Super(Initializer)
 {
@@ -905,8 +1044,11 @@ void UAutomationBlueprintFunctionLibrary::AutomationWaitForLoading(UObject* Worl
 	}
 }
 
-bool UAutomationBlueprintFunctionLibrary::TakeHighResScreenshot(int32 ResX, int32 ResY, FString Filename, ACameraActor* Camera, bool bMaskEnabled, bool bCaptureHDR)
+UAutomationEditorTask* UAutomationBlueprintFunctionLibrary::TakeHighResScreenshot(int32 ResX, int32 ResY, FString Filename, ACameraActor* Camera, bool bMaskEnabled, bool bCaptureHDR, EComparisonTolerance ComparisonTolerance, FString ComparisonNotes)
 {
+	UAutomationEditorTask* Task = NewObject<UAutomationEditorTask>();
+	FGCObjectScopeGuard TaskGuard(Task);
+
 #if WITH_EDITOR
 	if (FModuleManager::Get().IsModuleLoaded("LevelEditor"))
 	{
@@ -933,15 +1075,28 @@ bool UAutomationBlueprintFunctionLibrary::TakeHighResScreenshot(int32 ResX, int3
 
 			FinishLoadingBeforeScreenshot();
 
+#if (WITH_DEV_AUTOMATION_TESTS || WITH_PERF_AUTOMATION_TESTS)
+			if (GIsAutomationTesting)
+			{
+				if (FAutomationTestBase* CurrentTest = FAutomationTestFramework::Get().GetCurrentTest())
+				{
+					FString Context = CurrentTest->GetTestContext();
+					if (Context.IsEmpty()) { Context = CurrentTest->GetTestName(); }
+					FAutomationScreenshotOptions ComparisonOptions = FAutomationScreenshotOptions(ComparisonTolerance);
+					FAutomationHighResScreenshotGrabber* TempObject = new FAutomationHighResScreenshotGrabber(Context, Filename, ComparisonNotes, ComparisonOptions);
+				} //-V773
+			}
+#endif
+			Task->BindTask(MakeUnique<FScreenshotTakenState>());
 			LevelViewport->GetActiveViewport()->TakeHighResScreenShot();
 
-			return true;
+			return Task;
 		}
 
 		UE_LOG(AutomationFunctionLibrary, Error, TEXT("Screenshot size exceeds the maximum allowed texture size (%d x %d)"), GetMax2DTextureDimension(), GetMax2DTextureDimension());
 	}
 #endif
-	return false;
+	return Task;
 }
 
 FAutomationScreenshotOptions UAutomationBlueprintFunctionLibrary::GetDefaultScreenshotOptionsForGameplay(EComparisonTolerance Tolerance, float Delay)
diff --git a/Engine/Source/Developer/FunctionalTesting/Public/AutomationBlueprintFunctionLibrary.h b/Engine/Source/Developer/FunctionalTesting/Public/AutomationBlueprintFunctionLibrary.h
--- a/Engine/Source/Developer/FunctionalTesting/Public/AutomationBlueprintFunctionLibrary.h
+++ b/Engine/Source/Developer/FunctionalTesting/Public/AutomationBlueprintFunctionLibrary.h
@@ -7,11 +7,52 @@
 #include "Kismet/BlueprintFunctionLibrary.h"
 #include "AutomationScreenshotOptions.h"
 #include "HAL/IConsoleManager.h"
+#include "Templates/UniquePtr.h"
 
 #include "AutomationBlueprintFunctionLibrary.generated.h"
 
 class ACameraActor;
 
+/**
+ * FAutomationTaskStatusBase - abstract class for task status
+ */
+class FAutomationTaskStatusBase
+{
+public:
+	virtual ~FAutomationTaskStatusBase() = default;
+
+	bool IsDone() const { return Done; };
+	virtual void SetDone() { Done = true; };
+
+protected:
+	bool Done;
+};
+
+/**
+ * UAutomationEditorTask
+ */
+UCLASS(BlueprintType, Transient)
+class FUNCTIONALTESTING_API UAutomationEditorTask : public UObject
+{
+	GENERATED_BODY()
+
+public:
+	virtual ~UAutomationEditorTask() = default;
+
+	/** Query if the Editor task is done  */
+	UFUNCTION(BlueprintCallable, Category = "Automation")
+	bool IsTaskDone() const;
+
+	/** Query if a task was setup */
+	UFUNCTION(BlueprintCallable, Category = "Automation")
+	bool IsValidTask() const;
+
+	void BindTask(TUniquePtr<FAutomationTaskStatusBase> inTask);
+
+private:
+	TUniquePtr<FAutomationTaskStatusBase> Task;
+};
+
 /**
  * 
  */
@@ -80,8 +121,8 @@ public:
 	/**
 	* take high res screenshot in editor.
 	*/
-	UFUNCTION(BlueprintCallable, Category = "Automation")
-	static bool TakeHighResScreenshot(int32 ResX, int32 ResY, FString Filename, ACameraActor* Camera = nullptr, bool bMaskEnabled = false, bool bCaptureHDR = false);
+	UFUNCTION(BlueprintCallable, Category = "Automation", meta = (AdvancedDisplay="Camera, bMaskEnabled, bCaptureHDR, ComparisonTolerance, ComparisonNotes"))
+	static UAutomationEditorTask* TakeHighResScreenshot(int32 ResX, int32 ResY, FString Filename, ACameraActor* Camera = nullptr, bool bMaskEnabled = false, bool bCaptureHDR = false, EComparisonTolerance ComparisonTolerance = EComparisonTolerance::Low, FString ComparisonNotes = TEXT(""));
 
 	/**
 	 * 
diff --git a/Engine/Source/Editor/AnimGraph/Private/AnimGraphDetails.cpp b/Engine/Source/Editor/AnimGraph/Private/AnimGraphDetails.cpp
--- a/Engine/Source/Editor/AnimGraph/Private/AnimGraphDetails.cpp
+++ b/Engine/Source/Editor/AnimGraph/Private/AnimGraphDetails.cpp
@@ -145,6 +145,7 @@ void FAnimGraphDetails::CustomizeDetails(IDetailLayoutBuilder& DetailLayout)
 	];
 	InputsHeaderContentWidget->AddSlot()
 	.AutoWidth()
+	.VAlign(VAlign_Center)
 	[
 		SNew(SButton)
 		.ButtonStyle(FEditorStyle::Get(), "RoundButton")
diff --git a/Engine/Source/Editor/AnimationBlueprintEditor/Private/AnimationStateNodes/SGraphNodeAnimTransition.cpp b/Engine/Source/Editor/AnimationBlueprintEditor/Private/AnimationStateNodes/SGraphNodeAnimTransition.cpp
--- a/Engine/Source/Editor/AnimationBlueprintEditor/Private/AnimationStateNodes/SGraphNodeAnimTransition.cpp
+++ b/Engine/Source/Editor/AnimationBlueprintEditor/Private/AnimationStateNodes/SGraphNodeAnimTransition.cpp
@@ -190,11 +190,6 @@ void SGraphNodeAnimTransition::UpdateGraphNode()
 		[
 			SNew(SOverlay)
 			+SOverlay::Slot()
-			[
-				SNew(SImage)
-				.Image( FEditorStyle::GetBrush("Graph.TransitionNode.Body") )
-			]
-			+SOverlay::Slot()
 			[
 				SNew(SImage)
 				.Image( FEditorStyle::GetBrush("Graph.TransitionNode.ColorSpill") )
@@ -203,12 +198,7 @@ void SGraphNodeAnimTransition::UpdateGraphNode()
 			+SOverlay::Slot()
 			[
 				SNew(SImage)
-				.Image( FEditorStyle::GetBrush("Graph.TransitionNode.Icon") )
-			]
-			+SOverlay::Slot()
-			[
-				SNew(SImage)
-				.Image( FEditorStyle::GetBrush("Graph.TransitionNode.Gloss") )
+				.Image( this, &SGraphNodeAnimTransition::GetTransitionIconImage )
 			]
 		];
 }
@@ -333,6 +323,14 @@ FSlateColor SGraphNodeAnimTransition::GetTransitionColor() const
 	return StaticGetTransitionColor(TransNode, IsHovered());
 }
 
+const FSlateBrush* SGraphNodeAnimTransition::GetTransitionIconImage() const
+{
+	UAnimStateTransitionNode* TransNode = CastChecked<UAnimStateTransitionNode>(GraphNode);
+	return (TransNode->LogicType == ETransitionLogicType::TLT_Inertialization)
+		? FEditorStyle::GetBrush("Graph.TransitionNode.Icon_Inertialization")
+		: FEditorStyle::GetBrush("Graph.TransitionNode.Icon");
+}
+
 FString SGraphNodeAnimTransition::GetCurrentDuration() const
 {
 	UAnimStateTransitionNode* TransNode = CastChecked<UAnimStateTransitionNode>(GraphNode);
diff --git a/Engine/Source/Editor/AnimationBlueprintEditor/Private/AnimationStateNodes/SGraphNodeAnimTransition.h b/Engine/Source/Editor/AnimationBlueprintEditor/Private/AnimationStateNodes/SGraphNodeAnimTransition.h
--- a/Engine/Source/Editor/AnimationBlueprintEditor/Private/AnimationStateNodes/SGraphNodeAnimTransition.h
+++ b/Engine/Source/Editor/AnimationBlueprintEditor/Private/AnimationStateNodes/SGraphNodeAnimTransition.h
@@ -47,6 +47,7 @@ private:
 private:
 	FText GetPreviewCornerText(bool reverse) const;
 	FSlateColor GetTransitionColor() const;
+	const FSlateBrush* GetTransitionIconImage() const;
 
 	TSharedRef<SWidget> GenerateInlineDisplayOrEditingWidget(bool bShowGraphPreview);
 	TSharedRef<SWidget> GenerateRichTooltip();
diff --git a/Engine/Source/Editor/DetailCustomizations/Private/AnimTransitionNodeDetails.cpp b/Engine/Source/Editor/DetailCustomizations/Private/AnimTransitionNodeDetails.cpp
--- a/Engine/Source/Editor/DetailCustomizations/Private/AnimTransitionNodeDetails.cpp
+++ b/Engine/Source/Editor/DetailCustomizations/Private/AnimTransitionNodeDetails.cpp
@@ -73,8 +73,12 @@ void FAnimTransitionNodeDetails::CustomizeDetails( IDetailLayoutBuilder& DetailB
 		}
 	}
 
+	UAnimStateTransitionNode* TransNode = TransitionNode.Get();
 	IDetailCategoryBuilder& TransitionCategory = DetailBuilder.EditCategory("Transition", LOCTEXT("TransitionCategoryTitle", "Transition") );
 
+	// Added below via CreateBlendProfilePicker
+	DetailBuilder.HideProperty(GET_MEMBER_NAME_CHECKED(UAnimStateTransitionNode, BlendProfile));
+
 	if (bTransitionToConduit)
 	{
 		// Transitions to conduits are just shorthand for some other real transition;
@@ -82,6 +86,7 @@ void FAnimTransitionNodeDetails::CustomizeDetails( IDetailLayoutBuilder& DetailB
 		DetailBuilder.HideProperty(GET_MEMBER_NAME_CHECKED(UAnimStateTransitionNode, Bidirectional));
 		DetailBuilder.HideProperty(GET_MEMBER_NAME_CHECKED(UAnimStateTransitionNode, CrossfadeDuration));
 		DetailBuilder.HideProperty(GET_MEMBER_NAME_CHECKED(UAnimStateTransitionNode, BlendMode));
+		DetailBuilder.HideProperty(GET_MEMBER_NAME_CHECKED(UAnimStateTransitionNode, CustomBlendCurve));
 		DetailBuilder.HideProperty(GET_MEMBER_NAME_CHECKED(UAnimStateTransitionNode, LogicType));
 		DetailBuilder.HideProperty(GET_MEMBER_NAME_CHECKED(UAnimStateTransitionNode, PriorityOrder));
 	}
@@ -120,7 +125,6 @@ void FAnimTransitionNodeDetails::CustomizeDetails( IDetailLayoutBuilder& DetailB
 				]
 			];
 
-		UAnimStateTransitionNode* TransNode = TransitionNode.Get();
 		if (TransitionNode != NULL && SelectedObjects.Num() == 1)
 		{
 			// The sharing option for the rule
@@ -142,15 +146,6 @@ void FAnimTransitionNodeDetails::CustomizeDetails( IDetailLayoutBuilder& DetailB
 					FOnGetContent::CreateSP(this, &FAnimTransitionNodeDetails::OnGetShareableNodesMenu, true))
 			];
 
-
-// 			TransitionCategory.AddRow()
-// 				[
-// 					SNew( STextBlock )
-// 					.Text( TEXT("Crossfade Settings") )
-// 					.Font( IDetailLayoutBuilder::GetDetailFontBold() )
-// 				];
-
-
 			// Show the rule itself
 			UEdGraphPin* CanExecPin = NULL;
 			if (UAnimationTransitionGraph* TransGraph = Cast<UAnimationTransitionGraph>(TransNode->BoundGraph))
@@ -195,11 +190,24 @@ void FAnimTransitionNodeDetails::CustomizeDetails( IDetailLayoutBuilder& DetailB
 
 		//////////////////////////////////////////////////////////////////////////
 
+		auto BlendSettingsEnabled = [LogicTypeHandle]()
+		{
+			uint8 LogicType;
+			if (LogicTypeHandle->GetValue(LogicType) == FPropertyAccess::Result::Success)
+			{
+				return LogicType != ETransitionLogicType::TLT_Inertialization;
+			}
+			return true;
+		};
+
+		auto BlendSettingsEnabledAttribute = TAttribute<bool>::Create(TAttribute<bool>::FGetter::CreateLambda(BlendSettingsEnabled));
+
 		IDetailCategoryBuilder& CrossfadeCategory = DetailBuilder.EditCategory("BlendSettings", LOCTEXT("BlendSettingsCategoryTitle", "Blend Settings") );
 		if (TransitionNode != NULL && SelectedObjects.Num() == 1)
 		{
 			// The sharing option for the crossfade settings
 			CrossfadeCategory.AddCustomRow( LOCTEXT("TransitionCrossfadeSharingLabel", "Transition Crossfade Sharing") )
+			.IsEnabled(BlendSettingsEnabledAttribute)
 			.NameContent()
 			[
 				SNew(STextBlock)
@@ -218,12 +226,12 @@ void FAnimTransitionNodeDetails::CustomizeDetails( IDetailLayoutBuilder& DetailB
 			];
 		}
 
-		//@TODO: Gate editing these on shared non-authorative ones
+		//@TODO: Gate editing these on shared non-authoritative ones
 		CrossfadeCategory.AddProperty(GET_MEMBER_NAME_CHECKED(UAnimStateTransitionNode, CrossfadeDuration)).DisplayName( LOCTEXT("DurationLabel", "Duration") );
-		CrossfadeCategory.AddProperty(GET_MEMBER_NAME_CHECKED(UAnimStateTransitionNode, BlendMode)).DisplayName( LOCTEXT("ModeLabel", "Mode") );
-		CrossfadeCategory.AddProperty(GET_MEMBER_NAME_CHECKED(UAnimStateTransitionNode, CustomBlendCurve)).DisplayName(LOCTEXT("CurveLabel", "Custom Blend Curve"));
+		CrossfadeCategory.AddProperty(GET_MEMBER_NAME_CHECKED(UAnimStateTransitionNode, BlendMode)).DisplayName(LOCTEXT("ModeLabel", "Mode")).IsEnabled(BlendSettingsEnabledAttribute);
+		CrossfadeCategory.AddProperty(GET_MEMBER_NAME_CHECKED(UAnimStateTransitionNode, CustomBlendCurve)).DisplayName(LOCTEXT("CurveLabel", "Custom Blend Curve")).IsEnabled(BlendSettingsEnabledAttribute);
 
-		USkeleton* TargetSkeleton = TransitionNode->GetAnimBlueprint()->TargetSkeleton;
+		USkeleton* TargetSkeleton = TransNode ? TransNode->GetAnimBlueprint()->TargetSkeleton : nullptr;
 
 		if(TargetSkeleton)
 		{
@@ -241,7 +249,8 @@ void FAnimTransitionNodeDetails::CustomizeDetails( IDetailLayoutBuilder& DetailB
 			Args.bAllowClear = true;
 			Args.bAllowRemove = false;
 
-			CrossfadeCategory.AddProperty(BlendProfileHandle).CustomWidget(true)
+			CrossfadeCategory.AddCustomRow(LOCTEXT("BlendProfileLabel", "Blend Profile"))
+				.IsEnabled(BlendSettingsEnabledAttribute)
 				.NameContent()
 				[
 					BlendProfileHandle->CreatePropertyNameWidget()
diff --git a/Engine/Source/Editor/EditorStyle/Private/SlateEditorStyle.cpp b/Engine/Source/Editor/EditorStyle/Private/SlateEditorStyle.cpp
--- a/Engine/Source/Editor/EditorStyle/Private/SlateEditorStyle.cpp
+++ b/Engine/Source/Editor/EditorStyle/Private/SlateEditorStyle.cpp
@@ -4067,10 +4067,9 @@ void FSlateEditorStyle::FStyle::SetupGraphEditorStyles()
 			// Transition node
 			{
 				FMargin TestMargin(16.f/64.f, 16.f/28.f, 16.f/64.f, 4.f/28.f);
-				Set( "Graph.TransitionNode.Body", new BOX_BRUSH( "/Persona/StateMachineEditor/Trans_Node_Body", FMargin(16.f/64.f, 12.f/28.f) ) );
 				Set( "Graph.TransitionNode.ColorSpill", new BOX_BRUSH( "/Persona/StateMachineEditor/Trans_Node_ColorSpill", TestMargin ) );
-				Set( "Graph.TransitionNode.Gloss", new BOX_BRUSH( "/Persona/StateMachineEditor/Trans_Node_Gloss", TestMargin) );
 				Set( "Graph.TransitionNode.Icon", new IMAGE_BRUSH( "/Persona/StateMachineEditor/Trans_Node_Icon", FVector2D(25,25) ) );
+				Set( "Graph.TransitionNode.Icon_Inertialization", new IMAGE_BRUSH( "/Persona/StateMachineEditor/Trans_Node_Icon_Inertialization", FVector2D(25,25) ) );
 			}
 
 			// Transition rule tooltip name 
diff --git a/Engine/Source/Editor/Kismet/Private/BlueprintEditorModule.cpp b/Engine/Source/Editor/Kismet/Private/BlueprintEditorModule.cpp
--- a/Engine/Source/Editor/Kismet/Private/BlueprintEditorModule.cpp
+++ b/Engine/Source/Editor/Kismet/Private/BlueprintEditorModule.cpp
@@ -140,7 +140,14 @@ void FixSubObjectReferencesPostUndoRedo(const FTransaction* Transaction)
 	{
 		FixSubObjectReferencesPostUndoRedo(Blueprint->GeneratedClass->GetDefaultObject());
 		// Will cause a call to RefreshEditors()
-		FBlueprintEditorUtils::MarkBlueprintAsStructurallyModified(Blueprint);
+		if (Blueprint->RequiresMarkAsStructurallyModifiedOnUndo())
+		{
+			FBlueprintEditorUtils::MarkBlueprintAsStructurallyModified(Blueprint);
+		}
+		else
+		{
+			Blueprint->MarkPackageDirty();
+		}
 	}
 }
 
diff --git a/Engine/Source/Editor/Kismet/Private/SMyBlueprint.cpp b/Engine/Source/Editor/Kismet/Private/SMyBlueprint.cpp
--- a/Engine/Source/Editor/Kismet/Private/SMyBlueprint.cpp
+++ b/Engine/Source/Editor/Kismet/Private/SMyBlueprint.cpp
@@ -755,7 +755,7 @@ TSharedRef<SWidget> SMyBlueprint::CreateAddToSectionButton(int32 InSectionID, TW
 		.ForegroundColor(FEditorStyle::GetSlateColor("DefaultForeground"))
 		.ContentPadding(FMargin(2, 0))
 		.OnClicked(this, &SMyBlueprint::OnAddButtonClickedOnSection, InSectionID)
-		.IsEnabled(this, &SMyBlueprint::IsEditingMode)
+		.IsEnabled(this, &SMyBlueprint::CanAddNewElementToSection, InSectionID)
 		.HAlign(HAlign_Center)
 		.VAlign(VAlign_Center)
 		.AddMetaData<FTagMetaData>(FTagMetaData(MetaDataTag))
@@ -816,6 +816,39 @@ FReply SMyBlueprint::OnAddButtonClickedOnSection(int32 InSectionID)
 	return FReply::Handled();
 }
 
+bool SMyBlueprint::CanAddNewElementToSection(int32 InSectionID) const
+{
+	if (!IsEditingMode())
+	{
+		return false;
+	}
+
+	if (UBlueprint* CurrentBlueprint = GetBlueprintObj())
+	{
+		switch (InSectionID)
+		{
+		case NodeSectionID::VARIABLE:
+			return CurrentBlueprint->SupportsGlobalVariables();
+		case NodeSectionID::FUNCTION:
+			return CurrentBlueprint->SupportsFunctions();
+		case NodeSectionID::MACRO:
+			return CurrentBlueprint->SupportsMacros();
+		case NodeSectionID::DELEGATE:
+			return CurrentBlueprint->SupportsDelegates();
+		case NodeSectionID::GRAPH:
+			return CurrentBlueprint->SupportsEventGraphs();
+		case NodeSectionID::ANIMLAYER:
+			return CurrentBlueprint->SupportsAnimLayers();
+		case NodeSectionID::LOCAL_VARIABLE:
+			return CurrentBlueprint->SupportsLocalVariables();
+		default:
+			break;
+		}
+	}
+
+	return false;
+}
+
 EVisibility SMyBlueprint::OnGetSectionTextVisibility(TWeakPtr<SWidget> RowWidget, int32 InSectionID) const
 {
 	bool ShowText = RowWidget.Pin()->IsHovered();
@@ -2073,25 +2106,45 @@ TSharedRef<SWidget> SMyBlueprint::CreateAddNewMenuWidget()
 void SMyBlueprint::BuildAddNewMenu(FMenuBuilder& MenuBuilder)
 {
 	MenuBuilder.BeginSection("AddNewItem", LOCTEXT("AddOperations", "Add New"));
-	{
-		MenuBuilder.AddMenuEntry(FBlueprintEditorCommands::Get().AddNewVariable);
-		MenuBuilder.AddMenuEntry(FBlueprintEditorCommands::Get().AddNewLocalVariable);
-		MenuBuilder.AddMenuEntry(FBlueprintEditorCommands::Get().AddNewFunction);
 
-		// If we cannot handle Function Graphs, we cannot handle function overrides
-		if ( OverridableFunctionActions.Num() > 0 && BlueprintEditorPtr.Pin()->NewDocument_IsVisibleForType(FBlueprintEditor::CGT_NewFunctionGraph) )
+	if(UBlueprint* CurrentBlueprint = GetBlueprintObj())
+	{
+		if (CurrentBlueprint->SupportsGlobalVariables())
+		{
+			MenuBuilder.AddMenuEntry(FBlueprintEditorCommands::Get().AddNewVariable);
+		}
+		if (CurrentBlueprint->SupportsLocalVariables())
 		{
-			MenuBuilder.AddSubMenu(
-				LOCTEXT("OverrideFunction", "Override Function"),
-				FText::GetEmpty(),
-				FNewMenuDelegate::CreateSP(this, &SMyBlueprint::BuildOverridableFunctionsMenu),
-				false,
-				FSlateIcon(FEditorStyle::GetStyleSetName(), "BlueprintEditor.AddNewFunction.Small"));
+			MenuBuilder.AddMenuEntry(FBlueprintEditorCommands::Get().AddNewLocalVariable);
 		}
+		if (CurrentBlueprint->SupportsFunctions())
+		{
+			MenuBuilder.AddMenuEntry(FBlueprintEditorCommands::Get().AddNewFunction);
 
-		MenuBuilder.AddMenuEntry(FBlueprintEditorCommands::Get().AddNewMacroDeclaration);
-		MenuBuilder.AddMenuEntry(FBlueprintEditorCommands::Get().AddNewEventGraph);
-		MenuBuilder.AddMenuEntry(FBlueprintEditorCommands::Get().AddNewDelegate);
+			// If we cannot handle Function Graphs, we cannot handle function overrides
+			if (OverridableFunctionActions.Num() > 0 && BlueprintEditorPtr.Pin()->NewDocument_IsVisibleForType(FBlueprintEditor::CGT_NewFunctionGraph))
+			{
+				MenuBuilder.AddSubMenu(
+					LOCTEXT("OverrideFunction", "Override Function"),
+					FText::GetEmpty(),
+					FNewMenuDelegate::CreateSP(this, &SMyBlueprint::BuildOverridableFunctionsMenu),
+					false,
+					FSlateIcon(FEditorStyle::GetStyleSetName(), "BlueprintEditor.AddNewFunction.Small"));
+			}
+		}
+
+		if (CurrentBlueprint->SupportsMacros())
+		{
+			MenuBuilder.AddMenuEntry(FBlueprintEditorCommands::Get().AddNewMacroDeclaration);
+		}
+		if (CurrentBlueprint->SupportsEventGraphs())
+		{
+			MenuBuilder.AddMenuEntry(FBlueprintEditorCommands::Get().AddNewEventGraph);
+		}
+		if (CurrentBlueprint->SupportsDelegates())
+		{
+			MenuBuilder.AddMenuEntry(FBlueprintEditorCommands::Get().AddNewDelegate);
+		}
 	}
 	MenuBuilder.EndSection();
 }
diff --git a/Engine/Source/Editor/Persona/Private/SAnimViewportToolBar.cpp b/Engine/Source/Editor/Persona/Private/SAnimViewportToolBar.cpp
--- a/Engine/Source/Editor/Persona/Private/SAnimViewportToolBar.cpp
+++ b/Engine/Source/Editor/Persona/Private/SAnimViewportToolBar.cpp
@@ -1123,6 +1123,18 @@ void SAnimViewportToolBar::OnFloorOffsetChanged( float NewValue )
 	PinnedCommands->AddCustomWidget(TEXT("FloorOffsetWidget"));
 }
 
+void SAnimViewportToolBar::AddMenuExtender(FName MenuToExtend, FMenuExtensionDelegate MenuBuilderDelegate)
+{
+	TSharedRef<FExtender> Extender(new FExtender());
+	Extender->AddMenuExtension(
+		MenuToExtend,
+		EExtensionHook::After,
+		CommandList,
+		MenuBuilderDelegate
+	);
+	Extenders.Add(Extender);
+}
+
 TSharedRef<FExtender> SAnimViewportToolBar::GetViewMenuExtender(TSharedPtr<class SEditorViewport> InRealViewport)
 {
 	TSharedRef<FExtender> Extender(new FExtender());
diff --git a/Engine/Source/Editor/Persona/Private/SAnimViewportToolBar.h b/Engine/Source/Editor/Persona/Private/SAnimViewportToolBar.h
--- a/Engine/Source/Editor/Persona/Private/SAnimViewportToolBar.h
+++ b/Engine/Source/Editor/Persona/Private/SAnimViewportToolBar.h
@@ -162,6 +162,10 @@ private:
 	/** Build the follow bone widget (as a combo button) */
 	TSharedRef<SWidget> MakeFollowBoneComboWidget() const;
 
+public:
+	/** Extend the anim menu */
+	void AddMenuExtender(FName MenuToExtend, FMenuExtensionDelegate MenuBuilderDelegate);
+
 private:
 	/** Build the main toolbar */
 	TSharedRef<SWidget> MakeViewportToolbar(TSharedPtr<class SEditorViewport> InRealViewport);
diff --git a/Engine/Source/Editor/Persona/Private/SAnimationEditorViewport.cpp b/Engine/Source/Editor/Persona/Private/SAnimationEditorViewport.cpp
--- a/Engine/Source/Editor/Persona/Private/SAnimationEditorViewport.cpp
+++ b/Engine/Source/Editor/Persona/Private/SAnimationEditorViewport.cpp
@@ -411,6 +411,11 @@ TWeakPtr<SWidget> SAnimationEditorViewportTabBody::AddNotification(TAttribute<EM
 	return ContainingWidget;
 }
 
+void SAnimationEditorViewportTabBody::AddToolbarExtender(FName MenuToExtend, FMenuExtensionDelegate MenuBuilderDelegate)
+{
+	return ViewportWidget->ViewportToolbar->AddMenuExtender(MenuToExtend, MenuBuilderDelegate);
+}
+
 void SAnimationEditorViewportTabBody::RemoveNotification(const TWeakPtr<SWidget>& InContainingWidget)
 {
 	if(InContainingWidget.IsValid())
@@ -435,10 +440,15 @@ FReply SAnimationEditorViewportTabBody::OnKeyDown(const FGeometry& MyGeometry, c
 	{
 		return FReply::Handled();
 	}
+
+	if (OnKeyDownDelegate.IsBound())
+	{
+		return OnKeyDownDelegate.Execute(MyGeometry, InKeyEvent);
+	}
+
 	return FReply::Unhandled();
 }
 
-
 void SAnimationEditorViewportTabBody::Construct(const FArguments& InArgs, const TSharedRef<class IPersonaPreviewScene>& InPreviewScene, const TSharedRef<class FAssetEditorToolkit>& InAssetEditorToolkit, int32 InViewportIndex)
 {
 	UICommandList = MakeShareable(new FUICommandList_Pinnable);
diff --git a/Engine/Source/Editor/Persona/Private/SAnimationEditorViewport.h b/Engine/Source/Editor/Persona/Private/SAnimationEditorViewport.h
--- a/Engine/Source/Editor/Persona/Private/SAnimationEditorViewport.h
+++ b/Engine/Source/Editor/Persona/Private/SAnimationEditorViewport.h
@@ -156,6 +156,8 @@ protected:
 
 	/** Whether to show options relating to physics */
 	bool bShowPhysicsMenu;
+
+	friend class SAnimationEditorViewportTabBody;
 };
 
 //////////////////////////////////////////////////////////////////////////
@@ -218,6 +220,8 @@ public:
 	virtual FEditorViewportClient& GetViewportClient() const override;
 	virtual TSharedRef<IPinnedCommandList> GetPinnedCommandList() const override;
 	virtual TWeakPtr<SWidget> AddNotification(TAttribute<EMessageSeverity::Type> InSeverity, TAttribute<bool> InCanBeDismissed, const TSharedRef<SWidget>& InNotificationWidget) override;
+	virtual void AddToolbarExtender(FName MenuToExtend, FMenuExtensionDelegate MenuBuilderDelegate) override;
+	virtual FPersonaViewportKeyDownDelegate& GetKeyDownDelegate() override { return OnKeyDownDelegate; }
 	virtual void RemoveNotification(const TWeakPtr<SWidget>& InContainingWidget) override;
 
 	/** SWidget interface */
@@ -483,6 +487,8 @@ private:
 	void OnSetTurnTableMode(int32 ModeIndex);
 	bool IsTurnTableModeSelected(int32 ModeIndex) const;
 
+	FPersonaViewportKeyDownDelegate OnKeyDownDelegate;
+
 public:
 	/** Setup the camera follow mode */
 	void SetCameraFollowMode(EAnimationViewportCameraFollowMode InCameraFollowMode, FName InBoneName);
diff --git a/Engine/Source/Editor/Persona/Public/IPersonaViewport.h b/Engine/Source/Editor/Persona/Public/IPersonaViewport.h
--- a/Engine/Source/Editor/Persona/Public/IPersonaViewport.h
+++ b/Engine/Source/Editor/Persona/Public/IPersonaViewport.h
@@ -4,11 +4,14 @@
 
 #include "CoreMinimal.h"
 #include "Widgets/SCompoundWidget.h"
+#include "Framework/MultiBox/MultiBoxExtender.h"
 #include "Logging/TokenizedMessage.h"
 
 class FEditorViewportClient;
 class IPinnedCommandList;
 
+DECLARE_DELEGATE_RetVal_TwoParams(FReply, FPersonaViewportKeyDownDelegate, const FGeometry&, const FKeyEvent&);
+
 /** Opaque state interface for saving and restoring viewport state */
 struct IPersonaViewportState
 {
@@ -39,6 +42,18 @@ public:
 	 */
 	virtual TWeakPtr<SWidget> AddNotification(TAttribute<EMessageSeverity::Type> InSeverity, TAttribute<bool> InCanBeDismissed, const TSharedRef<SWidget>& InNotificationWidget) = 0;
 
+	/** 
+	 * Adds an extender to the viewport's toolbar
+	 * @param MenuToExtend        The name of the toolbar menu to extend
+	 * @param MenuBuilderDelegate The delegate to use when filling the menu
+	 */
+	virtual void AddToolbarExtender(FName MenuToExtend, FMenuExtensionDelegate MenuBuilderDelegate) = 0;
+
+	/**
+	 * Returns the delegate broadcasted two in case we have an unhandled key
+	 */
+	virtual FPersonaViewportKeyDownDelegate& GetKeyDownDelegate() = 0;
+
 	/** 
 	 * Remove a notification widget 
 	 * @param	InContainingWidget		The containing widget returned from AddNotification()
diff --git a/Engine/Source/Editor/PropertyEditor/Private/DetailLayoutBuilderImpl.cpp b/Engine/Source/Editor/PropertyEditor/Private/DetailLayoutBuilderImpl.cpp
--- a/Engine/Source/Editor/PropertyEditor/Private/DetailLayoutBuilderImpl.cpp
+++ b/Engine/Source/Editor/PropertyEditor/Private/DetailLayoutBuilderImpl.cpp
@@ -152,7 +152,7 @@ IDetailPropertyRow* FDetailLayoutBuilderImpl::EditDefaultProperty(TSharedPtr<IPr
 	return nullptr;
 }
 
-TSharedRef<IPropertyHandle> FDetailLayoutBuilderImpl::GetProperty( const FName PropertyPath, const UClass* ClassOutermost, FName InInstanceName )
+TSharedRef<IPropertyHandle> FDetailLayoutBuilderImpl::GetProperty( const FName PropertyPath, const UStruct* ClassOutermost, FName InInstanceName )
 {	
 	TSharedPtr<FPropertyHandleBase> PropertyHandle; 
 
@@ -184,7 +184,7 @@ void FDetailLayoutBuilderImpl::HideProperty( const TSharedPtr<IPropertyHandle> P
 	}
 }
 
-void FDetailLayoutBuilderImpl::HideProperty( FName PropertyPath, const UClass* ClassOutermost, FName InstanceName )
+void FDetailLayoutBuilderImpl::HideProperty( FName PropertyPath, const UStruct* ClassOutermost, FName InstanceName )
 {
 	TSharedPtr<FPropertyNode> PropertyNode = GetPropertyNode( PropertyPath, ClassOutermost, InstanceName );
 	if( PropertyNode.IsValid() )
@@ -372,7 +372,7 @@ void FDetailLayoutBuilderImpl::SetCurrentCustomizationClass( UStruct* CurrentCla
 	CurrentCustomizationVariableName = VariableName;
 }
 
-TSharedPtr<FPropertyNode> FDetailLayoutBuilderImpl::GetPropertyNode( const FName PropertyName, const UClass* ClassOutermost, FName InstanceName ) const
+TSharedPtr<FPropertyNode> FDetailLayoutBuilderImpl::GetPropertyNode( const FName PropertyName, const UStruct* ClassOutermost, FName InstanceName ) const
 {
 	TSharedPtr<FPropertyNode> PropertyNode = GetPropertyNodeInternal( PropertyName, ClassOutermost, InstanceName );
 	return PropertyNode;
@@ -453,7 +453,7 @@ TSharedPtr<FPropertyNode> FDetailLayoutBuilderImpl::GetPropertyNode( TSharedPtr<
  Example setup
 */
 
-TSharedPtr<FPropertyNode> FDetailLayoutBuilderImpl::GetPropertyNodeInternal( const FName PropertyPath, const UClass* ClassOutermost, FName InstanceName ) const
+TSharedPtr<FPropertyNode> FDetailLayoutBuilderImpl::GetPropertyNodeInternal( const FName PropertyPath, const UStruct* ClassOutermost, FName InstanceName ) const
 {
 	FName PropertyName;
 	TArray<FString> PathList;
diff --git a/Engine/Source/Editor/PropertyEditor/Private/DetailLayoutBuilderImpl.h b/Engine/Source/Editor/PropertyEditor/Private/DetailLayoutBuilderImpl.h
--- a/Engine/Source/Editor/PropertyEditor/Private/DetailLayoutBuilderImpl.h
+++ b/Engine/Source/Editor/PropertyEditor/Private/DetailLayoutBuilderImpl.h
@@ -36,10 +36,10 @@ public:
 	virtual IDetailPropertyRow& AddPropertyToCategory(TSharedPtr<IPropertyHandle> InPropertyHandle) override;
 	virtual FDetailWidgetRow& AddCustomRowToCategory(TSharedPtr<IPropertyHandle> InPropertyHandle, const FText& InCustomSearchString, bool bForAdvanced = false) override;
 	virtual IDetailPropertyRow* EditDefaultProperty(TSharedPtr<IPropertyHandle> InPropertyHandle) override;
-	virtual TSharedRef<IPropertyHandle> GetProperty(const FName PropertyPath, const UClass* ClassOutermost, FName InInstanceName) override;
+	virtual TSharedRef<IPropertyHandle> GetProperty(const FName PropertyPath, const UStruct* ClassOutermost, FName InInstanceName) override;
 	virtual FName GetTopLevelProperty() override;
 	virtual void HideProperty(const TSharedPtr<IPropertyHandle> Property) override;
-	virtual void HideProperty(FName PropertyPath, const UClass* ClassOutermost = NULL, FName InstanceName = NAME_None) override;
+	virtual void HideProperty(FName PropertyPath, const UStruct* ClassOutermost = NULL, FName InstanceName = NAME_None) override;
 	virtual void ForceRefreshDetails() override;
 	virtual TSharedPtr<FAssetThumbnailPool> GetThumbnailPool() const override;
 	virtual bool IsPropertyVisible(TSharedRef<IPropertyHandle> PropertyHandle) const override;
@@ -97,7 +97,7 @@ public:
 	 * @param ClassOutermost	The outer class of the property
 	 * @return The found property node
 	 */
-	TSharedPtr<FPropertyNode> GetPropertyNode(FName PropertyPath, const UClass* ClassOutermost, FName InstanceName) const;
+	TSharedPtr<FPropertyNode> GetPropertyNode(FName PropertyPath, const UStruct* ClassOutermost, FName InstanceName) const;
 
 	/**
 	 * Gets the property node from the provided handle
@@ -220,7 +220,7 @@ private:
 	 * @param bMarkPropertyAsCustomized	Whether or not to mark the property as customized so it does not appear in the default layout
 	 * @return The found property node
 	 */
-	TSharedPtr<FPropertyNode> GetPropertyNodeInternal( FName PropertyPath, const UClass* ClassOutermost, FName InstanceName ) const;
+	TSharedPtr<FPropertyNode> GetPropertyNodeInternal( FName PropertyPath, const UStruct* ClassOutermost, FName InstanceName ) const;
 
 	/**
 	 * Builds a list of simple and advanced categories that should be displayed
diff --git a/Engine/Source/Editor/PropertyEditor/Public/DetailLayoutBuilder.h b/Engine/Source/Editor/PropertyEditor/Public/DetailLayoutBuilder.h
--- a/Engine/Source/Editor/PropertyEditor/Public/DetailLayoutBuilder.h
+++ b/Engine/Source/Editor/PropertyEditor/Public/DetailLayoutBuilder.h
@@ -159,7 +159,7 @@ public:
 		To access StaticArray at index 2 from Struct2 in MyActor, your path would be MyStruct.StaticArray[2]" and your instance name is "Struct2"
 		To access MyFloat in MyActor you can just pass in "MyFloat" because the name of the property is unambiguous
 	 */
-	virtual TSharedRef<IPropertyHandle> GetProperty( const FName PropertyPath, const UClass* ClassOutermost = NULL, FName InstanceName = NAME_None )  = 0;
+	virtual TSharedRef<IPropertyHandle> GetProperty( const FName PropertyPath, const UStruct* ClassOutermost = NULL, FName InstanceName = NAME_None )  = 0;
 
 	/**
 	 * Gets the top level property, for showing the warning for experimental or early access class
@@ -185,7 +185,7 @@ public:
 	 * @param InstanceName				Optional instance name if multiple UProperty's of the same type exist. such as two identical structs, the instance name is one of the struct variable names)
 	 * See IDetailCategoryBuilder::GetProperty for clarification of parameters
 	 */
-	virtual void HideProperty( FName PropertyPath, const UClass* ClassOutermost = NULL, FName InstanceName = NAME_None ) = 0;
+	virtual void HideProperty( FName PropertyPath, const UStruct* ClassOutermost = NULL, FName InstanceName = NAME_None ) = 0;
 
 	/**
 	 * Refreshes the details view and regenerates all the customized layouts
