Duplicating CL#8835415 from //UE4/Dev-Anim

UE-80524 - silence Control Rig Controller Log

diff --git a/Engine/Plugins/Experimental/ControlRig/Source/ControlRigDeveloper/Private/ControlRigController.cpp b/Engine/Plugins/Experimental/ControlRig/Source/ControlRigDeveloper/Private/ControlRigController.cpp
--- a/Engine/Plugins/Experimental/ControlRig/Source/ControlRigDeveloper/Private/ControlRigController.cpp
+++ b/Engine/Plugins/Experimental/ControlRig/Source/ControlRigDeveloper/Private/ControlRigController.cpp
@@ -9,7 +9,7 @@
 
 UControlRigController::UControlRigController()
 	: Model(nullptr)
-	, bSuspendLog(false)
+	, bSuspendLog(true)
 {
 }
 
