Duplicating CL#8882254 from //UE4/Dev-Anim

Control Rig: Silence Control Rig Compiler

diff --git a/Engine/Plugins/Experimental/ControlRig/Source/ControlRigDeveloper/Private/ControlRigBlueprintCompiler.cpp b/Engine/Plugins/Experimental/ControlRig/Source/ControlRigDeveloper/Private/ControlRigBlueprintCompiler.cpp
--- a/Engine/Plugins/Experimental/ControlRig/Source/ControlRigDeveloper/Private/ControlRigBlueprintCompiler.cpp
+++ b/Engine/Plugins/Experimental/ControlRig/Source/ControlRigDeveloper/Private/ControlRigBlueprintCompiler.cpp
@@ -286,21 +286,6 @@ void FControlRigBlueprintCompilerContext::PostCompile()
 					}
 				}
 			}
-
-			{
-				DECLARE_SCOPE_HIERARCHICAL_COUNTER(PrintGraphToConsole)
-
-				int32 SortedPropertyLinkIndex = 1;
-				for (const FControlRigDAG::FNode Node : UnitOrder)
-				{
-					UE_LOG(LogControlRigCompiler, Log, TEXT("%d. %s"), SortedPropertyLinkIndex++, *Node.Name.ToString());
-					for (const FControlRigDAG::FPin& Pin : Node.Outputs)
-					{
-						const int32 Index = Pin.Link;
-						UE_LOG(LogControlRigCompiler, Log, TEXT("%d. %s -> %s"), SortedPropertyLinkIndex++, *PropertyLinks[Index].GetSourcePropertyPath(), *PropertyLinks[Index].GetDestPropertyPath());
-					}
-				}
-			}
 #endif
 
 			{
