UE-80685

Merging

//UE4/Dev-Anim/Engine/Plugins/Experimental/ControlRig/Source/...

to //UE4/Main/Engine/Plugins/Experimental/ControlRig/Source/...

diff --git a/Engine/Plugins/Experimental/ControlRig/Source/ControlRigDeveloper/Public/Graph/ControlRigGraphNode.h b/Engine/Plugins/Experimental/ControlRig/Source/ControlRigDeveloper/Public/Graph/ControlRigGraphNode.h
--- a/Engine/Plugins/Experimental/ControlRig/Source/ControlRigDeveloper/Public/Graph/ControlRigGraphNode.h
+++ b/Engine/Plugins/Experimental/ControlRig/Source/ControlRigDeveloper/Public/Graph/ControlRigGraphNode.h
@@ -122,7 +122,7 @@ class CONTROLRIGDEVELOPER_API UControlRigGraphNode : public UEdGraphNode
 	friend class FControlRigGraphPanelPinFactory;
 	friend class FControlRigEditor;
 	friend struct FControlRigBlueprintUtils;
-	friend class SGraphPinCurveFloat;
+	friend class SControlRigGraphPinCurveFloat;
 
 private:
 	/** The property we represent. For template nodes this represents the struct/property type name. */
diff --git a/Engine/Plugins/Experimental/ControlRig/Source/ControlRigEditor/Private/Editor/SControlRigGizmoNameList.cpp b/Engine/Plugins/Experimental/ControlRig/Source/ControlRigEditor/Private/Editor/SControlRigGizmoNameList.cpp
--- a/Engine/Plugins/Experimental/ControlRig/Source/ControlRigEditor/Private/Editor/SControlRigGizmoNameList.cpp
+++ b/Engine/Plugins/Experimental/ControlRig/Source/ControlRigEditor/Private/Editor/SControlRigGizmoNameList.cpp
@@ -2,7 +2,7 @@
 
 
 #include "Editor/SControlRigGizmoNameList.h"
-#include "Graph/SGraphPinNameListValueWidget.h"
+#include "Graph/SControlRigGraphPinNameListValueWidget.h"
 #include "Widgets/Input/SEditableTextBox.h"
 #include "EdGraphSchema_K2.h"
 #include "ScopedTransaction.h"
@@ -31,7 +31,7 @@ void SControlRigGizmoNameList::Construct(const FArguments& InArgs, FRigControl*
 		.MinDesiredWidth(150)
 		.MaxDesiredWidth(400)
 		[
-			SAssignNew(NameListComboBox, SGraphPinNameListValueWidget)
+			SAssignNew(NameListComboBox, SControlRigGraphPinNameListValueWidget)
 			.OptionsSource(&GetNameList())
 			.OnGenerateWidget(this, &SControlRigGizmoNameList::MakeNameListItemWidget)
 			.OnSelectionChanged(this, &SControlRigGizmoNameList::OnNameListChanged)
diff --git a/Engine/Plugins/Experimental/ControlRig/Source/ControlRigEditor/Private/Editor/SControlRigGizmoNameList.h b/Engine/Plugins/Experimental/ControlRig/Source/ControlRigEditor/Private/Editor/SControlRigGizmoNameList.h
--- a/Engine/Plugins/Experimental/ControlRig/Source/ControlRigEditor/Private/Editor/SControlRigGizmoNameList.h
+++ b/Engine/Plugins/Experimental/ControlRig/Source/ControlRigEditor/Private/Editor/SControlRigGizmoNameList.h
@@ -6,7 +6,7 @@
 #include "Widgets/DeclarativeSyntaxSupport.h"
 #include "Widgets/SWidget.h"
 #include "Widgets/Layout/SBox.h"
-#include "Graph/SGraphPinNameListValueWidget.h"
+#include "Graph/SControlRigGraphPinNameListValueWidget.h"
 #include "Rigs/RigHierarchyContainer.h"
 #include "ControlRigBlueprint.h"
 
@@ -35,7 +35,7 @@ protected:
 	void OnNameListComboBox();
 
 	FOnGetNameListContent OnGetNameListContent;
-	TSharedPtr<SGraphPinNameListValueWidget> NameListComboBox;
+	TSharedPtr<SControlRigGraphPinNameListValueWidget> NameListComboBox;
 	TArray<TSharedPtr<FString>> EmptyList;
 
 	FRigElementKey ControlKey;
diff --git a/Engine/Plugins/Experimental/ControlRig/Source/ControlRigEditor/Private/Graph/ControlRigGraphPanelPinFactory.cpp b/Engine/Plugins/Experimental/ControlRig/Source/ControlRigEditor/Private/Graph/ControlRigGraphPanelPinFactory.cpp
--- a/Engine/Plugins/Experimental/ControlRig/Source/ControlRigEditor/Private/Graph/ControlRigGraphPanelPinFactory.cpp
+++ b/Engine/Plugins/Experimental/ControlRig/Source/ControlRigEditor/Private/Graph/ControlRigGraphPanelPinFactory.cpp
@@ -4,8 +4,8 @@
 #include "Graph/ControlRigGraphSchema.h"
 #include "Graph/ControlRigGraphNode.h"
 #include "Graph/ControlRigGraph.h"
-#include "Graph/SGraphPinNameList.h"
-#include "Graph/SGraphPinCurveFloat.h"
+#include "Graph/SControlRigGraphPinNameList.h"
+#include "Graph/SControlRigGraphPinCurveFloat.h"
 #include "KismetPins/SGraphPinExec.h"
 #include "ControlRig.h"
 #include "NodeFactory.h"
@@ -43,22 +43,22 @@ TSharedPtr<SGraphPin> FControlRigGraphPanelPinFactory::CreatePin(UEdGraphPin* In
 						{
 							if (Property->HasMetaData(UControlRig::BoneNameMetaName))
 							{
-								return SNew(SGraphPinNameList, InPin)
+								return SNew(SControlRigGraphPinNameList, InPin)
 									.OnGetNameListContent_UObject(RigGraph, &UControlRigGraph::GetBoneNameList);
 							}
 							if (Property->HasMetaData(UControlRig::ControlNameMetaName))
 							{
-								return SNew(SGraphPinNameList, InPin)
+								return SNew(SControlRigGraphPinNameList, InPin)
 									.OnGetNameListContent_UObject(RigGraph, &UControlRigGraph::GetControlNameList);
 							}
 							if (Property->HasMetaData(UControlRig::SpaceNameMetaName))
 							{
-								return SNew(SGraphPinNameList, InPin)
+								return SNew(SControlRigGraphPinNameList, InPin)
 									.OnGetNameListContent_UObject(RigGraph, &UControlRigGraph::GetSpaceNameList);
 							}
 							else if (Property->HasMetaData(UControlRig::CurveNameMetaName))
 							{
-								return SNew(SGraphPinNameList, InPin)
+								return SNew(SControlRigGraphPinNameList, InPin)
 									.OnGetNameListContent_UObject(RigGraph, &UControlRigGraph::GetCurveNameList);
 							}
 						}
@@ -74,7 +74,7 @@ TSharedPtr<SGraphPin> FControlRigGraphPanelPinFactory::CreatePin(UEdGraphPin* In
 				}
 				else if (InPin->PinType.PinSubCategoryObject == FRuntimeFloatCurve::StaticStruct())
 				{
-					return SNew(SGraphPinCurveFloat, InPin);
+					return SNew(SControlRigGraphPinCurveFloat, InPin);
 				}
 			}
 		}
diff --git a/Engine/Plugins/Experimental/ControlRig/Source/ControlRigEditor/Private/Graph/SControlRigGraphPinCurveFloat.cpp b/Engine/Plugins/Experimental/ControlRig/Source/ControlRigEditor/Private/Graph/SControlRigGraphPinCurveFloat.cpp
new file mode 100644
--- /dev/null
+++ b/Engine/Plugins/Experimental/ControlRig/Source/ControlRigEditor/Private/Graph/SControlRigGraphPinCurveFloat.cpp
@@ -0,0 +1,128 @@
+// Copyright 1998-2019 Epic Games, Inc. All Rights Reserved.
+
+
+#include "Graph/SControlRigGraphPinCurveFloat.h"
+#include "Graph/ControlRigGraph.h"
+#include "Widgets/Layout/SBox.h"
+#include "Widgets/Input/SEditableTextBox.h"
+#include "EdGraphSchema_K2.h"
+#include "ScopedTransaction.h"
+#include "DetailLayoutBuilder.h"
+#include "PropertyPathHelpers.h"
+
+#include "IControlRigEditorModule.h"
+
+void SControlRigGraphPinCurveFloat::Construct(const FArguments& InArgs, UEdGraphPin* InGraphPinObj)
+{
+	SGraphPin::Construct(SGraphPin::FArguments(), InGraphPinObj);
+}
+
+TSharedRef<SWidget>	SControlRigGraphPinCurveFloat::GetDefaultValueWidget()
+{
+	UControlRigGraph* RigGraph = Cast<UControlRigGraph>(GraphPinObj->GetOwningNode()->GetGraph());
+
+	TSharedRef<SWidget> Widget = SNew(SBox)
+		.MinDesiredWidth(200)
+		.MaxDesiredWidth(400)
+		.MinDesiredHeight(175)
+		.MaxDesiredHeight(300)
+		[
+			SAssignNew(CurveEditor, SCurveEditor)
+			.Visibility(this, &SGraphPin::GetDefaultValueVisibility)
+			.ViewMinInput(0.f)
+			.ViewMaxInput(1.f)
+			.ViewMinOutput(0.f)
+			.ViewMaxOutput(1.f)
+			.TimelineLength(1.f)
+			.DesiredSize(FVector2D(300, 200))
+			.HideUI(true)
+		];
+
+	CurveEditor->SetCurveOwner(this);
+
+	return Widget;
+}
+
+TArray<FRichCurveEditInfoConst> SControlRigGraphPinCurveFloat::GetCurves() const
+{
+	TArray<FRichCurveEditInfoConst> Curves;
+	Curves.Add(Curve.GetRichCurveConst());
+	return Curves;
+}
+
+TArray<FRichCurveEditInfo> SControlRigGraphPinCurveFloat::GetCurves()
+{
+	TArray<FRichCurveEditInfo> Curves;
+	Curves.Add(UpdateAndGetCurve().GetRichCurve());
+	return Curves;
+}
+
+FRuntimeFloatCurve& SControlRigGraphPinCurveFloat::UpdateAndGetCurve()
+{
+	if (UEdGraphPin* Pin = GetPinObj())
+	{
+		FRuntimeFloatCurve::StaticStruct()->ImportText(*Pin->DefaultValue, &Curve, nullptr, EPropertyPortFlags::PPF_None, nullptr, FRuntimeFloatCurve::StaticStruct()->GetName(), true);
+	}
+	return Curve;
+}
+
+
+void SControlRigGraphPinCurveFloat::ModifyOwner()
+{
+	if (UEdGraphPin* Pin = GetPinObj())
+	{
+		Pin->Modify();
+	}
+}
+
+TArray<const UObject*> SControlRigGraphPinCurveFloat::GetOwners() const
+{
+	TArray<const UObject*> Owners;
+	if (UEdGraphPin* Pin = GetPinObj())
+	{
+		if (UEdGraphNode* Node = Pin->GetOwningNode())
+		{
+			Owners.Add(Node);
+		}
+	}
+	return Owners;
+}
+
+void SControlRigGraphPinCurveFloat::MakeTransactional()
+{
+}
+
+bool SControlRigGraphPinCurveFloat::IsValidCurve(FRichCurveEditInfo CurveInfo)
+{
+	if (UEdGraphPin* Pin = GetPinObj())
+	{
+		if (UControlRigGraphNode* Node = Cast<UControlRigGraphNode>(Pin->GetOwningNode()))
+		{
+			if (const UStructProperty* StructProperty = Node->GetUnitProperty())
+			{
+				FString NodeName, PropertyName;
+				Pin->PinName.ToString().Split(TEXT("."), &NodeName, &PropertyName);
+
+				if (UProperty* CurveProperty = StructProperty->Struct->FindPropertyByName(*PropertyName))
+				{
+					return true;
+				}
+			}
+		}
+	}
+	return false;
+}
+
+void SControlRigGraphPinCurveFloat::OnCurveChanged(const TArray<FRichCurveEditInfo>& ChangedCurveEditInfos)
+{
+	if (UEdGraphPin* Pin = GetPinObj())
+	{
+		if (UControlRigGraphNode* Node = Cast<UControlRigGraphNode>(Pin->GetOwningNode()))
+		{
+			Pin->DefaultValue.Empty();
+			FRuntimeFloatCurve::StaticStruct()->ExportText(Pin->DefaultValue, &Curve, nullptr, nullptr, EPropertyPortFlags::PPF_None, nullptr, true);
+			Node->CopyPinDefaultsToModel(Pin);
+		}
+	}
+	ModifyOwner();
+}
diff --git a/Engine/Plugins/Experimental/ControlRig/Source/ControlRigEditor/Private/Graph/SControlRigGraphPinCurveFloat.h b/Engine/Plugins/Experimental/ControlRig/Source/ControlRigEditor/Private/Graph/SControlRigGraphPinCurveFloat.h
new file mode 100644
--- /dev/null
+++ b/Engine/Plugins/Experimental/ControlRig/Source/ControlRigEditor/Private/Graph/SControlRigGraphPinCurveFloat.h
@@ -0,0 +1,39 @@
+// Copyright 1998-2019 Epic Games, Inc. All Rights Reserved.
+
+#pragma once
+
+#include "CoreMinimal.h"
+#include "Widgets/DeclarativeSyntaxSupport.h"
+#include "Widgets/SWidget.h"
+#include "SCurveEditor.h"
+#include "Curves/CurveFloat.h"
+#include "SGraphPin.h"
+
+class SControlRigGraphPinCurveFloat : public SGraphPin, public FCurveOwnerInterface
+{
+public:
+	SLATE_BEGIN_ARGS(SControlRigGraphPinCurveFloat) {}
+	SLATE_END_ARGS()
+
+	void Construct(const FArguments& InArgs, UEdGraphPin* InGraphPinObj);
+
+	/** FCurveOwnerInterface interface */
+	virtual TArray<FRichCurveEditInfoConst> GetCurves() const override;
+	virtual TArray<FRichCurveEditInfo> GetCurves() override;
+	virtual void ModifyOwner() override;
+	virtual TArray<const UObject*> GetOwners() const override;
+	virtual void MakeTransactional() override;
+	virtual void OnCurveChanged(const TArray<FRichCurveEditInfo>& ChangedCurveEditInfos) override;
+	virtual bool IsValidCurve(FRichCurveEditInfo CurveInfo) override;
+
+protected:
+
+	//~ Begin SGraphPin Interface
+	virtual TSharedRef<SWidget>	GetDefaultValueWidget() override;
+	//~ End SGraphPin Interface
+
+	FRuntimeFloatCurve& UpdateAndGetCurve();
+
+	TSharedPtr<SCurveEditor> CurveEditor;
+	FRuntimeFloatCurve Curve;
+};
diff --git a/Engine/Plugins/Experimental/ControlRig/Source/ControlRigEditor/Private/Graph/SControlRigGraphPinNameList.cpp b/Engine/Plugins/Experimental/ControlRig/Source/ControlRigEditor/Private/Graph/SControlRigGraphPinNameList.cpp
new file mode 100644
--- /dev/null
+++ b/Engine/Plugins/Experimental/ControlRig/Source/ControlRigEditor/Private/Graph/SControlRigGraphPinNameList.cpp
@@ -0,0 +1,98 @@
+// Copyright 1998-2019 Epic Games, Inc. All Rights Reserved.
+
+
+#include "Graph/SControlRigGraphPinNameList.h"
+#include "Graph/SControlRigGraphPinNameListValueWidget.h"
+#include "Widgets/Layout/SBox.h"
+#include "Widgets/Input/SEditableTextBox.h"
+#include "EdGraphSchema_K2.h"
+#include "ScopedTransaction.h"
+#include "DetailLayoutBuilder.h"
+
+void SControlRigGraphPinNameList::Construct(const FArguments& InArgs, UEdGraphPin* InGraphPinObj)
+{
+	this->OnGetNameListContent = InArgs._OnGetNameListContent;
+
+	SGraphPin::Construct(SGraphPin::FArguments(), InGraphPinObj);
+}
+
+TSharedRef<SWidget>	SControlRigGraphPinNameList::GetDefaultValueWidget()
+{
+	TSharedPtr<FString> InitialSelected;
+	for (TSharedPtr<FString> Item : GetNameList())
+	{
+		if (Item->Equals(GetNameListText().ToString()))
+		{
+			InitialSelected = Item;
+		}
+	}
+
+	return SNew(SBox)
+		.MinDesiredWidth(150)
+		.MaxDesiredWidth(400)
+		[
+			SAssignNew(NameListComboBox, SControlRigGraphPinNameListValueWidget)
+				.Visibility(this, &SGraphPin::GetDefaultValueVisibility)
+				.OptionsSource(&GetNameList())
+				.OnGenerateWidget(this, &SControlRigGraphPinNameList::MakeNameListItemWidget)
+				.OnSelectionChanged(this, &SControlRigGraphPinNameList::OnNameListChanged)
+				.OnComboBoxOpening(this, &SControlRigGraphPinNameList::OnNameListComboBox)
+				.InitiallySelectedItem(InitialSelected)
+				.Content()
+				[
+					SNew(STextBlock)
+					.Text(this, &SControlRigGraphPinNameList::GetNameListText)
+				]
+		];
+}
+
+const TArray<TSharedPtr<FString>>& SControlRigGraphPinNameList::GetNameList() const
+{
+	if (OnGetNameListContent.IsBound())
+	{
+		return OnGetNameListContent.Execute();
+	}
+	return EmptyList;
+}
+
+FText SControlRigGraphPinNameList::GetNameListText() const
+{
+	return FText::FromString( GraphPinObj->GetDefaultAsString() );
+}
+
+void SControlRigGraphPinNameList::SetNameListText(const FText& NewTypeInValue, ETextCommit::Type /*CommitInfo*/)
+{
+	if(!GraphPinObj->GetDefaultAsString().Equals(NewTypeInValue.ToString()))
+	{
+		const FScopedTransaction Transaction( NSLOCTEXT("GraphEditor", "ChangeNameListPinValue", "Change Bone Name Pin Value" ) );
+		GraphPinObj->Modify();
+		GraphPinObj->GetSchema()->TrySetDefaultValue(*GraphPinObj, NewTypeInValue.ToString());
+	}
+}
+
+TSharedRef<SWidget> SControlRigGraphPinNameList::MakeNameListItemWidget(TSharedPtr<FString> InItem)
+{
+	return 	SNew(STextBlock).Text(FText::FromString(*InItem));// .Font(FEditorStyle::GetFontStyle(TEXT("PropertyWindow.NormalFont")));
+}
+
+void SControlRigGraphPinNameList::OnNameListChanged(TSharedPtr<FString> NewSelection, ESelectInfo::Type SelectInfo)
+{
+	if (SelectInfo != ESelectInfo::Direct)
+	{
+		FString NewValue = *NewSelection.Get();
+		SetNameListText(FText::FromString(NewValue), ETextCommit::OnEnter);
+	}
+}
+
+void SControlRigGraphPinNameList::OnNameListComboBox()
+{
+	TSharedPtr<FString> CurrentlySelected;
+	for (TSharedPtr<FString> Item : GetNameList())
+	{
+		if (Item->Equals(GetNameListText().ToString()))
+		{
+			CurrentlySelected = Item;
+		}
+	}
+	NameListComboBox->SetSelectedItem(CurrentlySelected);
+}
diff --git a/Engine/Plugins/Experimental/ControlRig/Source/ControlRigEditor/Private/Graph/SControlRigGraphPinNameList.h b/Engine/Plugins/Experimental/ControlRig/Source/ControlRigEditor/Private/Graph/SControlRigGraphPinNameList.h
new file mode 100644
--- /dev/null
+++ b/Engine/Plugins/Experimental/ControlRig/Source/ControlRigEditor/Private/Graph/SControlRigGraphPinNameList.h
@@ -0,0 +1,42 @@
+// Copyright 1998-2019 Epic Games, Inc. All Rights Reserved.
+
+#pragma once
+
+#include "CoreMinimal.h"
+#include "Widgets/DeclarativeSyntaxSupport.h"
+#include "Widgets/SWidget.h"
+#include "SGraphPin.h"
+#include "SControlRigGraphPinNameListValueWidget.h"
+
+class SControlRigGraphPinNameList : public SGraphPin
+{
+public:
+
+	DECLARE_DELEGATE_RetVal( const TArray<TSharedPtr<FString>>&, FOnGetNameListContent );
+
+	SLATE_BEGIN_ARGS(SControlRigGraphPinNameList){}
+
+		SLATE_EVENT(FOnGetNameListContent, OnGetNameListContent)
+
+	SLATE_END_ARGS()
+
+	void Construct(const FArguments& InArgs, UEdGraphPin* InGraphPinObj);
+
+protected:
+	//~ Begin SGraphPin Interface
+	virtual TSharedRef<SWidget>	GetDefaultValueWidget() override;
+	//~ End SGraphPin Interface
+
+	const TArray<TSharedPtr<FString>>& GetNameList() const;
+	FText GetNameListText() const;
+	virtual void SetNameListText(const FText& NewTypeInValue, ETextCommit::Type CommitInfo);
+
+	TSharedRef<SWidget> MakeNameListItemWidget(TSharedPtr<FString> InItem);
+	void OnNameListChanged(TSharedPtr<FString> NewSelection, ESelectInfo::Type SelectInfo);
+	void OnNameListComboBox();
+
+	FOnGetNameListContent OnGetNameListContent;
+	TSharedPtr<SControlRigGraphPinNameListValueWidget> NameListComboBox;
+	TArray<TSharedPtr<FString>> EmptyList;
+
+};
diff --git a/Engine/Plugins/Experimental/ControlRig/Source/ControlRigEditor/Private/Graph/SControlRigGraphPinNameListValueWidget.cpp b/Engine/Plugins/Experimental/ControlRig/Source/ControlRigEditor/Private/Graph/SControlRigGraphPinNameListValueWidget.cpp
new file mode 100644
--- /dev/null
+++ b/Engine/Plugins/Experimental/ControlRig/Source/ControlRigEditor/Private/Graph/SControlRigGraphPinNameListValueWidget.cpp
@@ -0,0 +1,254 @@
+// Copyright 1998-2019 Epic Games, Inc. All Rights Reservekd.
+
+#include "SControlRigGraphPinNameListValueWidget.h"
+#include "DetailLayoutBuilder.h"
+#include "Framework/Application/SlateUser.h"
+
+#define LOCTEXT_NAMESPACE "GraphPinNameListValueWidget"
+
+void SControlRigGraphPinNameListValueWidget::Construct(const FArguments& InArgs)
+{
+	this->OnComboBoxOpening = InArgs._OnComboBoxOpening;
+	this->OnSelectionChanged = InArgs._OnSelectionChanged;
+	this->OnGenerateWidget = InArgs._OnGenerateWidget;
+
+	OptionsSource = InArgs._OptionsSource;
+	CustomScrollbar = InArgs._CustomScrollbar;
+
+	TSharedRef<SWidget> ComboBoxMenuContent =
+		SNew(SBox)
+		.MaxDesiredHeight(InArgs._MaxListHeight)
+		[
+			SNew(SVerticalBox)
+
+			+ SVerticalBox::Slot()
+			.AutoHeight()
+			[
+				SAssignNew(this->SearchField, SEditableTextBox)
+				.HintText(LOCTEXT("Search", "Search"))
+				.OnTextChanged(this, &SControlRigGraphPinNameListValueWidget::OnSearchTextChanged)
+			]
+
+			+ SVerticalBox::Slot()
+			[
+				SAssignNew(this->ComboListView, SComboListType)
+				.ListItemsSource(OptionsSource)
+				.OnGenerateRow(this, &SControlRigGraphPinNameListValueWidget::GenerateMenuItemRow)
+				.OnSelectionChanged(this, &SControlRigGraphPinNameListValueWidget::OnSelectionChanged_Internal)
+				.SelectionMode(ESelectionMode::Single)
+				.ExternalScrollbar(InArgs._CustomScrollbar)
+			]
+		];
+
+	// Set up content
+	TSharedPtr<SWidget> ButtonContent = InArgs._Content.Widget;
+	if (InArgs._Content.Widget == SNullWidget::NullWidget)
+	{
+		SAssignNew(ButtonContent, STextBlock)
+			.Text(NSLOCTEXT("SControlRigGraphPinNameListValueWidget", "ContentWarning", "No Content Provided"))
+			.ColorAndOpacity(FLinearColor::Red);
+	}
+
+
+	SComboButton::Construct(SComboButton::FArguments()
+		.Method(InArgs._Method)
+		.ButtonContent()
+		[
+			ButtonContent.ToSharedRef()
+		]
+		.MenuContent()
+		[
+			ComboBoxMenuContent
+		]
+		.HasDownArrow(InArgs._HasDownArrow)
+		.ContentPadding(InArgs._ContentPadding)
+		.OnMenuOpenChanged(this, &SControlRigGraphPinNameListValueWidget::OnMenuOpenChanged)
+		.IsFocusable(true)
+		);
+	SetMenuContentWidgetToFocus(ComboListView);
+
+	// Need to establish the selected item at point of construction so its available for querying
+	// NB: If you need a selection to fire use SetItemSelection rather than setting an IntiallySelectedItem
+	SelectedItem = InArgs._InitiallySelectedItem;
+	if (TListTypeTraits<TSharedPtr<FString>>::IsPtrValid(SelectedItem))
+	{
+		ComboListView->Private_SetItemSelection(SelectedItem, true);
+	}
+
+}
+
+void SControlRigGraphPinNameListValueWidget::ClearSelection()
+{
+	ComboListView->ClearSelection();
+}
+
+void SControlRigGraphPinNameListValueWidget::SetSelectedItem(TSharedPtr<FString> InSelectedItem)
+{
+	if (TListTypeTraits<TSharedPtr<FString>>::IsPtrValid(InSelectedItem))
+	{
+		ComboListView->SetSelection(InSelectedItem);
+	}
+	else
+	{
+		ComboListView->ClearSelection();
+	}
+}
+
+TSharedPtr<FString> SControlRigGraphPinNameListValueWidget::GetSelectedItem()
+{
+	return SelectedItem;
+}
+
+void SControlRigGraphPinNameListValueWidget::RefreshOptions()
+{
+	if (!ComboListView->IsPendingRefresh())
+	{
+		ComboListView->RequestListRefresh();
+	}
+}
+
+TSharedRef<ITableRow> SControlRigGraphPinNameListValueWidget::GenerateMenuItemRow(TSharedPtr<FString> InItem, const TSharedRef<STableViewBase>& OwnerTable)
+{
+	if (OnGenerateWidget.IsBound())
+	{
+		FString SearchToken = SearchField->GetText().ToString().ToLower();
+		EVisibility WidgetVisibility = EVisibility::Visible;
+		if (!SearchToken.IsEmpty())
+		{
+			FString SearchTokenUnderscores = SearchToken.Replace(TEXT(" "), TEXT("_"));
+			if ((InItem->ToLower().Find(SearchToken) < 0) &&
+				(InItem->ToLower().Find(SearchTokenUnderscores) < 0))
+			{
+				WidgetVisibility = EVisibility::Collapsed;
+			}
+		}
+		return SNew(SComboRow<TSharedPtr<FString>>, OwnerTable)
+			.Visibility(WidgetVisibility)
+			[
+				OnGenerateWidget.Execute(InItem)
+			];
+	}
+	else
+	{
+		return SNew(SComboRow<TSharedPtr<FString>>, OwnerTable)
+			[
+				SNew(STextBlock).Text(NSLOCTEXT("SlateCore", "ComboBoxMissingOnGenerateWidgetMethod", "Please provide a .OnGenerateWidget() handler."))
+			];
+
+	}
+}
+
+void SControlRigGraphPinNameListValueWidget::OnMenuOpenChanged(bool bOpen)
+{
+	if (bOpen == false)
+	{
+		if (TListTypeTraits<TSharedPtr<FString>>::IsPtrValid(SelectedItem))
+		{
+			// Ensure the ListView selection is set back to the last committed selection
+			ComboListView->SetSelection(SelectedItem, ESelectInfo::OnNavigation);
+			ComboListView->RequestScrollIntoView(SelectedItem, 0);
+		}
+
+		// Set focus back to ComboBox for users focusing the ListView that just closed
+		TSharedRef<SWidget> ThisRef = AsShared();
+		FSlateApplication::Get().ForEachUser([&ThisRef](FSlateUser& User) {
+			if (User.HasFocusedDescendants(ThisRef))
+			{
+				User.SetFocus(ThisRef);
+			}
+		});
+
+	}
+	else
+	{
+		RegisterActiveTimer(0.f, FWidgetActiveTimerDelegate::CreateSP(this, &SControlRigGraphPinNameListValueWidget::SetFocusPostConstruct));
+	}
+}
+
+EActiveTimerReturnType SControlRigGraphPinNameListValueWidget::SetFocusPostConstruct(double InCurrentTime, float InDeltaTime)
+{
+	if (SearchField.IsValid())
+	{
+		bool bSucceeded = false;
+		FSlateApplication::Get().ForEachUser([&](FSlateUser& User) {
+			bSucceeded |= User.SetFocus(SearchField.ToSharedRef());
+		});
+
+		if (bSucceeded)
+		{
+			return EActiveTimerReturnType::Stop;
+		}
+	}
+
+	return EActiveTimerReturnType::Continue;
+}
+
+void SControlRigGraphPinNameListValueWidget::OnSelectionChanged_Internal(TSharedPtr<FString> ProposedSelection, ESelectInfo::Type SelectInfo)
+{
+	// Ensure that the proposed selection is different
+	if (SelectInfo != ESelectInfo::OnNavigation)
+	{
+		// Ensure that the proposed selection is different from selected
+		if (ProposedSelection != SelectedItem)
+		{
+			SelectedItem = ProposedSelection;
+			OnSelectionChanged.ExecuteIfBound(ProposedSelection, SelectInfo);
+		}
+		// close combo even if user reselected item
+		this->SetIsOpen(false);
+	}
+}
+
+void SControlRigGraphPinNameListValueWidget::OnSearchTextChanged(const FText& ChangedText)
+{
+	FString SearchToken = ChangedText.ToString().ToLower();
+	FString SearchTokenUnderscores = SearchToken.Replace(TEXT(" "), TEXT("_"));
+
+	for (int32 i = 0; i < OptionsSource->Num(); i++)
+	{
+		TSharedPtr<ITableRow> Row = ComboListView->WidgetFromItem((*OptionsSource)[i]);
+		if (Row)
+		{
+			if (SearchToken.IsEmpty())
+			{
+				Row->AsWidget()->SetVisibility(EVisibility::Visible);
+			}
+			else if (((*OptionsSource)[i]->ToLower().Find(SearchToken) >= 0) ||
+				((*OptionsSource)[i]->ToLower().Find(SearchTokenUnderscores) >= 0))
+			{
+				Row->AsWidget()->SetVisibility(EVisibility::Visible);
+			}
+			else
+			{
+				Row->AsWidget()->SetVisibility(EVisibility::Collapsed);
+			}
+		}
+	}
+
+	ComboListView->RequestListRefresh();
+
+	SelectedItem = TSharedPtr< FString >();
+}
+
+FReply SControlRigGraphPinNameListValueWidget::OnButtonClicked()
+{
+	// if user clicked to close the combo menu
+	if (this->IsOpen())
+	{
+		// Re-select first selected item, just in case it was selected by navigation previously
+		TArray<TSharedPtr<FString>> SelectedItems = ComboListView->GetSelectedItems();
+		if (SelectedItems.Num() > 0)
+		{
+			OnSelectionChanged_Internal(SelectedItems[0], ESelectInfo::Direct);
+		}
+	}
+	else
+	{
+		SearchField->SetText(FText::GetEmpty());
+		OnComboBoxOpening.ExecuteIfBound();
+	}
+
+	return SComboButton::OnButtonClicked();
+}
+
+#undef LOCTEXT_NAMESPACE
diff --git a/Engine/Plugins/Experimental/ControlRig/Source/ControlRigEditor/Private/Graph/SControlRigGraphPinNameListValueWidget.h b/Engine/Plugins/Experimental/ControlRig/Source/ControlRigEditor/Private/Graph/SControlRigGraphPinNameListValueWidget.h
new file mode 100644
--- /dev/null
+++ b/Engine/Plugins/Experimental/ControlRig/Source/ControlRigEditor/Private/Graph/SControlRigGraphPinNameListValueWidget.h
@@ -0,0 +1,131 @@
+// Copyright 1998-2019 Epic Games, Inc. All Rights Reservekd.
+
+#pragma once
+
+#include "CoreMinimal.h"
+#include "Input/Reply.h"
+#include "Widgets/SWidget.h"
+#include "Widgets/Input/SComboButton.h"
+#include "Widgets/Input/SEditableTextBox.h"
+#include "Widgets/Views/SListView.h"
+#include "Widgets/Input/SComboBox.h"
+
+/**
+ * A searchable text combo box
+ */
+class SControlRigGraphPinNameListValueWidget : public SComboButton
+{
+public:
+	/** Type of list used for showing menu options. */
+	typedef SListView< TSharedPtr<FString> > SComboListType;
+	/** Delegate type used to generate widgets that represent Options */
+	typedef typename TSlateDelegates< TSharedPtr<FString> >::FOnGenerateWidget FOnGenerateWidget;
+	typedef typename TSlateDelegates< TSharedPtr<FString> >::FOnSelectionChanged FOnSelectionChanged;
+
+	SLATE_BEGIN_ARGS(SControlRigGraphPinNameListValueWidget)
+		: _Content()
+		, _ContentPadding(FMargin(3.0, 3.0))
+		, _OptionsSource()
+		, _OnSelectionChanged()
+		, _OnGenerateWidget()
+		, _InitiallySelectedItem(nullptr)
+		, _Method()
+		, _MaxListHeight(450.0f)
+		, _HasDownArrow(true)
+	{}
+
+	/** Slot for this button's content (optional) */
+	SLATE_DEFAULT_SLOT(FArguments, Content)
+
+		SLATE_ATTRIBUTE(FMargin, ContentPadding)
+
+		SLATE_ARGUMENT(const TArray< TSharedPtr<FString> >*, OptionsSource)
+		SLATE_EVENT(FOnSelectionChanged, OnSelectionChanged)
+		SLATE_EVENT(FOnGenerateWidget, OnGenerateWidget)
+
+		/** Called when combo box is opened, before list is actually created */
+		SLATE_EVENT(FOnComboBoxOpening, OnComboBoxOpening)
+
+		/** The custom scrollbar to use in the ListView */
+		SLATE_ARGUMENT(TSharedPtr<SScrollBar>, CustomScrollbar)
+
+		/** The option that should be selected when the combo box is first created */
+		SLATE_ARGUMENT(TSharedPtr<FString>, InitiallySelectedItem)
+
+		SLATE_ARGUMENT(TOptional<EPopupMethod>, Method)
+
+		/** The max height of the combo box menu */
+		SLATE_ARGUMENT(float, MaxListHeight)
+
+		/**
+		 * When false, the down arrow is not generated and it is up to the API consumer
+		 * to make their own visual hint that this is a drop down.
+		 */
+		SLATE_ARGUMENT(bool, HasDownArrow)
+
+	SLATE_END_ARGS()
+
+	/**
+	 * Construct the widget from a declaration
+	 *
+	 * @param InArgs   Declaration from which to construct the combo box
+	 */
+	void Construct(const FArguments& InArgs);
+
+	void ClearSelection();
+
+	void SetSelectedItem(TSharedPtr<FString> InSelectedItem);
+
+	/** @return the item currently selected by the combo box. */
+	TSharedPtr<FString> GetSelectedItem();
+
+	/**
+	 * Requests a list refresh after updating options
+	 * Call SetSelectedItem to update the selected item if required
+	 * @see SetSelectedItem
+	 */
+	void RefreshOptions();
+
+private:
+
+	/** Generate a row for the InItem in the combo box's list (passed in as OwnerTable). Do this by calling the user-specified OnGenerateWidget */
+	TSharedRef<ITableRow> GenerateMenuItemRow(TSharedPtr<FString> InItem, const TSharedRef<STableViewBase>& OwnerTable);
+
+	//** Called if the menu is closed
+	void OnMenuOpenChanged(bool bOpen);
+
+	/** Invoked when the selection in the list changes */
+	void OnSelectionChanged_Internal(TSharedPtr<FString> ProposedSelection, ESelectInfo::Type SelectInfo);
+
+	/** Invoked when the search text changes */
+	void OnSearchTextChanged(const FText& ChangedText);
+
+	/** Handle clicking on the content menu */
+	virtual FReply OnButtonClicked() override;
+
+	/** The item style to use. */
+	const FTableRowStyle* ItemStyle;
+
+private:
+	/** Delegate that is invoked when the selected item in the combo box changes */
+	FOnSelectionChanged OnSelectionChanged;
+	/** The item currently selected in the combo box */
+	TSharedPtr<FString> SelectedItem;
+	/** The search field used for the combox box's contents */
+	TSharedPtr< SEditableTextBox > SearchField;
+	/** The ListView that we pop up; visualized the available options. */
+	TSharedPtr< SComboListType > ComboListView;
+	/** The Scrollbar used in the ListView. */
+	TSharedPtr< SScrollBar > CustomScrollbar;
+	/** Delegate to invoke before the combo box is opening. */
+	FOnComboBoxOpening OnComboBoxOpening;
+	/** Delegate to invoke when we need to visualize an option as a widget. */
+	FOnGenerateWidget OnGenerateWidget;
+
+	const TArray< TSharedPtr<FString> >* OptionsSource;
+
+	/** Used to focus the name box immediately following construction */
+	EActiveTimerReturnType SetFocusPostConstruct(double InCurrentTime, float InDeltaTime);
+};
+
+
diff --git a/Engine/Plugins/Experimental/ControlRig/Source/ControlRigEditor/Private/Graph/SGraphPinCurveFloat.cpp b/Engine/Plugins/Experimental/ControlRig/Source/ControlRigEditor/Private/Graph/SGraphPinCurveFloat.cpp
deleted file mode 100644
diff --git a/Engine/Plugins/Experimental/ControlRig/Source/ControlRigEditor/Private/Graph/SGraphPinCurveFloat.h b/Engine/Plugins/Experimental/ControlRig/Source/ControlRigEditor/Private/Graph/SGraphPinCurveFloat.h
deleted file mode 100644
diff --git a/Engine/Plugins/Experimental/ControlRig/Source/ControlRigEditor/Private/Graph/SGraphPinNameList.cpp b/Engine/Plugins/Experimental/ControlRig/Source/ControlRigEditor/Private/Graph/SGraphPinNameList.cpp
deleted file mode 100644
diff --git a/Engine/Plugins/Experimental/ControlRig/Source/ControlRigEditor/Private/Graph/SGraphPinNameList.h b/Engine/Plugins/Experimental/ControlRig/Source/ControlRigEditor/Private/Graph/SGraphPinNameList.h
deleted file mode 100644
diff --git a/Engine/Plugins/Experimental/ControlRig/Source/ControlRigEditor/Private/Graph/SGraphPinNameListValueWidget.cpp b/Engine/Plugins/Experimental/ControlRig/Source/ControlRigEditor/Private/Graph/SGraphPinNameListValueWidget.cpp
deleted file mode 100644
diff --git a/Engine/Plugins/Experimental/ControlRig/Source/ControlRigEditor/Private/Graph/SGraphPinNameListValueWidget.h b/Engine/Plugins/Experimental/ControlRig/Source/ControlRigEditor/Private/Graph/SGraphPinNameListValueWidget.h
deleted file mode 100644
