Copying //UE4/Dev-Anim to Dev-Main (//UE4/Dev-Main)

diff --git a/Engine/Plugins/Experimental/ControlRig/Source/ControlRig/Private/AnimNode_ControlRig.cpp b/Engine/Plugins/Experimental/ControlRig/Source/ControlRig/Private/AnimNode_ControlRig.cpp
--- a/Engine/Plugins/Experimental/ControlRig/Source/ControlRig/Private/AnimNode_ControlRig.cpp
+++ b/Engine/Plugins/Experimental/ControlRig/Source/ControlRig/Private/AnimNode_ControlRig.cpp
@@ -15,7 +15,12 @@
 #endif
 
 FAnimNode_ControlRig::FAnimNode_ControlRig()
-	: ControlRig(nullptr)
+	: FAnimNode_ControlRigBase()
+	, ControlRig(nullptr)
+	, Alpha(1.f)
+	, AlphaInputType(EAnimAlphaInputType::Float)
+	, bAlphaBoolEnabled(true)
+	, AlphaCurveName(NAME_None)
 {
 }
 
@@ -49,6 +54,7 @@ FAnimNode_ControlRig::~FAnimNode_ControlRig()
 	}
 #endif // WITH_EDITOR
 }
+
 void FAnimNode_ControlRig::GatherDebugData(FNodeDebugData& DebugData)
 {
 	DECLARE_SCOPE_HIERARCHICAL_COUNTER_FUNC()
@@ -65,6 +71,28 @@ void FAnimNode_ControlRig::Update_AnyThread(const FAnimationUpdateContext& Conte
 
 	FAnimNode_ControlRigBase::Update_AnyThread(Context);
 	GetEvaluateGraphExposedInputs().Execute(Context);
+
+	// alpha handlers
+	InternalBlendAlpha = 0.f;
+	switch (AlphaInputType)
+	{
+	case EAnimAlphaInputType::Float:
+		InternalBlendAlpha = AlphaScaleBias.ApplyTo(AlphaScaleBiasClamp.ApplyTo(Alpha, Context.GetDeltaTime()));
+		break;
+	case EAnimAlphaInputType::Bool:
+		InternalBlendAlpha = AlphaBoolBlend.ApplyTo(bAlphaBoolEnabled, Context.GetDeltaTime());
+		break;
+	case EAnimAlphaInputType::Curve:
+		if (UAnimInstance* AnimInstance = Cast<UAnimInstance>(Context.AnimInstanceProxy->GetAnimInstanceObject()))
+		{
+			InternalBlendAlpha = AlphaScaleBiasClamp.ApplyTo(AnimInstance->GetCurveValue(AlphaCurveName), Context.GetDeltaTime());
+		}
+		break;
+	};
+
+	// Make sure Alpha is clamped between 0 and 1.
+	InternalBlendAlpha = FMath::Clamp<float>(InternalBlendAlpha, 0.f, 1.f);
+
 	PropagateInputProperties(Context.AnimInstanceProxy->GetAnimInstanceObject());
 	Source.Update(Context);
 }
@@ -76,6 +104,9 @@ void FAnimNode_ControlRig::Initialize_AnyThread(const FAnimationInitializeContex
 	FAnimNode_ControlRigBase::Initialize_AnyThread(Context);
 
 	Source.Initialize(Context);
+
+	AlphaBoolBlend.Reinitialize();
+	AlphaScaleBiasClamp.Reinitialize();
 }
 
 void FAnimNode_ControlRig::CacheBones_AnyThread(const FAnimationCacheBonesContext& Context)
diff --git a/Engine/Plugins/Experimental/ControlRig/Source/ControlRig/Private/AnimNode_ControlRigBase.cpp b/Engine/Plugins/Experimental/ControlRig/Source/ControlRig/Private/AnimNode_ControlRigBase.cpp
--- a/Engine/Plugins/Experimental/ControlRig/Source/ControlRig/Private/AnimNode_ControlRigBase.cpp
+++ b/Engine/Plugins/Experimental/ControlRig/Source/ControlRig/Private/AnimNode_ControlRigBase.cpp
@@ -8,16 +8,19 @@
 #include "AnimationRuntime.h"
 
 FAnimNode_ControlRigBase::FAnimNode_ControlRigBase()
+	: FAnimNode_CustomProperty()
+	, bUpdateInput(true)
+	, bExecute(true)
+	, InternalBlendAlpha (1.f)
 {
-	bUpdateInput = true;
-	bExecute = true;
+
 }
 
 void FAnimNode_ControlRigBase::OnInitializeAnimInstance(const FAnimInstanceProxy* InProxy, const UAnimInstance* InAnimInstance)
 {
 	DECLARE_SCOPE_HIERARCHICAL_COUNTER_FUNC()
 
-	FAnimNode_Base::OnInitializeAnimInstance(InProxy, InAnimInstance);
+	FAnimNode_CustomProperty::OnInitializeAnimInstance(InProxy, InAnimInstance);
 
 	USkeletalMeshComponent* Component = InAnimInstance->GetOwningComponent();
 	UControlRig* ControlRig = GetControlRig();
@@ -37,7 +40,8 @@ void FAnimNode_ControlRigBase::Initialize_AnyThread(const FAnimationInitializeCo
 {
 	DECLARE_SCOPE_HIERARCHICAL_COUNTER_FUNC()
 
-	FAnimNode_Base::Initialize_AnyThread(Context);
+	FAnimNode_CustomProperty::Initialize_AnyThread(Context);
+	Source.Initialize(Context);
 
 	if (UControlRig* ControlRig = GetControlRig())
 	{
@@ -49,14 +53,15 @@ void FAnimNode_ControlRigBase::Initialize_AnyThread(const FAnimationInitializeCo
 
 void FAnimNode_ControlRigBase::GatherDebugData(FNodeDebugData& DebugData)
 {
-
+	Source.GatherDebugData(DebugData.BranchFlow(1.f));
 }
 
 void FAnimNode_ControlRigBase::Update_AnyThread(const FAnimationUpdateContext& Context)
 {
 	DECLARE_SCOPE_HIERARCHICAL_COUNTER_FUNC()
 
-	FAnimNode_Base::Update_AnyThread(Context);
+	FAnimNode_CustomProperty::Update_AnyThread(Context);
+	Source.Update(Context);
 
 	if (bExecute)
 	{
@@ -150,12 +155,55 @@ void FAnimNode_ControlRigBase::Evaluate_AnyThread(FPoseContext& Output)
 {
 	DECLARE_SCOPE_HIERARCHICAL_COUNTER_FUNC()
 
+	FPoseContext SourcePose(Output);
+
+	if (Source.GetLinkNode())
+	{
+		Source.Evaluate(SourcePose);
+	}
+	else
+	{
+		// apply refpose
+		SourcePose.ResetToRefPose();
+	}
+
+	if (FAnimWeight::IsRelevant(InternalBlendAlpha))
+	{
+		if (FAnimWeight::IsFullWeight(InternalBlendAlpha))
+		{
+			ExecuteControlRig(SourcePose);
+			Output = SourcePose;
+		}
+		else 
+		{
+			// this blends additively - by weight
+			FPoseContext ControlRigPose(SourcePose);
+			ControlRigPose = SourcePose;
+			ExecuteControlRig(ControlRigPose);
+
+			FPoseContext AdditivePose(ControlRigPose);
+			AdditivePose = ControlRigPose;
+			FAnimationRuntime::ConvertPoseToAdditive(AdditivePose.Pose, SourcePose.Pose);
+			AdditivePose.Curve.ConvertToAdditive(SourcePose.Curve);
+			Output = SourcePose;
+			FAnimationRuntime::AccumulateAdditivePose(Output.Pose, AdditivePose.Pose, Output.Curve, AdditivePose.Curve, InternalBlendAlpha, AAT_LocalSpaceBase);
+		}
+	}
+	else // if not relevant, skip to run control rig
+		// this may cause issue if we have simulation node in the control rig that accumulates time
+	{
+		Output = SourcePose;
+	}
+}
+
+void FAnimNode_ControlRigBase::ExecuteControlRig(FPoseContext& InOutput)
+{
 	if (UControlRig* ControlRig = GetControlRig())
 	{
 		if (bUpdateInput)
 		{
 			// first update input to the system
-			UpdateInput(ControlRig, Output);
+			UpdateInput(ControlRig, InOutput);
 		}
 
 		if (bExecute)
@@ -165,12 +213,7 @@ void FAnimNode_ControlRigBase::Evaluate_AnyThread(FPoseContext& Output)
 		}
 
 		// now update output
-		UpdateOutput(ControlRig, Output);
-	}
-	else
-	{
-		// apply refpose
-		Output.ResetToRefPose();
+		UpdateOutput(ControlRig, InOutput);
 	}
 }
 
@@ -178,6 +221,9 @@ void FAnimNode_ControlRigBase::CacheBones_AnyThread(const FAnimationCacheBonesCo
 {
 	DECLARE_SCOPE_HIERARCHICAL_COUNTER_FUNC()
 
+	FAnimNode_CustomProperty::CacheBones_AnyThread(Context);
+	Source.CacheBones(Context);
+
 	if (UControlRig* ControlRig = GetControlRig())
 	{
 		// fill up node names
@@ -247,4 +293,5 @@ UClass* FAnimNode_ControlRigBase::GetTargetClass() const
 	}
 
 	return nullptr;
-}
\ No newline at end of file
+}
+
diff --git a/Engine/Plugins/Experimental/ControlRig/Source/ControlRig/Private/AnimNode_ControlRig_ExternalSource.cpp b/Engine/Plugins/Experimental/ControlRig/Source/ControlRig/Private/AnimNode_ControlRig_ExternalSource.cpp
--- a/Engine/Plugins/Experimental/ControlRig/Source/ControlRig/Private/AnimNode_ControlRig_ExternalSource.cpp
+++ b/Engine/Plugins/Experimental/ControlRig/Source/ControlRig/Private/AnimNode_ControlRig_ExternalSource.cpp
@@ -18,45 +18,3 @@ UControlRig* FAnimNode_ControlRig_ExternalSource::GetControlRig() const
 	return (ControlRig.IsValid()? ControlRig.Get() : nullptr);
 }
 
-void FAnimNode_ControlRig_ExternalSource::GatherDebugData(FNodeDebugData& DebugData)
-{
-	DECLARE_SCOPE_HIERARCHICAL_COUNTER_FUNC()
-
-	Source.GatherDebugData(DebugData.BranchFlow(1.f));
-}
-
-void FAnimNode_ControlRig_ExternalSource::Update_AnyThread(const FAnimationUpdateContext& Context)
-{
-	DECLARE_SCOPE_HIERARCHICAL_COUNTER_FUNC()
-
-	FAnimNode_ControlRigBase::Update_AnyThread(Context);
-	Source.Update(Context);
-}
-
-void FAnimNode_ControlRig_ExternalSource::Initialize_AnyThread(const FAnimationInitializeContext& Context)
-{
-	DECLARE_SCOPE_HIERARCHICAL_COUNTER_FUNC()
-
-	FAnimNode_ControlRigBase::Initialize_AnyThread(Context);
-	Source.Initialize(Context);
-}
-
-void FAnimNode_ControlRig_ExternalSource::CacheBones_AnyThread(const FAnimationCacheBonesContext& Context)
-{
-	DECLARE_SCOPE_HIERARCHICAL_COUNTER_FUNC()
-
-	FAnimNode_ControlRigBase::CacheBones_AnyThread(Context);
-	Source.CacheBones(Context);
-}
-
-void FAnimNode_ControlRig_ExternalSource::Evaluate_AnyThread(FPoseContext& Output)
-{
-	Output.ResetToRefPose();
-
-	if (Source.GetLinkNode())
-	{
-		Source.Evaluate(Output);
-	}
-
-	FAnimNode_ControlRigBase::Evaluate_AnyThread(Output);
-}
\ No newline at end of file
diff --git a/Engine/Plugins/Experimental/ControlRig/Source/ControlRig/Private/ControlRig.cpp b/Engine/Plugins/Experimental/ControlRig/Source/ControlRig/Private/ControlRig.cpp
--- a/Engine/Plugins/Experimental/ControlRig/Source/ControlRig/Private/ControlRig.cpp
+++ b/Engine/Plugins/Experimental/ControlRig/Source/ControlRig/Private/ControlRig.cpp
@@ -211,6 +211,10 @@ void UControlRig::Initialize(bool bInitRigUnits)
 
 	// cache requested inputs
 	// go through find requested inputs
+
+#if WITH_EDITOR
+	Hierarchy.ControlHierarchy.OnControlSelected.AddUObject(this, &UControlRig::HandleOnControlSelected);
+#endif
 }
 
 void UControlRig::InitializeFromCDO()
@@ -366,8 +370,6 @@ void UControlRig::Execute(const EControlRigState InState)
 
 FTransform UControlRig::GetGlobalTransform(const FName& BoneName) const
 {
-	DECLARE_SCOPE_HIERARCHICAL_COUNTER_FUNC()
-
 	int32 Index = Hierarchy.BoneHierarchy.GetIndex(BoneName);
 	if (Index != INDEX_NONE)
 	{
@@ -380,8 +382,6 @@ FTransform UControlRig::GetGlobalTransform(const FName& BoneName) const
 
 void UControlRig::SetGlobalTransform(const FName& BoneName, const FTransform& InTransform, bool bPropagateTransform)
 {
-	DECLARE_SCOPE_HIERARCHICAL_COUNTER_FUNC()
-
 	int32 Index = Hierarchy.BoneHierarchy.GetIndex(BoneName);
 	if (Index != INDEX_NONE)
 	{
@@ -391,38 +391,31 @@ void UControlRig::SetGlobalTransform(const FName& BoneName, const FTransform& In
 
 FTransform UControlRig::GetGlobalTransform(const int32 BoneIndex) const
 {
-	DECLARE_SCOPE_HIERARCHICAL_COUNTER_FUNC()
 	return Hierarchy.BoneHierarchy.GetGlobalTransform(BoneIndex);
 }
 
 void UControlRig::SetGlobalTransform(const int32 BoneIndex, const FTransform& InTransform, bool bPropagateTransform)
 {
-	DECLARE_SCOPE_HIERARCHICAL_COUNTER_FUNC()
 	Hierarchy.BoneHierarchy.SetGlobalTransform(BoneIndex, InTransform, bPropagateTransform);
 }
 
 float UControlRig::GetCurveValue(const FName& CurveName) const
 {
-	DECLARE_SCOPE_HIERARCHICAL_COUNTER_FUNC()
 	return Hierarchy.CurveContainer.GetValue(CurveName);
 }
 
 void UControlRig::SetCurveValue(const FName& CurveName, const float CurveValue)
 {
-	DECLARE_SCOPE_HIERARCHICAL_COUNTER_FUNC()
-
 	Hierarchy.CurveContainer.SetValue(CurveName, CurveValue);
 }
 
 float UControlRig::GetCurveValue(const int32 CurveIndex) const
 {
-	DECLARE_SCOPE_HIERARCHICAL_COUNTER_FUNC()
 	return Hierarchy.CurveContainer.GetValue(CurveIndex);
 }
 
 void UControlRig::SetCurveValue(const int32 CurveIndex, const float CurveValue)
 {
-	DECLARE_SCOPE_HIERARCHICAL_COUNTER_FUNC()
 	Hierarchy.CurveContainer.SetValue(CurveIndex, CurveValue);
 }
 
@@ -567,6 +560,15 @@ void UControlRig::AddReferencedObjects(UObject* InThis, FReferenceCollector& Col
 	}
 #endif // WITH_EDITOR
 }
+#if WITH_EDITOR
+//undo will clear out the transient Operators, need to recreate them
+void UControlRig::PostEditUndo()
+{
+	Super::PostEditUndo();
+	InstantiateOperatorsFromGeneratedClass();
+	ResolvePropertyPaths();
+}
+#endif // WITH_EDITOR
 
 void UControlRig::Serialize(FArchive& Ar)
 {
@@ -755,5 +757,43 @@ UControlRigGizmoLibrary* UControlRig::GetGizmoLibrary() const
 	return nullptr;
 }
 
+#if WITH_EDITOR
+
+void UControlRig::SelectControl(const FName& InControlName, bool bSelect)
+{
+	Hierarchy.ControlHierarchy.Select(InControlName, bSelect);
+}
+
+bool UControlRig::ClearControlSelection()
+{
+	return Hierarchy.ControlHierarchy.ClearSelection();
+}
+
+TArray<FName> UControlRig::CurrentControlSelection() const
+{
+	return Hierarchy.ControlHierarchy.CurrentSelection();
+}
+
+bool UControlRig::IsControlSelected(const FName& InControlName)const
+{
+	return Hierarchy.ControlHierarchy.IsSelected(InControlName);
+}
+
 // END IControlRigManipulatable interface
+
+void UControlRig::HandleOnControlSelected(FRigHierarchyContainer* InContainer, const FRigElementKey& InKey, bool bSelected)
+{
+	if (InKey.Type == ERigElementType::Control)
+	{
+		FRigControl* Control = FindControl(InKey.Name);
+		if (Control)
+		{
+			ControlSelected().Broadcast(this, *Control, bSelected);
+		}
+	}
+}
+#endif
+
 #undef LOCTEXT_NAMESPACE
+
+
diff --git a/Engine/Plugins/Experimental/ControlRig/Source/ControlRig/Private/ControlRigVM.cpp b/Engine/Plugins/Experimental/ControlRig/Source/ControlRig/Private/ControlRigVM.cpp
--- a/Engine/Plugins/Experimental/ControlRig/Source/ControlRig/Private/ControlRigVM.cpp
+++ b/Engine/Plugins/Experimental/ControlRig/Source/ControlRig/Private/ControlRigVM.cpp
@@ -27,7 +27,6 @@ namespace ControlRigVM
 		{
 		case EControlRigOpCode::Copy:
 		{
-			DECLARE_SCOPE_HIERARCHICAL_COUNTER(ControlRigVM::CopyPropertyValue)
 			PropertyPathHelpers::CopyPropertyValueFast(OuterObject, InOperator.CachedPropertyPath2, InOperator.CachedPropertyPath1);
 			return true;
 		}
diff --git a/Engine/Plugins/Experimental/ControlRig/Source/ControlRig/Private/Manipulatable/IControlRigManipulatable.cpp b/Engine/Plugins/Experimental/ControlRig/Source/ControlRig/Private/Manipulatable/IControlRigManipulatable.cpp
--- a/Engine/Plugins/Experimental/ControlRig/Source/ControlRig/Private/Manipulatable/IControlRigManipulatable.cpp
+++ b/Engine/Plugins/Experimental/ControlRig/Source/ControlRig/Private/Manipulatable/IControlRigManipulatable.cpp
@@ -11,6 +11,9 @@ IControlRigManipulatable::~IControlRigManipulatable()
 {
 	OnFilterControl.Clear();
 	OnControlModified.Clear();
+#if WITH_EDITOR
+	OnControlSelected.Clear();
+#endif
 }
 
 bool IControlRigManipulatable::SetControlGlobalTransform(const FName& InControlName, const FTransform& InGlobalTransform)
diff --git a/Engine/Plugins/Experimental/ControlRig/Source/ControlRig/Private/Math/ControlRigMathLibrary.cpp b/Engine/Plugins/Experimental/ControlRig/Source/ControlRig/Private/Math/ControlRigMathLibrary.cpp
--- a/Engine/Plugins/Experimental/ControlRig/Source/ControlRig/Private/Math/ControlRigMathLibrary.cpp
+++ b/Engine/Plugins/Experimental/ControlRig/Source/ControlRig/Private/Math/ControlRigMathLibrary.cpp
@@ -398,7 +398,7 @@ FTransform FControlRigMathLibrary::LerpTransform(const FTransform& A, const FTra
 	return Result;
 }
 
-void FControlRigMathLibrary::SolveBasicTwoBoneIK(FTransform& BoneA, FTransform& BoneB, FTransform& Effector, const FVector& PoleVector, const FVector& PrimaryAxis, const FVector& SecondaryAxis, float BoneALength, float BoneBLength, bool bEnableStretch, float StretchStartRatio, float StretchMaxRatio)
+void FControlRigMathLibrary::SolveBasicTwoBoneIK(FTransform& BoneA, FTransform& BoneB, FTransform& Effector, const FVector& PoleVector, const FVector& PrimaryAxis, const FVector& SecondaryAxis, float SecondaryAxisWeight, float BoneALength, float BoneBLength, bool bEnableStretch, float StretchStartRatio, float StretchMaxRatio)
 {
 	FVector RootPos = BoneA.GetLocation();
 	FVector ElbowPos = BoneB.GetLocation();
@@ -411,6 +411,7 @@ void FControlRigMathLibrary::SolveBasicTwoBoneIK(FTransform& BoneA, FTransform&
 
 	FVector Axis = BoneA.TransformVectorNoScale(PrimaryAxis);
 	FVector Target1 = BoneB.GetLocation() - BoneA.GetLocation();
+
 	if (!Target1.IsNearlyZero() && !Axis.IsNearlyZero())
 	{
 		Target1 = Target1.GetSafeNormal();
@@ -418,14 +419,24 @@ void FControlRigMathLibrary::SolveBasicTwoBoneIK(FTransform& BoneA, FTransform&
 		BoneA.SetRotation((Rotation1 * BoneA.GetRotation()).GetNormalized());
 
 		Axis = BoneA.TransformVectorNoScale(SecondaryAxis);
-		FVector Target2 = PoleVector - BoneA.GetLocation();
-		if (!Target2.IsNearlyZero() && !Axis.IsNearlyZero())
+
+		if (SecondaryAxisWeight > SMALL_NUMBER)
 		{
-			Target2 = Target2 - FVector::DotProduct(Target2, Target1) * Target1;
-			Target2 = Target2.GetSafeNormal();
+			FVector Target2 = BoneB.GetLocation() - (Effector.GetLocation() + BoneA.GetLocation()) * 0.5f;
+			if (!Target2.IsNearlyZero() && !Axis.IsNearlyZero())
+			{
+				Target2 = Target2 - FVector::DotProduct(Target2, Target1) * Target1;
+				Target2 = Target2.GetSafeNormal();
 
-			FQuat Rotation2 = FQuat::FindBetweenNormals(Axis, Target2);
-			BoneA.SetRotation((Rotation2 * BoneA.GetRotation()).GetNormalized());
+				FQuat Rotation2 = FQuat::FindBetweenNormals(Axis, Target2);
+				if (!FMath::IsNearlyEqual(SecondaryAxisWeight, 1.f))
+				{
+					FVector RotationAxis = Rotation2.GetRotationAxis();
+					float RotationAngle = Rotation2.GetAngle();
+					Rotation2 = FQuat(RotationAxis, RotationAngle * FMath::Clamp<float>(SecondaryAxisWeight, 0.f, 1.f));
+				}
+				BoneA.SetRotation((Rotation2 * BoneA.GetRotation()).GetNormalized());
+			}
 		}
 	}
 
@@ -437,15 +448,111 @@ void FControlRigMathLibrary::SolveBasicTwoBoneIK(FTransform& BoneA, FTransform&
 		FQuat Rotation1 = FQuat::FindBetweenNormals(Axis, Target1);
 		BoneB.SetRotation((Rotation1 * BoneB.GetRotation()).GetNormalized());
 
-		Axis = BoneB.TransformVectorNoScale(SecondaryAxis);
-		FVector Target2 = PoleVector - BoneB.GetLocation();
-		if (!Target2.IsNearlyZero() && !Axis.IsNearlyZero())
-		{
-			Target2 = Target2 - FVector::DotProduct(Target2, Target1) * Target1;
-			Target2 = Target2.GetSafeNormal();
+		if (SecondaryAxisWeight > SMALL_NUMBER)
+		{
+			Axis = BoneB.TransformVectorNoScale(SecondaryAxis);
+			FVector Target2 = BoneB.GetLocation() - (Effector.GetLocation() + BoneA.GetLocation()) * 0.5f;
+			if (!Target2.IsNearlyZero() && !Axis.IsNearlyZero())
+			{
+				Target2 = Target2 - FVector::DotProduct(Target2, Target1) * Target1;
+				Target2 = Target2.GetSafeNormal();
+
+				FQuat Rotation2 = FQuat::FindBetweenNormals(Axis, Target2);
+				if (!FMath::IsNearlyEqual(SecondaryAxisWeight, 1.f))
+				{
+					FVector RotationAxis = Rotation2.GetRotationAxis();
+					float RotationAngle = Rotation2.GetAngle();
+					Rotation2 = FQuat(RotationAxis, RotationAngle * FMath::Clamp<float>(SecondaryAxisWeight, 0.f, 1.f));
+				}
+				BoneB.SetRotation((Rotation2 * BoneB.GetRotation()).GetNormalized());
+			}
+		}
+	}
+}
+
+FVector FControlRigMathLibrary::ClampSpatially(const FVector& Value, EAxis::Type Axis, EControlRigClampSpatialMode::Type Type, float Minimum, float Maximum, FTransform Space)
+{
+	FVector Local = Space.InverseTransformPosition(Value);
 
-			FQuat Rotation2 = FQuat::FindBetweenNormals(Axis, Target2);
-			BoneB.SetRotation((Rotation2 * BoneB.GetRotation()).GetNormalized());
+	switch (Type)
+	{
+		case EControlRigClampSpatialMode::Plane:
+		{
+			switch (Axis)
+			{
+				case EAxis::X:
+				{
+					Local.X = FMath::Clamp<float>(Local.X, Minimum, Maximum);
+					break;
+				}
+				case EAxis::Y:
+				{
+					Local.Y = FMath::Clamp<float>(Local.Y, Minimum, Maximum);
+					break;
+				}
+				default:
+				{
+					Local.Z = FMath::Clamp<float>(Local.Z, Minimum, Maximum);
+					break;
+				}
+			}
+			break;
+		}
+		case EControlRigClampSpatialMode::Cylinder:
+		{
+			switch (Axis)
+			{
+				case EAxis::X:
+				{
+					FVector OnPlane = Local * FVector(0.f, 1.f, 1.f);
+					if (!OnPlane.IsNearlyZero())
+					{
+						float Length = OnPlane.Size();
+						OnPlane = OnPlane * FMath::Clamp<float>(Length, Minimum, Maximum) / Length;
+						Local.Y = OnPlane.Y;
+						Local.Z = OnPlane.Z;
+					}
+					break;
+				}
+				case EAxis::Y:
+				{
+					FVector OnPlane = Local * FVector(1.f, 0.f, 1.f);
+					if (!OnPlane.IsNearlyZero())
+					{
+						float Length = OnPlane.Size();
+						OnPlane = OnPlane * FMath::Clamp<float>(Length, Minimum, Maximum) / Length;
+						Local.X = OnPlane.X;
+						Local.Z = OnPlane.Z;
+					}
+					break;
+				}
+				default:
+				{
+					FVector OnPlane = Local * FVector(1.f, 1.f, 0.f);
+					if (!OnPlane.IsNearlyZero())
+					{
+						float Length = OnPlane.Size();
+						OnPlane = OnPlane * FMath::Clamp<float>(Length, Minimum, Maximum) / Length;
+						Local.X = OnPlane.X;
+						Local.Y = OnPlane.Y;
+					}
+					break;
+				}
+			}
+			break;
+		}
+		default:
+		case EControlRigClampSpatialMode::Sphere:
+		{
+			if (!Local.IsNearlyZero())
+			{
+				float Length = Local.Size();
+				Local = Local * FMath::Clamp<float>(Length, Minimum, Maximum) / Length;
+			}
+			break;
 		}
+
 	}
+
+	return Space.TransformPosition(Local);
 }
diff --git a/Engine/Plugins/Experimental/ControlRig/Source/ControlRig/Private/Rigs/RigBoneHierarchy.cpp b/Engine/Plugins/Experimental/ControlRig/Source/ControlRig/Private/Rigs/RigBoneHierarchy.cpp
--- a/Engine/Plugins/Experimental/ControlRig/Source/ControlRig/Private/Rigs/RigBoneHierarchy.cpp
+++ b/Engine/Plugins/Experimental/ControlRig/Source/ControlRig/Private/Rigs/RigBoneHierarchy.cpp
@@ -3,6 +3,7 @@
 #include "Rigs/RigBoneHierarchy.h"
 #include "ControlRig.h"
 #include "HelperUtil.h"
+#include "AnimationRuntime.h"
 
 ////////////////////////////////////////////////////////////////////////////////
 // FRigBoneHierarchy
@@ -49,7 +50,7 @@ FName FRigBoneHierarchy::GetSafeNewName(const FName& InPotentialNewName) const
 	return Name;
 }
 
-FRigBone& FRigBoneHierarchy::Add(const FName& InNewName, const FName& InParentName, const FTransform& InInitTransform)
+FRigBone& FRigBoneHierarchy::Add(const FName& InNewName, const FName& InParentName, ERigBoneType InType, const FTransform& InInitTransform)
 {
 	FRigBone NewBone;
 	NewBone.Name = GetSafeNewName(InNewName);
@@ -57,6 +58,7 @@ FRigBone& FRigBoneHierarchy::Add(const FName& InNewName, const FName& InParentNa
 	NewBone.ParentName = NewBone.ParentIndex == INDEX_NONE ? NAME_None : InParentName;
 	NewBone.InitialTransform = InInitTransform;
 	NewBone.GlobalTransform = InInitTransform;
+	NewBone.Type = InType;
 	RecalculateLocalTransform(NewBone);
 
 	FName NewBoneName = NewBone.Name;
@@ -71,9 +73,9 @@ FRigBone& FRigBoneHierarchy::Add(const FName& InNewName, const FName& InParentNa
 	return Bones[Index];
 }
 
-FRigBone& FRigBoneHierarchy::Add(const FName& InNewName, const FName& InParentName, const FTransform& InInitTransform, const FTransform& InLocalTransform, const FTransform& InGlobalTransform)
+FRigBone& FRigBoneHierarchy::Add(const FName& InNewName, const FName& InParentName, ERigBoneType InType, const FTransform& InInitTransform, const FTransform& InLocalTransform, const FTransform& InGlobalTransform)
 {
-	FRigBone& Bone = Add(InNewName, InParentName, InInitTransform);
+	FRigBone& Bone = Add(InNewName, InParentName, InType, InInitTransform);
 	Bone.LocalTransform = InLocalTransform;
 	Bone.GlobalTransform = InGlobalTransform;
 	return Bone;
@@ -375,10 +377,9 @@ FName FRigBoneHierarchy::Rename(const FName& InOldName, const FName& InNewName)
 
 void FRigBoneHierarchy::Sort()
 {
+	TArray<int32> RootArray;
 	TMap<int32, TArray<int32>> HierarchyTree;
 
-	TArray<int32> SortedArray;
-
 	// first figure out children and find roots
 	for (int32 Index = 0; Index < Bones.Num(); ++Index)
 	{
@@ -392,23 +393,30 @@ void FRigBoneHierarchy::Sort()
 		{
 			// as as a root
 			HierarchyTree.Add(Index);
-			// add them to the list first
-			SortedArray.Add(Index);
+			RootArray.Add(Index);
 		}
 	}
-
-	// now go through map and add to sorted array
-	for (int32 SortIndex = 0; SortIndex < SortedArray.Num(); ++SortIndex)
+	
+	struct Local
 	{
-		// add children of sorted array
-		TArray<int32>* ChildIndices = HierarchyTree.Find(SortedArray[SortIndex]);
-		if (ChildIndices)
+		static void VisitBone(int32 Index, const TMap<int32, TArray<int32>>& InHierarchyTree, TArray<int32>& SortedArray)
 		{
-			// append children now
-			SortedArray.Append(*ChildIndices);
-			// now sorted array will grow
-			// this is same as BFS as it starts from all roots, and going down
+			if (const TArray<int32>* ChildIndices = InHierarchyTree.Find(Index))
+			{
+				SortedArray.Append(*ChildIndices);
+				for (int32 ChildIndex : *ChildIndices)
+				{
+					VisitBone(ChildIndex, InHierarchyTree, SortedArray);
+				}
+			}
 		}
+	};
+
+	TArray<int32> SortedArray;
+	for (int32 RootIndex : RootArray)
+	{
+		SortedArray.Add(RootIndex);
+		Local::VisitBone(RootIndex, HierarchyTree, SortedArray);
 	}
 
 	check(SortedArray.Num() == Bones.Num());
@@ -432,8 +440,21 @@ void FRigBoneHierarchy::Sort()
 	}
 }
 
+void FRigBoneHierarchy::RefreshParentNames()
+{
+	for (FRigBone& Bone : Bones)
+	{
+		Bone.ParentIndex = INDEX_NONE;
+		if (Bone.ParentName != NAME_None)
+		{
+			Bone.ParentIndex = GetIndex(Bone.ParentName);
+		}
+	}
+}
+
 void FRigBoneHierarchy::RefreshMapping()
 {
+	RefreshParentNames();
 	Sort();
 
 	NameToIndexMapping.Empty();
@@ -448,7 +469,6 @@ void FRigBoneHierarchy::Initialize()
 {
 	RefreshMapping();
 
-	// update parent index
 	for (int32 Index = 0; Index < Bones.Num(); ++Index)
 	{
 		Bones[Index].ParentIndex = GetIndex(Bones[Index].ParentName);
@@ -575,4 +595,154 @@ bool FRigBoneHierarchy::IsSelected(const FName& InName) const
 	return Selection.Contains(InName);
 }
 
+TArray<FRigElementKey> FRigBoneHierarchy::ImportSkeleton(const FReferenceSkeleton& InSkeleton, const FName& InNameSpace, bool bReplaceExistingBones, bool bRemoveObsoleteBones, bool bSelectBones)
+{
+	TArray<FRigElementKey> AddedBones;
+	TArray<FName> BoneNamesToSelect;
+	TMap<FName, FName> BoneNameMap;
+
+	const TArray<FMeshBoneInfo>& BoneInfos = InSkeleton.GetRefBoneInfo();
+	const TArray<FTransform>& BonePoses = InSkeleton.GetRefBonePose();
+
+	struct Local
+	{
+		static FName DetermineBoneName(const FName& InBoneName, const FName& InNameSpace)
+		{
+			if (InNameSpace == NAME_None || InBoneName == NAME_None)
+			{
+				return InBoneName;
+			}
+			return *FString::Printf(TEXT("%s_%s"), *InNameSpace.ToString(), *InBoneName.ToString());
+		}
+	};
+
+	if (bReplaceExistingBones)
+	{
+		for (const FRigBone& ExistingBone : Bones)
+		{
+			BoneNameMap.Add(ExistingBone.Name, ExistingBone.Name);
+		}
+
+		for (int32 Index = 0; Index < BoneInfos.Num(); ++Index)
+		{
+			int32 ExistingBoneIndex = GetIndex(BoneInfos[Index].Name);
+			FName DesiredBoneName = Local::DetermineBoneName(BoneInfos[Index].Name, InNameSpace);
+			FName ParentName = (BoneInfos[Index].ParentIndex != INDEX_NONE) ? BoneInfos[BoneInfos[Index].ParentIndex].Name : NAME_None;
+			ParentName = Local::DetermineBoneName(ParentName, InNameSpace);
+
+			const FName* MappedParentNamePtr = BoneNameMap.Find(ParentName);
+			if (MappedParentNamePtr)
+			{
+				ParentName = *MappedParentNamePtr;
+			}
+
+			if (ExistingBoneIndex != INDEX_NONE)
+			{
+				// check it's parent
+				if (ParentName != NAME_None)
+				{
+					Bones[ExistingBoneIndex].ParentName = ParentName;
+				}
+
+				// reset it's initial transform
+				SetInitialTransform(ExistingBoneIndex, FAnimationRuntime::GetComponentSpaceTransform(InSkeleton, BonePoses, Index));
+
+				BoneNamesToSelect.Add(Bones[ExistingBoneIndex].Name);
+			}
+			else
+			{
+				FRigBone& AddedBone = Add(DesiredBoneName, ParentName, ERigBoneType::Imported, FAnimationRuntime::GetComponentSpaceTransform(InSkeleton, BonePoses, Index));
+				BoneNameMap.Add(DesiredBoneName, AddedBone.Name);
+				AddedBones.Add(AddedBone.GetElementKey());
+				BoneNamesToSelect.Add(AddedBone.Name);
+			}
+		}
+
+	}
+	else // import all as new
+	{
+		for (int32 Index = 0; Index < BoneInfos.Num(); ++Index)
+		{
+			FName DesiredBoneName = Local::DetermineBoneName(BoneInfos[Index].Name, InNameSpace);
+			FName ParentName = (BoneInfos[Index].ParentIndex != INDEX_NONE) ? BoneInfos[BoneInfos[Index].ParentIndex].Name : NAME_None;
+			ParentName = Local::DetermineBoneName(ParentName, InNameSpace);
+
+			const FName* MappedParentNamePtr = BoneNameMap.Find(ParentName);
+			if (MappedParentNamePtr)
+			{
+				ParentName = *MappedParentNamePtr;
+			}
+			
+			FRigBone& AddedBone = Add(DesiredBoneName, ParentName, ERigBoneType::Imported, FAnimationRuntime::GetComponentSpaceTransform(InSkeleton, BonePoses, Index));
+			BoneNameMap.Add(DesiredBoneName, AddedBone.Name);
+			AddedBones.Add(AddedBone.GetElementKey());
+			BoneNamesToSelect.Add(AddedBone.Name);
+		}
+	}
+
+	RefreshMapping();
+
+	if (bReplaceExistingBones && bRemoveObsoleteBones)
+	{
+		TMap<FName, int32> BoneNameToIndexInSkeleton;
+		for (const FMeshBoneInfo& BoneInfo : BoneInfos)
+		{
+			FName DesiredBoneName = Local::DetermineBoneName(BoneInfo.Name, InNameSpace);
+			BoneNameToIndexInSkeleton.Add(DesiredBoneName, BoneNameToIndexInSkeleton.Num());
+		}
+		
+		TArray<FName> BonesToDelete;
+		for (int32 ExistingBoneIndex = 0; ExistingBoneIndex < Bones.Num(); ExistingBoneIndex++)
+		{
+			const FRigBone& Bone = Bones[ExistingBoneIndex];
+			if (!BoneNameToIndexInSkeleton.Contains(Bone.Name))
+			{
+				if (Bone.Type == ERigBoneType::Imported)
+				{
+					BonesToDelete.Add(Bone.Name);
+				}
+			}
+		}
+
+		for (const FName& BoneToDelete : BonesToDelete)
+		{
+			const FRigBone& Bone = Bones[GetIndex(BoneToDelete)];
+			TArray<FName> Dependents;
+			for (int32 DependentIndex : Bone.Dependents)
+			{
+				if (Bones[DependentIndex].Type != ERigBoneType::Imported)
+				{
+					Dependents.Add(Bones[DependentIndex].Name);
+				}
+			}
+
+			Algo::Reverse(Dependents);
+			for (const FName& DependentName : Dependents)
+			{
+				Reparent(DependentName, NAME_None);
+			}
+		}
+
+		for (const FName& BoneToDelete : BonesToDelete)
+		{
+			Remove(BoneToDelete);
+			BoneNamesToSelect.Remove(BoneToDelete);
+		}
+	}
+
+	if (bSelectBones)
+	{
+		ClearSelection();
+		for (const FName& BoneNameToSelect : BoneNamesToSelect)
+		{
+			Select(BoneNameToSelect);
+		}
+	}
+
+	Initialize();
+
+	return AddedBones;
+}
+
+
 #endif
\ No newline at end of file
diff --git a/Engine/Plugins/Experimental/ControlRig/Source/ControlRig/Private/Rigs/RigHierarchyContainer.cpp b/Engine/Plugins/Experimental/ControlRig/Source/ControlRig/Private/Rigs/RigHierarchyContainer.cpp
--- a/Engine/Plugins/Experimental/ControlRig/Source/ControlRig/Private/Rigs/RigHierarchyContainer.cpp
+++ b/Engine/Plugins/Experimental/ControlRig/Source/ControlRig/Private/Rigs/RigHierarchyContainer.cpp
@@ -690,7 +690,7 @@ TArray<FRigElementKey> FRigHierarchyContainer::ImportFromText(const FString& InC
 							ParentName = ParentKey->Name;
 						}
 
-						FRigBone& NewElement = BoneHierarchy.Add(Element.Name, ParentName, Element.InitialTransform, Element.LocalTransform, Data.GlobalTransforms[Index]);
+						FRigBone& NewElement = BoneHierarchy.Add(Element.Name, ParentName, ERigBoneType::User, Element.InitialTransform, Element.LocalTransform, Data.GlobalTransforms[Index]);
 						ElementMap.FindOrAdd(Element.GetElementKey()) = NewElement.GetElementKey();
 						PastedKeys.Add(NewElement.GetElementKey());
 						break;
diff --git a/Engine/Plugins/Experimental/ControlRig/Source/ControlRig/Private/Sequencer/MovieSceneControlRigParameterTemplate.cpp b/Engine/Plugins/Experimental/ControlRig/Source/ControlRig/Private/Sequencer/MovieSceneControlRigParameterTemplate.cpp
new file mode 100644
--- /dev/null
+++ b/Engine/Plugins/Experimental/ControlRig/Source/ControlRig/Private/Sequencer/MovieSceneControlRigParameterTemplate.cpp
@@ -0,0 +1,988 @@
+// Copyright 1998-2019 Epic Games, Inc. All Rights Reserved.
+
+#include "Sequencer/MovieSceneControlRigParameterTemplate.h"
+#include "Evaluation/Blending/MovieSceneMultiChannelBlending.h"
+#include "Sequencer/MovieSceneControlRigParameterTrack.h"
+#include "ControlRig.h"
+#include "Evaluation/MovieSceneEvaluation.h"
+#include "IMovieScenePlayer.h"
+#include "Components/SkeletalMeshComponent.h"
+#include "Sequencer/ControlRigSequencerAnimInstance.h"
+#include "IControlRigObjectBinding.h"
+#include "Animation/AnimData/BoneMaskFilter.h"
+#include "ControlRigSkeletalMeshBinding.h"
+#include "Evaluation/Blending/BlendableTokenStack.h"
+#include "Evaluation/Blending/MovieSceneBlendingActuatorID.h"
+
+
+//#include "Particles/ParticleSystemComponent.h"
+
+DECLARE_CYCLE_STAT(TEXT("ControlRig Parameter Track Evaluate"), MovieSceneEval_ControlRigTemplateParameter_Evaluate, STATGROUP_MovieSceneEval);
+DECLARE_CYCLE_STAT(TEXT("ControlRig Parameter Track Token Execute"), MovieSceneEval_ControlRigParameterTrack_TokenExecute, STATGROUP_MovieSceneEval);
+
+template<typename T>
+struct TNameAndValue
+{
+	FName Name;
+	T Value;
+};
+
+
+/**
+ * Structure representing the animated value of a scalar parameter.
+ */
+struct FScalarParameterStringAndValue
+{
+	/** Creates a new FScalarParameterAndValue with a parameter name and a value. */
+	FScalarParameterStringAndValue(FName InParameterName, float InValue)
+	{
+		ParameterName = InParameterName;
+		ParameterString = InParameterName.ToString();		Value = InValue;
+	}
+
+	/** The name of the scalar parameter. */
+	FString ParameterString;
+	FName ParameterName;
+	/** The animated value of the scalar parameter. */
+	float Value;
+};
+
+/**
+ * Structure representing the animated value of a vector parameter.
+ */
+struct FVectorParameterStringAndValue
+{
+	/** Creates a new FVectorParameterAndValue with a parameter name and a value. */
+	FVectorParameterStringAndValue(FName InParameterName, FVector InValue)
+	{
+		ParameterName = InParameterName;
+		ParameterString = InParameterName.ToString();
+		Value = InValue;
+	}
+
+	/** The name of the vector parameter. */
+	FString ParameterString;
+	FName ParameterName;
+
+	/** The animated value of the vector parameter. */
+	FVector Value;
+};
+
+
+/**
+ * Structure representing the animated value of a color parameter.
+ */
+struct FColorParameterStringAndValue
+{
+	/** Creates a new FColorParameterAndValue with a parameter name and a value. */
+	FColorParameterStringAndValue(FName InParameterName, FLinearColor InValue)
+	{
+		ParameterName = InParameterName;
+		ParameterString = InParameterName.ToString();
+		Value = InValue;
+	}
+
+	/** The name of the color parameter. */
+	FString ParameterString;
+	FName ParameterName;
+
+	/** The animated value of the color parameter. */
+	FLinearColor Value;
+};
+
+struct FTransformParameterStringAndValue
+{
+
+	/** The name of the transform  parameter. */
+	FString ParameterString;
+	FName ParameterName;
+	/** Translation component */
+	FVector Translation;
+	/** Rotation component */
+	FRotator Rotation;
+	/** Scale component */
+	FVector Scale;
+
+	FTransformParameterStringAndValue(FName InParameterName, const FVector& InTranslation,
+		const FRotator& InRotation, const FVector& InScale) : Translation(InTranslation),
+		Rotation(InRotation), Scale(InScale)
+	{
+		ParameterName = InParameterName;
+		ParameterString = InParameterName.ToString();
+	}
+};
+
+struct FEvaluatedControlRigParameterSectionValues
+{
+	FEvaluatedControlRigParameterSectionValues() = default;
+
+	FEvaluatedControlRigParameterSectionValues(FEvaluatedControlRigParameterSectionValues&&) = default;
+	FEvaluatedControlRigParameterSectionValues& operator=(FEvaluatedControlRigParameterSectionValues&&) = default;
+
+	// Non-copyable
+	FEvaluatedControlRigParameterSectionValues(const FEvaluatedControlRigParameterSectionValues&) = delete;
+	FEvaluatedControlRigParameterSectionValues& operator=(const FEvaluatedControlRigParameterSectionValues&) = delete;
+
+	/** Array of evaluated scalar values */
+	TArray<FScalarParameterStringAndValue, TInlineAllocator<2>> ScalarValues;
+	/** Array of evaluated vector values */
+	TArray<FVectorParameterStringAndValue, TInlineAllocator<2>> VectorValues;
+	/** Array of evaluated color values */
+	TArray<FColorParameterStringAndValue, TInlineAllocator<2>> ColorValues;
+	/** Array of evaluated transform values */
+	TArray<FTransformParameterStringAndValue, TInlineAllocator<2>> TransformValues;
+};
+
+/** Token for control rig control parameters */
+struct FControlRigTrackTokenFloat
+{
+	FControlRigTrackTokenFloat() {}
+	
+	FControlRigTrackTokenFloat(float InValue)
+		:Value(InValue)
+	{}
+
+	float Value;
+
+};
+
+struct FControlRigTrackTokenVector
+{
+	FControlRigTrackTokenVector() {}
+	FControlRigTrackTokenVector(FVector InValue)
+		:Value(InValue)
+	{}
+
+	FVector Value;
+};
+
+struct FControlRigTrackTokenTransform
+{
+	FControlRigTrackTokenTransform() {}
+	FControlRigTrackTokenTransform(FTransform InValue)
+		: Value(InValue)
+	{}
+	FTransform Value;
+
+};
+
+
+
+// Specify a unique runtime type identifier for rig control track tokens
+template<> FMovieSceneAnimTypeID GetBlendingDataType<FControlRigTrackTokenFloat>()
+{
+	static FMovieSceneAnimTypeID TypeId = FMovieSceneAnimTypeID::Unique();
+	return TypeId;
+}
+
+template<> FMovieSceneAnimTypeID GetBlendingDataType<FControlRigTrackTokenVector>()
+{
+	static FMovieSceneAnimTypeID TypeId = FMovieSceneAnimTypeID::Unique();
+	return TypeId;
+}
+
+template<> FMovieSceneAnimTypeID GetBlendingDataType<FControlRigTrackTokenTransform>()
+{
+	static FMovieSceneAnimTypeID TypeId = FMovieSceneAnimTypeID::Unique();
+	return TypeId;
+}
+
+
+
+/** Define working data types for blending calculations  */
+template<>  struct TBlendableTokenTraits<FControlRigTrackTokenFloat>
+{
+	typedef MovieScene::TMaskedBlendable<float, 1> WorkingDataType;
+};
+
+template<> struct TBlendableTokenTraits<FControlRigTrackTokenVector>
+{
+	typedef MovieScene::TMaskedBlendable<float, 3> WorkingDataType;
+};
+
+template<>  struct TBlendableTokenTraits<FControlRigTrackTokenTransform>
+{
+	typedef MovieScene::TMaskedBlendable<float, 9> WorkingDataType;
+};
+
+
+
+
+namespace MovieScene
+{
+
+	void MultiChannelFromData(const FControlRigTrackTokenFloat& In, TMultiChannelValue<float, 1>& Out)
+	{
+		Out = { In.Value };
+	}
+
+	void ResolveChannelsToData(const TMultiChannelValue<float, 1>& In, FControlRigTrackTokenFloat& Out)
+	{
+		Out.Value = In[0];
+	}
+
+	void MultiChannelFromData(const FControlRigTrackTokenVector& In, TMultiChannelValue<float, 3>& Out)
+	{
+		Out = { In.Value.X, In.Value.Y, In.Value.Z };
+	}
+
+	void ResolveChannelsToData(const TMultiChannelValue<float, 3>& In, FControlRigTrackTokenVector& Out)
+	{
+		Out.Value = FVector(In[0], In[1], In[2]);
+	}
+
+	void MultiChannelFromData(const FControlRigTrackTokenTransform& In, TMultiChannelValue<float, 9>& Out)
+	{
+		FVector Translation = In.Value.GetTranslation();
+		FVector Rotation = In.Value.GetRotation().Rotator().Euler();
+		FVector Scale = In.Value.GetScale3D();
+		Out = { Translation.X, Translation.Y, Translation.Z, Rotation.X, Rotation.Y, Rotation.Z, Scale.X, Scale.Y, Scale.Z };
+
+	}
+
+	void ResolveChannelsToData(const TMultiChannelValue<float, 9>& In, FControlRigTrackTokenTransform& Out)
+	{
+		Out.Value = FTransform(
+			FRotator::MakeFromEuler(FVector(In[3], In[4], In[5])),
+			FVector(In[0], In[1], In[2]),
+			FVector(In[6], In[7], In[8])
+		);
+	}
+}
+
+
+
+struct FControlRigParameterPreAnimatedTokenProducer : IMovieScenePreAnimatedTokenProducer
+{
+	FControlRigParameterPreAnimatedTokenProducer(FMovieSceneSequenceIDRef InSequenceID)
+		: SequenceID(InSequenceID)
+	{}
+
+	virtual IMovieScenePreAnimatedTokenPtr CacheExistingState(UObject& Object) const
+	{
+
+		struct FToken : IMovieScenePreAnimatedToken
+		{
+			FToken(FMovieSceneSequenceIDRef InSequenceID)
+				: SequenceID(InSequenceID)
+			{}
+
+			virtual void RestoreState(UObject& InObject, IMovieScenePlayer& Player) override
+			{
+				if (UControlRig* ControlRig = Cast<UControlRig>(&InObject))
+				{
+					if (ControlRig->GetObjectBinding())
+					{
+						if (USkeletalMeshComponent* SkeletalMeshComponent = Cast<USkeletalMeshComponent>(ControlRig->GetObjectBinding()->GetBoundObject()))
+						{
+							if (UControlRigSequencerAnimInstance* AnimInstance = Cast<UControlRigSequencerAnimInstance>(SkeletalMeshComponent->GetAnimInstance()))
+							{
+								AnimInstance->ResetNodes();
+								AnimInstance->RecalcRequiredBones();
+							}
+							UAnimSequencerInstance::UnbindFromSkeletalMeshComponent(SkeletalMeshComponent);
+						}
+						for (TNameAndValue<float>& Value : ScalarValues)
+						{
+							ControlRig->SetControlValue<float>(Value.Name, Value.Value);
+						}
+
+						for (TNameAndValue<FVector>& Value : VectorValues)
+						{
+							ControlRig->SetControlValue<FVector>(Value.Name, Value.Value);
+						}
+
+						for (TNameAndValue<FTransform>& Value : TransformValues)
+						{
+							ControlRig->SetControlValue<FTransform>(Value.Name, Value.Value);
+						}
+						ControlRig->GetObjectBinding()->UnbindFromObject();
+					}
+				}
+			}
+
+			FMovieSceneSequenceID SequenceID;
+			TArray< TNameAndValue<float> > ScalarValues;
+			TArray< TNameAndValue<FVector> > VectorValues;
+			TArray< TNameAndValue<FTransform> > TransformValues;
+		};
+
+
+		FToken Token(SequenceID);
+
+		if (UControlRig* ControlRig = Cast<UControlRig>(&Object))
+		{
+
+			/*
+			//mz todo the other types
+
+			UENUM()
+			enum class ERigControlType : uint8
+			{
+				Bool,
+				Float,
+				Vector2D,
+				Position,
+				Scale,
+				Quat,
+				Rotator,
+				Transform
+			};
+			*/
+
+			const TArray<FRigControl>& Controls = ControlRig->AvailableControls();
+			FRigControlValue Value;
+			for (const FRigControl& RigControl : Controls)
+			{
+				switch (RigControl.ControlType)
+				{
+				case ERigControlType::Float:
+				{
+					float Val = RigControl.Value.Get<float>();
+					Token.ScalarValues.Add(TNameAndValue<float>{ RigControl.Name, Val });
+					break;
+				}
+				case ERigControlType::Position:
+				case ERigControlType::Scale:
+				case ERigControlType::Rotator:
+				{
+					FVector Val = RigControl.Value.Get<FVector>();
+					Token.VectorValues.Add(TNameAndValue<FVector>{ RigControl.Name, Val });
+					//mz todo specify rotator special so we can do quat interps
+					break;
+				}
+				case ERigControlType::Transform:
+				{
+					FTransform Val = RigControl.Value.Get<FTransform>();
+					Token.TransformValues.Add(TNameAndValue<FTransform>{ RigControl.Name, Val });
+					break;
+				}
+				}
+			}
+		}
+
+		return MoveTemp(Token);
+	}
+
+	FMovieSceneSequenceID SequenceID;
+	TArray< TNameAndValue<float> > ScalarValues;
+	TArray< TNameAndValue<FVector> > VectorValues;
+	TArray< TNameAndValue<FTransform> > TransformValues;
+
+};
+
+
+
+/* Simple token used for non-blendables*/
+struct FControlRigParameterExecutionToken : IMovieSceneExecutionToken
+{
+	FControlRigParameterExecutionToken(const UMovieSceneControlRigParameterSection* InSection)
+	:	Section(InSection)
+	{}
+	FControlRigParameterExecutionToken(FControlRigParameterExecutionToken&&) = default;
+	FControlRigParameterExecutionToken& operator=(FControlRigParameterExecutionToken&&) = default;
+
+	// Non-copyable
+	FControlRigParameterExecutionToken(const FControlRigParameterExecutionToken&) = delete;
+	FControlRigParameterExecutionToken& operator=(const FControlRigParameterExecutionToken&) = delete;
+
+	void BindToSequencerInstance(UControlRig* ControlRig)
+	{
+		check(ControlRig);
+		if (USkeletalMeshComponent* SkeletalMeshComponent = Cast<USkeletalMeshComponent>(ControlRig->GetObjectBinding()->GetBoundObject()))
+		{
+			bool bWasCreated = false;
+			if (UControlRigSequencerAnimInstance* AnimInstance = UAnimCustomInstance::BindToSkeletalMeshComponent<UControlRigSequencerAnimInstance>(SkeletalMeshComponent, bWasCreated))
+			{
+				AnimInstance->RecalcRequiredBones();
+			}
+		}
+	}
+
+	void UnBindFromSequencerInstance(UControlRig* ControlRig)
+	{
+		check(ControlRig);
+		if (USkeletalMeshComponent* SkeletalMeshComponent = Cast<USkeletalMeshComponent>(ControlRig->GetObjectBinding()->GetBoundObject()))
+		{
+			UAnimCustomInstance::UnbindFromSkeletalMeshComponent(SkeletalMeshComponent);
+		}
+	}
+
+	virtual void Execute(const FMovieSceneContext& Context, const FMovieSceneEvaluationOperand& Operand, FPersistentEvaluationData& PersistentData, IMovieScenePlayer& Player)
+	{
+		MOVIESCENE_DETAILED_SCOPE_CYCLE_COUNTER(MovieSceneEval_ControlRigParameterTrack_TokenExecute)
+		
+		FInputBlendPose DefaultBoneFilter;
+
+		bool bAdditive = false;
+		bool bApplyBoneFilter = false;
+		const FInputBlendPose* BoneFilter = &DefaultBoneFilter;
+
+		FMovieSceneSequenceID SequenceID = Operand.SequenceID;
+		UControlRig* ControlRig = Section->ControlRig;
+
+		// Update the animation's state
+		
+		if (ControlRig)
+		{
+			const UMovieSceneSequence* Sequence = Player.State.FindSequence(Operand.SequenceID);
+			TArrayView<TWeakObjectPtr<>> BoundSkelMeshes = Player.FindBoundObjects(Operand);
+
+			if (Sequence && BoundSkelMeshes.Num() > 0 && BoundSkelMeshes[0].Get())
+			{
+				if (!ControlRig->GetObjectBinding())
+				{
+					ControlRig->SetObjectBinding(MakeShared<FControlRigSkeletalMeshBinding>());
+				}
+				if (!ControlRig->GetObjectBinding()->GetBoundObject())
+				{
+					ControlRig->GetObjectBinding()->BindToObject(BoundSkelMeshes[0].Get());
+				}
+				BindToSequencerInstance(ControlRig);
+				//MZ TODO HANDLE BOOLS AND OTHER NON BLENDABLES
+				if (ControlRig->GetObjectBinding())
+				{
+					if (USkeletalMeshComponent* SkeletalMeshComponent = Cast<USkeletalMeshComponent>(ControlRig->GetObjectBinding()->GetBoundObject()))
+					{
+						if (UControlRigSequencerAnimInstance* AnimInstance = Cast<UControlRigSequencerAnimInstance>(SkeletalMeshComponent->GetAnimInstance()))
+						{
+							float Weight = 1.0f;
+							bool bStructureChanged = AnimInstance->UpdateControlRig(ControlRig, Operand.SequenceID.GetInternalValue(), bAdditive, bApplyBoneFilter, *BoneFilter, Weight, true, true);
+							if (bStructureChanged)
+							{
+								AnimInstance->RecalcRequiredBones();
+							}
+						}
+					}
+				}
+			}
+		}
+
+		// ensure that pre animated state is saved
+		Player.SavePreAnimatedState(*ControlRig, FMovieSceneControlRigParameterTemplate::GetAnimTypeID(), FControlRigParameterPreAnimatedTokenProducer(Operand.SequenceID));
+		
+	}
+
+	const UMovieSceneControlRigParameterSection* Section;
+
+};
+
+
+FMovieSceneControlRigParameterTemplate::FMovieSceneControlRigParameterTemplate(const UMovieSceneControlRigParameterSection& Section, const UMovieSceneControlRigParameterTrack& Track)
+	: FMovieSceneParameterSectionTemplate(Section) 
+{
+}
+
+
+
+struct TControlRigParameterActuatorFloat : TMovieSceneBlendingActuator<FControlRigTrackTokenFloat>
+{
+	TControlRigParameterActuatorFloat(FMovieSceneAnimTypeID& InAnimID, UControlRig* InControlRig, const FName& InParameterName, const UMovieSceneControlRigParameterSection* InSection)
+		: TMovieSceneBlendingActuator<FControlRigTrackTokenFloat>(FMovieSceneBlendingActuatorID(InAnimID))
+		, ControlRig(InControlRig)
+		, ParameterName(InParameterName)
+		, SectionData(InSection)
+	{}
+
+
+	FControlRigTrackTokenFloat RetrieveCurrentValue(UObject* InObject, IMovieScenePlayer* Player) const
+	{
+		if (ControlRig)
+		{
+			FRigControl* RigControl = ControlRig->FindControl(ParameterName);
+			if (RigControl && RigControl->ControlType == ERigControlType::Float)
+			{
+				float Val = RigControl->Value.Get<float>();
+				return FControlRigTrackTokenFloat(Val);
+			}
+
+		}
+		return FControlRigTrackTokenFloat();
+	}
+
+	void Actuate(UObject* InObject, const FControlRigTrackTokenFloat& InFinalValue, const TBlendableTokenStack<FControlRigTrackTokenFloat>& OriginalStack, const FMovieSceneContext& Context, FPersistentEvaluationData& PersistentData, IMovieScenePlayer& Player)
+	{
+		const UMovieSceneControlRigParameterSection* Section = nullptr;
+		Section = SectionData.Get();
+		
+		bool bWasDoNotKey = false;
+		if (Section)
+		{
+			bWasDoNotKey = Section->GetDoNotKey();
+			Section->SetDoNotKey(true);
+		}
+		if (ControlRig)
+		{
+			FRigControl* RigControl = ControlRig->FindControl(ParameterName);
+			if (RigControl && RigControl->ControlType == ERigControlType::Float)
+			{
+				ControlRig->SetControlValue<float>(ParameterName, InFinalValue.Value);
+			}
+		}
+		if (Section)
+		{
+			Section->SetDoNotKey(bWasDoNotKey);
+		}
+	}
+
+
+	virtual void Actuate(FMovieSceneInterrogationData& InterrogationData, const FControlRigTrackTokenFloat& InValue, const TBlendableTokenStack<FControlRigTrackTokenFloat>& OriginalStack, const FMovieSceneContext& Context) const override
+	{
+		FFloatInterrogationData Data;
+		Data.Val = InValue.Value;
+		Data.ParameterName = ParameterName;
+		InterrogationData.Add(FFloatInterrogationData(Data), UMovieSceneControlRigParameterSection::GetFloatInterrogationKey());
+	}
+
+	UControlRig *ControlRig;
+	FName ParameterName;
+	TWeakObjectPtr<const UMovieSceneControlRigParameterSection> SectionData;
+
+};
+
+
+
+struct TControlRigParameterActuatorVector : TMovieSceneBlendingActuator<FControlRigTrackTokenVector>
+{
+	TControlRigParameterActuatorVector(FMovieSceneAnimTypeID& InAnimID, UControlRig* InControlRig, const FName& InParameterName, const UMovieSceneControlRigParameterSection* InSection)
+		: TMovieSceneBlendingActuator<FControlRigTrackTokenVector>(FMovieSceneBlendingActuatorID(InAnimID))
+		, ControlRig(InControlRig)
+		, ParameterName(InParameterName)
+		, SectionData(InSection)
+	{}
+
+
+
+	FControlRigTrackTokenVector RetrieveCurrentValue(UObject* InObject, IMovieScenePlayer* Player) const
+	{
+		if (ControlRig)
+		{
+			FRigControl* RigControl = ControlRig->FindControl(ParameterName);
+			if (RigControl && (RigControl->ControlType == ERigControlType::Position || RigControl->ControlType == ERigControlType::Scale || RigControl->ControlType == ERigControlType::Rotator))
+			{
+				FVector Val = RigControl->Value.Get<FVector>();
+				return FControlRigTrackTokenVector(Val);
+			}
+		}
+		return FControlRigTrackTokenVector();
+	}
+
+	void Actuate(UObject* InObject, const FControlRigTrackTokenVector& InFinalValue, const TBlendableTokenStack<FControlRigTrackTokenVector>& OriginalStack, const FMovieSceneContext& Context, FPersistentEvaluationData& PersistentData, IMovieScenePlayer& Player)
+	{
+		const UMovieSceneControlRigParameterSection* Section = nullptr;
+		Section = SectionData.Get();
+		
+		bool bWasDoNotKey = false;
+		if (Section)
+		{
+			bWasDoNotKey = Section->GetDoNotKey();
+			Section->SetDoNotKey(true);
+		}
+		if (ControlRig)
+		{
+			FRigControl* RigControl = ControlRig->FindControl(ParameterName);
+			if (RigControl && (RigControl->ControlType == ERigControlType::Position || RigControl->ControlType == ERigControlType::Scale || RigControl->ControlType == ERigControlType::Rotator))
+			{
+				ControlRig->SetControlValue<FVector>(ParameterName, InFinalValue.Value);
+			}
+		}
+		if (Section)
+		{
+			Section->SetDoNotKey(bWasDoNotKey);
+		}
+	}
+	virtual void Actuate(FMovieSceneInterrogationData& InterrogationData, const FControlRigTrackTokenVector& InValue, const TBlendableTokenStack<FControlRigTrackTokenVector>& OriginalStack, const FMovieSceneContext& Context) const override
+	{
+		FVectorInterrogationData Data;
+		Data.Val = InValue.Value;
+		Data.ParameterName = ParameterName;		
+		InterrogationData.Add(FVectorInterrogationData(Data), UMovieSceneControlRigParameterSection::GetVectorInterrogationKey());
+	}
+
+	UControlRig *ControlRig;
+	FName ParameterName;
+	TWeakObjectPtr<const UMovieSceneControlRigParameterSection> SectionData;
+
+};
+
+
+
+struct TControlRigParameterActuatorTransform : TMovieSceneBlendingActuator<FControlRigTrackTokenTransform>
+{
+	TControlRigParameterActuatorTransform(FMovieSceneAnimTypeID& InAnimID, UControlRig* InControlRig, const FName& InParameterName, const UMovieSceneControlRigParameterSection* InSection)
+		: TMovieSceneBlendingActuator<FControlRigTrackTokenTransform>(FMovieSceneBlendingActuatorID(InAnimID))
+		, ControlRig(InControlRig)
+		, ParameterName(InParameterName)
+		, SectionData(InSection)
+	{}
+
+	FControlRigTrackTokenTransform RetrieveCurrentValue(UObject* InObject, IMovieScenePlayer* Player) const
+	{
+		if (ControlRig)
+		{
+			FRigControl* RigControl = ControlRig->FindControl(ParameterName);
+			if (RigControl && RigControl->ControlType == ERigControlType::Transform)
+			{
+				FTransform Val = RigControl->Value.Get<FTransform>();
+				return FControlRigTrackTokenTransform(Val);
+			}
+		}
+		return FControlRigTrackTokenTransform();
+	}
+
+	void Actuate(UObject* InObject, const FControlRigTrackTokenTransform& InFinalValue, const TBlendableTokenStack<FControlRigTrackTokenTransform>& OriginalStack, const FMovieSceneContext& Context, FPersistentEvaluationData& PersistentData, IMovieScenePlayer& Player)
+	{
+		const UMovieSceneControlRigParameterSection* Section = nullptr;
+		
+		Section = SectionData.Get();
+		if (Section)
+		{
+			UMovieSceneTrack * Track = Cast<UMovieSceneTrack>(Section->GetOuter());
+			if (Track && Track->GetSectionToKey())
+			{
+				Section = Cast<const UMovieSceneControlRigParameterSection>(Track->GetSectionToKey());
+			}
+		}
+
+		bool bWasDoNotKey = false;
+		if (Section)
+		{
+			bWasDoNotKey = Section->GetDoNotKey();
+			Section->SetDoNotKey(true);
+		}
+		if (ControlRig)
+		{
+			FRigControl* RigControl = ControlRig->FindControl(ParameterName);
+			if (RigControl && RigControl->ControlType == ERigControlType::Transform)
+			{
+				ControlRig->SetControlValue<FTransform>(ParameterName, InFinalValue.Value);
+			}
+		}
+		if (Section)
+		{
+			Section->SetDoNotKey(bWasDoNotKey);
+		}
+	}
+	virtual void Actuate(FMovieSceneInterrogationData& InterrogationData, const FControlRigTrackTokenTransform& InValue, const TBlendableTokenStack<FControlRigTrackTokenTransform>& OriginalStack, const FMovieSceneContext& Context) const override
+	{
+		FTransformInterrogationData Data;
+		Data.Val = InValue.Value;
+		Data.ParameterName = ParameterName;
+		InterrogationData.Add(FTransformInterrogationData(Data), UMovieSceneControlRigParameterSection::GetTransformInterrogationKey());
+	}
+	UControlRig *ControlRig;
+	FName ParameterName;
+	TWeakObjectPtr<const UMovieSceneControlRigParameterSection> SectionData;
+
+};
+
+
+
+void FMovieSceneControlRigParameterTemplate::Evaluate(const FMovieSceneEvaluationOperand& Operand, const FMovieSceneContext& Context, const FPersistentEvaluationData& PersistentData, FMovieSceneExecutionTokens& ExecutionTokens) const
+{
+
+	const FFrameTime Time = Context.GetTime();
+
+	const UMovieSceneControlRigParameterSection* Section = nullptr;
+	if (GetSourceSection())
+	{
+		Section = Cast<UMovieSceneControlRigParameterSection>(GetSourceSection());
+
+		//Do basic token
+		FControlRigParameterExecutionToken ExecutionToken(Section);
+		ExecutionTokens.Add(MoveTemp(ExecutionToken));
+
+		//Do blended tokens
+		FEvaluatedControlRigParameterSectionValues Values;
+		EvaluateCurvesWithMasks(Context, Values);
+		static TMovieSceneAnimTypeIDContainer<FString> ScalarAnimTypeIDsByName;
+		static TMovieSceneAnimTypeIDContainer<FString> VectorAnimTypeIDsByName;
+		static TMovieSceneAnimTypeIDContainer<FString> TransformAnimTypeIDsByName;
+
+
+		float Weight = 1.f;
+
+		//float Weight = EvaluateEasing(Context.GetTime());
+	//if (EnumHasAllFlags(TemplateData.Mask.GetChannels(), EMovieSceneTransformChannel::Weight))
+	//{
+	//	float ManualWeight = 1.f;
+	//	TemplateData.ManualWeight.Evaluate(Context.GetTime(), ManualWeight);
+	//	Weight *= ManualWeight;
+	//}
+
+		uint32 OperandHash = GetTypeHash(Operand);
+		FString UniqueActuator(FString::FromInt((int32)OperandHash));
+		for (const FScalarParameterStringAndValue& ScalarNameAndValue : Values.ScalarValues)
+		{
+			FString NewString(ScalarNameAndValue.ParameterString);
+			NewString.Append(UniqueActuator);
+			FMovieSceneAnimTypeID AnimTypeID = ScalarAnimTypeIDsByName.GetAnimTypeID(NewString);
+			FMovieSceneBlendingActuatorID ActuatorTypeID(AnimTypeID);
+
+			if (!ExecutionTokens.GetBlendingAccumulator().FindActuator< FControlRigTrackTokenFloat>(ActuatorTypeID))
+			{
+				ExecutionTokens.GetBlendingAccumulator().DefineActuator(ActuatorTypeID, MakeShared <TControlRigParameterActuatorFloat>(AnimTypeID, Section->ControlRig, ScalarNameAndValue.ParameterName,Section));
+			}
+			ExecutionTokens.BlendToken(ActuatorTypeID, TBlendableToken<FControlRigTrackTokenFloat>(ScalarNameAndValue.Value, Section->GetBlendType().Get(), Weight));
+		}
+
+		MovieScene::TMultiChannelValue<float, 3> VectorData;
+
+		for (const FVectorParameterStringAndValue& VectorNameAndValue : Values.VectorValues)
+		{
+			FString NewString(VectorNameAndValue.ParameterString);
+			NewString.Append(UniqueActuator);
+			FMovieSceneAnimTypeID AnimTypeID = VectorAnimTypeIDsByName.GetAnimTypeID(NewString);
+			FMovieSceneBlendingActuatorID ActuatorTypeID(AnimTypeID);
+
+			if (!ExecutionTokens.GetBlendingAccumulator().FindActuator< FControlRigTrackTokenVector>(ActuatorTypeID))
+			{
+				ExecutionTokens.GetBlendingAccumulator().DefineActuator(ActuatorTypeID, MakeShared <TControlRigParameterActuatorVector>(AnimTypeID, Section->ControlRig, VectorNameAndValue.ParameterName,Section));
+			}
+			VectorData.Set(0, VectorNameAndValue.Value.X);
+			VectorData.Set(1, VectorNameAndValue.Value.Y);
+			VectorData.Set(2, VectorNameAndValue.Value.Z);
+
+			ExecutionTokens.BlendToken(ActuatorTypeID, TBlendableToken<FControlRigTrackTokenVector>(VectorData, Section->GetBlendType().Get(), Weight));
+		}
+		MovieScene::TMultiChannelValue<float, 9> TransformData;
+
+		for (const FTransformParameterStringAndValue& TransformNameAndValue : Values.TransformValues)
+		{
+			FString NewString(TransformNameAndValue.ParameterString);
+			NewString.Append(UniqueActuator);
+			FMovieSceneAnimTypeID AnimTypeID = TransformAnimTypeIDsByName.GetAnimTypeID(NewString);
+			FMovieSceneBlendingActuatorID ActuatorTypeID(AnimTypeID);
+
+			if (!ExecutionTokens.GetBlendingAccumulator().FindActuator< FControlRigTrackTokenTransform>(ActuatorTypeID))
+			{
+				ExecutionTokens.GetBlendingAccumulator().DefineActuator(ActuatorTypeID, MakeShared <TControlRigParameterActuatorTransform>(AnimTypeID, Section->ControlRig, TransformNameAndValue.ParameterName, Section));
+			}
+
+			FTransform Transform(TransformNameAndValue.Rotation, TransformNameAndValue.Translation, TransformNameAndValue.Scale);
+
+			TransformData.Set(0, TransformNameAndValue.Translation.X);
+			TransformData.Set(1, TransformNameAndValue.Translation.Y);
+			TransformData.Set(2, TransformNameAndValue.Translation.Z);
+
+			TransformData.Set(3, TransformNameAndValue.Rotation.Roll);
+			TransformData.Set(4, TransformNameAndValue.Rotation.Pitch);
+			TransformData.Set(5, TransformNameAndValue.Rotation.Yaw);
+
+			TransformData.Set(6, TransformNameAndValue.Scale.X);
+			TransformData.Set(7, TransformNameAndValue.Scale.Y);
+			TransformData.Set(8, TransformNameAndValue.Scale.Z);
+			ExecutionTokens.BlendToken(ActuatorTypeID, TBlendableToken<FControlRigTrackTokenTransform>(TransformData, Section->GetBlendType().Get(), Weight));
+		}
+
+	}
+}
+
+
+void FMovieSceneControlRigParameterTemplate::EvaluateCurvesWithMasks(const FMovieSceneContext& Context, FEvaluatedControlRigParameterSectionValues& Values) const
+{
+	const FFrameTime Time = Context.GetTime();
+
+
+	const UMovieSceneControlRigParameterSection* Section = nullptr;
+	if (GetSourceSection())
+	{
+		Section = Cast<UMovieSceneControlRigParameterSection>(GetSourceSection());
+		TArray<bool> ControlsMask = Section->GetControlsMask();
+
+		int32 Index = 0;
+		for (const FScalarParameterNameAndCurve& Scalar : Scalars)
+		{
+			float Value = 0;
+			if (ControlsMask[Index])
+			{
+				Scalar.ParameterCurve.Evaluate(Time, Value);
+			}
+			
+			Values.ScalarValues.Emplace(Scalar.ParameterName, Value);
+			++Index;
+		}
+
+		for (const FVectorParameterNameAndCurves& Vector : Vectors)
+		{
+			FVector Value(ForceInitToZero);
+			if (ControlsMask[Index])
+			{
+				Vector.XCurve.Evaluate(Time, Value.X);
+				Vector.YCurve.Evaluate(Time, Value.Y);
+				Vector.ZCurve.Evaluate(Time, Value.Z);
+			}
+
+			Values.VectorValues.Emplace(Vector.ParameterName, Value);
+			++Index;
+		}
+
+		for (const FColorParameterNameAndCurves& Color : Colors)
+		{
+			FLinearColor ColorValue = FLinearColor::White;
+			if (ControlsMask[Index])
+			{
+				Color.RedCurve.Evaluate(Time, ColorValue.R);
+				Color.GreenCurve.Evaluate(Time, ColorValue.G);
+				Color.BlueCurve.Evaluate(Time, ColorValue.B);
+				Color.AlphaCurve.Evaluate(Time, ColorValue.A);
+			}
+		
+			Values.ColorValues.Emplace(Color.ParameterName, ColorValue);
+			++Index;
+		}
+
+		for (const FTransformParameterNameAndCurves& Transform : Transforms)
+		{
+			FVector Translation(ForceInitToZero), Scale(FVector::OneVector);
+			FRotator Rotator(0.0f, 0.0f, 0.0f);
+			EMovieSceneTransformChannel ChannelMask = Section->GetTransformMask().GetChannels();
+
+			if (ControlsMask[Index])
+			{
+				if (EnumHasAllFlags(ChannelMask, EMovieSceneTransformChannel::TranslationX))
+				{
+					Transform.Translation[0].Evaluate(Time, Translation[0]);
+				}
+				if (EnumHasAllFlags(ChannelMask, EMovieSceneTransformChannel::TranslationY))
+				{
+					Transform.Translation[1].Evaluate(Time, Translation[1]);
+				}
+				if (EnumHasAllFlags(ChannelMask, EMovieSceneTransformChannel::TranslationZ))
+				{
+					Transform.Translation[2].Evaluate(Time, Translation[2]);
+				}
+				if (EnumHasAllFlags(ChannelMask, EMovieSceneTransformChannel::RotationX))
+				{
+					Transform.Rotation[0].Evaluate(Time, Rotator.Roll);
+				}
+				if (EnumHasAllFlags(ChannelMask, EMovieSceneTransformChannel::RotationY))
+				{
+					Transform.Rotation[1].Evaluate(Time, Rotator.Pitch);
+				}
+				if (EnumHasAllFlags(ChannelMask, EMovieSceneTransformChannel::RotationZ))
+				{
+					Transform.Rotation[2].Evaluate(Time, Rotator.Yaw);
+				}
+				//mz todo quat interp...
+				if (EnumHasAllFlags(ChannelMask, EMovieSceneTransformChannel::ScaleX))
+				{
+					Transform.Scale[0].Evaluate(Time, Scale[0]);
+				}
+				if (EnumHasAllFlags(ChannelMask, EMovieSceneTransformChannel::ScaleY))
+				{
+					Transform.Scale[1].Evaluate(Time, Scale[1]);
+				}
+				if (EnumHasAllFlags(ChannelMask, EMovieSceneTransformChannel::ScaleZ))
+				{
+					Transform.Scale[2].Evaluate(Time, Scale[2]);
+				}
+			}
+			FTransformParameterStringAndValue NameAndValue(Transform.ParameterName, Translation, Rotator, Scale);
+			Values.TransformValues.Emplace(NameAndValue);
+			++Index;
+
+		}
+	}
+}
+
+
+FMovieSceneAnimTypeID FMovieSceneControlRigParameterTemplate::GetAnimTypeID()
+{
+	return TMovieSceneAnimTypeID<FMovieSceneControlRigParameterTemplate>();
+}
+
+
+void FMovieSceneControlRigParameterTemplate::Interrogate(const FMovieSceneContext& Context, FMovieSceneInterrogationData& Container, UObject* BindingOverride) const
+{
+	MOVIESCENE_DETAILED_SCOPE_CYCLE_COUNTER(MovieSceneEval_ControlRigTemplateParameter_Evaluate)
+
+		const FFrameTime Time = Context.GetTime();
+
+	const UMovieSceneControlRigParameterSection* Section = nullptr;
+	if (GetSourceSection())
+	{
+		Section = Cast<UMovieSceneControlRigParameterSection>(GetSourceSection());
+
+		//Do blended tokens
+		FEvaluatedControlRigParameterSectionValues Values;
+		EvaluateCurvesWithMasks(Context, Values);
+		static TMovieSceneAnimTypeIDContainer<FString> ScalarAnimTypeIDsByName;
+		static TMovieSceneAnimTypeIDContainer<FString> VectorAnimTypeIDsByName;
+		static TMovieSceneAnimTypeIDContainer<FString> TransformAnimTypeIDsByName;
+
+
+		float Weight = 1.f;
+
+		//float Weight = EvaluateEasing(Context.GetTime());
+	//if (EnumHasAllFlags(TemplateData.Mask.GetChannels(), EMovieSceneTransformChannel::Weight))
+	//{
+	//	float ManualWeight = 1.f;
+	//	TemplateData.ManualWeight.Evaluate(Context.GetTime(), ManualWeight);
+	//	Weight *= ManualWeight;
+	//}
+
+
+		for (const FScalarParameterStringAndValue& ScalarNameAndValue : Values.ScalarValues)
+		{
+			FMovieSceneAnimTypeID AnimTypeID = ScalarAnimTypeIDsByName.GetAnimTypeID(ScalarNameAndValue.ParameterString);
+			FMovieSceneBlendingActuatorID ActuatorTypeID(AnimTypeID);
+
+			if (!Container.GetAccumulator().FindActuator< FControlRigTrackTokenFloat>(ActuatorTypeID))
+			{
+				Container.GetAccumulator().DefineActuator(ActuatorTypeID, MakeShared <TControlRigParameterActuatorFloat>(AnimTypeID, Section->ControlRig, ScalarNameAndValue.ParameterName, Section));
+			}
+			Container.GetAccumulator().BlendToken(FMovieSceneEvaluationOperand(), ActuatorTypeID, FMovieSceneEvaluationScope(), Context,TBlendableToken<FControlRigTrackTokenFloat>(ScalarNameAndValue.Value, Section->GetBlendType().Get(), Weight));
+		}
+
+		MovieScene::TMultiChannelValue<float, 3> VectorData;
+
+		for (const FVectorParameterStringAndValue& VectorNameAndValue : Values.VectorValues)
+		{
+			FMovieSceneAnimTypeID AnimTypeID = VectorAnimTypeIDsByName.GetAnimTypeID(VectorNameAndValue.ParameterString);
+			FMovieSceneBlendingActuatorID ActuatorTypeID(AnimTypeID);
+
+			if (!Container.GetAccumulator().FindActuator< FControlRigTrackTokenVector>(ActuatorTypeID))
+			{
+				Container.GetAccumulator().DefineActuator(ActuatorTypeID, MakeShared <TControlRigParameterActuatorVector>(AnimTypeID, Section->ControlRig, VectorNameAndValue.ParameterName, Section));
+			}
+			VectorData.Set(0, VectorNameAndValue.Value.X);
+			VectorData.Set(1, VectorNameAndValue.Value.Y);
+			VectorData.Set(2, VectorNameAndValue.Value.Z);
+
+			Container.GetAccumulator().BlendToken(FMovieSceneEvaluationOperand(), ActuatorTypeID, FMovieSceneEvaluationScope(),Context,TBlendableToken<FControlRigTrackTokenVector>(VectorData, Section->GetBlendType().Get(), Weight));
+		}
+		MovieScene::TMultiChannelValue<float, 9> TransformData;
+
+		for (const FTransformParameterStringAndValue& TransformNameAndValue : Values.TransformValues)
+		{
+			uint32 ID = BindingOverride->GetUniqueID();
+
+
+			FMovieSceneAnimTypeID AnimTypeID = TransformAnimTypeIDsByName.GetAnimTypeID(TransformNameAndValue.ParameterString);
+			FMovieSceneBlendingActuatorID ActuatorTypeID(AnimTypeID);
+
+			if (!Container.GetAccumulator().FindActuator< FControlRigTrackTokenTransform>(ActuatorTypeID))
+			{
+				Container.GetAccumulator().DefineActuator(ActuatorTypeID, MakeShared <TControlRigParameterActuatorTransform>(AnimTypeID, Section->ControlRig, TransformNameAndValue.ParameterName, Section));
+			}
+
+			FTransform Transform(TransformNameAndValue.Rotation, TransformNameAndValue.Translation, TransformNameAndValue.Scale);
+
+			TransformData.Set(0, TransformNameAndValue.Translation.X);
+			TransformData.Set(1, TransformNameAndValue.Translation.Y);
+			TransformData.Set(2, TransformNameAndValue.Translation.Z);
+
+			TransformData.Set(3, TransformNameAndValue.Rotation.Roll);
+			TransformData.Set(4, TransformNameAndValue.Rotation.Pitch);
+			TransformData.Set(5, TransformNameAndValue.Rotation.Yaw);
+
+			TransformData.Set(6, TransformNameAndValue.Scale.X);
+			TransformData.Set(7, TransformNameAndValue.Scale.Y);
+			TransformData.Set(8, TransformNameAndValue.Scale.Z);
+			Container.GetAccumulator().BlendToken(FMovieSceneEvaluationOperand(),ActuatorTypeID, FMovieSceneEvaluationScope(), Context, TBlendableToken<FControlRigTrackTokenTransform>(TransformData, Section->GetBlendType().Get(), Weight));
+		}
+
+	}
+}
+
diff --git a/Engine/Plugins/Experimental/ControlRig/Source/ControlRig/Private/Units/Deprecated/RigUnit_ApplyFK.cpp b/Engine/Plugins/Experimental/ControlRig/Source/ControlRig/Private/Units/Deprecated/RigUnit_ApplyFK.cpp
--- a/Engine/Plugins/Experimental/ControlRig/Source/ControlRig/Private/Units/Deprecated/RigUnit_ApplyFK.cpp
+++ b/Engine/Plugins/Experimental/ControlRig/Source/ControlRig/Private/Units/Deprecated/RigUnit_ApplyFK.cpp
@@ -62,8 +62,8 @@ FTransform FRigUnit_ApplyFK::GetBaseTransform(int32 BoneIndex, const FRigBoneHie
 
 IMPLEMENT_RIGUNIT_AUTOMATION_TEST(FRigUnit_ApplyFK)
 {
-	BoneHierarchy.Add(TEXT("Root"), NAME_None, FTransform(FVector(1.f, 0.f, 0.f)));
-	BoneHierarchy.Add(TEXT("BoneA"), TEXT("Root"), FTransform(FVector(1.f, 2.f, 3.f)));
+	BoneHierarchy.Add(TEXT("Root"), NAME_None, ERigBoneType::User, FTransform(FVector(1.f, 0.f, 0.f)));
+	BoneHierarchy.Add(TEXT("BoneA"), TEXT("Root"), ERigBoneType::User, FTransform(FVector(1.f, 2.f, 3.f)));
 
 	Unit.ExecuteContext = ExecuteContext;
 	Unit.Joint = TEXT("BoneA");
diff --git a/Engine/Plugins/Experimental/ControlRig/Source/ControlRig/Private/Units/Hierarchy/RigUnit_AddBoneTransform.cpp b/Engine/Plugins/Experimental/ControlRig/Source/ControlRig/Private/Units/Hierarchy/RigUnit_AddBoneTransform.cpp
--- a/Engine/Plugins/Experimental/ControlRig/Source/ControlRig/Private/Units/Hierarchy/RigUnit_AddBoneTransform.cpp
+++ b/Engine/Plugins/Experimental/ControlRig/Source/ControlRig/Private/Units/Hierarchy/RigUnit_AddBoneTransform.cpp
@@ -25,15 +25,25 @@ FRigUnit_AddBoneTransform_Execute()
 			{
 				if (CachedBoneIndex != INDEX_NONE)
 				{
-					FTransform PreviousTransform = Hierarchy->GetGlobalTransform(CachedBoneIndex);
+					FTransform TargetTransform;
+					const FTransform PreviousTransform = Hierarchy->GetGlobalTransform(CachedBoneIndex);
+
 					if (bPostMultiply)
 					{
-						Hierarchy->SetGlobalTransform(CachedBoneIndex, PreviousTransform * Transform, bPropagateToChildren);
+						TargetTransform = PreviousTransform * Transform;
 					}
 					else
 					{
-						Hierarchy->SetGlobalTransform(CachedBoneIndex, Transform * PreviousTransform, bPropagateToChildren);
+						TargetTransform = Transform * PreviousTransform;
 					}
+
+					if (!FMath::IsNearlyEqual(Weight, 1.f))
+					{
+						float T = FMath::Clamp<float>(Weight, 0.f, 1.f);
+						TargetTransform = FControlRigMathLibrary::LerpTransform(PreviousTransform, TargetTransform, T);
+					}
+
+					Hierarchy->SetGlobalTransform(CachedBoneIndex, TargetTransform, bPropagateToChildren);
 				}
 			}
 			default:
@@ -49,8 +59,8 @@ FRigUnit_AddBoneTransform_Execute()
 
 IMPLEMENT_RIGUNIT_AUTOMATION_TEST(FRigUnit_AddBoneTransform)
 {
-	BoneHierarchy.Add(TEXT("Root"), NAME_None, FTransform(FVector(1.f, 0.f, 0.f)));
-	BoneHierarchy.Add(TEXT("BoneA"), TEXT("Root"), FTransform(FVector(1.f, 2.f, 0.f)));
+	BoneHierarchy.Add(TEXT("Root"), NAME_None, ERigBoneType::User, FTransform(FVector(1.f, 0.f, 0.f)));
+	BoneHierarchy.Add(TEXT("BoneA"), TEXT("Root"), ERigBoneType::User, FTransform(FVector(1.f, 2.f, 0.f)));
 	BoneHierarchy.Initialize();
 	Unit.ExecuteContext.Hierarchy = &HierarchyContainer;
 
diff --git a/Engine/Plugins/Experimental/ControlRig/Source/ControlRig/Private/Units/Hierarchy/RigUnit_AddBoneTransform.h b/Engine/Plugins/Experimental/ControlRig/Source/ControlRig/Private/Units/Hierarchy/RigUnit_AddBoneTransform.h
--- a/Engine/Plugins/Experimental/ControlRig/Source/ControlRig/Private/Units/Hierarchy/RigUnit_AddBoneTransform.h
+++ b/Engine/Plugins/Experimental/ControlRig/Source/ControlRig/Private/Units/Hierarchy/RigUnit_AddBoneTransform.h
@@ -15,7 +15,8 @@ struct FRigUnit_AddBoneTransform : public FRigUnitMutable
 	GENERATED_BODY()
 
 	FRigUnit_AddBoneTransform()
-		: bPostMultiply(false)
+		: Weight(1.f)
+		, bPostMultiply(false)
 		, bPropagateToChildren(false)
 		, CachedBoneIndex(INDEX_NONE)
 	{}
@@ -36,6 +37,12 @@ struct FRigUnit_AddBoneTransform : public FRigUnitMutable
 	UPROPERTY(meta = (Input))
 	FTransform Transform;
 
+	/**
+	 * The weight of the change - how much the change should be applied
+	 */
+	UPROPERTY(meta = (Input, UIMin = "0.0", UIMax = "1.0"))
+	float Weight;
+
 	/**
 	 * If set to true the transform will be post multiplied, otherwise pre multiplied.
 	 * Post multiplying means that the transform is understood as a parent space change,
diff --git a/Engine/Plugins/Experimental/ControlRig/Source/ControlRig/Private/Units/Hierarchy/RigUnit_BoneName.cpp b/Engine/Plugins/Experimental/ControlRig/Source/ControlRig/Private/Units/Hierarchy/RigUnit_BoneName.cpp
--- a/Engine/Plugins/Experimental/ControlRig/Source/ControlRig/Private/Units/Hierarchy/RigUnit_BoneName.cpp
+++ b/Engine/Plugins/Experimental/ControlRig/Source/ControlRig/Private/Units/Hierarchy/RigUnit_BoneName.cpp
@@ -7,3 +7,7 @@ FString FRigUnit_BoneName::GetUnitLabel() const
 {
 	return FString::Printf(TEXT("%s Name"), *Bone.ToString());
 }
+
+FRigUnit_BoneName_Execute()
+{
+}
\ No newline at end of file
diff --git a/Engine/Plugins/Experimental/ControlRig/Source/ControlRig/Private/Units/Hierarchy/RigUnit_BoneName.h b/Engine/Plugins/Experimental/ControlRig/Source/ControlRig/Private/Units/Hierarchy/RigUnit_BoneName.h
--- a/Engine/Plugins/Experimental/ControlRig/Source/ControlRig/Private/Units/Hierarchy/RigUnit_BoneName.h
+++ b/Engine/Plugins/Experimental/ControlRig/Source/ControlRig/Private/Units/Hierarchy/RigUnit_BoneName.h
@@ -8,7 +8,7 @@
 /**
  * BoneName is used to represent a bone name in the graph
  */
-USTRUCT(meta=(DisplayName="Bone Name", Category="Hierarchy", DocumentationPolicy = "Strict", Deprecated = "4.24"))
+USTRUCT(meta=(DisplayName="Bone Name", Category="Hierarchy", DocumentationPolicy = "Strict"))
 struct FRigUnit_BoneName : public FRigUnit
 {
 	GENERATED_BODY()
@@ -17,7 +17,9 @@ struct FRigUnit_BoneName : public FRigUnit
 	{}
 
 	virtual FString GetUnitLabel() const override;
-	virtual void Execute(const FRigUnitContext& Context) override {}
+
+	RIGVM_METHOD()
+	virtual void Execute(const FRigUnitContext& Context) override;
 
 	/**
 	 * The name of the Bone
diff --git a/Engine/Plugins/Experimental/ControlRig/Source/ControlRig/Private/Units/Hierarchy/RigUnit_GetBoneTransform.cpp b/Engine/Plugins/Experimental/ControlRig/Source/ControlRig/Private/Units/Hierarchy/RigUnit_GetBoneTransform.cpp
--- a/Engine/Plugins/Experimental/ControlRig/Source/ControlRig/Private/Units/Hierarchy/RigUnit_GetBoneTransform.cpp
+++ b/Engine/Plugins/Experimental/ControlRig/Source/ControlRig/Private/Units/Hierarchy/RigUnit_GetBoneTransform.cpp
@@ -57,8 +57,8 @@ FRigUnit_GetBoneTransform_Execute()
 
 IMPLEMENT_RIGUNIT_AUTOMATION_TEST(FRigUnit_GetBoneTransform)
 {
-	BoneHierarchy.Add(TEXT("Root"), NAME_None, FTransform(FVector(1.f, 0.f, 0.f)));
-	BoneHierarchy.Add(TEXT("BoneA"), TEXT("Root"), FTransform(FVector(1.f, 2.f, 3.f)));
+	BoneHierarchy.Add(TEXT("Root"), NAME_None, ERigBoneType::User, FTransform(FVector(1.f, 0.f, 0.f)));
+	BoneHierarchy.Add(TEXT("BoneA"), TEXT("Root"), ERigBoneType::User, FTransform(FVector(1.f, 2.f, 3.f)));
 	BoneHierarchy.Initialize();
 
 	Unit.Bone = TEXT("Unknown");
diff --git a/Engine/Plugins/Experimental/ControlRig/Source/ControlRig/Private/Units/Hierarchy/RigUnit_GetRelativeBoneTransform.cpp b/Engine/Plugins/Experimental/ControlRig/Source/ControlRig/Private/Units/Hierarchy/RigUnit_GetRelativeBoneTransform.cpp
--- a/Engine/Plugins/Experimental/ControlRig/Source/ControlRig/Private/Units/Hierarchy/RigUnit_GetRelativeBoneTransform.cpp
+++ b/Engine/Plugins/Experimental/ControlRig/Source/ControlRig/Private/Units/Hierarchy/RigUnit_GetRelativeBoneTransform.cpp
@@ -43,9 +43,9 @@ FRigUnit_GetRelativeBoneTransform_Execute()
 
 IMPLEMENT_RIGUNIT_AUTOMATION_TEST(FRigUnit_GetRelativeBoneTransform)
 {
-	BoneHierarchy.Add(TEXT("Root"), NAME_None, FTransform(FVector(1.f, 0.f, 0.f)));
-	BoneHierarchy.Add(TEXT("BoneA"), TEXT("Root"), FTransform(FVector(1.f, 2.f, 3.f)));
-	BoneHierarchy.Add(TEXT("BoneB"), TEXT("Root"), FTransform(FVector(-4.f, 0.f, 0.f)));
+	BoneHierarchy.Add(TEXT("Root"), NAME_None, ERigBoneType::User, FTransform(FVector(1.f, 0.f, 0.f)));
+	BoneHierarchy.Add(TEXT("BoneA"), TEXT("Root"), ERigBoneType::User, FTransform(FVector(1.f, 2.f, 3.f)));
+	BoneHierarchy.Add(TEXT("BoneB"), TEXT("Root"), ERigBoneType::User, FTransform(FVector(-4.f, 0.f, 0.f)));
 	BoneHierarchy.Initialize();
 
 	Unit.Bone = TEXT("Unknown");
diff --git a/Engine/Plugins/Experimental/ControlRig/Source/ControlRig/Private/Units/Hierarchy/RigUnit_SetBoneRotation.cpp b/Engine/Plugins/Experimental/ControlRig/Source/ControlRig/Private/Units/Hierarchy/RigUnit_SetBoneRotation.cpp
--- a/Engine/Plugins/Experimental/ControlRig/Source/ControlRig/Private/Units/Hierarchy/RigUnit_SetBoneRotation.cpp
+++ b/Engine/Plugins/Experimental/ControlRig/Source/ControlRig/Private/Units/Hierarchy/RigUnit_SetBoneRotation.cpp
@@ -30,14 +30,34 @@ FRigUnit_SetBoneRotation_Execute()
 						case EBoneGetterSetterMode::GlobalSpace:
 						{
 							FTransform Transform = Hierarchy->GetGlobalTransform(CachedBoneIndex);
-							Transform.SetRotation(Rotation);
+
+							if (FMath::IsNearlyEqual(Weight, 1.f))
+							{
+								Transform.SetRotation(Rotation);
+							}
+							else
+							{
+								float T = FMath::Clamp<float>(Weight, 0.f, 1.f);
+								Transform.SetRotation(FQuat::Slerp(Transform.GetRotation(), Rotation, T));
+							}
+
 							Hierarchy->SetGlobalTransform(CachedBoneIndex, Transform, bPropagateToChildren);
 							break;
 						}
 						case EBoneGetterSetterMode::LocalSpace:
 						{
 							FTransform Transform = Hierarchy->GetLocalTransform(CachedBoneIndex);
-							Transform.SetRotation(Rotation);
+
+							if (FMath::IsNearlyEqual(Weight, 1.f))
+							{
+								Transform.SetRotation(Rotation);
+							}
+							else
+							{
+								float T = FMath::Clamp<float>(Weight, 0.f, 1.f);
+								Transform.SetRotation(FQuat::Slerp(Transform.GetRotation(), Rotation, T));
+							}
+
 							Hierarchy->SetLocalTransform(CachedBoneIndex, Transform, bPropagateToChildren);
 							break;
 						}
@@ -61,9 +81,9 @@ FRigUnit_SetBoneRotation_Execute()
 
 IMPLEMENT_RIGUNIT_AUTOMATION_TEST(FRigUnit_SetBoneRotation)
 {
-	BoneHierarchy.Add(TEXT("Root"), NAME_None, FTransform(FQuat(FVector(-1.f, 0.f, 0.f), 0.1f)));
-	BoneHierarchy.Add(TEXT("BoneA"), TEXT("Root"), FTransform(FQuat(FVector(-1.f, 0.f, 0.f), 0.5f)));
-	BoneHierarchy.Add(TEXT("BoneB"), TEXT("BoneA"), FTransform(FQuat(FVector(-1.f, 0.f, 0.f), 0.7f)));
+	BoneHierarchy.Add(TEXT("Root"), NAME_None, ERigBoneType::User, FTransform(FQuat(FVector(-1.f, 0.f, 0.f), 0.1f)));
+	BoneHierarchy.Add(TEXT("BoneA"), TEXT("Root"), ERigBoneType::User, FTransform(FQuat(FVector(-1.f, 0.f, 0.f), 0.5f)));
+	BoneHierarchy.Add(TEXT("BoneB"), TEXT("BoneA"), ERigBoneType::User, FTransform(FQuat(FVector(-1.f, 0.f, 0.f), 0.7f)));
 	BoneHierarchy.Initialize();
 
 	Unit.ExecuteContext.Hierarchy = &HierarchyContainer;
diff --git a/Engine/Plugins/Experimental/ControlRig/Source/ControlRig/Private/Units/Hierarchy/RigUnit_SetBoneRotation.h b/Engine/Plugins/Experimental/ControlRig/Source/ControlRig/Private/Units/Hierarchy/RigUnit_SetBoneRotation.h
--- a/Engine/Plugins/Experimental/ControlRig/Source/ControlRig/Private/Units/Hierarchy/RigUnit_SetBoneRotation.h
+++ b/Engine/Plugins/Experimental/ControlRig/Source/ControlRig/Private/Units/Hierarchy/RigUnit_SetBoneRotation.h
@@ -17,6 +17,7 @@ struct FRigUnit_SetBoneRotation : public FRigUnitMutable
 	FRigUnit_SetBoneRotation()
 		: Rotation(FQuat::Identity)
 		, Space(EBoneGetterSetterMode::LocalSpace)
+		, Weight(1.f)
 		, bPropagateToChildren(false)
 		, CachedBoneIndex(INDEX_NONE)
 	{}
@@ -44,6 +45,12 @@ struct FRigUnit_SetBoneRotation : public FRigUnitMutable
 	UPROPERTY(meta = (Input))
 	EBoneGetterSetterMode Space;
 
+	/**
+	 * The weight of the change - how much the change should be applied
+	 */
+	UPROPERTY(meta = (Input, UIMin = "0.0", UIMax = "1.0"))
+	float Weight;
+
 	/**
 	 * If set to true all of the global transforms of the children 
 	 * of this bone will be recalculated based on their local transforms.
diff --git a/Engine/Plugins/Experimental/ControlRig/Source/ControlRig/Private/Units/Hierarchy/RigUnit_SetBoneTransform.cpp b/Engine/Plugins/Experimental/ControlRig/Source/ControlRig/Private/Units/Hierarchy/RigUnit_SetBoneTransform.cpp
--- a/Engine/Plugins/Experimental/ControlRig/Source/ControlRig/Private/Units/Hierarchy/RigUnit_SetBoneTransform.cpp
+++ b/Engine/Plugins/Experimental/ControlRig/Source/ControlRig/Private/Units/Hierarchy/RigUnit_SetBoneTransform.cpp
@@ -29,12 +29,30 @@ FRigUnit_SetBoneTransform_Execute()
 					{
 						case EBoneGetterSetterMode::GlobalSpace:
 						{
-							Hierarchy->SetGlobalTransform(CachedBoneIndex, Transform, bPropagateToChildren);
+							if (FMath::IsNearlyEqual(Weight, 1.f))
+							{
+								Hierarchy->SetGlobalTransform(CachedBoneIndex, Transform, bPropagateToChildren);
+							}
+							else
+							{
+								float T = FMath::Clamp<float>(Weight, 0.f, 1.f);
+								const FTransform PreviousTransform = Hierarchy->GetGlobalTransform(CachedBoneIndex);
+								Hierarchy->SetGlobalTransform(CachedBoneIndex, FControlRigMathLibrary::LerpTransform(PreviousTransform, Transform, T), bPropagateToChildren);
+							}
 							break;
 						}
 						case EBoneGetterSetterMode::LocalSpace:
 						{
-							Hierarchy->SetLocalTransform(CachedBoneIndex, Transform, bPropagateToChildren);
+							if (FMath::IsNearlyEqual(Weight, 1.f))
+							{
+								Hierarchy->SetLocalTransform(CachedBoneIndex, Transform, bPropagateToChildren);
+							}
+							else
+							{
+								float T = FMath::Clamp<float>(Weight, 0.f, 1.f);
+								const FTransform PreviousTransform = Hierarchy->GetLocalTransform(CachedBoneIndex);
+								Hierarchy->SetLocalTransform(CachedBoneIndex, FControlRigMathLibrary::LerpTransform(PreviousTransform, Transform, T), bPropagateToChildren);
+							}
 							break;
 						}
 						default:
@@ -57,9 +75,9 @@ FRigUnit_SetBoneTransform_Execute()
 
 IMPLEMENT_RIGUNIT_AUTOMATION_TEST(FRigUnit_SetBoneTransform)
 {
-	BoneHierarchy.Add(TEXT("Root"), NAME_None, FTransform(FVector(1.f, 0.f, 0.f)));
-	BoneHierarchy.Add(TEXT("BoneA"), TEXT("Root"), FTransform(FVector(1.f, 2.f, 3.f)));
-	BoneHierarchy.Add(TEXT("BoneB"), TEXT("BoneA"), FTransform(FVector(1.f, 5.f, 3.f)));
+	BoneHierarchy.Add(TEXT("Root"), NAME_None, ERigBoneType::User, FTransform(FVector(1.f, 0.f, 0.f)));
+	BoneHierarchy.Add(TEXT("BoneA"), TEXT("Root"), ERigBoneType::User, FTransform(FVector(1.f, 2.f, 3.f)));
+	BoneHierarchy.Add(TEXT("BoneB"), TEXT("BoneA"), ERigBoneType::User, FTransform(FVector(1.f, 5.f, 3.f)));
 	BoneHierarchy.Initialize();
 	Unit.ExecuteContext.Hierarchy = &HierarchyContainer;
 
diff --git a/Engine/Plugins/Experimental/ControlRig/Source/ControlRig/Private/Units/Hierarchy/RigUnit_SetBoneTransform.h b/Engine/Plugins/Experimental/ControlRig/Source/ControlRig/Private/Units/Hierarchy/RigUnit_SetBoneTransform.h
--- a/Engine/Plugins/Experimental/ControlRig/Source/ControlRig/Private/Units/Hierarchy/RigUnit_SetBoneTransform.h
+++ b/Engine/Plugins/Experimental/ControlRig/Source/ControlRig/Private/Units/Hierarchy/RigUnit_SetBoneTransform.h
@@ -16,6 +16,7 @@ struct FRigUnit_SetBoneTransform : public FRigUnitMutable
 
 	FRigUnit_SetBoneTransform()
 		: Space(EBoneGetterSetterMode::LocalSpace)
+		, Weight(1.f)
 		, bPropagateToChildren(false)
 		, CachedBoneIndex(INDEX_NONE)
 	{}
@@ -43,6 +44,12 @@ struct FRigUnit_SetBoneTransform : public FRigUnitMutable
 	 */
 	UPROPERTY(meta = (Input))
 	EBoneGetterSetterMode Space;
+	
+	/**
+	 * The weight of the change - how much the change should be applied
+	 */
+	UPROPERTY(meta = (Input, UIMin = "0.0", UIMax = "1.0"))
+	float Weight;
 
 	/**
 	 * If set to true all of the global transforms of the children 
diff --git a/Engine/Plugins/Experimental/ControlRig/Source/ControlRig/Private/Units/Hierarchy/RigUnit_SetBoneTranslation.cpp b/Engine/Plugins/Experimental/ControlRig/Source/ControlRig/Private/Units/Hierarchy/RigUnit_SetBoneTranslation.cpp
--- a/Engine/Plugins/Experimental/ControlRig/Source/ControlRig/Private/Units/Hierarchy/RigUnit_SetBoneTranslation.cpp
+++ b/Engine/Plugins/Experimental/ControlRig/Source/ControlRig/Private/Units/Hierarchy/RigUnit_SetBoneTranslation.cpp
@@ -30,14 +30,34 @@ FRigUnit_SetBoneTranslation_Execute()
 						case EBoneGetterSetterMode::GlobalSpace:
 						{
 							FTransform Transform = Hierarchy->GetGlobalTransform(CachedBoneIndex);
-							Transform.SetTranslation(Translation);
+
+							if (FMath::IsNearlyEqual(Weight, 1.f))
+							{
+								Transform.SetTranslation(Translation);
+							}
+							else
+							{
+								float T = FMath::Clamp<float>(Weight, 0.f, 1.f);
+								Transform.SetTranslation(FMath::Lerp<FVector>(Transform.GetTranslation(), Translation, T));
+							}
+
 							Hierarchy->SetGlobalTransform(CachedBoneIndex, Transform, bPropagateToChildren);
 							break;
 						}
 						case EBoneGetterSetterMode::LocalSpace:
 						{
 							FTransform Transform = Hierarchy->GetLocalTransform(CachedBoneIndex);
-							Transform.SetTranslation(Translation);
+
+							if (FMath::IsNearlyEqual(Weight, 1.f))
+							{
+								Transform.SetTranslation(Translation);
+							}
+							else
+							{
+								float T = FMath::Clamp<float>(Weight, 0.f, 1.f);
+								Transform.SetTranslation(FMath::Lerp<FVector>(Transform.GetTranslation(), Translation, T));
+							}
+
 							Hierarchy->SetLocalTransform(CachedBoneIndex, Transform, bPropagateToChildren);
 							break;
 						}
diff --git a/Engine/Plugins/Experimental/ControlRig/Source/ControlRig/Private/Units/Hierarchy/RigUnit_SetBoneTranslation.h b/Engine/Plugins/Experimental/ControlRig/Source/ControlRig/Private/Units/Hierarchy/RigUnit_SetBoneTranslation.h
--- a/Engine/Plugins/Experimental/ControlRig/Source/ControlRig/Private/Units/Hierarchy/RigUnit_SetBoneTranslation.h
+++ b/Engine/Plugins/Experimental/ControlRig/Source/ControlRig/Private/Units/Hierarchy/RigUnit_SetBoneTranslation.h
@@ -17,6 +17,7 @@ struct FRigUnit_SetBoneTranslation : public FRigUnitMutable
 	FRigUnit_SetBoneTranslation()
 		: Translation(FVector::ZeroVector)
 		, Space(EBoneGetterSetterMode::LocalSpace)
+		, Weight(1.f)
 		, bPropagateToChildren(false)
 		, CachedBoneIndex(INDEX_NONE)
 	{}
@@ -44,6 +45,12 @@ struct FRigUnit_SetBoneTranslation : public FRigUnitMutable
 	UPROPERTY(meta = (Input))
 	EBoneGetterSetterMode Space;
 
+	/**
+	 * The weight of the change - how much the change should be applied
+	 */
+	UPROPERTY(meta = (Input, UIMin = "0.0", UIMax = "1.0"))
+	float Weight;
+
 	/**
 	 * If set to true all of the global transforms of the children 
 	 * of this bone will be recalculated based on their local transforms.
diff --git a/Engine/Plugins/Experimental/ControlRig/Source/ControlRig/Private/Units/Hierarchy/RigUnit_SetControlTransform.cpp b/Engine/Plugins/Experimental/ControlRig/Source/ControlRig/Private/Units/Hierarchy/RigUnit_SetControlTransform.cpp
--- a/Engine/Plugins/Experimental/ControlRig/Source/ControlRig/Private/Units/Hierarchy/RigUnit_SetControlTransform.cpp
+++ b/Engine/Plugins/Experimental/ControlRig/Source/ControlRig/Private/Units/Hierarchy/RigUnit_SetControlTransform.cpp
@@ -19,7 +19,7 @@ FRigUnit_SetControlTransform_Execute()
 			case EControlRigState::Init:
 			{
 				CachedControlIndex = Hierarchy->GetIndex(Control);
-				// fall through to update
+				break;
 			}
 			case EControlRigState::Update:
 			{
diff --git a/Engine/Plugins/Experimental/ControlRig/Source/ControlRig/Private/Units/Hierarchy/RigUnit_SetRelativeBoneTransform.cpp b/Engine/Plugins/Experimental/ControlRig/Source/ControlRig/Private/Units/Hierarchy/RigUnit_SetRelativeBoneTransform.cpp
--- a/Engine/Plugins/Experimental/ControlRig/Source/ControlRig/Private/Units/Hierarchy/RigUnit_SetRelativeBoneTransform.cpp
+++ b/Engine/Plugins/Experimental/ControlRig/Source/ControlRig/Private/Units/Hierarchy/RigUnit_SetRelativeBoneTransform.cpp
@@ -27,8 +27,16 @@ FRigUnit_SetRelativeBoneTransform_Execute()
 				if (CachedBoneIndex != INDEX_NONE && CachedSpaceIndex != INDEX_NONE)
 				{
 					const FTransform SpaceTransform = Hierarchy->GetGlobalTransform(CachedSpaceIndex);
-					const FTransform Absolute = Transform * SpaceTransform;
-					Hierarchy->SetGlobalTransform(CachedBoneIndex, Absolute, bPropagateToChildren);
+					FTransform TargetTransform = Transform * SpaceTransform;
+
+					if (!FMath::IsNearlyEqual(Weight, 1.f))
+					{
+						float T = FMath::Clamp<float>(Weight, 0.f, 1.f);
+						const FTransform PreviousTransform = Hierarchy->GetGlobalTransform(CachedBoneIndex);
+						TargetTransform = FControlRigMathLibrary::LerpTransform(PreviousTransform, TargetTransform, T);
+					}
+
+					Hierarchy->SetGlobalTransform(CachedBoneIndex, TargetTransform, bPropagateToChildren);
 				}
 			}
 			default:
@@ -44,10 +52,10 @@ FRigUnit_SetRelativeBoneTransform_Execute()
 
 IMPLEMENT_RIGUNIT_AUTOMATION_TEST(FRigUnit_SetRelativeBoneTransform)
 {
-	BoneHierarchy.Add(TEXT("Root"), NAME_None, FTransform(FVector(1.f, 0.f, 0.f)));
-	BoneHierarchy.Add(TEXT("BoneA"), TEXT("Root"), FTransform(FVector(1.f, 2.f, 3.f)));
-	BoneHierarchy.Add(TEXT("BoneB"), TEXT("BoneA"), FTransform(FVector(1.f, 5.f, 3.f)));
-	BoneHierarchy.Add(TEXT("BoneC"), TEXT("Root"), FTransform(FVector(-4.f, 0.f, 0.f)));
+	BoneHierarchy.Add(TEXT("Root"), NAME_None, ERigBoneType::User, FTransform(FVector(1.f, 0.f, 0.f)));
+	BoneHierarchy.Add(TEXT("BoneA"), TEXT("Root"), ERigBoneType::User, FTransform(FVector(1.f, 2.f, 3.f)));
+	BoneHierarchy.Add(TEXT("BoneB"), TEXT("BoneA"), ERigBoneType::User, FTransform(FVector(1.f, 5.f, 3.f)));
+	BoneHierarchy.Add(TEXT("BoneC"), TEXT("Root"), ERigBoneType::User, FTransform(FVector(-4.f, 0.f, 0.f)));
 	BoneHierarchy.Initialize();
 	Unit.ExecuteContext.Hierarchy = &HierarchyContainer;
 
diff --git a/Engine/Plugins/Experimental/ControlRig/Source/ControlRig/Private/Units/Hierarchy/RigUnit_SetRelativeBoneTransform.h b/Engine/Plugins/Experimental/ControlRig/Source/ControlRig/Private/Units/Hierarchy/RigUnit_SetRelativeBoneTransform.h
--- a/Engine/Plugins/Experimental/ControlRig/Source/ControlRig/Private/Units/Hierarchy/RigUnit_SetRelativeBoneTransform.h
+++ b/Engine/Plugins/Experimental/ControlRig/Source/ControlRig/Private/Units/Hierarchy/RigUnit_SetRelativeBoneTransform.h
@@ -15,7 +15,8 @@ struct FRigUnit_SetRelativeBoneTransform : public FRigUnitMutable
 	GENERATED_BODY()
 
 	FRigUnit_SetRelativeBoneTransform()
-		: bPropagateToChildren(false)
+		: Weight(1.f)
+		, bPropagateToChildren(false)
 		, CachedBoneIndex(INDEX_NONE)
 		, CachedSpaceIndex(INDEX_NONE)
 	{}
@@ -42,6 +43,12 @@ struct FRigUnit_SetRelativeBoneTransform : public FRigUnitMutable
 	UPROPERTY(meta = (Input))
 	FTransform Transform;
 
+	/**
+	 * The weight of the change - how much the change should be applied
+	 */
+	UPROPERTY(meta = (Input, UIMin = "0.0", UIMax = "1.0"))
+	float Weight;
+
 	/**
 	 * If set to true all of the global transforms of the children 
 	 * of this bone will be recalculated based on their local transforms.
diff --git a/Engine/Plugins/Experimental/ControlRig/Source/ControlRig/Private/Units/Hierarchy/RigUnit_SetSpaceTransform.cpp b/Engine/Plugins/Experimental/ControlRig/Source/ControlRig/Private/Units/Hierarchy/RigUnit_SetSpaceTransform.cpp
--- a/Engine/Plugins/Experimental/ControlRig/Source/ControlRig/Private/Units/Hierarchy/RigUnit_SetSpaceTransform.cpp
+++ b/Engine/Plugins/Experimental/ControlRig/Source/ControlRig/Private/Units/Hierarchy/RigUnit_SetSpaceTransform.cpp
@@ -19,7 +19,7 @@ FRigUnit_SetSpaceTransform_Execute()
 			case EControlRigState::Init:
 			{
 				CachedSpaceIndex = Hierarchy->GetIndex(Space);
-				// fall through to update
+				break;
 			}
 			case EControlRigState::Update:
 			{
diff --git a/Engine/Plugins/Experimental/ControlRig/Source/ControlRig/Private/Units/Highlevel/Hierarchy/RigUnit_AimBone.cpp b/Engine/Plugins/Experimental/ControlRig/Source/ControlRig/Private/Units/Highlevel/Hierarchy/RigUnit_AimBone.cpp
--- a/Engine/Plugins/Experimental/ControlRig/Source/ControlRig/Private/Units/Highlevel/Hierarchy/RigUnit_AimBone.cpp
+++ b/Engine/Plugins/Experimental/ControlRig/Source/ControlRig/Private/Units/Highlevel/Hierarchy/RigUnit_AimBone.cpp
@@ -26,7 +26,6 @@ FRigUnit_AimBone_Execute()
 
 	if (Primary.Weight <= SMALL_NUMBER && Secondary.Weight <= SMALL_NUMBER)
 	{
-		UE_CONTROLRIG_RIGUNIT_REPORT_WARNING(TEXT("Weights are zero."));
 		return;
 	}
 	FTransform Transform = Hierarchy->GetGlobalTransform(BoneIndex);
diff --git a/Engine/Plugins/Experimental/ControlRig/Source/ControlRig/Private/Units/Highlevel/Hierarchy/RigUnit_CCDIK.cpp b/Engine/Plugins/Experimental/ControlRig/Source/ControlRig/Private/Units/Highlevel/Hierarchy/RigUnit_CCDIK.cpp
--- a/Engine/Plugins/Experimental/ControlRig/Source/ControlRig/Private/Units/Highlevel/Hierarchy/RigUnit_CCDIK.cpp
+++ b/Engine/Plugins/Experimental/ControlRig/Source/ControlRig/Private/Units/Highlevel/Hierarchy/RigUnit_CCDIK.cpp
@@ -99,13 +99,32 @@ FRigUnit_CCDIK_Execute()
 			// If we moved some bones, update bone transforms.
 			if (bBoneLocationUpdated)
 			{
-				for (int32 LinkIndex = 0; LinkIndex < BoneIndices.Num(); LinkIndex++)
+				if (FMath::IsNearlyEqual(Weight, 1.f))
 				{
-					const FCCDIKChainLink& CurrentLink = Chain[LinkIndex];
-					Hierarchy->SetGlobalTransform(BoneIndices[LinkIndex], CurrentLink.Transform, bPropagateToChildren);
+					for (int32 LinkIndex = 0; LinkIndex < BoneIndices.Num(); LinkIndex++)
+					{
+						const FCCDIKChainLink& CurrentLink = Chain[LinkIndex];
+						Hierarchy->SetGlobalTransform(BoneIndices[LinkIndex], CurrentLink.Transform, bPropagateToChildren);
+					}
+
+					Hierarchy->SetGlobalTransform(EffectorIndex, EffectorTransform, bPropagateToChildren);
 				}
+				else
+				{
+					float T = FMath::Clamp<float>(Weight, 0.f, 1.f);
 
-				Hierarchy->SetGlobalTransform(EffectorIndex, EffectorTransform, bPropagateToChildren);
+					for (int32 LinkIndex = 0; LinkIndex < BoneIndices.Num(); LinkIndex++)
+					{
+						const FCCDIKChainLink& CurrentLink = Chain[LinkIndex];
+						FTransform PreviousXfo = Hierarchy->GetGlobalTransform(BoneIndices[LinkIndex]);
+						FTransform Xfo = FControlRigMathLibrary::LerpTransform(PreviousXfo, CurrentLink.Transform, T);
+						Hierarchy->SetGlobalTransform(BoneIndices[LinkIndex], Xfo, bPropagateToChildren);
+					}
+
+					FTransform PreviousXfo = Hierarchy->GetGlobalTransform(EffectorIndex);
+					FTransform Xfo = FControlRigMathLibrary::LerpTransform(PreviousXfo, EffectorTransform, T);
+					Hierarchy->SetGlobalTransform(EffectorIndex, Xfo, bPropagateToChildren);
+				}
 			}
 		}
 	}
diff --git a/Engine/Plugins/Experimental/ControlRig/Source/ControlRig/Private/Units/Highlevel/Hierarchy/RigUnit_DistributeRotation.cpp b/Engine/Plugins/Experimental/ControlRig/Source/ControlRig/Private/Units/Highlevel/Hierarchy/RigUnit_DistributeRotation.cpp
--- a/Engine/Plugins/Experimental/ControlRig/Source/ControlRig/Private/Units/Highlevel/Hierarchy/RigUnit_DistributeRotation.cpp
+++ b/Engine/Plugins/Experimental/ControlRig/Source/ControlRig/Private/Units/Highlevel/Hierarchy/RigUnit_DistributeRotation.cpp
@@ -190,11 +190,11 @@ FRigUnit_DistributeRotation_Execute()
 
 IMPLEMENT_RIGUNIT_AUTOMATION_TEST(FRigUnit_DistributeRotation)
 {
-	BoneHierarchy.Add(TEXT("Root"), NAME_None, FTransform(FVector(1.f, 0.f, 0.f)));
-	BoneHierarchy.Add(TEXT("BoneA"), TEXT("Root"), FTransform(FVector(2.f, 0.f, 0.f)));
-	BoneHierarchy.Add(TEXT("BoneB"), TEXT("BoneA"), FTransform(FVector(2.f, 0.f, 0.f)));
-	BoneHierarchy.Add(TEXT("BoneC"), TEXT("BoneB"), FTransform(FVector(2.f, 0.f, 0.f)));
-	BoneHierarchy.Add(TEXT("BoneD"), TEXT("BoneC"), FTransform(FVector(2.f, 0.f, 0.f)));
+	BoneHierarchy.Add(TEXT("Root"), NAME_None, ERigBoneType::User, FTransform(FVector(1.f, 0.f, 0.f)));
+	BoneHierarchy.Add(TEXT("BoneA"), TEXT("Root"), ERigBoneType::User, FTransform(FVector(2.f, 0.f, 0.f)));
+	BoneHierarchy.Add(TEXT("BoneB"), TEXT("BoneA"), ERigBoneType::User, FTransform(FVector(2.f, 0.f, 0.f)));
+	BoneHierarchy.Add(TEXT("BoneC"), TEXT("BoneB"), ERigBoneType::User, FTransform(FVector(2.f, 0.f, 0.f)));
+	BoneHierarchy.Add(TEXT("BoneD"), TEXT("BoneC"), ERigBoneType::User, FTransform(FVector(2.f, 0.f, 0.f)));
 	BoneHierarchy.Initialize();
 	Unit.ExecuteContext.Hierarchy = &HierarchyContainer;
 
diff --git a/Engine/Plugins/Experimental/ControlRig/Source/ControlRig/Private/Units/Highlevel/Hierarchy/RigUnit_FABRIK.cpp b/Engine/Plugins/Experimental/ControlRig/Source/ControlRig/Private/Units/Highlevel/Hierarchy/RigUnit_FABRIK.cpp
--- a/Engine/Plugins/Experimental/ControlRig/Source/ControlRig/Private/Units/Highlevel/Hierarchy/RigUnit_FABRIK.cpp
+++ b/Engine/Plugins/Experimental/ControlRig/Source/ControlRig/Private/Units/Highlevel/Hierarchy/RigUnit_FABRIK.cpp
@@ -115,13 +115,32 @@ FRigUnit_FABRIK_Execute()
 				CurrentBoneTransform.SetTranslation(CurrentLink.Position);
 				CurrentBoneTransform.SetRotation(Hierarchy->GetGlobalTransform(CurrentLink.BoneIndex).GetRotation());
 
-				for (int32 LinkIndex = 0; LinkIndex < NumChainLinks; LinkIndex++)
+				if (FMath::IsNearlyEqual(Weight, 1.f))
 				{
-					FFABRIKChainLink const & LocalLink = Chain[LinkIndex];
-					Hierarchy->SetGlobalTransform(LocalLink.BoneIndex, Transforms[LocalLink.TransformIndex], bPropagateToChildren);
-				}
+					for (int32 LinkIndex = 0; LinkIndex < NumChainLinks; LinkIndex++)
+					{
+						FFABRIKChainLink const & LocalLink = Chain[LinkIndex];
+						Hierarchy->SetGlobalTransform(LocalLink.BoneIndex, Transforms[LocalLink.TransformIndex], bPropagateToChildren);
+					}
 
-				Hierarchy->SetGlobalTransform(EffectorIndex, EffectorTransform, bPropagateToChildren);
+					Hierarchy->SetGlobalTransform(EffectorIndex, EffectorTransform, bPropagateToChildren);
+				}
+				else
+				{
+					float T = FMath::Clamp<float>(Weight, 0.f, 1.f);
+
+					for (int32 LinkIndex = 0; LinkIndex < NumChainLinks; LinkIndex++)
+					{
+						FFABRIKChainLink const & LocalLink = Chain[LinkIndex];
+						FTransform PreviousXfo = Hierarchy->GetGlobalTransform(LocalLink.BoneIndex);
+						FTransform Xfo = FControlRigMathLibrary::LerpTransform(PreviousXfo, Transforms[LocalLink.TransformIndex], T);
+						Hierarchy->SetGlobalTransform(LocalLink.BoneIndex, Xfo, bPropagateToChildren);
+					}
+
+					FTransform PreviousXfo = Hierarchy->GetGlobalTransform(EffectorIndex);
+					FTransform Xfo = FControlRigMathLibrary::LerpTransform(PreviousXfo, EffectorTransform, T);
+					Hierarchy->SetGlobalTransform(EffectorIndex, Xfo, bPropagateToChildren);
+				}
 			}
 		}
 	}
diff --git a/Engine/Plugins/Experimental/ControlRig/Source/ControlRig/Private/Units/Highlevel/Hierarchy/RigUnit_ModifyBoneTransforms.cpp b/Engine/Plugins/Experimental/ControlRig/Source/ControlRig/Private/Units/Highlevel/Hierarchy/RigUnit_ModifyBoneTransforms.cpp
--- a/Engine/Plugins/Experimental/ControlRig/Source/ControlRig/Private/Units/Highlevel/Hierarchy/RigUnit_ModifyBoneTransforms.cpp
+++ b/Engine/Plugins/Experimental/ControlRig/Source/ControlRig/Private/Units/Highlevel/Hierarchy/RigUnit_ModifyBoneTransforms.cpp
@@ -111,9 +111,9 @@ FRigUnit_ModifyBoneTransforms_Execute()
 
 IMPLEMENT_RIGUNIT_AUTOMATION_TEST(FRigUnit_ModifyBoneTransforms)
 {
-	BoneHierarchy.Add(TEXT("Root"), NAME_None, FTransform(FVector(1.f, 0.f, 0.f)));
-	BoneHierarchy.Add(TEXT("BoneA"), TEXT("Root"), FTransform(FVector(1.f, 2.f, 3.f)));
-	BoneHierarchy.Add(TEXT("BoneB"), TEXT("Root"), FTransform(FVector(5.f, 6.f, 7.f)));
+	BoneHierarchy.Add(TEXT("Root"), NAME_None, ERigBoneType::User, FTransform(FVector(1.f, 0.f, 0.f)));
+	BoneHierarchy.Add(TEXT("BoneA"), TEXT("Root"), ERigBoneType::User, FTransform(FVector(1.f, 2.f, 3.f)));
+	BoneHierarchy.Add(TEXT("BoneB"), TEXT("Root"), ERigBoneType::User, FTransform(FVector(5.f, 6.f, 7.f)));
 	BoneHierarchy.Initialize();
 	Unit.ExecuteContext.Hierarchy = &HierarchyContainer;
 
