Copying //UE4/Dev-Anim to Dev-Main (//UE4/Dev-Main)

diff --git a/Engine/Plugins/Experimental/ControlRig/Source/ControlRig/Private/Units/Highlevel/Hierarchy/RigUnit_TwoBoneIKSimple.cpp b/Engine/Plugins/Experimental/ControlRig/Source/ControlRig/Private/Units/Highlevel/Hierarchy/RigUnit_TwoBoneIKSimple.cpp
--- a/Engine/Plugins/Experimental/ControlRig/Source/ControlRig/Private/Units/Highlevel/Hierarchy/RigUnit_TwoBoneIKSimple.cpp
+++ b/Engine/Plugins/Experimental/ControlRig/Source/ControlRig/Private/Units/Highlevel/Hierarchy/RigUnit_TwoBoneIKSimple.cpp
@@ -3,6 +3,7 @@
 #include "Units/Highlevel/Hierarchy/RigUnit_TwoBoneIKSimple.h"
 #include "Units/RigUnitContext.h"
 #include "Math/ControlRigMathLibrary.h"
+#include "TwoBoneIK.h"
 
 FRigUnit_TwoBoneIKSimple_Execute()
 {
@@ -70,7 +71,7 @@ FRigUnit_TwoBoneIKSimple_Execute()
 	TransformB.SetLocation(Hierarchy->GetGlobalTransform(BoneBIndex).GetLocation());
 	FTransform TransformC = Effector;
 
-	FControlRigMathLibrary::SolveBasicTwoBoneIK(TransformA, TransformB, TransformC, PoleTarget, PrimaryAxis, SecondaryAxis, LengthA, LengthB, bEnableStretch, StretchStartRatio, StretchMaximumRatio);
+	FControlRigMathLibrary::SolveBasicTwoBoneIK(TransformA, TransformB, TransformC, PoleTarget, PrimaryAxis, SecondaryAxis, SecondaryAxisWeight, LengthA, LengthB, bEnableStretch, StretchStartRatio, StretchMaximumRatio);
 
 	if (Context.DrawInterface != nullptr && DebugSettings.bEnabled)
 	{
@@ -96,4 +97,14 @@ FRigUnit_TwoBoneIKSimple_Execute()
 	Hierarchy->SetGlobalTransform(BoneAIndex, TransformA, bPropagateToChildren);
 	Hierarchy->SetGlobalTransform(BoneBIndex, TransformB, bPropagateToChildren);
 	Hierarchy->SetGlobalTransform(EffectorBoneIndex, TransformC, bPropagateToChildren);
+}
+
+FRigUnit_TwoBoneIKSimpleVectors_Execute()
+{
+	AnimationCore::SolveTwoBoneIK(Root, Elbow, Effector, PoleVector, Effector, Elbow, Effector, BoneALength, BoneBLength, bEnableStretch, StretchStartRatio, StretchMaximumRatio);
+}
+
+FRigUnit_TwoBoneIKSimpleTransforms_Execute()
+{
+	FControlRigMathLibrary::SolveBasicTwoBoneIK(Root, Elbow, Effector, PoleVector, PrimaryAxis, SecondaryAxis, SecondaryAxisWeight, BoneALength, BoneBLength, bEnableStretch, StretchStartRatio, StretchMaximumRatio);
 }
\ No newline at end of file
diff --git a/Engine/Plugins/Experimental/ControlRig/Source/ControlRig/Private/Units/Highlevel/Hierarchy/RigUnit_TwoBoneIKSimple.h b/Engine/Plugins/Experimental/ControlRig/Source/ControlRig/Private/Units/Highlevel/Hierarchy/RigUnit_TwoBoneIKSimple.h
--- a/Engine/Plugins/Experimental/ControlRig/Source/ControlRig/Private/Units/Highlevel/Hierarchy/RigUnit_TwoBoneIKSimple.h
+++ b/Engine/Plugins/Experimental/ControlRig/Source/ControlRig/Private/Units/Highlevel/Hierarchy/RigUnit_TwoBoneIKSimple.h
@@ -37,7 +37,7 @@ struct FRigUnit_TwoBoneIKSimple_DebugSettings
 };
 
 /**
- * Aligns the rotation of a primary and secondary axis of a bone to a world target.
+ * Solves the two bone IK given two bones.
  * Note: This node operates in world space!
  */
 USTRUCT(meta=(DisplayName="Basic IK", Category="Hierarchy", Keywords="TwoBone,IK"))
@@ -51,6 +51,7 @@ struct FRigUnit_TwoBoneIKSimple : public FRigUnit_HighlevelBaseMutable
 		Effector = FTransform::Identity;
 		PrimaryAxis = FVector(1.f, 0.f, 0.f);
 		SecondaryAxis = FVector(0.f, 1.f, 0.f);
+		SecondaryAxisWeight = 1.f;
 		PoleVector = FVector(0.f, 0.f, 1.f);
 		PoleVectorKind = EControlRigVectorKind::Direction;
 		bEnableStretch = false;
@@ -106,6 +107,12 @@ struct FRigUnit_TwoBoneIKSimple : public FRigUnit_HighlevelBaseMutable
 	UPROPERTY(meta = (Input))
 	FVector SecondaryAxis;
 
+	/**
+	 * Determines how much the secondary axis roll is being applied
+	 */
+	UPROPERTY(meta = (Input))
+	float SecondaryAxisWeight;
+
 	/** 
 	 * The polevector to use for the IK solver
 	 * This can be a location or direction.
@@ -186,3 +193,180 @@ struct FRigUnit_TwoBoneIKSimple : public FRigUnit_HighlevelBaseMutable
 	UPROPERTY()
 	int32 PoleVectorSpaceIndex;
 };
+
+/**
+ * Solves the two bone IK given positions
+ * Note: This node operates in world space!
+ */
+USTRUCT(meta = (DisplayName = "Basic IK Positions", Category = "Hierarchy", Keywords = "TwoBone,IK"))
+struct FRigUnit_TwoBoneIKSimpleVectors : public FRigUnit_HighlevelBase
+{
+	GENERATED_BODY()
+
+	FRigUnit_TwoBoneIKSimpleVectors()
+	{
+		Root = PoleVector = Effector = Elbow = FVector::ZeroVector;
+		bEnableStretch = false;
+		StretchStartRatio = 0.75f;
+		StretchMaximumRatio = 1.25f;
+		BoneALength = BoneBLength = 1.f;
+	}
+
+	RIGVM_METHOD()
+	virtual void Execute(const FRigUnitContext& Context) override;
+
+	/**
+	 * The position of the root of the triangle
+	 */
+	UPROPERTY(meta = (Input))
+	FVector Root;
+
+	/**
+	 * The position of the pole of the triangle
+	 */
+	UPROPERTY(meta = (Input))
+	FVector PoleVector;
+
+	/**
+	 * The position of the effector
+	 */
+	UPROPERTY(meta = (Input, Output))
+	FVector Effector;
+
+	/**
+	 * If set to true the stretch feature of the solver will be enabled
+	 */
+	UPROPERTY(meta = (Input))
+	bool bEnableStretch;
+
+	/**
+	 * The ratio where the stretch starts
+	 */
+	UPROPERTY(meta = (Input))
+	float StretchStartRatio;
+
+	/**
+	 * The maximum allowed stretch ratio
+	 */
+	UPROPERTY(meta = (Input))
+	float StretchMaximumRatio;
+
+	/**
+	 * The length of the first bone.
+	 * If set to 0.0 it will be determined by the hierarchy
+	 */
+	UPROPERTY(meta = (Input))
+	float BoneALength;
+
+	/**
+	 * The length of the second  bone.
+	 * If set to 0.0 it will be determined by the hierarchy
+	 */
+	UPROPERTY(meta = (Input))
+	float BoneBLength;
+
+	/**
+	 * The resulting elbow position
+	 */
+	UPROPERTY(meta = (Output))
+	FVector Elbow;
+};
+
+/**
+ * Solves the two bone IK given transforms
+ * Note: This node operates in world space!
+ */
+USTRUCT(meta = (DisplayName = "Basic IK Transforms", Category = "Hierarchy", Keywords = "TwoBone,IK"))
+struct FRigUnit_TwoBoneIKSimpleTransforms : public FRigUnit_HighlevelBase
+{
+	GENERATED_BODY()
+
+	FRigUnit_TwoBoneIKSimpleTransforms()
+	{
+		Root = Effector = Elbow = FTransform::Identity;
+		PoleVector = FVector::ZeroVector;
+		PrimaryAxis = FVector(1.f, 0.f, 0.f);
+		SecondaryAxis = FVector(0.f, 1.f, 0.f);
+		SecondaryAxisWeight = 1.f;
+		bEnableStretch = false;
+		StretchStartRatio = 0.75f;
+		StretchMaximumRatio = 1.25f;
+		BoneALength = BoneBLength = 1.f;
+	}
+
+	RIGVM_METHOD()
+	virtual void Execute(const FRigUnitContext& Context) override;
+
+	/**
+	 * The transform of the root of the triangle
+	 */
+	UPROPERTY(meta = (Input, Output))
+	FTransform Root;
+
+	/**
+	 * The position of the pole of the triangle
+	 */
+	UPROPERTY(meta = (Input))
+	FVector PoleVector;
+
+	/**
+	 * The transform of the effector
+	 */
+	UPROPERTY(meta = (Input, Output))
+	FTransform Effector;
+		/** 
+	 * The major axis being aligned - along the bone
+	 */
+	UPROPERTY(meta = (Input))
+	FVector PrimaryAxis;
+
+	/** 
+	 * The minor axis being aligned - towards the polevector
+	 */
+	UPROPERTY(meta = (Input))
+	FVector SecondaryAxis;
+
+	/**
+	 * Determines how much the secondary axis roll is being applied
+	 */
+	UPROPERTY(meta = (Input))
+	float SecondaryAxisWeight;
+
+	/**
+	 * If set to true the stretch feature of the solver will be enabled
+	 */
+	UPROPERTY(meta = (Input))
+	bool bEnableStretch;
+
+	/**
+	 * The ratio where the stretch starts
+	 */
+	UPROPERTY(meta = (Input))
+	float StretchStartRatio;
+
+	/**
+	 * The maximum allowed stretch ratio
+	 */
+	UPROPERTY(meta = (Input))
+	float StretchMaximumRatio;
+
+	/**
+	 * The length of the first bone.
+	 * If set to 0.0 it will be determined by the hierarchy
+	 */
+	UPROPERTY(meta = (Input))
+	float BoneALength;
+
+	/**
+	 * The length of the second  bone.
+	 * If set to 0.0 it will be determined by the hierarchy
+	 */
+	UPROPERTY(meta = (Input))
+	float BoneBLength;
+
+	/**
+	 * The resulting elbow transform
+	 */
+	UPROPERTY(meta = (Input, Output))
+	FTransform Elbow;
+};
diff --git a/Engine/Plugins/Experimental/ControlRig/Source/ControlRig/Private/Units/Math/RigUnit_MathTransform.cpp b/Engine/Plugins/Experimental/ControlRig/Source/ControlRig/Private/Units/Math/RigUnit_MathTransform.cpp
--- a/Engine/Plugins/Experimental/ControlRig/Source/ControlRig/Private/Units/Math/RigUnit_MathTransform.cpp
+++ b/Engine/Plugins/Experimental/ControlRig/Source/ControlRig/Private/Units/Math/RigUnit_MathTransform.cpp
@@ -1,6 +1,7 @@
 // Copyright 1998-2019 Epic Games, Inc. All Rights Reserved.
 
 #include "Units/Math/RigUnit_MathTransform.h"
+#include "Units/Math/RigUnit_MathVector.h"
 #include "Math/ControlRigMathLibrary.h"
 #include "Units/RigUnitContext.h"
 
@@ -26,6 +27,14 @@ FRigUnit_MathTransformMakeRelative_Execute()
 {
     DECLARE_SCOPE_HIERARCHICAL_COUNTER_RIGUNIT()
 	Local = Global.GetRelativeTransform(Parent);
+	Local.NormalizeRotation();
+}
+
+FRigUnit_MathTransformMakeAbsolute_Execute()
+{
+	DECLARE_SCOPE_HIERARCHICAL_COUNTER_RIGUNIT()
+	Global = Local * Parent;
+	Global.NormalizeRotation();
 }
 
 FRigUnit_MathTransformInverse_Execute()
@@ -66,3 +75,12 @@ FRigUnit_MathTransformFromSRT_Execute()
 	Transform.SetScale3D(Scale);
 	EulerTransform.FromFTransform(Transform);
 }
+
+FRigUnit_MathTransformClampSpatially_Execute()
+{
+	DECLARE_SCOPE_HIERARCHICAL_COUNTER_RIGUNIT()
+	FVector Position;
+	FRigUnit_MathVectorClampSpatially::StaticExecute(Value.GetTranslation(), Axis, Type, Minimum, Maximum, Space, bDrawDebug, DebugColor, DebugThickness, Position, RigUnitName, RigUnitStructName, ExecutionType, Context);
+	Result = Value;
+	Result.SetTranslation(Position);
+}
\ No newline at end of file
diff --git a/Engine/Plugins/Experimental/ControlRig/Source/ControlRig/Private/Units/Math/RigUnit_MathTransform.h b/Engine/Plugins/Experimental/ControlRig/Source/ControlRig/Private/Units/Math/RigUnit_MathTransform.h
--- a/Engine/Plugins/Experimental/ControlRig/Source/ControlRig/Private/Units/Math/RigUnit_MathTransform.h
+++ b/Engine/Plugins/Experimental/ControlRig/Source/ControlRig/Private/Units/Math/RigUnit_MathTransform.h
@@ -112,7 +112,7 @@ struct FRigUnit_MathTransformMul : public FRigUnit_MathTransformBinaryOp
 /**
  * Returns the relative local transform within a parent's transform
  */
-USTRUCT(meta=(DisplayName="Make Relative", PrototypeName="MakeRelative", Keywords="Local,Global"))
+USTRUCT(meta=(DisplayName="Make Relative", PrototypeName="MakeRelative", Keywords="Local,Global,Absolute"))
 struct FRigUnit_MathTransformMakeRelative : public FRigUnit_MathTransformBase
 {
 	GENERATED_BODY()
@@ -135,6 +135,32 @@ struct FRigUnit_MathTransformMakeRelative : public FRigUnit_MathTransformBase
 	FTransform Local;
 };
 
+/**
+ * Returns the absolute global transform within a parent's transform
+ */
+USTRUCT(meta = (DisplayName = "Make Absolute", PrototypeName = "MakeAbsolute", Keywords = "Local,Global,Relative"))
+struct FRigUnit_MathTransformMakeAbsolute : public FRigUnit_MathTransformBase
+{
+	GENERATED_BODY()
+
+	FRigUnit_MathTransformMakeAbsolute()
+	{
+		Global = Parent = Local = FTransform::Identity;
+	}
+
+	RIGVM_METHOD()
+	virtual void Execute(const FRigUnitContext& Context) override;
+
+	UPROPERTY(meta = (Input))
+	FTransform Local;
+
+	UPROPERTY(meta = (Input))
+	FTransform Parent;
+
+	UPROPERTY(meta = (Output))
+	FTransform Global;
+};
+
 /**
  * Returns the negative value
  */
@@ -263,7 +289,7 @@ struct FRigUnit_MathTransformTransformVector : public FRigUnit_MathTransformBase
 /**
  * Composes a Transform (and Euler Transform) from its components.
  */
-USTRUCT(meta=(DisplayName="Transform from SRT", PrototypeName="EulerTransform,Scale,Rotation,Orientation,Translation,Location"))
+USTRUCT(meta=(DisplayName="Transform from SRT", Keywords ="EulerTransform,Scale,Rotation,Orientation,Translation,Location"))
 struct FRigUnit_MathTransformFromSRT : public FRigUnit_MathTransformBase
 {
 	GENERATED_BODY()
@@ -298,4 +324,61 @@ struct FRigUnit_MathTransformFromSRT : public FRigUnit_MathTransformBase
 
 	UPROPERTY(meta=(Output))
 	FEulerTransform EulerTransform;
+};
+
+/**
+ * Clamps a position using a plane collision, cylindric collision or spherical collision.
+ */
+USTRUCT(meta = (DisplayName = "Clamp Spatially", PrototypeName = "ClampSpatially", Keywords = "Collide,Collision"))
+struct FRigUnit_MathTransformClampSpatially : public FRigUnit_MathTransformBase
+{
+	GENERATED_BODY()
+
+	FRigUnit_MathTransformClampSpatially()
+	{
+		Value = Result = FTransform::Identity;
+		Axis = EAxis::X;
+		Type = EControlRigClampSpatialMode::Plane;
+		Minimum = 0.f;
+		Maximum = 100.f;
+		Space = FTransform::Identity;
+		bDrawDebug = false;
+		DebugColor = FLinearColor::Red;
+		DebugThickness = 1.f;
+	}
+
+	RIGVM_METHOD()
+	virtual void Execute(const FRigUnitContext& Context) override;
+
+	UPROPERTY(meta = (Input))
+	FTransform Value;
+
+	UPROPERTY(meta = (Input))
+	TEnumAsByte<EAxis::Type> Axis;
+
+	UPROPERTY(meta = (Input))
+	TEnumAsByte<EControlRigClampSpatialMode::Type> Type;
+
+	UPROPERTY(meta = (Input))
+	float Minimum;
+
+	UPROPERTY(meta = (Input))
+	float Maximum;
+
+	// The space this spatial clamp happens within.
+	// The input position will be projected into this space.
+	UPROPERTY(meta = (Input))
+	FTransform Space;
+
+	UPROPERTY(meta = (Input))
+	bool bDrawDebug;
+
+	UPROPERTY(meta = (Input))
+	FLinearColor DebugColor;
+
+	UPROPERTY(meta = (Input))
+	float DebugThickness;
+
+	UPROPERTY(meta = (Output))
+	FTransform Result;
 };
\ No newline at end of file
diff --git a/Engine/Plugins/Experimental/ControlRig/Source/ControlRig/Private/Units/Math/RigUnit_MathVector.cpp b/Engine/Plugins/Experimental/ControlRig/Source/ControlRig/Private/Units/Math/RigUnit_MathVector.cpp
--- a/Engine/Plugins/Experimental/ControlRig/Source/ControlRig/Private/Units/Math/RigUnit_MathVector.cpp
+++ b/Engine/Plugins/Experimental/ControlRig/Source/ControlRig/Private/Units/Math/RigUnit_MathVector.cpp
@@ -281,6 +281,31 @@ FRigUnit_MathVectorUnit_Execute()
 	Result = Value.GetUnsafeNormal();
 }
 
+FRigUnit_MathVectorSetLength_Execute()
+{
+	DECLARE_SCOPE_HIERARCHICAL_COUNTER_RIGUNIT()
+	if (Value.IsNearlyZero())
+	{
+		UE_CONTROLRIG_RIGUNIT_REPORT_WARNING(TEXT("Value is nearly zero"));
+		Result = FVector::ZeroVector;
+		return;
+	}
+	Result = Value.GetUnsafeNormal() * Length;
+}
+
+FRigUnit_MathVectorClampLength_Execute()
+{
+	DECLARE_SCOPE_HIERARCHICAL_COUNTER_RIGUNIT()
+	if (Value.IsNearlyZero())
+	{
+		UE_CONTROLRIG_RIGUNIT_REPORT_WARNING(TEXT("Value is nearly zero"));
+		Result = FVector::ZeroVector;
+		return;
+	}
+	float Length = Value.Size();
+	Result = Value * FMath::Clamp<float>(Length, MinimumLength, MaximumLength) / Length;
+}
+
 FRigUnit_MathVectorMirror_Execute()
 {
     DECLARE_SCOPE_HIERARCHICAL_COUNTER_RIGUNIT()
@@ -354,4 +379,104 @@ FRigUnit_MathVectorBezierFourPoint_Execute()
 {
     DECLARE_SCOPE_HIERARCHICAL_COUNTER_RIGUNIT()
 	FControlRigMathLibrary::FourPointBezier(Bezier, T, Result, Tangent);
+}
+
+FRigUnit_MathVectorMakeBezierFourPoint_Execute()
+{
+}
+
+FRigUnit_MathVectorClampSpatially_Execute()
+{
+	DECLARE_SCOPE_HIERARCHICAL_COUNTER_RIGUNIT()
+	Result = FControlRigMathLibrary::ClampSpatially(Value, Axis, Type, Minimum, Maximum, Space);
+
+	if (Context.DrawInterface != nullptr && bDrawDebug)
+	{
+		switch (Type)
+		{
+			case EControlRigClampSpatialMode::Plane:
+			{
+				TArray<FVector> Points;
+				Points.SetNumUninitialized(2);
+					
+				switch (Axis)
+				{
+					case EAxis::X:
+					{
+						Points[0] = FVector(Minimum, 0.f, 0.f);
+						Points[1] = FVector(Maximum, 0.f, 0.f);
+						break;
+					}
+					case EAxis::Y:
+					{
+						Points[0] = FVector(0.f, Minimum, 0.f);
+						Points[1] = FVector(0.f, Maximum, 0.f);
+						break;
+					}
+					default:
+					{
+						Points[0] = FVector(0.f, 0.f, Minimum);
+						Points[1] = FVector(0.f, 0.f, Maximum);
+						break;
+					}
+				}
+
+				Context.DrawInterface->DrawLine(Space, Points[0], Points[1], DebugColor, DebugThickness);
+				Context.DrawInterface->DrawPoints(Space, Points, DebugThickness * 8.f, DebugColor);
+
+				break;
+			}
+			case EControlRigClampSpatialMode::Cylinder:
+			{
+				FTransform CircleTransform = FTransform::Identity;
+				switch (Axis)
+				{
+					case EAxis::X:
+					{
+						CircleTransform.SetRotation(FQuat(FVector(0.f, 1.f, 0.f), PI * 0.5f));
+						break;
+					}
+					case EAxis::Y:
+					{
+						CircleTransform.SetRotation(FQuat(FVector(1.f, 0.f, 0.f), PI * 0.5f));
+						break;
+					}
+					default:
+					{
+						break;
+					}
+				}
+				if (Minimum > SMALL_NUMBER)
+				{
+					Context.DrawInterface->DrawArc(Space, CircleTransform, Minimum, 0.f, PI * 2.f, DebugColor, DebugThickness, 16);
+				}
+				Context.DrawInterface->DrawArc(Space, CircleTransform, Maximum, 0.f, PI * 2.f, DebugColor, DebugThickness, 16);
+				break;
+			}
+			default:
+			case EControlRigClampSpatialMode::Sphere:
+			{
+				FTransform CircleTransform = FTransform::Identity;
+				if (Minimum > SMALL_NUMBER)
+				{
+					Context.DrawInterface->DrawArc(Space, CircleTransform, Minimum, 0.f, PI * 2.f, DebugColor, DebugThickness, 16);
+				}
+				Context.DrawInterface->DrawArc(Space, CircleTransform, Maximum, 0.f, PI * 2.f, DebugColor, DebugThickness, 16);
+				CircleTransform.SetRotation(FQuat(FVector(0.f, 1.f, 0.f), PI * 0.5f));
+				if (Minimum > SMALL_NUMBER)
+				{
+					Context.DrawInterface->DrawArc(Space, CircleTransform, Minimum, 0.f, PI * 2.f, DebugColor, DebugThickness, 16);
+				}
+				Context.DrawInterface->DrawArc(Space, CircleTransform, Maximum, 0.f, PI * 2.f, DebugColor, DebugThickness, 16);
+				CircleTransform.SetRotation(FQuat(FVector(1.f, 0.f, 0.f), PI * 0.5f));
+				if (Minimum > SMALL_NUMBER)
+				{
+					Context.DrawInterface->DrawArc(Space, CircleTransform, Minimum, 0.f, PI * 2.f, DebugColor, DebugThickness, 16);
+				}
+				Context.DrawInterface->DrawArc(Space, CircleTransform, Maximum, 0.f, PI * 2.f, DebugColor, DebugThickness, 16);
+				break;
+			}
+
+		}
+	}
 }
\ No newline at end of file
diff --git a/Engine/Plugins/Experimental/ControlRig/Source/ControlRig/Private/Units/Math/RigUnit_MathVector.h b/Engine/Plugins/Experimental/ControlRig/Source/ControlRig/Private/Units/Math/RigUnit_MathVector.h
--- a/Engine/Plugins/Experimental/ControlRig/Source/ControlRig/Private/Units/Math/RigUnit_MathVector.h
+++ b/Engine/Plugins/Experimental/ControlRig/Source/ControlRig/Private/Units/Math/RigUnit_MathVector.h
@@ -665,6 +665,65 @@ struct FRigUnit_MathVectorUnit : public FRigUnit_MathVectorUnaryOp
 	virtual void Execute(const FRigUnitContext& Context) override;
 };
 
+/**
+ * Sets the length of a given vector
+ */
+USTRUCT(meta = (DisplayName = "SetLength", PrototypeName = "SetLength", Keywords = "Unit,Normalize,Scale"))
+struct FRigUnit_MathVectorSetLength: public FRigUnit_MathVectorBase
+{
+	GENERATED_BODY()
+
+	FRigUnit_MathVectorSetLength()
+	{
+		Value = Result = FVector(1.f, 0.f, 0.f);
+		Length = 1.f;
+	}
+
+	RIGVM_METHOD()
+	virtual void Execute(const FRigUnitContext& Context) override;
+
+	UPROPERTY(meta = (Input))
+	FVector Value;
+
+	UPROPERTY(meta = (Input))
+	float Length;
+
+	UPROPERTY(meta = (Output))
+	FVector Result;
+};
+
+/**
+ * Clamps the length of a given vector between a minimum and maximum
+ */
+USTRUCT(meta = (DisplayName = "ClampLength", PrototypeName = "ClampLength", Keywords = "Unit,Normalize,Scale"))
+struct FRigUnit_MathVectorClampLength: public FRigUnit_MathVectorBase
+{
+	GENERATED_BODY()
+
+	FRigUnit_MathVectorClampLength()
+	{
+		Value = Result = FVector(1.f, 0.f, 0.f);
+		MinimumLength = 0.f;
+		MaximumLength = 1.f;
+	}
+
+
+	RIGVM_METHOD()
+	virtual void Execute(const FRigUnitContext& Context) override;
+
+	UPROPERTY(meta = (Input))
+	FVector Value;
+
+	UPROPERTY(meta = (Input))
+	float MinimumLength;
+
+	UPROPERTY(meta = (Input))
+	float MaximumLength;
+
+	UPROPERTY(meta = (Output))
+	FVector Result;
+};
+
 /**
  * Mirror a vector about a normal vector.
  */
@@ -818,6 +877,66 @@ struct FRigUnit_MathVectorMakeBezierFourPoint : public FRigUnit_MathVectorBase
 		Bezier = FCRFourPointBezier();
 	}
 
+	RIGVM_METHOD()
+	virtual void Execute(const FRigUnitContext& Context) override;
+
 	UPROPERTY(meta = (Input, Output))
 	FCRFourPointBezier Bezier;
+};
+
+/**
+ * Clamps a position using a plane collision, cylindric collision or spherical collision.
+ */
+USTRUCT(meta = (DisplayName = "Clamp Spatially", PrototypeName = "ClampSpatially", Keywords="Collide,Collision"))
+struct FRigUnit_MathVectorClampSpatially: public FRigUnit_MathVectorBase
+{
+	GENERATED_BODY()
+
+	FRigUnit_MathVectorClampSpatially()
+	{
+		Value = Result = FVector::ZeroVector;
+		Axis = EAxis::X;
+		Type = EControlRigClampSpatialMode::Plane;
+		Minimum = 0.f;
+		Maximum = 100.f;
+		Space = FTransform::Identity;
+		bDrawDebug = false;
+		DebugColor = FLinearColor::Red;
+		DebugThickness = 1.f;
+	}
+
+	RIGVM_METHOD()
+	virtual void Execute(const FRigUnitContext& Context) override;
+
+	UPROPERTY(meta = (Input))
+	FVector Value;
+
+	UPROPERTY(meta = (Input))
+	TEnumAsByte<EAxis::Type> Axis;
+
+	UPROPERTY(meta = (Input))
+	TEnumAsByte<EControlRigClampSpatialMode::Type> Type;
+
+	UPROPERTY(meta = (Input))
+	float Minimum;
+
+	UPROPERTY(meta = (Input))
+	float Maximum;
+
+	// The space this spatial clamp happens within.
+	// The input position will be projected into this space.
+	UPROPERTY(meta = (Input))
+	FTransform Space;
+
+	UPROPERTY(meta = (Input))
+	bool bDrawDebug;
+
+	UPROPERTY(meta = (Input))
+	FLinearColor DebugColor;
+
+	UPROPERTY(meta = (Input))
+	float DebugThickness;
+
+	UPROPERTY(meta = (Output))
+	FVector Result;
 };
\ No newline at end of file
diff --git a/Engine/Plugins/Experimental/ControlRig/Source/ControlRig/Private/Units/Math/RigUnit_MathVectorTests.cpp b/Engine/Plugins/Experimental/ControlRig/Source/ControlRig/Private/Units/Math/RigUnit_MathVectorTests.cpp
--- a/Engine/Plugins/Experimental/ControlRig/Source/ControlRig/Private/Units/Math/RigUnit_MathVectorTests.cpp
+++ b/Engine/Plugins/Experimental/ControlRig/Source/ControlRig/Private/Units/Math/RigUnit_MathVectorTests.cpp
@@ -321,6 +321,35 @@ IMPLEMENT_RIGUNIT_AUTOMATION_TEST(FRigUnit_MathVectorUnit)
 	return true;
 }
 
+IMPLEMENT_RIGUNIT_AUTOMATION_TEST(FRigUnit_MathVectorSetLength)
+{
+	Unit.Value = FVector(4.f, 0.f, 0.f);
+	Unit.Length = 2.f;
+	InitAndExecute();
+	AddErrorIfFalse(FRigUnit_MathVectorTest_Utils::IsNearlyEqual(Unit.Result, FVector(2.f, 0.f, 0.f), 0.001f), TEXT("unexpected result (0)"));
+	Unit.Value = FVector(0.f, 5.f, 0.f);
+	Unit.Length = 3.f;
+	InitAndExecute();
+	AddErrorIfFalse(FRigUnit_MathVectorTest_Utils::IsNearlyEqual(Unit.Result, FVector(0.f, 3.f, 0.f), 0.001f), TEXT("unexpected result (1)"));
+	return true;
+}
+
+IMPLEMENT_RIGUNIT_AUTOMATION_TEST(FRigUnit_MathVectorClampLength)
+{
+	Unit.Value = FVector(4.f, 0.f, 0.f);
+	Unit.MinimumLength = 2.f;
+	Unit.MaximumLength = 6.f;
+	InitAndExecute();
+	AddErrorIfFalse(FRigUnit_MathVectorTest_Utils::IsNearlyEqual(Unit.Result, FVector(4.f, 0.f, 0.f), 0.001f), TEXT("unexpected result (0)"));
+	Unit.Value = FVector(1.f, 0.f, 0.f);
+	InitAndExecute();
+	AddErrorIfFalse(FRigUnit_MathVectorTest_Utils::IsNearlyEqual(Unit.Result, FVector(2.f, 0.f, 0.f), 0.001f), TEXT("unexpected result (1)"));
+	Unit.Value = FVector(8.f, 0.f, 0.f);
+	InitAndExecute();
+	AddErrorIfFalse(FRigUnit_MathVectorTest_Utils::IsNearlyEqual(Unit.Result, FVector(6.f, 0.f, 0.f), 0.001f), TEXT("unexpected result (1)"));
+	return true;
+}
+
 IMPLEMENT_RIGUNIT_AUTOMATION_TEST(FRigUnit_MathVectorMirror)
 {
 	Unit.Value = FVector(1.f, 2.f, 3.f);
diff --git a/Engine/Plugins/Experimental/ControlRig/Source/ControlRig/Private/Units/Simulation/RigUnit_AlphaInterp.cpp b/Engine/Plugins/Experimental/ControlRig/Source/ControlRig/Private/Units/Simulation/RigUnit_AlphaInterp.cpp
new file mode 100644
--- /dev/null
+++ b/Engine/Plugins/Experimental/ControlRig/Source/ControlRig/Private/Units/Simulation/RigUnit_AlphaInterp.cpp
@@ -0,0 +1,29 @@
+// Copyright 1998-2019 Epic Games, Inc. All Rights Reserved.
+
+#include "Units/Simulation/RigUnit_AlphaInterp.h"
+#include "Units/RigUnitContext.h"
+
+FRigUnit_AlphaInterp_Execute()
+{
+    DECLARE_SCOPE_HIERARCHICAL_COUNTER_RIGUNIT()
+	if (Context.State == EControlRigState::Init)
+	{
+		ScaleBiasClamp.Reinitialize();
+		ScaleBiasClamp.bMapRange = bMapRange;
+		ScaleBiasClamp.bClampResult = bClampResult;
+		ScaleBiasClamp.bInterpResult = bInterpResult;
+		ScaleBiasClamp.InRange = InRange;
+		ScaleBiasClamp.OutRange = OutRange;
+		ScaleBiasClamp.ClampMin = ClampMin;
+		ScaleBiasClamp.ClampMax = ClampMax;
+	}
+	else
+	{
+		ScaleBiasClamp.Scale = Scale;
+		ScaleBiasClamp.Bias = Bias;
+		ScaleBiasClamp.InterpSpeedIncreasing = InterpSpeedIncreasing;
+		ScaleBiasClamp.InterpSpeedDecreasing = InterpSpeedDecreasing;
+
+		Result = ScaleBiasClamp.ApplyTo(Value, Context.DeltaTime);
+	}
+}
diff --git a/Engine/Plugins/Experimental/ControlRig/Source/ControlRig/Private/Units/Simulation/RigUnit_AlphaInterp.h b/Engine/Plugins/Experimental/ControlRig/Source/ControlRig/Private/Units/Simulation/RigUnit_AlphaInterp.h
new file mode 100644
--- /dev/null
+++ b/Engine/Plugins/Experimental/ControlRig/Source/ControlRig/Private/Units/Simulation/RigUnit_AlphaInterp.h
@@ -0,0 +1,78 @@
+// Copyright 1998-2019 Epic Games, Inc. All Rights Reserved.
+
+#pragma once
+
+#include "Units/Simulation/RigUnit_SimBase.h"
+#include "Animation/InputScaleBias.h"
+#include "RigUnit_AlphaInterp.generated.h"
+
+/**
+ * Adds a value over time over and over again
+ */
+USTRUCT(meta=(DisplayName="Alpha Interpolate", Keywords="Alpha,Lerp,LinearInterpolate"))
+struct FRigUnit_AlphaInterp : public FRigUnit_SimBase
+{
+	GENERATED_BODY()
+	
+	FRigUnit_AlphaInterp()
+	{
+		Value = Result = 0.f;
+		ScaleBiasClamp = FInputScaleBiasClamp();
+		bMapRange = ScaleBiasClamp.bMapRange;
+		bClampResult = ScaleBiasClamp.bClampResult;
+		bInterpResult = ScaleBiasClamp.bInterpResult;
+		InRange = ScaleBiasClamp.InRange;
+		OutRange = ScaleBiasClamp.OutRange;
+		Scale = ScaleBiasClamp.Scale;
+		Bias = ScaleBiasClamp.Bias;
+		ClampMin = ScaleBiasClamp.ClampMin;
+		ClampMax = ScaleBiasClamp.ClampMax;
+		InterpSpeedIncreasing = ScaleBiasClamp.InterpSpeedIncreasing;
+		InterpSpeedDecreasing = ScaleBiasClamp.InterpSpeedDecreasing;
+	}
+
+	RIGVM_METHOD()
+	virtual void Execute(const FRigUnitContext& Context) override;
+
+	UPROPERTY(meta=(Input))
+	float Value;
+
+	UPROPERTY(meta=(Input, Constant))
+	bool bMapRange;
+
+	UPROPERTY(meta=(Input, Constant))
+	bool bClampResult;
+
+	UPROPERTY(meta=(Input, Constant))
+	bool bInterpResult;
+
+	UPROPERTY(meta=(Input, Constant))
+	FInputRange InRange;
+
+	UPROPERTY(meta=(Input, Constant))
+	FInputRange OutRange;
+
+	UPROPERTY(meta=(Input))
+	float Scale;
+
+	UPROPERTY(meta=(Input))
+	float Bias;
+
+	UPROPERTY(meta=(Input, Constant))
+	float ClampMin;
+
+	UPROPERTY(meta=(Input, Constant))
+	float ClampMax;
+
+	UPROPERTY(meta=(Input))
+	float InterpSpeedIncreasing;
+
+	UPROPERTY(meta=(Input))
+	float InterpSpeedDecreasing;
+
+	UPROPERTY(meta=(Output))
+	float Result;
+
+	UPROPERTY()
+	FInputScaleBiasClamp ScaleBiasClamp;
+};
diff --git a/Engine/Plugins/Experimental/ControlRig/Source/ControlRig/Public/AnimNode_ControlRig.h b/Engine/Plugins/Experimental/ControlRig/Source/ControlRig/Public/AnimNode_ControlRig.h
--- a/Engine/Plugins/Experimental/ControlRig/Source/ControlRig/Public/AnimNode_ControlRig.h
+++ b/Engine/Plugins/Experimental/ControlRig/Source/ControlRig/Public/AnimNode_ControlRig.h
@@ -3,6 +3,7 @@
 #pragma once
 
 #include "ControlRig.h"
+#include "Animation/InputScaleBias.h"
 #include "AnimNode_ControlRigBase.h"
 #include "AnimNode_ControlRig.generated.h"
 
@@ -33,9 +34,6 @@ struct CONTROLRIG_API FAnimNode_ControlRig : public FAnimNode_ControlRigBase
 	FName GetIOMapping(bool bInput, const FName& SourceProperty) const;
 private:
 
-	UPROPERTY(EditAnywhere, Category = Links)
-	FPoseLink Source;
-
 	/** Cached ControlRig */
 	UPROPERTY(EditAnywhere, Category = ControlRig)
 	TSubclassOf<UControlRig> ControlRigClass;
@@ -43,6 +41,28 @@ private:
 	UPROPERTY(transient)
 	UControlRig* ControlRig;
 
+	// alpha value handler
+	UPROPERTY(EditAnywhere, Category = Settings, meta = (PinShownByDefault))
+	float Alpha;
+
+	UPROPERTY(EditAnywhere, Category = Settings)
+	EAnimAlphaInputType AlphaInputType;
+
+	UPROPERTY(EditAnywhere, Category = Settings, meta = (PinShownByDefault, DisplayName = "bEnabled", DisplayAfter = "AlphaScaleBias"))
+	uint8 bAlphaBoolEnabled : 1;
+
+	UPROPERTY(EditAnywhere, Category=Settings)
+	FInputScaleBias AlphaScaleBias;
+
+	UPROPERTY(EditAnywhere, Category = Settings, meta = (DisplayName = "Blend Settings"))
+	FInputAlphaBoolBlend AlphaBoolBlend;
+
+	UPROPERTY(EditAnywhere, Category = Settings, meta = (PinShownByDefault))
+	FName AlphaCurveName;
+
+	UPROPERTY(EditAnywhere, Category = Settings)
+	FInputScaleBiasClamp AlphaScaleBiasClamp;
+
 	// we only save mapping, 
 	// we have to query control rig when runtime 
 	// to ensure type and everything is still valid or not
diff --git a/Engine/Plugins/Experimental/ControlRig/Source/ControlRig/Public/AnimNode_ControlRigBase.h b/Engine/Plugins/Experimental/ControlRig/Source/ControlRig/Public/AnimNode_ControlRigBase.h
--- a/Engine/Plugins/Experimental/ControlRig/Source/ControlRig/Public/AnimNode_ControlRigBase.h
+++ b/Engine/Plugins/Experimental/ControlRig/Source/ControlRig/Public/AnimNode_ControlRigBase.h
@@ -30,6 +30,10 @@ struct CONTROLRIG_API FAnimNode_ControlRigBase : public FAnimNode_CustomProperty
 	virtual void CacheBones_AnyThread(const FAnimationCacheBonesContext& Context) override;
 
 protected:
+
+	UPROPERTY(EditAnywhere, Category = Links)
+	FPoseLink Source;
+
 	/** Rig Hierarchy bone name to required array index mapping */
 	UPROPERTY(transient)
 	TMap<FName, uint16> ControlRigBoneMapping;
@@ -51,10 +55,16 @@ protected:
 	UPROPERTY(transient)
 	bool bExecute;
 
+	// The below is alpha value support for control rig
+	float InternalBlendAlpha;
+
 	// update input/output to control rig
 	virtual void UpdateInput(UControlRig* ControlRig, const FPoseContext& InOutput);
 	virtual void UpdateOutput(UControlRig* ControlRig, FPoseContext& InOutput);
 	virtual UClass* GetTargetClass() const override;
+	
+	// execute control rig on the input pose and outputs the result
+	void ExecuteControlRig(FPoseContext& InOutput);
 
 	friend struct FControlRigSequencerAnimInstanceProxy;
 };
diff --git a/Engine/Plugins/Experimental/ControlRig/Source/ControlRig/Public/AnimNode_ControlRig_ExternalSource.h b/Engine/Plugins/Experimental/ControlRig/Source/ControlRig/Public/AnimNode_ControlRig_ExternalSource.h
--- a/Engine/Plugins/Experimental/ControlRig/Source/ControlRig/Public/AnimNode_ControlRig_ExternalSource.h
+++ b/Engine/Plugins/Experimental/ControlRig/Source/ControlRig/Public/AnimNode_ControlRig_ExternalSource.h
@@ -18,15 +18,6 @@ struct CONTROLRIG_API FAnimNode_ControlRig_ExternalSource : public FAnimNode_Con
 	void SetControlRig(UControlRig* InControlRig);
 	virtual UControlRig* GetControlRig() const;
 
-	virtual void Initialize_AnyThread(const FAnimationInitializeContext& Context) override;
-	virtual void GatherDebugData(FNodeDebugData& DebugData) override;
-	virtual void Update_AnyThread(const FAnimationUpdateContext& Context) override;
-	virtual void CacheBones_AnyThread(const FAnimationCacheBonesContext& Context) override;
-	virtual void Evaluate_AnyThread(FPoseContext& Output) override;
-
-	UPROPERTY(EditAnywhere, Category = Links)
-	FPoseLink Source;
-
 private:
 	UPROPERTY(transient)
 	TWeakObjectPtr<UControlRig> ControlRig;
diff --git a/Engine/Plugins/Experimental/ControlRig/Source/ControlRig/Public/ControlRig.h b/Engine/Plugins/Experimental/ControlRig/Source/ControlRig/Public/ControlRig.h
--- a/Engine/Plugins/Experimental/ControlRig/Source/ControlRig/Public/ControlRig.h
+++ b/Engine/Plugins/Experimental/ControlRig/Source/ControlRig/Public/ControlRig.h
@@ -1,7 +1,7 @@
 // Copyright 1998-2019 Epic Games, Inc. All Rights Reserved.
 
 #pragma once
-
+#include "UObject/Object.h"
 #include "CoreMinimal.h"
 #include "UObject/ObjectMacros.h"
 #include "UObject/ScriptMacros.h"
@@ -69,7 +69,10 @@ public:
 	UControlRig();
 
 	virtual void Serialize(FArchive& Ar) override;
-
+#if WITH_EDITOR
+	virtual void PostEditUndo() override;
+#endif
+	
 	/** Get the current delta time */
 	UFUNCTION(BlueprintPure, Category = "Animation")
 	float GetDeltaTime() const;
@@ -95,17 +98,23 @@ public:
 	virtual void Evaluate_AnyThread();
 
 	/** Setup bindings to a runtime object (or clear by passing in nullptr). */
-	void SetObjectBinding(TSharedPtr<IControlRigObjectBinding> InObjectBinding)
+	virtual void SetObjectBinding(TSharedPtr<IControlRigObjectBinding> InObjectBinding) override
 	{
 		ObjectBinding = InObjectBinding;
 	}
 
 	/** Get bindings to a runtime object */
-	TSharedPtr<IControlRigObjectBinding> GetObjectBinding() const
+	virtual TSharedPtr<IControlRigObjectBinding> GetObjectBinding() const override
 	{
 		return ObjectBinding;
 	}
-
+    /** Get OurName*/
+	virtual FString GetName() const override
+	{
+		FString ObjectName = (GetClass()->GetName());
+		ObjectName.RemoveFromEnd(TEXT("_C"));
+		return ObjectName;
+	}
 	FRigHierarchyContainer* GetHierarchy()
 	{
 		return &Hierarchy;
@@ -216,6 +225,12 @@ public:
 	virtual FRigControlValue GetControlValueFromGlobalTransform(const FName& InControlName, const FTransform& InGlobalTransform) override;
 	virtual bool SetControlSpace(const FName& InControlName, const FName& InSpaceName) override;
 	virtual UControlRigGizmoLibrary* GetGizmoLibrary() const override;
+#if WITH_EDITOR
+	virtual void SelectControl(const FName& InControlName, bool bSelect = true) override;
+	virtual bool ClearControlSelection() override;
+	virtual TArray<FName> CurrentControlSelection() const override;
+	virtual bool IsControlSelected(const FName& InControlName)const override;
+#endif
 	// END IControlRigManipulatable interface
 
 	DECLARE_EVENT_TwoParams(UControlRig, FControlRigExecuteEvent, class UControlRig*, const EControlRigState);
@@ -296,6 +311,10 @@ private:
 	/** Broadcasts a notification whenever the controlrig is executed / updated. */
 	FControlRigExecuteEvent ExecutedEvent;
 
+#if WITH_EDITOR
+	/** Handle a Control Being Selected */
+	void HandleOnControlSelected(FRigHierarchyContainer* InContainer, const FRigElementKey& InKey, bool bSelected);
+#endif
 	void ResolveInputOutputProperties();
 
 	void InitializeFromCDO();
diff --git a/Engine/Plugins/Experimental/ControlRig/Source/ControlRig/Public/Manipulatable/IControlRigManipulatable.h b/Engine/Plugins/Experimental/ControlRig/Source/ControlRig/Public/Manipulatable/IControlRigManipulatable.h
--- a/Engine/Plugins/Experimental/ControlRig/Source/ControlRig/Public/Manipulatable/IControlRigManipulatable.h
+++ b/Engine/Plugins/Experimental/ControlRig/Source/ControlRig/Public/Manipulatable/IControlRigManipulatable.h
@@ -16,6 +16,9 @@
  * and implements a series of setter functions to perform changes.
  *
  */
+
+class IControlRigObjectBinding;
+
 class CONTROLRIG_API IControlRigManipulatable
 {
 public:
@@ -27,6 +30,10 @@ public:
 
 	/** Bindable event for external objects to be notified of Control changes */
 	DECLARE_EVENT_TwoParams(IControlRigManipulatable, FControlModifiedEvent, IControlRigManipulatable*, const FRigControl&);
+#if WITH_EDITOR
+	/** Bindable event for external objects to be notified that a Control is Selected */
+	DECLARE_EVENT_ThreeParams(IControlRigManipulatable, FControlSelectedEvent, IControlRigManipulatable*, const FRigControl&,  bool);
+#endif
 
 	// Returns true if this manipulatable subject is currently
 	// available for manipulation / is enabled.
@@ -90,7 +97,7 @@ public:
 
 	// Sets the relative value of a Control
 	template<class T>
-	FORCEINLINE bool SetControlValue(const FName& InControlName, T InValue, bool bNotify = true)
+	FORCEINLINE void SetControlValue(const FName& InControlName, T InValue, bool bNotify = true)
 	{
 		SetControlValue(InControlName, FRigControlValue::Make<T>(InValue), bNotify);
 	}
@@ -108,22 +115,52 @@ public:
 	// Sets a Control's Space (for space switching), returns true when successful.
 	virtual bool SetControlSpace(const FName& InControlName, const FName& InSpaceName) = 0;
 
-	// Returns a event than can be used to subscribe to
+	// Returns a event that can be used to subscribe to
 	// filtering control data when needed
 	FFilterControlEvent& ControlFilter() { return OnFilterControl;  }
 	
-	// Returns a event than can be used to subscribe to
-	// notifications coming from the manipulated subject.
+	// Returns a event that can be used to subscribe to
+	// change notifications coming from the manipulated subject.
 	FControlModifiedEvent& ControlModified() { return OnControlModified;  }
 
+#if WITH_EDITOR
+	// Returns a event that can be used to subscribe to
+	// selection changes coming from the manipulated subject.
+	FControlSelectedEvent& ControlSelected() { return OnControlSelected; }
+
+	//Select or deselected the specified control  
+	virtual void SelectControl(const FName& InControlName, bool bSelect = true) = 0;
+
+	//Clear selection on all controls
+	virtual bool ClearControlSelection() = 0;
+
+	//Get the current selection
+	virtual TArray<FName> CurrentControlSelection() const = 0;
+
+	//Is the Specified Control Selected
+	virtual bool IsControlSelected(const FName& InControlName) const = 0;
+
+#endif
+	//Get Name
+	virtual FString GetName() const = 0;
+
 	// Returns the gizmo library used for generating gizmos
 	virtual UControlRigGizmoLibrary* GetGizmoLibrary() const { return nullptr; }
 
+	//Set Object Binding
+	virtual void SetObjectBinding(TSharedPtr<IControlRigObjectBinding> InObjectBinding) {};
+
+	//Get bindings to a runtime object 
+	virtual TSharedPtr<IControlRigObjectBinding> GetObjectBinding() const {	return nullptr;}
+
+
 private:
 
 	FFilterControlEvent OnFilterControl;
 	FControlModifiedEvent OnControlModified;
-
+#if WITH_EDITOR
+	FControlSelectedEvent OnControlSelected;
+#endif
 	/** True if currently manipulation is enabled */
 	bool bManipulationEnabled;
 };
\ No newline at end of file
diff --git a/Engine/Plugins/Experimental/ControlRig/Source/ControlRig/Public/Math/ControlRigMathLibrary.h b/Engine/Plugins/Experimental/ControlRig/Source/ControlRig/Public/Math/ControlRigMathLibrary.h
--- a/Engine/Plugins/Experimental/ControlRig/Source/ControlRig/Public/Math/ControlRigMathLibrary.h
+++ b/Engine/Plugins/Experimental/ControlRig/Source/ControlRig/Public/Math/ControlRigMathLibrary.h
@@ -3,6 +3,7 @@
 #pragma once
 
 #include "CoreMinimal.h"
+#include "ControlRigDefines.h"
 #include "ControlRigMathLibrary.generated.h"
 
 UENUM()
@@ -85,5 +86,6 @@ public:
 	static void FourPointBezier(const FCRFourPointBezier& Bezier, float T, FVector& OutPosition, FVector& OutTangent);
 	static float EaseFloat(float Value, EControlRigAnimEasingType Type);
 	static FTransform LerpTransform(const FTransform& A, const FTransform& B, float T);
-	static void SolveBasicTwoBoneIK(FTransform& BoneA, FTransform& BoneB, FTransform& Effector, const FVector& PoleVector, const FVector& PrimaryAxis, const FVector& SecondaryAxis, float BoneALength, float BoneBLength, bool bEnableStretch, float StretchStartRatio, float StretchMaxRatio);
+	static void SolveBasicTwoBoneIK(FTransform& BoneA, FTransform& BoneB, FTransform& Effector, const FVector& PoleVector, const FVector& PrimaryAxis, const FVector& SecondaryAxis, float SecondaryAxisWeight, float BoneALength, float BoneBLength, bool bEnableStretch, float StretchStartRatio, float StretchMaxRatio);
+	static FVector ClampSpatially(const FVector& Value, EAxis::Type Axis, EControlRigClampSpatialMode::Type Type, float Minimum, float Maximum, FTransform Space);
 };
\ No newline at end of file
diff --git a/Engine/Plugins/Experimental/ControlRig/Source/ControlRig/Public/Rigs/RigBoneHierarchy.h b/Engine/Plugins/Experimental/ControlRig/Source/ControlRig/Public/Rigs/RigBoneHierarchy.h
--- a/Engine/Plugins/Experimental/ControlRig/Source/ControlRig/Public/Rigs/RigBoneHierarchy.h
+++ b/Engine/Plugins/Experimental/ControlRig/Source/ControlRig/Public/Rigs/RigBoneHierarchy.h
@@ -4,10 +4,18 @@
 
 #include "CoreMinimal.h"
 #include "RigHierarchyDefines.h"
+#include "ReferenceSkeleton.h"
 #include "RigBoneHierarchy.generated.h"
 
 class UControlRig;
 
+UENUM()
+enum class ERigBoneType : uint8
+{
+	Imported,
+	User
+};
+
 USTRUCT(BlueprintType)
 struct CONTROLRIG_API FRigBone: public FRigElement
 {
@@ -21,6 +29,7 @@ struct CONTROLRIG_API FRigBone: public FRigElement
 		, GlobalTransform(FTransform::Identity)
 		, LocalTransform(FTransform::Identity)
 		, Dependents()
+		, Type(ERigBoneType::Imported)
 	{
 	}
 	virtual ~FRigBone() {}
@@ -45,6 +54,10 @@ struct CONTROLRIG_API FRigBone: public FRigElement
 	UPROPERTY(transient)
 	TArray<int32> Dependents;
 
+	/** the source of the bone to differentiate procedurally generated, imported etc */
+	UPROPERTY(VisibleAnywhere, Category = FRigElement)
+	ERigBoneType Type;
+
 	FORCEINLINE virtual ERigElementType GetElementType() const override
 	{
 		return ERigElementType::Bone;
@@ -81,9 +94,9 @@ struct CONTROLRIG_API FRigBoneHierarchy
 
 	FName GetSafeNewName(const FName& InPotentialNewName) const;
 
-	FRigBone& Add(const FName& InNewName, const FName& InParentName = NAME_None, const FTransform& InInitTransform = FTransform::Identity);
+	FRigBone& Add(const FName& InNewName, const FName& InParentName = NAME_None, ERigBoneType InType = ERigBoneType::Imported, const FTransform& InInitTransform = FTransform::Identity);
 
-	FRigBone& Add(const FName& InNewName, const FName& InParentName, const FTransform& InInitTransform, const FTransform& InLocalTransform, const FTransform& InGlobalTransform);
+	FRigBone& Add(const FName& InNewName, const FName& InParentName, ERigBoneType InType, const FTransform& InInitTransform, const FTransform& InLocalTransform, const FTransform& InGlobalTransform);
 
 	FRigBone Remove(const FName& InNameToRemove);
 
@@ -154,6 +167,8 @@ struct CONTROLRIG_API FRigBoneHierarchy
 	TArray<FName> CurrentSelection() const;
 	bool IsSelected(const FName& InName) const;
 
+	TArray<FRigElementKey> ImportSkeleton(const FReferenceSkeleton& InSkeleton, const FName& InNameSpace, bool bReplaceExistingBones, bool bRemoveObsoleteBones, bool bSelectBones);
+
 	FRigElementAdded OnBoneAdded;
 	FRigElementRemoved OnBoneRemoved;
 	FRigElementRenamed OnBoneRenamed;
@@ -186,6 +201,7 @@ private:
 	void RecalculateLocalTransform(FRigBone& InOutBone);
 	void RecalculateGlobalTransform(FRigBone& InOutBone);
 
+	void RefreshParentNames();
 	void RefreshMapping();
 	void Sort();
 
diff --git a/Engine/Plugins/Experimental/ControlRig/Source/ControlRig/Public/Sequencer/MovieSceneControlRigParameterSection.h b/Engine/Plugins/Experimental/ControlRig/Source/ControlRig/Public/Sequencer/MovieSceneControlRigParameterSection.h
new file mode 100644
--- /dev/null
+++ b/Engine/Plugins/Experimental/ControlRig/Source/ControlRig/Public/Sequencer/MovieSceneControlRigParameterSection.h
@@ -0,0 +1,191 @@
+// Copyright 1998-2019 Epic Games, Inc. All Rights Reserved.
+
+#pragma once
+
+#include "CoreMinimal.h"
+#include "Sections/MovieSceneParameterSection.h"
+#include "UObject/ObjectMacros.h"
+#include "Channels/MovieSceneFloatChannel.h"
+#include "Sections/MovieSceneSubSection.h"
+#include "ControlRig.h"
+#include "MovieSceneSequencePlayer.h"
+#include "Animation/AnimData/BoneMaskFilter.h"
+#include "MovieSceneObjectBindingID.h"
+#include "MovieSceneControlRigParameterSection.generated.h"
+
+
+class UMovieSceneControlRigSection;
+/**
+*  Data that's queried during an interrogtion
+*/
+struct FFloatInterrogationData
+{
+	float Val;
+	FName ParameterName;
+};
+
+
+struct FVectorInterrogationData
+{
+	FVector Val;
+	FName ParameterName;
+};
+
+struct FTransformInterrogationData
+{
+	FTransform Val;
+	FName ParameterName;
+};
+
+
+/**
+ * Movie scene section that controls animation controller animation
+ */
+UCLASS()
+class CONTROLRIG_API UMovieSceneControlRigParameterSection : public UMovieSceneParameterSection
+{
+	GENERATED_BODY()
+
+public:
+
+	/** Control Rig that controls us*/
+	UPROPERTY()
+	UControlRig* ControlRig;
+
+	/** Mask for controls themselves*/
+	UPROPERTY()
+	TArray<bool> ControlsMask;
+
+	/** Mask for Transform Mask*/
+	UPROPERTY()
+	FMovieSceneTransformMask TransformMask;
+
+	/** Blend this track in additively (using the reference pose as a base) */
+	UPROPERTY(EditAnywhere, Category = "Animation")
+	bool bAdditive;
+
+	/** Only apply bones that are in the filter */
+	UPROPERTY(EditAnywhere, Category = "Animation")
+	bool bApplyBoneFilter;
+
+	/** Per-bone filter to apply to our animation */
+	UPROPERTY(EditAnywhere, Category = "Animation", meta=(EditCondition=bApplyBoneFilter))
+	FInputBlendPose BoneFilter;
+
+	/** The weight curve for this animation controller section */
+	UPROPERTY()
+	FMovieSceneFloatChannel Weight;
+
+public:
+
+	UMovieSceneControlRigParameterSection();
+
+	const TArray<bool>& GetControlsMask() const
+	{
+		return ControlsMask;
+	}
+
+	bool GetControlsMask(int32 Index) const
+	{
+		if (Index >= 0 && Index < ControlsMask.Num())
+		{
+			return ControlsMask[Index];
+		}
+		return false;
+	}
+
+	void SetControlsMask(const TArray<bool>& InMask)
+	{
+		ControlsMask = InMask;
+		ReconstructChannelProxy();
+	}
+
+	void SetControlsMask(int32 Index, bool Val)
+	{
+		if (Index >= 0 && Index < ControlsMask.Num())
+		{
+			ControlsMask[Index] = Val;
+		}
+		ReconstructChannelProxy();
+	}
+
+
+	void FillControlsMask(bool Val)
+	{
+		ControlsMask.Init(Val, ControlsMask.Num());
+		ReconstructChannelProxy();
+
+	}
+
+	/**
+	* Access the transform mask that defines which channels this track should animate
+	*/
+	FMovieSceneTransformMask GetTransformMask() const
+	{
+		return TransformMask;
+	}
+
+	/**
+	 * Set the transform mask that defines which channels this track should animate
+	 */
+	void SetTransformMask(FMovieSceneTransformMask NewMask)
+	{
+		TransformMask = NewMask;
+		ReconstructChannelProxy();
+	}
+
+public:
+
+	/** Whether or not to key currently, maybe evaluating so don't*/
+	void  SetDoNotKey(bool bIn) const { bDoNotKey = bIn; }
+	/** Get Whether to key or not*/
+	bool GetDoNotKey() const { return bDoNotKey; }
+
+	/**  Whether or not this section his scalar*/
+	bool HasScalarParameter(FName InParameterName) const;
+
+	/**  Whether or not this section his scalar*/
+	bool HasVectorParameter(FName InParameterName) const;
+
+	/**  Whether or not this section his scalar*/
+	bool HasColorParameter(FName InParameterName) const;
+
+		/**  Whether or not this section his scalar*/
+	bool HasTransformParameter(FName InParameterName) const;
+
+	/** Adds specified scalar parameter. */
+	void AddScalarParameter(FName InParameterName,  TOptional<float> DefaultValue);
+
+	/** Adds a a key for a specific vector parameter. */
+	void AddVectorParameter(FName InParameterName, TOptional<FVector> DefaultValue);
+
+	/** Adds a a key for a specific color parameter. */
+	void AddColorParameter(FName InParameterName, TOptional<FLinearColor> DefaultValue);
+
+	/** Adds a a key for a specific color parameter*/
+	void AddTransformParameter(FName InParameterName, TOptional<FTransform> DefaultValue);
+
+public:
+	/**
+	* Access the interrogation key for control rig data 
+	*/
+	 static FMovieSceneInterrogationKey GetFloatInterrogationKey();
+	 static FMovieSceneInterrogationKey GetVectorInterrogationKey();
+	 static FMovieSceneInterrogationKey GetTransformInterrogationKey();
+
+
+protected:
+	virtual void ReconstructChannelProxy();
+
+protected:
+
+	//~ UMovieSceneSection interface
+	virtual void Serialize(FArchive& Ar) override;
+	virtual void PostEditImport() override;
+
+	// When true we do not set a key on the section, since it will be set because we changed the value
+	// We need this because control rig notifications are set on every change even when just changing sequencer time
+	// which forces a sequencer eval, not like the edito where changes are only set on UI changes(changing time doesn't send change delegate)
+	mutable bool bDoNotKey;
+
+};
diff --git a/Engine/Plugins/Experimental/ControlRig/Source/ControlRigDeveloper/Private/AnimGraphNode_ControlRig.cpp b/Engine/Plugins/Experimental/ControlRig/Source/ControlRigDeveloper/Private/AnimGraphNode_ControlRig.cpp
--- a/Engine/Plugins/Experimental/ControlRig/Source/ControlRigDeveloper/Private/AnimGraphNode_ControlRig.cpp
+++ b/Engine/Plugins/Experimental/ControlRig/Source/ControlRigDeveloper/Private/AnimGraphNode_ControlRig.cpp
@@ -11,6 +11,7 @@
 #include "SVariableMappingWidget.h"
 #include "ScopedTransaction.h"
 #include "Kismet2/BlueprintEditorUtils.h"
+#include "ControlRigBlueprint.h"
 
 #define LOCTEXT_NAMESPACE "AnimGraphNode_ControlRig"
 
@@ -70,6 +71,41 @@ void UAnimGraphNode_ControlRig::ReallocatePinsDuringReconstruction(TArray<UEdGra
 	Super::ReallocatePinsDuringReconstruction(OldPins);
 }
 
+void UAnimGraphNode_ControlRig::ValidateAnimNodeDuringCompilation(USkeleton* ForSkeleton, FCompilerResultsLog& MessageLog)
+{
+	if (UClass* TargetClass = GetTargetClass())
+	{
+		if (UControlRigBlueprint* Blueprint = Cast<UControlRigBlueprint>(TargetClass->ClassGeneratedBy))
+		{
+			const FRigBoneHierarchy& BoneHierarchy = Blueprint->HierarchyContainer.BoneHierarchy;
+			const FReferenceSkeleton& ReferenceSkeleton = ForSkeleton->GetReferenceSkeleton();
+			const TArray<FMeshBoneInfo>& BoneInfos = ReferenceSkeleton.GetRefBoneInfo();
+
+			for (const FMeshBoneInfo& BoneInfo : BoneInfos)
+			{
+				int32 BoneIndex = BoneHierarchy.GetIndex(BoneInfo.Name);
+				if (BoneIndex != INDEX_NONE)
+				{
+					FName DesiredParentName = NAME_None;
+					if (BoneInfo.ParentIndex != INDEX_NONE)
+					{
+						DesiredParentName = BoneInfos[BoneInfo.ParentIndex].Name;
+					}
+
+					const FRigBone& Bone = BoneHierarchy[BoneIndex];
+					if (DesiredParentName != Bone.ParentName)
+					{
+						FString Message = FString::Printf(TEXT("@@ - Hierarchy discrepancy for bone '%s' - different parents on Control Rig vs SkeletalMesh."), *BoneInfo.Name.ToString());
+						MessageLog.Warning(*Message, this);
+					}
+				}
+			}
+		}
+	}
+
+	Super::ValidateAnimNodeDuringCompilation(ForSkeleton, MessageLog);
+}
+
 void UAnimGraphNode_ControlRig::RebuildExposedProperties()
 {
 	DECLARE_SCOPE_HIERARCHICAL_COUNTER_FUNC()
@@ -190,6 +226,7 @@ void UAnimGraphNode_ControlRig::CustomizeDetails(IDetailLayoutBuilder& DetailBui
 		return;
 	}
 
+	// input/output exposure feature START
 	RebuildExposedProperties();
 
 	IDetailCategoryBuilder& InputCategoryBuilder = DetailBuilder.EditCategory(FName(TEXT("Input")));
@@ -229,6 +266,35 @@ void UAnimGraphNode_ControlRig::CustomizeDetails(IDetailLayoutBuilder& DetailBui
 	{
 		ClassHandle->SetOnPropertyValueChanged(FSimpleDelegate::CreateUObject(this, &UAnimGraphNode_ControlRig::OnInstanceClassChanged, &DetailBuilder));
 	}
+
+	// input/output exposure feature END
+
+	// alpha property blending support START
+	TSharedRef<IPropertyHandle> NodeHandle = DetailBuilder.GetProperty(FName(TEXT("Node")), GetClass());
+
+	if (Node.AlphaInputType != EAnimAlphaInputType::Bool)
+	{
+		DetailBuilder.HideProperty(NodeHandle->GetChildHandle(GET_MEMBER_NAME_CHECKED(FAnimNode_ControlRig, bAlphaBoolEnabled)));
+		DetailBuilder.HideProperty(NodeHandle->GetChildHandle(GET_MEMBER_NAME_CHECKED(FAnimNode_ControlRig, AlphaBoolBlend)));
+	}
+
+	if (Node.AlphaInputType != EAnimAlphaInputType::Float)
+	{
+		DetailBuilder.HideProperty(NodeHandle->GetChildHandle(GET_MEMBER_NAME_CHECKED(FAnimNode_ControlRig, Alpha)));
+		DetailBuilder.HideProperty(NodeHandle->GetChildHandle(GET_MEMBER_NAME_CHECKED(FAnimNode_ControlRig, AlphaScaleBias)));
+	}
+
+	if (Node.AlphaInputType != EAnimAlphaInputType::Curve)
+	{
+		DetailBuilder.HideProperty(NodeHandle->GetChildHandle(GET_MEMBER_NAME_CHECKED(FAnimNode_ControlRig, AlphaCurveName)));
+	}
+
+	if ((Node.AlphaInputType != EAnimAlphaInputType::Float)
+		&& (Node.AlphaInputType != EAnimAlphaInputType::Curve))
+	{
+		DetailBuilder.HideProperty(NodeHandle->GetChildHandle(GET_MEMBER_NAME_CHECKED(FAnimNode_ControlRig, AlphaScaleBiasClamp)));
+	}
+	// alpha property blending support END
 }
 
 void UAnimGraphNode_ControlRig::GetIOProperties(bool bInput, TMap<FName, FControlRigIOVariable>& OutVars) const
@@ -363,12 +429,77 @@ void UAnimGraphNode_ControlRig::PostEditChangeProperty(FPropertyChangedEvent& Pr
 			bRequiresNodeReconstruct = true;
 			RebuildExposedProperties();
 		}
+
+		if (ChangedProperty->GetFName() == GET_MEMBER_NAME_STRING_CHECKED(FAnimNode_ControlRig, AlphaInputType))
+		{
+			FScopedTransaction Transaction(LOCTEXT("ChangeAlphaInputType", "Change Alpha Input Type"));
+			Modify();
+
+			// Break links to pins going away
+			for (int32 PinIndex = 0; PinIndex < Pins.Num(); ++PinIndex)
+			{
+				UEdGraphPin* Pin = Pins[PinIndex];
+				if (Pin->PinName == GET_MEMBER_NAME_STRING_CHECKED(FAnimNode_ControlRig, Alpha))
+				{
+					if (Node.AlphaInputType != EAnimAlphaInputType::Float)
+					{
+						Pin->BreakAllPinLinks();
+					}
+				}
+				else if (Pin->PinName == GET_MEMBER_NAME_STRING_CHECKED(FAnimNode_ControlRig, bAlphaBoolEnabled))
+				{
+					if (Node.AlphaInputType != EAnimAlphaInputType::Bool)
+					{
+						Pin->BreakAllPinLinks();
+					}
+				}
+				else if (Pin->PinName == GET_MEMBER_NAME_STRING_CHECKED(FAnimNode_ControlRig, AlphaCurveName))
+				{
+					if (Node.AlphaInputType != EAnimAlphaInputType::Curve)
+					{
+						Pin->BreakAllPinLinks();
+					}
+				}
+			}
+
+			bRequiresNodeReconstruct = true;
+		}
 	}
 
 	if (bRequiresNodeReconstruct)
 	{
 		ReconstructNode();
+		FBlueprintEditorUtils::MarkBlueprintAsStructurallyModified(GetBlueprint());
 	}
 }
 
+void UAnimGraphNode_ControlRig::CustomizePinData(UEdGraphPin* Pin, FName SourcePropertyName, int32 ArrayIndex) const
+{
+	Super::CustomizePinData(Pin, SourcePropertyName, ArrayIndex);
+
+	if (Pin->PinName == GET_MEMBER_NAME_STRING_CHECKED(FAnimNode_ControlRig, Alpha))
+	{
+		Pin->bHidden = (Node.AlphaInputType != EAnimAlphaInputType::Float);
+
+		if (!Pin->bHidden)
+		{
+			Pin->PinFriendlyName = Node.AlphaScaleBias.GetFriendlyName(Node.AlphaScaleBiasClamp.GetFriendlyName(Pin->PinFriendlyName));
+		}
+	}
+
+	if (Pin->PinName == GET_MEMBER_NAME_STRING_CHECKED(FAnimNode_ControlRig, bAlphaBoolEnabled))
+	{
+		Pin->bHidden = (Node.AlphaInputType != EAnimAlphaInputType::Bool);
+	}
+
+	if (Pin->PinName == GET_MEMBER_NAME_STRING_CHECKED(FAnimNode_ControlRig, AlphaCurveName))
+	{
+		Pin->bHidden = (Node.AlphaInputType != EAnimAlphaInputType::Curve);
+
+		if (!Pin->bHidden)
+		{
+			Pin->PinFriendlyName = Node.AlphaScaleBiasClamp.GetFriendlyName(Pin->PinFriendlyName);
+		}
+	}
+}
 #undef LOCTEXT_NAMESPACE
diff --git a/Engine/Plugins/Experimental/ControlRig/Source/ControlRigDeveloper/Private/ControlRigModel.cpp b/Engine/Plugins/Experimental/ControlRig/Source/ControlRigDeveloper/Private/ControlRigModel.cpp
--- a/Engine/Plugins/Experimental/ControlRig/Source/ControlRigDeveloper/Private/ControlRigModel.cpp
+++ b/Engine/Plugins/Experimental/ControlRig/Source/ControlRigDeveloper/Private/ControlRigModel.cpp
@@ -2555,6 +2555,14 @@ void UControlRigModel::ConfigurePinFromField(FControlRigModelPin& Pin, UProperty
 	{
 		Pin.CustomWidgetName = UControlRig::BoneNameMetaName;
 	}
+	else if (Property->HasMetaData(UControlRig::ControlNameMetaName))
+	{
+		Pin.CustomWidgetName = UControlRig::ControlNameMetaName;
+	}
+	else if (Property->HasMetaData(UControlRig::SpaceNameMetaName))
+	{
+		Pin.CustomWidgetName = UControlRig::SpaceNameMetaName;
+	}
 	else if (Property->HasMetaData(UControlRig::CurveNameMetaName))
 	{
 		Pin.CustomWidgetName = UControlRig::CurveNameMetaName;
diff --git a/Engine/Plugins/Experimental/ControlRig/Source/ControlRigEditor/Private/ControlRigEditorModule.cpp b/Engine/Plugins/Experimental/ControlRig/Source/ControlRigEditor/Private/ControlRigEditorModule.cpp
--- a/Engine/Plugins/Experimental/ControlRig/Source/ControlRigEditor/Private/ControlRigEditorModule.cpp
+++ b/Engine/Plugins/Experimental/ControlRig/Source/ControlRigEditor/Private/ControlRigEditorModule.cpp
@@ -12,9 +12,7 @@
 #include "GraphEditorActions.h"
 #include "Kismet2/KismetEditorUtilities.h"
 #include "ISequencerModule.h"
-#include "ControlRigTrackEditor.h"
 #include "IAssetTools.h"
-#include "ControlRigSequenceActions.h"
 #include "ControlRigEditorStyle.h"
 #include "Framework/Docking/LayoutExtender.h"
 #include "Framework/Application/SlateApplication.h"
@@ -25,16 +23,12 @@
 #include "WorkflowOrientedApp/WorkflowTabManager.h"
 #include "Modules/ModuleManager.h"
 #include "Widgets/Docking/SDockTab.h"
-#include "Sequencer/ControlRigSequence.h"
 #include "EditorModeRegistry.h"
 #include "ControlRigEditMode.h"
-#include "ControlRigEditorObjectBinding.h"
-#include "ControlRigEditorObjectSpawner.h"
 #include "ILevelSequenceModule.h"
-#include "ControlRigBindingTrackEditor.h"
 #include "EditorModeManager.h"
 #include "ControlRigEditMode.h"
-#include "Sequencer/MovieSceneControlRigSection.h"
+#include "Sequencer/MovieSceneControlRigParameterSection.h"
 #include "MovieSceneControlRigSectionDetailsCustomization.h"
 #include "ControlRigEditModeCommands.h"
 #include "Materials/Material.h"
@@ -77,6 +71,7 @@
 #include "WorkspaceMenuStructure.h"
 #include "WorkspaceMenuStructureModule.h"
 #include "Subsystems/AssetEditorSubsystem.h"
+#include "ControlRigParameterTrackEditor.h"
 
 #define LOCTEXT_NAMESPACE "ControlRigEditorModule"
 
@@ -101,10 +96,6 @@ void FControlRigEditorModule::StartupModule()
 	FControlRigStackCommands::Register();
 	FControlRigEditorStyle::Get();
 
-	CommandBindings = MakeShareable(new FUICommandList());
-
-	BindCommands();
-
 	MenuExtensibilityManager = MakeShareable(new FExtensibilityManager);
 	ToolBarExtensibilityManager = MakeShareable(new FExtensibilityManager);
 
@@ -119,7 +110,7 @@ void FControlRigEditorModule::StartupModule()
 	FPropertyEditorModule& PropertyEditorModule = FModuleManager::LoadModuleChecked<FPropertyEditorModule>("PropertyEditor");
 	ClassesToUnregisterOnShutdown.Reset();
 
-	ClassesToUnregisterOnShutdown.Add(UMovieSceneControlRigSection::StaticClass()->GetFName());
+	ClassesToUnregisterOnShutdown.Add(UMovieSceneControlRigParameterSection::StaticClass()->GetFName());
 	PropertyEditorModule.RegisterCustomClassLayout(ClassesToUnregisterOnShutdown.Last(), FOnGetDetailCustomizationInstance::CreateStatic(&FMovieSceneControlRigSectionDetailsCustomization::MakeInstance));
 
 	ClassesToUnregisterOnShutdown.Add(UControlRigSequenceExporterSettings::StaticClass()->GetFName());
@@ -140,7 +131,6 @@ void FControlRigEditorModule::StartupModule()
 	// same as ClassesToUnregisterOnShutdown but for properties, there is none right now
 	PropertiesToUnregisterOnShutdown.Reset();
 
-
 	// Register asset tools
 	auto RegisterAssetTypeAction = [this](const TSharedRef<IAssetTypeActions>& InAssetTypeAction)
 	{
@@ -149,41 +139,12 @@ void FControlRigEditorModule::StartupModule()
 		AssetTools.RegisterAssetTypeActions(InAssetTypeAction);
 	};
 
-	RegisterAssetTypeAction(MakeShareable(new FControlRigSequenceActions()));
 	RegisterAssetTypeAction(MakeShareable(new FControlRigBlueprintActions()));
 	RegisterAssetTypeAction(MakeShareable(new FControlRigGizmoLibraryActions()));
 	
 	// Register sequencer track editor
 	ISequencerModule& SequencerModule = FModuleManager::Get().LoadModuleChecked<ISequencerModule>("Sequencer");
-	SequencerCreatedHandle = SequencerModule.RegisterOnSequencerCreated(FOnSequencerCreated::FDelegate::CreateRaw(this, &FControlRigEditorModule::HandleSequencerCreated));
-	ControlRigTrackCreateEditorHandle = SequencerModule.RegisterTrackEditor(FOnCreateTrackEditor::CreateStatic(&FControlRigTrackEditor::CreateTrackEditor));
-	ControlRigBindingTrackCreateEditorHandle = SequencerModule.RegisterTrackEditor(FOnCreateTrackEditor::CreateStatic(&FControlRigBindingTrackEditor::CreateTrackEditor));
-	ControlRigEditorObjectBindingHandle = SequencerModule.RegisterEditorObjectBinding(FOnCreateEditorObjectBinding::CreateStatic(&FControlRigEditorObjectBinding::CreateEditorObjectBinding));
-
-	SequencerToolbarExtender = MakeShareable(new FExtender());
-	SequencerToolbarExtender->AddToolBarExtension(
-		"Level Sequence Separator",
-		EExtensionHook::Before,
-		CommandBindings,
-		FToolBarExtensionDelegate::CreateLambda([](FToolBarBuilder& ToolBarBuilder)
-		{
-			ToolBarBuilder.AddToolBarButton(FControlRigEditModeCommands::Get().ExportAnimSequence);
-		}));
-
-	SequencerModule.GetToolBarExtensibilityManager()->AddExtender(SequencerToolbarExtender);
-
-	// Register for assets being opened
-	AssetEditorOpenedHandle = GEditor->GetEditorSubsystem<UAssetEditorSubsystem>()->OnAssetEditorOpened().AddRaw(this, &FControlRigEditorModule::HandleAssetEditorOpened);
-
-	// Register level sequence spawner
-	ILevelSequenceModule& LevelSequenceModule = FModuleManager::LoadModuleChecked<ILevelSequenceModule>("LevelSequence");
-	LevelSequenceSpawnerDelegateHandle = LevelSequenceModule.RegisterObjectSpawner(FOnCreateMovieSceneObjectSpawner::CreateStatic(&FControlRigEditorObjectSpawner::CreateObjectSpawner));
-
-	TrajectoryMaterial = LoadObject<UMaterial>(nullptr, TEXT("/ControlRig/M_Traj.M_Traj"));
-	if (TrajectoryMaterial.IsValid())
-	{
-		TrajectoryMaterial->AddToRoot();
-	}
+	ControlRigParameterTrackCreateEditorHandle = SequencerModule.RegisterTrackEditor(FOnCreateTrackEditor::CreateStatic(&FControlRigParameterTrackEditor::CreateTrackEditor));
 
 	FEditorModeRegistry::Get().RegisterMode<FControlRigEditMode>(
 		FControlRigEditMode::ModeName,
@@ -197,126 +158,6 @@ void FControlRigEditorModule::StartupModule()
 		FSlateIcon(FControlRigEditorStyle::Get().GetStyleSetName(), "ControlRigEditMode", "ControlRigEditMode.Small"),
 		false);
 
-	FContentBrowserModule& ContentBrowserModule = FModuleManager::LoadModuleChecked<FContentBrowserModule>(TEXT("ContentBrowser"));
-	ContentBrowserModule.GetAllAssetViewContextMenuExtenders().Add(FContentBrowserMenuExtender_SelectedAssets::CreateLambda([this](const TArray<FAssetData>& SelectedAssets)
-	{
-		TSharedRef<FExtender> Extender = MakeShared<FExtender>();
-
-		if (SelectedAssets.ContainsByPredicate([](const FAssetData& AssetData) { return AssetData.GetClass() == UAnimSequence::StaticClass(); }))
-		{
-			Extender->AddMenuExtension(
-				"GetAssetActions",
-				EExtensionHook::After,
-				CommandBindings,
-				FMenuExtensionDelegate::CreateLambda([this, SelectedAssets](FMenuBuilder& MenuBuilder)
-				{
-					const TSharedPtr<FUICommandInfo>& ImportFromRigSequence = FControlRigEditModeCommands::Get().ImportFromRigSequence;
-					MenuBuilder.AddMenuEntry(
-						ImportFromRigSequence->GetLabel(),
-						ImportFromRigSequence->GetDescription(),
-						ImportFromRigSequence->GetIcon(),
-						FUIAction(FExecuteAction::CreateRaw(this, &FControlRigEditorModule::ImportFromRigSequence, SelectedAssets)));
-				}));
-
-			// only add this if we find a control rig sequence targeting this anim sequence in the asset registry
-			FAssetRegistryModule& AssetRegistryModule = FModuleManager::LoadModuleChecked<FAssetRegistryModule>("AssetRegistry");
-			
-			bool bCanReimport = false;
-			if (SelectedAssets.Num() > 0)
-			{
-				// It's faster to find all assets with this tag and then query them against the selection then it is to 
-				// query the asset registry each time for a tag with a particular value
-				const FName LastExportedToAnimationSequenceTagName = GET_MEMBER_NAME_CHECKED(UControlRigSequence, LastExportedToAnimationSequence);
-				TArray<FAssetData> FoundAssets;
-				{
-					TArray<FName> Tags;
-					Tags.Add(LastExportedToAnimationSequenceTagName);
-					AssetRegistryModule.Get().GetAssetsByTags(Tags, FoundAssets);
-				}
-
-				if (FoundAssets.Num() > 0)
-				{
-					for(const FAssetData& AssetData : SelectedAssets)
-					{
-						const bool bFoundAsset = FoundAssets.ContainsByPredicate([&AssetData, LastExportedToAnimationSequenceTagName](const FAssetData& FoundAsset)
-						{
-							const FName TagValue = FoundAsset.GetTagValueRef<FName>(LastExportedToAnimationSequenceTagName);
-							return TagValue == AssetData.ObjectPath;
-						});
-
-						if (bFoundAsset)
-						{
-							bCanReimport = true;
-							break;
-						}
-					}
-				}
-			}
-
-			if (bCanReimport)
-			{
-				Extender->AddMenuExtension(
-					"GetAssetActions",
-					EExtensionHook::After,
-					CommandBindings,
-					FMenuExtensionDelegate::CreateLambda([this, SelectedAssets](FMenuBuilder& MenuBuilder)
-					{
-						const TSharedPtr<FUICommandInfo>& ReImportFromRigSequence = FControlRigEditModeCommands::Get().ReImportFromRigSequence;
-						MenuBuilder.AddMenuEntry(
-							ReImportFromRigSequence->GetLabel(),
-							ReImportFromRigSequence->GetDescription(),
-							ReImportFromRigSequence->GetIcon(),
-							FUIAction(FExecuteAction::CreateRaw(this, &FControlRigEditorModule::ReImportFromRigSequence, SelectedAssets)));
-					}));
-			}
-		}
-		else if (SelectedAssets.ContainsByPredicate([](const FAssetData& AssetData) { return AssetData.GetClass() == UControlRigSequence::StaticClass(); }))
-		{
-			Extender->AddMenuExtension(
-				"CommonAssetActions",
-				EExtensionHook::Before,
-				CommandBindings,
-				FMenuExtensionDelegate::CreateLambda([this, SelectedAssets](FMenuBuilder& MenuBuilder)
-				{
-					MenuBuilder.BeginSection("ControlRigActions", LOCTEXT("ControlRigActions", "Control Rig Sequence Actions"));
-					{
-						const TSharedPtr<FUICommandInfo>& ExportAnimSequence = FControlRigEditModeCommands::Get().ExportAnimSequence;
-						MenuBuilder.AddMenuEntry(
-							ExportAnimSequence->GetLabel(),
-							ExportAnimSequence->GetDescription(),
-							ExportAnimSequence->GetIcon(),
-							FUIAction(FExecuteAction::CreateRaw(this, &FControlRigEditorModule::ExportToAnimSequence, SelectedAssets)));
-
-						bool bCanReExport = false;
-						for (const FAssetData& AssetData : SelectedAssets)
-						{
-							if(UControlRigSequence* ControlRigSequence = Cast<UControlRigSequence>(AssetData.GetAsset()))
-							{
-								if (ControlRigSequence->LastExportedToAnimationSequence.IsValid())
-								{
-									bCanReExport = true;
-									break;
-								}
-							}
-						}
-
-						if (bCanReExport)
-						{
-							const TSharedPtr<FUICommandInfo>& ReExportAnimSequence = FControlRigEditModeCommands::Get().ReExportAnimSequence;
-							MenuBuilder.AddMenuEntry(
-								ReExportAnimSequence->GetLabel(),
-								ReExportAnimSequence->GetDescription(),
-								ReExportAnimSequence->GetIcon(),
-								FUIAction(FExecuteAction::CreateRaw(this, &FControlRigEditorModule::ReExportToAnimSequence, SelectedAssets)));
-						}
-					}
-					MenuBuilder.EndSection();
-				}));
-		}
-	
-		return Extender;
-	}));
-	ContentBrowserMenuExtenderHandle = ContentBrowserModule.GetAllAssetViewContextMenuExtenders().Last().GetHandle();
 
 	ControlRigGraphPanelNodeFactory = MakeShared<FControlRigGraphPanelNodeFactory>();
 	FEdGraphUtilities::RegisterVisualNodeFactory(ControlRigGraphPanelNodeFactory);
@@ -359,44 +200,14 @@ void FControlRigEditorModule::ShutdownModule()
 	FEdGraphUtilities::UnregisterVisualPinFactory(ControlRigGraphPanelPinFactory);
 	FEdGraphUtilities::UnregisterVisualNodeFactory(ControlRigGraphPanelNodeFactory);
 
-	FContentBrowserModule* ContentBrowserModule = FModuleManager::GetModulePtr<FContentBrowserModule>(TEXT("ContentBrowser"));
-	if (ContentBrowserModule)
-	{
-		ContentBrowserModule->GetAllAssetViewContextMenuExtenders().RemoveAll([=](const FContentBrowserMenuExtender_SelectedAssets& InDelegate) { return ContentBrowserMenuExtenderHandle == InDelegate.GetHandle(); });
-	}
-
-	if (TrajectoryMaterial.IsValid())
-	{
-		TrajectoryMaterial->RemoveFromRoot();
-	}
-
-	if (GEditor)
-	{
-		if (UAssetEditorSubsystem* AssetEditorSubsystem = GEditor->GetEditorSubsystem<UAssetEditorSubsystem>())
-		{
-			AssetEditorSubsystem->OnAssetEditorOpened().Remove(AssetEditorOpenedHandle);
-		}
-	}
-
 	FEditorModeRegistry::Get().UnregisterMode(FControlRigEditorEditMode::ModeName);
 	FEditorModeRegistry::Get().UnregisterMode(FControlRigEditMode::ModeName);
 
-	ILevelSequenceModule* LevelSequenceModule = FModuleManager::GetModulePtr<ILevelSequenceModule>("LevelSequence");
-	if (LevelSequenceModule)
-	{
-		LevelSequenceModule->UnregisterObjectSpawner(LevelSequenceSpawnerDelegateHandle);
-	}
 
 	ISequencerModule* SequencerModule = FModuleManager::GetModulePtr<ISequencerModule>("Sequencer");
 	if (SequencerModule)
 	{
-		SequencerModule->UnregisterOnSequencerCreated(SequencerCreatedHandle);
-		SequencerModule->UnRegisterTrackEditor(ControlRigTrackCreateEditorHandle);
-		SequencerModule->UnRegisterTrackEditor(ControlRigBindingTrackCreateEditorHandle);
-		SequencerModule->UnRegisterEditorObjectBinding(ControlRigEditorObjectBindingHandle);
-
-		SequencerModule->GetToolBarExtensibilityManager()->RemoveExtender(SequencerToolbarExtender);
-		SequencerToolbarExtender = nullptr;
+		SequencerModule->UnRegisterTrackEditor(ControlRigParameterTrackCreateEditorHandle);
 	}
 
 	FAssetToolsModule* AssetToolsModule = FModuleManager::GetModulePtr<FAssetToolsModule>("AssetTools");
@@ -429,8 +240,6 @@ void FControlRigEditorModule::ShutdownModule()
 			PropertyEditorModule->UnregisterCustomPropertyTypeLayout(PropertiesToUnregisterOnShutdown[Index]);
 		}
 	}
-
-	CommandBindings = nullptr;
 }
 
 void FControlRigEditorModule::HandleNewBlueprintCreated(UBlueprint* InBlueprint)
@@ -444,310 +253,6 @@ void FControlRigEditorModule::HandleNewBlueprintCreated(UBlueprint* InBlueprint)
 	InBlueprint->LastEditedDocuments.AddUnique(ControlRigGraph);
 }
 
-void FControlRigEditorModule::HandleSequencerCreated(TSharedRef<ISequencer> InSequencer)
-{
-	TWeakPtr<ISequencer> LocalSequencer = InSequencer;
-
-	// Record the last sequencer we opened that was editing a control rig sequence
-	UMovieSceneSequence* FocusedSequence = InSequencer->GetFocusedMovieSceneSequence();
-	if (UControlRigSequence* FocusedControlRigSequence = ExactCast<UControlRigSequence>(FocusedSequence))
-	{
-		WeakSequencer = InSequencer;
-	}
-
-	// We want to be informed of sequence activations (subsequences or not)
-	auto HandleActivateSequence = [this, LocalSequencer](FMovieSceneSequenceIDRef Ref)
-	{
-		if (LocalSequencer.IsValid())
-		{
-			TSharedRef<ISequencer> Sequencer = LocalSequencer.Pin().ToSharedRef();
-			UMovieSceneSequence* Sequence = Sequencer->GetFocusedMovieSceneSequence();
-			if (UControlRigSequence* ControlRigSequence = ExactCast<UControlRigSequence>(Sequence))
-			{
-				WeakSequencer = LocalSequencer;
-
-				GLevelEditorModeTools().ActivateMode(FControlRigEditMode::ModeName);
-
-				if (FControlRigEditMode* ControlRigEditMode = static_cast<FControlRigEditMode*>(GLevelEditorModeTools().GetActiveMode(FControlRigEditMode::ModeName)))
-				{
-					ControlRigEditMode->SetSequencer(Sequencer);
-				}
-			}
-			else
-			{
-				if (FControlRigEditMode* ControlRigEditMode = static_cast<FControlRigEditMode*>(GLevelEditorModeTools().GetActiveMode(FControlRigEditMode::ModeName)))
-				{
-					ControlRigEditMode->SetSequencer(nullptr);
-					ControlRigEditMode->SetObjects(TWeakObjectPtr<>(), FGuid(), nullptr);
-				}
-			}
-		}
-	};
-
-	InSequencer->OnActivateSequence().AddLambda(HandleActivateSequence);
-
-	// Call into activation callback to handle initial activation
-	FMovieSceneSequenceID SequenceID = MovieSceneSequenceID::Root;
-	HandleActivateSequence(SequenceID);
-
-	InSequencer->GetSelectionChangedObjectGuids().AddLambda([LocalSequencer](TArray<FGuid> InObjectBindings)
-	{
-		if (LocalSequencer.IsValid())
-		{
-			TSharedRef<ISequencer> Sequencer = LocalSequencer.Pin().ToSharedRef();
-			UMovieSceneSequence* Sequence = Sequencer->GetFocusedMovieSceneSequence();
-			if (UControlRigSequence* ControlRigSequence = ExactCast<UControlRigSequence>(Sequence))
-			{
-				TWeakObjectPtr<> SelectedObject;
-				FGuid ObjectBinding;
-				if(InObjectBindings.Num() > 0)
-				{
-					ObjectBinding = InObjectBindings[0];
-					TArrayView<TWeakObjectPtr<>> BoundObjects = Sequencer->FindBoundObjects(ObjectBinding, Sequencer->GetFocusedTemplateID());
-					if(BoundObjects.Num() > 0)
-					{
-						SelectedObject = BoundObjects[0];
-					}
-				}
-
-				if (SelectedObject.IsValid())
-				{
-					GLevelEditorModeTools().ActivateMode(FControlRigEditMode::ModeName);
-					if (FControlRigEditMode* ControlRigEditMode = static_cast<FControlRigEditMode*>(GLevelEditorModeTools().GetActiveMode(FControlRigEditMode::ModeName)))
-					{
-						ControlRigEditMode->SetObjects(SelectedObject, ObjectBinding, nullptr);
-					}
-				}
-			}
-		}
-	});
-
-	InSequencer->OnMovieSceneDataChanged().AddLambda([LocalSequencer](EMovieSceneDataChangeType DataChangeType)
-	{
-		if (LocalSequencer.IsValid())
-		{
-			TSharedRef<ISequencer> Sequencer = LocalSequencer.Pin().ToSharedRef();
-			UMovieSceneSequence* Sequence = Sequencer->GetFocusedMovieSceneSequence();
-			if (UControlRigSequence* ControlRigSequence = ExactCast<UControlRigSequence>(Sequence))
-			{
-				if (FControlRigEditMode* ControlRigEditMode = static_cast<FControlRigEditMode*>(GLevelEditorModeTools().GetActiveMode(FControlRigEditMode::ModeName)))
-				{
-					ControlRigEditMode->RefreshObjects();
-				}
-			}
-		}
-	});
-
-	InSequencer->GetSelectionChangedTracks().AddLambda([LocalSequencer](TArray<UMovieSceneTrack*> InTracks)
-	{
-		if (LocalSequencer.IsValid())
-		{
-			TSharedRef<ISequencer> Sequencer = LocalSequencer.Pin().ToSharedRef();
-			UMovieSceneSequence* Sequence = Sequencer->GetFocusedMovieSceneSequence();
-			if (UControlRigSequence* ControlRigSequence = ExactCast<UControlRigSequence>(Sequence))
-			{
-				TArray<FString> PropertyPaths;
-
-				// Look for any property tracks that might drive our rig manipulators
-				for (UMovieSceneTrack* Track : InTracks)
-				{
-					if (UMovieScenePropertyTrack* PropertyTrack = Cast<UMovieScenePropertyTrack>(Track))
-					{
-						PropertyPaths.Add(PropertyTrack->GetPropertyPath());
-					}
-				}
-
-				if (FControlRigEditMode* ControlRigEditMode = static_cast<FControlRigEditMode*>(GLevelEditorModeTools().GetActiveMode(FControlRigEditMode::ModeName)))
-				{
-					ControlRigEditMode->ClearRigElementSelection(FRigElementTypeHelper::ToMask(ERigElementType::All));
-					for (int32 Index = 0; Index < PropertyPaths.Num(); ++Index)
-					{
-						ControlRigEditMode->SetRigElementSelection(ERigElementType::Control, FName(*PropertyPaths[Index], FNAME_Find), true);
-					}
-					
-				}
-			}
-		}
-	});
-
-	InSequencer->OnPostSave().AddLambda([](ISequencer& InSequencerThatSaved)
-	{
-		UMovieSceneSequence* Sequence = InSequencerThatSaved.GetFocusedMovieSceneSequence();
-		if (UControlRigSequence* ControlRigSequence = ExactCast<UControlRigSequence>(Sequence))
-		{
-			if (FControlRigEditMode* ControlRigEditMode = static_cast<FControlRigEditMode*>(GLevelEditorModeTools().GetActiveMode(FControlRigEditMode::ModeName)))
-			{
-				ControlRigEditMode->ReBindToActor();
-			}
-		}
-	});
-
-	InSequencer->OnGetIsTrackVisible().BindRaw(this, &FControlRigEditorModule::IsTrackVisible);
-}
-
-void FControlRigEditorModule::HandleAssetEditorOpened(UObject* InAsset)
-{
-	if (UControlRigSequence* ControlRigSequence = ExactCast<UControlRigSequence>(InAsset))
-	{
-		GLevelEditorModeTools().ActivateMode(FControlRigEditMode::ModeName);
-
-		if (FControlRigEditMode* ControlRigEditMode = static_cast<FControlRigEditMode*>(GLevelEditorModeTools().GetActiveMode(FControlRigEditMode::ModeName)))
-		{
-			ControlRigEditMode->ReBindToActor();
-		}
-	}
-}
-
-void FControlRigEditorModule::OnInitializeSequence(UControlRigSequence* Sequence)
-{
-	auto* ProjectSettings = GetDefault<UMovieSceneToolsProjectSettings>();
-	UMovieScene* MovieScene = Sequence->GetMovieScene();
-	
-	FFrameNumber StartFrame = (ProjectSettings->DefaultStartTime * MovieScene->GetTickResolution()).RoundToFrame();
-	int32        Duration   = (ProjectSettings->DefaultDuration  * MovieScene->GetTickResolution()).RoundToFrame().Value;
-
-	MovieScene->SetPlaybackRange(StartFrame, Duration);
-}
-
-void FControlRigEditorModule::BindCommands()
-{
-	const FControlRigEditModeCommands& Commands = FControlRigEditModeCommands::Get();
-
-	CommandBindings->MapAction(
-		Commands.ExportAnimSequence,
-		FExecuteAction::CreateRaw(this, &FControlRigEditorModule::ExportAnimSequenceFromSequencer),
-		FCanExecuteAction(),
-		FGetActionCheckState(), 
-		FIsActionButtonVisible::CreateRaw(this, &FControlRigEditorModule::CanExportAnimSequenceFromSequencer));
-}
-
-bool FControlRigEditorModule::CanExportAnimSequenceFromSequencer() const
-{
-	if (WeakSequencer.IsValid())
-	{
-		TSharedRef<ISequencer> Sequencer = WeakSequencer.Pin().ToSharedRef();
-		return ExactCast<UControlRigSequence>(Sequencer->GetFocusedMovieSceneSequence()) != nullptr;
-	}
-
-	return false;
-}
-
-void FControlRigEditorModule::ExportAnimSequenceFromSequencer()
-{
-	// if we have an active sequencer, get the sequence
-	UControlRigSequence* ControlRigSequence = nullptr;
-	if (WeakSequencer.IsValid())
-	{
-		TSharedRef<ISequencer> Sequencer = WeakSequencer.Pin().ToSharedRef();
-		ControlRigSequence = ExactCast<UControlRigSequence>(Sequencer->GetFocusedMovieSceneSequence());
-	}
-
-	// If we are bound to an actor in the edit mode, auto pick skeletal mesh to use for binding
-	USkeletalMesh* SkeletalMesh = nullptr;
-	if (FControlRigEditMode* ControlRigEditMode = static_cast<FControlRigEditMode*>(GLevelEditorModeTools().GetActiveMode(FControlRigEditMode::ModeName)))
-	{
-		TLazyObjectPtr<AActor> ActorPtr = ControlRigEditMode->GetSettings()->Actor;
-		if (ActorPtr.IsValid())
-		{
-			if (USkeletalMeshComponent* SkeletalMeshComponent = ActorPtr->FindComponentByClass<USkeletalMeshComponent>())
-			{
-				SkeletalMesh = SkeletalMeshComponent->SkeletalMesh;
-			}
-		}
-	}
-
-	if (ControlRigSequence)
-	{
-		ControlRigSequenceConverter::Convert(ControlRigSequence, nullptr, SkeletalMesh);
-	}
-}
-
-void FControlRigEditorModule::ExportToAnimSequence(TArray<FAssetData> InAssetData)
-{
-	for (const FAssetData& AssetData : InAssetData)
-	{
-		UControlRigSequence* ControlRigSequence = Cast<UControlRigSequence>(AssetData.GetAsset());
-		if (ControlRigSequence)
-		{
-			ControlRigSequenceConverter::Convert(ControlRigSequence, nullptr, nullptr);
-		}
-	}
-}
-
-void FControlRigEditorModule::ReExportToAnimSequence(TArray<FAssetData> InAssetData)
-{
-	for (const FAssetData& AssetData : InAssetData)
-	{
-		UControlRigSequence* ControlRigSequence = Cast<UControlRigSequence>(AssetData.GetAsset());
-		if (ControlRigSequence)
-		{
-			UAnimSequence* AnimSequence = ControlRigSequence->LastExportedToAnimationSequence.LoadSynchronous();
-			USkeletalMesh* SkeletalMesh = ControlRigSequence->LastExportedUsingSkeletalMesh.LoadSynchronous();
-			bool bShowDialog = AnimSequence == nullptr || SkeletalMesh == nullptr;
-
-			ControlRigSequenceConverter::Convert(ControlRigSequence, AnimSequence, SkeletalMesh, bShowDialog);
-		}
-	}
-}
-
-void FControlRigEditorModule::ImportFromRigSequence(TArray<FAssetData> InAssetData)
-{
-	for (const FAssetData& AssetData : InAssetData)
-	{
-		UAnimSequence* AnimSequence = Cast<UAnimSequence>(AssetData.GetAsset());
-		if (AnimSequence)
-		{
-			ControlRigSequenceConverter::Convert(nullptr, AnimSequence, nullptr);
-		}
-	}
-}
-
-void FControlRigEditorModule::ReImportFromRigSequence(TArray<FAssetData> InAssetData)
-{
-	for (const FAssetData& AssetData : InAssetData)
-	{
-		UAnimSequence* AnimSequence = Cast<UAnimSequence>(AssetData.GetAsset());
-		USkeletalMesh* SkeletalMesh = nullptr;
-		UControlRigSequence* ControlRigSequence = nullptr;
-
-		FAssetRegistryModule& AssetRegistryModule = FModuleManager::LoadModuleChecked<FAssetRegistryModule>("AssetRegistry");
-
-		TMultiMap<FName, FString> TagsAndValues;
-		TagsAndValues.Add(GET_MEMBER_NAME_CHECKED(UControlRigSequence, LastExportedToAnimationSequence), AssetData.ObjectPath.ToString());
-
-		TArray<FAssetData> FoundAssets;
-		AssetRegistryModule.Get().GetAssetsByTagValues(TagsAndValues, FoundAssets);
-
-		if (FoundAssets.Num() > 0)
-		{
-			ControlRigSequence = Cast<UControlRigSequence>(FoundAssets[0].GetAsset());
-			if (ControlRigSequence)
-			{
-				SkeletalMesh = ControlRigSequence->LastExportedUsingSkeletalMesh.LoadSynchronous();
-			}
-		}
-
-		bool bShowDialog = ControlRigSequence == nullptr || AnimSequence == nullptr || SkeletalMesh == nullptr;
-
-		ControlRigSequenceConverter::Convert(ControlRigSequence, AnimSequence, SkeletalMesh, bShowDialog);
-	}
-}
-
-bool FControlRigEditorModule::IsTrackVisible(const UMovieSceneTrack* InTrack)
-{
-	if (FControlRigEditMode* ControlRigEditMode = static_cast<FControlRigEditMode*>(GLevelEditorModeTools().GetActiveMode(FControlRigEditMode::ModeName)))
-	{		
-		// If nothing selected, show all nodes
-		if (ControlRigEditMode->AreRigElementsSelected(FRigElementTypeHelper::ToMask(ERigElementType::Control)) == 0)
-		{
-			return true;
-		}
-
-		return ControlRigEditMode->SelectedRigElements.Contains(FRigElementKey(InTrack->GetTrackName(), ERigElementType::Control));
-	}
-
-	return true;
-}
 
 TSharedRef<IControlRigEditor> FControlRigEditorModule::CreateControlRigEditor(const EToolkitMode::Type Mode, const TSharedPtr< IToolkitHost >& InitToolkitHost, class UControlRigBlueprint* InBlueprint)
 {
diff --git a/Engine/Plugins/Experimental/ControlRig/Source/ControlRigEditor/Private/ControlRigEditorModule.h b/Engine/Plugins/Experimental/ControlRig/Source/ControlRigEditor/Private/ControlRigEditorModule.h
--- a/Engine/Plugins/Experimental/ControlRig/Source/ControlRigEditor/Private/ControlRigEditorModule.h
+++ b/Engine/Plugins/Experimental/ControlRig/Source/ControlRigEditor/Private/ControlRigEditorModule.h
@@ -9,8 +9,6 @@
 
 class UBlueprint;
 class IAssetTypeActions;
-class UControlRigSequence;
-class ISequencer;
 class UMaterial;
 class UAnimSequence;
 class USkeletalMesh;
@@ -32,15 +30,14 @@ public:
 	virtual TSharedRef<IControlRigEditor> CreateControlRigEditor(const EToolkitMode::Type Mode, const TSharedPtr<IToolkitHost>& InitToolkitHost, UControlRigBlueprint* Blueprint) override;
 	DECLARE_DELEGATE_RetVal_TwoParams(TSharedRef<FExtender>, FControlRigEditorToolbarExtender, const TSharedRef<FUICommandList> /*InCommandList*/, TSharedRef<IControlRigEditor> /*InControlRigEditor*/);
 	virtual TArray<FControlRigEditorToolbarExtender>& GetAllControlRigEditorToolbarExtenders() override { return ControlRigEditorToolbarExtenders; }
-
 	/** IHasMenuExtensibility interface */
 	virtual TSharedPtr<FExtensibilityManager> GetMenuExtensibilityManager() override { return MenuExtensibilityManager; }
 
 	/** IHasToolBarExtensibility interface */
 	virtual TSharedPtr<FExtensibilityManager> GetToolBarExtensibilityManager() override { return ToolBarExtensibilityManager; }
 
-	UMaterial* GetTrajectoryMaterial() const { return TrajectoryMaterial.Get(); }
 
+	
 	virtual void RegisterRigUnitEditorClass(FName RigUnitClassName, TSubclassOf<URigUnitEditor_Base> Class) override;
 	virtual void UnregisterRigUnitEditorClass(FName RigUnitClassName) override;
 	
@@ -56,74 +53,11 @@ private:
 	/** Handle a new animation controller blueprint being created */
 	void HandleNewBlueprintCreated(UBlueprint* InBlueprint);
 
-	/** Handle a new sequencer instance being created */
-	void HandleSequencerCreated(TSharedRef<ISequencer> InSequencer);
-
-	/** Handle an asset being opened */
-	void HandleAssetEditorOpened(UObject* InAsset);
-
-	/** Called to setup a new sequence's defaults */
-	static void OnInitializeSequence(UControlRigSequence* Sequence);
-
-	/** Whether we can export the current control rig sequence as an anim sequence */
-	bool CanExportAnimSequenceFromSequencer() const;
-
-	/** Export the current control rig sequence as an anim sequence */
-	void ExportAnimSequenceFromSequencer();
-
-	/** Export control rig sequence(s) to anim sequence(s) */
-	void ExportToAnimSequence(TArray<FAssetData> InAssetData);
-
-	/** Re-export control rig sequence(s) to anim sequence(s) using the previous export settings */
-	void ReExportToAnimSequence(TArray<FAssetData> InAssetData);
-
-	/** Import animation sequence(s) from a source rig sequence */
-	void ImportFromRigSequence(TArray<FAssetData> InAssetData);
-
-	/** Re-import animation sequence(s) from their source rig sequence(s) */
-	void ReImportFromRigSequence(TArray<FAssetData> InAssetData);
-
-	/** Bind our module-level commands */
-	void BindCommands();
-
-	/** Whether the track is visible in the sequencer node tree */
-	bool IsTrackVisible(const UMovieSceneTrack* InTrack);
-
-	/** Handle for our sequencer track editor */
-	FDelegateHandle ControlRigTrackCreateEditorHandle;
-
-	/** Handle for our sequencer binding track editor */
-	FDelegateHandle ControlRigBindingTrackCreateEditorHandle;
-
-	/** Handle for our sequencer object binding */
-	FDelegateHandle ControlRigEditorObjectBindingHandle;
-
-	/** Handle for our level sequence spawner */
-	FDelegateHandle LevelSequenceSpawnerDelegateHandle;
-
-	/** Handle for tracking ISequencer creation */
-	FDelegateHandle SequencerCreatedHandle;
-
-	/** Handle for tracking asset editors opening */
-	FDelegateHandle AssetEditorOpenedHandle;
+	/** Handle for our sequencer control rig parameter track editor */
+	FDelegateHandle ControlRigParameterTrackCreateEditorHandle;
 
 	TArray<TSharedRef<IAssetTypeActions>> RegisteredAssetTypeActions;
 
-	/** Material used for trajectories */
-	TWeakObjectPtr<UMaterial> TrajectoryMaterial;
-
-	/** Toolbar extender for Sequencer */
-	TSharedPtr<FExtender> SequencerToolbarExtender;
-
-	/** Command bindings for keyboard shortcuts */
-	TSharedPtr<FUICommandList> CommandBindings;
-
-	/** Weak pointer to the last sequencer that was opened */
-	TWeakPtr<ISequencer> WeakSequencer;
-
-	/** Delegate handle used to extend the content browser asset menu */
-	FDelegateHandle ContentBrowserMenuExtenderHandle;
-
 	/** StaticClass is not safe on shutdown, so we cache the name, and use this to unregister on shut down */
 	TArray<FName> ClassesToUnregisterOnShutdown;
 	TArray<FName> PropertiesToUnregisterOnShutdown;
diff --git a/Engine/Plugins/Experimental/ControlRig/Source/ControlRigEditor/Private/ControlRigEditorStyle.h b/Engine/Plugins/Experimental/ControlRig/Source/ControlRigEditor/Private/ControlRigEditorStyle.h
--- a/Engine/Plugins/Experimental/ControlRig/Source/ControlRigEditor/Private/ControlRigEditorStyle.h
+++ b/Engine/Plugins/Experimental/ControlRig/Source/ControlRigEditor/Private/ControlRigEditorStyle.h
@@ -82,7 +82,8 @@ public:
 
 		// Tree styles
 		{
-			Set("ControlRig.Tree.Bone", new IMAGE_BRUSH("Slate/Bone_16x", Icon16x16));
+			Set("ControlRig.Tree.BoneUser", new IMAGE_BRUSH("Slate/BoneUser_16x", Icon16x16));
+			Set("ControlRig.Tree.BoneImported", new IMAGE_BRUSH("Slate/BoneImported_16x", Icon16x16));
 			Set("ControlRig.Tree.Control", new IMAGE_BRUSH("Slate/Control_16x", Icon16x16));
 			Set("ControlRig.Tree.Space", new IMAGE_BRUSH("Slate/Space_16x", Icon16x16));
 		}
