Copying //UE4/Dev-Anim to Dev-Main (//UE4/Dev-Main)

diff --git a/Engine/Plugins/Experimental/ControlRig/Source/ControlRigEditor/Private/EditMode/ControlRigEditMode.cpp b/Engine/Plugins/Experimental/ControlRig/Source/ControlRigEditor/Private/EditMode/ControlRigEditMode.cpp
--- a/Engine/Plugins/Experimental/ControlRig/Source/ControlRigEditor/Private/EditMode/ControlRigEditMode.cpp
+++ b/Engine/Plugins/Experimental/ControlRig/Source/ControlRigEditor/Private/EditMode/ControlRigEditMode.cpp
@@ -80,7 +80,7 @@ FControlRigEditMode::FControlRigEditMode()
 }
 
 FControlRigEditMode::~FControlRigEditMode()
-{
+{	
 	CommandBindings = nullptr;
 
 #if WITH_EDITOR
@@ -88,32 +88,6 @@ FControlRigEditMode::~FControlRigEditMode()
 #endif
 }
 
-void FControlRigEditMode::SetSequencer(TSharedPtr<ISequencer> InSequencer)
-{
-	static bool bRecursionGuard = false;
-	if (!bRecursionGuard)
-	{
-		TGuardValue<bool> ScopeGuard(bRecursionGuard, true);
-
-		Settings->Sequence = nullptr;
-
-		WeakSequencer = InSequencer;
-		if(UsesToolkits())
-		{
-			StaticCastSharedPtr<SControlRigEditModeTools>(Toolkit->GetInlineContent())->SetSequencer(InSequencer);
-		}
-
-		if (InSequencer.IsValid())
-		{
-			if (UControlRigSequence* Sequence = ExactCast<UControlRigSequence>(InSequencer->GetFocusedMovieSceneSequence()))
-			{
-				Settings->Sequence = Sequence;
-				ReBindToActor();
-			}
-		}
-	}
-}
-
 void FControlRigEditMode::SetObjects(const TWeakObjectPtr<>& InSelectedObject, const FGuid& InObjectBinding, UObject* BindingObject)
 {
 	WeakControlRigEditing = Cast<UControlRig>(InSelectedObject.Get());
@@ -146,93 +120,24 @@ void FControlRigEditMode::SetObjects_Internal()
 		SelectedObjects.Add(Settings);
 	}
 
-	if(WeakControlRigEditing.IsValid())
+	if (WeakControlRigEditing.IsValid())
 	{
-		SelectedObjects.Add(WeakControlRigEditing);
+		//Don't add the WeakControlRig Editing...SelectedObjects.Add(WeakControlRigEditing);
 		WeakControlRigEditing.Get()->DrawInterface = &DrawInterface;
-	}
-
-	if(UsesToolkits())
-	{
-		StaticCastSharedPtr<SControlRigEditModeTools>(Toolkit->GetInlineContent())->SetDetailsObjects(SelectedObjects);
-	}
-
-	// create default manipulation layer
-	RecreateManipulationLayer();
-	HandleSelectionChanged();
-}
-
-void FControlRigEditMode::HandleBindToActor(AActor* InActor, bool bFocus)
-{
-	static bool bRecursionGuard = false;
-	if (!bRecursionGuard)
-	{
-		TGuardValue<bool> ScopeGuard(bRecursionGuard, true);
-
-		if(IsInLevelEditor())
+		if (IsInLevelEditor())
 		{
-			FControlRigBindingTemplate::SetObjectBinding(InActor);
+			WeakControlRigEditing->Hierarchy.OnElementSelected.AddSP(this, &FControlRigEditMode::OnRigElementSelected);
 		}
 
-		if (WeakSequencer.IsValid())
-		{
-			TSharedRef<ISequencer> Sequencer = WeakSequencer.Pin().ToSharedRef();
-
-			// Modify the sequence
-			if (UControlRigSequence* Sequence = ExactCast<UControlRigSequence>(Sequencer->GetFocusedMovieSceneSequence()))
-			{
-				Sequence->Modify(false);
-
-				// Also modify the binding tracks in the sequence, so bindings get regenerated to this actor
-				UMovieScene* MovieScene = Sequence->GetMovieScene();
-				for (UMovieSceneSection* Section : MovieScene->GetAllSections())
-				{
-					if (UMovieSceneSpawnSection* SpawnSection = Cast<UMovieSceneSpawnSection>(Section))
-					{
-						SpawnSection->TryModify(false);
-					}
-				}
-
-				// now notify the sequence (will rebind when it re-evaluates
-				Sequencer->NotifyMovieSceneDataChanged(EMovieSceneDataChangeType::TrackValueChangedRefreshImmediately);
-
-				// we tick first
-				ManipulationLayer->TickManipulatableObjects(0.f);
-
-				// Now re-display our objects in the details panel (they may have changed)
-				if (MovieScene->GetSpawnableCount() > 0)
-				{
-					FGuid SpawnableGuid = MovieScene->GetSpawnable(0).GetGuid();
-					TWeakObjectPtr<> BoundObject = Sequencer->FindSpawnedObjectOrTemplate(SpawnableGuid);
-					SetObjects(BoundObject, SpawnableGuid, nullptr);
-				}
-			}
 
-			if (bFocus && InActor && IsInLevelEditor())
-			{
-				const bool bNotifySelectionChanged = false;
-				const bool bDeselectBSP = true;
-				const bool bWarnAboutTooManyActors = false;
-				const bool bSelectEvenIfHidden = true;
-
-				// Select & focus the actor
-				GEditor->GetSelectedActors()->Modify();
-				GEditor->GetSelectedActors()->BeginBatchSelectOperation();
-				GEditor->SelectNone(bNotifySelectionChanged, bDeselectBSP, bWarnAboutTooManyActors);
-				GEditor->SelectActor(InActor, true, bNotifySelectionChanged, bSelectEvenIfHidden);
-				GEditor->Exec(InActor->GetWorld(), TEXT("CAMERA ALIGN ACTIVEVIEWPORTONLY"));
-				GEditor->SelectNone(bNotifySelectionChanged, bDeselectBSP, bWarnAboutTooManyActors);
-				GEditor->GetSelectedActors()->EndBatchSelectOperation();
-			}
+		if (UsesToolkits())
+		{
+			StaticCastSharedPtr<SControlRigEditModeTools>(Toolkit->GetInlineContent())->SetDetailsObjects(SelectedObjects);
 		}
-	}
-}
 
-void FControlRigEditMode::ReBindToActor()
-{
-	if (Settings->Actor.IsValid())
-	{
-		HandleBindToActor(Settings->Actor.Get(), false);
+		// create default manipulation layer
+		RecreateManipulationLayer();
+		HandleSelectionChanged();
 	}
 }
 
@@ -284,10 +189,12 @@ void FControlRigEditMode::Exit()
 
 	if (ManipulationLayer)
 	{
+		SelectNone();
 		ManipulationLayer->DestroyLayer();
 		ManipulationLayer = nullptr;
 	}
-
+	//clear actors
+	GizmoActors.SetNum(0);
 	// Call parent implementation
 	FEdMode::Exit();
 }
@@ -298,7 +205,10 @@ void FControlRigEditMode::Tick(FEditorViewportClient* ViewportClient, float Delt
 
 	ViewportClient->Invalidate();
 
-	RecalcPivotTransform();
+	if (ManipulationLayer)
+	{
+		RecalcPivotTransform();
+	}
 
 	if (bRecreateManipulationLayerRequired)
 	{
@@ -319,37 +229,35 @@ void FControlRigEditMode::Tick(FEditorViewportClient* ViewportClient, float Delt
 		bRecreateManipulationLayerRequired = false;
 	}
 
-	// If we have detached from sequencer, unbind the settings UI
-	if (!WeakSequencer.IsValid() && Settings->Sequence != nullptr)
+	// We need to tick here since changing a bone for example
+	// might have changed the transform of the Control
+	if (ManipulationLayer)
 	{
-		Settings->Sequence = nullptr;
-		RefreshObjects();
+		ManipulationLayer->TickManipulatableObjects(DeltaTime);
 	}
 }
 
 void FControlRigEditMode::Render(const FSceneView* View, FViewport* Viewport, FPrimitiveDrawInterface* PDI)
-{
-	bool bRender = true;
-	if (WeakSequencer.IsValid())
-	{
-		bRender = WeakSequencer.Pin()->GetPlaybackStatus() != EMovieScenePlayerStatus::Playing || Settings->bShowManipulatorsDuringPlayback;
-	}
-
-	if (!WeakControlRigEditing.IsValid())
+{	
+	if (!WeakControlRigEditing.IsValid() || ManipulationLayer == nullptr)
 	{
 		DrawInterface.DrawInstructions.Reset();
 		return;
 	}
 
-	// Force off manipulators if hide flag is set
-	if (Settings->bHideManipulators)
-	{
-		bRender = false;
-	}
+	bool bRender = !Settings->bHideManipulators;
 
 	FTransform ComponentTransform = (ManipulationLayer)? ManipulationLayer->GetSkeletalMeshComponentTransform() : FTransform::Identity;
 	if (bRender)
 	{
+		for (AControlRigGizmoActor* Actor : GizmoActors)
+		{
+			//Actor->SetActorHiddenInGame(bIsHidden);
+			if (GIsEditor && Actor->GetWorld() != nullptr && !Actor->GetWorld()->IsPlayInEditor())
+			{
+				Actor->SetIsTemporarilyHiddenInEditor(false);
+			}
+		}
 		if (Settings->bDisplayHierarchy)
 		{
 			UControlRig* ControlRig = WeakControlRigEditing.Get();
@@ -387,17 +295,14 @@ void FControlRigEditMode::Render(const FSceneView* View, FViewport* Viewport, FP
 				PDI->DrawLine(ElementTransform.GetTranslation(), ElementTransform.TransformPosition(FVector(0.f, 0.f, Scale)), FLinearColor::Blue, SDPG_Foreground);
 			}
 		}
-	}
-
-
-	for (const FControlRigDrawInterface::FDrawIntruction& Instruction : DrawInterface.DrawInstructions)
-	{
-		if (Instruction.Positions.Num() == 0)
-		{
-			continue;
-		}
-		switch (Instruction.DrawType)
+		for (const FControlRigDrawInterface::FDrawIntruction& Instruction : DrawInterface.DrawInstructions)
 		{
+			if (Instruction.Positions.Num() == 0)
+			{
+				continue;
+			}
+			switch (Instruction.DrawType)
+			{
 			case FControlRigDrawInterface::EDrawType_Point:
 			{
 				for (const FVector& Point : Instruction.Positions)
@@ -411,7 +316,7 @@ void FControlRigEditMode::Render(const FSceneView* View, FViewport* Viewport, FP
 				const TArray<FVector>& Points = Instruction.Positions;
 				for (int32 PointIndex = 0; PointIndex < Points.Num() - 1; PointIndex += 2)
 				{
-					PDI->DrawLine(ComponentTransform.TransformPosition(Points[PointIndex]), ComponentTransform.TransformPosition(Points[PointIndex+1]), Instruction.Color, SDPG_Foreground, Instruction.Thickness);
+					PDI->DrawLine(ComponentTransform.TransformPosition(Points[PointIndex]), ComponentTransform.TransformPosition(Points[PointIndex + 1]), Instruction.Color, SDPG_Foreground, Instruction.Thickness);
 				}
 				break;
 			}
@@ -424,8 +329,23 @@ void FControlRigEditMode::Render(const FSceneView* View, FViewport* Viewport, FP
 				}
 				break;
 			}
+			}
 		}
 	}
+	else
+	{
+		for (AControlRigGizmoActor* Actor : GizmoActors)
+		{
+			//Actor->SetActorHiddenInGame(bIsHidden);
+			if (GIsEditor && Actor->GetWorld() != nullptr && !Actor->GetWorld()->IsPlayInEditor())
+			{
+				Actor->SetIsTemporarilyHiddenInEditor(true);
+			}
+		}
+	}
+
+	//draw debug info
+	
 	DrawInterface.DrawInstructions.Reset();
 }
 
@@ -479,7 +399,10 @@ bool FControlRigEditMode::StartTracking(FEditorViewportClient* InViewportClient,
 	{
 		GEditor->BeginTransaction(LOCTEXT("MoveControlTransaction", "Move Control"));
 
-		ManipulationLayer->BeginTransaction();
+		if (ManipulationLayer != nullptr)
+		{
+			ManipulationLayer->BeginTransaction();
+		}
 		for (AControlRigGizmoActor* GizmoActor : GizmoActors)
 		{
 			GizmoActor->SetManipulating(true);
@@ -516,7 +439,7 @@ bool FControlRigEditMode::UsesTransformWidget(FWidget::EWidgetMode CheckMode) co
 {
 	for (const AControlRigGizmoActor* GizmoActor : GizmoActors)
 	{
-		if (GizmoActor->IsSelected())
+		if (GizmoActor->IsSelected() && ManipulationLayer != nullptr)
 		{
 			return ManipulationLayer->ModeSupportedByGizmoActor(GizmoActor, CheckMode);
 		}
@@ -532,7 +455,7 @@ bool FControlRigEditMode::UsesTransformWidget(FWidget::EWidgetMode CheckMode) co
 
 FVector FControlRigEditMode::GetWidgetLocation() const
 {
-	if (AreRigElementSelectedAndMovable())
+	if (AreRigElementSelectedAndMovable() && ManipulationLayer != nullptr)
 	{
 		FTransform ComponentTransform = ManipulationLayer->GetSkeletalMeshComponentTransform();
 		return ComponentTransform.TransformPosition(PivotTransform.GetLocation());
@@ -566,7 +489,7 @@ bool FControlRigEditMode::HandleClick(FEditorViewportClient* InViewportClient, H
 			if (ActorHitProxy->Actor->IsA<AControlRigGizmoActor>())
 			{
 				AControlRigGizmoActor* GizmoActor = CastChecked<AControlRigGizmoActor>(ActorHitProxy->Actor);
-				const FControlData* ControlData = ManipulationLayer->GetControlDataFromGizmo(GizmoActor);
+				const FControlData* ControlData = ManipulationLayer != nullptr ? ManipulationLayer->GetControlDataFromGizmo(GizmoActor) :  nullptr;
 
 				if (ControlData)
 				{
@@ -580,18 +503,18 @@ bool FControlRigEditMode::HandleClick(FEditorViewportClient* InViewportClient, H
 						ClearRigElementSelection(FRigElementTypeHelper::ToMask(ERigElementType::Control));
 						SetRigElementSelection(ERigElementType::Control, ControlName, true);
 					}
+
+				}
+				// for now we show this menu all the time if body is selected
+				// if we want some global menu, we'll have to move this
+				if (Click.GetKey() == EKeys::RightMouseButton)
+				{
+					OpenContextMenu(InViewportClient);
 				}
-			}
-		}
 
-		// for now we show this menu all the time if body is selected
-		// if we want some global menu, we'll have to move this
-		if (Click.GetKey() == EKeys::RightMouseButton)
-		{
-			OpenContextMenu(InViewportClient);
+				return true;
+			}
 		}
-
-		return true;
 	}
 
 	// for now we show this menu all the time if body is selected
@@ -605,12 +528,6 @@ bool FControlRigEditMode::HandleClick(FEditorViewportClient* InViewportClient, H
 	// clear selected controls
 	ClearRigElementSelection(FRigElementTypeHelper::ToMask(ERigElementType::All));
 
-	// If we are animating then swallow clicks so we dont select things other than controls
-	if(WeakSequencer.IsValid() && WeakSequencer.Pin()->GetFocusedMovieSceneSequence()->IsA<UControlRigSequence>())
-	{
-		return true;
-	}
-
 	return FEdMode::HandleClick(InViewportClient, HitProxy, Click);
 }
 
@@ -647,24 +564,28 @@ void FControlRigEditMode::OpenContextMenu(FEditorViewportClient* InViewportClien
 
 bool FControlRigEditMode::IntersectSelect(bool InSelect, const TFunctionRef<bool(const AControlRigGizmoActor*, const FTransform&)>& Intersects)
 {
-	FTransform ComponentTransform = ManipulationLayer->GetSkeletalMeshComponentTransform();
-
-	bool bSelected = false;
-	for (AControlRigGizmoActor* GizmoActor: GizmoActors)
+	if (ManipulationLayer != nullptr)
 	{
-		const FTransform ControlTransform = GizmoActor->GetGlobalTransform() * ComponentTransform;
-		if (Intersects(GizmoActor, ControlTransform))
+		FTransform ComponentTransform = ManipulationLayer->GetSkeletalMeshComponentTransform();
+
+		bool bSelected = false;
+		for (AControlRigGizmoActor* GizmoActor : GizmoActors)
 		{
-			const FControlData* ControlData = ManipulationLayer->GetControlDataFromGizmo(GizmoActor);
-			if (ControlData)
+			const FTransform ControlTransform = GizmoActor->GetGlobalTransform() * ComponentTransform;
+			if (Intersects(GizmoActor, ControlTransform))
 			{
-				SetRigElementSelection(ERigElementType::Control, ControlData->ControlName, InSelect);
-				bSelected = true;
+				const FControlData* ControlData = ManipulationLayer->GetControlDataFromGizmo(GizmoActor);
+				if (ControlData)
+				{
+					SetRigElementSelection(ERigElementType::Control, ControlData->ControlName, InSelect);
+					bSelected = true;
+				}
 			}
 		}
-	}
 
-	return bSelected;
+		return bSelected;
+	}
+	return false;
 }
 
 bool FControlRigEditMode::BoxSelect(FBox& InBox, bool InSelect)
@@ -737,9 +658,9 @@ bool FControlRigEditMode::InputDelta(FEditorViewportClient* InViewportClient, FV
 		const bool bDoTranslation = !Drag.IsZero() && (WidgetMode == FWidget::WM_Translate || WidgetMode == FWidget::WM_TranslateRotateZ);
 		const bool bDoScale = !Scale.IsZero() && WidgetMode == FWidget::WM_Scale;
 
-		FTransform ComponentTransform = ManipulationLayer->GetSkeletalMeshComponentTransform();
-		if (AreRigElementsSelected(FRigElementTypeHelper::ToMask(ERigElementType::Control)))
+		if (ManipulationLayer != nullptr && AreRigElementsSelected(FRigElementTypeHelper::ToMask(ERigElementType::Control)))
 		{
+			FTransform ComponentTransform = ManipulationLayer->GetSkeletalMeshComponentTransform();
 			for (AControlRigGizmoActor* GizmoActor : GizmoActors)
 			{
 				if (GizmoActor->IsSelected())
@@ -758,8 +679,10 @@ bool FControlRigEditMode::InputDelta(FEditorViewportClient* InViewportClient, FV
 			}
 			return true;
 		}
-		else if (AreRigElementSelectedAndMovable())
+		else if (ManipulationLayer != nullptr && AreRigElementSelectedAndMovable())
 		{
+			FTransform ComponentTransform = ManipulationLayer->GetSkeletalMeshComponentTransform();
+
 			// set Bone transform
 			// that will set initial Bone transform
 			for (int32 Index = 0; Index < SelectedRigElements.Num(); ++Index)
@@ -831,12 +754,20 @@ bool FControlRigEditMode::IsCompatibleWith(FEditorModeID OtherModeID) const
 
 void FControlRigEditMode::AddReferencedObjects( FReferenceCollector& Collector )
 {
-	Collector.AddReferencedObject(Settings);
-	Collector.AddReferencedObject(ManipulationLayer);
-
-	for (AControlRigGizmoActor* GizmoActor : GizmoActors)
+	if (Settings)
 	{
-		Collector.AddReferencedObject(GizmoActor);
+		Collector.AddReferencedObject(Settings);
+	}
+	if (ManipulationLayer)
+	{
+		Collector.AddReferencedObject(ManipulationLayer);
+	}
+	if (GizmoActors.Num() > 0)
+	{
+		for (AControlRigGizmoActor* GizmoActor : GizmoActors)
+		{
+			Collector.AddReferencedObject(GizmoActor);
+		}
 	}
 }
 
@@ -853,6 +784,10 @@ void FControlRigEditMode::ClearRigElementSelection(uint32 InTypes)
 	{
 		Blueprint->HierarchyContainer.ClearSelection();
 	}
+	if(IsInLevelEditor())
+	{
+		WeakControlRigEditing->Hierarchy.ClearSelection();
+	}
 }
 
 // internal private function that doesn't use guarding.
@@ -868,6 +803,10 @@ void FControlRigEditMode::SetRigElementSelectionInternal(ERigElementType Type, c
 	{
 		Blueprint->HierarchyContainer.Select(FRigElementKey(InRigElementName, Type), bSelected);
 	}
+	if(IsInLevelEditor())
+	{
+		WeakControlRigEditing->Hierarchy.Select(FRigElementKey(InRigElementName, Type), bSelected);
+	}
 }
 
 void FControlRigEditMode::SetRigElementSelection(ERigElementType Type, const FName& InRigElementName, bool bSelected)
@@ -933,72 +872,13 @@ int32 FControlRigEditMode::GetNumSelectedRigElements(uint32 InTypes) const
 	return 0;
 }
 
-void FControlRigEditMode::HandleObjectSpawned(FGuid InObjectBinding, UObject* SpawnedObject, IMovieScenePlayer& Player)
-{
-	if (WeakSequencer.IsValid())
-	{
-		// check whether this spawned object is from our sequence
-		TSharedPtr<ISequencer> Sequencer = WeakSequencer.Pin();
-		if (Sequencer.Get() == &Player)
-		{
-			RefreshObjects();
-
-			// check if the object is being displayed currently
-			if (ControlRigGuid == InObjectBinding)
-			{
-				if (WeakControlRigEditing != Cast<UControlRig>(SpawnedObject))
-				{
-					WeakControlRigEditing = Cast<UControlRig>(SpawnedObject);
-					SetObjects_Internal();
-				}
-				return;
-			}
-
-			// We didnt find an existing Guid, so set up our internal cache
-			if (!ControlRigGuid.IsValid())
-			{
-				SetObjects(SpawnedObject, InObjectBinding, nullptr);
-				if (UControlRig* ControlRig = Cast<UControlRig>(SpawnedObject))
-				{
-					if (TSharedPtr<IControlRigObjectBinding> ObjectBinding = ControlRig->GetObjectBinding())
-					{
-						if (Settings->Actor.IsValid() && ObjectBinding->GetBoundObject() == nullptr)
-						{
-							ObjectBinding->BindToObject(Settings->Actor.Get());
-						}
-					}
-				}
-				ReBindToActor();
-			}
-		}
-	}
-}
 
 void FControlRigEditMode::RefreshObjects()
 {
-	if (WeakSequencer.IsValid())
-	{
-		TSharedRef<ISequencer> Sequencer = WeakSequencer.Pin().ToSharedRef();
-		UMovieScene* MovieScene = Sequencer->GetFocusedMovieSceneSequence() ? Sequencer->GetFocusedMovieSceneSequence()->GetMovieScene() : nullptr;
-		if (MovieScene)
-		{
-			// check if we have an invalid Guid & invalidate Guid if so
-			if (ControlRigGuid.IsValid() && MovieScene->FindSpawnable(ControlRigGuid) == nullptr)
-			{
-				ControlRigGuid.Invalidate();
-				WeakControlRigEditing = nullptr;
-			}
+	WeakControlRigEditing = nullptr;
+	ControlRigGuid.Invalidate();
 
-			SetObjects_Internal();
-		}
-	}
-	else
-	{
-		WeakControlRigEditing = nullptr;
-		ControlRigGuid.Invalidate();
-
-		SetObjects_Internal();
-	}
+	SetObjects_Internal();
 }
 
 void FControlRigEditMode::EnableRigElementEditing(bool bEnabled)
@@ -1086,23 +966,6 @@ void FControlRigEditMode::RecalcPivotTransform()
 
 void FControlRigEditMode::HandleSelectionChanged()
 {
-	if (WeakSequencer.IsValid())
-	{
-		TArray<FString> SelectedNames;
-		for (int32 Index=0; Index<SelectedRigElements.Num(); ++Index)
-		{
-			// if it's control, make sure it's selected
-			if (SelectedRigElements[Index].Type == ERigElementType::Control)
-			{
-				SelectedNames.Add(SelectedRigElements[Index].Name.ToString());
-			}
-		}
-
-		if (SelectedNames.Num() > 0)
-		{
-			WeakSequencer.Pin()->SelectByPropertyPaths(SelectedNames);
-		}
-	}
 
 	for (AControlRigGizmoActor* GizmoActor : GizmoActors)
 	{
@@ -1114,14 +977,7 @@ void FControlRigEditMode::HandleSelectionChanged()
 		}
 	}
 
-	if (WeakSequencer.IsValid())
-	{
-		if (WeakSequencer.Pin()->GetSequencerSettings()->GetShowSelectedNodesOnly())
-		{
-			WeakSequencer.Pin()->RefreshTree();
-		}
-	}
-
+	
 	// update the pivot transform of our selected objects (they could be animating)
 	RecalcPivotTransform();
 }
@@ -1130,10 +986,6 @@ void FControlRigEditMode::BindCommands()
 {
 	const FControlRigEditModeCommands& Commands = FControlRigEditModeCommands::Get();
 
-	CommandBindings->MapAction(
-		Commands.SetKey,
-		FExecuteAction::CreateRaw(this, &FControlRigEditMode::SetKeysForSelectedManipulators));
-
 	CommandBindings->MapAction(
 		Commands.ToggleManipulators,
 		FExecuteAction::CreateRaw(this, &FControlRigEditMode::ToggleManipulators));
@@ -1161,7 +1013,7 @@ bool FControlRigEditMode::GetRigElementGlobalTransform(const FRigElementKey& InE
 		AControlRigGizmoActor* GizmoActor = GetGizmoFromControlName(InElement.Name);
 		if (GizmoActor && ensure(GizmoActor->IsSelected()))
 		{
-			if (ManipulationLayer->GetGlobalTransform(GizmoActor, InElement.Name, OutGlobalTransform))
+			if (ManipulationLayer != nullptr && ManipulationLayer->GetGlobalTransform(GizmoActor, InElement.Name, OutGlobalTransform))
 			{
 				return true;
 			}
@@ -1184,21 +1036,6 @@ bool FControlRigEditMode::GetRigElementGlobalTransform(const FRigElementKey& InE
 	return false;
 }
 
-void FControlRigEditMode::SetKeysForSelectedManipulators()
-{
-	TSharedPtr<ISequencer> Sequencer = WeakSequencer.Pin();
-
-	if (Sequencer)
-	{
-		for (const AControlRigGizmoActor* GizmoActor : GizmoActors)
-		{
-			if (GizmoActor->IsSelected())
-			{
-				ManipulationLayer->AddKey(GizmoActor, Sequencer);
-			}
-		}
-	}
-}
 
 void FControlRigEditMode::ToggleManipulators()
 {
@@ -1322,19 +1159,20 @@ bool FControlRigEditMode::IsTransformDelegateAvailable() const
 
 bool FControlRigEditMode::AreRigElementSelectedAndMovable() const
 {
-	if (IsTransformDelegateAvailable())
+	if (AreRigElementsSelected(FRigElementTypeHelper::ToMask(ERigElementType::Control)))
 	{
-		// you can move control anytime
-		if (AreRigElementsSelected(FRigElementTypeHelper::ToMask(ERigElementType::Control)))
-		{
-			return true;
-		}
+		return true;
+	}
 
-		// if bone, we want to make sure execution is stopped
-		if (bEnableRigElementDefaultPoseEditing && AreRigElementsSelected(FRigElementTypeHelper::ToMask(ERigElementType::Bone)))
-		{
-			return true;
-		}
+	// if bone, we want to make sure execution is stopped
+	if (bEnableRigElementDefaultPoseEditing && AreRigElementsSelected(FRigElementTypeHelper::ToMask(ERigElementType::Bone)))
+	{
+		return true;
+	}
+
+	if (IsTransformDelegateAvailable())
+	{
+		return true;
 	}
 
 	return false;
@@ -1419,7 +1257,7 @@ void FControlRigEditMode::OnControlUISettingChanged(FRigHierarchyContainer* Cont
 
 AControlRigGizmoActor* FControlRigEditMode::GetGizmoFromControlName(const FName& InControlName) const
 {
-	return ManipulationLayer->GetGizmoFromControlName(InControlName);
+	return ManipulationLayer ? ManipulationLayer->GetGizmoFromControlName(InControlName) : nullptr;
 }
 
 #undef LOCTEXT_NAMESPACE
diff --git a/Engine/Plugins/Experimental/ControlRig/Source/ControlRigEditor/Private/EditMode/ControlRigEditMode.h b/Engine/Plugins/Experimental/ControlRig/Source/ControlRigEditor/Private/EditMode/ControlRigEditMode.h
--- a/Engine/Plugins/Experimental/ControlRig/Source/ControlRigEditor/Private/EditMode/ControlRigEditMode.h
+++ b/Engine/Plugins/Experimental/ControlRig/Source/ControlRigEditor/Private/EditMode/ControlRigEditMode.h
@@ -42,9 +42,6 @@ public:
 	/** Set the objects to be displayed in the details panel */
 	void SetObjects(const TWeakObjectPtr<>& InSelectedObject, const FGuid& InObjectBinding, UObject* BindingObject);
 
-	/** Set the sequencer we are bound to */
-	void SetSequencer(TSharedPtr<ISequencer> InSequencer);
-
 	/** This edit mode is re-used between the level editor and the control rig editor. Calling this indicates which context we are in */
 	virtual bool IsInLevelEditor() const { return true; }
 
@@ -80,18 +77,6 @@ public:
 	/** FGCObject interface */
 	virtual void AddReferencedObjects( FReferenceCollector& Collector ) override;
 
-	/** 
-	 * Lets the edit mode know that an object has just been spawned. 
-	 * Allows us to redisplay different underlying objects in the details panel.
-	 */
-	void HandleObjectSpawned(FGuid InObjectBinding, UObject* SpawnedObject, IMovieScenePlayer& Player);
-
-	/** Re-bind to the current actor - used when sequence, selection etc. changes */
-	void ReBindToActor();
-
-	/** Bind us to an actor for editing */
-	void HandleBindToActor(AActor* InActor, bool bFocus);
-
 	/** Refresh our internal object list (they may have changed) */
 	void RefreshObjects();
 
@@ -121,6 +106,8 @@ public:
 	/** Enable RigElement Editing */
 	void EnableRigElementEditing(bool bEnabled);
 
+	/** Get Control Rig, could be more than one later,  we are animating. Currently used by Sequencer for cross selection.*/
+	UControlRig*  GetControlRig() { return WeakControlRigEditing.IsValid() ? WeakControlRigEditing.Get() : nullptr; }
 protected:
 
 	// Gizmo related functions wrt enable/selection
@@ -140,9 +127,6 @@ protected:
 	/** Handle selection internally */
 	void HandleSelectionChanged();
 
-	/** Set keys on all selected manipulators */
-	void SetKeysForSelectedManipulators();
-
 	/** Toggles visibility of manipulators in the viewport */
 	void ToggleManipulators();
 
@@ -176,9 +160,6 @@ protected:
 	/** The draw interface to use for the control rig */
 	FControlRigDrawInterface DrawInterface;
 
-	/** Sequencer we are currently bound to */
-	TWeakPtr<ISequencer> WeakSequencer;
-
 	/** Guard value for selection */
 	bool bSelecting;
 
diff --git a/Engine/Plugins/Experimental/ControlRig/Source/ControlRigEditor/Private/EditMode/ControlRigEditModeSettings.cpp b/Engine/Plugins/Experimental/ControlRig/Source/ControlRigEditor/Private/EditMode/ControlRigEditModeSettings.cpp
--- a/Engine/Plugins/Experimental/ControlRig/Source/ControlRigEditor/Private/EditMode/ControlRigEditModeSettings.cpp
+++ b/Engine/Plugins/Experimental/ControlRig/Source/ControlRigEditor/Private/EditMode/ControlRigEditModeSettings.cpp
@@ -10,51 +10,10 @@
 
 void UControlRigEditModeSettings::PreEditChange(UProperty* PropertyAboutToChange)
 {
-	if (PropertyAboutToChange)
-	{
-		if (PropertyAboutToChange->GetFName() == GET_MEMBER_NAME_CHECKED(UControlRigEditModeSettings, Actor))
-		{
-			PrevActor = Actor.Get();
-		}
-	}
+	
 }
 
 void UControlRigEditModeSettings::PostEditChangeProperty(struct FPropertyChangedEvent& PropertyChangedEvent)
 {
-	if (UProperty* Property = PropertyChangedEvent.Property)
-	{
-		if (Property->GetFName() == GET_MEMBER_NAME_CHECKED(UControlRigEditModeSettings, Actor))
-		{
-			if (Actor.IsValid())
-			{
-				if (Actor->FindComponentByClass<USkeletalMeshComponent>())
-				{
-					// user set the actor we are targeting, so bind our (standalone) sequence to it
-					if (FControlRigEditMode* ControlRigEditMode = static_cast<FControlRigEditMode*>(GLevelEditorModeTools().GetActiveMode(FControlRigEditMode::ModeName)))
-					{
-						ControlRigEditMode->HandleBindToActor(Actor.Get(), true);
-					}
-				}
-				else
-				{
-					Actor = PrevActor;
-				}
-			}
-			else
-			{
-				// user cleared the actor we are targeting, so (un)bind our sequence to it
-				if (FControlRigEditMode* ControlRigEditMode = static_cast<FControlRigEditMode*>(GLevelEditorModeTools().GetActiveMode(FControlRigEditMode::ModeName)))
-				{
-					ControlRigEditMode->HandleBindToActor(nullptr, true);
-				}
-			}
-		}
-		else if (Property->GetFName() == GET_MEMBER_NAME_CHECKED(UControlRigEditModeSettings, Sequence))
-		{
-			if (Sequence)
-			{
-				GEditor->GetEditorSubsystem<UAssetEditorSubsystem>()->OpenEditorForAsset(Sequence);
-			}
-		}
-	}
+	
 }
diff --git a/Engine/Plugins/Experimental/ControlRig/Source/ControlRigEditor/Private/EditMode/ControlRigEditModeSettings.h b/Engine/Plugins/Experimental/ControlRig/Source/ControlRigEditor/Private/EditMode/ControlRigEditModeSettings.h
--- a/Engine/Plugins/Experimental/ControlRig/Source/ControlRigEditor/Private/EditMode/ControlRigEditModeSettings.h
+++ b/Engine/Plugins/Experimental/ControlRig/Source/ControlRigEditor/Private/EditMode/ControlRigEditModeSettings.h
@@ -19,7 +19,6 @@ class UControlRigEditModeSettings : public UObject
 
 	UControlRigEditModeSettings()
 		: bDisplayHierarchy(false)
-		, bShowManipulatorsDuringPlayback(true)
 		, bHideManipulators(false)
 		, bDisplayAxesOnSelection(true)
 		, AxisScale(10.f)
@@ -30,22 +29,11 @@ class UControlRigEditModeSettings : public UObject
 	virtual void PostEditChangeProperty(struct FPropertyChangedEvent& PropertyChangedEvent) override;
 
 public:
-	/** Sequence to animate */
-	UPROPERTY(EditAnywhere, Category = "Animation", NoClear)
-	UControlRigSequence* Sequence;
-
-	/** The actor we are currently animating */
-	UPROPERTY(EditAnywhere, Category = "Animation")
-	TLazyObjectPtr<AActor> Actor;
 
 	/** Whether to show all nodes in the hierarchy being animated */
 	UPROPERTY(EditAnywhere, AdvancedDisplay, Category = "Animation")
 	bool bDisplayHierarchy;
 
-	/** Whether to show manipulators when animations are being played back */
-	UPROPERTY(EditAnywhere, AdvancedDisplay, Category = "Animation")
-	bool bShowManipulatorsDuringPlayback;
-
 	/** Should we always hide manipulators in viewport */
 	UPROPERTY(EditAnywhere, AdvancedDisplay, Category = "Animation")
 	bool bHideManipulators;
@@ -58,7 +46,4 @@ public:
 	UPROPERTY(EditAnywhere, AdvancedDisplay, Category = "Animation")
 	float AxisScale;
 
-private:
-	/** Cache the previous actor for pre/post edit change handling */
-	AActor* PrevActor;
 };
\ No newline at end of file
diff --git a/Engine/Plugins/Experimental/ControlRig/Source/ControlRigEditor/Private/EditMode/SControlRigEditModeTools.cpp b/Engine/Plugins/Experimental/ControlRig/Source/ControlRigEditor/Private/EditMode/SControlRigEditModeTools.cpp
--- a/Engine/Plugins/Experimental/ControlRig/Source/ControlRigEditor/Private/EditMode/SControlRigEditModeTools.cpp
+++ b/Engine/Plugins/Experimental/ControlRig/Source/ControlRigEditor/Private/EditMode/SControlRigEditModeTools.cpp
@@ -73,7 +73,7 @@ void SControlRigEditModeTools::Construct(const FArguments& InArgs, UWorld* InWor
 		+ SScrollBox::Slot()
 		[
 			SNew(SVerticalBox)
-
+			/*
 			+ SVerticalBox::Slot()
 			.AutoHeight()
 			[
@@ -87,6 +87,7 @@ void SControlRigEditModeTools::Construct(const FArguments& InArgs, UWorld* InWor
 					SAssignNew(ControlPicker, SControlPicker, InWorld)
 				]
 			]
+			*/
 
 			+SVerticalBox::Slot()
 			.AutoHeight()
@@ -112,10 +113,8 @@ void SControlRigEditModeTools::SetDetailsObjects(const TArray<TWeakObjectPtr<>>&
 		}
 	}
 
-	ControlPicker->SetControlRig(Rig);
+	//ControlPicker->SetControlRig(Rig);
 
-	// Expand when you have a rig, collapse when set to null
-	PickerExpander->SetExpanded(Rig != nullptr);
 }
 
 void SControlRigEditModeTools::SetSequencer(TSharedPtr<ISequencer> InSequencer)
diff --git a/Engine/Plugins/Experimental/ControlRig/Source/ControlRigEditor/Private/Editor/ControlRigEditor.cpp b/Engine/Plugins/Experimental/ControlRig/Source/ControlRigEditor/Private/Editor/ControlRigEditor.cpp
--- a/Engine/Plugins/Experimental/ControlRig/Source/ControlRigEditor/Private/Editor/ControlRigEditor.cpp
+++ b/Engine/Plugins/Experimental/ControlRig/Source/ControlRigEditor/Private/Editor/ControlRigEditor.cpp
@@ -999,6 +999,12 @@ void FControlRigEditor::HandleModelModified(const UControlRigModel* InModel, ECo
 {
 	DECLARE_SCOPE_HIERARCHICAL_COUNTER_FUNC()
 
+	UControlRigBlueprint* RigBlueprint = Cast<UControlRigBlueprint>(GetBlueprintObj());
+	if (RigBlueprint == nullptr)
+	{
+		return;
+	}
+
 	switch (InType)
 	{
 		case EControlRigModelNotifType::NodeSelected:
@@ -1007,26 +1013,79 @@ void FControlRigEditor::HandleModelModified(const UControlRigModel* InModel, ECo
 			const FControlRigModelNode* Node = (const FControlRigModelNode*)InPayload;
 			if (Node != nullptr)
 			{
-				if (!bIsSelecting)
+				if (FocusedGraphEdPtr.IsValid())
 				{
-					TGuardValue<bool> SelectingGuard(bIsSelecting, true);
-					if (FocusedGraphEdPtr.IsValid())
+					TSharedPtr<SGraphEditor> FocusedGraphEd = FocusedGraphEdPtr.Pin();
+					if (UControlRigGraph* RigGraph = Cast<UControlRigGraph>(FocusedGraphEd->GetCurrentGraph()))
 					{
-						TSharedPtr<SGraphEditor> FocusedGraphEd = FocusedGraphEdPtr.Pin();
-						if (UControlRigGraph* RigGraph = Cast<UControlRigGraph>(FocusedGraphEd->GetCurrentGraph()))
+						TArray<FControlRigModelNode> Selection;
+						bool bSelect = InType == EControlRigModelNotifType::NodeSelected;
+						if (bSelect)
+						{
+							Selection = InModel->SelectedNodes();
+						}
+						else
+						{
+							Selection.Add(*Node);
+						}
+
+						TArray<FRigElementKey> RigElementsToSelect;
+						for (const FControlRigModelNode& SelectedNode : Selection)
+						{
+							if (UControlRigGraphNode* RigNode = Cast<UControlRigGraphNode>(RigGraph->FindNodeFromPropertyName(SelectedNode.Name)))
+							{
+								for (const FControlRigModelPin& Pin : Node->Pins)
+								{
+									if (Pin.Type.PinCategory == UEdGraphSchema_K2::PC_Name)
+									{
+										if (Pin.CustomWidgetName == TEXT("BoneName"))
+										{
+											RigElementsToSelect.Add(FRigElementKey(*Pin.DefaultValue, ERigElementType::Bone));
+										}
+										else if (Pin.CustomWidgetName == TEXT("ControlName"))
+										{
+											RigElementsToSelect.Add(FRigElementKey(*Pin.DefaultValue, ERigElementType::Control));
+										}
+										else if (Pin.CustomWidgetName == TEXT("SpaceName"))
+										{
+											RigElementsToSelect.Add(FRigElementKey(*Pin.DefaultValue, ERigElementType::Space));
+										}
+										else if (Pin.CustomWidgetName == TEXT("CurveName"))
+										{
+											RigElementsToSelect.Add(FRigElementKey(*Pin.DefaultValue, ERigElementType::Curve));
+										}
+									}
+								}
+							}
+						}
+
+						if (RigElementsToSelect.Num() > 0)
+						{
+							if (bSelect)
+							{
+								RigBlueprint->HierarchyContainer.ClearSelection();
+							}
+							for (const FRigElementKey& RigElementToSelect : RigElementsToSelect)
+							{
+								RigBlueprint->HierarchyContainer.Select(RigElementToSelect, bSelect);
+							}
+						}
+
+						if (!bIsSelecting)
 						{
+							TGuardValue<bool> SelectingGuard(bIsSelecting, true);
 							if (UEdGraphNode* EdNode = RigGraph->FindNodeFromPropertyName(Node->Name))
 							{
 								FocusedGraphEd->SetNodeSelection(EdNode, InType == EControlRigModelNotifType::NodeSelected);
 							}
+							break;
 						}
 					}
-					break;
 				}
 
 				if (InType == EControlRigModelNotifType::NodeSelected)
 				{
-					UClass* Class = GetBlueprintObj()->GeneratedClass.Get();
+					UClass* Class = RigBlueprint->GeneratedClass.Get();
 					if (Class)
 					{
 						if (UProperty* Property = Class->FindPropertyByName(Node->Name))
@@ -1650,6 +1709,12 @@ void FControlRigEditor::OnFinishedChangingProperties(const FPropertyChangedEvent
 	if (ControlRig && ControlRigBP && RigElementInDetailPanel)
 	{
 		UControlRig* DebuggedControlRig = Cast<UControlRig>(GetBlueprintObj()->GetObjectBeingDebugged());
+
+		if (DebuggedControlRig && DebuggedControlRig->GetHierarchy()->GetIndex(RigElementInDetailPanel) == INDEX_NONE)
+		{
+			return;
+		}
+
 		UScriptStruct* ScriptStruct = Cast< UScriptStruct>(PropertyChangedEvent.Property->GetOuter());
 		if (ScriptStruct)
 		{
diff --git a/Engine/Plugins/Experimental/ControlRig/Source/ControlRigEditor/Private/Editor/ControlRigEditorEditMode.cpp b/Engine/Plugins/Experimental/ControlRig/Source/ControlRigEditor/Private/Editor/ControlRigEditorEditMode.cpp
--- a/Engine/Plugins/Experimental/ControlRig/Source/ControlRigEditor/Private/Editor/ControlRigEditorEditMode.cpp
+++ b/Engine/Plugins/Experimental/ControlRig/Source/ControlRigEditor/Private/Editor/ControlRigEditorEditMode.cpp
@@ -13,11 +13,17 @@ bool FControlRigEditorEditMode::GetCameraTarget(FSphere& OutTarget) const
 
 	for (int32 Index = 0; Index < SelectedRigElements.Num(); ++Index)
 	{
-		if (SelectedRigElements[Index].Type == ERigElementType::Bone || SelectedRigElements[Index].Type == ERigElementType::Control)
+		if (SelectedRigElements[Index].Type == ERigElementType::Bone || SelectedRigElements[Index].Type == ERigElementType::Space)
 		{
 			FTransform Transform = OnGetRigElementTransformDelegate.Execute(SelectedRigElements[Index], false);
 			Box += Transform.GetLocation();
 		}
+		else if (SelectedRigElements[Index].Type == ERigElementType::Control)
+		{
+			FTransform Transform = OnGetRigElementTransformDelegate.Execute(SelectedRigElements[Index], false);
+			Box += Transform.TransformPosition(FVector::OneVector * 50.f); // we assume 100cm as the base size for a control
+			Box += Transform.TransformPosition(FVector::OneVector * -50.f);
+		}
 	}
 
 	if(Box.IsValid)
diff --git a/Engine/Plugins/Experimental/ControlRig/Source/ControlRigEditor/Private/Editor/SRigHierarchy.cpp b/Engine/Plugins/Experimental/ControlRig/Source/ControlRigEditor/Private/Editor/SRigHierarchy.cpp
--- a/Engine/Plugins/Experimental/ControlRig/Source/ControlRigEditor/Private/Editor/SRigHierarchy.cpp
+++ b/Engine/Plugins/Experimental/ControlRig/Source/ControlRigEditor/Private/Editor/SRigHierarchy.cpp
@@ -31,6 +31,7 @@
 #include "ControlRig.h"
 #include "ControlRigEditorStyle.h"
 #include "HAL/PlatformApplicationMisc.h"
+#include "Dialogs/Dialogs.h"
 
 #define LOCTEXT_NAMESPACE "SRigHierarchy"
 
@@ -114,9 +115,36 @@ void SRigHierarchyItem::Construct(const FArguments& InArgs, TSharedPtr<FControlR
 			Brush = FControlRigEditorStyle::Get().GetBrush("ControlRig.Tree.Space");
 			break;
 		}
+		case ERigElementType::Bone:
+		{
+			ERigBoneType BoneType = ERigBoneType::User;
+
+			int32 BoneIndex = InHierarchy->ControlRigBlueprint->HierarchyContainer.GetIndex(InRigTreeElement->Key);
+			if (BoneIndex != INDEX_NONE)
+			{
+				BoneType = InHierarchy->ControlRigBlueprint->HierarchyContainer.BoneHierarchy[BoneIndex].Type;
+			}
+
+			switch (BoneType)
+			{
+				case ERigBoneType::Imported:
+				{
+					Brush = FControlRigEditorStyle::Get().GetBrush("ControlRig.Tree.BoneImported");
+					break;
+				}
+				case ERigBoneType::User:
+				default:
+				{
+					Brush = FControlRigEditorStyle::Get().GetBrush("ControlRig.Tree.BoneUser");
+					break;
+				}
+			}
+
+			break;
+		}
 		default:
 		{
-			Brush = FControlRigEditorStyle::Get().GetBrush("ControlRig.Tree.Bone");
+			Brush = FControlRigEditorStyle::Get().GetBrush("ControlRig.Tree.BoneUser");
 			break;
 		}
 	}
@@ -568,10 +596,14 @@ void SRigHierarchy::AddElement(FRigElementKey InKey, FRigElementKey InParentKey)
 			RootElements.Add(NewItem);
 		}
 	}
-	else if (InKey.Name.ToString().Contains(FilteredString))
+	else
 	{
-		TSharedPtr<FRigTreeElement> NewItem = MakeShared<FRigTreeElement>(InKey, SharedThis(this));
-		RootElements.Add(NewItem);
+		FString FilteredStringUnderScores = FilteredString.Replace(TEXT(" "), TEXT("_"));
+		if (InKey.Name.ToString().Contains(FilteredString) || InKey.Name.ToString().Contains(FilteredStringUnderScores))	
+		{
+			TSharedPtr<FRigTreeElement> NewItem = MakeShared<FRigTreeElement>(InKey, SharedThis(this));
+			RootElements.Add(NewItem);
+		}
 	}
 }
 
@@ -827,19 +859,7 @@ void SRigHierarchy::RefreshHierarchy(const FAssetData& InAssetData)
 		ControlRigBlueprint->Modify();
 
 		const FReferenceSkeleton& RefSkeleton = Mesh->RefSkeleton;
-		const TArray<FMeshBoneInfo>& BoneInfos = RefSkeleton.GetRefBoneInfo();
-		const TArray<FTransform>& BonePoses = RefSkeleton.GetRefBonePose();
-
-		for (int32 Index = 0; Index < RefSkeleton.GetNum(); ++Index)
-		{
-			// only add if you don't have it. This may change in the future
-			int32 RigIndex = Hierarchy->BoneHierarchy.GetIndex(BoneInfos[Index].Name);
-			if (RigIndex != INDEX_NONE)
-			{
-				// @todo: add optimized version without sorting, but if no sort, we should make sure not to use find index function
-				Hierarchy->BoneHierarchy.SetInitialTransform(RigIndex, FAnimationRuntime::GetComponentSpaceTransform(RefSkeleton, BonePoses, Index));
-			}
-		}
+		Hierarchy->BoneHierarchy.ImportSkeleton(RefSkeleton, NAME_None, true, true, true);
 	}
 }
 void SRigHierarchy::CreateImportMenu(FMenuBuilder& MenuBuilder)
@@ -879,28 +899,7 @@ void SRigHierarchy::ImportHierarchy(const FAssetData& InAssetData)
 		ControlRigBlueprint->Modify();
 
 		const FReferenceSkeleton& RefSkeleton = Mesh->RefSkeleton;
-		const TArray<FMeshBoneInfo>& BoneInfos = RefSkeleton.GetRefBoneInfo();
-		const TArray<FTransform>& BonePoses = RefSkeleton.GetRefBonePose();
-
-		Hierarchy->ClearSelection();
-
-		TArray<FName> AddedBones;
-		for (int32 Index = 0; Index < RefSkeleton.GetNum(); ++Index)
-		{
-			// only add if you don't have it. This may change in the future
-			if (Hierarchy->BoneHierarchy.GetIndex(BoneInfos[Index].Name) == INDEX_NONE)
-			{
-				// @todo: add optimized version without sorting, but if no sort, we should make sure not to use find index function
-				FName ParentName = (BoneInfos[Index].ParentIndex != INDEX_NONE) ? BoneInfos[BoneInfos[Index].ParentIndex].Name : NAME_None;
-				Hierarchy->BoneHierarchy.Add(BoneInfos[Index].Name, ParentName, FAnimationRuntime::GetComponentSpaceTransform(RefSkeleton, BonePoses, Index));
-				AddedBones.Add(BoneInfos[Index].Name);
-			}
-		}
-
-		for (const FName& AddedBone : AddedBones)
-		{
-			Hierarchy->BoneHierarchy.Select(AddedBone);
-		}
+		Hierarchy->BoneHierarchy.ImportSkeleton(RefSkeleton, NAME_None, false, false, true);
 
 		FSlateApplication::Get().DismissAllMenus();
 		RefreshTreeView();
@@ -959,6 +958,9 @@ void SRigHierarchy::HandleDeleteItem()
 		// clear detail view display
 		ControlRigEditor.Pin()->ClearDetailObject();
 
+		bool bConfirmedByUser = false;
+		bool bDeleteImportedBones = false;
+
 		TArray<TSharedPtr<FRigTreeElement>> SelectedItems = TreeView->GetSelectedItems();
 		for (int32 ItemIndex = 0; ItemIndex < SelectedItems.Num(); ++ItemIndex)
 		{
@@ -969,8 +971,31 @@ void SRigHierarchy::HandleDeleteItem()
 			{
 				case ERigElementType::Bone:
 				{
-					if (Hierarchy->BoneHierarchy.GetIndex(SelectedKey.Name) != INDEX_NONE)
+					int32 BoneIndex = Hierarchy->BoneHierarchy.GetIndex(SelectedKey.Name);
+					if (BoneIndex != INDEX_NONE)
 					{
+						const FRigBone& Bone = Hierarchy->BoneHierarchy[BoneIndex];
+						if (Bone.Type == ERigBoneType::Imported && Bone.ParentIndex != INDEX_NONE)
+						{
+							if (!bConfirmedByUser)
+							{
+								FText ConfirmDelete = LOCTEXT("ConfirmDeleteBoneHierarchy",
+									"Deleting imported(white) bones can cause issues with animation - are you sure ?");
+
+								FSuppressableWarningDialog::FSetupInfo Info(ConfirmDelete, LOCTEXT("DeleteImportedBone", "Delete Imported Bone"), "DeleteImportedBoneHierarchy_Warning");
+								Info.ConfirmText = LOCTEXT("DeleteImportedBoneHierarchy_Yes", "Yes");
+								Info.CancelText = LOCTEXT("DeleteImportedBoneHierarchy_No", "No");
+
+								FSuppressableWarningDialog DeleteImportedBonesInHierarchy(Info);
+								bDeleteImportedBones = DeleteImportedBonesInHierarchy.ShowModal() != FSuppressableWarningDialog::Cancel;
+								bConfirmedByUser = true;
+							}
+
+							if (!bDeleteImportedBones)
+							{
+								break;
+							}
+						}
 						Hierarchy->BoneHierarchy.Remove(SelectedKey.Name);
 					}
 					break;
@@ -1040,7 +1065,7 @@ void SRigHierarchy::HandleNewItem(ERigElementType InElementType)
 				{
 					FName ParentName = ParentKey.Type == ERigElementType::Bone ? ParentKey.Name : NAME_None;
 					ParentTransform = ParentName == NAME_None ? FTransform::Identity : ParentTransform;
-					Hierarchy->BoneHierarchy.Add(NewElementName, ParentName, ParentTransform);
+					Hierarchy->BoneHierarchy.Add(NewElementName, ParentName, ERigBoneType::User, ParentTransform);
 					break;
 				}
 				case ERigElementType::Control:
@@ -1147,7 +1172,7 @@ void SRigHierarchy::HandleDuplicateItem()
 					FName ParentName = Hierarchy->BoneHierarchy[Key.Name].ParentName;
 
 					const FName NewName = CreateUniqueName(Key.Name, Key.Type);
-					FRigBone& NewBone = Hierarchy->BoneHierarchy.Add(NewName, ParentName, Transform);
+					FRigBone& NewBone = Hierarchy->BoneHierarchy.Add(NewName, ParentName, ERigBoneType::User, Transform);
 					NewKeys.Add(NewBone.GetElementKey());
 					break;
 				}
@@ -1205,6 +1230,29 @@ void SRigHierarchy::HandleRenameItem()
 		TArray<TSharedPtr<FRigTreeElement>> SelectedItems = TreeView->GetSelectedItems();
 		if (SelectedItems.Num() == 1)
 		{
+			if (SelectedItems[0]->Key.Type == ERigElementType::Bone)
+			{
+				int32 BoneIndex = Hierarchy->GetIndex(SelectedItems[0]->Key);
+				if (BoneIndex != INDEX_NONE)
+				{
+					const FRigBone& Bone = Hierarchy->BoneHierarchy[BoneIndex];
+					if (Bone.Type == ERigBoneType::Imported)
+					{
+						FText ConfirmRename = LOCTEXT("RenameDeleteBoneHierarchy",
+							"Renaming imported(white) bones can cause issues with animation - are you sure ?");
+
+						FSuppressableWarningDialog::FSetupInfo Info(ConfirmRename, LOCTEXT("RenameImportedBone", "Rename Imported Bone"), "RenameImportedBoneHierarchy_Warning");
+						Info.ConfirmText = LOCTEXT("RenameImportedBoneHierarchy_Yes", "Yes");
+						Info.CancelText = LOCTEXT("RenameImportedBoneHierarchy_No", "No");
+
+						FSuppressableWarningDialog RenameImportedBonesInHierarchy(Info);
+						if (RenameImportedBonesInHierarchy.ShowModal() == FSuppressableWarningDialog::Cancel)
+						{
+							return;
+						}
+					}
+				}
+			}
 			SelectedItems[0]->RequestRename();
 		}
 	}
@@ -1354,7 +1402,7 @@ FReply SRigHierarchy::OnDragDetected(const FGeometry& MyGeometry, const FPointer
 		DraggedElements.Add(SelectedItem->Key);
 	}
 
-	if (MouseEvent.IsMouseButtonDown(EKeys::LeftMouseButton))
+	if (MouseEvent.IsMouseButtonDown(EKeys::LeftMouseButton) && DraggedElements.Num() > 0)
 	{
 		if (ControlRigEditor.IsValid())
 		{
@@ -1451,7 +1499,31 @@ FReply SRigHierarchy::OnAcceptDrop(const FDragDropEvent& DragDropEvent, EItemDro
 					DraggedKey.Type,
 					Container->GetIndex(DraggedKey)))
 				{
-					return FReply::Unhandled();;
+					return FReply::Unhandled();
+				}
+
+				if (DraggedKey.Type == ERigElementType::Bone)
+				{
+					int32 BoneIndex = Container->BoneHierarchy.GetIndex(DraggedKey.Name);
+					if (BoneIndex != INDEX_NONE)
+					{
+						const FRigBone& Bone = Container->BoneHierarchy[BoneIndex];
+						if (Bone.Type == ERigBoneType::Imported && Bone.ParentIndex != INDEX_NONE)
+						{
+							FText ConfirmReparent = LOCTEXT("ConfirmReparentBoneHierarchy",
+								"Reparenting imported(white) bones can cause issues with animation - are you sure ?");
+
+							FSuppressableWarningDialog::FSetupInfo Info(ConfirmReparent, LOCTEXT("ReparentImportedBone", "Reparent Imported Bone"), "ReparentImportedBoneHierarchy_Warning");
+							Info.ConfirmText = LOCTEXT("ReparentImportedBoneHierarchy_Yes", "Yes");
+							Info.CancelText = LOCTEXT("ReparentImportedBoneHierarchy_No", "No");
+
+							FSuppressableWarningDialog ReparentImportedBonesInHierarchy(Info);
+							if (ReparentImportedBonesInHierarchy.ShowModal() == FSuppressableWarningDialog::Cancel)
+							{
+								return FReply::Unhandled();
+							}
+						}
+					}
 				}
 			}
 
