Copying //UE4/Dev-Anim to Dev-Main (//UE4/Dev-Main)

diff --git a/Engine/Plugins/Experimental/ControlRig/Source/ControlRigEditor/Private/Sequencer/ControlRigParameterTrackEditor.cpp b/Engine/Plugins/Experimental/ControlRig/Source/ControlRigEditor/Private/Sequencer/ControlRigParameterTrackEditor.cpp
new file mode 100644
--- /dev/null
+++ b/Engine/Plugins/Experimental/ControlRig/Source/ControlRigEditor/Private/Sequencer/ControlRigParameterTrackEditor.cpp
@@ -0,0 +1,1315 @@
+// Copyright 1998-2019 Epic Games, Inc. All Rights Reserved.
+
+#include "ControlRigParameterTrackEditor.h"
+#include "Sequencer/MovieSceneControlRigParameterTrack.h"
+#include "Sequencer/MovieSceneControlRigParameterSection.h"
+#include "Framework/Commands/Commands.h"
+#include "Rendering/DrawElements.h"
+#include "Widgets/SBoxPanel.h"
+#include "Framework/MultiBox/MultiBoxBuilder.h"
+#include "GameFramework/Actor.h"
+#include "AssetData.h"
+#include "Modules/ModuleManager.h"
+#include "Layout/WidgetPath.h"
+#include "Framework/Application/MenuStack.h"
+#include "Framework/Application/SlateApplication.h"
+#include "Widgets/Layout/SBox.h"
+#include "SequencerSectionPainter.h"
+#include "Components/SkeletalMeshComponent.h"
+#include "Editor/UnrealEdEngine.h"
+#include "UnrealEdGlobals.h"
+#include "ClassViewerModule.h"
+#include "ClassViewerFilter.h"
+#include "CommonMovieSceneTools.h"
+#include "AssetRegistryModule.h"
+#include "IContentBrowserSingleton.h"
+#include "ContentBrowserModule.h"
+#include "MatineeImportTools.h"
+#include "Matinee/InterpTrackAnimControl.h"
+#include "SequencerUtilities.h"
+#include "ISectionLayoutBuilder.h"
+#include "EditorStyleSet.h"
+#include "MovieSceneTimeHelpers.h"
+#include "Fonts/FontMeasure.h"
+#include "AnimationEditorUtils.h"
+#include "Misc/MessageDialog.h"
+#include "Framework/Notifications/NotificationManager.h"
+#include "Widgets/Notifications/SNotificationList.h"
+#include "Toolkits/AssetEditorManager.h"
+#include "Engine/SCS_Node.h"
+#include "Engine/SimpleConstructionScript.h"
+#include "Engine/Blueprint.h"
+#include "ControlRig.h"
+#include "ControlRigEditMode.h"
+#include "EditorModeManager.h"
+#include "Engine/Selection.h"
+#include "Manipulatable/IControlRigManipulatable.h"
+#include "ControlRigSkeletalMeshBinding.h"
+#include "LevelEditorViewport.h"
+#include "IKeyArea.h"
+
+#define LOCTEXT_NAMESPACE "FControlRigParameterTrackEditor"
+
+
+
+class FControlRigTrackCommands
+	: public TCommands<FControlRigTrackCommands>
+{
+public:
+
+	FControlRigTrackCommands()
+		: TCommands<FControlRigTrackCommands>
+		(
+			"ControlRigTrack",
+			NSLOCTEXT("Contexts", "ControlRigTrack", "ControlRigTrack"),
+			NAME_None,
+			FEditorStyle::GetStyleSetName() // Icon Style Set
+			)
+		, BindingCount(0)
+	{ }
+
+	/** Sets a transform key at the current time for the selected actor */
+	TSharedPtr< FUICommandInfo > AddTransformKey;
+
+	/** Sets a translation key at the current time for the selected actor */
+	TSharedPtr< FUICommandInfo > AddTranslationKey;
+
+	/** Sets a rotation key at the current time for the selected actor */
+	TSharedPtr< FUICommandInfo > AddRotationKey;
+
+	/** Sets a scale key at the current time for the selected actor */
+	TSharedPtr< FUICommandInfo > AddScaleKey;
+
+	/**
+	 * Initialize commands
+	 */
+	virtual void RegisterCommands() override;
+
+	mutable uint32 BindingCount;
+};
+
+void FControlRigTrackCommands::RegisterCommands()
+{
+	UI_COMMAND(AddTransformKey, "Add Transform Key", "Add a transform key at the current time for the selected control rig control.", EUserInterfaceActionType::Button, FInputChord(EKeys::S));
+	UI_COMMAND(AddTranslationKey, "Add Translation Key", "Add a translation key at the current time for the selected control rig control.", EUserInterfaceActionType::Button, FInputChord(EModifierKey::Shift, EKeys::W));
+	UI_COMMAND(AddRotationKey, "Add Rotation Key", "Add a rotation key at the current time for the selected control rig control.", EUserInterfaceActionType::Button, FInputChord(EModifierKey::Shift, EKeys::E));
+	UI_COMMAND(AddScaleKey, "Add Scale Key", "Add a scale key at the current time for the selected control rig control.", EUserInterfaceActionType::Button, FInputChord(EModifierKey::Shift, EKeys::R));
+}
+
+static USkeleton* GetSkeletonFromComponent(UActorComponent* InComponent)
+{
+	USkeletalMeshComponent* SkeletalMeshComp = Cast<USkeletalMeshComponent>(InComponent);
+	if (SkeletalMeshComp && SkeletalMeshComp->SkeletalMesh && SkeletalMeshComp->SkeletalMesh->Skeleton)
+	{
+		// @todo Multiple actors, multiple components
+		return SkeletalMeshComp->SkeletalMesh->Skeleton;
+	}
+
+	return nullptr;
+}
+
+static USkeleton* AcquireSkeletonFromObjectGuid(const FGuid& Guid, UObject** Object, TSharedPtr<ISequencer> SequencerPtr)
+{
+	UObject* BoundObject = SequencerPtr.IsValid() ? SequencerPtr->FindSpawnedObjectOrTemplate(Guid) : nullptr;
+	*Object = BoundObject;
+	if (AActor* Actor = Cast<AActor>(BoundObject))
+	{
+		for (UActorComponent* Component : Actor->GetComponents())
+		{
+			if (USkeleton* Skeleton = GetSkeletonFromComponent(Component))
+			{
+				return Skeleton;
+			}
+		}
+
+		AActor* ActorCDO = Cast<AActor>(Actor->GetClass()->GetDefaultObject());
+		if (ActorCDO)
+		{
+			for (UActorComponent* Component : ActorCDO->GetComponents())
+			{
+				if (USkeleton* Skeleton = GetSkeletonFromComponent(Component))
+				{
+					return Skeleton;
+				}
+			}
+		}
+
+		UBlueprintGeneratedClass* ActorBlueprintGeneratedClass = Cast<UBlueprintGeneratedClass>(Actor->GetClass());
+		if (ActorBlueprintGeneratedClass)
+		{
+			const TArray<USCS_Node*>& ActorBlueprintNodes = ActorBlueprintGeneratedClass->SimpleConstructionScript->GetAllNodes();
+
+			for (USCS_Node* Node : ActorBlueprintNodes)
+			{
+				if (Node->ComponentClass->IsChildOf(USkeletalMeshComponent::StaticClass()))
+				{
+					if (USkeleton* Skeleton = GetSkeletonFromComponent(Node->GetActualComponentTemplate(ActorBlueprintGeneratedClass)))
+					{
+						return Skeleton;
+					}
+				}
+			}
+		}
+	}
+	else if (USkeletalMeshComponent* SkeletalMeshComponent = Cast<USkeletalMeshComponent>(BoundObject))
+	{
+		if (USkeleton* Skeleton = GetSkeletonFromComponent(SkeletalMeshComponent))
+		{
+			return Skeleton;
+		}
+	}
+
+	return nullptr;
+}
+
+FControlRigParameterTrackEditor::FControlRigParameterTrackEditor(TSharedRef<ISequencer> InSequencer)
+	: FKeyframeTrackEditor<UMovieSceneControlRigParameterTrack>(InSequencer), bIsDoingASelection(false)
+{
+	SelectionChangedHandle = InSequencer->GetSelectionChangedTracks().AddRaw(this, &FControlRigParameterTrackEditor::OnSelectionChanged);
+	SequencerChangedHandle = InSequencer->OnMovieSceneDataChanged().AddRaw(this, &FControlRigParameterTrackEditor::OnSequencerDataChanged);
+
+	FControlRigTrackCommands::Register();
+	//register all modified/selections for control rigs
+	UMovieScene* MovieScene = InSequencer->GetFocusedMovieSceneSequence()->GetMovieScene();
+	const TArray<FMovieSceneBinding>& Bindings = MovieScene->GetBindings();
+	for(const FMovieSceneBinding& Binding : Bindings)
+	{
+		UMovieSceneControlRigParameterTrack* Track = Cast<UMovieSceneControlRigParameterTrack>(MovieScene->FindTrack(UMovieSceneControlRigParameterTrack::StaticClass(), Binding.GetObjectGuid(), NAME_None));
+		if (Track && Track->GetControlRig())
+		{
+			Track->GetControlRig()->ControlModified().AddRaw(this, &FControlRigParameterTrackEditor::HandleControlModified);
+			Track->GetControlRig()->ControlSelected().AddRaw(this, &FControlRigParameterTrackEditor::HandleControlSelected);
+		}
+	}
+}
+
+FControlRigParameterTrackEditor::~FControlRigParameterTrackEditor()
+{
+}
+
+void FControlRigParameterTrackEditor::OnRelease()
+{
+	if (GetSequencer().IsValid())
+	{
+		if (SelectionChangedHandle.IsValid())
+		{
+			GetSequencer()->GetSelectionChangedTracks().Remove(SelectionChangedHandle);
+		}
+		if (SequencerChangedHandle.IsValid())
+		{
+			GetSequencer()->OnMovieSceneDataChanged().Remove(SequencerChangedHandle);
+		}
+		if (GetSequencer()->GetFocusedMovieSceneSequence() && GetSequencer()->GetFocusedMovieSceneSequence()->GetMovieScene())
+		{
+			UMovieScene* MovieScene = GetSequencer()->GetFocusedMovieSceneSequence()->GetMovieScene();
+			const TArray<FMovieSceneBinding>& Bindings = MovieScene->GetBindings();
+			for (const FMovieSceneBinding& Binding : Bindings)
+			{
+				UMovieSceneControlRigParameterTrack* Track = Cast<UMovieSceneControlRigParameterTrack>(MovieScene->FindTrack(UMovieSceneControlRigParameterTrack::StaticClass(), Binding.GetObjectGuid(), NAME_None));
+				if (Track && Track->GetControlRig())
+				{
+					Track->GetControlRig()->ControlModified().RemoveAll(this);
+					Track->GetControlRig()->ControlSelected().RemoveAll(this);
+				}
+			}
+		}
+	}
+	FControlRigEditMode* ControlRigEditMode = static_cast<FControlRigEditMode*>(GLevelEditorModeTools().FindMode(FControlRigEditMode::ModeName));
+	if (ControlRigEditMode)
+	{
+		if (GLevelEditorModeTools().GetActiveMode(FControlRigEditMode::ModeName))
+		{
+			if (GLevelEditorModeTools().HasToolkitHost())
+			{
+				GLevelEditorModeTools().DeactivateMode(FControlRigEditMode::ModeName);
+			}
+		}
+		ControlRigEditMode->SetObjects(nullptr, FGuid(), nullptr);
+	}
+
+	const FControlRigTrackCommands& Commands = FControlRigTrackCommands::Get();
+	Commands.BindingCount--;
+
+	if (Commands.BindingCount < 1)
+	{
+		FControlRigTrackCommands::Unregister();
+	}
+
+}
+
+TSharedRef<ISequencerTrackEditor> FControlRigParameterTrackEditor::CreateTrackEditor(TSharedRef<ISequencer> InSequencer)
+{
+	return MakeShareable(new FControlRigParameterTrackEditor(InSequencer));
+}
+
+
+bool FControlRigParameterTrackEditor::SupportsType(TSubclassOf<UMovieSceneTrack> Type) const
+{
+	return Type == UMovieSceneControlRigParameterTrack::StaticClass();
+}
+
+
+TSharedRef<ISequencerSection> FControlRigParameterTrackEditor::MakeSectionInterface(UMovieSceneSection& SectionObject, UMovieSceneTrack& Track, FGuid ObjectBinding)
+{
+	check(SupportsType(SectionObject.GetOuter()->GetClass()));
+
+	return MakeShareable(new FControlRigParameterSection(SectionObject, GetSequencer()));
+}
+
+
+void FControlRigParameterTrackEditor::BuildObjectBindingTrackMenu(FMenuBuilder& MenuBuilder, const TArray<FGuid>& ObjectBindings, const UClass* ObjectClass)
+{
+
+	if (ObjectClass->IsChildOf(USkeletalMeshComponent::StaticClass()) || ObjectClass->IsChildOf(AActor::StaticClass()) || ObjectClass->IsChildOf(UChildActorComponent::StaticClass()))
+	{
+		const TSharedPtr<ISequencer> ParentSequencer = GetSequencer();
+		UObject *BoundObject = nullptr;
+		USkeleton* Skeleton = AcquireSkeletonFromObjectGuid(ObjectBindings[0], &BoundObject, ParentSequencer);
+
+		if (Skeleton)
+		{
+			//if there are any other control rigs we don't allow it for now..
+			//mz todo will allow later
+			UMovieScene* MovieScene = GetSequencer()->GetFocusedMovieSceneSequence()->GetMovieScene();
+			UMovieSceneControlRigParameterTrack* ExistingTrack = Cast<UMovieSceneControlRigParameterTrack>(MovieScene->FindTrack(UMovieSceneControlRigParameterTrack::StaticClass(), ObjectBindings[0], NAME_None));
+			if (!ExistingTrack)
+			{
+				UMovieSceneTrack* Track = nullptr;
+
+				MenuBuilder.AddSubMenu(
+					LOCTEXT("AddControlRig", "Animation ControlRig"), NSLOCTEXT("Sequencer", "AddControlRigTooltip", "Adds an animation ControlRig track."),
+					FNewMenuDelegate::CreateRaw(this, &FControlRigParameterTrackEditor::AddControlRigSubMenu, ObjectBindings, Track)
+				);
+			}
+		}
+	}
+}
+
+void FControlRigParameterTrackEditor::AddControlRigSubMenu(FMenuBuilder& MenuBuilder, TArray<FGuid> ObjectBindings, UMovieSceneTrack* Track)
+{
+	/*
+	MenuBuilder.BeginSection(TEXT("ChooseSequence"), LOCTEXT("ChooseSequence", "Choose Sequence"));
+	{
+		FAssetPickerConfig AssetPickerConfig;
+		{
+			AssetPickerConfig.OnAssetSelected = FOnAssetSelected::CreateRaw(this, &FControlRigTrackEditor::OnSequencerAssetSelected, ObjectBindings, Track);
+			AssetPickerConfig.OnAssetEnterPressed = FOnAssetEnterPressed::CreateRaw(this, &FControlRigTrackEditor::OnSequencerAssetEnterPressed, ObjectBindings, Track);
+			AssetPickerConfig.bAllowNullSelection = false;
+			AssetPickerConfig.InitialAssetViewType = EAssetViewType::List;
+			AssetPickerConfig.Filter.bRecursiveClasses = true;
+			AssetPickerConfig.Filter.ClassNames.Add(UControlRigSequence::StaticClass()->GetFName());
+		}
+
+		FContentBrowserModule& ContentBrowserModule = FModuleManager::Get().LoadModuleChecked<FContentBrowserModule>(TEXT("ContentBrowser"));
+
+		TSharedPtr<SBox> MenuEntry = SNew(SBox)
+			.WidthOverride(300.0f)
+			.HeightOverride(300.0f)
+			[
+				ContentBrowserModule.Get().CreateAssetPicker(AssetPickerConfig)
+			];
+
+		MenuBuilder.AddWidget(MenuEntry.ToSharedRef(), FText::GetEmpty(), true);
+	}
+	MenuBuilder.EndSection();
+	*/
+
+
+	const TSharedPtr<ISequencer> ParentSequencer = GetSequencer();
+	UObject *BoundObject = nullptr;
+	//todo support multiple bindings?
+	USkeleton* Skeleton = AcquireSkeletonFromObjectGuid(ObjectBindings[0], &BoundObject, GetSequencer());
+
+	if (Skeleton)
+	{
+		//MenuBuilder.AddSubMenu(
+		//	LOCTEXT("AddControlRigTrack", "ControlRigTrack"), NSLOCTEXT("ControlRig", "AddControlRigTrack", "Adds a Control Rigtrack."),
+		//	FNewMenuDelegate::CreateRaw(this, &FControlRigParameterTrackEditor::AddAnimationSubMenu, BoundObject, ObjectBindings[0], Skeleton)
+		//);
+
+
+		class FControlRigClassFilter : public IClassViewerFilter
+		{
+		public:
+			bool IsClassAllowed(const FClassViewerInitializationOptions& InInitOptions, const UClass* InClass, TSharedRef< FClassViewerFilterFuncs > InFilterFuncs) override
+			{
+				const bool bChildOfObjectClass = InClass->IsChildOf(UControlRig::StaticClass());
+				const bool bMatchesFlags = !InClass->HasAnyClassFlags(CLASS_Hidden | CLASS_HideDropDown | CLASS_Deprecated | CLASS_Abstract);
+				return bChildOfObjectClass && bMatchesFlags;
+			}
+
+			virtual bool IsUnloadedClassAllowed(const FClassViewerInitializationOptions& InInitOptions, const TSharedRef< const IUnloadedBlueprintData > InUnloadedClassData, TSharedRef< FClassViewerFilterFuncs > InFilterFuncs) override
+			{
+				const bool bChildOfObjectClass = InUnloadedClassData->IsChildOf(UControlRig::StaticClass());
+				const bool bMatchesFlags = !InUnloadedClassData->HasAnyClassFlags(CLASS_Hidden | CLASS_HideDropDown | CLASS_Deprecated | CLASS_Abstract);
+				return bChildOfObjectClass && bMatchesFlags;
+			}
+		};
+
+		FClassViewerInitializationOptions Options;
+		Options.bShowUnloadedBlueprints = true;
+		Options.NameTypeToDisplay = EClassViewerNameTypeToDisplay::DisplayName;
+
+		TSharedPtr<FControlRigClassFilter> ClassFilter = MakeShareable(new FControlRigClassFilter);
+		Options.ClassFilter = ClassFilter;
+		Options.bShowNoneOption = false;
+
+		FClassViewerModule& ClassViewerModule = FModuleManager::LoadModuleChecked<FClassViewerModule>("ClassViewer");
+		//	FNewMenuDelegate::CreateRaw(this, &FControlRigParameterTrackEditor::AddAnimationSubMenu, BoundObject, ObjectBindings[0], Skeleton)
+
+
+
+		TSharedRef<SWidget> ClassViewer = ClassViewerModule.CreateClassViewer(Options, FOnClassPicked::CreateRaw(this, &FControlRigParameterTrackEditor::AddControlRig, BoundObject, ObjectBindings[0], Skeleton));
+		MenuBuilder.AddWidget(ClassViewer, FText::GetEmpty(), true);
+
+		/*
+		FAssetPickerConfig AssetPickerConfig;
+		{
+			AssetPickerConfig.OnAssetSelected = FOnAssetSelected::CreateRaw(this, &FControlRigParameterTrackEditor::OnControlRigAssetSelected, ObjectBindings, Skeleton);
+			AssetPickerConfig.OnAssetEnterPressed = FOnAssetEnterPressed::CreateRaw(this, &FControlRigParameterTrackEditor::OnControlRigAssetEnterPressed, ObjectBindings, Skeleton);
+			AssetPickerConfig.bAllowNullSelection = false;
+			AssetPickerConfig.InitialAssetViewType = EAssetViewType::List;
+			AssetPickerConfig.OnShouldFilterAsset = FOnShouldFilterAsset::CreateRaw(this, &FControlRigParameterTrackEditor::ShouldFilterAsset);
+			AssetPickerConfig.Filter.bRecursiveClasses = true;
+			AssetPickerConfig.Filter.ClassNames.Add((UControlRig::StaticClass())->GetFName());
+			AssetPickerConfig.Filter.TagsAndValues.Add(TEXT("Skeleton"), FAssetData(Skeleton).GetExportTextName());
+		}
+
+		FContentBrowserModule& ContentBrowserModule = FModuleManager::Get().LoadModuleChecked<FContentBrowserModule>(TEXT("ContentBrowser"));
+
+		TSharedPtr<SBox> MenuEntry = SNew(SBox)
+			.WidthOverride(300.0f)
+			.HeightOverride(300.f)
+			[
+				ContentBrowserModule.Get().CreateAssetPicker(AssetPickerConfig)
+			];
+
+		MenuBuilder.AddWidget(MenuEntry.ToSharedRef(), FText::GetEmpty(), true);
+		*/
+	}
+
+}
+/*
+bool FControlRigParameterTrackEditor::ShouldFilterAsset(const FAssetData& AssetData)
+{
+	if (AssetData.AssetClass == UControlRig::StaticClass()->GetFName())
+	{
+		return true;
+	}
+	return false;
+}
+
+void FControlRigParameterTrackEditor::OnControlRigAssetSelected(const FAssetData& AssetData, TArray<FGuid> ObjectBindings, USkeleton* Skeleton)
+{
+	FSlateApplication::Get().DismissAllMenus();
+
+	UObject* SelectedObject = AssetData.GetAsset();
+	TSharedPtr<ISequencer> SequencerPtr = GetSequencer();
+
+	if (SelectedObject && SelectedObject->IsA(UControlRig::StaticClass()) && SequencerPtr.IsValid())
+	{
+		UControlRig* ControlRig = CastChecked<UControlRig>(AssetData.GetAsset());
+
+		const FScopedTransaction Transaction(LOCTEXT("AddAnimation_Transaction", "Add Animation"));
+
+		for (FGuid ObjectBinding : ObjectBindings)
+		{
+			UObject* Object = SequencerPtr->FindSpawnedObjectOrTemplate(ObjectBinding);
+			int32 RowIndex = INDEX_NONE;
+			//AnimatablePropertyChanged(FOnKeyProperty::CreateRaw(this, &FControlRigParameterTrackEditor::AddControlRig, BoundObject, ObjectBindings[0], Skeleton)));
+		}
+	}
+}
+
+void FControlRigParameterTrackEditor::OnControlRigAssetEnterPressed(const TArray<FAssetData>& AssetData, TArray<FGuid> ObjectBindings, USkeleton* Skeleton)
+{
+	if (AssetData.Num() > 0)
+	{
+		OnControlRigAssetSelected(AssetData[0].GetAsset(), ObjectBindings, Skeleton);
+	}
+}
+
+*/
+
+void FControlRigParameterTrackEditor::AddControlRig(UClass* InClass, UObject* BoundSkelMesh, FGuid SkelMeshBinding, USkeleton* Skeleton)
+{
+	FSlateApplication::Get().DismissAllMenus();
+	const TSharedPtr<ISequencer> SequencerParent = GetSequencer();
+
+	if (InClass && InClass->IsChildOf(UControlRig::StaticClass()) && SequencerParent.IsValid())
+	{
+		FScopedTransaction AddControlRigTrackTransaction(LOCTEXT("AddControlRigTrack_Transaction", "Add Control Rig Track"));
+
+		UMovieSceneControlRigParameterTrack* Track = Cast<UMovieSceneControlRigParameterTrack>(AddTrack(GetSequencer()->GetFocusedMovieSceneSequence()->GetMovieScene(), SkelMeshBinding, UMovieSceneControlRigParameterTrack::StaticClass(), NAME_None));
+		if (Track)
+		{
+
+			FString ObjectName = (InClass->GetName());
+			ObjectName.RemoveFromEnd(TEXT("_C"));
+
+			UControlRig* ControlRig = NewObject<UControlRig>(Track, InClass, FName(*ObjectName), RF_Transactional);
+
+			ControlRig->SetObjectBinding(MakeShared<FControlRigSkeletalMeshBinding>());
+			ControlRig->GetObjectBinding()->BindToObject(BoundSkelMesh);
+			ControlRig->Initialize();
+			ControlRig->Execute(EControlRigState::Update);
+
+			SequencerParent->NotifyMovieSceneDataChanged(EMovieSceneDataChangeType::MovieSceneStructureItemsChanged);
+
+			Track->Modify();
+			UMovieSceneSection* NewSection = Track->CreateControlRigSection(0, ControlRig);
+
+			//mz todo need to have multiple rigs with same class
+			Track->SetTrackName(FName(*ObjectName));
+			Track->SetDisplayName(FText::FromString(ObjectName));
+
+			GetSequencer()->EmptySelection();
+			GetSequencer()->SelectSection(NewSection);
+			GetSequencer()->ThrobSectionSelection();
+			GetSequencer()->NotifyMovieSceneDataChanged(EMovieSceneDataChangeType::MovieSceneStructureItemAdded);
+
+			FControlRigEditMode* ControlRigEditMode = static_cast<FControlRigEditMode*>(GLevelEditorModeTools().GetActiveMode(FControlRigEditMode::ModeName));
+			if (!ControlRigEditMode)
+			{
+				GLevelEditorModeTools().ActivateMode(FControlRigEditMode::ModeName);
+				ControlRigEditMode = static_cast<FControlRigEditMode*>(GLevelEditorModeTools().GetActiveMode(FControlRigEditMode::ModeName));
+
+			}
+			if (ControlRigEditMode)
+			{
+				ControlRigEditMode->SetObjects(ControlRig, FGuid(), nullptr);
+			}
+			ControlRig->ControlModified().AddRaw(this, &FControlRigParameterTrackEditor::HandleControlModified);
+			ControlRig->ControlSelected().AddRaw(this, &FControlRigParameterTrackEditor::HandleControlSelected);
+		}
+	}
+}
+
+
+void FControlRigParameterTrackEditor::BindCommands(TSharedRef<FUICommandList> SequencerCommandBindings)
+{
+	const FControlRigTrackCommands& Commands = FControlRigTrackCommands::Get();
+
+	CommandBindings = MakeShared<FUICommandList>();
+	CommandBindings->MapAction(
+		Commands.AddTransformKey,
+		FExecuteAction::CreateSP(this, &FControlRigParameterTrackEditor::OnAddTransformKeysForSelectedObjects, EMovieSceneTransformChannel::All),
+		FCanExecuteAction::CreateSP(this, &FControlRigParameterTrackEditor::CanAddTransformKeysForSelectedObjects));
+
+	CommandBindings->MapAction(
+		Commands.AddTranslationKey,
+		FExecuteAction::CreateSP(this, &FControlRigParameterTrackEditor::OnAddTransformKeysForSelectedObjects, EMovieSceneTransformChannel::Translation),
+		FCanExecuteAction::CreateSP(this, &FControlRigParameterTrackEditor::CanAddTransformKeysForSelectedObjects));
+
+	CommandBindings->MapAction(
+		Commands.AddRotationKey,
+		FExecuteAction::CreateSP(this, &FControlRigParameterTrackEditor::OnAddTransformKeysForSelectedObjects, EMovieSceneTransformChannel::Rotation),
+		FCanExecuteAction::CreateSP(this, &FControlRigParameterTrackEditor::CanAddTransformKeysForSelectedObjects));
+
+	CommandBindings->MapAction(
+		Commands.AddScaleKey,
+		FExecuteAction::CreateSP(this, &FControlRigParameterTrackEditor::OnAddTransformKeysForSelectedObjects, EMovieSceneTransformChannel::Scale),
+		FCanExecuteAction::CreateSP(this, &FControlRigParameterTrackEditor::CanAddTransformKeysForSelectedObjects));
+
+	Commands.BindingCount++;
+
+	// Add these bindings to Sequencer
+	SequencerCommandBindings->Append(CommandBindings.ToSharedRef());
+
+	// Also add them to the Curve Editor 
+	GetSequencer()->GetCommandBindings(ESequencerCommandBindings::CurveEditor)->Append(CommandBindings.ToSharedRef());
+}
+
+bool FControlRigParameterTrackEditor::CanAddTransformKeysForSelectedObjects()
+{
+	if (!GetSequencer()->IsAllowedToChange())
+	{
+		return false;
+	}
+
+	FControlRigEditMode* ControlRigEditMode = static_cast<FControlRigEditMode*>(GLevelEditorModeTools().GetActiveMode(FControlRigEditMode::ModeName));
+	if (ControlRigEditMode && ControlRigEditMode->GetControlRig())
+	{
+		UControlRig* ControlRig = ControlRigEditMode->GetControlRig();
+		FString OurName = ControlRig->GetName();
+		FName Name(*OurName);
+		if (TSharedPtr<IControlRigObjectBinding> ObjectBinding = ControlRig->GetObjectBinding())
+		{
+			TArray<FName> ControlNames = ControlRig->CurrentControlSelection();
+			return (ControlNames.Num() > 0);
+		}
+	}
+	return false;
+}
+
+void FControlRigParameterTrackEditor::OnAddTransformKeysForSelectedObjects(EMovieSceneTransformChannel Channel)
+{
+	if (!GetSequencer()->IsAllowedToChange())
+	{
+		return;
+	}
+
+	FControlRigEditMode* ControlRigEditMode = static_cast<FControlRigEditMode*>(GLevelEditorModeTools().GetActiveMode(FControlRigEditMode::ModeName));
+	if (ControlRigEditMode && ControlRigEditMode->GetControlRig())
+	{
+		UControlRig* ControlRig = ControlRigEditMode->GetControlRig();
+		FString OurName = ControlRig->GetName();
+		FName Name(*OurName);
+		if (TSharedPtr<IControlRigObjectBinding> ObjectBinding = ControlRig->GetObjectBinding())
+		{
+			TArray<FName> ControlNames = ControlRig->CurrentControlSelection();
+			for (const FName& ControlName : ControlNames)
+			{
+				USceneComponent* Component = Cast<USkeletalMeshComponent>(ObjectBinding->GetBoundObject());
+				if (Component)
+				{
+					UObject *ActorObject = Component->GetOwner();
+					AddControlKeys(ActorObject, ControlRig, Name, ControlName, Channel, ESequencerKeyMode::ManualKey);
+				}
+			}
+		}
+	}
+}
+
+void FControlRigParameterTrackEditor::OnSequencerDataChanged(EMovieSceneDataChangeType DataChangeType)
+{
+	UMovieScene* MovieScene = GetSequencer()->GetFocusedMovieSceneSequence()->GetMovieScene();
+	FControlRigEditMode* ControlRigEditMode = static_cast<FControlRigEditMode*>(GLevelEditorModeTools().FindMode(FControlRigEditMode::ModeName));
+	
+	//if we have a valid control rig edit mode need to check and see the control rig in that mode is still in a track
+	//if not we get rid of it.
+	if (ControlRigEditMode && ControlRigEditMode->GetControlRig() != nullptr && MovieScene && (DataChangeType == EMovieSceneDataChangeType::MovieSceneStructureItemRemoved ||
+		DataChangeType == EMovieSceneDataChangeType::Unknown))
+	{
+		const TArray<FMovieSceneBinding>& Bindings = MovieScene->GetBindings();
+		for (const FMovieSceneBinding& Binding : Bindings)
+		{
+			UMovieSceneControlRigParameterTrack* Track = Cast<UMovieSceneControlRigParameterTrack>(MovieScene->FindTrack(UMovieSceneControlRigParameterTrack::StaticClass(), Binding.GetObjectGuid(), NAME_None));
+			if (Track && Track->GetControlRig() == ControlRigEditMode->GetControlRig())
+			{
+				return; //just exit out we still have a good track
+			}
+		}
+		//okay no good track so deactive it and delete it's Control Rig and bingings.
+		if (GLevelEditorModeTools().GetActiveMode(FControlRigEditMode::ModeName))
+		{
+			if (GLevelEditorModeTools().HasToolkitHost())
+			{
+				GLevelEditorModeTools().DeactivateMode(FControlRigEditMode::ModeName);
+			}
+		}
+		ControlRigEditMode->SetObjects(nullptr, FGuid(), nullptr);
+	}
+}
+
+void FControlRigParameterTrackEditor::OnSelectionChanged(TArray<UMovieSceneTrack*> InTracks)
+{
+	if (bIsDoingASelection)
+	{
+		return;
+	}
+
+	TGuardValue<bool> Guard(bIsDoingASelection, true);
+	TArray<FString> StringArray;
+	FControlRigEditMode* ControlRigEditMode = static_cast<FControlRigEditMode*>(GLevelEditorModeTools().GetActiveMode(FControlRigEditMode::ModeName));
+	UControlRig* ControlRig = nullptr;
+	//Always clear the control rig(s) in the edit mode.
+	if (ControlRigEditMode)
+	{
+		ControlRig = ControlRigEditMode->GetControlRig();
+		if (ControlRig)
+		{
+			ControlRig->ClearControlSelection();
+		}
+	}
+	TArray<const IKeyArea*> KeyAreas;
+	GetSequencer()->GetSelectedKeyAreas(KeyAreas);
+	if (KeyAreas.Num() <= 0)
+	{
+		return;
+	}
+
+	TMap<UControlRig *, TSet<FName>> RigsAndControls;
+	for (const IKeyArea* KeyArea : KeyAreas)
+	{
+		UMovieSceneControlRigParameterSection* MovieSection = Cast<UMovieSceneControlRigParameterSection>(KeyArea->GetOwningSection());
+		if (MovieSection)
+		{
+			ControlRig = MovieSection->ControlRig;
+			//Only create the edit mode if we have a KeyAra selected and it's not set and we have some boundobjects.
+			if (!ControlRigEditMode )
+			{
+				GLevelEditorModeTools().ActivateMode(FControlRigEditMode::ModeName);
+				ControlRigEditMode = static_cast<FControlRigEditMode*>(GLevelEditorModeTools().GetActiveMode(FControlRigEditMode::ModeName));
+				if (TSharedPtr<IControlRigObjectBinding> ObjectBinding = ControlRig->GetObjectBinding())
+				{
+					if (ControlRigEditMode)
+					{
+						ControlRigEditMode->SetObjects(ControlRig, FGuid(), nullptr);
+					}
+				}
+			}
+			else
+			{
+				if (ControlRigEditMode->GetControlRig() != ControlRig)
+				{
+					ControlRigEditMode->SetObjects(ControlRig, FGuid(), nullptr);
+				}
+			}
+			
+			const FMovieSceneChannelMetaData* MetaData = KeyArea->GetChannel().GetMetaData();
+			if (MetaData)
+			{
+				StringArray.SetNum(0);
+				FString String = MetaData->Name.ToString();
+				String.ParseIntoArray(StringArray, TEXT("."));
+				if (StringArray.Num() > 0)
+				{
+					FName ControlName(*StringArray[0]);
+					RigsAndControls.FindOrAdd(ControlRig).Add(ControlName);
+				}
+			}
+		}
+	}
+	for (TPair<UControlRig *, TSet<FName>> Pair : RigsAndControls)
+	{
+		Pair.Key->ClearControlSelection();
+		for (const FName Name : Pair.Value)
+		{
+			Pair.Key->SelectControl(Name, true);
+		}
+	}
+}
+
+
+void FControlRigParameterTrackEditor::HandleControlSelected(IControlRigManipulatable* ControlRigManp, const FRigControl& Control, bool bSelected)
+{
+	if (!bIsDoingASelection)
+	{
+		TGuardValue<bool> Guard(bIsDoingASelection, true);
+		FString OurName = ControlRigManp->GetName();
+		FName ControlRigName(*OurName);
+		if (TSharedPtr<IControlRigObjectBinding> ObjectBinding = ControlRigManp->GetObjectBinding())
+		{
+			UObject *ActorObject = nullptr;
+			USceneComponent* Component = Cast<USkeletalMeshComponent>(ObjectBinding->GetBoundObject());
+			if (!Component)
+			{
+				return;
+			}
+			ActorObject = Component->GetOwner();
+			bool bCreateTrack = false;
+			bool bCreateHandle = false;
+			FFindOrCreateHandleResult HandleResult = FindOrCreateHandleToObject(ActorObject, bCreateHandle, NAME_None);
+			FGuid ObjectHandle = HandleResult.Handle;
+			if (!ObjectHandle.IsValid())
+			{
+				return;
+			}
+			FFindOrCreateTrackResult TrackResult = FindOrCreateTrackForObject(ObjectHandle, UMovieSceneControlRigParameterTrack::StaticClass(), ControlRigName, bCreateTrack);
+			UMovieSceneControlRigParameterTrack* Track = CastChecked<UMovieSceneControlRigParameterTrack>(TrackResult.Track, ECastCheckedType::NullAllowed);
+			if (Track)
+			{
+				float Weight = 1.0f;
+				TArray<IKeyArea> KeyAreas;
+				TArray<FString> StringArray;
+
+				for (UMovieSceneSection* Section : Track->GetAllSections())
+				{
+					UMovieSceneControlRigParameterSection* ParamSection = Cast<UMovieSceneControlRigParameterSection>(Section);
+					if (ParamSection)
+					{
+						FMovieSceneChannelProxy& ChannelProxy = Section->GetChannelProxy();
+						for (const FMovieSceneChannelEntry& Entry : Section->GetChannelProxy().GetAllEntries())
+						{
+							const FName ChannelTypeName = Entry.GetChannelTypeName();
+
+							// One editor data ptr per channel
+							TArrayView<FMovieSceneChannel* const>        Channels = Entry.GetChannels();
+							TArrayView<const FMovieSceneChannelMetaData> AllMetaData = Entry.GetMetaData();
+
+							for (int32 Index = 0; Index < Channels.Num(); ++Index)
+							{
+								FMovieSceneChannelHandle ChannelHandle = ChannelProxy.MakeHandle(ChannelTypeName, Index);
+								const FMovieSceneChannelMetaData& MetaData = AllMetaData[Index];
+								FString String = MetaData.Name.ToString();
+								String.ParseIntoArray(StringArray, TEXT("."));
+								if (StringArray.Num() > 0)
+								{
+									FName ControlName(*StringArray[0]);
+									if (ControlName == Control.Name)
+									{
+										IKeyArea KeyArea(Section, ChannelHandle);
+										KeyAreas.Add(KeyArea);
+									}
+								}
+							}
+						}
+					}
+				}
+				if (KeyAreas.Num())
+				{
+					GetSequencer()->SelectByKeyAreas(KeyAreas, true, bSelected);
+				}
+			}
+		}
+	}
+}
+
+void FControlRigParameterTrackEditor::HandleControlModified(IControlRigManipulatable* ControlRigManp, const FRigControl& Control)
+{
+	if (GetSequencer().IsValid() && !GetSequencer()->IsAllowedToChange())
+	{
+		return;
+	}
+	FString OurName = ControlRigManp->GetName();
+	FName Name(*OurName);
+	if (TSharedPtr<IControlRigObjectBinding> ObjectBinding = ControlRigManp->GetObjectBinding())
+	{
+		UObject *ActorObject = nullptr;
+		USceneComponent* Component = Cast<USkeletalMeshComponent>(ObjectBinding->GetBoundObject());
+		if (Component)
+		{
+			ActorObject = Component->GetOwner();
+			AddControlKeys(ActorObject, ControlRigManp, Name, Control.Name, EMovieSceneTransformChannel::All, ESequencerKeyMode::AutoKey);
+		}
+	}
+}
+
+void FControlRigParameterTrackEditor::GetControlRigKeys(IControlRigManipulatable* Manip, FName ParameterName, EMovieSceneTransformChannel ChannelsToKey, FGeneratedTrackKeys& OutGeneratedKeys)
+{
+	const TArray<FRigControl>& Controls = Manip->AvailableControls();
+	// If key all is enabled, for a key on all the channels
+	if (GetSequencer()->GetKeyGroupMode() == EKeyGroupMode::KeyAll)
+	{
+		ChannelsToKey = EMovieSceneTransformChannel::All;
+	}
+
+	int32 ChannelIndex = 0;
+	for (const FRigControl& RigControl : Controls)
+	{
+		switch (RigControl.ControlType)
+		{
+		case ERigControlType::Float:
+		{
+			float Val = RigControl.Value.Get<float>();
+			OutGeneratedKeys.Add(FMovieSceneChannelValueSetter::Create<FMovieSceneFloatChannel>(ChannelIndex++, Val, true));
+			break;
+		}
+		case ERigControlType::Position:
+		case ERigControlType::Scale:
+		case ERigControlType::Rotator:
+		{
+			FVector Val = RigControl.Value.Get<FVector>();
+			OutGeneratedKeys.Add(FMovieSceneChannelValueSetter::Create<FMovieSceneFloatChannel>(ChannelIndex++, Val.X, true));
+			OutGeneratedKeys.Add(FMovieSceneChannelValueSetter::Create<FMovieSceneFloatChannel>(ChannelIndex++, Val.Y, true));
+			OutGeneratedKeys.Add(FMovieSceneChannelValueSetter::Create<FMovieSceneFloatChannel>(ChannelIndex++, Val.Z, true));
+			break;
+		}
+
+		case ERigControlType::Transform:
+		{
+			FTransform Val = RigControl.Value.Get<FTransform>();
+			FVector CurrentVector = Val.GetTranslation();
+			bool bKeyX = EnumHasAnyFlags(ChannelsToKey, EMovieSceneTransformChannel::TranslationX);
+			bool bKeyY = EnumHasAnyFlags(ChannelsToKey, EMovieSceneTransformChannel::TranslationY);
+			bool bKeyZ = EnumHasAnyFlags(ChannelsToKey, EMovieSceneTransformChannel::TranslationZ);
+			if (GetSequencer()->GetKeyGroupMode() == EKeyGroupMode::KeyGroup && (bKeyX || bKeyY || bKeyZ))
+			{
+				bKeyX = bKeyY = bKeyZ = true;
+			}
+			OutGeneratedKeys.Add(FMovieSceneChannelValueSetter::Create<FMovieSceneFloatChannel>(ChannelIndex++, CurrentVector.X, bKeyX));
+			OutGeneratedKeys.Add(FMovieSceneChannelValueSetter::Create<FMovieSceneFloatChannel>(ChannelIndex++, CurrentVector.Y, bKeyY));
+			OutGeneratedKeys.Add(FMovieSceneChannelValueSetter::Create<FMovieSceneFloatChannel>(ChannelIndex++, CurrentVector.Z, bKeyZ));
+
+			CurrentVector = Val.GetRotation().Euler();
+			bKeyX = EnumHasAnyFlags(ChannelsToKey, EMovieSceneTransformChannel::RotationX);
+			bKeyY = EnumHasAnyFlags(ChannelsToKey, EMovieSceneTransformChannel::RotationY);
+			bKeyZ = EnumHasAnyFlags(ChannelsToKey, EMovieSceneTransformChannel::RotationZ);
+			if (GetSequencer()->GetKeyGroupMode() == EKeyGroupMode::KeyGroup && (bKeyX || bKeyY || bKeyZ))
+			{
+				bKeyX = bKeyY = bKeyZ = true;
+			}
+			OutGeneratedKeys.Add(FMovieSceneChannelValueSetter::Create<FMovieSceneFloatChannel>(ChannelIndex++, CurrentVector.X, bKeyX));
+			OutGeneratedKeys.Add(FMovieSceneChannelValueSetter::Create<FMovieSceneFloatChannel>(ChannelIndex++, CurrentVector.Y, bKeyY));
+			OutGeneratedKeys.Add(FMovieSceneChannelValueSetter::Create<FMovieSceneFloatChannel>(ChannelIndex++, CurrentVector.Z, bKeyZ));
+
+			CurrentVector = Val.GetScale3D();
+			bKeyX = EnumHasAnyFlags(ChannelsToKey, EMovieSceneTransformChannel::ScaleX);
+			bKeyY = EnumHasAnyFlags(ChannelsToKey, EMovieSceneTransformChannel::ScaleY);
+			bKeyZ = EnumHasAnyFlags(ChannelsToKey, EMovieSceneTransformChannel::ScaleZ);
+			if (GetSequencer()->GetKeyGroupMode() == EKeyGroupMode::KeyGroup && (bKeyX || bKeyY || bKeyZ))
+			{
+				bKeyX = bKeyY = bKeyZ = true;
+			}
+			OutGeneratedKeys.Add(FMovieSceneChannelValueSetter::Create<FMovieSceneFloatChannel>(ChannelIndex++, CurrentVector.X, bKeyX));
+			OutGeneratedKeys.Add(FMovieSceneChannelValueSetter::Create<FMovieSceneFloatChannel>(ChannelIndex++, CurrentVector.Y, bKeyY));
+			OutGeneratedKeys.Add(FMovieSceneChannelValueSetter::Create<FMovieSceneFloatChannel>(ChannelIndex++, CurrentVector.Z, bKeyZ));
+
+			break;
+		}
+		//mz todo the other types
+		}
+	}
+}
+
+FKeyPropertyResult FControlRigParameterTrackEditor::AddKeysToControlRigHandle(UObject *Object, IControlRigManipulatable* Manip,
+	FGuid ObjectHandle, FFrameNumber KeyTime, FGeneratedTrackKeys& GeneratedKeys,
+	ESequencerKeyMode KeyMode, TSubclassOf<UMovieSceneTrack> TrackClass, FName ControlRigName)
+{
+	EAutoChangeMode AutoChangeMode = GetSequencer()->GetAutoChangeMode();
+	EAllowEditsMode AllowEditsMode = GetSequencer()->GetAllowEditsMode();
+
+	bool bCreateTrack =
+		(KeyMode == ESequencerKeyMode::AutoKey && (AutoChangeMode == EAutoChangeMode::AutoTrack || AutoChangeMode == EAutoChangeMode::All)) ||
+		KeyMode == ESequencerKeyMode::ManualKey ||
+		KeyMode == ESequencerKeyMode::ManualKeyForced ||
+		AllowEditsMode == EAllowEditsMode::AllowSequencerEditsOnly;
+
+	bool bCreateSection = false;
+		// we don't do this, maybe revisit if a bug occurs, but currently extends sections on autokey.
+	    //bCreateTrack || (KeyMode == ESequencerKeyMode::AutoKey && (AutoChangeMode != EAutoChangeMode::None));
+
+	// Try to find an existing Track, and if one doesn't exist check the key params and create one if requested.
+
+	FFindOrCreateTrackResult TrackResult = FindOrCreateTrackForObject(ObjectHandle, TrackClass, ControlRigName, bCreateTrack);
+	UMovieSceneControlRigParameterTrack* Track = CastChecked<UMovieSceneControlRigParameterTrack>(TrackResult.Track, ECastCheckedType::NullAllowed);
+
+	bool bTrackCreated = TrackResult.bWasCreated;
+
+	bool bSectionCreated = false;
+	FKeyPropertyResult KeyPropertyResult;
+
+	if (Track)
+	{
+		float Weight = 1.0f;
+
+		UMovieSceneSection* SectionToKey = bCreateSection ? Track->FindOrExtendSection(KeyTime, Weight) : Track->FindSection(KeyTime);
+
+		// If there's no overlapping section to key, create one only if a track was newly created. Otherwise, skip keying altogether
+		// so that the user is forced to create a section to key on.
+		if (bTrackCreated && !SectionToKey)
+		{
+			Track->Modify();
+			SectionToKey = Track->FindOrAddSection(KeyTime, bSectionCreated);
+			if (bSectionCreated && GetSequencer()->GetInfiniteKeyAreas())
+			{
+				SectionToKey->SetRange(TRange<FFrameNumber>::All());
+			}
+		}
+
+		if (SectionToKey && SectionToKey->GetRange().Contains(KeyTime))
+		{
+			if (!bTrackCreated)
+			{
+				ModifyOurGeneratedKeysByCurrentAndWeight(Object, Manip, Track, SectionToKey, KeyTime, GeneratedKeys, Weight);
+			}
+			UMovieSceneControlRigParameterSection* ParamSection = Cast<UMovieSceneControlRigParameterSection>(SectionToKey);
+			if (!ParamSection->GetDoNotKey())
+			{
+				KeyPropertyResult |= AddKeysToSection(SectionToKey, KeyTime, GeneratedKeys, KeyMode);
+			}
+		}
+	}
+
+	KeyPropertyResult.bTrackCreated |= bTrackCreated || bSectionCreated;
+	return KeyPropertyResult;
+}
+
+FKeyPropertyResult FControlRigParameterTrackEditor::AddKeysToControlRig(
+	UObject* Object, IControlRigManipulatable* Manip, FFrameNumber KeyTime, FGeneratedTrackKeys& GeneratedKeys,
+	ESequencerKeyMode KeyMode, TSubclassOf<UMovieSceneTrack> TrackClass, FName PropertyName)
+{
+	FKeyPropertyResult KeyPropertyResult;
+	EAutoChangeMode AutoChangeMode = GetSequencer()->GetAutoChangeMode();
+	EAllowEditsMode AllowEditsMode = GetSequencer()->GetAllowEditsMode();
+	bool bCreateHandle =
+		(KeyMode == ESequencerKeyMode::AutoKey && (AutoChangeMode == EAutoChangeMode::All)) ||
+		KeyMode == ESequencerKeyMode::ManualKey ||
+		KeyMode == ESequencerKeyMode::ManualKeyForced ||
+		AllowEditsMode == EAllowEditsMode::AllowSequencerEditsOnly;
+
+	FFindOrCreateHandleResult HandleResult = FindOrCreateHandleToObject(Object, bCreateHandle, NAME_None);
+	FGuid ObjectHandle = HandleResult.Handle;
+	KeyPropertyResult.bHandleCreated = HandleResult.bWasCreated;
+	if (ObjectHandle.IsValid())
+	{
+		KeyPropertyResult |= AddKeysToControlRigHandle(Object, Manip, ObjectHandle, KeyTime, GeneratedKeys, KeyMode, TrackClass, PropertyName);
+	}
+
+	return KeyPropertyResult;
+}
+
+void FControlRigParameterTrackEditor::AddControlKeys(UObject* InObject, IControlRigManipulatable* Manip, FName ControlRigName, FName RigControlName, EMovieSceneTransformChannel ChannelsToKey, ESequencerKeyMode KeyMode)
+{
+	if (GetSequencer().IsValid() && !GetSequencer()->IsAllowedToChange())
+	{
+		return;
+	}
+	bool bCreateTrack = false;
+	bool bCreateHandle = false;
+	FFindOrCreateHandleResult HandleResult = FindOrCreateHandleToObject(InObject, bCreateHandle, NAME_None);
+	FGuid ObjectHandle = HandleResult.Handle;
+	if (!ObjectHandle.IsValid())
+	{
+		return;
+	}
+	FFindOrCreateTrackResult TrackResult = FindOrCreateTrackForObject(ObjectHandle, UMovieSceneControlRigParameterTrack::StaticClass(), ControlRigName, bCreateTrack);
+	UMovieSceneControlRigParameterTrack* Track = CastChecked<UMovieSceneControlRigParameterTrack>(TrackResult.Track, ECastCheckedType::NullAllowed);
+	if (Track)
+	{
+		FFrameNumber  FrameTime = GetTimeForKey();
+		UMovieSceneSection* Section = Track->FindSection(FrameTime);
+		UMovieSceneControlRigParameterSection* ParamSection = Cast<UMovieSceneControlRigParameterSection>(Section);
+
+		if (ParamSection && ParamSection->GetDoNotKey())
+		{
+			return;
+		}
+	}
+
+	TSharedRef<FGeneratedTrackKeys> GeneratedKeys = MakeShared<FGeneratedTrackKeys>();
+
+	GetControlRigKeys(Manip, RigControlName, ChannelsToKey, *GeneratedKeys);
+
+	auto OnKeyProperty = [=](FFrameNumber Time) -> FKeyPropertyResult
+	{
+		return this->AddKeysToControlRig(InObject, Manip, Time, *GeneratedKeys, KeyMode, UMovieSceneControlRigParameterTrack::StaticClass(), ControlRigName);
+	};
+
+	AnimatablePropertyChanged(FOnKeyProperty::CreateLambda(OnKeyProperty));
+}
+
+bool FControlRigParameterTrackEditor::ModifyOurGeneratedKeysByCurrentAndWeight(UObject *Object, IControlRigManipulatable* Manip, UMovieSceneTrack *Track, UMovieSceneSection* SectionToKey, FFrameNumber KeyTime, FGeneratedTrackKeys& GeneratedTotalKeys, float Weight) const
+{
+	FFrameRate TickResolution = GetSequencer()->GetFocusedTickResolution();
+	FMovieSceneEvaluationTrack EvalTrack = Track->GenerateTrackTemplate();
+
+	FMovieSceneInterrogationData InterrogationData;
+	GetSequencer()->GetEvaluationTemplate().CopyActuators(InterrogationData.GetAccumulator());
+
+	FMovieSceneContext Context(FMovieSceneEvaluationRange(KeyTime, GetSequencer()->GetFocusedTickResolution()));
+	EvalTrack.Interrogate(Context, InterrogationData, Object);
+	const TArray<FRigControl>& Controls = Manip->AvailableControls();
+
+	FMovieSceneChannelProxy& Proxy = SectionToKey->GetChannelProxy();
+	int32 ChannelIndex = 0;
+	for (const FRigControl& RigControl : Controls)
+	{
+		switch (RigControl.ControlType)
+		{
+		case ERigControlType::Float:
+		{
+			for (const FFloatInterrogationData& Val : InterrogationData.Iterate<FFloatInterrogationData>(UMovieSceneControlRigParameterSection::GetFloatInterrogationKey()))
+			{
+				if (Val.ParameterName == RigControl.Name)
+				{
+					GeneratedTotalKeys[ChannelIndex]->ModifyByCurrentAndWeight(Proxy, KeyTime, (void *)&Val.Val, Weight);
+					break;
+				}
+			}
+			++ChannelIndex;
+			break;
+		}
+		case ERigControlType::Position:
+		case ERigControlType::Scale:
+		case ERigControlType::Rotator:
+		{
+			for (const FVectorInterrogationData& Val : InterrogationData.Iterate<FVectorInterrogationData>(UMovieSceneControlRigParameterSection::GetVectorInterrogationKey()))
+			{
+				if (Val.ParameterName == RigControl.Name)
+				{
+					GeneratedTotalKeys[ChannelIndex]->ModifyByCurrentAndWeight(Proxy, KeyTime, (void *)&Val.Val, Weight);
+					GeneratedTotalKeys[ChannelIndex + 1]->ModifyByCurrentAndWeight(Proxy, KeyTime, (void *)&Val.Val, Weight);
+					GeneratedTotalKeys[ChannelIndex + 2]->ModifyByCurrentAndWeight(Proxy, KeyTime, (void *)&Val.Val, Weight);
+
+					break;
+				}
+			}
+			ChannelIndex += 3;
+
+			break;
+		}
+
+		case ERigControlType::Transform:
+		{
+			for (const FTransformInterrogationData& Val : InterrogationData.Iterate<FTransformInterrogationData>(UMovieSceneControlRigParameterSection::GetTransformInterrogationKey()))
+			{
+				if (Val.ParameterName == RigControl.Name)
+				{
+					FVector CurrentPos = Val.Val.GetTranslation();
+					FRotator CurrentRot = Val.Val.GetRotation().Rotator();
+					FVector CurrentScale = Val.Val.GetScale3D();
+					GeneratedTotalKeys[ChannelIndex]->ModifyByCurrentAndWeight(Proxy, KeyTime, (void *)&CurrentPos.X, Weight);
+					GeneratedTotalKeys[ChannelIndex + 1]->ModifyByCurrentAndWeight(Proxy, KeyTime, (void *)&CurrentPos.Y, Weight);
+					GeneratedTotalKeys[ChannelIndex + 2]->ModifyByCurrentAndWeight(Proxy, KeyTime, (void *)&CurrentPos.Z, Weight);
+					GeneratedTotalKeys[ChannelIndex + 3]->ModifyByCurrentAndWeight(Proxy, KeyTime, (void *)&CurrentRot.Roll, Weight);
+					GeneratedTotalKeys[ChannelIndex + 4]->ModifyByCurrentAndWeight(Proxy, KeyTime, (void *)&CurrentRot.Pitch, Weight);
+					GeneratedTotalKeys[ChannelIndex + 5]->ModifyByCurrentAndWeight(Proxy, KeyTime, (void *)&CurrentRot.Yaw, Weight);
+					GeneratedTotalKeys[ChannelIndex + 6]->ModifyByCurrentAndWeight(Proxy, KeyTime, (void *)&CurrentScale.X, Weight);
+					GeneratedTotalKeys[ChannelIndex + 7]->ModifyByCurrentAndWeight(Proxy, KeyTime, (void *)&CurrentScale.Y, Weight);
+					GeneratedTotalKeys[ChannelIndex + 8]->ModifyByCurrentAndWeight(Proxy, KeyTime, (void *)&CurrentScale.Z, Weight);
+
+					break;
+				}
+			}
+			ChannelIndex += 9;
+
+			break;
+		}
+		//mz todo the other types
+		}
+	}
+	return true;
+}
+
+void FControlRigParameterSection::BuildSectionContextMenu(FMenuBuilder& MenuBuilder, const FGuid& InObjectBinding)
+{
+	UMovieSceneControlRigParameterSection* ParameterSection = CastChecked<UMovieSceneControlRigParameterSection>(WeakSection.Get());
+	TSharedPtr<ISequencer> SequencerPtr = WeakSequencer.Pin();
+
+	UControlRig* ControlRig = ParameterSection->ControlRig;
+	if (ControlRig)
+	{
+		const TArray<FRigControl>& Controls = ControlRig->AvailableControls();
+
+		auto MakeUIAction = [=](EMovieSceneTransformChannel ChannelsToToggle)
+		{
+			return FUIAction(
+				FExecuteAction::CreateLambda([=]
+			{
+				FScopedTransaction Transaction(LOCTEXT("SetActiveChannelsTransaction", "Set Active Channels"));
+				ParameterSection->Modify();
+				EMovieSceneTransformChannel Channels = ParameterSection->GetTransformMask().GetChannels();
+
+				if (EnumHasAllFlags(Channels, ChannelsToToggle) || (Channels & ChannelsToToggle) == EMovieSceneTransformChannel::None)
+				{
+					ParameterSection->SetTransformMask(ParameterSection->GetTransformMask().GetChannels() ^ ChannelsToToggle);
+				}
+				else
+				{
+					ParameterSection->SetTransformMask(ParameterSection->GetTransformMask().GetChannels() | ChannelsToToggle);
+				}
+
+				// Restore pre-animated state for the bound objects so that inactive channels will return to their default values.
+				for (TWeakObjectPtr<> WeakObject : SequencerPtr->FindBoundObjects(InObjectBinding, SequencerPtr->GetFocusedTemplateID()))
+				{
+					if (UObject* Object = WeakObject.Get())
+					{
+						SequencerPtr->RestorePreAnimatedState();
+					}
+				}
+
+				SequencerPtr->NotifyMovieSceneDataChanged(EMovieSceneDataChangeType::MovieSceneStructureItemsChanged);
+			}
+				),
+				FCanExecuteAction(),
+				FGetActionCheckState::CreateLambda([=]
+			{
+				EMovieSceneTransformChannel Channels = ParameterSection->GetTransformMask().GetChannels();
+				if (EnumHasAllFlags(Channels, ChannelsToToggle))
+				{
+					return ECheckBoxState::Checked;
+				}
+				else if (EnumHasAnyFlags(Channels, ChannelsToToggle))
+				{
+					return ECheckBoxState::Undetermined;
+				}
+				return ECheckBoxState::Unchecked;
+			})
+				);
+		};
+		auto ToggleControls = [=](int32 Index)
+		{
+			return FUIAction(
+				FExecuteAction::CreateLambda([=]
+			{
+				FScopedTransaction Transaction(LOCTEXT("RigSectionFilter", "Toggle Rig Control Filters"));
+				ParameterSection->Modify();
+				if (Index >= 0)
+				{
+					ParameterSection->SetControlsMask(Index, !ParameterSection->GetControlsMask(Index));
+				}
+				else
+				{
+					ParameterSection->FillControlsMask(!ParameterSection->GetControlsMask(0));
+				}
+				SequencerPtr->NotifyMovieSceneDataChanged(EMovieSceneDataChangeType::MovieSceneStructureItemsChanged);
+
+			}
+				),
+				FCanExecuteAction(),
+				FGetActionCheckState::CreateLambda([=]
+			{
+				TArray<bool> ControlBool = ParameterSection->GetControlsMask();
+				if (Index >= 0)
+				{
+					if (ControlBool[Index])
+					{
+						return ECheckBoxState::Checked;
+
+					}
+					else
+					{
+						return ECheckBoxState::Unchecked;
+					}
+				}
+				else
+				{
+					TOptional<bool> FirstVal;
+					for (bool Val : ControlBool)
+					{
+						if (FirstVal.IsSet())
+						{
+							if (Val != FirstVal)
+							{
+								return ECheckBoxState::Undetermined;
+							}
+						}
+						else
+						{
+							FirstVal = Val;
+						}
+
+					}
+					return (FirstVal.GetValue()) ? ECheckBoxState::Checked : ECheckBoxState::Unchecked;
+				}
+			})
+				);
+		};
+		MenuBuilder.BeginSection(NAME_None, LOCTEXT("RigSectionFilter", "Filter Controls"));
+		{
+			MenuBuilder.AddSubMenu(
+				LOCTEXT("ToggleRigControlsText", "Toggle Rig Controls"), LOCTEXT("ToggleRigControlsText_Tooltip", "Toggle Rig Controls"),
+				FNewMenuDelegate::CreateLambda([=](FMenuBuilder& SubMenuBuilder) {
+				int32 Index = 0;
+				for (const FRigControl& RigControl : Controls)
+				{
+					const FName RigName = RigControl.Name;
+					FText Name = FText::FromName(RigName);
+					FText Text = FText::Format(LOCTEXT("RigControlToggle", "{0}"), Name);
+					FText TooltipText = FText::Format(LOCTEXT("RigControlToggleTooltip", "Toggle Rig Control {0}"), Name);
+					SubMenuBuilder.AddMenuEntry(
+						Text, TooltipText,
+						FSlateIcon(), ToggleControls(Index++), NAME_None, EUserInterfaceActionType::ToggleButton);
+				}
+			}),
+				ToggleControls(-1),
+				NAME_None,
+				EUserInterfaceActionType::ToggleButton);
+
+			MenuBuilder.AddSubMenu(
+				LOCTEXT("AllTranslation", "Translation"), LOCTEXT("AllTranslation_ToolTip", "Causes this section to affect the translation of rig control transforms"),
+				FNewMenuDelegate::CreateLambda([=](FMenuBuilder& SubMenuBuilder) {
+				SubMenuBuilder.AddMenuEntry(
+					LOCTEXT("TranslationX", "X"), LOCTEXT("TranslationX_ToolTip", "Causes this section to affect the X channel of the transform's translation"),
+					FSlateIcon(), MakeUIAction(EMovieSceneTransformChannel::TranslationX), NAME_None, EUserInterfaceActionType::ToggleButton);
+				SubMenuBuilder.AddMenuEntry(
+					LOCTEXT("TranslationY", "Y"), LOCTEXT("TranslationY_ToolTip", "Causes this section to affect the Y channel of the transform's translation"),
+					FSlateIcon(), MakeUIAction(EMovieSceneTransformChannel::TranslationY), NAME_None, EUserInterfaceActionType::ToggleButton);
+				SubMenuBuilder.AddMenuEntry(
+					LOCTEXT("TranslationZ", "Z"), LOCTEXT("TranslationZ_ToolTip", "Causes this section to affect the Z channel of the transform's translation"),
+					FSlateIcon(), MakeUIAction(EMovieSceneTransformChannel::TranslationZ), NAME_None, EUserInterfaceActionType::ToggleButton);
+			}),
+				MakeUIAction(EMovieSceneTransformChannel::Translation),
+				NAME_None,
+				EUserInterfaceActionType::ToggleButton);
+
+			MenuBuilder.AddSubMenu(
+				LOCTEXT("AllRotation", "Rotation"), LOCTEXT("AllRotation_ToolTip", "Causes this section to affect the rotation of the rig control transform"),
+				FNewMenuDelegate::CreateLambda([=](FMenuBuilder& SubMenuBuilder) {
+				SubMenuBuilder.AddMenuEntry(
+					LOCTEXT("RotationX", "Roll (X)"), LOCTEXT("RotationX_ToolTip", "Causes this section to affect the roll (X) channel the transform's rotation"),
+					FSlateIcon(), MakeUIAction(EMovieSceneTransformChannel::RotationX), NAME_None, EUserInterfaceActionType::ToggleButton);
+				SubMenuBuilder.AddMenuEntry(
+					LOCTEXT("RotationY", "Pitch (Y)"), LOCTEXT("RotationY_ToolTip", "Causes this section to affect the pitch (Y) channel the transform's rotation"),
+					FSlateIcon(), MakeUIAction(EMovieSceneTransformChannel::RotationY), NAME_None, EUserInterfaceActionType::ToggleButton);
+				SubMenuBuilder.AddMenuEntry(
+					LOCTEXT("RotationZ", "Yaw (Z)"), LOCTEXT("RotationZ_ToolTip", "Causes this section to affect the yaw (Z) channel the transform's rotation"),
+					FSlateIcon(), MakeUIAction(EMovieSceneTransformChannel::RotationZ), NAME_None, EUserInterfaceActionType::ToggleButton);
+			}),
+				MakeUIAction(EMovieSceneTransformChannel::Rotation),
+				NAME_None,
+				EUserInterfaceActionType::ToggleButton);
+
+			MenuBuilder.AddSubMenu(
+				LOCTEXT("AllScale", "Scale"), LOCTEXT("AllScale_ToolTip", "Causes this section to affect the scale of the rig control transform"),
+				FNewMenuDelegate::CreateLambda([=](FMenuBuilder& SubMenuBuilder) {
+				SubMenuBuilder.AddMenuEntry(
+					LOCTEXT("ScaleX", "X"), LOCTEXT("ScaleX_ToolTip", "Causes this section to affect the X channel of the transform's scale"),
+					FSlateIcon(), MakeUIAction(EMovieSceneTransformChannel::ScaleX), NAME_None, EUserInterfaceActionType::ToggleButton);
+				SubMenuBuilder.AddMenuEntry(
+					LOCTEXT("ScaleY", "Y"), LOCTEXT("ScaleY_ToolTip", "Causes this section to affect the Y channel of the transform's scale"),
+					FSlateIcon(), MakeUIAction(EMovieSceneTransformChannel::ScaleY), NAME_None, EUserInterfaceActionType::ToggleButton);
+				SubMenuBuilder.AddMenuEntry(
+					LOCTEXT("ScaleZ", "Z"), LOCTEXT("ScaleZ_ToolTip", "Causes this section to affect the Z channel of the transform's scale"),
+					FSlateIcon(), MakeUIAction(EMovieSceneTransformChannel::ScaleZ), NAME_None, EUserInterfaceActionType::ToggleButton);
+			}),
+				MakeUIAction(EMovieSceneTransformChannel::Scale),
+				NAME_None,
+				EUserInterfaceActionType::ToggleButton);
+
+			//mz todo h
+			MenuBuilder.AddMenuEntry(
+				LOCTEXT("Weight", "Weight"), LOCTEXT("Weight_ToolTip", "Causes this section to be applied with a user-specified weight curve"),
+				FSlateIcon(), MakeUIAction(EMovieSceneTransformChannel::Weight), NAME_None, EUserInterfaceActionType::ToggleButton);
+		}
+		MenuBuilder.EndSection();
+	}
+}
+
+//mz todo
+bool FControlRigParameterSection::RequestDeleteCategory(const TArray<FName>& CategoryNamePaths)
+{
+	UMovieSceneControlRigParameterSection* ParameterSection = CastChecked<UMovieSceneControlRigParameterSection>(WeakSection.Get());
+	TSharedPtr<ISequencer> SequencerPtr = WeakSequencer.Pin();
+	/*
+	const FScopedTransaction Transaction(LOCTEXT("DeleteTransformCategory", "Delete transform category"));
+
+	if (ParameterSection->TryModify())
+	{
+		FName CategoryName = CategoryNamePaths[CategoryNamePaths.Num() - 1];
+
+		EMovieSceneTransformChannel Channel = ParameterSection->GetTransformMask().GetChannels();
+		EMovieSceneTransformChannel ChannelToRemove = ParameterSection->GetTransformMaskByName(CategoryName).GetChannels();
+
+		Channel = Channel ^ ChannelToRemove;
+
+		ParameterSection->SetTransformMask(Channel);
+
+		SequencerPtr->NotifyMovieSceneDataChanged(EMovieSceneDataChangeType::MovieSceneStructureItemsChanged);
+		return true;
+	}
+	*/
+	return false;
+}
+
+bool FControlRigParameterSection::RequestDeleteKeyArea(const TArray<FName>& KeyAreaNamePaths)
+{
+	UMovieSceneControlRigParameterSection* ParameterSection = CastChecked<UMovieSceneControlRigParameterSection>(WeakSection.Get());
+	TSharedPtr<ISequencer> SequencerPtr = WeakSequencer.Pin();
+
+	/*
+	const FScopedTransaction Transaction(LOCTEXT("DeleteTransformChannel", "Delete transform channel"));
+
+	if (ParameterSection->TryModify())
+	{
+		// Only delete the last key area path which is the channel. ie. TranslationX as opposed to Translation
+		FName KeyAreaName = KeyAreaNamePaths[KeyAreaNamePaths.Num() - 1];
+
+		EMovieSceneTransformChannel Channel = ParameterSection->GetTransformMask().GetChannels();
+		EMovieSceneTransformChannel ChannelToRemove = ParameterSection->GetTransformMaskByName(KeyAreaName).GetChannels();
+
+		Channel = Channel ^ ChannelToRemove;
+
+		ParameterSection->SetTransformMask(Channel);
+
+		SequencerPtr->NotifyMovieSceneDataChanged(EMovieSceneDataChangeType::MovieSceneStructureItemsChanged);
+		return true;
+	}
+	*/
+	return true;
+}
+
+
+#undef LOCTEXT_NAMESPACE
diff --git a/Engine/Plugins/Experimental/ControlRig/Source/ControlRigEditor/Private/Sequencer/ControlRigParameterTrackEditor.h b/Engine/Plugins/Experimental/ControlRig/Source/ControlRigEditor/Private/Sequencer/ControlRigParameterTrackEditor.h
new file mode 100644
--- /dev/null
+++ b/Engine/Plugins/Experimental/ControlRig/Source/ControlRigEditor/Private/Sequencer/ControlRigParameterTrackEditor.h
@@ -0,0 +1,148 @@
+// Copyright 1998-2019 Epic Games, Inc. All Rights Reserved.
+
+#pragma once
+
+#include "CoreMinimal.h"
+#include "Misc/Guid.h"
+#include "Templates/SubclassOf.h"
+#include "Widgets/SWidget.h"
+#include "ISequencer.h"
+#include "MovieSceneTrack.h"
+#include "ISequencerSection.h"
+#include "ISequencerTrackEditor.h"
+#include "Sequencer/MovieSceneControlRigParameterTrack.h"
+#include "KeyframeTrackEditor.h"
+#include "Sections/MovieScene3DTransformSection.h"
+
+struct FAssetData;
+class FMenuBuilder;
+class USkeleton;
+
+/**
+ * Tools for animation tracks
+ */
+class FControlRigParameterTrackEditor : public FKeyframeTrackEditor<UMovieSceneControlRigParameterTrack>
+{
+public:
+	/**
+	 * Constructor
+	 *
+	 * @param InSequencer The sequencer instance to be used by this tool
+	 */
+	FControlRigParameterTrackEditor(TSharedRef<ISequencer> InSequencer);
+
+	/** Virtual destructor. */
+	virtual ~FControlRigParameterTrackEditor();
+
+	/**
+	 * Creates an instance of this class.  Called by a sequencer
+	 *
+	 * @param OwningSequencer The sequencer instance to be used by this tool
+	 * @return The new instance of this class
+	 */
+	static TSharedRef<ISequencerTrackEditor> CreateTrackEditor(TSharedRef<ISequencer> OwningSequencer);
+
+public:
+
+	// ISequencerTrackEditor interface
+	virtual void OnRelease() override;
+	virtual void BindCommands(TSharedRef<FUICommandList> SequencerCommandBindings) override;
+	virtual void BuildObjectBindingTrackMenu(FMenuBuilder& MenuBuilder, const TArray<FGuid>& ObjectBindings, const UClass* ObjectClass) override;
+	virtual TSharedRef<ISequencerSection> MakeSectionInterface(UMovieSceneSection& SectionObject, UMovieSceneTrack& Track, FGuid ObjectBinding) override;
+	virtual bool SupportsType(TSubclassOf<UMovieSceneTrack> Type) const override;
+
+private:
+
+	void AddControlRigSubMenu(FMenuBuilder& MenuBuilder, TArray<FGuid> ObjectBindings, UMovieSceneTrack* Track);
+
+	/** Control Rig Picked */
+	void AddControlRig(UClass* InClass, UObject* BoundObject, FGuid ObjectBinding, USkeleton* Skeleton);
+
+	/** Delegate for Selection Changed Event */
+	void OnSelectionChanged(TArray<UMovieSceneTrack*> InTracks);
+
+	/** Delegate for MovieScene Changing so we can see if our track got deleted*/
+	void OnSequencerDataChanged(EMovieSceneDataChangeType DataChangeType);
+
+	/** Control Rig Delegates*/
+	void HandleControlModified(IControlRigManipulatable* Subject, const FRigControl& Control);
+	void HandleControlSelected(IControlRigManipulatable* Subject, const FRigControl& Control, bool bSelected);
+
+private:
+	//MZ TODO Store map of control rig to parameter values to see if they changed?
+
+	/** Command Bindings added by the Transform Track Editor to Sequencer and curve editor. */
+	TSharedPtr<FUICommandList> CommandBindings;
+public:
+
+	void AddControlKeys(UObject* InObject, IControlRigManipulatable* Manip, FName PropertyName, FName ParameterName, EMovieSceneTransformChannel ChannelsToKey, ESequencerKeyMode KeyMode);
+	void GetControlRigKeys(IControlRigManipulatable* Manip, FName ParameterName, EMovieSceneTransformChannel ChannelsToKey, FGeneratedTrackKeys& OutGeneratedKeys);
+	FKeyPropertyResult AddKeysToControlRig(
+		UObject* Object, IControlRigManipulatable* Manip, FFrameNumber KeyTime, FGeneratedTrackKeys& GeneratedKeys,
+		ESequencerKeyMode KeyMode, TSubclassOf<UMovieSceneTrack> TrackClass, FName PropertyName);
+	FKeyPropertyResult AddKeysToControlRigHandle(UObject *Object, IControlRigManipulatable* Manip,
+		FGuid ObjectHandle, FFrameNumber KeyTime, FGeneratedTrackKeys& GeneratedKeys,
+		ESequencerKeyMode KeyMode, TSubclassOf<UMovieSceneTrack> TrackClass, FName PropertyName);
+	/**
+	 * Modify the passed in Generated Keys by the current tracks values and weight at the passed in time.
+
+	 * @param Object The handle to the object modify
+	 * @param Track The track we are modifying
+	 * @param SectionToKey The Sections Channels we will be modifiying
+	 * @param Time The Time at which to evaluate
+	 * @param InOutGeneratedTrackKeys The Keys we need to modify. We change these values.
+	 * @param Weight The weight we need to modify the values by.
+	 */
+	bool ModifyOurGeneratedKeysByCurrentAndWeight(UObject* Object, IControlRigManipulatable* Manip, UMovieSceneTrack *Track, UMovieSceneSection* SectionToKey, FFrameNumber Time, FGeneratedTrackKeys& InOutGeneratedTotalKeys, float Weight) const;
+
+	/*
+	bool ShouldFilterAsset(const FAssetData& AssetData);
+
+	void OnControlRigAssetSelected(const FAssetData& AssetData, TArray<FGuid> ObjectBindings, USkeleton* Skeleton);
+
+	void OnControlRigAssetEnterPressed(const TArray<FAssetData>& AssetData, TArray<FGuid> ObjectBindings,  USkeleton* Skeleton);
+	*/
+private:
+	bool CanAddTransformKeysForSelectedObjects();
+	void OnAddTransformKeysForSelectedObjects(EMovieSceneTransformChannel Channel);
+	FDelegateHandle SelectionChangedHandle;
+	FDelegateHandle SequencerChangedHandle;
+
+private:
+
+	/** Guard to stop infinite loops when handling control selections*/
+	bool bIsDoingASelection;
+
+};
+
+
+/** Class for control rig sections */
+class FControlRigParameterSection : public FSequencerSection
+{
+public:
+
+	/**
+	* Creates a new control rig property section.
+	*
+	* @param InSection The section object which is being displayed and edited.
+	* @param InSequencer The sequencer which is controlling this parameter section.
+	*/
+	FControlRigParameterSection(UMovieSceneSection& InSection, TWeakPtr<ISequencer> InSequencer)
+		: FSequencerSection(InSection), WeakSequencer(InSequencer)
+	{
+	}
+
+public:
+
+	virtual void BuildSectionContextMenu(FMenuBuilder& MenuBuilder, const FGuid& InObjectBinding) override;
+
+	//~ ISequencerSection interface
+	virtual bool RequestDeleteCategory(const TArray<FName>& CategoryNamePath) override;
+	virtual bool RequestDeleteKeyArea(const TArray<FName>& KeyAreaNamePath) override;
+
+protected:
+
+	/** The sequencer which is controlling this section. */
+	TWeakPtr<ISequencer> WeakSequencer;
+
+};
\ No newline at end of file
diff --git a/Engine/Plugins/Experimental/ControlRig/Source/ControlRigManipulation/Private/ControlRigGizmoActor.cpp b/Engine/Plugins/Experimental/ControlRig/Source/ControlRigManipulation/Private/ControlRigGizmoActor.cpp
--- a/Engine/Plugins/Experimental/ControlRig/Source/ControlRigManipulation/Private/ControlRigGizmoActor.cpp
+++ b/Engine/Plugins/Experimental/ControlRig/Source/ControlRigManipulation/Private/ControlRigGizmoActor.cpp
@@ -24,6 +24,8 @@ AControlRigGizmoActor::AControlRigGizmoActor(const FObjectInitializer& ObjectIni
 
 	RootComponent = ActorRootComponent;
 	StaticMeshComponent->SetupAttachment(RootComponent);
+	StaticMeshComponent->bCastStaticShadow = false;
+	StaticMeshComponent->bCastDynamicShadow = false;
 }
 
 void AControlRigGizmoActor::SetEnabled(bool bInEnabled)
@@ -108,7 +110,7 @@ namespace FControlRigGizmoHelper
 		ActorSpawnParameters.bTemporaryEditorActor = true;
 		ActorSpawnParameters.bHideFromSceneOutliner = true;
 		ActorSpawnParameters.SpawnCollisionHandlingOverride = ESpawnActorCollisionHandlingMethod::AlwaysSpawn;
-
+		ActorSpawnParameters.ObjectFlags = RF_Transient;
 		return ActorSpawnParameters;
 	}
 
diff --git a/Engine/Plugins/Experimental/ControlRig/Source/ControlRigManipulation/Private/DefaultControlRigManipulationLayer.cpp b/Engine/Plugins/Experimental/ControlRig/Source/ControlRigManipulation/Private/DefaultControlRigManipulationLayer.cpp
--- a/Engine/Plugins/Experimental/ControlRig/Source/ControlRigManipulation/Private/DefaultControlRigManipulationLayer.cpp
+++ b/Engine/Plugins/Experimental/ControlRig/Source/ControlRigManipulation/Private/DefaultControlRigManipulationLayer.cpp
@@ -19,9 +19,6 @@ void UDefaultControlRigManipulationLayer::CreateLayer()
 
 void UDefaultControlRigManipulationLayer::DestroyLayer()
 {
-	ExecuteDelegateHandles.Empty();
-	InitializeDelegateHandle.Empty();
-
 	IControlRigManipulationLayer::DestroyLayer();
 }
 
@@ -49,17 +46,22 @@ void UDefaultControlRigManipulationLayer::RemoveManipulatableObject(IControlRigM
 void UDefaultControlRigManipulationLayer::OnControlRigAdded(UControlRig* InControlRig)
 {
 	// bind execution delegate
-	ExecuteDelegateHandles.Add(InControlRig->OnExecuted().AddUObject(this, &UDefaultControlRigManipulationLayer::OnRigExecuted));
-	InitializeDelegateHandle.Add(InControlRig->OnInitialized().AddUObject(this, &UDefaultControlRigManipulationLayer::OnRigInitialized));
+
 	ControlModifiedDelegateHandles.Add(InControlRig->ControlModified().AddUObject(this, &UDefaultControlRigManipulationLayer::OnControlModified));
 
 	// ensure the size matches
-	ensure(ExecuteDelegateHandles.Num() == ManipulatableObjects.Num());
-	ensure(InitializeDelegateHandle.Num() == ManipulatableObjects.Num());
 	ensure(ControlModifiedDelegateHandles.Num() == ManipulatableObjects.Num());
 
 	// object binding. This overwrites if there were multi
 	SetObjectBinding(InControlRig->GetObjectBinding());
+
+	// currently all the manipulatable mesh component is supposed to be same
+	// if that changes, this code has to change
+	USkeletalMeshComponent* MeshComponent = GetSkeletalMeshComponent();
+	if (MeshComponent)
+	{
+		MeshComponent->OnBoneTransformsFinalized.AddDynamic(this, &UDefaultControlRigManipulationLayer::PostPoseUpdate);
+	}
 }
 
 void UDefaultControlRigManipulationLayer::OnControlRigRemoved(UControlRig* InControlRig)
@@ -68,14 +70,16 @@ void UDefaultControlRigManipulationLayer::OnControlRigRemoved(UControlRig* InCon
 	int32 Found = ManipulatableObjects.Find(InControlRig);
 	if (Found != INDEX_NONE)
 	{
-		if (ExecuteDelegateHandles.IsValidIndex(Found) && ExecuteDelegateHandles[Found].IsValid())
+		// last one 
+		// currently all the manipulatable mesh component is supposed to be same
+		// if that changes, this code has to change
+		if (ManipulatableObjects.Num() == 1)
 		{
-			InControlRig->OnExecuted().Remove(ExecuteDelegateHandles[Found]);
-		}
-
-		if (InitializeDelegateHandle.IsValidIndex(Found) && InitializeDelegateHandle[Found].IsValid())
-		{
-			InControlRig->OnInitialized().Remove(InitializeDelegateHandle[Found]);
+			USkeletalMeshComponent* MeshComponent = GetSkeletalMeshComponent();
+			if (MeshComponent)
+			{
+				MeshComponent->OnBoneTransformsFinalized.RemoveDynamic(this, &UDefaultControlRigManipulationLayer::PostPoseUpdate);
+			}
 		}
 
 		if (ControlModifiedDelegateHandles.IsValidIndex(Found) && ControlModifiedDelegateHandles[Found].IsValid())
@@ -83,7 +87,6 @@ void UDefaultControlRigManipulationLayer::OnControlRigRemoved(UControlRig* InCon
 			InControlRig->ControlModified().Remove(ControlModifiedDelegateHandles[Found]);
 		}
 	}
-	
 }
 
 // gizmo related 
@@ -321,10 +324,21 @@ bool UDefaultControlRigManipulationLayer::CreateGizmoActors(UWorld* World, TArra
 			OutGizmoActors.Add(GizmoActor);
 		}
 	}
-	
+
+	WorldPtr = World;
+	OnWorldCleanupHandle = FWorldDelegates::OnWorldCleanup.AddUObject(this, &UDefaultControlRigManipulationLayer::OnWorldCleanup);
 	return (OutGizmoActors.Num() > 0);
 }
 
+void UDefaultControlRigManipulationLayer::OnWorldCleanup(UWorld* World, bool bSessionEnded, bool bCleanupResources)
+{
+	// if world gets cleaned up first, we destroy gizmo actors
+	if (WorldPtr == World)
+	{
+		DestroyGizmosActors();
+	}
+}
+
 void UDefaultControlRigManipulationLayer::DestroyGizmosActors()
 {
 	// clear previous gizmo actors
@@ -340,21 +354,8 @@ void UDefaultControlRigManipulationLayer::DestroyGizmosActors()
 	}
 
 	ResetControlData();
-}
 
-void UDefaultControlRigManipulationLayer::OnRigExecuted(class UControlRig* InRig, const EControlRigState State)
-{
-	FTransform ComponentTransform = GetSkeletalMeshComponentTransform();
-
-	// after executing rig, update gizmo
-	for (auto Iter=GizmoToControlMap.CreateIterator(); Iter; ++Iter)
-	{
-		TickGizmo(Iter.Key(), ComponentTransform);
-	}
-}
-
-void UDefaultControlRigManipulationLayer::OnRigInitialized(class UControlRig* InRig, const EControlRigState State)
-{
+	FWorldDelegates::OnWorldCleanup.Remove(OnWorldCleanupHandle);
 }
 
 void UDefaultControlRigManipulationLayer::OnControlModified(IControlRigManipulatable* InManipulatable, const FRigControl& InControl)
@@ -424,4 +425,14 @@ void UDefaultControlRigManipulationLayer::BeginTransaction()
 
 void UDefaultControlRigManipulationLayer::EndTransaction()
 {
+}
+
+void UDefaultControlRigManipulationLayer::PostPoseUpdate()
+{
+	FTransform ComponentTransform = GetSkeletalMeshComponentTransform();
+	// after executing rig, update gizmo
+	for (auto Iter = GizmoToControlMap.CreateIterator(); Iter; ++Iter)
+	{
+		TickGizmo(Iter.Key(), ComponentTransform);
+	}
 }
\ No newline at end of file
diff --git a/Engine/Plugins/Experimental/ControlRig/Source/ControlRigManipulation/Public/DefaultControlRigManipulationLayer.h b/Engine/Plugins/Experimental/ControlRig/Source/ControlRigManipulation/Public/DefaultControlRigManipulationLayer.h
--- a/Engine/Plugins/Experimental/ControlRig/Source/ControlRigManipulation/Public/DefaultControlRigManipulationLayer.h
+++ b/Engine/Plugins/Experimental/ControlRig/Source/ControlRigManipulation/Public/DefaultControlRigManipulationLayer.h
@@ -50,8 +50,6 @@ public:
 	virtual void MoveGizmo(AControlRigGizmoActor* GizmoActor, const bool bTranslation, FVector& InDrag, 
 		const bool bRotation, FRotator& InRot, const bool bScale, FVector& InScale, const FTransform& ToWorldTransform) override;
 	virtual void TickGizmo(AControlRigGizmoActor* GizmoActor, const FTransform& ComponentTransform) override;
-	// need templatizing input type @fixme: add implementation
-	virtual void AddKey(const AControlRigGizmoActor* GizmoActor, TSharedPtr<ISequencer> Sequencer) override {};
 	virtual bool ModeSupportedByGizmoActor(const AControlRigGizmoActor* GizmoActor, FWidget::EWidgetMode InMode) const override;
 	// IControlRigManipulationLayer END
 
@@ -87,11 +85,9 @@ private:
 	// control data related
 	virtual void ResetControlData();
 
-	void OnRigExecuted(class UControlRig* InRig, const EControlRigState State);
-	TArray<FDelegateHandle> ExecuteDelegateHandles;
-
-	void OnRigInitialized(class UControlRig* InRig, const EControlRigState State);
-	TArray<FDelegateHandle> InitializeDelegateHandle;
+	// Post pose update handler
+	UFUNCTION()
+	virtual void PostPoseUpdate();
 
 	void OnControlModified(IControlRigManipulatable* InManipulatable, const FRigControl& InControl);
 	TArray<FDelegateHandle> ControlModifiedDelegateHandles;
@@ -100,4 +96,8 @@ private:
 	void OnControlRigRemoved(UControlRig* InControlRig);
 
 	void GetGizmoCreationParams(TArray<FGizmoActorCreationParam>& OutCreationParams);
+	// world clean up handlers
+	FDelegateHandle OnWorldCleanupHandle;
+	void OnWorldCleanup(UWorld* World, bool bSessionEnded, bool bCleanupResources);
+	UWorld* WorldPtr = nullptr;
 };
diff --git a/Engine/Plugins/Experimental/ControlRig/Source/ControlRigManipulation/Public/IControlRigManipulationLayer.h b/Engine/Plugins/Experimental/ControlRig/Source/ControlRigManipulation/Public/IControlRigManipulationLayer.h
--- a/Engine/Plugins/Experimental/ControlRig/Source/ControlRigManipulation/Public/IControlRigManipulationLayer.h
+++ b/Engine/Plugins/Experimental/ControlRig/Source/ControlRigManipulation/Public/IControlRigManipulationLayer.h
@@ -37,9 +37,6 @@ public:
 	virtual void TickGizmo(AControlRigGizmoActor* GizmoActor, const FTransform& ComponentTransform) = 0;
 	virtual bool ModeSupportedByGizmoActor(const AControlRigGizmoActor* GizmoActor, FWidget::EWidgetMode InMode) const = 0;
 
-	// add key of current gizmo actor transform data
-	virtual void AddKey(const AControlRigGizmoActor* GizmoActor, TSharedPtr<ISequencer> Sequencer) = 0;
-
 protected:
 
 	// because manipulatable is created externally often. Manipulation Layer doesn't have any ownership on these object
diff --git a/Engine/Source/Editor/AnimGraph/Classes/AnimationGraphSchema.h b/Engine/Source/Editor/AnimGraph/Classes/AnimationGraphSchema.h
--- a/Engine/Source/Editor/AnimGraph/Classes/AnimationGraphSchema.h
+++ b/Engine/Source/Editor/AnimGraph/Classes/AnimationGraphSchema.h
@@ -13,6 +13,7 @@
 class FMenuBuilder;
 class UAnimationAsset;
 class UAnimBlueprint;
+class UPhysicsAsset;
 struct FBPInterfaceDescription;
 
 UCLASS(MinimalAPI)
@@ -77,6 +78,9 @@ class UAnimationGraphSchema : public UEdGraphSchema_K2
 	/** Spawn the correct node in the Animation Graph using the given AnimationAsset at the supplied location */
 	static void SpawnNodeFromAsset(UAnimationAsset* Asset, const FVector2D& GraphPosition, UEdGraph* Graph, UEdGraphPin* PinIfAvailable);
 
+	/** Spawn a rigid body node if we drop a physics asset on the graph */
+	static void SpawnRigidBodyNodeFromAsset(UPhysicsAsset* Asset, const FVector2D& GraphPosition, UEdGraph* Graph);
+
 	/** Update the specified node to a new asset */
 	static void UpdateNodeWithAsset(class UK2Node* K2Node, UAnimationAsset* Asset);
 
diff --git a/Engine/Source/Editor/AnimGraph/Private/AnimationGraphSchema.cpp b/Engine/Source/Editor/AnimGraph/Private/AnimationGraphSchema.cpp
--- a/Engine/Source/Editor/AnimGraph/Private/AnimationGraphSchema.cpp
+++ b/Engine/Source/Editor/AnimGraph/Private/AnimationGraphSchema.cpp
@@ -35,6 +35,7 @@
 #include "Animation/AnimMontage.h"
 #include "AnimGraphNode_LinkedInputPose.h"
 #include "AnimGraphNode_LinkedAnimLayer.h"
+#include "AnimGraphNode_RigidBody.h"
 
 #define LOCTEXT_NAMESPACE "AnimationGraphSchema"
 
@@ -410,6 +411,20 @@ void UAnimationGraphSchema::SpawnNodeFromAsset(UAnimationAsset* Asset, const FVe
 	}
 }
 
+void UAnimationGraphSchema::SpawnRigidBodyNodeFromAsset(UPhysicsAsset* Asset, const FVector2D& GraphPosition, UEdGraph* Graph)
+{
+	check(Graph);
+	check(Graph->GetSchema()->IsA(UAnimationGraphSchema::StaticClass()));
+	check(Asset);
+
+	FEdGraphSchemaAction_K2NewNode Action;
+
+	UAnimGraphNode_RigidBody* NewNode = NewObject<UAnimGraphNode_RigidBody>(GetTransientPackage());
+	NewNode->Node.OverridePhysicsAsset = Asset;
+	Action.NodeTemplate = NewNode;
+
+	Action.PerformAction(Graph, nullptr, GraphPosition);
+}
 
 void UAnimationGraphSchema::UpdateNodeWithAsset(UK2Node* K2Node, UAnimationAsset* Asset)
 {
@@ -431,18 +446,21 @@ void UAnimationGraphSchema::UpdateNodeWithAsset(UK2Node* K2Node, UAnimationAsset
 	}
 }
 
-
 void UAnimationGraphSchema::DroppedAssetsOnGraph( const TArray<FAssetData>& Assets, const FVector2D& GraphPosition, UEdGraph* Graph ) const 
 {
-	UAnimationAsset* Asset = FAssetData::GetFirstAsset<UAnimationAsset>(Assets);
-	if ((Asset != NULL) && (Graph != NULL))
+	if (Graph != NULL)
 	{
-		SpawnNodeFromAsset(Asset, GraphPosition, Graph, NULL);
+		if (UAnimationAsset* AnimationAsset = FAssetData::GetFirstAsset<UAnimationAsset>(Assets))
+		{
+			SpawnNodeFromAsset(AnimationAsset, GraphPosition, Graph, NULL);
+		}
+		else if (UPhysicsAsset* PhysicsAsset = FAssetData::GetFirstAsset<UPhysicsAsset>(Assets))
+		{
+			SpawnRigidBodyNodeFromAsset(PhysicsAsset, GraphPosition, Graph);
+		}
 	}
 }
 
-
-
 void UAnimationGraphSchema::DroppedAssetsOnNode(const TArray<FAssetData>& Assets, const FVector2D& GraphPosition, UEdGraphNode* Node) const
 {
 	UAnimationAsset* Asset = FAssetData::GetFirstAsset<UAnimationAsset>(Assets);
@@ -528,11 +546,10 @@ void UAnimationGraphSchema::GetAssetsPinHoverMessage(const TArray<FAssetData>& A
 
 void UAnimationGraphSchema::GetAssetsGraphHoverMessage(const TArray<FAssetData>& Assets, const UEdGraph* HoverGraph, FString& OutTooltipText, bool& OutOkIcon) const
 {
-	UAnimationAsset* Asset = FAssetData::GetFirstAsset<UAnimationAsset>(Assets);
-	if (Asset)
+	if (UAnimationAsset* AnimationAsset = FAssetData::GetFirstAsset<UAnimationAsset>(Assets))
 	{
 		UAnimBlueprint* AnimBlueprint = Cast<UAnimBlueprint>(FBlueprintEditorUtils::FindBlueprintForGraph(HoverGraph));
-		const bool bSkelMatch = (AnimBlueprint != NULL) && (AnimBlueprint->TargetSkeleton == Asset->GetSkeleton());
+		const bool bSkelMatch = (AnimBlueprint != NULL) && (AnimBlueprint->TargetSkeleton == AnimationAsset->GetSkeleton());
 		if (!bSkelMatch)
 		{
 			OutOkIcon = false;
@@ -549,6 +566,11 @@ void UAnimationGraphSchema::GetAssetsGraphHoverMessage(const TArray<FAssetData>&
 			OutTooltipText = TEXT("");
 		}
 	}
+	else if(UPhysicsAsset* PhysicsAsset = FAssetData::GetFirstAsset<UPhysicsAsset>(Assets))
+	{
+		OutOkIcon = true;
+		OutTooltipText = TEXT("");
+	}
 }
 
 void UAnimationGraphSchema::GetContextMenuActions(UToolMenu* Menu, UGraphNodeContextMenuContext* Context) const
diff --git a/Engine/Source/Editor/AnimationBlueprintEditor/Private/AnimationBlueprintEditor.cpp b/Engine/Source/Editor/AnimationBlueprintEditor/Private/AnimationBlueprintEditor.cpp
--- a/Engine/Source/Editor/AnimationBlueprintEditor/Private/AnimationBlueprintEditor.cpp
+++ b/Engine/Source/Editor/AnimationBlueprintEditor/Private/AnimationBlueprintEditor.cpp
@@ -310,7 +310,7 @@ void FAnimationBlueprintEditor::InitAnimationBlueprintEditor(const EToolkitMode:
 
 	// register customization of Slot node for this Animation Blueprint Editor
 	// this is so that you can open the manage window per Animation Blueprint Editor
-	PersonaModule.CustomizeSlotNodeDetails(Inspector->GetPropertyView().ToSharedRef(), FOnInvokeTab::CreateSP(this, &FAssetEditorToolkit::InvokeTab));
+	PersonaModule.CustomizeBlueprintEditorDetails(Inspector->GetPropertyView().ToSharedRef(), FOnInvokeTab::CreateSP(this, &FAssetEditorToolkit::InvokeTab));
 
 	if(bNewlyCreated && InAnimBlueprint->BlueprintType == BPTYPE_Interface)
 	{
@@ -396,7 +396,7 @@ void FAnimationBlueprintEditor::ExtendToolbar()
 		}
 	};
 
-    ToolbarExtender->AddToolBarExtension(
+	ToolbarExtender->AddToolBarExtension(
 		"Asset",
 		EExtensionHook::After,
 		GetToolkitCommands(),
@@ -1211,6 +1211,11 @@ FLinearColor FAnimationBlueprintEditor::GetWorldCentricTabColorScale() const
 	return FLinearColor( 0.5f, 0.25f, 0.35f, 0.5f );
 }
 
+IAnimationSequenceBrowser* FAnimationBlueprintEditor::GetAssetBrowser() const
+{
+	return SequenceBrowser.Pin().Get();
+}
+
 void FAnimationBlueprintEditor::OnActiveTabChanged( TSharedPtr<SDockTab> PreviouslyActive, TSharedPtr<SDockTab> NewlyActivated )
 {
 	if (!NewlyActivated.IsValid())
@@ -1483,7 +1488,7 @@ void FAnimationBlueprintEditor::HandleOpenNewAsset(UObject* InNewAsset)
 
 void FAnimationBlueprintEditor::AddReferencedObjects( FReferenceCollector& Collector )
 {
-    Collector.AddReferencedObject( EditorOptions );
+	Collector.AddReferencedObject( EditorOptions );
 }
 
 FAnimNode_Base* FAnimationBlueprintEditor::FindAnimNode(UAnimGraphNode_Base* AnimGraphNode) const
@@ -1535,7 +1540,7 @@ void FAnimationBlueprintEditor::OnSelectedNodesChangedImpl(const TSet<class UObj
 		}
 	}
 
-    bSelectRegularNode = false;
+	bSelectRegularNode = false;
 	for (FGraphPanelSelectionSet::TConstIterator It(NewSelection); It; ++It)
 	{
 		UEdGraphNode_Comment* SeqNode = Cast<UEdGraphNode_Comment>(*It);
@@ -1548,7 +1553,7 @@ void FAnimationBlueprintEditor::OnSelectedNodesChangedImpl(const TSet<class UObj
 		}
 	}
 
-    if (bHideUnrelatedNodes && !bLockNodeFadeState)
+	if (bHideUnrelatedNodes && !bLockNodeFadeState)
 	{
 		ResetAllNodesUnrelatedStates();
 
@@ -1813,5 +1818,10 @@ void FAnimationBlueprintEditor::HandlePreviewAnimBlueprintCompiled(UBlueprint* I
 	}
 }
 
+void FAnimationBlueprintEditor::HandleAnimationSequenceBrowserCreated(const TSharedRef<IAnimationSequenceBrowser>& InSequenceBrowser)
+{
+	SequenceBrowser = InSequenceBrowser;
+}
+
 #undef LOCTEXT_NAMESPACE
 
diff --git a/Engine/Source/Editor/AnimationBlueprintEditor/Private/AnimationBlueprintEditor.h b/Engine/Source/Editor/AnimationBlueprintEditor/Private/AnimationBlueprintEditor.h
--- a/Engine/Source/Editor/AnimationBlueprintEditor/Private/AnimationBlueprintEditor.h
+++ b/Engine/Source/Editor/AnimationBlueprintEditor/Private/AnimationBlueprintEditor.h
@@ -91,6 +91,7 @@ public:
 	/** IAnimationBlueprintEditor interface */
 	virtual const FEdGraphPinType& GetLastGraphPinTypeUsed() const override { return LastGraphPinType; }
 	virtual void SetLastGraphPinTypeUsed(const FEdGraphPinType& InType) override { LastGraphPinType = InType; }
+	virtual IAnimationSequenceBrowser* GetAssetBrowser() const override;
 
 	/** IHasPersonaToolkit interface */
 	virtual TSharedRef<class IPersonaToolkit> GetPersonaToolkit() const { return PersonaToolkit.ToSharedRef(); }
@@ -287,6 +288,8 @@ private:
 	 */
 	virtual void SaveEditorSettings();
 
+	void HandleAnimationSequenceBrowserCreated(const TSharedRef<IAnimationSequenceBrowser>& InSequenceBrowser);
+
 	/** The extender to pass to the level editor to extend it's window menu */
 	TSharedPtr<FExtender> MenuExtender;
 
@@ -305,6 +308,9 @@ private:
 	// selected anim graph node 
 	TWeakObjectPtr<class UAnimGraphNode_Base> SelectedAnimGraphNode;
 
+	/** Sequence Browser **/
+	TWeakPtr<class IAnimationSequenceBrowser> SequenceBrowser;
+
 	/** Delegate handle registered for when pin default values change */
 	FDelegateHandle OnPinDefaultValueChangedHandle;
 
diff --git a/Engine/Source/Editor/AnimationBlueprintEditor/Private/AnimationBlueprintEditorMode.cpp b/Engine/Source/Editor/AnimationBlueprintEditor/Private/AnimationBlueprintEditorMode.cpp
--- a/Engine/Source/Editor/AnimationBlueprintEditor/Private/AnimationBlueprintEditorMode.cpp
+++ b/Engine/Source/Editor/AnimationBlueprintEditor/Private/AnimationBlueprintEditorMode.cpp
@@ -132,7 +132,7 @@ FAnimationBlueprintEditorMode::FAnimationBlueprintEditorMode(const TSharedRef<FA
 	PersonaModule.RegisterPersonaViewportTabFactories(TabFactories, InAnimationBlueprintEditor, ViewportArgs);
 
 	TabFactories.RegisterFactory(PersonaModule.CreateAdvancedPreviewSceneTabFactory(InAnimationBlueprintEditor, InAnimationBlueprintEditor->GetPersonaToolkit()->GetPreviewScene()));
-	TabFactories.RegisterFactory(PersonaModule.CreateAnimationAssetBrowserTabFactory(InAnimationBlueprintEditor, InAnimationBlueprintEditor->GetPersonaToolkit(), FOnOpenNewAsset::CreateSP(&InAnimationBlueprintEditor.Get(), &FAnimationBlueprintEditor::HandleOpenNewAsset), FOnAnimationSequenceBrowserCreated(), true));
+	TabFactories.RegisterFactory(PersonaModule.CreateAnimationAssetBrowserTabFactory(InAnimationBlueprintEditor, InAnimationBlueprintEditor->GetPersonaToolkit(), FOnOpenNewAsset::CreateSP(&InAnimationBlueprintEditor.Get(), &FAnimationBlueprintEditor::HandleOpenNewAsset), FOnAnimationSequenceBrowserCreated::CreateSP(&InAnimationBlueprintEditor.Get(), &FAnimationBlueprintEditor::HandleAnimationSequenceBrowserCreated), true));
 	TabFactories.RegisterFactory(PersonaModule.CreateAnimBlueprintPreviewTabFactory(InAnimationBlueprintEditor, InAnimationBlueprintEditor->GetPersonaToolkit()->GetPreviewScene()));
 	TabFactories.RegisterFactory(PersonaModule.CreateAnimBlueprintAssetOverridesTabFactory(InAnimationBlueprintEditor, InAnimationBlueprintEditor->GetPersonaToolkit()->GetAnimBlueprint(), InAnimationBlueprintEditor->OnPostUndo));
 	TabFactories.RegisterFactory(PersonaModule.CreateSkeletonSlotNamesTabFactory(InAnimationBlueprintEditor, InAnimationBlueprintEditor->GetSkeletonTree()->GetEditableSkeleton(), InAnimationBlueprintEditor->OnPostUndo, FOnObjectSelected::CreateSP(&InAnimationBlueprintEditor.Get(), &FAnimationBlueprintEditor::HandleObjectSelected)));
diff --git a/Engine/Source/Editor/AnimationBlueprintEditor/Private/AnimationStateNodes/SGraphNodeAnimState.cpp b/Engine/Source/Editor/AnimationBlueprintEditor/Private/AnimationStateNodes/SGraphNodeAnimState.cpp
--- a/Engine/Source/Editor/AnimationBlueprintEditor/Private/AnimationStateNodes/SGraphNodeAnimState.cpp
+++ b/Engine/Source/Editor/AnimationBlueprintEditor/Private/AnimationStateNodes/SGraphNodeAnimState.cpp
@@ -83,37 +83,38 @@ void SGraphNodeAnimState::Construct(const FArguments& InArgs, UAnimStateNodeBase
 
 void SGraphNodeAnimState::GetStateInfoPopup(UEdGraphNode* GraphNode, TArray<FGraphInformationPopupInfo>& Popups)
 {
-	UAnimBlueprint* AnimBlueprint = Cast<UAnimBlueprint>(FBlueprintEditorUtils::FindBlueprintForNodeChecked(GraphNode));
-	check(AnimBlueprint);
-	UAnimInstance* ActiveObject = Cast<UAnimInstance>(AnimBlueprint->GetObjectBeingDebugged());
-	UAnimBlueprintGeneratedClass* Class = AnimBlueprint->GetAnimBlueprintGeneratedClass();
+	UAnimBlueprint* AnimBlueprint = Cast<UAnimBlueprint>(FBlueprintEditorUtils::FindBlueprintForNode(GraphNode));
+	if(AnimBlueprint)
+	{
+		UAnimInstance* ActiveObject = Cast<UAnimInstance>(AnimBlueprint->GetObjectBeingDebugged());
+		UAnimBlueprintGeneratedClass* Class = AnimBlueprint->GetAnimBlueprintGeneratedClass();
 
-	//FKismetNodeInfoContext* K2Context = (FKismetNodeInfoContext*)Context;
-	FLinearColor CurrentStateColor(1.f, 0.5f, 0.25f);
+		FLinearColor CurrentStateColor(1.f, 0.5f, 0.25f);
 
-	// Display various types of debug data
-	if ((ActiveObject != NULL) && (Class != NULL))
-	{
-		if (FStateMachineDebugData* DebugInfo = Class->GetAnimBlueprintDebugData().StateMachineDebugData.Find(GraphNode->GetGraph()))
+		// Display various types of debug data
+		if ((ActiveObject != NULL) && (Class != NULL))
 		{
-			if (Class->AnimNodeProperties.Num())
+			if (FStateMachineDebugData* DebugInfo = Class->GetAnimBlueprintDebugData().StateMachineDebugData.Find(GraphNode->GetGraph()))
 			{
-				UAnimationStateMachineGraph* TypedGraph = CastChecked<UAnimationStateMachineGraph>(GraphNode->GetGraph());
-
-				if (FAnimNode_StateMachine* CurrentInstance = Class->GetPropertyInstance<FAnimNode_StateMachine>(ActiveObject, TypedGraph->OwnerAnimGraphNode))  
+				if (Class->AnimNodeProperties.Num())
 				{
-					if (int32* pStateIndex = DebugInfo->NodeToStateIndex.Find(GraphNode))
+					UAnimationStateMachineGraph* TypedGraph = CastChecked<UAnimationStateMachineGraph>(GraphNode->GetGraph());
+
+					if (FAnimNode_StateMachine* CurrentInstance = Class->GetPropertyInstance<FAnimNode_StateMachine>(ActiveObject, TypedGraph->OwnerAnimGraphNode))  
 					{
-						const float Weight = CurrentInstance->GetStateWeight(*pStateIndex);
-						if (Weight > 0.0f)
+						if (int32* pStateIndex = DebugInfo->NodeToStateIndex.Find(GraphNode))
 						{
-							FString StateText = FString::Printf(TEXT("%.1f%%"), Weight * 100.0f);
-							if (*pStateIndex == CurrentInstance->GetCurrentState())
+							const float Weight = CurrentInstance->GetStateWeight(*pStateIndex);
+							if (Weight > 0.0f)
 							{
-								StateText += FString::Printf(TEXT("\nActive for %.2f secs"), CurrentInstance->GetCurrentStateElapsedTime());
-							}
+								FString StateText = FString::Printf(TEXT("%.1f%%"), Weight * 100.0f);
+								if (*pStateIndex == CurrentInstance->GetCurrentState())
+								{
+									StateText += FString::Printf(TEXT("\nActive for %.2f secs"), CurrentInstance->GetCurrentStateElapsedTime());
+								}
 
-							new (Popups) FGraphInformationPopupInfo(NULL, CurrentStateColor, StateText);
+								new (Popups) FGraphInformationPopupInfo(NULL, CurrentStateColor, StateText);
+							}
 						}
 					}
 				}
@@ -129,34 +130,34 @@ void SGraphNodeAnimState::GetNodeInfoPopups(FNodeInfoContext* Context, TArray<FG
 
 FSlateColor SGraphNodeAnimState::GetBorderBackgroundColor() const
 {
-	UAnimBlueprint* AnimBlueprint = Cast<UAnimBlueprint>(FBlueprintEditorUtils::FindBlueprintForNodeChecked(GraphNode));
-	check(AnimBlueprint);
-	UAnimInstance* ActiveObject = Cast<UAnimInstance>(AnimBlueprint->GetObjectBeingDebugged());
-	UAnimBlueprintGeneratedClass* Class = AnimBlueprint->GetAnimBlueprintGeneratedClass();
-
-	//FKismetNodeInfoContext* K2Context = (FKismetNodeInfoContext*)Context;
-
 	FLinearColor InactiveStateColor(0.08f, 0.08f, 0.08f);
 	FLinearColor ActiveStateColorDim(0.4f, 0.3f, 0.15f);
 	FLinearColor ActiveStateColorBright(1.f, 0.6f, 0.35f);
 
-	// Display various types of debug data
-	if ((ActiveObject != NULL) && (Class != NULL))
+	UAnimBlueprint* AnimBlueprint = Cast<UAnimBlueprint>(FBlueprintEditorUtils::FindBlueprintForNode(GraphNode));
+	if(AnimBlueprint)
 	{
-		if (FStateMachineDebugData* DebugInfo = Class->GetAnimBlueprintDebugData().StateMachineDebugData.Find(GraphNode->GetGraph()))
+		UAnimInstance* ActiveObject = Cast<UAnimInstance>(AnimBlueprint->GetObjectBeingDebugged());
+		UAnimBlueprintGeneratedClass* Class = AnimBlueprint->GetAnimBlueprintGeneratedClass();
+
+		// Display various types of debug data
+		if ((ActiveObject != NULL) && (Class != NULL))
 		{
-			if (Class->AnimNodeProperties.Num())
+			if (FStateMachineDebugData* DebugInfo = Class->GetAnimBlueprintDebugData().StateMachineDebugData.Find(GraphNode->GetGraph()))
 			{
-				UAnimationStateMachineGraph* TypedGraph = CastChecked<UAnimationStateMachineGraph>(GraphNode->GetGraph());
-
-				if (FAnimNode_StateMachine* CurrentInstance = Class->GetPropertyInstance<FAnimNode_StateMachine>(ActiveObject, TypedGraph->OwnerAnimGraphNode))
+				if (Class->AnimNodeProperties.Num())
 				{
-					if (int32* pStateIndex = DebugInfo->NodeToStateIndex.Find(GraphNode))
+					UAnimationStateMachineGraph* TypedGraph = CastChecked<UAnimationStateMachineGraph>(GraphNode->GetGraph());
+
+					if (FAnimNode_StateMachine* CurrentInstance = Class->GetPropertyInstance<FAnimNode_StateMachine>(ActiveObject, TypedGraph->OwnerAnimGraphNode))
 					{
-						const float Weight = CurrentInstance->GetStateWeight(*pStateIndex);
-						if (Weight > 0.0f)
+						if (int32* pStateIndex = DebugInfo->NodeToStateIndex.Find(GraphNode))
 						{
-							return FMath::Lerp<FLinearColor>(ActiveStateColorDim, ActiveStateColorBright, Weight);
+							const float Weight = CurrentInstance->GetStateWeight(*pStateIndex);
+							if (Weight > 0.0f)
+							{
+								return FMath::Lerp<FLinearColor>(ActiveStateColorDim, ActiveStateColorBright, Weight);
+							}
 						}
 					}
 				}
diff --git a/Engine/Source/Editor/AnimationBlueprintEditor/Public/IAnimationBlueprintEditor.h b/Engine/Source/Editor/AnimationBlueprintEditor/Public/IAnimationBlueprintEditor.h
--- a/Engine/Source/Editor/AnimationBlueprintEditor/Public/IAnimationBlueprintEditor.h
+++ b/Engine/Source/Editor/AnimationBlueprintEditor/Public/IAnimationBlueprintEditor.h
@@ -6,6 +6,8 @@
 #include "BlueprintEditor.h"
 #include "IHasPersonaToolkit.h"
 
+class IAnimationSequenceBrowser;
+
 class IAnimationBlueprintEditor : public FBlueprintEditor, public IHasPersonaToolkit
 {
 public:
@@ -14,4 +16,7 @@ public:
 
 	/** Set the last pin type we used to create a graph pin */
 	virtual void SetLastGraphPinTypeUsed(const FEdGraphPinType& InType) = 0;
+
+	/** Get the asset browser we host */
+	virtual IAnimationSequenceBrowser* GetAssetBrowser() const = 0;
 };
diff --git a/Engine/Source/Editor/AnimationEditor/Private/AnimationEditor.cpp b/Engine/Source/Editor/AnimationEditor/Private/AnimationEditor.cpp
--- a/Engine/Source/Editor/AnimationEditor/Private/AnimationEditor.cpp
+++ b/Engine/Source/Editor/AnimationEditor/Private/AnimationEditor.cpp
@@ -460,6 +460,11 @@ void FAnimationEditor::SetAnimationAsset(UAnimationAsset* AnimAsset)
 	HandleOpenNewAsset(AnimAsset);
 }
 
+IAnimationSequenceBrowser* FAnimationEditor::GetAssetBrowser() const
+{
+	return SequenceBrowser.Pin().Get();
+}
+
 void FAnimationEditor::HandleOpenNewAsset(UObject* InNewAsset)
 {
 	if (UAnimationAsset* NewAnimationAsset = Cast<UAnimationAsset>(InNewAsset))
diff --git a/Engine/Source/Editor/AnimationEditor/Private/AnimationEditor.h b/Engine/Source/Editor/AnimationEditor/Private/AnimationEditor.h
--- a/Engine/Source/Editor/AnimationEditor/Private/AnimationEditor.h
+++ b/Engine/Source/Editor/AnimationEditor/Private/AnimationEditor.h
@@ -56,6 +56,7 @@ public:
 
 	/** IAnimationEditor interface */
 	virtual void SetAnimationAsset(UAnimationAsset* AnimAsset) override;
+	virtual IAnimationSequenceBrowser* GetAssetBrowser() const override;
 
 	/** IHasPersonaToolkit interface */
 	virtual TSharedRef<class IPersonaToolkit> GetPersonaToolkit() const override { return PersonaToolkit.ToSharedRef(); }
diff --git a/Engine/Source/Editor/AnimationEditor/Public/IAnimationEditor.h b/Engine/Source/Editor/AnimationEditor/Public/IAnimationEditor.h
--- a/Engine/Source/Editor/AnimationEditor/Public/IAnimationEditor.h
+++ b/Engine/Source/Editor/AnimationEditor/Public/IAnimationEditor.h
@@ -7,10 +7,14 @@
 #include "IHasPersonaToolkit.h"
 
 class UAnimationAsset;
+class IAnimationSequenceBrowser;
 
 class IAnimationEditor : public FWorkflowCentricApplication, public IHasPersonaToolkit
 {
 public:
 	/** Set the animation asset of the editor. */
 	virtual void SetAnimationAsset(UAnimationAsset* AnimAsset) = 0;
+
+	/** Get the asset browser we host */
+	virtual IAnimationSequenceBrowser* GetAssetBrowser() const = 0;
 };
diff --git a/Engine/Source/Editor/BlueprintGraph/Private/BlueprintActionFilter.cpp b/Engine/Source/Editor/BlueprintGraph/Private/BlueprintActionFilter.cpp
--- a/Engine/Source/Editor/BlueprintGraph/Private/BlueprintActionFilter.cpp
+++ b/Engine/Source/Editor/BlueprintGraph/Private/BlueprintActionFilter.cpp
@@ -995,7 +995,7 @@ static bool BlueprintActionFilterImpl::IsIncompatibleImpureNode(FBlueprintAction
 
 	for (UEdGraph* Graph : FilterContext.Graphs)
 	{
-		if (const UEdGraphSchema_K2* K2Schema = GetDefault<UEdGraphSchema_K2>())
+		if (const UEdGraphSchema_K2* K2Schema = Cast<UEdGraphSchema_K2>(Graph->GetSchema()))
 		{
 			bAllowImpureNodes &= K2Schema->DoesGraphSupportImpureFunctions(Graph);
 		}
diff --git a/Engine/Source/Editor/ContentBrowser/Private/SFilterList.cpp b/Engine/Source/Editor/ContentBrowser/Private/SFilterList.cpp
--- a/Engine/Source/Editor/ContentBrowser/Private/SFilterList.cpp
+++ b/Engine/Source/Editor/ContentBrowser/Private/SFilterList.cpp
@@ -26,6 +26,16 @@
 
 #define LOCTEXT_NAMESPACE "ContentBrowser"
 
+/** Helper struct to avoid friending the whole of SFilterList */
+struct FFrontendFilterExternalActivationHelper
+{
+	static void BindToFilter(TSharedRef<SFilterList> InFilterList, TSharedRef<FFrontendFilter> InFrontendFilter)
+	{
+		TWeakPtr<FFrontendFilter> WeakFilter = InFrontendFilter;
+		InFrontendFilter->SetActiveEvent.AddSP(&InFilterList.Get(), &SFilterList::OnSetFilterActive, WeakFilter);
+	}
+};
+
 /** A class for check boxes in the filter list. If you double click a filter checkbox, you will enable it and disable all others */
 class SFilterCheckBox : public SCheckBox
 {
@@ -516,10 +526,15 @@ void SFilterList::Construct( const FArguments& InArgs )
 
 	AllFrontendFilterCategories.AddUnique(DefaultCategory);
 
-	// Auto add all inverse filters
+	
 	for (auto FilterIt = AllFrontendFilters.CreateConstIterator(); FilterIt; ++FilterIt)
 	{
-		SetFrontendFilterActive(*FilterIt, false);
+		// Bind external activation event
+		TSharedRef<FFrontendFilter> Filter = *FilterIt;
+		FFrontendFilterExternalActivationHelper::BindToFilter(SharedThis(this), Filter);
+
+		// Auto add all inverse filters
+		SetFrontendFilterActive(Filter, false);
 	}
 
 	FilterBox = SNew(SWrapBox)
@@ -1483,4 +1498,17 @@ void SFilterList::OnResetFilters()
 	RemoveAllFilters();
 }
 
+void SFilterList::OnSetFilterActive(bool bInActive, TWeakPtr<FFrontendFilter> InWeakFilter)
+{
+	TSharedPtr<FFrontendFilter> Filter = InWeakFilter.Pin();
+	if (Filter.IsValid())
+	{
+		if (!IsFrontendFilterInUse(Filter.ToSharedRef()))
+		{
+			TSharedRef<SFilter> NewFilter = AddFilter(Filter.ToSharedRef());
+			NewFilter->SetEnabled(bInActive);
+		}
+	}
+}
+
 #undef LOCTEXT_NAMESPACE
diff --git a/Engine/Source/Editor/ContentBrowser/Private/SFilterList.h b/Engine/Source/Editor/ContentBrowser/Private/SFilterList.h
--- a/Engine/Source/Editor/ContentBrowser/Private/SFilterList.h
+++ b/Engine/Source/Editor/ContentBrowser/Private/SFilterList.h
@@ -131,6 +131,9 @@ private:
 	/** Called when reset filters option is pressed */
 	void OnResetFilters();
 
+	/** Called to set a filter active externally */
+	void OnSetFilterActive(bool bInActive, TWeakPtr<FFrontendFilter> InWeakFilter);
+
 private:
 	/** The horizontal box which contains all the filters */
 	TSharedPtr<SWrapBox> FilterBox;
@@ -155,4 +158,6 @@ private:
 
 	/** Delegate for when filters have changed */
 	FOnFilterChanged OnFilterChanged;
+
+	friend struct FFrontendFilterExternalActivationHelper;
 };
diff --git a/Engine/Source/Editor/ContentBrowser/Public/FrontendFilterBase.h b/Engine/Source/Editor/ContentBrowser/Public/FrontendFilterBase.h
--- a/Engine/Source/Editor/ContentBrowser/Public/FrontendFilterBase.h
+++ b/Engine/Source/Editor/ContentBrowser/Public/FrontendFilterBase.h
@@ -58,6 +58,9 @@ public:
 	/** Called when the state of a particular Content Browser is being loaded from INI */
 	virtual void LoadSettings(const FString& IniFilename, const FString& IniSection, const FString& SettingsString) {}
 
+	/** Set this filter as active/inactive */
+	void SetActive(bool bInActive) { SetActiveEvent.Broadcast(bInActive); }
+
 	// IFilter implementation
 	DECLARE_DERIVED_EVENT( FFrontendFilter, IFilter<FAssetFilterType>::FChangedEvent, FChangedEvent );
 	virtual FChangedEvent& OnChanged() override { return ChangedEvent; }
@@ -70,5 +73,11 @@ protected:
 private:
 	FChangedEvent ChangedEvent;
 
+	/** This event is broadcast to set this filter active in the content browser it is being used in */
+	DECLARE_EVENT_OneParam(FFrontendFilter, FSetActiveEvent, bool);
+	FSetActiveEvent SetActiveEvent;
+
 	TSharedPtr<FFrontendFilterCategory> FilterCategory;
+
+	friend struct FFrontendFilterExternalActivationHelper;
 };
diff --git a/Engine/Source/Editor/DetailCustomizations/Private/DetailCustomizations.cpp b/Engine/Source/Editor/DetailCustomizations/Private/DetailCustomizations.cpp
--- a/Engine/Source/Editor/DetailCustomizations/Private/DetailCustomizations.cpp
+++ b/Engine/Source/Editor/DetailCustomizations/Private/DetailCustomizations.cpp
@@ -55,7 +55,6 @@
 #include "BlackboardEntryDetails.h"
 #include "AIDataProviderValueDetails.h"
 #include "EnvQueryParamInstanceCustomization.h"
-#include "SkeletonNotifyDetails.h"
 #include "Customizations/ColorStructCustomization.h"
 #include "SlateColorCustomization.h"
 #include "CurveStructCustomization.h"
@@ -317,7 +316,6 @@ void FDetailCustomizationsModule::RegisterObjectCustomizations()
 
 	RegisterCustomClassLayout("EditorAnimSegment", FOnGetDetailCustomizationInstance::CreateStatic(&FAnimMontageSegmentDetails::MakeInstance));
 	RegisterCustomClassLayout("EditorAnimCompositeSegment", FOnGetDetailCustomizationInstance::CreateStatic(&FAnimMontageSegmentDetails::MakeInstance));
-	RegisterCustomClassLayout("EditorSkeletonNotifyObj", FOnGetDetailCustomizationInstance::CreateStatic(&FSkeletonNotifyDetails::MakeInstance));
 	RegisterCustomClassLayout("AnimStateNode", FOnGetDetailCustomizationInstance::CreateStatic(&FAnimStateNodeDetails::MakeInstance));
 	RegisterCustomClassLayout("AnimStateTransitionNode", FOnGetDetailCustomizationInstance::CreateStatic(&FAnimTransitionNodeDetails::MakeInstance));
 	RegisterCustomClassLayout("AnimGraphNode_Trail", FOnGetDetailCustomizationInstance::CreateStatic(&FAnimTrailNodeDetails::MakeInstance));
