Fix deprecation warning.
Remove duplicate GetActiveMode call.
#rnx

diff --git a/Engine/Plugins/Experimental/ControlRig/Source/ControlRigEditor/Private/Sequencer/ControlRigParameterTrackEditor.cpp b/Engine/Plugins/Experimental/ControlRig/Source/ControlRigEditor/Private/Sequencer/ControlRigParameterTrackEditor.cpp
--- a/Engine/Plugins/Experimental/ControlRig/Source/ControlRigEditor/Private/Sequencer/ControlRigParameterTrackEditor.cpp
+++ b/Engine/Plugins/Experimental/ControlRig/Source/ControlRigEditor/Private/Sequencer/ControlRigParameterTrackEditor.cpp
@@ -214,16 +214,14 @@ void FControlRigParameterTrackEditor::OnRelease()
 			}
 		}
 	}
-	FControlRigEditMode* ControlRigEditMode = static_cast<FControlRigEditMode*>(GLevelEditorModeTools().FindMode(FControlRigEditMode::ModeName));
+	FControlRigEditMode* ControlRigEditMode = static_cast<FControlRigEditMode*>(GLevelEditorModeTools().GetActiveMode(FControlRigEditMode::ModeName));
 	if (ControlRigEditMode)
 	{
-		if (GLevelEditorModeTools().GetActiveMode(FControlRigEditMode::ModeName))
+		if (GLevelEditorModeTools().HasToolkitHost())
 		{
-			if (GLevelEditorModeTools().HasToolkitHost())
-			{
-				GLevelEditorModeTools().DeactivateMode(FControlRigEditMode::ModeName);
-			}
+			GLevelEditorModeTools().DeactivateMode(FControlRigEditMode::ModeName);
 		}
+
 		ControlRigEditMode->SetObjects(nullptr, FGuid(), nullptr);
 	}
 
