Resolve additional deprecation/duplicate call
#rnx

diff --git a/Engine/Plugins/Experimental/ControlRig/Source/ControlRigEditor/Private/Sequencer/ControlRigParameterTrackEditor.cpp b/Engine/Plugins/Experimental/ControlRig/Source/ControlRigEditor/Private/Sequencer/ControlRigParameterTrackEditor.cpp
--- a/Engine/Plugins/Experimental/ControlRig/Source/ControlRigEditor/Private/Sequencer/ControlRigParameterTrackEditor.cpp
+++ b/Engine/Plugins/Experimental/ControlRig/Source/ControlRigEditor/Private/Sequencer/ControlRigParameterTrackEditor.cpp
@@ -572,7 +572,7 @@ void FControlRigParameterTrackEditor::OnAddTransformKeysForSelectedObjects(EMovi
 void FControlRigParameterTrackEditor::OnSequencerDataChanged(EMovieSceneDataChangeType DataChangeType)
 {
 	UMovieScene* MovieScene = GetSequencer()->GetFocusedMovieSceneSequence()->GetMovieScene();
-	FControlRigEditMode* ControlRigEditMode = static_cast<FControlRigEditMode*>(GLevelEditorModeTools().FindMode(FControlRigEditMode::ModeName));
+	FControlRigEditMode* ControlRigEditMode = static_cast<FControlRigEditMode*>(GLevelEditorModeTools().GetActiveMode(FControlRigEditMode::ModeName));
 	
 	//if we have a valid control rig edit mode need to check and see the control rig in that mode is still in a track
 	//if not we get rid of it.
@@ -589,12 +589,9 @@ void FControlRigParameterTrackEditor::OnSequencerDataChanged(EMovieSceneDataChan
 			}
 		}
 		//okay no good track so deactive it and delete it's Control Rig and bingings.
-		if (GLevelEditorModeTools().GetActiveMode(FControlRigEditMode::ModeName))
+		if (GLevelEditorModeTools().HasToolkitHost())
 		{
-			if (GLevelEditorModeTools().HasToolkitHost())
-			{
-				GLevelEditorModeTools().DeactivateMode(FControlRigEditMode::ModeName);
-			}
+			GLevelEditorModeTools().DeactivateMode(FControlRigEditMode::ModeName);
 		}
 		ControlRigEditMode->SetObjects(nullptr, FGuid(), nullptr);
 	}
