Merging //UE4/Main@9388041 to Dev-RenderPlat-Staging

diff --git a/Engine/Plugins/Experimental/ControlRig/Source/ControlRig/Public/AnimNode_ControlRig.h b/Engine/Plugins/Experimental/ControlRig/Source/ControlRig/Public/AnimNode_ControlRig.h
--- a/Engine/Plugins/Experimental/ControlRig/Source/ControlRig/Public/AnimNode_ControlRig.h
+++ b/Engine/Plugins/Experimental/ControlRig/Source/ControlRig/Public/AnimNode_ControlRig.h
@@ -3,6 +3,7 @@
 #pragma once
 
 #include "ControlRig.h"
+#include "Animation/InputScaleBias.h"
 #include "AnimNode_ControlRigBase.h"
 #include "AnimNode_ControlRig.generated.h"
 
@@ -33,9 +34,6 @@ struct CONTROLRIG_API FAnimNode_ControlRig : public FAnimNode_ControlRigBase
 	FName GetIOMapping(bool bInput, const FName& SourceProperty) const;
 private:
 
-	UPROPERTY(EditAnywhere, Category = Links)
-	FPoseLink Source;
-
 	/** Cached ControlRig */
 	UPROPERTY(EditAnywhere, Category = ControlRig)
 	TSubclassOf<UControlRig> ControlRigClass;
@@ -43,6 +41,28 @@ private:
 	UPROPERTY(transient)
 	UControlRig* ControlRig;
 
+	// alpha value handler
+	UPROPERTY(EditAnywhere, Category = Settings, meta = (PinShownByDefault))
+	float Alpha;
+
+	UPROPERTY(EditAnywhere, Category = Settings)
+	EAnimAlphaInputType AlphaInputType;
+
+	UPROPERTY(EditAnywhere, Category = Settings, meta = (PinShownByDefault, DisplayName = "bEnabled", DisplayAfter = "AlphaScaleBias"))
+	uint8 bAlphaBoolEnabled : 1;
+
+	UPROPERTY(EditAnywhere, Category=Settings)
+	FInputScaleBias AlphaScaleBias;
+
+	UPROPERTY(EditAnywhere, Category = Settings, meta = (DisplayName = "Blend Settings"))
+	FInputAlphaBoolBlend AlphaBoolBlend;
+
+	UPROPERTY(EditAnywhere, Category = Settings, meta = (PinShownByDefault))
+	FName AlphaCurveName;
+
+	UPROPERTY(EditAnywhere, Category = Settings)
+	FInputScaleBiasClamp AlphaScaleBiasClamp;
+
 	// we only save mapping, 
 	// we have to query control rig when runtime 
 	// to ensure type and everything is still valid or not
diff --git a/Engine/Plugins/Experimental/ControlRig/Source/ControlRig/Public/AnimNode_ControlRigBase.h b/Engine/Plugins/Experimental/ControlRig/Source/ControlRig/Public/AnimNode_ControlRigBase.h
--- a/Engine/Plugins/Experimental/ControlRig/Source/ControlRig/Public/AnimNode_ControlRigBase.h
+++ b/Engine/Plugins/Experimental/ControlRig/Source/ControlRig/Public/AnimNode_ControlRigBase.h
@@ -30,6 +30,10 @@ struct CONTROLRIG_API FAnimNode_ControlRigBase : public FAnimNode_CustomProperty
 	virtual void CacheBones_AnyThread(const FAnimationCacheBonesContext& Context) override;
 
 protected:
+
+	UPROPERTY(EditAnywhere, Category = Links)
+	FPoseLink Source;
+
 	/** Rig Hierarchy bone name to required array index mapping */
 	UPROPERTY(transient)
 	TMap<FName, uint16> ControlRigBoneMapping;
@@ -51,10 +55,16 @@ protected:
 	UPROPERTY(transient)
 	bool bExecute;
 
+	// The below is alpha value support for control rig
+	float InternalBlendAlpha;
+
 	// update input/output to control rig
 	virtual void UpdateInput(UControlRig* ControlRig, const FPoseContext& InOutput);
 	virtual void UpdateOutput(UControlRig* ControlRig, FPoseContext& InOutput);
 	virtual UClass* GetTargetClass() const override;
+	
+	// execute control rig on the input pose and outputs the result
+	void ExecuteControlRig(FPoseContext& InOutput);
 
 	friend struct FControlRigSequencerAnimInstanceProxy;
 };
diff --git a/Engine/Plugins/Experimental/ControlRig/Source/ControlRig/Public/AnimNode_ControlRig_ExternalSource.h b/Engine/Plugins/Experimental/ControlRig/Source/ControlRig/Public/AnimNode_ControlRig_ExternalSource.h
--- a/Engine/Plugins/Experimental/ControlRig/Source/ControlRig/Public/AnimNode_ControlRig_ExternalSource.h
+++ b/Engine/Plugins/Experimental/ControlRig/Source/ControlRig/Public/AnimNode_ControlRig_ExternalSource.h
@@ -18,15 +18,6 @@ struct CONTROLRIG_API FAnimNode_ControlRig_ExternalSource : public FAnimNode_Con
 	void SetControlRig(UControlRig* InControlRig);
 	virtual UControlRig* GetControlRig() const;
 
-	virtual void Initialize_AnyThread(const FAnimationInitializeContext& Context) override;
-	virtual void GatherDebugData(FNodeDebugData& DebugData) override;
-	virtual void Update_AnyThread(const FAnimationUpdateContext& Context) override;
-	virtual void CacheBones_AnyThread(const FAnimationCacheBonesContext& Context) override;
-	virtual void Evaluate_AnyThread(FPoseContext& Output) override;
-
-	UPROPERTY(EditAnywhere, Category = Links)
-	FPoseLink Source;
-
 private:
 	UPROPERTY(transient)
 	TWeakObjectPtr<UControlRig> ControlRig;
diff --git a/Engine/Plugins/Experimental/ControlRig/Source/ControlRig/Public/ControlRig.h b/Engine/Plugins/Experimental/ControlRig/Source/ControlRig/Public/ControlRig.h
--- a/Engine/Plugins/Experimental/ControlRig/Source/ControlRig/Public/ControlRig.h
+++ b/Engine/Plugins/Experimental/ControlRig/Source/ControlRig/Public/ControlRig.h
@@ -1,7 +1,7 @@
 // Copyright 1998-2019 Epic Games, Inc. All Rights Reserved.
 
 #pragma once
-
+#include "UObject/Object.h"
 #include "CoreMinimal.h"
 #include "UObject/ObjectMacros.h"
 #include "UObject/ScriptMacros.h"
@@ -69,7 +69,10 @@ public:
 	UControlRig();
 
 	virtual void Serialize(FArchive& Ar) override;
-
+#if WITH_EDITOR
+	virtual void PostEditUndo() override;
+#endif
+	
 	/** Get the current delta time */
 	UFUNCTION(BlueprintPure, Category = "Animation")
 	float GetDeltaTime() const;
@@ -95,17 +98,23 @@ public:
 	virtual void Evaluate_AnyThread();
 
 	/** Setup bindings to a runtime object (or clear by passing in nullptr). */
-	void SetObjectBinding(TSharedPtr<IControlRigObjectBinding> InObjectBinding)
+	virtual void SetObjectBinding(TSharedPtr<IControlRigObjectBinding> InObjectBinding) override
 	{
 		ObjectBinding = InObjectBinding;
 	}
 
 	/** Get bindings to a runtime object */
-	TSharedPtr<IControlRigObjectBinding> GetObjectBinding() const
+	virtual TSharedPtr<IControlRigObjectBinding> GetObjectBinding() const override
 	{
 		return ObjectBinding;
 	}
-
+    /** Get OurName*/
+	virtual FString GetName() const override
+	{
+		FString ObjectName = (GetClass()->GetName());
+		ObjectName.RemoveFromEnd(TEXT("_C"));
+		return ObjectName;
+	}
 	FRigHierarchyContainer* GetHierarchy()
 	{
 		return &Hierarchy;
@@ -216,6 +225,12 @@ public:
 	virtual FRigControlValue GetControlValueFromGlobalTransform(const FName& InControlName, const FTransform& InGlobalTransform) override;
 	virtual bool SetControlSpace(const FName& InControlName, const FName& InSpaceName) override;
 	virtual UControlRigGizmoLibrary* GetGizmoLibrary() const override;
+#if WITH_EDITOR
+	virtual void SelectControl(const FName& InControlName, bool bSelect = true) override;
+	virtual bool ClearControlSelection() override;
+	virtual TArray<FName> CurrentControlSelection() const override;
+	virtual bool IsControlSelected(const FName& InControlName)const override;
+#endif
 	// END IControlRigManipulatable interface
 
 	DECLARE_EVENT_TwoParams(UControlRig, FControlRigExecuteEvent, class UControlRig*, const EControlRigState);
@@ -296,6 +311,10 @@ private:
 	/** Broadcasts a notification whenever the controlrig is executed / updated. */
 	FControlRigExecuteEvent ExecutedEvent;
 
+#if WITH_EDITOR
+	/** Handle a Control Being Selected */
+	void HandleOnControlSelected(FRigHierarchyContainer* InContainer, const FRigElementKey& InKey, bool bSelected);
+#endif
 	void ResolveInputOutputProperties();
 
 	void InitializeFromCDO();
diff --git a/Engine/Plugins/Experimental/ControlRig/Source/ControlRig/Public/Manipulatable/IControlRigManipulatable.h b/Engine/Plugins/Experimental/ControlRig/Source/ControlRig/Public/Manipulatable/IControlRigManipulatable.h
--- a/Engine/Plugins/Experimental/ControlRig/Source/ControlRig/Public/Manipulatable/IControlRigManipulatable.h
+++ b/Engine/Plugins/Experimental/ControlRig/Source/ControlRig/Public/Manipulatable/IControlRigManipulatable.h
@@ -16,6 +16,9 @@
  * and implements a series of setter functions to perform changes.
  *
  */
+
+class IControlRigObjectBinding;
+
 class CONTROLRIG_API IControlRigManipulatable
 {
 public:
@@ -27,6 +30,10 @@ public:
 
 	/** Bindable event for external objects to be notified of Control changes */
 	DECLARE_EVENT_TwoParams(IControlRigManipulatable, FControlModifiedEvent, IControlRigManipulatable*, const FRigControl&);
+#if WITH_EDITOR
+	/** Bindable event for external objects to be notified that a Control is Selected */
+	DECLARE_EVENT_ThreeParams(IControlRigManipulatable, FControlSelectedEvent, IControlRigManipulatable*, const FRigControl&,  bool);
+#endif
 
 	// Returns true if this manipulatable subject is currently
 	// available for manipulation / is enabled.
@@ -90,7 +97,7 @@ public:
 
 	// Sets the relative value of a Control
 	template<class T>
-	FORCEINLINE bool SetControlValue(const FName& InControlName, T InValue, bool bNotify = true)
+	FORCEINLINE void SetControlValue(const FName& InControlName, T InValue, bool bNotify = true)
 	{
 		SetControlValue(InControlName, FRigControlValue::Make<T>(InValue), bNotify);
 	}
@@ -108,22 +115,52 @@ public:
 	// Sets a Control's Space (for space switching), returns true when successful.
 	virtual bool SetControlSpace(const FName& InControlName, const FName& InSpaceName) = 0;
 
-	// Returns a event than can be used to subscribe to
+	// Returns a event that can be used to subscribe to
 	// filtering control data when needed
 	FFilterControlEvent& ControlFilter() { return OnFilterControl;  }
 	
-	// Returns a event than can be used to subscribe to
-	// notifications coming from the manipulated subject.
+	// Returns a event that can be used to subscribe to
+	// change notifications coming from the manipulated subject.
 	FControlModifiedEvent& ControlModified() { return OnControlModified;  }
 
+#if WITH_EDITOR
+	// Returns a event that can be used to subscribe to
+	// selection changes coming from the manipulated subject.
+	FControlSelectedEvent& ControlSelected() { return OnControlSelected; }
+
+	//Select or deselected the specified control  
+	virtual void SelectControl(const FName& InControlName, bool bSelect = true) = 0;
+
+	//Clear selection on all controls
+	virtual bool ClearControlSelection() = 0;
+
+	//Get the current selection
+	virtual TArray<FName> CurrentControlSelection() const = 0;
+
+	//Is the Specified Control Selected
+	virtual bool IsControlSelected(const FName& InControlName) const = 0;
+
+#endif
+	//Get Name
+	virtual FString GetName() const = 0;
+
 	// Returns the gizmo library used for generating gizmos
 	virtual UControlRigGizmoLibrary* GetGizmoLibrary() const { return nullptr; }
 
+	//Set Object Binding
+	virtual void SetObjectBinding(TSharedPtr<IControlRigObjectBinding> InObjectBinding) {};
+
+	//Get bindings to a runtime object 
+	virtual TSharedPtr<IControlRigObjectBinding> GetObjectBinding() const {	return nullptr;}
+
+
 private:
 
 	FFilterControlEvent OnFilterControl;
 	FControlModifiedEvent OnControlModified;
-
+#if WITH_EDITOR
+	FControlSelectedEvent OnControlSelected;
+#endif
 	/** True if currently manipulation is enabled */
 	bool bManipulationEnabled;
 };
\ No newline at end of file
diff --git a/Engine/Plugins/Experimental/ControlRig/Source/ControlRig/Public/Math/ControlRigMathLibrary.h b/Engine/Plugins/Experimental/ControlRig/Source/ControlRig/Public/Math/ControlRigMathLibrary.h
--- a/Engine/Plugins/Experimental/ControlRig/Source/ControlRig/Public/Math/ControlRigMathLibrary.h
+++ b/Engine/Plugins/Experimental/ControlRig/Source/ControlRig/Public/Math/ControlRigMathLibrary.h
@@ -3,6 +3,7 @@
 #pragma once
 
 #include "CoreMinimal.h"
+#include "ControlRigDefines.h"
 #include "ControlRigMathLibrary.generated.h"
 
 UENUM()
@@ -85,5 +86,6 @@ public:
 	static void FourPointBezier(const FCRFourPointBezier& Bezier, float T, FVector& OutPosition, FVector& OutTangent);
 	static float EaseFloat(float Value, EControlRigAnimEasingType Type);
 	static FTransform LerpTransform(const FTransform& A, const FTransform& B, float T);
-	static void SolveBasicTwoBoneIK(FTransform& BoneA, FTransform& BoneB, FTransform& Effector, const FVector& PoleVector, const FVector& PrimaryAxis, const FVector& SecondaryAxis, float BoneALength, float BoneBLength, bool bEnableStretch, float StretchStartRatio, float StretchMaxRatio);
+	static void SolveBasicTwoBoneIK(FTransform& BoneA, FTransform& BoneB, FTransform& Effector, const FVector& PoleVector, const FVector& PrimaryAxis, const FVector& SecondaryAxis, float SecondaryAxisWeight, float BoneALength, float BoneBLength, bool bEnableStretch, float StretchStartRatio, float StretchMaxRatio);
+	static FVector ClampSpatially(const FVector& Value, EAxis::Type Axis, EControlRigClampSpatialMode::Type Type, float Minimum, float Maximum, FTransform Space);
 };
\ No newline at end of file
diff --git a/Engine/Plugins/Experimental/ControlRig/Source/ControlRig/Public/Rigs/RigBoneHierarchy.h b/Engine/Plugins/Experimental/ControlRig/Source/ControlRig/Public/Rigs/RigBoneHierarchy.h
--- a/Engine/Plugins/Experimental/ControlRig/Source/ControlRig/Public/Rigs/RigBoneHierarchy.h
+++ b/Engine/Plugins/Experimental/ControlRig/Source/ControlRig/Public/Rigs/RigBoneHierarchy.h
@@ -4,10 +4,18 @@
 
 #include "CoreMinimal.h"
 #include "RigHierarchyDefines.h"
+#include "ReferenceSkeleton.h"
 #include "RigBoneHierarchy.generated.h"
 
 class UControlRig;
 
+UENUM()
+enum class ERigBoneType : uint8
+{
+	Imported,
+	User
+};
+
 USTRUCT(BlueprintType)
 struct CONTROLRIG_API FRigBone: public FRigElement
 {
@@ -21,6 +29,7 @@ struct CONTROLRIG_API FRigBone: public FRigElement
 		, GlobalTransform(FTransform::Identity)
 		, LocalTransform(FTransform::Identity)
 		, Dependents()
+		, Type(ERigBoneType::Imported)
 	{
 	}
 	virtual ~FRigBone() {}
@@ -45,6 +54,10 @@ struct CONTROLRIG_API FRigBone: public FRigElement
 	UPROPERTY(transient)
 	TArray<int32> Dependents;
 
+	/** the source of the bone to differentiate procedurally generated, imported etc */
+	UPROPERTY(VisibleAnywhere, Category = FRigElement)
+	ERigBoneType Type;
+
 	FORCEINLINE virtual ERigElementType GetElementType() const override
 	{
 		return ERigElementType::Bone;
@@ -81,9 +94,9 @@ struct CONTROLRIG_API FRigBoneHierarchy
 
 	FName GetSafeNewName(const FName& InPotentialNewName) const;
 
-	FRigBone& Add(const FName& InNewName, const FName& InParentName = NAME_None, const FTransform& InInitTransform = FTransform::Identity);
+	FRigBone& Add(const FName& InNewName, const FName& InParentName = NAME_None, ERigBoneType InType = ERigBoneType::Imported, const FTransform& InInitTransform = FTransform::Identity);
 
-	FRigBone& Add(const FName& InNewName, const FName& InParentName, const FTransform& InInitTransform, const FTransform& InLocalTransform, const FTransform& InGlobalTransform);
+	FRigBone& Add(const FName& InNewName, const FName& InParentName, ERigBoneType InType, const FTransform& InInitTransform, const FTransform& InLocalTransform, const FTransform& InGlobalTransform);
 
 	FRigBone Remove(const FName& InNameToRemove);
 
@@ -154,6 +167,8 @@ struct CONTROLRIG_API FRigBoneHierarchy
 	TArray<FName> CurrentSelection() const;
 	bool IsSelected(const FName& InName) const;
 
+	TArray<FRigElementKey> ImportSkeleton(const FReferenceSkeleton& InSkeleton, const FName& InNameSpace, bool bReplaceExistingBones, bool bRemoveObsoleteBones, bool bSelectBones);
+
 	FRigElementAdded OnBoneAdded;
 	FRigElementRemoved OnBoneRemoved;
 	FRigElementRenamed OnBoneRenamed;
@@ -186,6 +201,7 @@ private:
 	void RecalculateLocalTransform(FRigBone& InOutBone);
 	void RecalculateGlobalTransform(FRigBone& InOutBone);
 
+	void RefreshParentNames();
 	void RefreshMapping();
 	void Sort();
 
diff --git a/Engine/Plugins/Experimental/ControlRig/Source/ControlRig/Public/Sequencer/MovieSceneControlRigParameterSection.h b/Engine/Plugins/Experimental/ControlRig/Source/ControlRig/Public/Sequencer/MovieSceneControlRigParameterSection.h
new file mode 100644
--- /dev/null
+++ b/Engine/Plugins/Experimental/ControlRig/Source/ControlRig/Public/Sequencer/MovieSceneControlRigParameterSection.h
@@ -0,0 +1,191 @@
+// Copyright 1998-2019 Epic Games, Inc. All Rights Reserved.
+
+#pragma once
+
+#include "CoreMinimal.h"
+#include "Sections/MovieSceneParameterSection.h"
+#include "UObject/ObjectMacros.h"
+#include "Channels/MovieSceneFloatChannel.h"
+#include "Sections/MovieSceneSubSection.h"
+#include "ControlRig.h"
+#include "MovieSceneSequencePlayer.h"
+#include "Animation/AnimData/BoneMaskFilter.h"
+#include "MovieSceneObjectBindingID.h"
+#include "MovieSceneControlRigParameterSection.generated.h"
+
+
+class UMovieSceneControlRigSection;
+/**
+*  Data that's queried during an interrogtion
+*/
+struct FFloatInterrogationData
+{
+	float Val;
+	FName ParameterName;
+};
+
+
+struct FVectorInterrogationData
+{
+	FVector Val;
+	FName ParameterName;
+};
+
+struct FTransformInterrogationData
+{
+	FTransform Val;
+	FName ParameterName;
+};
+
+
+/**
+ * Movie scene section that controls animation controller animation
+ */
+UCLASS()
+class CONTROLRIG_API UMovieSceneControlRigParameterSection : public UMovieSceneParameterSection
+{
+	GENERATED_BODY()
+
+public:
+
+	/** Control Rig that controls us*/
+	UPROPERTY()
+	UControlRig* ControlRig;
+
+	/** Mask for controls themselves*/
+	UPROPERTY()
+	TArray<bool> ControlsMask;
+
+	/** Mask for Transform Mask*/
+	UPROPERTY()
+	FMovieSceneTransformMask TransformMask;
+
+	/** Blend this track in additively (using the reference pose as a base) */
+	UPROPERTY(EditAnywhere, Category = "Animation")
+	bool bAdditive;
+
+	/** Only apply bones that are in the filter */
+	UPROPERTY(EditAnywhere, Category = "Animation")
+	bool bApplyBoneFilter;
+
+	/** Per-bone filter to apply to our animation */
+	UPROPERTY(EditAnywhere, Category = "Animation", meta=(EditCondition=bApplyBoneFilter))
+	FInputBlendPose BoneFilter;
+
+	/** The weight curve for this animation controller section */
+	UPROPERTY()
+	FMovieSceneFloatChannel Weight;
+
+public:
+
+	UMovieSceneControlRigParameterSection();
+
+	const TArray<bool>& GetControlsMask() const
+	{
+		return ControlsMask;
+	}
+
+	bool GetControlsMask(int32 Index) const
+	{
+		if (Index >= 0 && Index < ControlsMask.Num())
+		{
+			return ControlsMask[Index];
+		}
+		return false;
+	}
+
+	void SetControlsMask(const TArray<bool>& InMask)
+	{
+		ControlsMask = InMask;
+		ReconstructChannelProxy();
+	}
+
+	void SetControlsMask(int32 Index, bool Val)
+	{
+		if (Index >= 0 && Index < ControlsMask.Num())
+		{
+			ControlsMask[Index] = Val;
+		}
+		ReconstructChannelProxy();
+	}
+
+
+	void FillControlsMask(bool Val)
+	{
+		ControlsMask.Init(Val, ControlsMask.Num());
+		ReconstructChannelProxy();
+
+	}
+
+	/**
+	* Access the transform mask that defines which channels this track should animate
+	*/
+	FMovieSceneTransformMask GetTransformMask() const
+	{
+		return TransformMask;
+	}
+
+	/**
+	 * Set the transform mask that defines which channels this track should animate
+	 */
+	void SetTransformMask(FMovieSceneTransformMask NewMask)
+	{
+		TransformMask = NewMask;
+		ReconstructChannelProxy();
+	}
+
+public:
+
+	/** Whether or not to key currently, maybe evaluating so don't*/
+	void  SetDoNotKey(bool bIn) const { bDoNotKey = bIn; }
+	/** Get Whether to key or not*/
+	bool GetDoNotKey() const { return bDoNotKey; }
+
+	/**  Whether or not this section his scalar*/
+	bool HasScalarParameter(FName InParameterName) const;
+
+	/**  Whether or not this section his scalar*/
+	bool HasVectorParameter(FName InParameterName) const;
+
+	/**  Whether or not this section his scalar*/
+	bool HasColorParameter(FName InParameterName) const;
+
+		/**  Whether or not this section his scalar*/
+	bool HasTransformParameter(FName InParameterName) const;
+
+	/** Adds specified scalar parameter. */
+	void AddScalarParameter(FName InParameterName,  TOptional<float> DefaultValue);
+
+	/** Adds a a key for a specific vector parameter. */
+	void AddVectorParameter(FName InParameterName, TOptional<FVector> DefaultValue);
+
+	/** Adds a a key for a specific color parameter. */
+	void AddColorParameter(FName InParameterName, TOptional<FLinearColor> DefaultValue);
+
+	/** Adds a a key for a specific color parameter*/
+	void AddTransformParameter(FName InParameterName, TOptional<FTransform> DefaultValue);
+
+public:
+	/**
+	* Access the interrogation key for control rig data 
+	*/
+	 static FMovieSceneInterrogationKey GetFloatInterrogationKey();
+	 static FMovieSceneInterrogationKey GetVectorInterrogationKey();
+	 static FMovieSceneInterrogationKey GetTransformInterrogationKey();
+
+
+protected:
+	virtual void ReconstructChannelProxy();
+
+protected:
+
+	//~ UMovieSceneSection interface
+	virtual void Serialize(FArchive& Ar) override;
+	virtual void PostEditImport() override;
+
+	// When true we do not set a key on the section, since it will be set because we changed the value
+	// We need this because control rig notifications are set on every change even when just changing sequencer time
+	// which forces a sequencer eval, not like the edito where changes are only set on UI changes(changing time doesn't send change delegate)
+	mutable bool bDoNotKey;
+
+};
diff --git a/Engine/Plugins/Experimental/ControlRig/Source/ControlRigDeveloper/Private/AnimGraphNode_ControlRig.cpp b/Engine/Plugins/Experimental/ControlRig/Source/ControlRigDeveloper/Private/AnimGraphNode_ControlRig.cpp
--- a/Engine/Plugins/Experimental/ControlRig/Source/ControlRigDeveloper/Private/AnimGraphNode_ControlRig.cpp
+++ b/Engine/Plugins/Experimental/ControlRig/Source/ControlRigDeveloper/Private/AnimGraphNode_ControlRig.cpp
@@ -11,6 +11,7 @@
 #include "SVariableMappingWidget.h"
 #include "ScopedTransaction.h"
 #include "Kismet2/BlueprintEditorUtils.h"
+#include "ControlRigBlueprint.h"
 
 #define LOCTEXT_NAMESPACE "AnimGraphNode_ControlRig"
 
@@ -70,6 +71,41 @@ void UAnimGraphNode_ControlRig::ReallocatePinsDuringReconstruction(TArray<UEdGra
 	Super::ReallocatePinsDuringReconstruction(OldPins);
 }
 
+void UAnimGraphNode_ControlRig::ValidateAnimNodeDuringCompilation(USkeleton* ForSkeleton, FCompilerResultsLog& MessageLog)
+{
+	if (UClass* TargetClass = GetTargetClass())
+	{
+		if (UControlRigBlueprint* Blueprint = Cast<UControlRigBlueprint>(TargetClass->ClassGeneratedBy))
+		{
+			const FRigBoneHierarchy& BoneHierarchy = Blueprint->HierarchyContainer.BoneHierarchy;
+			const FReferenceSkeleton& ReferenceSkeleton = ForSkeleton->GetReferenceSkeleton();
+			const TArray<FMeshBoneInfo>& BoneInfos = ReferenceSkeleton.GetRefBoneInfo();
+
+			for (const FMeshBoneInfo& BoneInfo : BoneInfos)
+			{
+				int32 BoneIndex = BoneHierarchy.GetIndex(BoneInfo.Name);
+				if (BoneIndex != INDEX_NONE)
+				{
+					FName DesiredParentName = NAME_None;
+					if (BoneInfo.ParentIndex != INDEX_NONE)
+					{
+						DesiredParentName = BoneInfos[BoneInfo.ParentIndex].Name;
+					}
+
+					const FRigBone& Bone = BoneHierarchy[BoneIndex];
+					if (DesiredParentName != Bone.ParentName)
+					{
+						FString Message = FString::Printf(TEXT("@@ - Hierarchy discrepancy for bone '%s' - different parents on Control Rig vs SkeletalMesh."), *BoneInfo.Name.ToString());
+						MessageLog.Warning(*Message, this);
+					}
+				}
+			}
+		}
+	}
+
+	Super::ValidateAnimNodeDuringCompilation(ForSkeleton, MessageLog);
+}
+
 void UAnimGraphNode_ControlRig::RebuildExposedProperties()
 {
 	DECLARE_SCOPE_HIERARCHICAL_COUNTER_FUNC()
@@ -190,6 +226,7 @@ void UAnimGraphNode_ControlRig::CustomizeDetails(IDetailLayoutBuilder& DetailBui
 		return;
 	}
 
+	// input/output exposure feature START
 	RebuildExposedProperties();
 
 	IDetailCategoryBuilder& InputCategoryBuilder = DetailBuilder.EditCategory(FName(TEXT("Input")));
@@ -229,6 +266,35 @@ void UAnimGraphNode_ControlRig::CustomizeDetails(IDetailLayoutBuilder& DetailBui
 	{
 		ClassHandle->SetOnPropertyValueChanged(FSimpleDelegate::CreateUObject(this, &UAnimGraphNode_ControlRig::OnInstanceClassChanged, &DetailBuilder));
 	}
+
+	// input/output exposure feature END
+
+	// alpha property blending support START
+	TSharedRef<IPropertyHandle> NodeHandle = DetailBuilder.GetProperty(FName(TEXT("Node")), GetClass());
+
+	if (Node.AlphaInputType != EAnimAlphaInputType::Bool)
+	{
+		DetailBuilder.HideProperty(NodeHandle->GetChildHandle(GET_MEMBER_NAME_CHECKED(FAnimNode_ControlRig, bAlphaBoolEnabled)));
+		DetailBuilder.HideProperty(NodeHandle->GetChildHandle(GET_MEMBER_NAME_CHECKED(FAnimNode_ControlRig, AlphaBoolBlend)));
+	}
+
+	if (Node.AlphaInputType != EAnimAlphaInputType::Float)
+	{
+		DetailBuilder.HideProperty(NodeHandle->GetChildHandle(GET_MEMBER_NAME_CHECKED(FAnimNode_ControlRig, Alpha)));
+		DetailBuilder.HideProperty(NodeHandle->GetChildHandle(GET_MEMBER_NAME_CHECKED(FAnimNode_ControlRig, AlphaScaleBias)));
+	}
+
+	if (Node.AlphaInputType != EAnimAlphaInputType::Curve)
+	{
+		DetailBuilder.HideProperty(NodeHandle->GetChildHandle(GET_MEMBER_NAME_CHECKED(FAnimNode_ControlRig, AlphaCurveName)));
+	}
+
+	if ((Node.AlphaInputType != EAnimAlphaInputType::Float)
+		&& (Node.AlphaInputType != EAnimAlphaInputType::Curve))
+	{
+		DetailBuilder.HideProperty(NodeHandle->GetChildHandle(GET_MEMBER_NAME_CHECKED(FAnimNode_ControlRig, AlphaScaleBiasClamp)));
+	}
+	// alpha property blending support END
 }
 
 void UAnimGraphNode_ControlRig::GetIOProperties(bool bInput, TMap<FName, FControlRigIOVariable>& OutVars) const
@@ -363,12 +429,77 @@ void UAnimGraphNode_ControlRig::PostEditChangeProperty(FPropertyChangedEvent& Pr
 			bRequiresNodeReconstruct = true;
 			RebuildExposedProperties();
 		}
+
+		if (ChangedProperty->GetFName() == GET_MEMBER_NAME_STRING_CHECKED(FAnimNode_ControlRig, AlphaInputType))
+		{
+			FScopedTransaction Transaction(LOCTEXT("ChangeAlphaInputType", "Change Alpha Input Type"));
+			Modify();
+
+			// Break links to pins going away
+			for (int32 PinIndex = 0; PinIndex < Pins.Num(); ++PinIndex)
+			{
+				UEdGraphPin* Pin = Pins[PinIndex];
+				if (Pin->PinName == GET_MEMBER_NAME_STRING_CHECKED(FAnimNode_ControlRig, Alpha))
+				{
+					if (Node.AlphaInputType != EAnimAlphaInputType::Float)
+					{
+						Pin->BreakAllPinLinks();
+					}
+				}
+				else if (Pin->PinName == GET_MEMBER_NAME_STRING_CHECKED(FAnimNode_ControlRig, bAlphaBoolEnabled))
+				{
+					if (Node.AlphaInputType != EAnimAlphaInputType::Bool)
+					{
+						Pin->BreakAllPinLinks();
+					}
+				}
+				else if (Pin->PinName == GET_MEMBER_NAME_STRING_CHECKED(FAnimNode_ControlRig, AlphaCurveName))
+				{
+					if (Node.AlphaInputType != EAnimAlphaInputType::Curve)
+					{
+						Pin->BreakAllPinLinks();
+					}
+				}
+			}
+
+			bRequiresNodeReconstruct = true;
+		}
 	}
 
 	if (bRequiresNodeReconstruct)
 	{
 		ReconstructNode();
+		FBlueprintEditorUtils::MarkBlueprintAsStructurallyModified(GetBlueprint());
 	}
 }
 
+void UAnimGraphNode_ControlRig::CustomizePinData(UEdGraphPin* Pin, FName SourcePropertyName, int32 ArrayIndex) const
+{
+	Super::CustomizePinData(Pin, SourcePropertyName, ArrayIndex);
+
+	if (Pin->PinName == GET_MEMBER_NAME_STRING_CHECKED(FAnimNode_ControlRig, Alpha))
+	{
+		Pin->bHidden = (Node.AlphaInputType != EAnimAlphaInputType::Float);
+
+		if (!Pin->bHidden)
+		{
+			Pin->PinFriendlyName = Node.AlphaScaleBias.GetFriendlyName(Node.AlphaScaleBiasClamp.GetFriendlyName(Pin->PinFriendlyName));
+		}
+	}
+
+	if (Pin->PinName == GET_MEMBER_NAME_STRING_CHECKED(FAnimNode_ControlRig, bAlphaBoolEnabled))
+	{
+		Pin->bHidden = (Node.AlphaInputType != EAnimAlphaInputType::Bool);
+	}
+
+	if (Pin->PinName == GET_MEMBER_NAME_STRING_CHECKED(FAnimNode_ControlRig, AlphaCurveName))
+	{
+		Pin->bHidden = (Node.AlphaInputType != EAnimAlphaInputType::Curve);
+
+		if (!Pin->bHidden)
+		{
+			Pin->PinFriendlyName = Node.AlphaScaleBiasClamp.GetFriendlyName(Pin->PinFriendlyName);
+		}
+	}
+}
 #undef LOCTEXT_NAMESPACE
diff --git a/Engine/Plugins/Experimental/ControlRig/Source/ControlRigDeveloper/Private/ControlRigModel.cpp b/Engine/Plugins/Experimental/ControlRig/Source/ControlRigDeveloper/Private/ControlRigModel.cpp
--- a/Engine/Plugins/Experimental/ControlRig/Source/ControlRigDeveloper/Private/ControlRigModel.cpp
+++ b/Engine/Plugins/Experimental/ControlRig/Source/ControlRigDeveloper/Private/ControlRigModel.cpp
@@ -2555,6 +2555,14 @@ void UControlRigModel::ConfigurePinFromField(FControlRigModelPin& Pin, UProperty
 	{
 		Pin.CustomWidgetName = UControlRig::BoneNameMetaName;
 	}
+	else if (Property->HasMetaData(UControlRig::ControlNameMetaName))
+	{
+		Pin.CustomWidgetName = UControlRig::ControlNameMetaName;
+	}
+	else if (Property->HasMetaData(UControlRig::SpaceNameMetaName))
+	{
+		Pin.CustomWidgetName = UControlRig::SpaceNameMetaName;
+	}
 	else if (Property->HasMetaData(UControlRig::CurveNameMetaName))
 	{
 		Pin.CustomWidgetName = UControlRig::CurveNameMetaName;
diff --git a/Engine/Plugins/Experimental/ControlRig/Source/ControlRigEditor/Private/ControlRigEditorModule.cpp b/Engine/Plugins/Experimental/ControlRig/Source/ControlRigEditor/Private/ControlRigEditorModule.cpp
--- a/Engine/Plugins/Experimental/ControlRig/Source/ControlRigEditor/Private/ControlRigEditorModule.cpp
+++ b/Engine/Plugins/Experimental/ControlRig/Source/ControlRigEditor/Private/ControlRigEditorModule.cpp
@@ -12,9 +12,7 @@
 #include "GraphEditorActions.h"
 #include "Kismet2/KismetEditorUtilities.h"
 #include "ISequencerModule.h"
-#include "ControlRigTrackEditor.h"
 #include "IAssetTools.h"
-#include "ControlRigSequenceActions.h"
 #include "ControlRigEditorStyle.h"
 #include "Framework/Docking/LayoutExtender.h"
 #include "Framework/Application/SlateApplication.h"
@@ -25,16 +23,12 @@
 #include "WorkflowOrientedApp/WorkflowTabManager.h"
 #include "Modules/ModuleManager.h"
 #include "Widgets/Docking/SDockTab.h"
-#include "Sequencer/ControlRigSequence.h"
 #include "EditorModeRegistry.h"
 #include "ControlRigEditMode.h"
-#include "ControlRigEditorObjectBinding.h"
-#include "ControlRigEditorObjectSpawner.h"
 #include "ILevelSequenceModule.h"
-#include "ControlRigBindingTrackEditor.h"
 #include "EditorModeManager.h"
 #include "ControlRigEditMode.h"
-#include "Sequencer/MovieSceneControlRigSection.h"
+#include "Sequencer/MovieSceneControlRigParameterSection.h"
 #include "MovieSceneControlRigSectionDetailsCustomization.h"
 #include "ControlRigEditModeCommands.h"
 #include "Materials/Material.h"
@@ -77,6 +71,7 @@
 #include "WorkspaceMenuStructure.h"
 #include "WorkspaceMenuStructureModule.h"
 #include "Subsystems/AssetEditorSubsystem.h"
+#include "ControlRigParameterTrackEditor.h"
 
 #define LOCTEXT_NAMESPACE "ControlRigEditorModule"
 
@@ -101,10 +96,6 @@ void FControlRigEditorModule::StartupModule()
 	FControlRigStackCommands::Register();
 	FControlRigEditorStyle::Get();
 
-	CommandBindings = MakeShareable(new FUICommandList());
-
-	BindCommands();
-
 	MenuExtensibilityManager = MakeShareable(new FExtensibilityManager);
 	ToolBarExtensibilityManager = MakeShareable(new FExtensibilityManager);
 
@@ -119,7 +110,7 @@ void FControlRigEditorModule::StartupModule()
 	FPropertyEditorModule& PropertyEditorModule = FModuleManager::LoadModuleChecked<FPropertyEditorModule>("PropertyEditor");
 	ClassesToUnregisterOnShutdown.Reset();
 
-	ClassesToUnregisterOnShutdown.Add(UMovieSceneControlRigSection::StaticClass()->GetFName());
+	ClassesToUnregisterOnShutdown.Add(UMovieSceneControlRigParameterSection::StaticClass()->GetFName());
 	PropertyEditorModule.RegisterCustomClassLayout(ClassesToUnregisterOnShutdown.Last(), FOnGetDetailCustomizationInstance::CreateStatic(&FMovieSceneControlRigSectionDetailsCustomization::MakeInstance));
 
 	ClassesToUnregisterOnShutdown.Add(UControlRigSequenceExporterSettings::StaticClass()->GetFName());
@@ -140,7 +131,6 @@ void FControlRigEditorModule::StartupModule()
 	// same as ClassesToUnregisterOnShutdown but for properties, there is none right now
 	PropertiesToUnregisterOnShutdown.Reset();
 
-
 	// Register asset tools
 	auto RegisterAssetTypeAction = [this](const TSharedRef<IAssetTypeActions>& InAssetTypeAction)
 	{
@@ -149,41 +139,12 @@ void FControlRigEditorModule::StartupModule()
 		AssetTools.RegisterAssetTypeActions(InAssetTypeAction);
 	};
 
-	RegisterAssetTypeAction(MakeShareable(new FControlRigSequenceActions()));
 	RegisterAssetTypeAction(MakeShareable(new FControlRigBlueprintActions()));
 	RegisterAssetTypeAction(MakeShareable(new FControlRigGizmoLibraryActions()));
 	
 	// Register sequencer track editor
 	ISequencerModule& SequencerModule = FModuleManager::Get().LoadModuleChecked<ISequencerModule>("Sequencer");
-	SequencerCreatedHandle = SequencerModule.RegisterOnSequencerCreated(FOnSequencerCreated::FDelegate::CreateRaw(this, &FControlRigEditorModule::HandleSequencerCreated));
-	ControlRigTrackCreateEditorHandle = SequencerModule.RegisterTrackEditor(FOnCreateTrackEditor::CreateStatic(&FControlRigTrackEditor::CreateTrackEditor));
-	ControlRigBindingTrackCreateEditorHandle = SequencerModule.RegisterTrackEditor(FOnCreateTrackEditor::CreateStatic(&FControlRigBindingTrackEditor::CreateTrackEditor));
-	ControlRigEditorObjectBindingHandle = SequencerModule.RegisterEditorObjectBinding(FOnCreateEditorObjectBinding::CreateStatic(&FControlRigEditorObjectBinding::CreateEditorObjectBinding));
-
-	SequencerToolbarExtender = MakeShareable(new FExtender());
-	SequencerToolbarExtender->AddToolBarExtension(
-		"Level Sequence Separator",
-		EExtensionHook::Before,
-		CommandBindings,
-		FToolBarExtensionDelegate::CreateLambda([](FToolBarBuilder& ToolBarBuilder)
-		{
-			ToolBarBuilder.AddToolBarButton(FControlRigEditModeCommands::Get().ExportAnimSequence);
-		}));
-
-	SequencerModule.GetToolBarExtensibilityManager()->AddExtender(SequencerToolbarExtender);
-
-	// Register for assets being opened
-	AssetEditorOpenedHandle = GEditor->GetEditorSubsystem<UAssetEditorSubsystem>()->OnAssetEditorOpened().AddRaw(this, &FControlRigEditorModule::HandleAssetEditorOpened);
-
-	// Register level sequence spawner
-	ILevelSequenceModule& LevelSequenceModule = FModuleManager::LoadModuleChecked<ILevelSequenceModule>("LevelSequence");
-	LevelSequenceSpawnerDelegateHandle = LevelSequenceModule.RegisterObjectSpawner(FOnCreateMovieSceneObjectSpawner::CreateStatic(&FControlRigEditorObjectSpawner::CreateObjectSpawner));
-
-	TrajectoryMaterial = LoadObject<UMaterial>(nullptr, TEXT("/ControlRig/M_Traj.M_Traj"));
-	if (TrajectoryMaterial.IsValid())
-	{
-		TrajectoryMaterial->AddToRoot();
-	}
+	ControlRigParameterTrackCreateEditorHandle = SequencerModule.RegisterTrackEditor(FOnCreateTrackEditor::CreateStatic(&FControlRigParameterTrackEditor::CreateTrackEditor));
 
 	FEditorModeRegistry::Get().RegisterMode<FControlRigEditMode>(
 		FControlRigEditMode::ModeName,
@@ -197,126 +158,6 @@ void FControlRigEditorModule::StartupModule()
 		FSlateIcon(FControlRigEditorStyle::Get().GetStyleSetName(), "ControlRigEditMode", "ControlRigEditMode.Small"),
 		false);
 
-	FContentBrowserModule& ContentBrowserModule = FModuleManager::LoadModuleChecked<FContentBrowserModule>(TEXT("ContentBrowser"));
-	ContentBrowserModule.GetAllAssetViewContextMenuExtenders().Add(FContentBrowserMenuExtender_SelectedAssets::CreateLambda([this](const TArray<FAssetData>& SelectedAssets)
-	{
-		TSharedRef<FExtender> Extender = MakeShared<FExtender>();
-
-		if (SelectedAssets.ContainsByPredicate([](const FAssetData& AssetData) { return AssetData.GetClass() == UAnimSequence::StaticClass(); }))
-		{
-			Extender->AddMenuExtension(
-				"GetAssetActions",
-				EExtensionHook::After,
-				CommandBindings,
-				FMenuExtensionDelegate::CreateLambda([this, SelectedAssets](FMenuBuilder& MenuBuilder)
-				{
-					const TSharedPtr<FUICommandInfo>& ImportFromRigSequence = FControlRigEditModeCommands::Get().ImportFromRigSequence;
-					MenuBuilder.AddMenuEntry(
-						ImportFromRigSequence->GetLabel(),
-						ImportFromRigSequence->GetDescription(),
-						ImportFromRigSequence->GetIcon(),
-						FUIAction(FExecuteAction::CreateRaw(this, &FControlRigEditorModule::ImportFromRigSequence, SelectedAssets)));
-				}));
-
-			// only add this if we find a control rig sequence targeting this anim sequence in the asset registry
-			FAssetRegistryModule& AssetRegistryModule = FModuleManager::LoadModuleChecked<FAssetRegistryModule>("AssetRegistry");
-			
-			bool bCanReimport = false;
-			if (SelectedAssets.Num() > 0)
-			{
-				// It's faster to find all assets with this tag and then query them against the selection then it is to 
-				// query the asset registry each time for a tag with a particular value
-				const FName LastExportedToAnimationSequenceTagName = GET_MEMBER_NAME_CHECKED(UControlRigSequence, LastExportedToAnimationSequence);
-				TArray<FAssetData> FoundAssets;
-				{
-					TArray<FName> Tags;
-					Tags.Add(LastExportedToAnimationSequenceTagName);
-					AssetRegistryModule.Get().GetAssetsByTags(Tags, FoundAssets);
-				}
-
-				if (FoundAssets.Num() > 0)
-				{
-					for(const FAssetData& AssetData : SelectedAssets)
-					{
-						const bool bFoundAsset = FoundAssets.ContainsByPredicate([&AssetData, LastExportedToAnimationSequenceTagName](const FAssetData& FoundAsset)
-						{
-							const FName TagValue = FoundAsset.GetTagValueRef<FName>(LastExportedToAnimationSequenceTagName);
-							return TagValue == AssetData.ObjectPath;
-						});
-
-						if (bFoundAsset)
-						{
-							bCanReimport = true;
-							break;
-						}
-					}
-				}
-			}
-
-			if (bCanReimport)
-			{
-				Extender->AddMenuExtension(
-					"GetAssetActions",
-					EExtensionHook::After,
-					CommandBindings,
-					FMenuExtensionDelegate::CreateLambda([this, SelectedAssets](FMenuBuilder& MenuBuilder)
-					{
-						const TSharedPtr<FUICommandInfo>& ReImportFromRigSequence = FControlRigEditModeCommands::Get().ReImportFromRigSequence;
-						MenuBuilder.AddMenuEntry(
-							ReImportFromRigSequence->GetLabel(),
-							ReImportFromRigSequence->GetDescription(),
-							ReImportFromRigSequence->GetIcon(),
-							FUIAction(FExecuteAction::CreateRaw(this, &FControlRigEditorModule::ReImportFromRigSequence, SelectedAssets)));
-					}));
-			}
-		}
-		else if (SelectedAssets.ContainsByPredicate([](const FAssetData& AssetData) { return AssetData.GetClass() == UControlRigSequence::StaticClass(); }))
-		{
-			Extender->AddMenuExtension(
-				"CommonAssetActions",
-				EExtensionHook::Before,
-				CommandBindings,
-				FMenuExtensionDelegate::CreateLambda([this, SelectedAssets](FMenuBuilder& MenuBuilder)
-				{
-					MenuBuilder.BeginSection("ControlRigActions", LOCTEXT("ControlRigActions", "Control Rig Sequence Actions"));
-					{
-						const TSharedPtr<FUICommandInfo>& ExportAnimSequence = FControlRigEditModeCommands::Get().ExportAnimSequence;
-						MenuBuilder.AddMenuEntry(
-							ExportAnimSequence->GetLabel(),
-							ExportAnimSequence->GetDescription(),
-							ExportAnimSequence->GetIcon(),
-							FUIAction(FExecuteAction::CreateRaw(this, &FControlRigEditorModule::ExportToAnimSequence, SelectedAssets)));
-
-						bool bCanReExport = false;
-						for (const FAssetData& AssetData : SelectedAssets)
-						{
-							if(UControlRigSequence* ControlRigSequence = Cast<UControlRigSequence>(AssetData.GetAsset()))
-							{
-								if (ControlRigSequence->LastExportedToAnimationSequence.IsValid())
-								{
-									bCanReExport = true;
-									break;
-								}
-							}
-						}
-
-						if (bCanReExport)
-						{
-							const TSharedPtr<FUICommandInfo>& ReExportAnimSequence = FControlRigEditModeCommands::Get().ReExportAnimSequence;
-							MenuBuilder.AddMenuEntry(
-								ReExportAnimSequence->GetLabel(),
-								ReExportAnimSequence->GetDescription(),
-								ReExportAnimSequence->GetIcon(),
-								FUIAction(FExecuteAction::CreateRaw(this, &FControlRigEditorModule::ReExportToAnimSequence, SelectedAssets)));
-						}
-					}
-					MenuBuilder.EndSection();
-				}));
-		}
-	
-		return Extender;
-	}));
-	ContentBrowserMenuExtenderHandle = ContentBrowserModule.GetAllAssetViewContextMenuExtenders().Last().GetHandle();
 
 	ControlRigGraphPanelNodeFactory = MakeShared<FControlRigGraphPanelNodeFactory>();
 	FEdGraphUtilities::RegisterVisualNodeFactory(ControlRigGraphPanelNodeFactory);
@@ -359,44 +200,14 @@ void FControlRigEditorModule::ShutdownModule()
 	FEdGraphUtilities::UnregisterVisualPinFactory(ControlRigGraphPanelPinFactory);
 	FEdGraphUtilities::UnregisterVisualNodeFactory(ControlRigGraphPanelNodeFactory);
 
-	FContentBrowserModule* ContentBrowserModule = FModuleManager::GetModulePtr<FContentBrowserModule>(TEXT("ContentBrowser"));
-	if (ContentBrowserModule)
-	{
-		ContentBrowserModule->GetAllAssetViewContextMenuExtenders().RemoveAll([=](const FContentBrowserMenuExtender_SelectedAssets& InDelegate) { return ContentBrowserMenuExtenderHandle == InDelegate.GetHandle(); });
-	}
-
-	if (TrajectoryMaterial.IsValid())
-	{
-		TrajectoryMaterial->RemoveFromRoot();
-	}
-
-	if (GEditor)
-	{
-		if (UAssetEditorSubsystem* AssetEditorSubsystem = GEditor->GetEditorSubsystem<UAssetEditorSubsystem>())
-		{
-			AssetEditorSubsystem->OnAssetEditorOpened().Remove(AssetEditorOpenedHandle);
-		}
-	}
-
 	FEditorModeRegistry::Get().UnregisterMode(FControlRigEditorEditMode::ModeName);
 	FEditorModeRegistry::Get().UnregisterMode(FControlRigEditMode::ModeName);
 
-	ILevelSequenceModule* LevelSequenceModule = FModuleManager::GetModulePtr<ILevelSequenceModule>("LevelSequence");
-	if (LevelSequenceModule)
-	{
-		LevelSequenceModule->UnregisterObjectSpawner(LevelSequenceSpawnerDelegateHandle);
-	}
 
 	ISequencerModule* SequencerModule = FModuleManager::GetModulePtr<ISequencerModule>("Sequencer");
 	if (SequencerModule)
 	{
-		SequencerModule->UnregisterOnSequencerCreated(SequencerCreatedHandle);
-		SequencerModule->UnRegisterTrackEditor(ControlRigTrackCreateEditorHandle);
-		SequencerModule->UnRegisterTrackEditor(ControlRigBindingTrackCreateEditorHandle);
-		SequencerModule->UnRegisterEditorObjectBinding(ControlRigEditorObjectBindingHandle);
-
-		SequencerModule->GetToolBarExtensibilityManager()->RemoveExtender(SequencerToolbarExtender);
-		SequencerToolbarExtender = nullptr;
+		SequencerModule->UnRegisterTrackEditor(ControlRigParameterTrackCreateEditorHandle);
 	}
 
 	FAssetToolsModule* AssetToolsModule = FModuleManager::GetModulePtr<FAssetToolsModule>("AssetTools");
@@ -429,8 +240,6 @@ void FControlRigEditorModule::ShutdownModule()
 			PropertyEditorModule->UnregisterCustomPropertyTypeLayout(PropertiesToUnregisterOnShutdown[Index]);
 		}
 	}
-
-	CommandBindings = nullptr;
 }
 
 void FControlRigEditorModule::HandleNewBlueprintCreated(UBlueprint* InBlueprint)
@@ -444,310 +253,6 @@ void FControlRigEditorModule::HandleNewBlueprintCreated(UBlueprint* InBlueprint)
 	InBlueprint->LastEditedDocuments.AddUnique(ControlRigGraph);
 }
 
-void FControlRigEditorModule::HandleSequencerCreated(TSharedRef<ISequencer> InSequencer)
-{
-	TWeakPtr<ISequencer> LocalSequencer = InSequencer;
-
-	// Record the last sequencer we opened that was editing a control rig sequence
-	UMovieSceneSequence* FocusedSequence = InSequencer->GetFocusedMovieSceneSequence();
-	if (UControlRigSequence* FocusedControlRigSequence = ExactCast<UControlRigSequence>(FocusedSequence))
-	{
-		WeakSequencer = InSequencer;
-	}
-
-	// We want to be informed of sequence activations (subsequences or not)
-	auto HandleActivateSequence = [this, LocalSequencer](FMovieSceneSequenceIDRef Ref)
-	{
-		if (LocalSequencer.IsValid())
-		{
-			TSharedRef<ISequencer> Sequencer = LocalSequencer.Pin().ToSharedRef();
-			UMovieSceneSequence* Sequence = Sequencer->GetFocusedMovieSceneSequence();
-			if (UControlRigSequence* ControlRigSequence = ExactCast<UControlRigSequence>(Sequence))
-			{
-				WeakSequencer = LocalSequencer;
-
-				GLevelEditorModeTools().ActivateMode(FControlRigEditMode::ModeName);
-
-				if (FControlRigEditMode* ControlRigEditMode = static_cast<FControlRigEditMode*>(GLevelEditorModeTools().GetActiveMode(FControlRigEditMode::ModeName)))
-				{
-					ControlRigEditMode->SetSequencer(Sequencer);
-				}
-			}
-			else
-			{
-				if (FControlRigEditMode* ControlRigEditMode = static_cast<FControlRigEditMode*>(GLevelEditorModeTools().GetActiveMode(FControlRigEditMode::ModeName)))
-				{
-					ControlRigEditMode->SetSequencer(nullptr);
-					ControlRigEditMode->SetObjects(TWeakObjectPtr<>(), FGuid(), nullptr);
-				}
-			}
-		}
-	};
-
-	InSequencer->OnActivateSequence().AddLambda(HandleActivateSequence);
-
-	// Call into activation callback to handle initial activation
-	FMovieSceneSequenceID SequenceID = MovieSceneSequenceID::Root;
-	HandleActivateSequence(SequenceID);
-
-	InSequencer->GetSelectionChangedObjectGuids().AddLambda([LocalSequencer](TArray<FGuid> InObjectBindings)
-	{
-		if (LocalSequencer.IsValid())
-		{
-			TSharedRef<ISequencer> Sequencer = LocalSequencer.Pin().ToSharedRef();
-			UMovieSceneSequence* Sequence = Sequencer->GetFocusedMovieSceneSequence();
-			if (UControlRigSequence* ControlRigSequence = ExactCast<UControlRigSequence>(Sequence))
-			{
-				TWeakObjectPtr<> SelectedObject;
-				FGuid ObjectBinding;
-				if(InObjectBindings.Num() > 0)
-				{
-					ObjectBinding = InObjectBindings[0];
-					TArrayView<TWeakObjectPtr<>> BoundObjects = Sequencer->FindBoundObjects(ObjectBinding, Sequencer->GetFocusedTemplateID());
-					if(BoundObjects.Num() > 0)
-					{
-						SelectedObject = BoundObjects[0];
-					}
-				}
-
-				if (SelectedObject.IsValid())
-				{
-					GLevelEditorModeTools().ActivateMode(FControlRigEditMode::ModeName);
-					if (FControlRigEditMode* ControlRigEditMode = static_cast<FControlRigEditMode*>(GLevelEditorModeTools().GetActiveMode(FControlRigEditMode::ModeName)))
-					{
-						ControlRigEditMode->SetObjects(SelectedObject, ObjectBinding, nullptr);
-					}
-				}
-			}
-		}
-	});
-
-	InSequencer->OnMovieSceneDataChanged().AddLambda([LocalSequencer](EMovieSceneDataChangeType DataChangeType)
-	{
-		if (LocalSequencer.IsValid())
-		{
-			TSharedRef<ISequencer> Sequencer = LocalSequencer.Pin().ToSharedRef();
-			UMovieSceneSequence* Sequence = Sequencer->GetFocusedMovieSceneSequence();
-			if (UControlRigSequence* ControlRigSequence = ExactCast<UControlRigSequence>(Sequence))
-			{
-				if (FControlRigEditMode* ControlRigEditMode = static_cast<FControlRigEditMode*>(GLevelEditorModeTools().GetActiveMode(FControlRigEditMode::ModeName)))
-				{
-					ControlRigEditMode->RefreshObjects();
-				}
-			}
-		}
-	});
-
-	InSequencer->GetSelectionChangedTracks().AddLambda([LocalSequencer](TArray<UMovieSceneTrack*> InTracks)
-	{
-		if (LocalSequencer.IsValid())
-		{
-			TSharedRef<ISequencer> Sequencer = LocalSequencer.Pin().ToSharedRef();
-			UMovieSceneSequence* Sequence = Sequencer->GetFocusedMovieSceneSequence();
-			if (UControlRigSequence* ControlRigSequence = ExactCast<UControlRigSequence>(Sequence))
-			{
-				TArray<FString> PropertyPaths;
-
-				// Look for any property tracks that might drive our rig manipulators
-				for (UMovieSceneTrack* Track : InTracks)
-				{
-					if (UMovieScenePropertyTrack* PropertyTrack = Cast<UMovieScenePropertyTrack>(Track))
-					{
-						PropertyPaths.Add(PropertyTrack->GetPropertyPath());
-					}
-				}
-
-				if (FControlRigEditMode* ControlRigEditMode = static_cast<FControlRigEditMode*>(GLevelEditorModeTools().GetActiveMode(FControlRigEditMode::ModeName)))
-				{
-					ControlRigEditMode->ClearRigElementSelection(FRigElementTypeHelper::ToMask(ERigElementType::All));
-					for (int32 Index = 0; Index < PropertyPaths.Num(); ++Index)
-					{
-						ControlRigEditMode->SetRigElementSelection(ERigElementType::Control, FName(*PropertyPaths[Index], FNAME_Find), true);
-					}
-					
-				}
-			}
-		}
-	});
-
-	InSequencer->OnPostSave().AddLambda([](ISequencer& InSequencerThatSaved)
-	{
-		UMovieSceneSequence* Sequence = InSequencerThatSaved.GetFocusedMovieSceneSequence();
-		if (UControlRigSequence* ControlRigSequence = ExactCast<UControlRigSequence>(Sequence))
-		{
-			if (FControlRigEditMode* ControlRigEditMode = static_cast<FControlRigEditMode*>(GLevelEditorModeTools().GetActiveMode(FControlRigEditMode::ModeName)))
-			{
-				ControlRigEditMode->ReBindToActor();
-			}
-		}
-	});
-
-	InSequencer->OnGetIsTrackVisible().BindRaw(this, &FControlRigEditorModule::IsTrackVisible);
-}
-
-void FControlRigEditorModule::HandleAssetEditorOpened(UObject* InAsset)
-{
-	if (UControlRigSequence* ControlRigSequence = ExactCast<UControlRigSequence>(InAsset))
-	{
-		GLevelEditorModeTools().ActivateMode(FControlRigEditMode::ModeName);
-
-		if (FControlRigEditMode* ControlRigEditMode = static_cast<FControlRigEditMode*>(GLevelEditorModeTools().GetActiveMode(FControlRigEditMode::ModeName)))
-		{
-			ControlRigEditMode->ReBindToActor();
-		}
-	}
-}
-
-void FControlRigEditorModule::OnInitializeSequence(UControlRigSequence* Sequence)
-{
-	auto* ProjectSettings = GetDefault<UMovieSceneToolsProjectSettings>();
-	UMovieScene* MovieScene = Sequence->GetMovieScene();
-	
-	FFrameNumber StartFrame = (ProjectSettings->DefaultStartTime * MovieScene->GetTickResolution()).RoundToFrame();
-	int32        Duration   = (ProjectSettings->DefaultDuration  * MovieScene->GetTickResolution()).RoundToFrame().Value;
-
-	MovieScene->SetPlaybackRange(StartFrame, Duration);
-}
-
-void FControlRigEditorModule::BindCommands()
-{
-	const FControlRigEditModeCommands& Commands = FControlRigEditModeCommands::Get();
-
-	CommandBindings->MapAction(
-		Commands.ExportAnimSequence,
-		FExecuteAction::CreateRaw(this, &FControlRigEditorModule::ExportAnimSequenceFromSequencer),
-		FCanExecuteAction(),
-		FGetActionCheckState(), 
-		FIsActionButtonVisible::CreateRaw(this, &FControlRigEditorModule::CanExportAnimSequenceFromSequencer));
-}
-
-bool FControlRigEditorModule::CanExportAnimSequenceFromSequencer() const
-{
-	if (WeakSequencer.IsValid())
-	{
-		TSharedRef<ISequencer> Sequencer = WeakSequencer.Pin().ToSharedRef();
-		return ExactCast<UControlRigSequence>(Sequencer->GetFocusedMovieSceneSequence()) != nullptr;
-	}
-
-	return false;
-}
-
-void FControlRigEditorModule::ExportAnimSequenceFromSequencer()
-{
-	// if we have an active sequencer, get the sequence
-	UControlRigSequence* ControlRigSequence = nullptr;
-	if (WeakSequencer.IsValid())
-	{
-		TSharedRef<ISequencer> Sequencer = WeakSequencer.Pin().ToSharedRef();
-		ControlRigSequence = ExactCast<UControlRigSequence>(Sequencer->GetFocusedMovieSceneSequence());
-	}
-
-	// If we are bound to an actor in the edit mode, auto pick skeletal mesh to use for binding
-	USkeletalMesh* SkeletalMesh = nullptr;
-	if (FControlRigEditMode* ControlRigEditMode = static_cast<FControlRigEditMode*>(GLevelEditorModeTools().GetActiveMode(FControlRigEditMode::ModeName)))
-	{
-		TLazyObjectPtr<AActor> ActorPtr = ControlRigEditMode->GetSettings()->Actor;
-		if (ActorPtr.IsValid())
-		{
-			if (USkeletalMeshComponent* SkeletalMeshComponent = ActorPtr->FindComponentByClass<USkeletalMeshComponent>())
-			{
-				SkeletalMesh = SkeletalMeshComponent->SkeletalMesh;
-			}
-		}
-	}
-
-	if (ControlRigSequence)
-	{
-		ControlRigSequenceConverter::Convert(ControlRigSequence, nullptr, SkeletalMesh);
-	}
-}
-
-void FControlRigEditorModule::ExportToAnimSequence(TArray<FAssetData> InAssetData)
-{
-	for (const FAssetData& AssetData : InAssetData)
-	{
-		UControlRigSequence* ControlRigSequence = Cast<UControlRigSequence>(AssetData.GetAsset());
-		if (ControlRigSequence)
-		{
-			ControlRigSequenceConverter::Convert(ControlRigSequence, nullptr, nullptr);
-		}
-	}
-}
-
-void FControlRigEditorModule::ReExportToAnimSequence(TArray<FAssetData> InAssetData)
-{
-	for (const FAssetData& AssetData : InAssetData)
-	{
-		UControlRigSequence* ControlRigSequence = Cast<UControlRigSequence>(AssetData.GetAsset());
-		if (ControlRigSequence)
-		{
-			UAnimSequence* AnimSequence = ControlRigSequence->LastExportedToAnimationSequence.LoadSynchronous();
-			USkeletalMesh* SkeletalMesh = ControlRigSequence->LastExportedUsingSkeletalMesh.LoadSynchronous();
-			bool bShowDialog = AnimSequence == nullptr || SkeletalMesh == nullptr;
-
-			ControlRigSequenceConverter::Convert(ControlRigSequence, AnimSequence, SkeletalMesh, bShowDialog);
-		}
-	}
-}
-
-void FControlRigEditorModule::ImportFromRigSequence(TArray<FAssetData> InAssetData)
-{
-	for (const FAssetData& AssetData : InAssetData)
-	{
-		UAnimSequence* AnimSequence = Cast<UAnimSequence>(AssetData.GetAsset());
-		if (AnimSequence)
-		{
-			ControlRigSequenceConverter::Convert(nullptr, AnimSequence, nullptr);
-		}
-	}
-}
-
-void FControlRigEditorModule::ReImportFromRigSequence(TArray<FAssetData> InAssetData)
-{
-	for (const FAssetData& AssetData : InAssetData)
-	{
-		UAnimSequence* AnimSequence = Cast<UAnimSequence>(AssetData.GetAsset());
-		USkeletalMesh* SkeletalMesh = nullptr;
-		UControlRigSequence* ControlRigSequence = nullptr;
-
-		FAssetRegistryModule& AssetRegistryModule = FModuleManager::LoadModuleChecked<FAssetRegistryModule>("AssetRegistry");
-
-		TMultiMap<FName, FString> TagsAndValues;
-		TagsAndValues.Add(GET_MEMBER_NAME_CHECKED(UControlRigSequence, LastExportedToAnimationSequence), AssetData.ObjectPath.ToString());
-
-		TArray<FAssetData> FoundAssets;
-		AssetRegistryModule.Get().GetAssetsByTagValues(TagsAndValues, FoundAssets);
-
-		if (FoundAssets.Num() > 0)
-		{
-			ControlRigSequence = Cast<UControlRigSequence>(FoundAssets[0].GetAsset());
-			if (ControlRigSequence)
-			{
-				SkeletalMesh = ControlRigSequence->LastExportedUsingSkeletalMesh.LoadSynchronous();
-			}
-		}
-
-		bool bShowDialog = ControlRigSequence == nullptr || AnimSequence == nullptr || SkeletalMesh == nullptr;
-
-		ControlRigSequenceConverter::Convert(ControlRigSequence, AnimSequence, SkeletalMesh, bShowDialog);
-	}
-}
-
-bool FControlRigEditorModule::IsTrackVisible(const UMovieSceneTrack* InTrack)
-{
-	if (FControlRigEditMode* ControlRigEditMode = static_cast<FControlRigEditMode*>(GLevelEditorModeTools().GetActiveMode(FControlRigEditMode::ModeName)))
-	{		
-		// If nothing selected, show all nodes
-		if (ControlRigEditMode->AreRigElementsSelected(FRigElementTypeHelper::ToMask(ERigElementType::Control)) == 0)
-		{
-			return true;
-		}
-
-		return ControlRigEditMode->SelectedRigElements.Contains(FRigElementKey(InTrack->GetTrackName(), ERigElementType::Control));
-	}
-
-	return true;
-}
 
 TSharedRef<IControlRigEditor> FControlRigEditorModule::CreateControlRigEditor(const EToolkitMode::Type Mode, const TSharedPtr< IToolkitHost >& InitToolkitHost, class UControlRigBlueprint* InBlueprint)
 {
diff --git a/Engine/Plugins/Experimental/ControlRig/Source/ControlRigEditor/Private/ControlRigEditorModule.h b/Engine/Plugins/Experimental/ControlRig/Source/ControlRigEditor/Private/ControlRigEditorModule.h
--- a/Engine/Plugins/Experimental/ControlRig/Source/ControlRigEditor/Private/ControlRigEditorModule.h
+++ b/Engine/Plugins/Experimental/ControlRig/Source/ControlRigEditor/Private/ControlRigEditorModule.h
@@ -9,8 +9,6 @@
 
 class UBlueprint;
 class IAssetTypeActions;
-class UControlRigSequence;
-class ISequencer;
 class UMaterial;
 class UAnimSequence;
 class USkeletalMesh;
@@ -32,15 +30,14 @@ public:
 	virtual TSharedRef<IControlRigEditor> CreateControlRigEditor(const EToolkitMode::Type Mode, const TSharedPtr<IToolkitHost>& InitToolkitHost, UControlRigBlueprint* Blueprint) override;
 	DECLARE_DELEGATE_RetVal_TwoParams(TSharedRef<FExtender>, FControlRigEditorToolbarExtender, const TSharedRef<FUICommandList> /*InCommandList*/, TSharedRef<IControlRigEditor> /*InControlRigEditor*/);
 	virtual TArray<FControlRigEditorToolbarExtender>& GetAllControlRigEditorToolbarExtenders() override { return ControlRigEditorToolbarExtenders; }
-
 	/** IHasMenuExtensibility interface */
 	virtual TSharedPtr<FExtensibilityManager> GetMenuExtensibilityManager() override { return MenuExtensibilityManager; }
 
 	/** IHasToolBarExtensibility interface */
 	virtual TSharedPtr<FExtensibilityManager> GetToolBarExtensibilityManager() override { return ToolBarExtensibilityManager; }
 
-	UMaterial* GetTrajectoryMaterial() const { return TrajectoryMaterial.Get(); }
 
+	
 	virtual void RegisterRigUnitEditorClass(FName RigUnitClassName, TSubclassOf<URigUnitEditor_Base> Class) override;
 	virtual void UnregisterRigUnitEditorClass(FName RigUnitClassName) override;
 	
@@ -56,74 +53,11 @@ private:
 	/** Handle a new animation controller blueprint being created */
 	void HandleNewBlueprintCreated(UBlueprint* InBlueprint);
 
-	/** Handle a new sequencer instance being created */
-	void HandleSequencerCreated(TSharedRef<ISequencer> InSequencer);
-
-	/** Handle an asset being opened */
-	void HandleAssetEditorOpened(UObject* InAsset);
-
-	/** Called to setup a new sequence's defaults */
-	static void OnInitializeSequence(UControlRigSequence* Sequence);
-
-	/** Whether we can export the current control rig sequence as an anim sequence */
-	bool CanExportAnimSequenceFromSequencer() const;
-
-	/** Export the current control rig sequence as an anim sequence */
-	void ExportAnimSequenceFromSequencer();
-
-	/** Export control rig sequence(s) to anim sequence(s) */
-	void ExportToAnimSequence(TArray<FAssetData> InAssetData);
-
-	/** Re-export control rig sequence(s) to anim sequence(s) using the previous export settings */
-	void ReExportToAnimSequence(TArray<FAssetData> InAssetData);
-
-	/** Import animation sequence(s) from a source rig sequence */
-	void ImportFromRigSequence(TArray<FAssetData> InAssetData);
-
-	/** Re-import animation sequence(s) from their source rig sequence(s) */
-	void ReImportFromRigSequence(TArray<FAssetData> InAssetData);
-
-	/** Bind our module-level commands */
-	void BindCommands();
-
-	/** Whether the track is visible in the sequencer node tree */
-	bool IsTrackVisible(const UMovieSceneTrack* InTrack);
-
-	/** Handle for our sequencer track editor */
-	FDelegateHandle ControlRigTrackCreateEditorHandle;
-
-	/** Handle for our sequencer binding track editor */
-	FDelegateHandle ControlRigBindingTrackCreateEditorHandle;
-
-	/** Handle for our sequencer object binding */
-	FDelegateHandle ControlRigEditorObjectBindingHandle;
-
-	/** Handle for our level sequence spawner */
-	FDelegateHandle LevelSequenceSpawnerDelegateHandle;
-
-	/** Handle for tracking ISequencer creation */
-	FDelegateHandle SequencerCreatedHandle;
-
-	/** Handle for tracking asset editors opening */
-	FDelegateHandle AssetEditorOpenedHandle;
+	/** Handle for our sequencer control rig parameter track editor */
+	FDelegateHandle ControlRigParameterTrackCreateEditorHandle;
 
 	TArray<TSharedRef<IAssetTypeActions>> RegisteredAssetTypeActions;
 
-	/** Material used for trajectories */
-	TWeakObjectPtr<UMaterial> TrajectoryMaterial;
-
-	/** Toolbar extender for Sequencer */
-	TSharedPtr<FExtender> SequencerToolbarExtender;
-
-	/** Command bindings for keyboard shortcuts */
-	TSharedPtr<FUICommandList> CommandBindings;
-
-	/** Weak pointer to the last sequencer that was opened */
-	TWeakPtr<ISequencer> WeakSequencer;
-
-	/** Delegate handle used to extend the content browser asset menu */
-	FDelegateHandle ContentBrowserMenuExtenderHandle;
-
 	/** StaticClass is not safe on shutdown, so we cache the name, and use this to unregister on shut down */
 	TArray<FName> ClassesToUnregisterOnShutdown;
 	TArray<FName> PropertiesToUnregisterOnShutdown;
diff --git a/Engine/Plugins/Experimental/ControlRig/Source/ControlRigEditor/Private/ControlRigEditorStyle.h b/Engine/Plugins/Experimental/ControlRig/Source/ControlRigEditor/Private/ControlRigEditorStyle.h
--- a/Engine/Plugins/Experimental/ControlRig/Source/ControlRigEditor/Private/ControlRigEditorStyle.h
+++ b/Engine/Plugins/Experimental/ControlRig/Source/ControlRigEditor/Private/ControlRigEditorStyle.h
@@ -82,7 +82,8 @@ public:
 
 		// Tree styles
 		{
-			Set("ControlRig.Tree.Bone", new IMAGE_BRUSH("Slate/Bone_16x", Icon16x16));
+			Set("ControlRig.Tree.BoneUser", new IMAGE_BRUSH("Slate/BoneUser_16x", Icon16x16));
+			Set("ControlRig.Tree.BoneImported", new IMAGE_BRUSH("Slate/BoneImported_16x", Icon16x16));
 			Set("ControlRig.Tree.Control", new IMAGE_BRUSH("Slate/Control_16x", Icon16x16));
 			Set("ControlRig.Tree.Space", new IMAGE_BRUSH("Slate/Space_16x", Icon16x16));
 		}
diff --git a/Engine/Plugins/Experimental/ControlRig/Source/ControlRigEditor/Private/EditMode/ControlRigEditMode.cpp b/Engine/Plugins/Experimental/ControlRig/Source/ControlRigEditor/Private/EditMode/ControlRigEditMode.cpp
--- a/Engine/Plugins/Experimental/ControlRig/Source/ControlRigEditor/Private/EditMode/ControlRigEditMode.cpp
+++ b/Engine/Plugins/Experimental/ControlRig/Source/ControlRigEditor/Private/EditMode/ControlRigEditMode.cpp
@@ -80,7 +80,7 @@ FControlRigEditMode::FControlRigEditMode()
 }
 
 FControlRigEditMode::~FControlRigEditMode()
-{
+{	
 	CommandBindings = nullptr;
 
 #if WITH_EDITOR
@@ -88,32 +88,6 @@ FControlRigEditMode::~FControlRigEditMode()
 #endif
 }
 
-void FControlRigEditMode::SetSequencer(TSharedPtr<ISequencer> InSequencer)
-{
-	static bool bRecursionGuard = false;
-	if (!bRecursionGuard)
-	{
-		TGuardValue<bool> ScopeGuard(bRecursionGuard, true);
-
-		Settings->Sequence = nullptr;
-
-		WeakSequencer = InSequencer;
-		if(UsesToolkits())
-		{
-			StaticCastSharedPtr<SControlRigEditModeTools>(Toolkit->GetInlineContent())->SetSequencer(InSequencer);
-		}
-
-		if (InSequencer.IsValid())
-		{
-			if (UControlRigSequence* Sequence = ExactCast<UControlRigSequence>(InSequencer->GetFocusedMovieSceneSequence()))
-			{
-				Settings->Sequence = Sequence;
-				ReBindToActor();
-			}
-		}
-	}
-}
-
 void FControlRigEditMode::SetObjects(const TWeakObjectPtr<>& InSelectedObject, const FGuid& InObjectBinding, UObject* BindingObject)
 {
 	WeakControlRigEditing = Cast<UControlRig>(InSelectedObject.Get());
@@ -146,93 +120,24 @@ void FControlRigEditMode::SetObjects_Internal()
 		SelectedObjects.Add(Settings);
 	}
 
-	if(WeakControlRigEditing.IsValid())
+	if (WeakControlRigEditing.IsValid())
 	{
-		SelectedObjects.Add(WeakControlRigEditing);
+		//Don't add the WeakControlRig Editing...SelectedObjects.Add(WeakControlRigEditing);
 		WeakControlRigEditing.Get()->DrawInterface = &DrawInterface;
-	}
-
-	if(UsesToolkits())
-	{
-		StaticCastSharedPtr<SControlRigEditModeTools>(Toolkit->GetInlineContent())->SetDetailsObjects(SelectedObjects);
-	}
-
-	// create default manipulation layer
-	RecreateManipulationLayer();
-	HandleSelectionChanged();
-}
-
-void FControlRigEditMode::HandleBindToActor(AActor* InActor, bool bFocus)
-{
-	static bool bRecursionGuard = false;
-	if (!bRecursionGuard)
-	{
-		TGuardValue<bool> ScopeGuard(bRecursionGuard, true);
-
-		if(IsInLevelEditor())
+		if (IsInLevelEditor())
 		{
-			FControlRigBindingTemplate::SetObjectBinding(InActor);
+			WeakControlRigEditing->Hierarchy.OnElementSelected.AddSP(this, &FControlRigEditMode::OnRigElementSelected);
 		}
 
-		if (WeakSequencer.IsValid())
-		{
-			TSharedRef<ISequencer> Sequencer = WeakSequencer.Pin().ToSharedRef();
-
-			// Modify the sequence
-			if (UControlRigSequence* Sequence = ExactCast<UControlRigSequence>(Sequencer->GetFocusedMovieSceneSequence()))
-			{
-				Sequence->Modify(false);
-
-				// Also modify the binding tracks in the sequence, so bindings get regenerated to this actor
-				UMovieScene* MovieScene = Sequence->GetMovieScene();
-				for (UMovieSceneSection* Section : MovieScene->GetAllSections())
-				{
-					if (UMovieSceneSpawnSection* SpawnSection = Cast<UMovieSceneSpawnSection>(Section))
-					{
-						SpawnSection->TryModify(false);
-					}
-				}
-
-				// now notify the sequence (will rebind when it re-evaluates
-				Sequencer->NotifyMovieSceneDataChanged(EMovieSceneDataChangeType::TrackValueChangedRefreshImmediately);
-
-				// we tick first
-				ManipulationLayer->TickManipulatableObjects(0.f);
-
-				// Now re-display our objects in the details panel (they may have changed)
-				if (MovieScene->GetSpawnableCount() > 0)
-				{
-					FGuid SpawnableGuid = MovieScene->GetSpawnable(0).GetGuid();
-					TWeakObjectPtr<> BoundObject = Sequencer->FindSpawnedObjectOrTemplate(SpawnableGuid);
-					SetObjects(BoundObject, SpawnableGuid, nullptr);
-				}
-			}
 
-			if (bFocus && InActor && IsInLevelEditor())
-			{
-				const bool bNotifySelectionChanged = false;
-				const bool bDeselectBSP = true;
-				const bool bWarnAboutTooManyActors = false;
-				const bool bSelectEvenIfHidden = true;
-
-				// Select & focus the actor
-				GEditor->GetSelectedActors()->Modify();
-				GEditor->GetSelectedActors()->BeginBatchSelectOperation();
-				GEditor->SelectNone(bNotifySelectionChanged, bDeselectBSP, bWarnAboutTooManyActors);
-				GEditor->SelectActor(InActor, true, bNotifySelectionChanged, bSelectEvenIfHidden);
-				GEditor->Exec(InActor->GetWorld(), TEXT("CAMERA ALIGN ACTIVEVIEWPORTONLY"));
-				GEditor->SelectNone(bNotifySelectionChanged, bDeselectBSP, bWarnAboutTooManyActors);
-				GEditor->GetSelectedActors()->EndBatchSelectOperation();
-			}
+		if (UsesToolkits())
+		{
+			StaticCastSharedPtr<SControlRigEditModeTools>(Toolkit->GetInlineContent())->SetDetailsObjects(SelectedObjects);
 		}
-	}
-}
 
-void FControlRigEditMode::ReBindToActor()
-{
-	if (Settings->Actor.IsValid())
-	{
-		HandleBindToActor(Settings->Actor.Get(), false);
+		// create default manipulation layer
+		RecreateManipulationLayer();
+		HandleSelectionChanged();
 	}
 }
 
@@ -284,10 +189,12 @@ void FControlRigEditMode::Exit()
 
 	if (ManipulationLayer)
 	{
+		SelectNone();
 		ManipulationLayer->DestroyLayer();
 		ManipulationLayer = nullptr;
 	}
-
+	//clear actors
+	GizmoActors.SetNum(0);
 	// Call parent implementation
 	FEdMode::Exit();
 }
@@ -298,7 +205,10 @@ void FControlRigEditMode::Tick(FEditorViewportClient* ViewportClient, float Delt
 
 	ViewportClient->Invalidate();
 
-	RecalcPivotTransform();
+	if (ManipulationLayer)
+	{
+		RecalcPivotTransform();
+	}
 
 	if (bRecreateManipulationLayerRequired)
 	{
@@ -319,37 +229,35 @@ void FControlRigEditMode::Tick(FEditorViewportClient* ViewportClient, float Delt
 		bRecreateManipulationLayerRequired = false;
 	}
 
-	// If we have detached from sequencer, unbind the settings UI
-	if (!WeakSequencer.IsValid() && Settings->Sequence != nullptr)
+	// We need to tick here since changing a bone for example
+	// might have changed the transform of the Control
+	if (ManipulationLayer)
 	{
-		Settings->Sequence = nullptr;
-		RefreshObjects();
+		ManipulationLayer->TickManipulatableObjects(DeltaTime);
 	}
 }
 
 void FControlRigEditMode::Render(const FSceneView* View, FViewport* Viewport, FPrimitiveDrawInterface* PDI)
-{
-	bool bRender = true;
-	if (WeakSequencer.IsValid())
-	{
-		bRender = WeakSequencer.Pin()->GetPlaybackStatus() != EMovieScenePlayerStatus::Playing || Settings->bShowManipulatorsDuringPlayback;
-	}
-
-	if (!WeakControlRigEditing.IsValid())
+{	
+	if (!WeakControlRigEditing.IsValid() || ManipulationLayer == nullptr)
 	{
 		DrawInterface.DrawInstructions.Reset();
 		return;
 	}
 
-	// Force off manipulators if hide flag is set
-	if (Settings->bHideManipulators)
-	{
-		bRender = false;
-	}
+	bool bRender = !Settings->bHideManipulators;
 
 	FTransform ComponentTransform = (ManipulationLayer)? ManipulationLayer->GetSkeletalMeshComponentTransform() : FTransform::Identity;
 	if (bRender)
 	{
+		for (AControlRigGizmoActor* Actor : GizmoActors)
+		{
+			//Actor->SetActorHiddenInGame(bIsHidden);
+			if (GIsEditor && Actor->GetWorld() != nullptr && !Actor->GetWorld()->IsPlayInEditor())
+			{
+				Actor->SetIsTemporarilyHiddenInEditor(false);
+			}
+		}
 		if (Settings->bDisplayHierarchy)
 		{
 			UControlRig* ControlRig = WeakControlRigEditing.Get();
@@ -387,17 +295,14 @@ void FControlRigEditMode::Render(const FSceneView* View, FViewport* Viewport, FP
 				PDI->DrawLine(ElementTransform.GetTranslation(), ElementTransform.TransformPosition(FVector(0.f, 0.f, Scale)), FLinearColor::Blue, SDPG_Foreground);
 			}
 		}
-	}
-
-
-	for (const FControlRigDrawInterface::FDrawIntruction& Instruction : DrawInterface.DrawInstructions)
-	{
-		if (Instruction.Positions.Num() == 0)
-		{
-			continue;
-		}
-		switch (Instruction.DrawType)
+		for (const FControlRigDrawInterface::FDrawIntruction& Instruction : DrawInterface.DrawInstructions)
 		{
+			if (Instruction.Positions.Num() == 0)
+			{
+				continue;
+			}
+			switch (Instruction.DrawType)
+			{
 			case FControlRigDrawInterface::EDrawType_Point:
 			{
 				for (const FVector& Point : Instruction.Positions)
@@ -411,7 +316,7 @@ void FControlRigEditMode::Render(const FSceneView* View, FViewport* Viewport, FP
 				const TArray<FVector>& Points = Instruction.Positions;
 				for (int32 PointIndex = 0; PointIndex < Points.Num() - 1; PointIndex += 2)
 				{
-					PDI->DrawLine(ComponentTransform.TransformPosition(Points[PointIndex]), ComponentTransform.TransformPosition(Points[PointIndex+1]), Instruction.Color, SDPG_Foreground, Instruction.Thickness);
+					PDI->DrawLine(ComponentTransform.TransformPosition(Points[PointIndex]), ComponentTransform.TransformPosition(Points[PointIndex + 1]), Instruction.Color, SDPG_Foreground, Instruction.Thickness);
 				}
 				break;
 			}
@@ -424,8 +329,23 @@ void FControlRigEditMode::Render(const FSceneView* View, FViewport* Viewport, FP
 				}
 				break;
 			}
+			}
 		}
 	}
+	else
+	{
+		for (AControlRigGizmoActor* Actor : GizmoActors)
+		{
+			//Actor->SetActorHiddenInGame(bIsHidden);
+			if (GIsEditor && Actor->GetWorld() != nullptr && !Actor->GetWorld()->IsPlayInEditor())
+			{
+				Actor->SetIsTemporarilyHiddenInEditor(true);
+			}
+		}
+	}
+
+	//draw debug info
+	
 	DrawInterface.DrawInstructions.Reset();
 }
 
@@ -479,7 +399,10 @@ bool FControlRigEditMode::StartTracking(FEditorViewportClient* InViewportClient,
 	{
 		GEditor->BeginTransaction(LOCTEXT("MoveControlTransaction", "Move Control"));
 
-		ManipulationLayer->BeginTransaction();
+		if (ManipulationLayer != nullptr)
+		{
+			ManipulationLayer->BeginTransaction();
+		}
 		for (AControlRigGizmoActor* GizmoActor : GizmoActors)
 		{
 			GizmoActor->SetManipulating(true);
@@ -516,7 +439,7 @@ bool FControlRigEditMode::UsesTransformWidget(FWidget::EWidgetMode CheckMode) co
 {
 	for (const AControlRigGizmoActor* GizmoActor : GizmoActors)
 	{
-		if (GizmoActor->IsSelected())
+		if (GizmoActor->IsSelected() && ManipulationLayer != nullptr)
 		{
 			return ManipulationLayer->ModeSupportedByGizmoActor(GizmoActor, CheckMode);
 		}
@@ -532,7 +455,7 @@ bool FControlRigEditMode::UsesTransformWidget(FWidget::EWidgetMode CheckMode) co
 
 FVector FControlRigEditMode::GetWidgetLocation() const
 {
-	if (AreRigElementSelectedAndMovable())
+	if (AreRigElementSelectedAndMovable() && ManipulationLayer != nullptr)
 	{
 		FTransform ComponentTransform = ManipulationLayer->GetSkeletalMeshComponentTransform();
 		return ComponentTransform.TransformPosition(PivotTransform.GetLocation());
@@ -566,7 +489,7 @@ bool FControlRigEditMode::HandleClick(FEditorViewportClient* InViewportClient, H
 			if (ActorHitProxy->Actor->IsA<AControlRigGizmoActor>())
 			{
 				AControlRigGizmoActor* GizmoActor = CastChecked<AControlRigGizmoActor>(ActorHitProxy->Actor);
-				const FControlData* ControlData = ManipulationLayer->GetControlDataFromGizmo(GizmoActor);
+				const FControlData* ControlData = ManipulationLayer != nullptr ? ManipulationLayer->GetControlDataFromGizmo(GizmoActor) :  nullptr;
 
 				if (ControlData)
 				{
@@ -580,18 +503,18 @@ bool FControlRigEditMode::HandleClick(FEditorViewportClient* InViewportClient, H
 						ClearRigElementSelection(FRigElementTypeHelper::ToMask(ERigElementType::Control));
 						SetRigElementSelection(ERigElementType::Control, ControlName, true);
 					}
+
+				}
+				// for now we show this menu all the time if body is selected
+				// if we want some global menu, we'll have to move this
+				if (Click.GetKey() == EKeys::RightMouseButton)
+				{
+					OpenContextMenu(InViewportClient);
 				}
-			}
-		}
 
-		// for now we show this menu all the time if body is selected
-		// if we want some global menu, we'll have to move this
-		if (Click.GetKey() == EKeys::RightMouseButton)
-		{
-			OpenContextMenu(InViewportClient);
+				return true;
+			}
 		}
-
-		return true;
 	}
 
 	// for now we show this menu all the time if body is selected
@@ -605,12 +528,6 @@ bool FControlRigEditMode::HandleClick(FEditorViewportClient* InViewportClient, H
 	// clear selected controls
 	ClearRigElementSelection(FRigElementTypeHelper::ToMask(ERigElementType::All));
 
-	// If we are animating then swallow clicks so we dont select things other than controls
-	if(WeakSequencer.IsValid() && WeakSequencer.Pin()->GetFocusedMovieSceneSequence()->IsA<UControlRigSequence>())
-	{
-		return true;
-	}
-
 	return FEdMode::HandleClick(InViewportClient, HitProxy, Click);
 }
 
@@ -647,24 +564,28 @@ void FControlRigEditMode::OpenContextMenu(FEditorViewportClient* InViewportClien
 
 bool FControlRigEditMode::IntersectSelect(bool InSelect, const TFunctionRef<bool(const AControlRigGizmoActor*, const FTransform&)>& Intersects)
 {
-	FTransform ComponentTransform = ManipulationLayer->GetSkeletalMeshComponentTransform();
-
-	bool bSelected = false;
-	for (AControlRigGizmoActor* GizmoActor: GizmoActors)
+	if (ManipulationLayer != nullptr)
 	{
-		const FTransform ControlTransform = GizmoActor->GetGlobalTransform() * ComponentTransform;
-		if (Intersects(GizmoActor, ControlTransform))
+		FTransform ComponentTransform = ManipulationLayer->GetSkeletalMeshComponentTransform();
+
+		bool bSelected = false;
+		for (AControlRigGizmoActor* GizmoActor : GizmoActors)
 		{
-			const FControlData* ControlData = ManipulationLayer->GetControlDataFromGizmo(GizmoActor);
-			if (ControlData)
+			const FTransform ControlTransform = GizmoActor->GetGlobalTransform() * ComponentTransform;
+			if (Intersects(GizmoActor, ControlTransform))
 			{
-				SetRigElementSelection(ERigElementType::Control, ControlData->ControlName, InSelect);
-				bSelected = true;
+				const FControlData* ControlData = ManipulationLayer->GetControlDataFromGizmo(GizmoActor);
+				if (ControlData)
+				{
+					SetRigElementSelection(ERigElementType::Control, ControlData->ControlName, InSelect);
+					bSelected = true;
+				}
 			}
 		}
-	}
 
-	return bSelected;
+		return bSelected;
+	}
+	return false;
 }
 
 bool FControlRigEditMode::BoxSelect(FBox& InBox, bool InSelect)
@@ -737,9 +658,9 @@ bool FControlRigEditMode::InputDelta(FEditorViewportClient* InViewportClient, FV
 		const bool bDoTranslation = !Drag.IsZero() && (WidgetMode == FWidget::WM_Translate || WidgetMode == FWidget::WM_TranslateRotateZ);
 		const bool bDoScale = !Scale.IsZero() && WidgetMode == FWidget::WM_Scale;
 
-		FTransform ComponentTransform = ManipulationLayer->GetSkeletalMeshComponentTransform();
-		if (AreRigElementsSelected(FRigElementTypeHelper::ToMask(ERigElementType::Control)))
+		if (ManipulationLayer != nullptr && AreRigElementsSelected(FRigElementTypeHelper::ToMask(ERigElementType::Control)))
 		{
+			FTransform ComponentTransform = ManipulationLayer->GetSkeletalMeshComponentTransform();
 			for (AControlRigGizmoActor* GizmoActor : GizmoActors)
 			{
 				if (GizmoActor->IsSelected())
@@ -758,8 +679,10 @@ bool FControlRigEditMode::InputDelta(FEditorViewportClient* InViewportClient, FV
 			}
 			return true;
 		}
-		else if (AreRigElementSelectedAndMovable())
+		else if (ManipulationLayer != nullptr && AreRigElementSelectedAndMovable())
 		{
+			FTransform ComponentTransform = ManipulationLayer->GetSkeletalMeshComponentTransform();
+
 			// set Bone transform
 			// that will set initial Bone transform
 			for (int32 Index = 0; Index < SelectedRigElements.Num(); ++Index)
@@ -831,12 +754,20 @@ bool FControlRigEditMode::IsCompatibleWith(FEditorModeID OtherModeID) const
 
 void FControlRigEditMode::AddReferencedObjects( FReferenceCollector& Collector )
 {
-	Collector.AddReferencedObject(Settings);
-	Collector.AddReferencedObject(ManipulationLayer);
-
-	for (AControlRigGizmoActor* GizmoActor : GizmoActors)
+	if (Settings)
 	{
-		Collector.AddReferencedObject(GizmoActor);
+		Collector.AddReferencedObject(Settings);
+	}
+	if (ManipulationLayer)
+	{
+		Collector.AddReferencedObject(ManipulationLayer);
+	}
+	if (GizmoActors.Num() > 0)
+	{
+		for (AControlRigGizmoActor* GizmoActor : GizmoActors)
+		{
+			Collector.AddReferencedObject(GizmoActor);
+		}
 	}
 }
 
@@ -853,6 +784,10 @@ void FControlRigEditMode::ClearRigElementSelection(uint32 InTypes)
 	{
 		Blueprint->HierarchyContainer.ClearSelection();
 	}
+	if(IsInLevelEditor())
+	{
+		WeakControlRigEditing->Hierarchy.ClearSelection();
+	}
 }
 
 // internal private function that doesn't use guarding.
@@ -868,6 +803,10 @@ void FControlRigEditMode::SetRigElementSelectionInternal(ERigElementType Type, c
 	{
 		Blueprint->HierarchyContainer.Select(FRigElementKey(InRigElementName, Type), bSelected);
 	}
+	if(IsInLevelEditor())
+	{
+		WeakControlRigEditing->Hierarchy.Select(FRigElementKey(InRigElementName, Type), bSelected);
+	}
 }
 
 void FControlRigEditMode::SetRigElementSelection(ERigElementType Type, const FName& InRigElementName, bool bSelected)
@@ -933,72 +872,13 @@ int32 FControlRigEditMode::GetNumSelectedRigElements(uint32 InTypes) const
 	return 0;
 }
 
-void FControlRigEditMode::HandleObjectSpawned(FGuid InObjectBinding, UObject* SpawnedObject, IMovieScenePlayer& Player)
-{
-	if (WeakSequencer.IsValid())
-	{
-		// check whether this spawned object is from our sequence
-		TSharedPtr<ISequencer> Sequencer = WeakSequencer.Pin();
-		if (Sequencer.Get() == &Player)
-		{
-			RefreshObjects();
-
-			// check if the object is being displayed currently
-			if (ControlRigGuid == InObjectBinding)
-			{
-				if (WeakControlRigEditing != Cast<UControlRig>(SpawnedObject))
-				{
-					WeakControlRigEditing = Cast<UControlRig>(SpawnedObject);
-					SetObjects_Internal();
-				}
-				return;
-			}
-
-			// We didnt find an existing Guid, so set up our internal cache
-			if (!ControlRigGuid.IsValid())
-			{
-				SetObjects(SpawnedObject, InObjectBinding, nullptr);
-				if (UControlRig* ControlRig = Cast<UControlRig>(SpawnedObject))
-				{
-					if (TSharedPtr<IControlRigObjectBinding> ObjectBinding = ControlRig->GetObjectBinding())
-					{
-						if (Settings->Actor.IsValid() && ObjectBinding->GetBoundObject() == nullptr)
-						{
-							ObjectBinding->BindToObject(Settings->Actor.Get());
-						}
-					}
-				}
-				ReBindToActor();
-			}
-		}
-	}
-}
 
 void FControlRigEditMode::RefreshObjects()
 {
-	if (WeakSequencer.IsValid())
-	{
-		TSharedRef<ISequencer> Sequencer = WeakSequencer.Pin().ToSharedRef();
-		UMovieScene* MovieScene = Sequencer->GetFocusedMovieSceneSequence() ? Sequencer->GetFocusedMovieSceneSequence()->GetMovieScene() : nullptr;
-		if (MovieScene)
-		{
-			// check if we have an invalid Guid & invalidate Guid if so
-			if (ControlRigGuid.IsValid() && MovieScene->FindSpawnable(ControlRigGuid) == nullptr)
-			{
-				ControlRigGuid.Invalidate();
-				WeakControlRigEditing = nullptr;
-			}
+	WeakControlRigEditing = nullptr;
+	ControlRigGuid.Invalidate();
 
-			SetObjects_Internal();
-		}
-	}
-	else
-	{
-		WeakControlRigEditing = nullptr;
-		ControlRigGuid.Invalidate();
-
-		SetObjects_Internal();
-	}
+	SetObjects_Internal();
 }
 
 void FControlRigEditMode::EnableRigElementEditing(bool bEnabled)
@@ -1086,23 +966,6 @@ void FControlRigEditMode::RecalcPivotTransform()
 
 void FControlRigEditMode::HandleSelectionChanged()
 {
-	if (WeakSequencer.IsValid())
-	{
-		TArray<FString> SelectedNames;
-		for (int32 Index=0; Index<SelectedRigElements.Num(); ++Index)
-		{
-			// if it's control, make sure it's selected
-			if (SelectedRigElements[Index].Type == ERigElementType::Control)
-			{
-				SelectedNames.Add(SelectedRigElements[Index].Name.ToString());
-			}
-		}
-
-		if (SelectedNames.Num() > 0)
-		{
-			WeakSequencer.Pin()->SelectByPropertyPaths(SelectedNames);
-		}
-	}
 
 	for (AControlRigGizmoActor* GizmoActor : GizmoActors)
 	{
@@ -1114,14 +977,7 @@ void FControlRigEditMode::HandleSelectionChanged()
 		}
 	}
 
-	if (WeakSequencer.IsValid())
-	{
-		if (WeakSequencer.Pin()->GetSequencerSettings()->GetShowSelectedNodesOnly())
-		{
-			WeakSequencer.Pin()->RefreshTree();
-		}
-	}
-
+	
 	// update the pivot transform of our selected objects (they could be animating)
 	RecalcPivotTransform();
 }
@@ -1130,10 +986,6 @@ void FControlRigEditMode::BindCommands()
 {
 	const FControlRigEditModeCommands& Commands = FControlRigEditModeCommands::Get();
 
-	CommandBindings->MapAction(
-		Commands.SetKey,
-		FExecuteAction::CreateRaw(this, &FControlRigEditMode::SetKeysForSelectedManipulators));
-
 	CommandBindings->MapAction(
 		Commands.ToggleManipulators,
 		FExecuteAction::CreateRaw(this, &FControlRigEditMode::ToggleManipulators));
@@ -1161,7 +1013,7 @@ bool FControlRigEditMode::GetRigElementGlobalTransform(const FRigElementKey& InE
 		AControlRigGizmoActor* GizmoActor = GetGizmoFromControlName(InElement.Name);
 		if (GizmoActor && ensure(GizmoActor->IsSelected()))
 		{
-			if (ManipulationLayer->GetGlobalTransform(GizmoActor, InElement.Name, OutGlobalTransform))
+			if (ManipulationLayer != nullptr && ManipulationLayer->GetGlobalTransform(GizmoActor, InElement.Name, OutGlobalTransform))
 			{
 				return true;
 			}
@@ -1184,21 +1036,6 @@ bool FControlRigEditMode::GetRigElementGlobalTransform(const FRigElementKey& InE
 	return false;
 }
 
-void FControlRigEditMode::SetKeysForSelectedManipulators()
-{
-	TSharedPtr<ISequencer> Sequencer = WeakSequencer.Pin();
-
-	if (Sequencer)
-	{
-		for (const AControlRigGizmoActor* GizmoActor : GizmoActors)
-		{
-			if (GizmoActor->IsSelected())
-			{
-				ManipulationLayer->AddKey(GizmoActor, Sequencer);
-			}
-		}
-	}
-}
 
 void FControlRigEditMode::ToggleManipulators()
 {
@@ -1322,19 +1159,20 @@ bool FControlRigEditMode::IsTransformDelegateAvailable() const
 
 bool FControlRigEditMode::AreRigElementSelectedAndMovable() const
 {
-	if (IsTransformDelegateAvailable())
+	if (AreRigElementsSelected(FRigElementTypeHelper::ToMask(ERigElementType::Control)))
 	{
-		// you can move control anytime
-		if (AreRigElementsSelected(FRigElementTypeHelper::ToMask(ERigElementType::Control)))
-		{
-			return true;
-		}
+		return true;
+	}
 
-		// if bone, we want to make sure execution is stopped
-		if (bEnableRigElementDefaultPoseEditing && AreRigElementsSelected(FRigElementTypeHelper::ToMask(ERigElementType::Bone)))
-		{
-			return true;
-		}
+	// if bone, we want to make sure execution is stopped
+	if (bEnableRigElementDefaultPoseEditing && AreRigElementsSelected(FRigElementTypeHelper::ToMask(ERigElementType::Bone)))
+	{
+		return true;
+	}
+
+	if (IsTransformDelegateAvailable())
+	{
+		return true;
 	}
 
 	return false;
@@ -1419,7 +1257,7 @@ void FControlRigEditMode::OnControlUISettingChanged(FRigHierarchyContainer* Cont
 
 AControlRigGizmoActor* FControlRigEditMode::GetGizmoFromControlName(const FName& InControlName) const
 {
-	return ManipulationLayer->GetGizmoFromControlName(InControlName);
+	return ManipulationLayer ? ManipulationLayer->GetGizmoFromControlName(InControlName) : nullptr;
 }
 
 #undef LOCTEXT_NAMESPACE
diff --git a/Engine/Plugins/Experimental/ControlRig/Source/ControlRigEditor/Private/EditMode/ControlRigEditMode.h b/Engine/Plugins/Experimental/ControlRig/Source/ControlRigEditor/Private/EditMode/ControlRigEditMode.h
--- a/Engine/Plugins/Experimental/ControlRig/Source/ControlRigEditor/Private/EditMode/ControlRigEditMode.h
+++ b/Engine/Plugins/Experimental/ControlRig/Source/ControlRigEditor/Private/EditMode/ControlRigEditMode.h
@@ -42,9 +42,6 @@ public:
 	/** Set the objects to be displayed in the details panel */
 	void SetObjects(const TWeakObjectPtr<>& InSelectedObject, const FGuid& InObjectBinding, UObject* BindingObject);
 
-	/** Set the sequencer we are bound to */
-	void SetSequencer(TSharedPtr<ISequencer> InSequencer);
-
 	/** This edit mode is re-used between the level editor and the control rig editor. Calling this indicates which context we are in */
 	virtual bool IsInLevelEditor() const { return true; }
 
@@ -80,18 +77,6 @@ public:
 	/** FGCObject interface */
 	virtual void AddReferencedObjects( FReferenceCollector& Collector ) override;
 
-	/** 
-	 * Lets the edit mode know that an object has just been spawned. 
-	 * Allows us to redisplay different underlying objects in the details panel.
-	 */
-	void HandleObjectSpawned(FGuid InObjectBinding, UObject* SpawnedObject, IMovieScenePlayer& Player);
-
-	/** Re-bind to the current actor - used when sequence, selection etc. changes */
-	void ReBindToActor();
-
-	/** Bind us to an actor for editing */
-	void HandleBindToActor(AActor* InActor, bool bFocus);
-
 	/** Refresh our internal object list (they may have changed) */
 	void RefreshObjects();
 
@@ -121,6 +106,8 @@ public:
 	/** Enable RigElement Editing */
 	void EnableRigElementEditing(bool bEnabled);
 
+	/** Get Control Rig, could be more than one later,  we are animating. Currently used by Sequencer for cross selection.*/
+	UControlRig*  GetControlRig() { return WeakControlRigEditing.IsValid() ? WeakControlRigEditing.Get() : nullptr; }
 protected:
 
 	// Gizmo related functions wrt enable/selection
@@ -140,9 +127,6 @@ protected:
 	/** Handle selection internally */
 	void HandleSelectionChanged();
 
-	/** Set keys on all selected manipulators */
-	void SetKeysForSelectedManipulators();
-
 	/** Toggles visibility of manipulators in the viewport */
 	void ToggleManipulators();
 
@@ -176,9 +160,6 @@ protected:
 	/** The draw interface to use for the control rig */
 	FControlRigDrawInterface DrawInterface;
 
-	/** Sequencer we are currently bound to */
-	TWeakPtr<ISequencer> WeakSequencer;
-
 	/** Guard value for selection */
 	bool bSelecting;
 
diff --git a/Engine/Plugins/Experimental/ControlRig/Source/ControlRigEditor/Private/EditMode/ControlRigEditModeSettings.cpp b/Engine/Plugins/Experimental/ControlRig/Source/ControlRigEditor/Private/EditMode/ControlRigEditModeSettings.cpp
--- a/Engine/Plugins/Experimental/ControlRig/Source/ControlRigEditor/Private/EditMode/ControlRigEditModeSettings.cpp
+++ b/Engine/Plugins/Experimental/ControlRig/Source/ControlRigEditor/Private/EditMode/ControlRigEditModeSettings.cpp
@@ -10,51 +10,10 @@
 
 void UControlRigEditModeSettings::PreEditChange(UProperty* PropertyAboutToChange)
 {
-	if (PropertyAboutToChange)
-	{
-		if (PropertyAboutToChange->GetFName() == GET_MEMBER_NAME_CHECKED(UControlRigEditModeSettings, Actor))
-		{
-			PrevActor = Actor.Get();
-		}
-	}
+	
 }
 
 void UControlRigEditModeSettings::PostEditChangeProperty(struct FPropertyChangedEvent& PropertyChangedEvent)
 {
-	if (UProperty* Property = PropertyChangedEvent.Property)
-	{
-		if (Property->GetFName() == GET_MEMBER_NAME_CHECKED(UControlRigEditModeSettings, Actor))
-		{
-			if (Actor.IsValid())
-			{
-				if (Actor->FindComponentByClass<USkeletalMeshComponent>())
-				{
-					// user set the actor we are targeting, so bind our (standalone) sequence to it
-					if (FControlRigEditMode* ControlRigEditMode = static_cast<FControlRigEditMode*>(GLevelEditorModeTools().GetActiveMode(FControlRigEditMode::ModeName)))
-					{
-						ControlRigEditMode->HandleBindToActor(Actor.Get(), true);
-					}
-				}
-				else
-				{
-					Actor = PrevActor;
-				}
-			}
-			else
-			{
-				// user cleared the actor we are targeting, so (un)bind our sequence to it
-				if (FControlRigEditMode* ControlRigEditMode = static_cast<FControlRigEditMode*>(GLevelEditorModeTools().GetActiveMode(FControlRigEditMode::ModeName)))
-				{
-					ControlRigEditMode->HandleBindToActor(nullptr, true);
-				}
-			}
-		}
-		else if (Property->GetFName() == GET_MEMBER_NAME_CHECKED(UControlRigEditModeSettings, Sequence))
-		{
-			if (Sequence)
-			{
-				GEditor->GetEditorSubsystem<UAssetEditorSubsystem>()->OpenEditorForAsset(Sequence);
-			}
-		}
-	}
+	
 }
diff --git a/Engine/Plugins/Experimental/ControlRig/Source/ControlRigEditor/Private/EditMode/ControlRigEditModeSettings.h b/Engine/Plugins/Experimental/ControlRig/Source/ControlRigEditor/Private/EditMode/ControlRigEditModeSettings.h
--- a/Engine/Plugins/Experimental/ControlRig/Source/ControlRigEditor/Private/EditMode/ControlRigEditModeSettings.h
+++ b/Engine/Plugins/Experimental/ControlRig/Source/ControlRigEditor/Private/EditMode/ControlRigEditModeSettings.h
@@ -19,7 +19,6 @@ class UControlRigEditModeSettings : public UObject
 
 	UControlRigEditModeSettings()
 		: bDisplayHierarchy(false)
-		, bShowManipulatorsDuringPlayback(true)
 		, bHideManipulators(false)
 		, bDisplayAxesOnSelection(true)
 		, AxisScale(10.f)
@@ -30,22 +29,11 @@ class UControlRigEditModeSettings : public UObject
 	virtual void PostEditChangeProperty(struct FPropertyChangedEvent& PropertyChangedEvent) override;
 
 public:
-	/** Sequence to animate */
-	UPROPERTY(EditAnywhere, Category = "Animation", NoClear)
-	UControlRigSequence* Sequence;
-
-	/** The actor we are currently animating */
-	UPROPERTY(EditAnywhere, Category = "Animation")
-	TLazyObjectPtr<AActor> Actor;
 
 	/** Whether to show all nodes in the hierarchy being animated */
 	UPROPERTY(EditAnywhere, AdvancedDisplay, Category = "Animation")
 	bool bDisplayHierarchy;
 
-	/** Whether to show manipulators when animations are being played back */
-	UPROPERTY(EditAnywhere, AdvancedDisplay, Category = "Animation")
-	bool bShowManipulatorsDuringPlayback;
-
 	/** Should we always hide manipulators in viewport */
 	UPROPERTY(EditAnywhere, AdvancedDisplay, Category = "Animation")
 	bool bHideManipulators;
@@ -58,7 +46,4 @@ public:
 	UPROPERTY(EditAnywhere, AdvancedDisplay, Category = "Animation")
 	float AxisScale;
 
-private:
-	/** Cache the previous actor for pre/post edit change handling */
-	AActor* PrevActor;
 };
\ No newline at end of file
diff --git a/Engine/Plugins/Experimental/ControlRig/Source/ControlRigEditor/Private/EditMode/SControlRigEditModeTools.cpp b/Engine/Plugins/Experimental/ControlRig/Source/ControlRigEditor/Private/EditMode/SControlRigEditModeTools.cpp
--- a/Engine/Plugins/Experimental/ControlRig/Source/ControlRigEditor/Private/EditMode/SControlRigEditModeTools.cpp
+++ b/Engine/Plugins/Experimental/ControlRig/Source/ControlRigEditor/Private/EditMode/SControlRigEditModeTools.cpp
@@ -73,7 +73,7 @@ void SControlRigEditModeTools::Construct(const FArguments& InArgs, UWorld* InWor
 		+ SScrollBox::Slot()
 		[
 			SNew(SVerticalBox)
-
+			/*
 			+ SVerticalBox::Slot()
 			.AutoHeight()
 			[
@@ -87,6 +87,7 @@ void SControlRigEditModeTools::Construct(const FArguments& InArgs, UWorld* InWor
 					SAssignNew(ControlPicker, SControlPicker, InWorld)
 				]
 			]
+			*/
 
 			+SVerticalBox::Slot()
 			.AutoHeight()
@@ -112,10 +113,8 @@ void SControlRigEditModeTools::SetDetailsObjects(const TArray<TWeakObjectPtr<>>&
 		}
 	}
 
-	ControlPicker->SetControlRig(Rig);
+	//ControlPicker->SetControlRig(Rig);
 
-	// Expand when you have a rig, collapse when set to null
-	PickerExpander->SetExpanded(Rig != nullptr);
 }
 
 void SControlRigEditModeTools::SetSequencer(TSharedPtr<ISequencer> InSequencer)
diff --git a/Engine/Plugins/Experimental/ControlRig/Source/ControlRigEditor/Private/Editor/ControlRigEditor.cpp b/Engine/Plugins/Experimental/ControlRig/Source/ControlRigEditor/Private/Editor/ControlRigEditor.cpp
--- a/Engine/Plugins/Experimental/ControlRig/Source/ControlRigEditor/Private/Editor/ControlRigEditor.cpp
+++ b/Engine/Plugins/Experimental/ControlRig/Source/ControlRigEditor/Private/Editor/ControlRigEditor.cpp
@@ -999,6 +999,12 @@ void FControlRigEditor::HandleModelModified(const UControlRigModel* InModel, ECo
 {
 	DECLARE_SCOPE_HIERARCHICAL_COUNTER_FUNC()
 
+	UControlRigBlueprint* RigBlueprint = Cast<UControlRigBlueprint>(GetBlueprintObj());
+	if (RigBlueprint == nullptr)
+	{
+		return;
+	}
+
 	switch (InType)
 	{
 		case EControlRigModelNotifType::NodeSelected:
@@ -1007,26 +1013,79 @@ void FControlRigEditor::HandleModelModified(const UControlRigModel* InModel, ECo
 			const FControlRigModelNode* Node = (const FControlRigModelNode*)InPayload;
 			if (Node != nullptr)
 			{
-				if (!bIsSelecting)
+				if (FocusedGraphEdPtr.IsValid())
 				{
-					TGuardValue<bool> SelectingGuard(bIsSelecting, true);
-					if (FocusedGraphEdPtr.IsValid())
+					TSharedPtr<SGraphEditor> FocusedGraphEd = FocusedGraphEdPtr.Pin();
+					if (UControlRigGraph* RigGraph = Cast<UControlRigGraph>(FocusedGraphEd->GetCurrentGraph()))
 					{
-						TSharedPtr<SGraphEditor> FocusedGraphEd = FocusedGraphEdPtr.Pin();
-						if (UControlRigGraph* RigGraph = Cast<UControlRigGraph>(FocusedGraphEd->GetCurrentGraph()))
+						TArray<FControlRigModelNode> Selection;
+						bool bSelect = InType == EControlRigModelNotifType::NodeSelected;
+						if (bSelect)
+						{
+							Selection = InModel->SelectedNodes();
+						}
+						else
+						{
+							Selection.Add(*Node);
+						}
+
+						TArray<FRigElementKey> RigElementsToSelect;
+						for (const FControlRigModelNode& SelectedNode : Selection)
+						{
+							if (UControlRigGraphNode* RigNode = Cast<UControlRigGraphNode>(RigGraph->FindNodeFromPropertyName(SelectedNode.Name)))
+							{
+								for (const FControlRigModelPin& Pin : Node->Pins)
+								{
+									if (Pin.Type.PinCategory == UEdGraphSchema_K2::PC_Name)
+									{
+										if (Pin.CustomWidgetName == TEXT("BoneName"))
+										{
+											RigElementsToSelect.Add(FRigElementKey(*Pin.DefaultValue, ERigElementType::Bone));
+										}
+										else if (Pin.CustomWidgetName == TEXT("ControlName"))
+										{
+											RigElementsToSelect.Add(FRigElementKey(*Pin.DefaultValue, ERigElementType::Control));
+										}
+										else if (Pin.CustomWidgetName == TEXT("SpaceName"))
+										{
+											RigElementsToSelect.Add(FRigElementKey(*Pin.DefaultValue, ERigElementType::Space));
+										}
+										else if (Pin.CustomWidgetName == TEXT("CurveName"))
+										{
+											RigElementsToSelect.Add(FRigElementKey(*Pin.DefaultValue, ERigElementType::Curve));
+										}
+									}
+								}
+							}
+						}
+
+						if (RigElementsToSelect.Num() > 0)
+						{
+							if (bSelect)
+							{
+								RigBlueprint->HierarchyContainer.ClearSelection();
+							}
+							for (const FRigElementKey& RigElementToSelect : RigElementsToSelect)
+							{
+								RigBlueprint->HierarchyContainer.Select(RigElementToSelect, bSelect);
+							}
+						}
+
+						if (!bIsSelecting)
 						{
+							TGuardValue<bool> SelectingGuard(bIsSelecting, true);
 							if (UEdGraphNode* EdNode = RigGraph->FindNodeFromPropertyName(Node->Name))
 							{
 								FocusedGraphEd->SetNodeSelection(EdNode, InType == EControlRigModelNotifType::NodeSelected);
 							}
+							break;
 						}
 					}
-					break;
 				}
 
 				if (InType == EControlRigModelNotifType::NodeSelected)
 				{
-					UClass* Class = GetBlueprintObj()->GeneratedClass.Get();
+					UClass* Class = RigBlueprint->GeneratedClass.Get();
 					if (Class)
 					{
 						if (UProperty* Property = Class->FindPropertyByName(Node->Name))
@@ -1650,6 +1709,12 @@ void FControlRigEditor::OnFinishedChangingProperties(const FPropertyChangedEvent
 	if (ControlRig && ControlRigBP && RigElementInDetailPanel)
 	{
 		UControlRig* DebuggedControlRig = Cast<UControlRig>(GetBlueprintObj()->GetObjectBeingDebugged());
+
+		if (DebuggedControlRig && DebuggedControlRig->GetHierarchy()->GetIndex(RigElementInDetailPanel) == INDEX_NONE)
+		{
+			return;
+		}
+
 		UScriptStruct* ScriptStruct = Cast< UScriptStruct>(PropertyChangedEvent.Property->GetOuter());
 		if (ScriptStruct)
 		{
diff --git a/Engine/Plugins/Experimental/ControlRig/Source/ControlRigEditor/Private/Editor/ControlRigEditorEditMode.cpp b/Engine/Plugins/Experimental/ControlRig/Source/ControlRigEditor/Private/Editor/ControlRigEditorEditMode.cpp
--- a/Engine/Plugins/Experimental/ControlRig/Source/ControlRigEditor/Private/Editor/ControlRigEditorEditMode.cpp
+++ b/Engine/Plugins/Experimental/ControlRig/Source/ControlRigEditor/Private/Editor/ControlRigEditorEditMode.cpp
@@ -13,11 +13,17 @@ bool FControlRigEditorEditMode::GetCameraTarget(FSphere& OutTarget) const
 
 	for (int32 Index = 0; Index < SelectedRigElements.Num(); ++Index)
 	{
-		if (SelectedRigElements[Index].Type == ERigElementType::Bone || SelectedRigElements[Index].Type == ERigElementType::Control)
+		if (SelectedRigElements[Index].Type == ERigElementType::Bone || SelectedRigElements[Index].Type == ERigElementType::Space)
 		{
 			FTransform Transform = OnGetRigElementTransformDelegate.Execute(SelectedRigElements[Index], false);
 			Box += Transform.GetLocation();
 		}
+		else if (SelectedRigElements[Index].Type == ERigElementType::Control)
+		{
+			FTransform Transform = OnGetRigElementTransformDelegate.Execute(SelectedRigElements[Index], false);
+			Box += Transform.TransformPosition(FVector::OneVector * 50.f); // we assume 100cm as the base size for a control
+			Box += Transform.TransformPosition(FVector::OneVector * -50.f);
+		}
 	}
 
 	if(Box.IsValid)
