Merging //UE4/Main@9388041 to Dev-RenderPlat-Staging

diff --git a/Engine/Plugins/Experimental/ControlRig/Source/ControlRigEditor/Private/Editor/SRigHierarchy.cpp b/Engine/Plugins/Experimental/ControlRig/Source/ControlRigEditor/Private/Editor/SRigHierarchy.cpp
--- a/Engine/Plugins/Experimental/ControlRig/Source/ControlRigEditor/Private/Editor/SRigHierarchy.cpp
+++ b/Engine/Plugins/Experimental/ControlRig/Source/ControlRigEditor/Private/Editor/SRigHierarchy.cpp
@@ -31,6 +31,7 @@
 #include "ControlRig.h"
 #include "ControlRigEditorStyle.h"
 #include "HAL/PlatformApplicationMisc.h"
+#include "Dialogs/Dialogs.h"
 
 #define LOCTEXT_NAMESPACE "SRigHierarchy"
 
@@ -114,9 +115,36 @@ void SRigHierarchyItem::Construct(const FArguments& InArgs, TSharedPtr<FControlR
 			Brush = FControlRigEditorStyle::Get().GetBrush("ControlRig.Tree.Space");
 			break;
 		}
+		case ERigElementType::Bone:
+		{
+			ERigBoneType BoneType = ERigBoneType::User;
+
+			int32 BoneIndex = InHierarchy->ControlRigBlueprint->HierarchyContainer.GetIndex(InRigTreeElement->Key);
+			if (BoneIndex != INDEX_NONE)
+			{
+				BoneType = InHierarchy->ControlRigBlueprint->HierarchyContainer.BoneHierarchy[BoneIndex].Type;
+			}
+
+			switch (BoneType)
+			{
+				case ERigBoneType::Imported:
+				{
+					Brush = FControlRigEditorStyle::Get().GetBrush("ControlRig.Tree.BoneImported");
+					break;
+				}
+				case ERigBoneType::User:
+				default:
+				{
+					Brush = FControlRigEditorStyle::Get().GetBrush("ControlRig.Tree.BoneUser");
+					break;
+				}
+			}
+
+			break;
+		}
 		default:
 		{
-			Brush = FControlRigEditorStyle::Get().GetBrush("ControlRig.Tree.Bone");
+			Brush = FControlRigEditorStyle::Get().GetBrush("ControlRig.Tree.BoneUser");
 			break;
 		}
 	}
@@ -568,10 +596,14 @@ void SRigHierarchy::AddElement(FRigElementKey InKey, FRigElementKey InParentKey)
 			RootElements.Add(NewItem);
 		}
 	}
-	else if (InKey.Name.ToString().Contains(FilteredString))
+	else
 	{
-		TSharedPtr<FRigTreeElement> NewItem = MakeShared<FRigTreeElement>(InKey, SharedThis(this));
-		RootElements.Add(NewItem);
+		FString FilteredStringUnderScores = FilteredString.Replace(TEXT(" "), TEXT("_"));
+		if (InKey.Name.ToString().Contains(FilteredString) || InKey.Name.ToString().Contains(FilteredStringUnderScores))	
+		{
+			TSharedPtr<FRigTreeElement> NewItem = MakeShared<FRigTreeElement>(InKey, SharedThis(this));
+			RootElements.Add(NewItem);
+		}
 	}
 }
 
@@ -827,19 +859,7 @@ void SRigHierarchy::RefreshHierarchy(const FAssetData& InAssetData)
 		ControlRigBlueprint->Modify();
 
 		const FReferenceSkeleton& RefSkeleton = Mesh->RefSkeleton;
-		const TArray<FMeshBoneInfo>& BoneInfos = RefSkeleton.GetRefBoneInfo();
-		const TArray<FTransform>& BonePoses = RefSkeleton.GetRefBonePose();
-
-		for (int32 Index = 0; Index < RefSkeleton.GetNum(); ++Index)
-		{
-			// only add if you don't have it. This may change in the future
-			int32 RigIndex = Hierarchy->BoneHierarchy.GetIndex(BoneInfos[Index].Name);
-			if (RigIndex != INDEX_NONE)
-			{
-				// @todo: add optimized version without sorting, but if no sort, we should make sure not to use find index function
-				Hierarchy->BoneHierarchy.SetInitialTransform(RigIndex, FAnimationRuntime::GetComponentSpaceTransform(RefSkeleton, BonePoses, Index));
-			}
-		}
+		Hierarchy->BoneHierarchy.ImportSkeleton(RefSkeleton, NAME_None, true, true, true);
 	}
 }
 void SRigHierarchy::CreateImportMenu(FMenuBuilder& MenuBuilder)
@@ -879,28 +899,7 @@ void SRigHierarchy::ImportHierarchy(const FAssetData& InAssetData)
 		ControlRigBlueprint->Modify();
 
 		const FReferenceSkeleton& RefSkeleton = Mesh->RefSkeleton;
-		const TArray<FMeshBoneInfo>& BoneInfos = RefSkeleton.GetRefBoneInfo();
-		const TArray<FTransform>& BonePoses = RefSkeleton.GetRefBonePose();
-
-		Hierarchy->ClearSelection();
-
-		TArray<FName> AddedBones;
-		for (int32 Index = 0; Index < RefSkeleton.GetNum(); ++Index)
-		{
-			// only add if you don't have it. This may change in the future
-			if (Hierarchy->BoneHierarchy.GetIndex(BoneInfos[Index].Name) == INDEX_NONE)
-			{
-				// @todo: add optimized version without sorting, but if no sort, we should make sure not to use find index function
-				FName ParentName = (BoneInfos[Index].ParentIndex != INDEX_NONE) ? BoneInfos[BoneInfos[Index].ParentIndex].Name : NAME_None;
-				Hierarchy->BoneHierarchy.Add(BoneInfos[Index].Name, ParentName, FAnimationRuntime::GetComponentSpaceTransform(RefSkeleton, BonePoses, Index));
-				AddedBones.Add(BoneInfos[Index].Name);
-			}
-		}
-
-		for (const FName& AddedBone : AddedBones)
-		{
-			Hierarchy->BoneHierarchy.Select(AddedBone);
-		}
+		Hierarchy->BoneHierarchy.ImportSkeleton(RefSkeleton, NAME_None, false, false, true);
 
 		FSlateApplication::Get().DismissAllMenus();
 		RefreshTreeView();
@@ -959,6 +958,9 @@ void SRigHierarchy::HandleDeleteItem()
 		// clear detail view display
 		ControlRigEditor.Pin()->ClearDetailObject();
 
+		bool bConfirmedByUser = false;
+		bool bDeleteImportedBones = false;
+
 		TArray<TSharedPtr<FRigTreeElement>> SelectedItems = TreeView->GetSelectedItems();
 		for (int32 ItemIndex = 0; ItemIndex < SelectedItems.Num(); ++ItemIndex)
 		{
@@ -969,8 +971,31 @@ void SRigHierarchy::HandleDeleteItem()
 			{
 				case ERigElementType::Bone:
 				{
-					if (Hierarchy->BoneHierarchy.GetIndex(SelectedKey.Name) != INDEX_NONE)
+					int32 BoneIndex = Hierarchy->BoneHierarchy.GetIndex(SelectedKey.Name);
+					if (BoneIndex != INDEX_NONE)
 					{
+						const FRigBone& Bone = Hierarchy->BoneHierarchy[BoneIndex];
+						if (Bone.Type == ERigBoneType::Imported && Bone.ParentIndex != INDEX_NONE)
+						{
+							if (!bConfirmedByUser)
+							{
+								FText ConfirmDelete = LOCTEXT("ConfirmDeleteBoneHierarchy",
+									"Deleting imported(white) bones can cause issues with animation - are you sure ?");
+
+								FSuppressableWarningDialog::FSetupInfo Info(ConfirmDelete, LOCTEXT("DeleteImportedBone", "Delete Imported Bone"), "DeleteImportedBoneHierarchy_Warning");
+								Info.ConfirmText = LOCTEXT("DeleteImportedBoneHierarchy_Yes", "Yes");
+								Info.CancelText = LOCTEXT("DeleteImportedBoneHierarchy_No", "No");
+
+								FSuppressableWarningDialog DeleteImportedBonesInHierarchy(Info);
+								bDeleteImportedBones = DeleteImportedBonesInHierarchy.ShowModal() != FSuppressableWarningDialog::Cancel;
+								bConfirmedByUser = true;
+							}
+
+							if (!bDeleteImportedBones)
+							{
+								break;
+							}
+						}
 						Hierarchy->BoneHierarchy.Remove(SelectedKey.Name);
 					}
 					break;
@@ -1040,7 +1065,7 @@ void SRigHierarchy::HandleNewItem(ERigElementType InElementType)
 				{
 					FName ParentName = ParentKey.Type == ERigElementType::Bone ? ParentKey.Name : NAME_None;
 					ParentTransform = ParentName == NAME_None ? FTransform::Identity : ParentTransform;
-					Hierarchy->BoneHierarchy.Add(NewElementName, ParentName, ParentTransform);
+					Hierarchy->BoneHierarchy.Add(NewElementName, ParentName, ERigBoneType::User, ParentTransform);
 					break;
 				}
 				case ERigElementType::Control:
@@ -1147,7 +1172,7 @@ void SRigHierarchy::HandleDuplicateItem()
 					FName ParentName = Hierarchy->BoneHierarchy[Key.Name].ParentName;
 
 					const FName NewName = CreateUniqueName(Key.Name, Key.Type);
-					FRigBone& NewBone = Hierarchy->BoneHierarchy.Add(NewName, ParentName, Transform);
+					FRigBone& NewBone = Hierarchy->BoneHierarchy.Add(NewName, ParentName, ERigBoneType::User, Transform);
 					NewKeys.Add(NewBone.GetElementKey());
 					break;
 				}
@@ -1205,6 +1230,29 @@ void SRigHierarchy::HandleRenameItem()
 		TArray<TSharedPtr<FRigTreeElement>> SelectedItems = TreeView->GetSelectedItems();
 		if (SelectedItems.Num() == 1)
 		{
+			if (SelectedItems[0]->Key.Type == ERigElementType::Bone)
+			{
+				int32 BoneIndex = Hierarchy->GetIndex(SelectedItems[0]->Key);
+				if (BoneIndex != INDEX_NONE)
+				{
+					const FRigBone& Bone = Hierarchy->BoneHierarchy[BoneIndex];
+					if (Bone.Type == ERigBoneType::Imported)
+					{
+						FText ConfirmRename = LOCTEXT("RenameDeleteBoneHierarchy",
+							"Renaming imported(white) bones can cause issues with animation - are you sure ?");
+
+						FSuppressableWarningDialog::FSetupInfo Info(ConfirmRename, LOCTEXT("RenameImportedBone", "Rename Imported Bone"), "RenameImportedBoneHierarchy_Warning");
+						Info.ConfirmText = LOCTEXT("RenameImportedBoneHierarchy_Yes", "Yes");
+						Info.CancelText = LOCTEXT("RenameImportedBoneHierarchy_No", "No");
+
+						FSuppressableWarningDialog RenameImportedBonesInHierarchy(Info);
+						if (RenameImportedBonesInHierarchy.ShowModal() == FSuppressableWarningDialog::Cancel)
+						{
+							return;
+						}
+					}
+				}
+			}
 			SelectedItems[0]->RequestRename();
 		}
 	}
@@ -1354,7 +1402,7 @@ FReply SRigHierarchy::OnDragDetected(const FGeometry& MyGeometry, const FPointer
 		DraggedElements.Add(SelectedItem->Key);
 	}
 
-	if (MouseEvent.IsMouseButtonDown(EKeys::LeftMouseButton))
+	if (MouseEvent.IsMouseButtonDown(EKeys::LeftMouseButton) && DraggedElements.Num() > 0)
 	{
 		if (ControlRigEditor.IsValid())
 		{
@@ -1451,7 +1499,31 @@ FReply SRigHierarchy::OnAcceptDrop(const FDragDropEvent& DragDropEvent, EItemDro
 					DraggedKey.Type,
 					Container->GetIndex(DraggedKey)))
 				{
-					return FReply::Unhandled();;
+					return FReply::Unhandled();
+				}
+
+				if (DraggedKey.Type == ERigElementType::Bone)
+				{
+					int32 BoneIndex = Container->BoneHierarchy.GetIndex(DraggedKey.Name);
+					if (BoneIndex != INDEX_NONE)
+					{
+						const FRigBone& Bone = Container->BoneHierarchy[BoneIndex];
+						if (Bone.Type == ERigBoneType::Imported && Bone.ParentIndex != INDEX_NONE)
+						{
+							FText ConfirmReparent = LOCTEXT("ConfirmReparentBoneHierarchy",
+								"Reparenting imported(white) bones can cause issues with animation - are you sure ?");
+
+							FSuppressableWarningDialog::FSetupInfo Info(ConfirmReparent, LOCTEXT("ReparentImportedBone", "Reparent Imported Bone"), "ReparentImportedBoneHierarchy_Warning");
+							Info.ConfirmText = LOCTEXT("ReparentImportedBoneHierarchy_Yes", "Yes");
+							Info.CancelText = LOCTEXT("ReparentImportedBoneHierarchy_No", "No");
+
+							FSuppressableWarningDialog ReparentImportedBonesInHierarchy(Info);
+							if (ReparentImportedBonesInHierarchy.ShowModal() == FSuppressableWarningDialog::Cancel)
+							{
+								return FReply::Unhandled();
+							}
+						}
+					}
 				}
 			}
 
diff --git a/Engine/Plugins/Experimental/ControlRig/Source/ControlRigEditor/Private/Sequencer/ControlRigParameterTrackEditor.cpp b/Engine/Plugins/Experimental/ControlRig/Source/ControlRigEditor/Private/Sequencer/ControlRigParameterTrackEditor.cpp
new file mode 100644
--- /dev/null
+++ b/Engine/Plugins/Experimental/ControlRig/Source/ControlRigEditor/Private/Sequencer/ControlRigParameterTrackEditor.cpp
@@ -0,0 +1,1310 @@
+// Copyright 1998-2019 Epic Games, Inc. All Rights Reserved.
+
+#include "ControlRigParameterTrackEditor.h"
+#include "Sequencer/MovieSceneControlRigParameterTrack.h"
+#include "Sequencer/MovieSceneControlRigParameterSection.h"
+#include "Framework/Commands/Commands.h"
+#include "Rendering/DrawElements.h"
+#include "Widgets/SBoxPanel.h"
+#include "Framework/MultiBox/MultiBoxBuilder.h"
+#include "GameFramework/Actor.h"
+#include "AssetData.h"
+#include "Modules/ModuleManager.h"
+#include "Layout/WidgetPath.h"
+#include "Framework/Application/MenuStack.h"
+#include "Framework/Application/SlateApplication.h"
+#include "Widgets/Layout/SBox.h"
+#include "SequencerSectionPainter.h"
+#include "Components/SkeletalMeshComponent.h"
+#include "Editor/UnrealEdEngine.h"
+#include "UnrealEdGlobals.h"
+#include "ClassViewerModule.h"
+#include "ClassViewerFilter.h"
+#include "CommonMovieSceneTools.h"
+#include "AssetRegistryModule.h"
+#include "IContentBrowserSingleton.h"
+#include "ContentBrowserModule.h"
+#include "MatineeImportTools.h"
+#include "Matinee/InterpTrackAnimControl.h"
+#include "SequencerUtilities.h"
+#include "ISectionLayoutBuilder.h"
+#include "EditorStyleSet.h"
+#include "MovieSceneTimeHelpers.h"
+#include "Fonts/FontMeasure.h"
+#include "AnimationEditorUtils.h"
+#include "Misc/MessageDialog.h"
+#include "Framework/Notifications/NotificationManager.h"
+#include "Widgets/Notifications/SNotificationList.h"
+#include "Toolkits/AssetEditorManager.h"
+#include "Engine/SCS_Node.h"
+#include "Engine/SimpleConstructionScript.h"
+#include "Engine/Blueprint.h"
+#include "ControlRig.h"
+#include "ControlRigEditMode.h"
+#include "EditorModeManager.h"
+#include "Engine/Selection.h"
+#include "Manipulatable/IControlRigManipulatable.h"
+#include "ControlRigSkeletalMeshBinding.h"
+#include "LevelEditorViewport.h"
+#include "IKeyArea.h"
+
+#define LOCTEXT_NAMESPACE "FControlRigParameterTrackEditor"
+
+
+
+class FControlRigTrackCommands
+	: public TCommands<FControlRigTrackCommands>
+{
+public:
+
+	FControlRigTrackCommands()
+		: TCommands<FControlRigTrackCommands>
+		(
+			"ControlRigTrack",
+			NSLOCTEXT("Contexts", "ControlRigTrack", "ControlRigTrack"),
+			NAME_None,
+			FEditorStyle::GetStyleSetName() // Icon Style Set
+			)
+		, BindingCount(0)
+	{ }
+
+	/** Sets a transform key at the current time for the selected actor */
+	TSharedPtr< FUICommandInfo > AddTransformKey;
+
+	/** Sets a translation key at the current time for the selected actor */
+	TSharedPtr< FUICommandInfo > AddTranslationKey;
+
+	/** Sets a rotation key at the current time for the selected actor */
+	TSharedPtr< FUICommandInfo > AddRotationKey;
+
+	/** Sets a scale key at the current time for the selected actor */
+	TSharedPtr< FUICommandInfo > AddScaleKey;
+
+	/**
+	 * Initialize commands
+	 */
+	virtual void RegisterCommands() override;
+
+	mutable uint32 BindingCount;
+};
+
+void FControlRigTrackCommands::RegisterCommands()
+{
+	UI_COMMAND(AddTransformKey, "Add Transform Key", "Add a transform key at the current time for the selected control rig control.", EUserInterfaceActionType::Button, FInputChord(EKeys::S));
+	UI_COMMAND(AddTranslationKey, "Add Translation Key", "Add a translation key at the current time for the selected control rig control.", EUserInterfaceActionType::Button, FInputChord(EModifierKey::Shift, EKeys::W));
+	UI_COMMAND(AddRotationKey, "Add Rotation Key", "Add a rotation key at the current time for the selected control rig control.", EUserInterfaceActionType::Button, FInputChord(EModifierKey::Shift, EKeys::E));
+	UI_COMMAND(AddScaleKey, "Add Scale Key", "Add a scale key at the current time for the selected control rig control.", EUserInterfaceActionType::Button, FInputChord(EModifierKey::Shift, EKeys::R));
+}
+
+static USkeleton* GetSkeletonFromComponent(UActorComponent* InComponent)
+{
+	USkeletalMeshComponent* SkeletalMeshComp = Cast<USkeletalMeshComponent>(InComponent);
+	if (SkeletalMeshComp && SkeletalMeshComp->SkeletalMesh && SkeletalMeshComp->SkeletalMesh->Skeleton)
+	{
+		// @todo Multiple actors, multiple components
+		return SkeletalMeshComp->SkeletalMesh->Skeleton;
+	}
+
+	return nullptr;
+}
+
+static USkeleton* AcquireSkeletonFromObjectGuid(const FGuid& Guid, UObject** Object, TSharedPtr<ISequencer> SequencerPtr)
+{
+	UObject* BoundObject = SequencerPtr.IsValid() ? SequencerPtr->FindSpawnedObjectOrTemplate(Guid) : nullptr;
+	*Object = BoundObject;
+	if (AActor* Actor = Cast<AActor>(BoundObject))
+	{
+		for (UActorComponent* Component : Actor->GetComponents())
+		{
+			if (USkeleton* Skeleton = GetSkeletonFromComponent(Component))
+			{
+				return Skeleton;
+			}
+		}
+
+		AActor* ActorCDO = Cast<AActor>(Actor->GetClass()->GetDefaultObject());
+		if (ActorCDO)
+		{
+			for (UActorComponent* Component : ActorCDO->GetComponents())
+			{
+				if (USkeleton* Skeleton = GetSkeletonFromComponent(Component))
+				{
+					return Skeleton;
+				}
+			}
+		}
+
+		UBlueprintGeneratedClass* ActorBlueprintGeneratedClass = Cast<UBlueprintGeneratedClass>(Actor->GetClass());
+		if (ActorBlueprintGeneratedClass)
+		{
+			const TArray<USCS_Node*>& ActorBlueprintNodes = ActorBlueprintGeneratedClass->SimpleConstructionScript->GetAllNodes();
+
+			for (USCS_Node* Node : ActorBlueprintNodes)
+			{
+				if (Node->ComponentClass->IsChildOf(USkeletalMeshComponent::StaticClass()))
+				{
+					if (USkeleton* Skeleton = GetSkeletonFromComponent(Node->GetActualComponentTemplate(ActorBlueprintGeneratedClass)))
+					{
+						return Skeleton;
+					}
+				}
+			}
+		}
+	}
+	else if (USkeletalMeshComponent* SkeletalMeshComponent = Cast<USkeletalMeshComponent>(BoundObject))
+	{
+		if (USkeleton* Skeleton = GetSkeletonFromComponent(SkeletalMeshComponent))
+		{
+			return Skeleton;
+		}
+	}
+
+	return nullptr;
+}
+
+FControlRigParameterTrackEditor::FControlRigParameterTrackEditor(TSharedRef<ISequencer> InSequencer)
+	: FKeyframeTrackEditor<UMovieSceneControlRigParameterTrack>(InSequencer), bIsDoingASelection(false)
+{
+	SelectionChangedHandle = InSequencer->GetSelectionChangedTracks().AddRaw(this, &FControlRigParameterTrackEditor::OnSelectionChanged);
+	SequencerChangedHandle = InSequencer->OnMovieSceneDataChanged().AddRaw(this, &FControlRigParameterTrackEditor::OnSequencerDataChanged);
+
+	FControlRigTrackCommands::Register();
+	//register all modified/selections for control rigs
+	UMovieScene* MovieScene = InSequencer->GetFocusedMovieSceneSequence()->GetMovieScene();
+	const TArray<FMovieSceneBinding>& Bindings = MovieScene->GetBindings();
+	for(const FMovieSceneBinding& Binding : Bindings)
+	{
+		UMovieSceneControlRigParameterTrack* Track = Cast<UMovieSceneControlRigParameterTrack>(MovieScene->FindTrack(UMovieSceneControlRigParameterTrack::StaticClass(), Binding.GetObjectGuid(), NAME_None));
+		if (Track && Track->GetControlRig())
+		{
+			Track->GetControlRig()->ControlModified().AddRaw(this, &FControlRigParameterTrackEditor::HandleControlModified);
+			Track->GetControlRig()->ControlSelected().AddRaw(this, &FControlRigParameterTrackEditor::HandleControlSelected);
+		}
+	}
+}
+
+FControlRigParameterTrackEditor::~FControlRigParameterTrackEditor()
+{
+}
+
+void FControlRigParameterTrackEditor::OnRelease()
+{
+	if (GetSequencer().IsValid())
+	{
+		if (SelectionChangedHandle.IsValid())
+		{
+			GetSequencer()->GetSelectionChangedTracks().Remove(SelectionChangedHandle);
+		}
+		if (SequencerChangedHandle.IsValid())
+		{
+			GetSequencer()->OnMovieSceneDataChanged().Remove(SequencerChangedHandle);
+		}
+		if (GetSequencer()->GetFocusedMovieSceneSequence() && GetSequencer()->GetFocusedMovieSceneSequence()->GetMovieScene())
+		{
+			UMovieScene* MovieScene = GetSequencer()->GetFocusedMovieSceneSequence()->GetMovieScene();
+			const TArray<FMovieSceneBinding>& Bindings = MovieScene->GetBindings();
+			for (const FMovieSceneBinding& Binding : Bindings)
+			{
+				UMovieSceneControlRigParameterTrack* Track = Cast<UMovieSceneControlRigParameterTrack>(MovieScene->FindTrack(UMovieSceneControlRigParameterTrack::StaticClass(), Binding.GetObjectGuid(), NAME_None));
+				if (Track && Track->GetControlRig())
+				{
+					Track->GetControlRig()->ControlModified().RemoveAll(this);
+					Track->GetControlRig()->ControlSelected().RemoveAll(this);
+				}
+			}
+		}
+	}
+	FControlRigEditMode* ControlRigEditMode = static_cast<FControlRigEditMode*>(GLevelEditorModeTools().GetActiveMode(FControlRigEditMode::ModeName));
+	if (ControlRigEditMode)
+	{
+		if (GLevelEditorModeTools().HasToolkitHost())
+		{
+			GLevelEditorModeTools().DeactivateMode(FControlRigEditMode::ModeName);
+		}
+
+		ControlRigEditMode->SetObjects(nullptr, FGuid(), nullptr);
+	}
+
+	const FControlRigTrackCommands& Commands = FControlRigTrackCommands::Get();
+	Commands.BindingCount--;
+
+	if (Commands.BindingCount < 1)
+	{
+		FControlRigTrackCommands::Unregister();
+	}
+
+}
+
+TSharedRef<ISequencerTrackEditor> FControlRigParameterTrackEditor::CreateTrackEditor(TSharedRef<ISequencer> InSequencer)
+{
+	return MakeShareable(new FControlRigParameterTrackEditor(InSequencer));
+}
+
+
+bool FControlRigParameterTrackEditor::SupportsType(TSubclassOf<UMovieSceneTrack> Type) const
+{
+	return Type == UMovieSceneControlRigParameterTrack::StaticClass();
+}
+
+
+TSharedRef<ISequencerSection> FControlRigParameterTrackEditor::MakeSectionInterface(UMovieSceneSection& SectionObject, UMovieSceneTrack& Track, FGuid ObjectBinding)
+{
+	check(SupportsType(SectionObject.GetOuter()->GetClass()));
+
+	return MakeShareable(new FControlRigParameterSection(SectionObject, GetSequencer()));
+}
+
+
+void FControlRigParameterTrackEditor::BuildObjectBindingTrackMenu(FMenuBuilder& MenuBuilder, const TArray<FGuid>& ObjectBindings, const UClass* ObjectClass)
+{
+
+	if (ObjectClass->IsChildOf(USkeletalMeshComponent::StaticClass()) || ObjectClass->IsChildOf(AActor::StaticClass()) || ObjectClass->IsChildOf(UChildActorComponent::StaticClass()))
+	{
+		const TSharedPtr<ISequencer> ParentSequencer = GetSequencer();
+		UObject *BoundObject = nullptr;
+		USkeleton* Skeleton = AcquireSkeletonFromObjectGuid(ObjectBindings[0], &BoundObject, ParentSequencer);
+
+		if (Skeleton)
+		{
+			//if there are any other control rigs we don't allow it for now..
+			//mz todo will allow later
+			UMovieScene* MovieScene = GetSequencer()->GetFocusedMovieSceneSequence()->GetMovieScene();
+			UMovieSceneControlRigParameterTrack* ExistingTrack = Cast<UMovieSceneControlRigParameterTrack>(MovieScene->FindTrack(UMovieSceneControlRigParameterTrack::StaticClass(), ObjectBindings[0], NAME_None));
+			if (!ExistingTrack)
+			{
+				UMovieSceneTrack* Track = nullptr;
+
+				MenuBuilder.AddSubMenu(
+					LOCTEXT("AddControlRig", "Animation ControlRig"), NSLOCTEXT("Sequencer", "AddControlRigTooltip", "Adds an animation ControlRig track."),
+					FNewMenuDelegate::CreateRaw(this, &FControlRigParameterTrackEditor::AddControlRigSubMenu, ObjectBindings, Track)
+				);
+			}
+		}
+	}
+}
+
+void FControlRigParameterTrackEditor::AddControlRigSubMenu(FMenuBuilder& MenuBuilder, TArray<FGuid> ObjectBindings, UMovieSceneTrack* Track)
+{
+	/*
+	MenuBuilder.BeginSection(TEXT("ChooseSequence"), LOCTEXT("ChooseSequence", "Choose Sequence"));
+	{
+		FAssetPickerConfig AssetPickerConfig;
+		{
+			AssetPickerConfig.OnAssetSelected = FOnAssetSelected::CreateRaw(this, &FControlRigTrackEditor::OnSequencerAssetSelected, ObjectBindings, Track);
+			AssetPickerConfig.OnAssetEnterPressed = FOnAssetEnterPressed::CreateRaw(this, &FControlRigTrackEditor::OnSequencerAssetEnterPressed, ObjectBindings, Track);
+			AssetPickerConfig.bAllowNullSelection = false;
+			AssetPickerConfig.InitialAssetViewType = EAssetViewType::List;
+			AssetPickerConfig.Filter.bRecursiveClasses = true;
+			AssetPickerConfig.Filter.ClassNames.Add(UControlRigSequence::StaticClass()->GetFName());
+		}
+
+		FContentBrowserModule& ContentBrowserModule = FModuleManager::Get().LoadModuleChecked<FContentBrowserModule>(TEXT("ContentBrowser"));
+
+		TSharedPtr<SBox> MenuEntry = SNew(SBox)
+			.WidthOverride(300.0f)
+			.HeightOverride(300.0f)
+			[
+				ContentBrowserModule.Get().CreateAssetPicker(AssetPickerConfig)
+			];
+
+		MenuBuilder.AddWidget(MenuEntry.ToSharedRef(), FText::GetEmpty(), true);
+	}
+	MenuBuilder.EndSection();
+	*/
+
+
+	const TSharedPtr<ISequencer> ParentSequencer = GetSequencer();
+	UObject *BoundObject = nullptr;
+	//todo support multiple bindings?
+	USkeleton* Skeleton = AcquireSkeletonFromObjectGuid(ObjectBindings[0], &BoundObject, GetSequencer());
+
+	if (Skeleton)
+	{
+		//MenuBuilder.AddSubMenu(
+		//	LOCTEXT("AddControlRigTrack", "ControlRigTrack"), NSLOCTEXT("ControlRig", "AddControlRigTrack", "Adds a Control Rigtrack."),
+		//	FNewMenuDelegate::CreateRaw(this, &FControlRigParameterTrackEditor::AddAnimationSubMenu, BoundObject, ObjectBindings[0], Skeleton)
+		//);
+
+
+		class FControlRigClassFilter : public IClassViewerFilter
+		{
+		public:
+			bool IsClassAllowed(const FClassViewerInitializationOptions& InInitOptions, const UClass* InClass, TSharedRef< FClassViewerFilterFuncs > InFilterFuncs) override
+			{
+				const bool bChildOfObjectClass = InClass->IsChildOf(UControlRig::StaticClass());
+				const bool bMatchesFlags = !InClass->HasAnyClassFlags(CLASS_Hidden | CLASS_HideDropDown | CLASS_Deprecated | CLASS_Abstract);
+				return bChildOfObjectClass && bMatchesFlags;
+			}
+
+			virtual bool IsUnloadedClassAllowed(const FClassViewerInitializationOptions& InInitOptions, const TSharedRef< const IUnloadedBlueprintData > InUnloadedClassData, TSharedRef< FClassViewerFilterFuncs > InFilterFuncs) override
+			{
+				const bool bChildOfObjectClass = InUnloadedClassData->IsChildOf(UControlRig::StaticClass());
+				const bool bMatchesFlags = !InUnloadedClassData->HasAnyClassFlags(CLASS_Hidden | CLASS_HideDropDown | CLASS_Deprecated | CLASS_Abstract);
+				return bChildOfObjectClass && bMatchesFlags;
+			}
+		};
+
+		FClassViewerInitializationOptions Options;
+		Options.bShowUnloadedBlueprints = true;
+		Options.NameTypeToDisplay = EClassViewerNameTypeToDisplay::DisplayName;
+
+		TSharedPtr<FControlRigClassFilter> ClassFilter = MakeShareable(new FControlRigClassFilter);
+		Options.ClassFilter = ClassFilter;
+		Options.bShowNoneOption = false;
+
+		FClassViewerModule& ClassViewerModule = FModuleManager::LoadModuleChecked<FClassViewerModule>("ClassViewer");
+		//	FNewMenuDelegate::CreateRaw(this, &FControlRigParameterTrackEditor::AddAnimationSubMenu, BoundObject, ObjectBindings[0], Skeleton)
+
+
+
+		TSharedRef<SWidget> ClassViewer = ClassViewerModule.CreateClassViewer(Options, FOnClassPicked::CreateRaw(this, &FControlRigParameterTrackEditor::AddControlRig, BoundObject, ObjectBindings[0], Skeleton));
+		MenuBuilder.AddWidget(ClassViewer, FText::GetEmpty(), true);
+
+		/*
+		FAssetPickerConfig AssetPickerConfig;
+		{
+			AssetPickerConfig.OnAssetSelected = FOnAssetSelected::CreateRaw(this, &FControlRigParameterTrackEditor::OnControlRigAssetSelected, ObjectBindings, Skeleton);
+			AssetPickerConfig.OnAssetEnterPressed = FOnAssetEnterPressed::CreateRaw(this, &FControlRigParameterTrackEditor::OnControlRigAssetEnterPressed, ObjectBindings, Skeleton);
+			AssetPickerConfig.bAllowNullSelection = false;
+			AssetPickerConfig.InitialAssetViewType = EAssetViewType::List;
+			AssetPickerConfig.OnShouldFilterAsset = FOnShouldFilterAsset::CreateRaw(this, &FControlRigParameterTrackEditor::ShouldFilterAsset);
+			AssetPickerConfig.Filter.bRecursiveClasses = true;
+			AssetPickerConfig.Filter.ClassNames.Add((UControlRig::StaticClass())->GetFName());
+			AssetPickerConfig.Filter.TagsAndValues.Add(TEXT("Skeleton"), FAssetData(Skeleton).GetExportTextName());
+		}
+
+		FContentBrowserModule& ContentBrowserModule = FModuleManager::Get().LoadModuleChecked<FContentBrowserModule>(TEXT("ContentBrowser"));
+
+		TSharedPtr<SBox> MenuEntry = SNew(SBox)
+			.WidthOverride(300.0f)
+			.HeightOverride(300.f)
+			[
+				ContentBrowserModule.Get().CreateAssetPicker(AssetPickerConfig)
+			];
+
+		MenuBuilder.AddWidget(MenuEntry.ToSharedRef(), FText::GetEmpty(), true);
+		*/
+	}
+
+}
+/*
+bool FControlRigParameterTrackEditor::ShouldFilterAsset(const FAssetData& AssetData)
+{
+	if (AssetData.AssetClass == UControlRig::StaticClass()->GetFName())
+	{
+		return true;
+	}
+	return false;
+}
+
+void FControlRigParameterTrackEditor::OnControlRigAssetSelected(const FAssetData& AssetData, TArray<FGuid> ObjectBindings, USkeleton* Skeleton)
+{
+	FSlateApplication::Get().DismissAllMenus();
+
+	UObject* SelectedObject = AssetData.GetAsset();
+	TSharedPtr<ISequencer> SequencerPtr = GetSequencer();
+
+	if (SelectedObject && SelectedObject->IsA(UControlRig::StaticClass()) && SequencerPtr.IsValid())
+	{
+		UControlRig* ControlRig = CastChecked<UControlRig>(AssetData.GetAsset());
+
+		const FScopedTransaction Transaction(LOCTEXT("AddAnimation_Transaction", "Add Animation"));
+
+		for (FGuid ObjectBinding : ObjectBindings)
+		{
+			UObject* Object = SequencerPtr->FindSpawnedObjectOrTemplate(ObjectBinding);
+			int32 RowIndex = INDEX_NONE;
+			//AnimatablePropertyChanged(FOnKeyProperty::CreateRaw(this, &FControlRigParameterTrackEditor::AddControlRig, BoundObject, ObjectBindings[0], Skeleton)));
+		}
+	}
+}
+
+void FControlRigParameterTrackEditor::OnControlRigAssetEnterPressed(const TArray<FAssetData>& AssetData, TArray<FGuid> ObjectBindings, USkeleton* Skeleton)
+{
+	if (AssetData.Num() > 0)
+	{
+		OnControlRigAssetSelected(AssetData[0].GetAsset(), ObjectBindings, Skeleton);
+	}
+}
+
+*/
+
+void FControlRigParameterTrackEditor::AddControlRig(UClass* InClass, UObject* BoundSkelMesh, FGuid SkelMeshBinding, USkeleton* Skeleton)
+{
+	FSlateApplication::Get().DismissAllMenus();
+	const TSharedPtr<ISequencer> SequencerParent = GetSequencer();
+
+	if (InClass && InClass->IsChildOf(UControlRig::StaticClass()) && SequencerParent.IsValid())
+	{
+		FScopedTransaction AddControlRigTrackTransaction(LOCTEXT("AddControlRigTrack_Transaction", "Add Control Rig Track"));
+
+		UMovieSceneControlRigParameterTrack* Track = Cast<UMovieSceneControlRigParameterTrack>(AddTrack(GetSequencer()->GetFocusedMovieSceneSequence()->GetMovieScene(), SkelMeshBinding, UMovieSceneControlRigParameterTrack::StaticClass(), NAME_None));
+		if (Track)
+		{
+
+			FString ObjectName = (InClass->GetName());
+			ObjectName.RemoveFromEnd(TEXT("_C"));
+
+			UControlRig* ControlRig = NewObject<UControlRig>(Track, InClass, FName(*ObjectName), RF_Transactional);
+
+			ControlRig->SetObjectBinding(MakeShared<FControlRigSkeletalMeshBinding>());
+			ControlRig->GetObjectBinding()->BindToObject(BoundSkelMesh);
+			ControlRig->Initialize();
+			ControlRig->Execute(EControlRigState::Update);
+
+			SequencerParent->NotifyMovieSceneDataChanged(EMovieSceneDataChangeType::MovieSceneStructureItemsChanged);
+
+			Track->Modify();
+			UMovieSceneSection* NewSection = Track->CreateControlRigSection(0, ControlRig);
+
+			//mz todo need to have multiple rigs with same class
+			Track->SetTrackName(FName(*ObjectName));
+			Track->SetDisplayName(FText::FromString(ObjectName));
+
+			GetSequencer()->EmptySelection();
+			GetSequencer()->SelectSection(NewSection);
+			GetSequencer()->ThrobSectionSelection();
+			GetSequencer()->NotifyMovieSceneDataChanged(EMovieSceneDataChangeType::MovieSceneStructureItemAdded);
+
+			FControlRigEditMode* ControlRigEditMode = static_cast<FControlRigEditMode*>(GLevelEditorModeTools().GetActiveMode(FControlRigEditMode::ModeName));
+			if (!ControlRigEditMode)
+			{
+				GLevelEditorModeTools().ActivateMode(FControlRigEditMode::ModeName);
+				ControlRigEditMode = static_cast<FControlRigEditMode*>(GLevelEditorModeTools().GetActiveMode(FControlRigEditMode::ModeName));
+
+			}
+			if (ControlRigEditMode)
+			{
+				ControlRigEditMode->SetObjects(ControlRig, FGuid(), nullptr);
+			}
+			ControlRig->ControlModified().AddRaw(this, &FControlRigParameterTrackEditor::HandleControlModified);
+			ControlRig->ControlSelected().AddRaw(this, &FControlRigParameterTrackEditor::HandleControlSelected);
+		}
+	}
+}
+
+
+void FControlRigParameterTrackEditor::BindCommands(TSharedRef<FUICommandList> SequencerCommandBindings)
+{
+	const FControlRigTrackCommands& Commands = FControlRigTrackCommands::Get();
+
+	CommandBindings = MakeShared<FUICommandList>();
+	CommandBindings->MapAction(
+		Commands.AddTransformKey,
+		FExecuteAction::CreateSP(this, &FControlRigParameterTrackEditor::OnAddTransformKeysForSelectedObjects, EMovieSceneTransformChannel::All),
+		FCanExecuteAction::CreateSP(this, &FControlRigParameterTrackEditor::CanAddTransformKeysForSelectedObjects));
+
+	CommandBindings->MapAction(
+		Commands.AddTranslationKey,
+		FExecuteAction::CreateSP(this, &FControlRigParameterTrackEditor::OnAddTransformKeysForSelectedObjects, EMovieSceneTransformChannel::Translation),
+		FCanExecuteAction::CreateSP(this, &FControlRigParameterTrackEditor::CanAddTransformKeysForSelectedObjects));
+
+	CommandBindings->MapAction(
+		Commands.AddRotationKey,
+		FExecuteAction::CreateSP(this, &FControlRigParameterTrackEditor::OnAddTransformKeysForSelectedObjects, EMovieSceneTransformChannel::Rotation),
+		FCanExecuteAction::CreateSP(this, &FControlRigParameterTrackEditor::CanAddTransformKeysForSelectedObjects));
+
+	CommandBindings->MapAction(
+		Commands.AddScaleKey,
+		FExecuteAction::CreateSP(this, &FControlRigParameterTrackEditor::OnAddTransformKeysForSelectedObjects, EMovieSceneTransformChannel::Scale),
+		FCanExecuteAction::CreateSP(this, &FControlRigParameterTrackEditor::CanAddTransformKeysForSelectedObjects));
+
+	Commands.BindingCount++;
+
+	// Add these bindings to Sequencer
+	SequencerCommandBindings->Append(CommandBindings.ToSharedRef());
+
+	// Also add them to the Curve Editor 
+	GetSequencer()->GetCommandBindings(ESequencerCommandBindings::CurveEditor)->Append(CommandBindings.ToSharedRef());
+}
+
+bool FControlRigParameterTrackEditor::CanAddTransformKeysForSelectedObjects()
+{
+	if (!GetSequencer()->IsAllowedToChange())
+	{
+		return false;
+	}
+
+	FControlRigEditMode* ControlRigEditMode = static_cast<FControlRigEditMode*>(GLevelEditorModeTools().GetActiveMode(FControlRigEditMode::ModeName));
+	if (ControlRigEditMode && ControlRigEditMode->GetControlRig())
+	{
+		UControlRig* ControlRig = ControlRigEditMode->GetControlRig();
+		FString OurName = ControlRig->GetName();
+		FName Name(*OurName);
+		if (TSharedPtr<IControlRigObjectBinding> ObjectBinding = ControlRig->GetObjectBinding())
+		{
+			TArray<FName> ControlNames = ControlRig->CurrentControlSelection();
+			return (ControlNames.Num() > 0);
+		}
+	}
+	return false;
+}
+
+void FControlRigParameterTrackEditor::OnAddTransformKeysForSelectedObjects(EMovieSceneTransformChannel Channel)
+{
+	if (!GetSequencer()->IsAllowedToChange())
+	{
+		return;
+	}
+
+	FControlRigEditMode* ControlRigEditMode = static_cast<FControlRigEditMode*>(GLevelEditorModeTools().GetActiveMode(FControlRigEditMode::ModeName));
+	if (ControlRigEditMode && ControlRigEditMode->GetControlRig())
+	{
+		UControlRig* ControlRig = ControlRigEditMode->GetControlRig();
+		FString OurName = ControlRig->GetName();
+		FName Name(*OurName);
+		if (TSharedPtr<IControlRigObjectBinding> ObjectBinding = ControlRig->GetObjectBinding())
+		{
+			TArray<FName> ControlNames = ControlRig->CurrentControlSelection();
+			for (const FName& ControlName : ControlNames)
+			{
+				USceneComponent* Component = Cast<USkeletalMeshComponent>(ObjectBinding->GetBoundObject());
+				if (Component)
+				{
+					UObject *ActorObject = Component->GetOwner();
+					AddControlKeys(ActorObject, ControlRig, Name, ControlName, Channel, ESequencerKeyMode::ManualKey);
+				}
+			}
+		}
+	}
+}
+
+void FControlRigParameterTrackEditor::OnSequencerDataChanged(EMovieSceneDataChangeType DataChangeType)
+{
+	UMovieScene* MovieScene = GetSequencer()->GetFocusedMovieSceneSequence()->GetMovieScene();
+	FControlRigEditMode* ControlRigEditMode = static_cast<FControlRigEditMode*>(GLevelEditorModeTools().GetActiveMode(FControlRigEditMode::ModeName));
+	
+	//if we have a valid control rig edit mode need to check and see the control rig in that mode is still in a track
+	//if not we get rid of it.
+	if (ControlRigEditMode && ControlRigEditMode->GetControlRig() != nullptr && MovieScene && (DataChangeType == EMovieSceneDataChangeType::MovieSceneStructureItemRemoved ||
+		DataChangeType == EMovieSceneDataChangeType::Unknown))
+	{
+		const TArray<FMovieSceneBinding>& Bindings = MovieScene->GetBindings();
+		for (const FMovieSceneBinding& Binding : Bindings)
+		{
+			UMovieSceneControlRigParameterTrack* Track = Cast<UMovieSceneControlRigParameterTrack>(MovieScene->FindTrack(UMovieSceneControlRigParameterTrack::StaticClass(), Binding.GetObjectGuid(), NAME_None));
+			if (Track && Track->GetControlRig() == ControlRigEditMode->GetControlRig())
+			{
+				return; //just exit out we still have a good track
+			}
+		}
+		//okay no good track so deactive it and delete it's Control Rig and bingings.
+		if (GLevelEditorModeTools().HasToolkitHost())
+		{
+			GLevelEditorModeTools().DeactivateMode(FControlRigEditMode::ModeName);
+		}
+		ControlRigEditMode->SetObjects(nullptr, FGuid(), nullptr);
+	}
+}
+
+void FControlRigParameterTrackEditor::OnSelectionChanged(TArray<UMovieSceneTrack*> InTracks)
+{
+	if (bIsDoingASelection)
+	{
+		return;
+	}
+
+	TGuardValue<bool> Guard(bIsDoingASelection, true);
+	TArray<FString> StringArray;
+	FControlRigEditMode* ControlRigEditMode = static_cast<FControlRigEditMode*>(GLevelEditorModeTools().GetActiveMode(FControlRigEditMode::ModeName));
+	UControlRig* ControlRig = nullptr;
+	//Always clear the control rig(s) in the edit mode.
+	if (ControlRigEditMode)
+	{
+		ControlRig = ControlRigEditMode->GetControlRig();
+		if (ControlRig)
+		{
+			ControlRig->ClearControlSelection();
+		}
+	}
+	TArray<const IKeyArea*> KeyAreas;
+	GetSequencer()->GetSelectedKeyAreas(KeyAreas);
+	if (KeyAreas.Num() <= 0)
+	{
+		return;
+	}
+
+	TMap<UControlRig *, TSet<FName>> RigsAndControls;
+	for (const IKeyArea* KeyArea : KeyAreas)
+	{
+		UMovieSceneControlRigParameterSection* MovieSection = Cast<UMovieSceneControlRigParameterSection>(KeyArea->GetOwningSection());
+		if (MovieSection)
+		{
+			ControlRig = MovieSection->ControlRig;
+			//Only create the edit mode if we have a KeyAra selected and it's not set and we have some boundobjects.
+			if (!ControlRigEditMode )
+			{
+				GLevelEditorModeTools().ActivateMode(FControlRigEditMode::ModeName);
+				ControlRigEditMode = static_cast<FControlRigEditMode*>(GLevelEditorModeTools().GetActiveMode(FControlRigEditMode::ModeName));
+				if (TSharedPtr<IControlRigObjectBinding> ObjectBinding = ControlRig->GetObjectBinding())
+				{
+					if (ControlRigEditMode)
+					{
+						ControlRigEditMode->SetObjects(ControlRig, FGuid(), nullptr);
+					}
+				}
+			}
+			else
+			{
+				if (ControlRigEditMode->GetControlRig() != ControlRig)
+				{
+					ControlRigEditMode->SetObjects(ControlRig, FGuid(), nullptr);
+				}
+			}
+			
+			const FMovieSceneChannelMetaData* MetaData = KeyArea->GetChannel().GetMetaData();
+			if (MetaData)
+			{
+				StringArray.SetNum(0);
+				FString String = MetaData->Name.ToString();
+				String.ParseIntoArray(StringArray, TEXT("."));
+				if (StringArray.Num() > 0)
+				{
+					FName ControlName(*StringArray[0]);
+					RigsAndControls.FindOrAdd(ControlRig).Add(ControlName);
+				}
+			}
+		}
+	}
+	for (TPair<UControlRig *, TSet<FName>> Pair : RigsAndControls)
+	{
+		Pair.Key->ClearControlSelection();
+		for (const FName Name : Pair.Value)
+		{
+			Pair.Key->SelectControl(Name, true);
+		}
+	}
+}
+
+
+void FControlRigParameterTrackEditor::HandleControlSelected(IControlRigManipulatable* ControlRigManp, const FRigControl& Control, bool bSelected)
+{
+	if (!bIsDoingASelection)
+	{
+		TGuardValue<bool> Guard(bIsDoingASelection, true);
+		FString OurName = ControlRigManp->GetName();
+		FName ControlRigName(*OurName);
+		if (TSharedPtr<IControlRigObjectBinding> ObjectBinding = ControlRigManp->GetObjectBinding())
+		{
+			UObject *ActorObject = nullptr;
+			USceneComponent* Component = Cast<USkeletalMeshComponent>(ObjectBinding->GetBoundObject());
+			if (!Component)
+			{
+				return;
+			}
+			ActorObject = Component->GetOwner();
+			bool bCreateTrack = false;
+			bool bCreateHandle = false;
+			FFindOrCreateHandleResult HandleResult = FindOrCreateHandleToObject(ActorObject, bCreateHandle, NAME_None);
+			FGuid ObjectHandle = HandleResult.Handle;
+			if (!ObjectHandle.IsValid())
+			{
+				return;
+			}
+			FFindOrCreateTrackResult TrackResult = FindOrCreateTrackForObject(ObjectHandle, UMovieSceneControlRigParameterTrack::StaticClass(), ControlRigName, bCreateTrack);
+			UMovieSceneControlRigParameterTrack* Track = CastChecked<UMovieSceneControlRigParameterTrack>(TrackResult.Track, ECastCheckedType::NullAllowed);
+			if (Track)
+			{
+				float Weight = 1.0f;
+				TArray<IKeyArea> KeyAreas;
+				TArray<FString> StringArray;
+
+				for (UMovieSceneSection* Section : Track->GetAllSections())
+				{
+					UMovieSceneControlRigParameterSection* ParamSection = Cast<UMovieSceneControlRigParameterSection>(Section);
+					if (ParamSection)
+					{
+						FMovieSceneChannelProxy& ChannelProxy = Section->GetChannelProxy();
+						for (const FMovieSceneChannelEntry& Entry : Section->GetChannelProxy().GetAllEntries())
+						{
+							const FName ChannelTypeName = Entry.GetChannelTypeName();
+
+							// One editor data ptr per channel
+							TArrayView<FMovieSceneChannel* const>        Channels = Entry.GetChannels();
+							TArrayView<const FMovieSceneChannelMetaData> AllMetaData = Entry.GetMetaData();
+
+							for (int32 Index = 0; Index < Channels.Num(); ++Index)
+							{
+								FMovieSceneChannelHandle ChannelHandle = ChannelProxy.MakeHandle(ChannelTypeName, Index);
+								const FMovieSceneChannelMetaData& MetaData = AllMetaData[Index];
+								FString String = MetaData.Name.ToString();
+								String.ParseIntoArray(StringArray, TEXT("."));
+								if (StringArray.Num() > 0)
+								{
+									FName ControlName(*StringArray[0]);
+									if (ControlName == Control.Name)
+									{
+										IKeyArea KeyArea(Section, ChannelHandle);
+										KeyAreas.Add(KeyArea);
+									}
+								}
+							}
+						}
+					}
+				}
+				if (KeyAreas.Num())
+				{
+					GetSequencer()->SelectByKeyAreas(KeyAreas, true, bSelected);
+				}
+			}
+		}
+	}
+}
+
+void FControlRigParameterTrackEditor::HandleControlModified(IControlRigManipulatable* ControlRigManp, const FRigControl& Control)
+{
+	if (GetSequencer().IsValid() && !GetSequencer()->IsAllowedToChange())
+	{
+		return;
+	}
+	FString OurName = ControlRigManp->GetName();
+	FName Name(*OurName);
+	if (TSharedPtr<IControlRigObjectBinding> ObjectBinding = ControlRigManp->GetObjectBinding())
+	{
+		UObject *ActorObject = nullptr;
+		USceneComponent* Component = Cast<USkeletalMeshComponent>(ObjectBinding->GetBoundObject());
+		if (Component)
+		{
+			ActorObject = Component->GetOwner();
+			AddControlKeys(ActorObject, ControlRigManp, Name, Control.Name, EMovieSceneTransformChannel::All, ESequencerKeyMode::AutoKey);
+		}
+	}
+}
+
+void FControlRigParameterTrackEditor::GetControlRigKeys(IControlRigManipulatable* Manip, FName ParameterName, EMovieSceneTransformChannel ChannelsToKey, FGeneratedTrackKeys& OutGeneratedKeys)
+{
+	const TArray<FRigControl>& Controls = Manip->AvailableControls();
+	// If key all is enabled, for a key on all the channels
+	if (GetSequencer()->GetKeyGroupMode() == EKeyGroupMode::KeyAll)
+	{
+		ChannelsToKey = EMovieSceneTransformChannel::All;
+	}
+
+	int32 ChannelIndex = 0;
+	for (const FRigControl& RigControl : Controls)
+	{
+		switch (RigControl.ControlType)
+		{
+		case ERigControlType::Float:
+		{
+			float Val = RigControl.Value.Get<float>();
+			OutGeneratedKeys.Add(FMovieSceneChannelValueSetter::Create<FMovieSceneFloatChannel>(ChannelIndex++, Val, true));
+			break;
+		}
+		case ERigControlType::Position:
+		case ERigControlType::Scale:
+		case ERigControlType::Rotator:
+		{
+			FVector Val = RigControl.Value.Get<FVector>();
+			OutGeneratedKeys.Add(FMovieSceneChannelValueSetter::Create<FMovieSceneFloatChannel>(ChannelIndex++, Val.X, true));
+			OutGeneratedKeys.Add(FMovieSceneChannelValueSetter::Create<FMovieSceneFloatChannel>(ChannelIndex++, Val.Y, true));
+			OutGeneratedKeys.Add(FMovieSceneChannelValueSetter::Create<FMovieSceneFloatChannel>(ChannelIndex++, Val.Z, true));
+			break;
+		}
+
+		case ERigControlType::Transform:
+		{
+			FTransform Val = RigControl.Value.Get<FTransform>();
+			FVector CurrentVector = Val.GetTranslation();
+			bool bKeyX = EnumHasAnyFlags(ChannelsToKey, EMovieSceneTransformChannel::TranslationX);
+			bool bKeyY = EnumHasAnyFlags(ChannelsToKey, EMovieSceneTransformChannel::TranslationY);
+			bool bKeyZ = EnumHasAnyFlags(ChannelsToKey, EMovieSceneTransformChannel::TranslationZ);
+			if (GetSequencer()->GetKeyGroupMode() == EKeyGroupMode::KeyGroup && (bKeyX || bKeyY || bKeyZ))
+			{
+				bKeyX = bKeyY = bKeyZ = true;
+			}
+			OutGeneratedKeys.Add(FMovieSceneChannelValueSetter::Create<FMovieSceneFloatChannel>(ChannelIndex++, CurrentVector.X, bKeyX));
+			OutGeneratedKeys.Add(FMovieSceneChannelValueSetter::Create<FMovieSceneFloatChannel>(ChannelIndex++, CurrentVector.Y, bKeyY));
+			OutGeneratedKeys.Add(FMovieSceneChannelValueSetter::Create<FMovieSceneFloatChannel>(ChannelIndex++, CurrentVector.Z, bKeyZ));
+
+			CurrentVector = Val.GetRotation().Euler();
+			bKeyX = EnumHasAnyFlags(ChannelsToKey, EMovieSceneTransformChannel::RotationX);
+			bKeyY = EnumHasAnyFlags(ChannelsToKey, EMovieSceneTransformChannel::RotationY);
+			bKeyZ = EnumHasAnyFlags(ChannelsToKey, EMovieSceneTransformChannel::RotationZ);
+			if (GetSequencer()->GetKeyGroupMode() == EKeyGroupMode::KeyGroup && (bKeyX || bKeyY || bKeyZ))
+			{
+				bKeyX = bKeyY = bKeyZ = true;
+			}
+			OutGeneratedKeys.Add(FMovieSceneChannelValueSetter::Create<FMovieSceneFloatChannel>(ChannelIndex++, CurrentVector.X, bKeyX));
+			OutGeneratedKeys.Add(FMovieSceneChannelValueSetter::Create<FMovieSceneFloatChannel>(ChannelIndex++, CurrentVector.Y, bKeyY));
+			OutGeneratedKeys.Add(FMovieSceneChannelValueSetter::Create<FMovieSceneFloatChannel>(ChannelIndex++, CurrentVector.Z, bKeyZ));
+
+			CurrentVector = Val.GetScale3D();
+			bKeyX = EnumHasAnyFlags(ChannelsToKey, EMovieSceneTransformChannel::ScaleX);
+			bKeyY = EnumHasAnyFlags(ChannelsToKey, EMovieSceneTransformChannel::ScaleY);
+			bKeyZ = EnumHasAnyFlags(ChannelsToKey, EMovieSceneTransformChannel::ScaleZ);
+			if (GetSequencer()->GetKeyGroupMode() == EKeyGroupMode::KeyGroup && (bKeyX || bKeyY || bKeyZ))
+			{
+				bKeyX = bKeyY = bKeyZ = true;
+			}
+			OutGeneratedKeys.Add(FMovieSceneChannelValueSetter::Create<FMovieSceneFloatChannel>(ChannelIndex++, CurrentVector.X, bKeyX));
+			OutGeneratedKeys.Add(FMovieSceneChannelValueSetter::Create<FMovieSceneFloatChannel>(ChannelIndex++, CurrentVector.Y, bKeyY));
+			OutGeneratedKeys.Add(FMovieSceneChannelValueSetter::Create<FMovieSceneFloatChannel>(ChannelIndex++, CurrentVector.Z, bKeyZ));
+
+			break;
+		}
+		//mz todo the other types
+		}
+	}
+}
+
+FKeyPropertyResult FControlRigParameterTrackEditor::AddKeysToControlRigHandle(UObject *Object, IControlRigManipulatable* Manip,
+	FGuid ObjectHandle, FFrameNumber KeyTime, FGeneratedTrackKeys& GeneratedKeys,
+	ESequencerKeyMode KeyMode, TSubclassOf<UMovieSceneTrack> TrackClass, FName ControlRigName)
+{
+	EAutoChangeMode AutoChangeMode = GetSequencer()->GetAutoChangeMode();
+	EAllowEditsMode AllowEditsMode = GetSequencer()->GetAllowEditsMode();
+
+	bool bCreateTrack =
+		(KeyMode == ESequencerKeyMode::AutoKey && (AutoChangeMode == EAutoChangeMode::AutoTrack || AutoChangeMode == EAutoChangeMode::All)) ||
+		KeyMode == ESequencerKeyMode::ManualKey ||
+		KeyMode == ESequencerKeyMode::ManualKeyForced ||
+		AllowEditsMode == EAllowEditsMode::AllowSequencerEditsOnly;
+
+	bool bCreateSection = false;
+		// we don't do this, maybe revisit if a bug occurs, but currently extends sections on autokey.
+	    //bCreateTrack || (KeyMode == ESequencerKeyMode::AutoKey && (AutoChangeMode != EAutoChangeMode::None));
+
+	// Try to find an existing Track, and if one doesn't exist check the key params and create one if requested.
+
+	FFindOrCreateTrackResult TrackResult = FindOrCreateTrackForObject(ObjectHandle, TrackClass, ControlRigName, bCreateTrack);
+	UMovieSceneControlRigParameterTrack* Track = CastChecked<UMovieSceneControlRigParameterTrack>(TrackResult.Track, ECastCheckedType::NullAllowed);
+
+	bool bTrackCreated = TrackResult.bWasCreated;
+
+	bool bSectionCreated = false;
+	FKeyPropertyResult KeyPropertyResult;
+
+	if (Track)
+	{
+		float Weight = 1.0f;
+
+		UMovieSceneSection* SectionToKey = bCreateSection ? Track->FindOrExtendSection(KeyTime, Weight) : Track->FindSection(KeyTime);
+
+		// If there's no overlapping section to key, create one only if a track was newly created. Otherwise, skip keying altogether
+		// so that the user is forced to create a section to key on.
+		if (bTrackCreated && !SectionToKey)
+		{
+			Track->Modify();
+			SectionToKey = Track->FindOrAddSection(KeyTime, bSectionCreated);
+			if (bSectionCreated && GetSequencer()->GetInfiniteKeyAreas())
+			{
+				SectionToKey->SetRange(TRange<FFrameNumber>::All());
+			}
+		}
+
+		if (SectionToKey && SectionToKey->GetRange().Contains(KeyTime))
+		{
+			if (!bTrackCreated)
+			{
+				ModifyOurGeneratedKeysByCurrentAndWeight(Object, Manip, Track, SectionToKey, KeyTime, GeneratedKeys, Weight);
+			}
+			UMovieSceneControlRigParameterSection* ParamSection = Cast<UMovieSceneControlRigParameterSection>(SectionToKey);
+			if (!ParamSection->GetDoNotKey())
+			{
+				KeyPropertyResult |= AddKeysToSection(SectionToKey, KeyTime, GeneratedKeys, KeyMode);
+			}
+		}
+	}
+
+	KeyPropertyResult.bTrackCreated |= bTrackCreated || bSectionCreated;
+	return KeyPropertyResult;
+}
+
+FKeyPropertyResult FControlRigParameterTrackEditor::AddKeysToControlRig(
+	UObject* Object, IControlRigManipulatable* Manip, FFrameNumber KeyTime, FGeneratedTrackKeys& GeneratedKeys,
+	ESequencerKeyMode KeyMode, TSubclassOf<UMovieSceneTrack> TrackClass, FName PropertyName)
+{
+	FKeyPropertyResult KeyPropertyResult;
+	EAutoChangeMode AutoChangeMode = GetSequencer()->GetAutoChangeMode();
+	EAllowEditsMode AllowEditsMode = GetSequencer()->GetAllowEditsMode();
+	bool bCreateHandle =
+		(KeyMode == ESequencerKeyMode::AutoKey && (AutoChangeMode == EAutoChangeMode::All)) ||
+		KeyMode == ESequencerKeyMode::ManualKey ||
+		KeyMode == ESequencerKeyMode::ManualKeyForced ||
+		AllowEditsMode == EAllowEditsMode::AllowSequencerEditsOnly;
+
+	FFindOrCreateHandleResult HandleResult = FindOrCreateHandleToObject(Object, bCreateHandle, NAME_None);
+	FGuid ObjectHandle = HandleResult.Handle;
+	KeyPropertyResult.bHandleCreated = HandleResult.bWasCreated;
+	if (ObjectHandle.IsValid())
+	{
+		KeyPropertyResult |= AddKeysToControlRigHandle(Object, Manip, ObjectHandle, KeyTime, GeneratedKeys, KeyMode, TrackClass, PropertyName);
+	}
+
+	return KeyPropertyResult;
+}
+
+void FControlRigParameterTrackEditor::AddControlKeys(UObject* InObject, IControlRigManipulatable* Manip, FName ControlRigName, FName RigControlName, EMovieSceneTransformChannel ChannelsToKey, ESequencerKeyMode KeyMode)
+{
+	if (GetSequencer().IsValid() && !GetSequencer()->IsAllowedToChange())
+	{
+		return;
+	}
+	bool bCreateTrack = false;
+	bool bCreateHandle = false;
+	FFindOrCreateHandleResult HandleResult = FindOrCreateHandleToObject(InObject, bCreateHandle, NAME_None);
+	FGuid ObjectHandle = HandleResult.Handle;
+	if (!ObjectHandle.IsValid())
+	{
+		return;
+	}
+	FFindOrCreateTrackResult TrackResult = FindOrCreateTrackForObject(ObjectHandle, UMovieSceneControlRigParameterTrack::StaticClass(), ControlRigName, bCreateTrack);
+	UMovieSceneControlRigParameterTrack* Track = CastChecked<UMovieSceneControlRigParameterTrack>(TrackResult.Track, ECastCheckedType::NullAllowed);
+	if (Track)
+	{
+		FFrameNumber  FrameTime = GetTimeForKey();
+		UMovieSceneSection* Section = Track->FindSection(FrameTime);
+		UMovieSceneControlRigParameterSection* ParamSection = Cast<UMovieSceneControlRigParameterSection>(Section);
+
+		if (ParamSection && ParamSection->GetDoNotKey())
+		{
+			return;
+		}
+	}
+
+	TSharedRef<FGeneratedTrackKeys> GeneratedKeys = MakeShared<FGeneratedTrackKeys>();
+
+	GetControlRigKeys(Manip, RigControlName, ChannelsToKey, *GeneratedKeys);
+
+	auto OnKeyProperty = [=](FFrameNumber Time) -> FKeyPropertyResult
+	{
+		return this->AddKeysToControlRig(InObject, Manip, Time, *GeneratedKeys, KeyMode, UMovieSceneControlRigParameterTrack::StaticClass(), ControlRigName);
+	};
+
+	AnimatablePropertyChanged(FOnKeyProperty::CreateLambda(OnKeyProperty));
+}
+
+bool FControlRigParameterTrackEditor::ModifyOurGeneratedKeysByCurrentAndWeight(UObject *Object, IControlRigManipulatable* Manip, UMovieSceneTrack *Track, UMovieSceneSection* SectionToKey, FFrameNumber KeyTime, FGeneratedTrackKeys& GeneratedTotalKeys, float Weight) const
+{
+	FFrameRate TickResolution = GetSequencer()->GetFocusedTickResolution();
+	FMovieSceneEvaluationTrack EvalTrack = Track->GenerateTrackTemplate();
+
+	FMovieSceneInterrogationData InterrogationData;
+	GetSequencer()->GetEvaluationTemplate().CopyActuators(InterrogationData.GetAccumulator());
+
+	FMovieSceneContext Context(FMovieSceneEvaluationRange(KeyTime, GetSequencer()->GetFocusedTickResolution()));
+	EvalTrack.Interrogate(Context, InterrogationData, Object);
+	const TArray<FRigControl>& Controls = Manip->AvailableControls();
+
+	FMovieSceneChannelProxy& Proxy = SectionToKey->GetChannelProxy();
+	int32 ChannelIndex = 0;
+	for (const FRigControl& RigControl : Controls)
+	{
+		switch (RigControl.ControlType)
+		{
+		case ERigControlType::Float:
+		{
+			for (const FFloatInterrogationData& Val : InterrogationData.Iterate<FFloatInterrogationData>(UMovieSceneControlRigParameterSection::GetFloatInterrogationKey()))
+			{
+				if (Val.ParameterName == RigControl.Name)
+				{
+					GeneratedTotalKeys[ChannelIndex]->ModifyByCurrentAndWeight(Proxy, KeyTime, (void *)&Val.Val, Weight);
+					break;
+				}
+			}
+			++ChannelIndex;
+			break;
+		}
+		case ERigControlType::Position:
+		case ERigControlType::Scale:
+		case ERigControlType::Rotator:
+		{
+			for (const FVectorInterrogationData& Val : InterrogationData.Iterate<FVectorInterrogationData>(UMovieSceneControlRigParameterSection::GetVectorInterrogationKey()))
+			{
+				if (Val.ParameterName == RigControl.Name)
+				{
+					GeneratedTotalKeys[ChannelIndex]->ModifyByCurrentAndWeight(Proxy, KeyTime, (void *)&Val.Val, Weight);
+					GeneratedTotalKeys[ChannelIndex + 1]->ModifyByCurrentAndWeight(Proxy, KeyTime, (void *)&Val.Val, Weight);
+					GeneratedTotalKeys[ChannelIndex + 2]->ModifyByCurrentAndWeight(Proxy, KeyTime, (void *)&Val.Val, Weight);
+
+					break;
+				}
+			}
+			ChannelIndex += 3;
+
+			break;
+		}
+
+		case ERigControlType::Transform:
+		{
+			for (const FTransformInterrogationData& Val : InterrogationData.Iterate<FTransformInterrogationData>(UMovieSceneControlRigParameterSection::GetTransformInterrogationKey()))
+			{
+				if (Val.ParameterName == RigControl.Name)
+				{
+					FVector CurrentPos = Val.Val.GetTranslation();
+					FRotator CurrentRot = Val.Val.GetRotation().Rotator();
+					FVector CurrentScale = Val.Val.GetScale3D();
+					GeneratedTotalKeys[ChannelIndex]->ModifyByCurrentAndWeight(Proxy, KeyTime, (void *)&CurrentPos.X, Weight);
+					GeneratedTotalKeys[ChannelIndex + 1]->ModifyByCurrentAndWeight(Proxy, KeyTime, (void *)&CurrentPos.Y, Weight);
+					GeneratedTotalKeys[ChannelIndex + 2]->ModifyByCurrentAndWeight(Proxy, KeyTime, (void *)&CurrentPos.Z, Weight);
+					GeneratedTotalKeys[ChannelIndex + 3]->ModifyByCurrentAndWeight(Proxy, KeyTime, (void *)&CurrentRot.Roll, Weight);
+					GeneratedTotalKeys[ChannelIndex + 4]->ModifyByCurrentAndWeight(Proxy, KeyTime, (void *)&CurrentRot.Pitch, Weight);
+					GeneratedTotalKeys[ChannelIndex + 5]->ModifyByCurrentAndWeight(Proxy, KeyTime, (void *)&CurrentRot.Yaw, Weight);
+					GeneratedTotalKeys[ChannelIndex + 6]->ModifyByCurrentAndWeight(Proxy, KeyTime, (void *)&CurrentScale.X, Weight);
+					GeneratedTotalKeys[ChannelIndex + 7]->ModifyByCurrentAndWeight(Proxy, KeyTime, (void *)&CurrentScale.Y, Weight);
+					GeneratedTotalKeys[ChannelIndex + 8]->ModifyByCurrentAndWeight(Proxy, KeyTime, (void *)&CurrentScale.Z, Weight);
+
+					break;
+				}
+			}
+			ChannelIndex += 9;
+
+			break;
+		}
+		//mz todo the other types
+		}
+	}
+	return true;
+}
+
+void FControlRigParameterSection::BuildSectionContextMenu(FMenuBuilder& MenuBuilder, const FGuid& InObjectBinding)
+{
+	UMovieSceneControlRigParameterSection* ParameterSection = CastChecked<UMovieSceneControlRigParameterSection>(WeakSection.Get());
+	TSharedPtr<ISequencer> SequencerPtr = WeakSequencer.Pin();
+
+	UControlRig* ControlRig = ParameterSection->ControlRig;
+	if (ControlRig)
+	{
+		const TArray<FRigControl>& Controls = ControlRig->AvailableControls();
+
+		auto MakeUIAction = [=](EMovieSceneTransformChannel ChannelsToToggle)
+		{
+			return FUIAction(
+				FExecuteAction::CreateLambda([=]
+			{
+				FScopedTransaction Transaction(LOCTEXT("SetActiveChannelsTransaction", "Set Active Channels"));
+				ParameterSection->Modify();
+				EMovieSceneTransformChannel Channels = ParameterSection->GetTransformMask().GetChannels();
+
+				if (EnumHasAllFlags(Channels, ChannelsToToggle) || (Channels & ChannelsToToggle) == EMovieSceneTransformChannel::None)
+				{
+					ParameterSection->SetTransformMask(ParameterSection->GetTransformMask().GetChannels() ^ ChannelsToToggle);
+				}
+				else
+				{
+					ParameterSection->SetTransformMask(ParameterSection->GetTransformMask().GetChannels() | ChannelsToToggle);
+				}
+
+				// Restore pre-animated state for the bound objects so that inactive channels will return to their default values.
+				for (TWeakObjectPtr<> WeakObject : SequencerPtr->FindBoundObjects(InObjectBinding, SequencerPtr->GetFocusedTemplateID()))
+				{
+					if (UObject* Object = WeakObject.Get())
+					{
+						SequencerPtr->RestorePreAnimatedState();
+					}
+				}
+
+				SequencerPtr->NotifyMovieSceneDataChanged(EMovieSceneDataChangeType::MovieSceneStructureItemsChanged);
+			}
+				),
+				FCanExecuteAction(),
+				FGetActionCheckState::CreateLambda([=]
+			{
+				EMovieSceneTransformChannel Channels = ParameterSection->GetTransformMask().GetChannels();
+				if (EnumHasAllFlags(Channels, ChannelsToToggle))
+				{
+					return ECheckBoxState::Checked;
+				}
+				else if (EnumHasAnyFlags(Channels, ChannelsToToggle))
+				{
+					return ECheckBoxState::Undetermined;
+				}
+				return ECheckBoxState::Unchecked;
+			})
+				);
+		};
+		auto ToggleControls = [=](int32 Index)
+		{
+			return FUIAction(
+				FExecuteAction::CreateLambda([=]
+			{
+				FScopedTransaction Transaction(LOCTEXT("RigSectionFilter", "Toggle Rig Control Filters"));
+				ParameterSection->Modify();
+				if (Index >= 0)
+				{
+					ParameterSection->SetControlsMask(Index, !ParameterSection->GetControlsMask(Index));
+				}
+				else
+				{
+					ParameterSection->FillControlsMask(!ParameterSection->GetControlsMask(0));
+				}
+				SequencerPtr->NotifyMovieSceneDataChanged(EMovieSceneDataChangeType::MovieSceneStructureItemsChanged);
+
+			}
+				),
+				FCanExecuteAction(),
+				FGetActionCheckState::CreateLambda([=]
+			{
+				TArray<bool> ControlBool = ParameterSection->GetControlsMask();
+				if (Index >= 0)
+				{
+					if (ControlBool[Index])
+					{
+						return ECheckBoxState::Checked;
+
+					}
+					else
+					{
+						return ECheckBoxState::Unchecked;
+					}
+				}
+				else
+				{
+					TOptional<bool> FirstVal;
+					for (bool Val : ControlBool)
+					{
+						if (FirstVal.IsSet())
+						{
+							if (Val != FirstVal)
+							{
+								return ECheckBoxState::Undetermined;
+							}
+						}
+						else
+						{
+							FirstVal = Val;
+						}
+
+					}
+					return (FirstVal.GetValue()) ? ECheckBoxState::Checked : ECheckBoxState::Unchecked;
+				}
+			})
+				);
+		};
+		MenuBuilder.BeginSection(NAME_None, LOCTEXT("RigSectionFilter", "Filter Controls"));
+		{
+			MenuBuilder.AddSubMenu(
+				LOCTEXT("ToggleRigControlsText", "Toggle Rig Controls"), LOCTEXT("ToggleRigControlsText_Tooltip", "Toggle Rig Controls"),
+				FNewMenuDelegate::CreateLambda([=](FMenuBuilder& SubMenuBuilder) {
+				int32 Index = 0;
+				for (const FRigControl& RigControl : Controls)
+				{
+					const FName RigName = RigControl.Name;
+					FText Name = FText::FromName(RigName);
+					FText Text = FText::Format(LOCTEXT("RigControlToggle", "{0}"), Name);
+					FText TooltipText = FText::Format(LOCTEXT("RigControlToggleTooltip", "Toggle Rig Control {0}"), Name);
+					SubMenuBuilder.AddMenuEntry(
+						Text, TooltipText,
+						FSlateIcon(), ToggleControls(Index++), NAME_None, EUserInterfaceActionType::ToggleButton);
+				}
+			}),
+				ToggleControls(-1),
+				NAME_None,
+				EUserInterfaceActionType::ToggleButton);
+
+			MenuBuilder.AddSubMenu(
+				LOCTEXT("AllTranslation", "Translation"), LOCTEXT("AllTranslation_ToolTip", "Causes this section to affect the translation of rig control transforms"),
+				FNewMenuDelegate::CreateLambda([=](FMenuBuilder& SubMenuBuilder) {
+				SubMenuBuilder.AddMenuEntry(
+					LOCTEXT("TranslationX", "X"), LOCTEXT("TranslationX_ToolTip", "Causes this section to affect the X channel of the transform's translation"),
+					FSlateIcon(), MakeUIAction(EMovieSceneTransformChannel::TranslationX), NAME_None, EUserInterfaceActionType::ToggleButton);
+				SubMenuBuilder.AddMenuEntry(
+					LOCTEXT("TranslationY", "Y"), LOCTEXT("TranslationY_ToolTip", "Causes this section to affect the Y channel of the transform's translation"),
+					FSlateIcon(), MakeUIAction(EMovieSceneTransformChannel::TranslationY), NAME_None, EUserInterfaceActionType::ToggleButton);
+				SubMenuBuilder.AddMenuEntry(
+					LOCTEXT("TranslationZ", "Z"), LOCTEXT("TranslationZ_ToolTip", "Causes this section to affect the Z channel of the transform's translation"),
+					FSlateIcon(), MakeUIAction(EMovieSceneTransformChannel::TranslationZ), NAME_None, EUserInterfaceActionType::ToggleButton);
+			}),
+				MakeUIAction(EMovieSceneTransformChannel::Translation),
+				NAME_None,
+				EUserInterfaceActionType::ToggleButton);
+
+			MenuBuilder.AddSubMenu(
+				LOCTEXT("AllRotation", "Rotation"), LOCTEXT("AllRotation_ToolTip", "Causes this section to affect the rotation of the rig control transform"),
+				FNewMenuDelegate::CreateLambda([=](FMenuBuilder& SubMenuBuilder) {
+				SubMenuBuilder.AddMenuEntry(
+					LOCTEXT("RotationX", "Roll (X)"), LOCTEXT("RotationX_ToolTip", "Causes this section to affect the roll (X) channel the transform's rotation"),
+					FSlateIcon(), MakeUIAction(EMovieSceneTransformChannel::RotationX), NAME_None, EUserInterfaceActionType::ToggleButton);
+				SubMenuBuilder.AddMenuEntry(
+					LOCTEXT("RotationY", "Pitch (Y)"), LOCTEXT("RotationY_ToolTip", "Causes this section to affect the pitch (Y) channel the transform's rotation"),
+					FSlateIcon(), MakeUIAction(EMovieSceneTransformChannel::RotationY), NAME_None, EUserInterfaceActionType::ToggleButton);
+				SubMenuBuilder.AddMenuEntry(
+					LOCTEXT("RotationZ", "Yaw (Z)"), LOCTEXT("RotationZ_ToolTip", "Causes this section to affect the yaw (Z) channel the transform's rotation"),
+					FSlateIcon(), MakeUIAction(EMovieSceneTransformChannel::RotationZ), NAME_None, EUserInterfaceActionType::ToggleButton);
+			}),
+				MakeUIAction(EMovieSceneTransformChannel::Rotation),
+				NAME_None,
+				EUserInterfaceActionType::ToggleButton);
+
+			MenuBuilder.AddSubMenu(
+				LOCTEXT("AllScale", "Scale"), LOCTEXT("AllScale_ToolTip", "Causes this section to affect the scale of the rig control transform"),
+				FNewMenuDelegate::CreateLambda([=](FMenuBuilder& SubMenuBuilder) {
+				SubMenuBuilder.AddMenuEntry(
+					LOCTEXT("ScaleX", "X"), LOCTEXT("ScaleX_ToolTip", "Causes this section to affect the X channel of the transform's scale"),
+					FSlateIcon(), MakeUIAction(EMovieSceneTransformChannel::ScaleX), NAME_None, EUserInterfaceActionType::ToggleButton);
+				SubMenuBuilder.AddMenuEntry(
+					LOCTEXT("ScaleY", "Y"), LOCTEXT("ScaleY_ToolTip", "Causes this section to affect the Y channel of the transform's scale"),
+					FSlateIcon(), MakeUIAction(EMovieSceneTransformChannel::ScaleY), NAME_None, EUserInterfaceActionType::ToggleButton);
+				SubMenuBuilder.AddMenuEntry(
+					LOCTEXT("ScaleZ", "Z"), LOCTEXT("ScaleZ_ToolTip", "Causes this section to affect the Z channel of the transform's scale"),
+					FSlateIcon(), MakeUIAction(EMovieSceneTransformChannel::ScaleZ), NAME_None, EUserInterfaceActionType::ToggleButton);
+			}),
+				MakeUIAction(EMovieSceneTransformChannel::Scale),
+				NAME_None,
+				EUserInterfaceActionType::ToggleButton);
+
+			//mz todo h
+			MenuBuilder.AddMenuEntry(
+				LOCTEXT("Weight", "Weight"), LOCTEXT("Weight_ToolTip", "Causes this section to be applied with a user-specified weight curve"),
+				FSlateIcon(), MakeUIAction(EMovieSceneTransformChannel::Weight), NAME_None, EUserInterfaceActionType::ToggleButton);
+		}
+		MenuBuilder.EndSection();
+	}
+}
+
+//mz todo
+bool FControlRigParameterSection::RequestDeleteCategory(const TArray<FName>& CategoryNamePaths)
+{
+	UMovieSceneControlRigParameterSection* ParameterSection = CastChecked<UMovieSceneControlRigParameterSection>(WeakSection.Get());
+	TSharedPtr<ISequencer> SequencerPtr = WeakSequencer.Pin();
+	/*
+	const FScopedTransaction Transaction(LOCTEXT("DeleteTransformCategory", "Delete transform category"));
+
+	if (ParameterSection->TryModify())
+	{
+		FName CategoryName = CategoryNamePaths[CategoryNamePaths.Num() - 1];
+
+		EMovieSceneTransformChannel Channel = ParameterSection->GetTransformMask().GetChannels();
+		EMovieSceneTransformChannel ChannelToRemove = ParameterSection->GetTransformMaskByName(CategoryName).GetChannels();
+
+		Channel = Channel ^ ChannelToRemove;
+
+		ParameterSection->SetTransformMask(Channel);
+
+		SequencerPtr->NotifyMovieSceneDataChanged(EMovieSceneDataChangeType::MovieSceneStructureItemsChanged);
+		return true;
+	}
+	*/
+	return false;
+}
+
+bool FControlRigParameterSection::RequestDeleteKeyArea(const TArray<FName>& KeyAreaNamePaths)
+{
+	UMovieSceneControlRigParameterSection* ParameterSection = CastChecked<UMovieSceneControlRigParameterSection>(WeakSection.Get());
+	TSharedPtr<ISequencer> SequencerPtr = WeakSequencer.Pin();
+
+	/*
+	const FScopedTransaction Transaction(LOCTEXT("DeleteTransformChannel", "Delete transform channel"));
+
+	if (ParameterSection->TryModify())
+	{
+		// Only delete the last key area path which is the channel. ie. TranslationX as opposed to Translation
+		FName KeyAreaName = KeyAreaNamePaths[KeyAreaNamePaths.Num() - 1];
+
+		EMovieSceneTransformChannel Channel = ParameterSection->GetTransformMask().GetChannels();
+		EMovieSceneTransformChannel ChannelToRemove = ParameterSection->GetTransformMaskByName(KeyAreaName).GetChannels();
+
+		Channel = Channel ^ ChannelToRemove;
+
+		ParameterSection->SetTransformMask(Channel);
+
+		SequencerPtr->NotifyMovieSceneDataChanged(EMovieSceneDataChangeType::MovieSceneStructureItemsChanged);
+		return true;
+	}
+	*/
+	return true;
+}
+
+
+#undef LOCTEXT_NAMESPACE
diff --git a/Engine/Plugins/Experimental/ControlRig/Source/ControlRigEditor/Private/Sequencer/ControlRigParameterTrackEditor.h b/Engine/Plugins/Experimental/ControlRig/Source/ControlRigEditor/Private/Sequencer/ControlRigParameterTrackEditor.h
new file mode 100644
--- /dev/null
+++ b/Engine/Plugins/Experimental/ControlRig/Source/ControlRigEditor/Private/Sequencer/ControlRigParameterTrackEditor.h
@@ -0,0 +1,148 @@
+// Copyright 1998-2019 Epic Games, Inc. All Rights Reserved.
+
+#pragma once
+
+#include "CoreMinimal.h"
+#include "Misc/Guid.h"
+#include "Templates/SubclassOf.h"
+#include "Widgets/SWidget.h"
+#include "ISequencer.h"
+#include "MovieSceneTrack.h"
+#include "ISequencerSection.h"
+#include "ISequencerTrackEditor.h"
+#include "Sequencer/MovieSceneControlRigParameterTrack.h"
+#include "KeyframeTrackEditor.h"
+#include "Sections/MovieScene3DTransformSection.h"
+
+struct FAssetData;
+class FMenuBuilder;
+class USkeleton;
+
+/**
+ * Tools for animation tracks
+ */
+class FControlRigParameterTrackEditor : public FKeyframeTrackEditor<UMovieSceneControlRigParameterTrack>
+{
+public:
+	/**
+	 * Constructor
+	 *
+	 * @param InSequencer The sequencer instance to be used by this tool
+	 */
+	FControlRigParameterTrackEditor(TSharedRef<ISequencer> InSequencer);
+
+	/** Virtual destructor. */
+	virtual ~FControlRigParameterTrackEditor();
+
+	/**
+	 * Creates an instance of this class.  Called by a sequencer
+	 *
+	 * @param OwningSequencer The sequencer instance to be used by this tool
+	 * @return The new instance of this class
+	 */
+	static TSharedRef<ISequencerTrackEditor> CreateTrackEditor(TSharedRef<ISequencer> OwningSequencer);
+
+public:
+
+	// ISequencerTrackEditor interface
+	virtual void OnRelease() override;
+	virtual void BindCommands(TSharedRef<FUICommandList> SequencerCommandBindings) override;
+	virtual void BuildObjectBindingTrackMenu(FMenuBuilder& MenuBuilder, const TArray<FGuid>& ObjectBindings, const UClass* ObjectClass) override;
+	virtual TSharedRef<ISequencerSection> MakeSectionInterface(UMovieSceneSection& SectionObject, UMovieSceneTrack& Track, FGuid ObjectBinding) override;
+	virtual bool SupportsType(TSubclassOf<UMovieSceneTrack> Type) const override;
+
+private:
+
+	void AddControlRigSubMenu(FMenuBuilder& MenuBuilder, TArray<FGuid> ObjectBindings, UMovieSceneTrack* Track);
+
+	/** Control Rig Picked */
+	void AddControlRig(UClass* InClass, UObject* BoundObject, FGuid ObjectBinding, USkeleton* Skeleton);
+
+	/** Delegate for Selection Changed Event */
+	void OnSelectionChanged(TArray<UMovieSceneTrack*> InTracks);
+
+	/** Delegate for MovieScene Changing so we can see if our track got deleted*/
+	void OnSequencerDataChanged(EMovieSceneDataChangeType DataChangeType);
+
+	/** Control Rig Delegates*/
+	void HandleControlModified(IControlRigManipulatable* Subject, const FRigControl& Control);
+	void HandleControlSelected(IControlRigManipulatable* Subject, const FRigControl& Control, bool bSelected);
+
+private:
+	//MZ TODO Store map of control rig to parameter values to see if they changed?
+
+	/** Command Bindings added by the Transform Track Editor to Sequencer and curve editor. */
+	TSharedPtr<FUICommandList> CommandBindings;
+public:
+
+	void AddControlKeys(UObject* InObject, IControlRigManipulatable* Manip, FName PropertyName, FName ParameterName, EMovieSceneTransformChannel ChannelsToKey, ESequencerKeyMode KeyMode);
+	void GetControlRigKeys(IControlRigManipulatable* Manip, FName ParameterName, EMovieSceneTransformChannel ChannelsToKey, FGeneratedTrackKeys& OutGeneratedKeys);
+	FKeyPropertyResult AddKeysToControlRig(
+		UObject* Object, IControlRigManipulatable* Manip, FFrameNumber KeyTime, FGeneratedTrackKeys& GeneratedKeys,
+		ESequencerKeyMode KeyMode, TSubclassOf<UMovieSceneTrack> TrackClass, FName PropertyName);
+	FKeyPropertyResult AddKeysToControlRigHandle(UObject *Object, IControlRigManipulatable* Manip,
+		FGuid ObjectHandle, FFrameNumber KeyTime, FGeneratedTrackKeys& GeneratedKeys,
+		ESequencerKeyMode KeyMode, TSubclassOf<UMovieSceneTrack> TrackClass, FName PropertyName);
+	/**
+	 * Modify the passed in Generated Keys by the current tracks values and weight at the passed in time.
+
+	 * @param Object The handle to the object modify
+	 * @param Track The track we are modifying
+	 * @param SectionToKey The Sections Channels we will be modifiying
+	 * @param Time The Time at which to evaluate
+	 * @param InOutGeneratedTrackKeys The Keys we need to modify. We change these values.
+	 * @param Weight The weight we need to modify the values by.
+	 */
+	bool ModifyOurGeneratedKeysByCurrentAndWeight(UObject* Object, IControlRigManipulatable* Manip, UMovieSceneTrack *Track, UMovieSceneSection* SectionToKey, FFrameNumber Time, FGeneratedTrackKeys& InOutGeneratedTotalKeys, float Weight) const;
+
+	/*
+	bool ShouldFilterAsset(const FAssetData& AssetData);
+
+	void OnControlRigAssetSelected(const FAssetData& AssetData, TArray<FGuid> ObjectBindings, USkeleton* Skeleton);
+
+	void OnControlRigAssetEnterPressed(const TArray<FAssetData>& AssetData, TArray<FGuid> ObjectBindings,  USkeleton* Skeleton);
+	*/
+private:
+	bool CanAddTransformKeysForSelectedObjects();
+	void OnAddTransformKeysForSelectedObjects(EMovieSceneTransformChannel Channel);
+	FDelegateHandle SelectionChangedHandle;
+	FDelegateHandle SequencerChangedHandle;
+
+private:
+
+	/** Guard to stop infinite loops when handling control selections*/
+	bool bIsDoingASelection;
+
+};
+
+
+/** Class for control rig sections */
+class FControlRigParameterSection : public FSequencerSection
+{
+public:
+
+	/**
+	* Creates a new control rig property section.
+	*
+	* @param InSection The section object which is being displayed and edited.
+	* @param InSequencer The sequencer which is controlling this parameter section.
+	*/
+	FControlRigParameterSection(UMovieSceneSection& InSection, TWeakPtr<ISequencer> InSequencer)
+		: FSequencerSection(InSection), WeakSequencer(InSequencer)
+	{
+	}
+
+public:
+
+	virtual void BuildSectionContextMenu(FMenuBuilder& MenuBuilder, const FGuid& InObjectBinding) override;
+
+	//~ ISequencerSection interface
+	virtual bool RequestDeleteCategory(const TArray<FName>& CategoryNamePath) override;
+	virtual bool RequestDeleteKeyArea(const TArray<FName>& KeyAreaNamePath) override;
+
+protected:
+
+	/** The sequencer which is controlling this section. */
+	TWeakPtr<ISequencer> WeakSequencer;
+
+};
\ No newline at end of file
diff --git a/Engine/Plugins/Experimental/ControlRig/Source/ControlRigManipulation/Private/ControlRigGizmoActor.cpp b/Engine/Plugins/Experimental/ControlRig/Source/ControlRigManipulation/Private/ControlRigGizmoActor.cpp
--- a/Engine/Plugins/Experimental/ControlRig/Source/ControlRigManipulation/Private/ControlRigGizmoActor.cpp
+++ b/Engine/Plugins/Experimental/ControlRig/Source/ControlRigManipulation/Private/ControlRigGizmoActor.cpp
@@ -24,6 +24,8 @@ AControlRigGizmoActor::AControlRigGizmoActor(const FObjectInitializer& ObjectIni
 
 	RootComponent = ActorRootComponent;
 	StaticMeshComponent->SetupAttachment(RootComponent);
+	StaticMeshComponent->bCastStaticShadow = false;
+	StaticMeshComponent->bCastDynamicShadow = false;
 }
 
 void AControlRigGizmoActor::SetEnabled(bool bInEnabled)
@@ -108,7 +110,7 @@ namespace FControlRigGizmoHelper
 		ActorSpawnParameters.bTemporaryEditorActor = true;
 		ActorSpawnParameters.bHideFromSceneOutliner = true;
 		ActorSpawnParameters.SpawnCollisionHandlingOverride = ESpawnActorCollisionHandlingMethod::AlwaysSpawn;
-
+		ActorSpawnParameters.ObjectFlags = RF_Transient;
 		return ActorSpawnParameters;
 	}
 
diff --git a/Engine/Plugins/Experimental/ControlRig/Source/ControlRigManipulation/Private/DefaultControlRigManipulationLayer.cpp b/Engine/Plugins/Experimental/ControlRig/Source/ControlRigManipulation/Private/DefaultControlRigManipulationLayer.cpp
--- a/Engine/Plugins/Experimental/ControlRig/Source/ControlRigManipulation/Private/DefaultControlRigManipulationLayer.cpp
+++ b/Engine/Plugins/Experimental/ControlRig/Source/ControlRigManipulation/Private/DefaultControlRigManipulationLayer.cpp
@@ -19,9 +19,6 @@ void UDefaultControlRigManipulationLayer::CreateLayer()
 
 void UDefaultControlRigManipulationLayer::DestroyLayer()
 {
-	ExecuteDelegateHandles.Empty();
-	InitializeDelegateHandle.Empty();
-
 	IControlRigManipulationLayer::DestroyLayer();
 }
 
@@ -49,17 +46,22 @@ void UDefaultControlRigManipulationLayer::RemoveManipulatableObject(IControlRigM
 void UDefaultControlRigManipulationLayer::OnControlRigAdded(UControlRig* InControlRig)
 {
 	// bind execution delegate
-	ExecuteDelegateHandles.Add(InControlRig->OnExecuted().AddUObject(this, &UDefaultControlRigManipulationLayer::OnRigExecuted));
-	InitializeDelegateHandle.Add(InControlRig->OnInitialized().AddUObject(this, &UDefaultControlRigManipulationLayer::OnRigInitialized));
+
 	ControlModifiedDelegateHandles.Add(InControlRig->ControlModified().AddUObject(this, &UDefaultControlRigManipulationLayer::OnControlModified));
 
 	// ensure the size matches
-	ensure(ExecuteDelegateHandles.Num() == ManipulatableObjects.Num());
-	ensure(InitializeDelegateHandle.Num() == ManipulatableObjects.Num());
 	ensure(ControlModifiedDelegateHandles.Num() == ManipulatableObjects.Num());
 
 	// object binding. This overwrites if there were multi
 	SetObjectBinding(InControlRig->GetObjectBinding());
+
+	// currently all the manipulatable mesh component is supposed to be same
+	// if that changes, this code has to change
+	USkeletalMeshComponent* MeshComponent = GetSkeletalMeshComponent();
+	if (MeshComponent)
+	{
+		MeshComponent->OnBoneTransformsFinalized.AddDynamic(this, &UDefaultControlRigManipulationLayer::PostPoseUpdate);
+	}
 }
 
 void UDefaultControlRigManipulationLayer::OnControlRigRemoved(UControlRig* InControlRig)
@@ -68,14 +70,16 @@ void UDefaultControlRigManipulationLayer::OnControlRigRemoved(UControlRig* InCon
 	int32 Found = ManipulatableObjects.Find(InControlRig);
 	if (Found != INDEX_NONE)
 	{
-		if (ExecuteDelegateHandles.IsValidIndex(Found) && ExecuteDelegateHandles[Found].IsValid())
+		// last one 
+		// currently all the manipulatable mesh component is supposed to be same
+		// if that changes, this code has to change
+		if (ManipulatableObjects.Num() == 1)
 		{
-			InControlRig->OnExecuted().Remove(ExecuteDelegateHandles[Found]);
-		}
-
-		if (InitializeDelegateHandle.IsValidIndex(Found) && InitializeDelegateHandle[Found].IsValid())
-		{
-			InControlRig->OnInitialized().Remove(InitializeDelegateHandle[Found]);
+			USkeletalMeshComponent* MeshComponent = GetSkeletalMeshComponent();
+			if (MeshComponent)
+			{
+				MeshComponent->OnBoneTransformsFinalized.RemoveDynamic(this, &UDefaultControlRigManipulationLayer::PostPoseUpdate);
+			}
 		}
 
 		if (ControlModifiedDelegateHandles.IsValidIndex(Found) && ControlModifiedDelegateHandles[Found].IsValid())
@@ -83,7 +87,6 @@ void UDefaultControlRigManipulationLayer::OnControlRigRemoved(UControlRig* InCon
 			InControlRig->ControlModified().Remove(ControlModifiedDelegateHandles[Found]);
 		}
 	}
-	
 }
 
 // gizmo related 
@@ -321,10 +324,21 @@ bool UDefaultControlRigManipulationLayer::CreateGizmoActors(UWorld* World, TArra
 			OutGizmoActors.Add(GizmoActor);
 		}
 	}
-	
+
+	WorldPtr = World;
+	OnWorldCleanupHandle = FWorldDelegates::OnWorldCleanup.AddUObject(this, &UDefaultControlRigManipulationLayer::OnWorldCleanup);
 	return (OutGizmoActors.Num() > 0);
 }
 
+void UDefaultControlRigManipulationLayer::OnWorldCleanup(UWorld* World, bool bSessionEnded, bool bCleanupResources)
+{
+	// if world gets cleaned up first, we destroy gizmo actors
+	if (WorldPtr == World)
+	{
+		DestroyGizmosActors();
+	}
+}
+
 void UDefaultControlRigManipulationLayer::DestroyGizmosActors()
 {
 	// clear previous gizmo actors
@@ -340,21 +354,8 @@ void UDefaultControlRigManipulationLayer::DestroyGizmosActors()
 	}
 
 	ResetControlData();
-}
 
-void UDefaultControlRigManipulationLayer::OnRigExecuted(class UControlRig* InRig, const EControlRigState State)
-{
-	FTransform ComponentTransform = GetSkeletalMeshComponentTransform();
-
-	// after executing rig, update gizmo
-	for (auto Iter=GizmoToControlMap.CreateIterator(); Iter; ++Iter)
-	{
-		TickGizmo(Iter.Key(), ComponentTransform);
-	}
-}
-
-void UDefaultControlRigManipulationLayer::OnRigInitialized(class UControlRig* InRig, const EControlRigState State)
-{
+	FWorldDelegates::OnWorldCleanup.Remove(OnWorldCleanupHandle);
 }
 
 void UDefaultControlRigManipulationLayer::OnControlModified(IControlRigManipulatable* InManipulatable, const FRigControl& InControl)
@@ -424,4 +425,14 @@ void UDefaultControlRigManipulationLayer::BeginTransaction()
 
 void UDefaultControlRigManipulationLayer::EndTransaction()
 {
+}
+
+void UDefaultControlRigManipulationLayer::PostPoseUpdate()
+{
+	FTransform ComponentTransform = GetSkeletalMeshComponentTransform();
+	// after executing rig, update gizmo
+	for (auto Iter = GizmoToControlMap.CreateIterator(); Iter; ++Iter)
+	{
+		TickGizmo(Iter.Key(), ComponentTransform);
+	}
 }
\ No newline at end of file
diff --git a/Engine/Plugins/Experimental/ControlRig/Source/ControlRigManipulation/Public/DefaultControlRigManipulationLayer.h b/Engine/Plugins/Experimental/ControlRig/Source/ControlRigManipulation/Public/DefaultControlRigManipulationLayer.h
--- a/Engine/Plugins/Experimental/ControlRig/Source/ControlRigManipulation/Public/DefaultControlRigManipulationLayer.h
+++ b/Engine/Plugins/Experimental/ControlRig/Source/ControlRigManipulation/Public/DefaultControlRigManipulationLayer.h
@@ -50,8 +50,6 @@ public:
 	virtual void MoveGizmo(AControlRigGizmoActor* GizmoActor, const bool bTranslation, FVector& InDrag, 
 		const bool bRotation, FRotator& InRot, const bool bScale, FVector& InScale, const FTransform& ToWorldTransform) override;
 	virtual void TickGizmo(AControlRigGizmoActor* GizmoActor, const FTransform& ComponentTransform) override;
-	// need templatizing input type @fixme: add implementation
-	virtual void AddKey(const AControlRigGizmoActor* GizmoActor, TSharedPtr<ISequencer> Sequencer) override {};
 	virtual bool ModeSupportedByGizmoActor(const AControlRigGizmoActor* GizmoActor, FWidget::EWidgetMode InMode) const override;
 	// IControlRigManipulationLayer END
 
@@ -87,11 +85,9 @@ private:
 	// control data related
 	virtual void ResetControlData();
 
-	void OnRigExecuted(class UControlRig* InRig, const EControlRigState State);
-	TArray<FDelegateHandle> ExecuteDelegateHandles;
-
-	void OnRigInitialized(class UControlRig* InRig, const EControlRigState State);
-	TArray<FDelegateHandle> InitializeDelegateHandle;
+	// Post pose update handler
+	UFUNCTION()
+	virtual void PostPoseUpdate();
 
 	void OnControlModified(IControlRigManipulatable* InManipulatable, const FRigControl& InControl);
 	TArray<FDelegateHandle> ControlModifiedDelegateHandles;
@@ -100,4 +96,8 @@ private:
 	void OnControlRigRemoved(UControlRig* InControlRig);
 
 	void GetGizmoCreationParams(TArray<FGizmoActorCreationParam>& OutCreationParams);
+	// world clean up handlers
+	FDelegateHandle OnWorldCleanupHandle;
+	void OnWorldCleanup(UWorld* World, bool bSessionEnded, bool bCleanupResources);
+	UWorld* WorldPtr = nullptr;
 };
diff --git a/Engine/Plugins/Experimental/ControlRig/Source/ControlRigManipulation/Public/IControlRigManipulationLayer.h b/Engine/Plugins/Experimental/ControlRig/Source/ControlRigManipulation/Public/IControlRigManipulationLayer.h
--- a/Engine/Plugins/Experimental/ControlRig/Source/ControlRigManipulation/Public/IControlRigManipulationLayer.h
+++ b/Engine/Plugins/Experimental/ControlRig/Source/ControlRigManipulation/Public/IControlRigManipulationLayer.h
@@ -37,9 +37,6 @@ public:
 	virtual void TickGizmo(AControlRigGizmoActor* GizmoActor, const FTransform& ComponentTransform) = 0;
 	virtual bool ModeSupportedByGizmoActor(const AControlRigGizmoActor* GizmoActor, FWidget::EWidgetMode InMode) const = 0;
 
-	// add key of current gizmo actor transform data
-	virtual void AddKey(const AControlRigGizmoActor* GizmoActor, TSharedPtr<ISequencer> Sequencer) = 0;
-
 protected:
 
 	// because manipulatable is created externally often. Manipulation Layer doesn't have any ownership on these object
diff --git a/Engine/Plugins/Experimental/GeometryCollectionPlugin/Source/GeometryCollectionExampleCore/Private/GeometryCollection/GeometryCollectionExampleSimulation.cpp b/Engine/Plugins/Experimental/GeometryCollectionPlugin/Source/GeometryCollectionExampleCore/Private/GeometryCollection/GeometryCollectionExampleSimulation.cpp
--- a/Engine/Plugins/Experimental/GeometryCollectionPlugin/Source/GeometryCollectionExampleCore/Private/GeometryCollection/GeometryCollectionExampleSimulation.cpp
+++ b/Engine/Plugins/Experimental/GeometryCollectionPlugin/Source/GeometryCollectionExampleCore/Private/GeometryCollection/GeometryCollectionExampleSimulation.cpp
@@ -29,7 +29,6 @@ namespace GeometryCollectionExample
 	template<class T>
 	void RigidBodiesFallingUnderGravity()
 	{
-#if INCLUDE_CHAOS
 		TUniquePtr<Chaos::TChaosPhysicsMaterial<T>> PhysicalMaterial = nullptr;
 		TSharedPtr<FGeometryCollection> RestCollection = nullptr;
 		TSharedPtr<FGeometryDynamicCollection> DynamicCollection = nullptr;
@@ -66,7 +65,6 @@ namespace GeometryCollectionExample
 		FChaosSolversModule::GetModule()->DestroySolver(Solver);
 
 		delete PhysObject;
-#endif
 	}
 	template void RigidBodiesFallingUnderGravity<float>();
 
@@ -75,7 +73,6 @@ namespace GeometryCollectionExample
 	template<class T>
 	void RigidBodiesCollidingWithSolverFloor()
 	{
-#if INCLUDE_CHAOS
 		TUniquePtr<Chaos::TChaosPhysicsMaterial<T>> PhysicalMaterial = nullptr;
 		TSharedPtr<FGeometryCollection> RestCollection = nullptr;
 		TSharedPtr<FGeometryDynamicCollection> DynamicCollection = nullptr;
@@ -119,7 +116,6 @@ namespace GeometryCollectionExample
 		FChaosSolversModule::GetModule()->DestroySolver(Solver);
 
 		delete PhysObject;
-#endif
 	}
 	template void RigidBodiesCollidingWithSolverFloor<float>();
 
@@ -127,7 +123,6 @@ namespace GeometryCollectionExample
 	template<class T>
 	void RigidBodiesSingleSphereCollidingWithSolverFloor()
 	{
-#if INCLUDE_CHAOS
 		TUniquePtr<Chaos::TChaosPhysicsMaterial<T>> PhysicalMaterial = nullptr;
 		TSharedPtr<FGeometryCollection> RestCollection = nullptr;
 		TSharedPtr<FGeometryDynamicCollection> DynamicCollection = nullptr;
@@ -181,7 +176,6 @@ namespace GeometryCollectionExample
 		FChaosSolversModule::GetModule()->DestroySolver(Solver);
 
 		delete PhysObject;
-#endif
 	}
 	template void RigidBodiesSingleSphereCollidingWithSolverFloor<float>();
 
@@ -189,7 +183,6 @@ namespace GeometryCollectionExample
 	template<class T>
 	void RigidBodiesSingleSphereIntersectingWithSolverFloor()
 	{
-#if INCLUDE_CHAOS
 		TUniquePtr<Chaos::TChaosPhysicsMaterial<T>> PhysicalMaterial = nullptr;
 		TSharedPtr<FGeometryCollection> RestCollection = nullptr;
 		TSharedPtr<FGeometryDynamicCollection> DynamicCollection = nullptr;
@@ -236,7 +229,6 @@ namespace GeometryCollectionExample
 
 		delete PhysObject;
 
-#endif
 	}
 	template void RigidBodiesSingleSphereIntersectingWithSolverFloor<float>();
 
@@ -244,7 +236,6 @@ namespace GeometryCollectionExample
 	template<class T>
 	void RigidBodiesKinematic()
 	{
-#if INCLUDE_CHAOS
 		TUniquePtr<Chaos::TChaosPhysicsMaterial<T>> PhysicalMaterial = nullptr;
 		TSharedPtr<FGeometryCollection> RestCollection = nullptr;
 		TSharedPtr<FGeometryDynamicCollection> DynamicCollection = nullptr;
@@ -283,7 +274,6 @@ namespace GeometryCollectionExample
 
 		delete PhysObject;
 
-#endif
 	}
 	template void RigidBodiesKinematic<float>();
 
@@ -291,7 +281,6 @@ namespace GeometryCollectionExample
 	template<class T>
 	void RigidBodiesSleepingActivation()
 	{
-#if INCLUDE_CHAOS
 		TUniquePtr<Chaos::TChaosPhysicsMaterial<T>> PhysicalMaterial = nullptr;
 		TSharedPtr<FGeometryCollection> RestCollection = nullptr;
 		TSharedPtr<FGeometryDynamicCollection> DynamicCollection = nullptr;
@@ -343,14 +332,12 @@ namespace GeometryCollectionExample
 		FChaosSolversModule::GetModule()->DestroySolver(Solver);
 
 		delete PhysObject;
-#endif
 	}
 	template void RigidBodiesSleepingActivation<float>();
 
 	template<class T>
 	void RigidBodies_CollisionGroup()
 	{
-#if INCLUDE_CHAOS
 		TUniquePtr<Chaos::TChaosPhysicsMaterial<T>> PhysicalMaterial = nullptr;
 		TSharedPtr<FGeometryCollection> RestCollection = nullptr;
 		TSharedPtr<FGeometryDynamicCollection> DynamicCollection = nullptr;
@@ -437,7 +424,6 @@ namespace GeometryCollectionExample
 
 		FChaosSolversModule::GetModule()->DestroySolver(Solver);
 		delete PhysObject;
-#endif
 	}
 	template void RigidBodies_CollisionGroup<float>();
 
@@ -446,7 +432,6 @@ namespace GeometryCollectionExample
 	template<class T>
 	void RigidBodies_Initialize_ParticleImplicitCollisionGeometry()
 	{
-#if INCLUDE_CHAOS
 		typedef Chaos::TVector<T, 3> Vec;
 
 		typename SimulationObjects<T>::FParameters P;
@@ -500,7 +485,6 @@ namespace GeometryCollectionExample
 			EXPECT_LT((FMath::Abs(FMath::Abs(Vert.X) + FMath::Abs(Vert.Y) + FMath::Abs(Vert.Z))-1.5), KINDA_SMALL_NUMBER );
 		}
 		delete Object;
-#endif
 	}
 	template void RigidBodies_Initialize_ParticleImplicitCollisionGeometry<float>();
 
diff --git a/Engine/Plugins/Experimental/GeometryCollectionPlugin/Source/GeometryCollectionExampleCore/Private/GeometryCollection/GeometryCollectionExampleSimulationSolver.cpp b/Engine/Plugins/Experimental/GeometryCollectionPlugin/Source/GeometryCollectionExampleCore/Private/GeometryCollection/GeometryCollectionExampleSimulationSolver.cpp
--- a/Engine/Plugins/Experimental/GeometryCollectionPlugin/Source/GeometryCollectionExampleCore/Private/GeometryCollection/GeometryCollectionExampleSimulationSolver.cpp
+++ b/Engine/Plugins/Experimental/GeometryCollectionPlugin/Source/GeometryCollectionExampleCore/Private/GeometryCollection/GeometryCollectionExampleSimulationSolver.cpp
@@ -37,7 +37,6 @@ namespace GeometryCollectionExample
 	template<class T>
 	void Solver_AdvanceNoObjects()
 	{
-#if INCLUDE_CHAOS
 		Chaos::FPBDRigidsSolver* Solver = FChaosSolversModule::GetModule()->CreateSolver(true);
 		Solver->SetHasFloor(false);
 		Solver->SetEnabled(true);
@@ -45,7 +44,6 @@ namespace GeometryCollectionExample
 		FinalizeSolver(*Solver);
 
 		FChaosSolversModule::GetModule()->DestroySolver(Solver);
-#endif
 	}
 	template void Solver_AdvanceNoObjects<float>();
 
@@ -53,7 +51,6 @@ namespace GeometryCollectionExample
 	template<class T>
 	void Solver_AdvanceDisabledObjects()
 	{
-#if INCLUDE_CHAOS
 		TSharedPtr<FGeometryCollection> RestCollection = GeometryCollection::MakeCubeElement(FTransform::Identity, FVector(1.0));
 		TSharedPtr<FGeometryDynamicCollection> DynamicCollection = GeometryCollectionToGeometryDynamicCollection(RestCollection.Get());
 
@@ -94,14 +91,12 @@ namespace GeometryCollectionExample
 		FChaosSolversModule::GetModule()->DestroySolver(Solver);
 
 		delete PhysObject;
-#endif
 	}
 	template void Solver_AdvanceDisabledObjects<float>();
 
 	template<class T>
 	void Solver_AdvanceDisabledClusteredObjects()
 	{
-#if INCLUDE_CHAOS
 		TSharedPtr<FGeometryCollection> RestCollection = GeometryCollection::MakeCubeElement(FTransform(FQuat::MakeFromEuler(FVector(0, 0, 0.)), FVector(0, -10, 10)), FVector(1.0));
 		RestCollection->AppendGeometry(*GeometryCollection::MakeCubeElement(FTransform(FQuat::MakeFromEuler(FVector(0, 0, 0.)), FVector(0, 10, 10)), FVector(1.0)));
 		EXPECT_EQ(RestCollection->Transform.Num(), 2);
@@ -151,7 +146,6 @@ namespace GeometryCollectionExample
 
 		delete PhysObject;
 
-#endif
 	}
 	template void Solver_AdvanceDisabledClusteredObjects<float>();
 
@@ -159,7 +153,6 @@ namespace GeometryCollectionExample
 	template<class T>
 	void Solver_ValidateReverseMapping()
 	{
-#if INCLUDE_CHAOS
 		TUniquePtr<Chaos::TChaosPhysicsMaterial<T>> PhysicalMaterial = MakeUnique<Chaos::TChaosPhysicsMaterial<T>>();
 		InitMaterialToZero(PhysicalMaterial);
 
@@ -215,11 +208,9 @@ namespace GeometryCollectionExample
 
 		FChaosSolversModule::GetModule()->DestroySolver(Solver);
 
-#endif
 	}
 	template void Solver_ValidateReverseMapping<float>();
 
-#if INCLUDE_CHAOS
 	template<class T>
 	void CommonInit(int32 NumObjects, bool UseClusters, Chaos::FPBDRigidsSolver** SolverInOut, TUniquePtr<Chaos::TChaosPhysicsMaterial<T>>& PhysicalMaterial, TArray<TSharedPtr<FGeometryCollection>>& RestArray, TArray<TSharedPtr<FGeometryDynamicCollection>>& DynamicArray)
 	{
@@ -273,13 +264,10 @@ namespace GeometryCollectionExample
 
 		Solver->AdvanceSolverBy(1 / 24.);
 	}
-#endif // INCLUDE_CHAOS
-
 
 	class EventHarvester
 	{
 	public:
-#if INCLUDE_CHAOS
 		EventHarvester(Chaos::FPBDRigidsSolver* Solver)
 		{
 			Solver->GetEventManager()->RegisterHandler<Chaos::FCollisionEventData>(Chaos::EEventType::Collision, this, &EventHarvester::HandleCollisionEvents);
@@ -298,14 +286,11 @@ namespace GeometryCollectionExample
 
 		Chaos::FCollisionEventData CollisionEventData;
 		Chaos::FBreakingEventData BreakingEventData;
-#endif
 	};
 
 	template<class T>
 	void Solver_CollisionEventFilter()
 	{
-#if INCLUDE_CHAOS
-
 		float TestMassThreshold = 6.0f;
 
 		Chaos::FPBDRigidsSolver* Solver;
@@ -362,7 +347,6 @@ namespace GeometryCollectionExample
 
 #endif
 		FChaosSolversModule::GetModule()->DestroySolver(Solver);
-#endif
 
 	}
 	template void Solver_CollisionEventFilter<float>();
@@ -370,8 +354,6 @@ namespace GeometryCollectionExample
 	template<class T>
 	void Solver_BreakingEventFilter()
 	{
-#if INCLUDE_CHAOS
-
 		float TestMass = 6.0f;
 
 		Chaos::FPBDRigidsSolver* Solver;
@@ -434,7 +416,6 @@ namespace GeometryCollectionExample
 
 #endif
 		FChaosSolversModule::GetModule()->DestroySolver(Solver);
-#endif
 	}
 	template void Solver_BreakingEventFilter<float>();
 }
diff --git a/Engine/Plugins/Experimental/GeometryCollectionPlugin/Source/GeometryCollectionExampleCore/Private/GeometryCollection/GeometryCollectionExampleSimulationStreaming.cpp b/Engine/Plugins/Experimental/GeometryCollectionPlugin/Source/GeometryCollectionExampleCore/Private/GeometryCollection/GeometryCollectionExampleSimulationStreaming.cpp
--- a/Engine/Plugins/Experimental/GeometryCollectionPlugin/Source/GeometryCollectionExampleCore/Private/GeometryCollection/GeometryCollectionExampleSimulationStreaming.cpp
+++ b/Engine/Plugins/Experimental/GeometryCollectionPlugin/Source/GeometryCollectionExampleCore/Private/GeometryCollection/GeometryCollectionExampleSimulationStreaming.cpp
@@ -26,7 +26,6 @@ namespace GeometryCollectionExample
 	template<class T>
 	void RigidBodies_Streaming_StartSolverEmpty()
 	{
-#if INCLUDE_CHAOS
 		Chaos::FPBDRigidsSolver* Solver = FChaosSolversModule::GetModule()->CreateSolver(true);
 		Solver->SetHasFloor(false);
 		Solver->SetEnabled(true);
@@ -57,7 +56,6 @@ namespace GeometryCollectionExample
 			delete Obj;
 #endif
 		delete Solver;
-#endif
 	}
 	template void RigidBodies_Streaming_StartSolverEmpty<float>();
 
@@ -65,7 +63,6 @@ namespace GeometryCollectionExample
 	template<class T>
 	void RigidBodies_Streaming_BulkInitialization()
 	{
-#if INCLUDE_CHAOS
 		typename SimulationObjects<T>::FParameters P;
 		P.CollisionGroup = -1;
 
@@ -105,7 +102,6 @@ namespace GeometryCollectionExample
 			delete Obj;
 #endif
 		delete Solver;
-#endif
 	}
 	template void RigidBodies_Streaming_BulkInitialization<float>(); 
 
@@ -114,7 +110,6 @@ namespace GeometryCollectionExample
 	template<class T>
 	void RigidBodies_Streaming_DeferedClusteringInitialization()
 	{
-#if INCLUDE_CHAOS
 		typename SimulationObjects<T>::FParameters P;
 		P.CollisionGroup = -1;
 		P.EnableClustering = true;
@@ -175,7 +170,6 @@ namespace GeometryCollectionExample
 			delete Obj;
 #endif
 		delete Solver;
-#endif
 	}
 	template void RigidBodies_Streaming_DeferedClusteringInitialization<float>();
 
diff --git a/Engine/Plugins/Experimental/GeometryCollectionPlugin/Source/GeometryCollectionExampleCore/Private/GeometryCollection/GeometryCollectionExampleUtility.cpp b/Engine/Plugins/Experimental/GeometryCollectionPlugin/Source/GeometryCollectionExampleCore/Private/GeometryCollection/GeometryCollectionExampleUtility.cpp
--- a/Engine/Plugins/Experimental/GeometryCollectionPlugin/Source/GeometryCollectionExampleCore/Private/GeometryCollection/GeometryCollectionExampleUtility.cpp
+++ b/Engine/Plugins/Experimental/GeometryCollectionPlugin/Source/GeometryCollectionExampleCore/Private/GeometryCollection/GeometryCollectionExampleUtility.cpp
@@ -40,14 +40,12 @@ namespace GeometryCollectionExample {
 
 	void FinalizeSolver(Chaos::FPBDRigidsSolver& InSolver)
 	{
-#if INCLUDE_CHAOS
 		InSolver.ForEachPhysicsProxy([](auto* Object)
 		{
 			Object->BufferPhysicsResults();
 			Object->FlipBuffer();
 			Object->PullFromPhysicsState();
 		});
-#endif
 	}
 
 
@@ -56,7 +54,6 @@ namespace GeometryCollectionExample {
 	TSharedPtr<FGeometryCollection>	CreateClusteredBody(FVector Position)
 	{
 		TSharedPtr<FGeometryCollection> RestCollection;
-#if INCLUDE_CHAOS
 		RestCollection = GeometryCollection::MakeCubeElement(FTransform(FQuat::MakeFromEuler(FVector(0.f)), Position), FVector(1.0));
 		RestCollection->AppendGeometry(*GeometryCollection::MakeCubeElement(FTransform(FQuat::MakeFromEuler(FVector(0.f)), FVector(30.f)), FVector(1.0)));
 
@@ -66,14 +63,12 @@ namespace GeometryCollectionExample {
 
 		GeometryCollectionAlgo::ParentTransforms(RestCollection.Get(), 2, { 0,1 });
 
-#endif
 		return RestCollection;
 	}
 
 	TSharedPtr<FGeometryCollection>	CreateClusteredBody_TwoParents_TwoBodies(FVector Position)
 	{
 		TSharedPtr<FGeometryCollection> RestCollection;
-#if INCLUDE_CHAOS
 		RestCollection = GeometryCollection::MakeCubeElement(FTransform(FQuat::MakeFromEuler(FVector(0.f)), Position), FVector(1.0));
 		RestCollection->AppendGeometry(*GeometryCollection::MakeCubeElement(FTransform(FQuat::MakeFromEuler(FVector(0.f)), FVector(30.f)), FVector(1.0)));
 
@@ -86,14 +81,12 @@ namespace GeometryCollectionExample {
 		GeometryCollectionAlgo::ParentTransforms(RestCollection.Get(), 2, { 0,1 });
 		GeometryCollectionAlgo::ParentTransforms(RestCollection.Get(), 3, { 2 });
 
-#endif
 		return RestCollection;
 	}
 
 	TSharedPtr<FGeometryCollection>	CreateClusteredBody_FourParents_OneBody(FVector Position)
 	{
 		TSharedPtr<FGeometryCollection> RestCollection;
-#if INCLUDE_CHAOS
 		RestCollection = GeometryCollection::MakeCubeElement(FTransform(FQuat::MakeFromEuler(FVector(0.f)), Position), FVector(1.0));
 
 		RestCollection->AddElements(4, FGeometryCollection::TransformGroup);
@@ -109,7 +102,6 @@ namespace GeometryCollectionExample {
 		GeometryCollectionAlgo::ParentTransforms(RestCollection.Get(), 3, { 2 });
 		GeometryCollectionAlgo::ParentTransforms(RestCollection.Get(), 4, { 3 });
 
-#endif
 		return RestCollection;
 	}
 
@@ -118,7 +110,6 @@ namespace GeometryCollectionExample {
 	TSharedPtr<FGeometryCollection>	CreateClusteredBody_TwoByTwo_ThreeTransform(FVector Position)
 	{
 		TSharedPtr<FGeometryCollection> RestCollection;
-#if INCLUDE_CHAOS
 		RestCollection = GeometryCollection::MakeCubeElement(FTransform(FQuat::MakeFromEuler(FVector(0.f)), FVector(0,0,0)), FVector(1.0));
 		RestCollection->AppendGeometry(*GeometryCollection::MakeCubeElement(FTransform(FQuat::MakeFromEuler(FVector(0.f)), FVector(100,0,0)), FVector(1.0)));
 		RestCollection->AppendGeometry(*GeometryCollection::MakeCubeElement(FTransform(FQuat::MakeFromEuler(FVector(0.f)), FVector(200,0,0)), FVector(1.0)));
@@ -140,7 +131,6 @@ namespace GeometryCollectionExample {
 		GeometryCollectionAlgo::ParentTransforms(RestCollection.Get(), 5, { 2,3 });
 		GeometryCollectionAlgo::ParentTransforms(RestCollection.Get(), 6, { 4,5 });
 
-#endif
 		return RestCollection;
 	}
 
@@ -149,7 +139,6 @@ namespace GeometryCollectionExample {
 	TSharedPtr<FGeometryCollection>	CreateClusteredBody_ThreeByTwo_ThreeTransform(FVector Position)
 	{
 		TSharedPtr<FGeometryCollection> RestCollection;
-#if INCLUDE_CHAOS
 		RestCollection = GeometryCollection::MakeCubeElement(FTransform(FQuat::MakeFromEuler(FVector(0.f)), FVector(0, 0, 0)), FVector(1.0));
 		RestCollection->AppendGeometry(*GeometryCollection::MakeCubeElement(FTransform(FQuat::MakeFromEuler(FVector(0.f)), FVector(100, 0, 0)), FVector(1.0)));
 		RestCollection->AppendGeometry(*GeometryCollection::MakeCubeElement(FTransform(FQuat::MakeFromEuler(FVector(0.f)), FVector(200, 0, 0)), FVector(1.0)));
@@ -175,7 +164,6 @@ namespace GeometryCollectionExample {
 		GeometryCollectionAlgo::ParentTransforms(RestCollection.Get(), 7, { 3,4,5 });
 		GeometryCollectionAlgo::ParentTransforms(RestCollection.Get(), 8, { 6,7 });
 
-#endif
 		return RestCollection;
 	}
 
@@ -183,7 +171,6 @@ namespace GeometryCollectionExample {
 	TSharedPtr<FGeometryCollection>	CreateClusteredBody_FracturedGeometry(FVector Position)
 	{
 		TSharedPtr<FGeometryCollection> RestCollection;
-#if INCLUDE_CHAOS
 
 		RestCollection = TSharedPtr<FGeometryCollection>(FGeometryCollection::NewGeometryCollection(
 			FracturedGeometry::RawVertexArray,
@@ -213,7 +200,6 @@ namespace GeometryCollectionExample {
 		RestCollection->SimulationType[12] = FGeometryCollection::ESimulationTypes::FST_Clustered;
 		RestCollection->SimulationType[0] = FGeometryCollection::ESimulationTypes::FST_Clustered;
 
-#endif
 		return RestCollection;
 	}
 
