Control Rig: Disable in menu

diff --git a/Engine/Plugins/Experimental/ControlRig/Source/ControlRigEditor/Private/Sequencer/ControlRigSequenceFactory.cpp b/Engine/Plugins/Experimental/ControlRig/Source/ControlRigEditor/Private/Sequencer/ControlRigSequenceFactory.cpp
--- a/Engine/Plugins/Experimental/ControlRig/Source/ControlRigEditor/Private/Sequencer/ControlRigSequenceFactory.cpp
+++ b/Engine/Plugins/Experimental/ControlRig/Source/ControlRigEditor/Private/Sequencer/ControlRigSequenceFactory.cpp
@@ -31,7 +31,7 @@ UObject* UControlRigSequenceFactory::FactoryCreateNew(UClass* Class, UObject* In
 
 bool UControlRigSequenceFactory::ShouldShowInNewMenu() const
 {
-	return true;
+	return false;
 }
 
 #undef LOCTEXT_NAMESPACE
