Control Rig: Ensure to guard against non-existing viewport

diff --git a/Engine/Plugins/Experimental/ControlRig/Source/ControlRigEditor/Private/Editor/ControlRigEditor.cpp b/Engine/Plugins/Experimental/ControlRig/Source/ControlRigEditor/Private/Editor/ControlRigEditor.cpp
--- a/Engine/Plugins/Experimental/ControlRig/Source/ControlRigEditor/Private/Editor/ControlRigEditor.cpp
+++ b/Engine/Plugins/Experimental/ControlRig/Source/ControlRigEditor/Private/Editor/ControlRigEditor.cpp
@@ -208,8 +208,12 @@ void FControlRigEditor::InitControlRigEditor(const EToolkitMode::Type Mode, cons
 
 	// Activate our edit mode
 //	GetAssetEditorModeManager()->SetToolkitHost(GetToolkitHost());
-	GetAssetEditorModeManager()->SetDefaultMode(FControlRigEditorEditMode::ModeName);
-	GetAssetEditorModeManager()->ActivateMode(FControlRigEditorEditMode::ModeName);
+	if (GetAssetEditorModeManager() != nullptr)
+	{
+		GetAssetEditorModeManager()->SetDefaultMode(FControlRigEditorEditMode::ModeName);
+		GetAssetEditorModeManager()->ActivateMode(FControlRigEditorEditMode::ModeName);
+	}
+
 	if (FControlRigEditMode* EditMode = GetEditMode())
 	{
 		EditMode->OnGetRigElementTransform() = FOnGetRigElementTransform::CreateSP(this, &FControlRigEditor::GetRigElementTransform);
@@ -1985,6 +1989,16 @@ void FControlRigEditor::OnControlUISettingChanged(FRigHierarchyContainer* Contai
 	OnHierarchyChanged();
 }
 
+FControlRigEditorEditMode* FControlRigEditor::GetEditMode() const
+{
+	if (GetAssetEditorModeManager() == nullptr)
+	{
+		return nullptr;
+	}
+	return static_cast<FControlRigEditorEditMode*>(GetAssetEditorModeManager()->GetActiveMode(FControlRigEditorEditMode::ModeName));
+}
+
+
 void FControlRigEditor::OnCurveContainerChanged()
 {
 	DECLARE_SCOPE_HIERARCHICAL_COUNTER_FUNC()
diff --git a/Engine/Plugins/Experimental/ControlRig/Source/ControlRigEditor/Private/Editor/ControlRigEditor.h b/Engine/Plugins/Experimental/ControlRig/Source/ControlRigEditor/Private/Editor/ControlRigEditor.h
--- a/Engine/Plugins/Experimental/ControlRig/Source/ControlRigEditor/Private/Editor/ControlRigEditor.h
+++ b/Engine/Plugins/Experimental/ControlRig/Source/ControlRigEditor/Private/Editor/ControlRigEditor.h
@@ -103,7 +103,7 @@ public:
 	TSharedRef<SBorder> GetToolbox() { return Toolbox.ToSharedRef(); }
 
 	/** Get the edit mode */
-	FControlRigEditorEditMode* GetEditMode() const { return static_cast<FControlRigEditorEditMode*>(GetAssetEditorModeManager()->GetActiveMode(FControlRigEditorEditMode::ModeName)); }
+	FControlRigEditorEditMode* GetEditMode() const;
 
 	// this changes everytime you compile, so don't cache it expecting it will last. 
 	UControlRig* GetInstanceRig() const { return ControlRig;  }
