ControlRig: UpdateInput when bExecute is disabled

diff --git a/Engine/Plugins/Experimental/ControlRig/Source/ControlRigEditor/Private/Editor/ControlRigEditor.cpp b/Engine/Plugins/Experimental/ControlRig/Source/ControlRigEditor/Private/Editor/ControlRigEditor.cpp
--- a/Engine/Plugins/Experimental/ControlRig/Source/ControlRigEditor/Private/Editor/ControlRigEditor.cpp
+++ b/Engine/Plugins/Experimental/ControlRig/Source/ControlRigEditor/Private/Editor/ControlRigEditor.cpp
@@ -1520,7 +1520,7 @@ void FControlRigEditor::UpdateControlRig()
 
 			// initialize is moved post reinstance
 			FInputBlendPose Filter;
-			AnimInstance->UpdateControlRig(ControlRig, 0, false, false, Filter, 1.0f, bExecutionControlRig, bExecutionControlRig);
+			AnimInstance->UpdateControlRig(ControlRig, 0, false, false, Filter, 1.0f, true, bExecutionControlRig);
 			AnimInstance->RecalcRequiredBones();
 
 			// since rig has changed, rebuild draw skeleton
