Control Rig: Additional map for caching the previous parameter type on nodes

This will become obsolete once we merge the RigVM for 4.25.

diff --git a/Engine/Plugins/Experimental/ControlRig/Source/ControlRigDeveloper/Private/ControlRigBlueprint.cpp b/Engine/Plugins/Experimental/ControlRig/Source/ControlRigDeveloper/Private/ControlRigBlueprint.cpp
--- a/Engine/Plugins/Experimental/ControlRig/Source/ControlRigDeveloper/Private/ControlRigBlueprint.cpp
+++ b/Engine/Plugins/Experimental/ControlRig/Source/ControlRigDeveloper/Private/ControlRigBlueprint.cpp
@@ -207,6 +207,7 @@ void UControlRigBlueprint::PopulateModelFromGraph(const UControlRigGraph* InGrap
 
 		InitializeModel();
 		ModelController->Clear();
+		NodeToParameterType.Reset();
 
 		for (const UEdGraphNode* Node : InGraph->Nodes)
 		{
@@ -244,6 +245,7 @@ void UControlRigBlueprint::PopulateModelFromGraph(const UControlRigGraph* InGrap
 						ParameterType = EControlRigModelParameterType::Hidden;
 					}
 					ModelController->AddParameter(NodeName, DataType, ParameterType, NodePosition, false);
+					NodeToParameterType.Add(NodeName, ParameterType);
 				}
 				else
 				{
@@ -330,6 +332,7 @@ void UControlRigBlueprint::HandleModelModified(const UControlRigModel* InModel,
 			case EControlRigModelNotifType::ModelCleared:
 			{
 				LastNameFromNotification = NAME_None;
+				NodeToParameterType.Reset();
 
 				for (const FControlRigModelNode& Node : InModel->Nodes())
 				{
@@ -352,6 +355,7 @@ void UControlRigBlueprint::HandleModelModified(const UControlRigModel* InModel,
 						{
 							FControlRigBlueprintUtils::AddPropertyMember(this, Node->Pins[0].Type, *Node->Name.ToString());
 							HandleModelModified(InModel, EControlRigModelNotifType::NodeChanged, InPayload);
+							NodeToParameterType.Add(Node->Name, Node->ParameterType);
 							bValidNode = true;
 							break;
 						}
@@ -393,6 +397,7 @@ void UControlRigBlueprint::HandleModelModified(const UControlRigModel* InModel,
 					FBlueprintEditorUtils::RemoveMemberVariable(this, Node->Name);
 					FBlueprintEditorUtils::MarkBlueprintAsModified(this);
 					WatchedPins.Reset();
+					NodeToParameterType.Remove(Node->Name);
 				}
 				break;
 			}
@@ -410,17 +415,24 @@ void UControlRigBlueprint::HandleModelModified(const UControlRigModel* InModel,
 						UProperty* Property = GeneratedClass->FindPropertyByName(Node->Name);
 						if (Property != nullptr)
 						{
-							bool bWasInput = Property->HasMetaData(UControlRig::AnimationInputMetaName);
-							bool bWasOutput = Property->HasMetaData(UControlRig::AnimationOutputMetaName);
-
 							EControlRigModelParameterType WasParameterType = EControlRigModelParameterType::Hidden;
-							if (bWasInput)
+							if (NodeToParameterType.Contains(Node->Name))
 							{
-								WasParameterType = EControlRigModelParameterType::Input;
+								WasParameterType = NodeToParameterType.FindChecked(Node->Name);
 							}
-							else if (bWasOutput)
+							else
 							{
-								WasParameterType = EControlRigModelParameterType::Output;
+								bool bWasInput = Property->HasMetaData(UControlRig::AnimationInputMetaName);
+								bool bWasOutput = Property->HasMetaData(UControlRig::AnimationOutputMetaName);
+								if (bWasInput)
+								{
+									WasParameterType = EControlRigModelParameterType::Input;
+								}
+								else if (bWasOutput)
+								{
+									WasParameterType = EControlRigModelParameterType::Output;
+								}
+								NodeToParameterType.Add(Node->Name, WasParameterType);
 							}
 
 							if (WasParameterType != Node->ParameterType)
@@ -436,6 +448,8 @@ void UControlRigBlueprint::HandleModelModified(const UControlRigModel* InModel,
 								{
 									Property->SetMetaData(UControlRig::AnimationOutputMetaName, TEXT("True"));
 								}
+
+								NodeToParameterType[Node->Name] = Node->ParameterType;
 								FBlueprintEditorUtils::MarkBlueprintAsModified(this);
 								UpdateParametersOnControlRig();
 							}
diff --git a/Engine/Plugins/Experimental/ControlRig/Source/ControlRigDeveloper/Public/ControlRigBlueprint.h b/Engine/Plugins/Experimental/ControlRig/Source/ControlRigDeveloper/Public/ControlRigBlueprint.h
--- a/Engine/Plugins/Experimental/ControlRig/Source/ControlRigDeveloper/Public/ControlRigBlueprint.h
+++ b/Engine/Plugins/Experimental/ControlRig/Source/ControlRigDeveloper/Public/ControlRigBlueprint.h
@@ -154,6 +154,9 @@ public:
 	UPROPERTY(transient)
 	UControlRigController* ModelController;
 
+	// hack until RigVM refactoring
+	TMap<FName, EControlRigModelParameterType> NodeToParameterType;
+
 	bool bSuspendModelNotificationsForSelf;
 	bool bSuspendModelNotificationsForOthers;
 	FName LastNameFromNotification;
