Control Rig: Check firstval has been set

#rnx

diff --git a/Engine/Plugins/Experimental/ControlRig/Source/ControlRigEditor/Private/Sequencer/ControlRigParameterTrackEditor.cpp b/Engine/Plugins/Experimental/ControlRig/Source/ControlRigEditor/Private/Sequencer/ControlRigParameterTrackEditor.cpp
--- a/Engine/Plugins/Experimental/ControlRig/Source/ControlRigEditor/Private/Sequencer/ControlRigParameterTrackEditor.cpp
+++ b/Engine/Plugins/Experimental/ControlRig/Source/ControlRigEditor/Private/Sequencer/ControlRigParameterTrackEditor.cpp
@@ -1167,7 +1167,7 @@ void FControlRigParameterSection::BuildSectionContextMenu(FMenuBuilder& MenuBuil
 						}
 
 					}
-					return (FirstVal.GetValue()) ? ECheckBoxState::Checked : ECheckBoxState::Unchecked;
+					return (FirstVal.IsSet() && FirstVal.GetValue()) ? ECheckBoxState::Checked : ECheckBoxState::Unchecked;
 				}
 			})
 				);
