Control Rig: Ensure to re-initialize all rig units in case of a hierarchical change.

diff --git a/Engine/Plugins/Experimental/ControlRig/Source/ControlRigDeveloper/Private/ControlRigBlueprint.cpp b/Engine/Plugins/Experimental/ControlRig/Source/ControlRigDeveloper/Private/ControlRigBlueprint.cpp
--- a/Engine/Plugins/Experimental/ControlRig/Source/ControlRigDeveloper/Private/ControlRigBlueprint.cpp
+++ b/Engine/Plugins/Experimental/ControlRig/Source/ControlRigDeveloper/Private/ControlRigBlueprint.cpp
@@ -715,6 +715,7 @@ void UControlRigBlueprint::PropagateHierarchyFromBPToInstances(bool bInitialize)
 			}
 
 			DefaultObject->Hierarchy = HierarchyContainer;
+			DefaultObject->Initialize(true);
 
 			TArray<UObject*> ArchetypeInstances;
 			DefaultObject->GetArchetypeInstances(ArchetypeInstances);
@@ -723,6 +724,7 @@ void UControlRigBlueprint::PropagateHierarchyFromBPToInstances(bool bInitialize)
 				if (UControlRig* InstanceRig = Cast<UControlRig>(ArchetypeInstance))
 				{
 					InstanceRig->Hierarchy = HierarchyContainer;
+					InstanceRig->Initialize(true);
 				}
 			}
 		}
