Fixed crash during editor shutdown

diff --git a/Engine/Plugins/Experimental/ControlRig/Source/ControlRigEditor/Private/ControlRigEditorModule.cpp b/Engine/Plugins/Experimental/ControlRig/Source/ControlRigEditor/Private/ControlRigEditorModule.cpp
--- a/Engine/Plugins/Experimental/ControlRig/Source/ControlRigEditor/Private/ControlRigEditorModule.cpp
+++ b/Engine/Plugins/Experimental/ControlRig/Source/ControlRigEditor/Private/ControlRigEditorModule.cpp
@@ -356,9 +356,12 @@ void FControlRigEditorModule::ShutdownModule()
 		TrajectoryMaterial->RemoveFromRoot();
 	}
 
-	if (UAssetEditorSubsystem* AssetEditorSubsystem = GEditor->GetEditorSubsystem<UAssetEditorSubsystem>())
+	if (GEditor)
 	{
-		AssetEditorSubsystem->OnAssetEditorOpened().Remove(AssetEditorOpenedHandle);
+		if (UAssetEditorSubsystem* AssetEditorSubsystem = GEditor->GetEditorSubsystem<UAssetEditorSubsystem>())
+		{
+			AssetEditorSubsystem->OnAssetEditorOpened().Remove(AssetEditorOpenedHandle);
+		}
 	}
 
 	FEditorModeRegistry::Get().UnregisterMode(FControlRigEditorEditMode::ModeName);
