Control Rig: transform constraint node has incorrectly cached offsets

This fix has two parts:

a) Since the merge of 4.24 we've changed the behaviour of the getbonetransform node to only return a valid value during update. since the transform constraint node decided to cache results during init it was always caching identity matrices - thus the incorrect results during runtime.

b) the transform constraint node needs to be able to rely on the current pose / input pose for caching the constraint data - thus this cannot be done during init. instead we clear the data during init and then cache it during update if required.

diff --git a/Engine/Plugins/Experimental/ControlRig/Source/ControlRig/Private/Units/Hierarchy/RigUnit_GetBoneTransform.cpp b/Engine/Plugins/Experimental/ControlRig/Source/ControlRig/Private/Units/Hierarchy/RigUnit_GetBoneTransform.cpp
--- a/Engine/Plugins/Experimental/ControlRig/Source/ControlRig/Private/Units/Hierarchy/RigUnit_GetBoneTransform.cpp
+++ b/Engine/Plugins/Experimental/ControlRig/Source/ControlRig/Private/Units/Hierarchy/RigUnit_GetBoneTransform.cpp
@@ -19,7 +19,6 @@ FRigUnit_GetBoneTransform_Execute()
 			case EControlRigState::Init:
 			{
 				CachedBoneIndex = Hierarchy->GetIndex(Bone);
-				break;
 			}
 			case EControlRigState::Update:
 			{
diff --git a/Engine/Plugins/Experimental/ControlRig/Source/ControlRig/Private/Units/Hierarchy/RigUnit_GetControlInitialTransform.cpp b/Engine/Plugins/Experimental/ControlRig/Source/ControlRig/Private/Units/Hierarchy/RigUnit_GetControlInitialTransform.cpp
--- a/Engine/Plugins/Experimental/ControlRig/Source/ControlRig/Private/Units/Hierarchy/RigUnit_GetControlInitialTransform.cpp
+++ b/Engine/Plugins/Experimental/ControlRig/Source/ControlRig/Private/Units/Hierarchy/RigUnit_GetControlInitialTransform.cpp
@@ -19,7 +19,6 @@ FRigUnit_GetControlInitialTransform_Execute()
 			case EControlRigState::Init:
 			{
 				CachedControlIndex = Hierarchy->GetIndex(Control);
-				break;				
 			}
 			case EControlRigState::Update:
 			{
diff --git a/Engine/Plugins/Experimental/ControlRig/Source/ControlRig/Private/Units/Highlevel/Hierarchy/RigUnit_TransformConstraint.cpp b/Engine/Plugins/Experimental/ControlRig/Source/ControlRig/Private/Units/Highlevel/Hierarchy/RigUnit_TransformConstraint.cpp
--- a/Engine/Plugins/Experimental/ControlRig/Source/ControlRig/Private/Units/Highlevel/Hierarchy/RigUnit_TransformConstraint.cpp
+++ b/Engine/Plugins/Experimental/ControlRig/Source/ControlRig/Private/Units/Highlevel/Hierarchy/RigUnit_TransformConstraint.cpp
@@ -16,57 +16,59 @@ FRigUnit_TransformConstraint_Execute()
 	{
 		ConstraintData.Reset();
 		ConstraintDataToTargets.Reset();
-
+	}
+	else if (Context.State == EControlRigState::Update)
+	{
 		FRigBoneHierarchy* Hierarchy = ExecuteContext.GetBones();
 		if (Hierarchy)
 		{
-			int32 BoneIndex = Hierarchy->GetIndex(Bone);
-			if (BoneIndex != INDEX_NONE)
+			if (ConstraintData.Num() != Targets.Num())
 			{
-				const int32 TargetNum = Targets.Num();
-				if (TargetNum > 0)
+				ConstraintData.Reset();
+				ConstraintDataToTargets.Reset();
+
+				int32 BoneIndex = Hierarchy->GetIndex(Bone);
+				if (BoneIndex != INDEX_NONE)
 				{
-					const FTransform SourceTransform = Hierarchy->GetGlobalTransform(BoneIndex);
-					FTransform InputBaseTransform = UtilityHelpers::GetBaseTransformByMode(BaseTransformSpace, [Hierarchy](const FName& JointName) { return Hierarchy->GetGlobalTransform(JointName); },
-						(*Hierarchy)[BoneIndex].ParentName, BaseBone, BaseTransform);
-					for (int32 TargetIndex = 0; TargetIndex < TargetNum; ++TargetIndex)
+					const int32 TargetNum = Targets.Num();
+					if (TargetNum > 0)
 					{
-						// talk to Rob about the implication of support both of this
-						const bool bTranslationFilterValid = Targets[TargetIndex].Filter.TranslationFilter.IsValid();
-						const bool bRotationFilterValid = Targets[TargetIndex].Filter.RotationFilter.IsValid();
-						const bool bScaleFilterValid = Targets[TargetIndex].Filter.ScaleFilter.IsValid();
-
-						if (bTranslationFilterValid && bRotationFilterValid && bScaleFilterValid)
-						{
-							AddConstraintData(Targets, ETransformConstraintType::Parent, TargetIndex, SourceTransform, InputBaseTransform, ConstraintData, ConstraintDataToTargets);
-						}
-						else
+						const FTransform SourceTransform = Hierarchy->GetGlobalTransform(BoneIndex);
+						FTransform InputBaseTransform = UtilityHelpers::GetBaseTransformByMode(BaseTransformSpace, [Hierarchy](const FName& JointName) { return Hierarchy->GetGlobalTransform(JointName); },
+							(*Hierarchy)[BoneIndex].ParentName, BaseBone, BaseTransform);
+						for (int32 TargetIndex = 0; TargetIndex < TargetNum; ++TargetIndex)
 						{
-							if (bTranslationFilterValid)
-							{
-								AddConstraintData(Targets, ETransformConstraintType::Translation, TargetIndex, SourceTransform, InputBaseTransform, ConstraintData, ConstraintDataToTargets);
-							}
+							// talk to Rob about the implication of support both of this
+							const bool bTranslationFilterValid = Targets[TargetIndex].Filter.TranslationFilter.IsValid();
+							const bool bRotationFilterValid = Targets[TargetIndex].Filter.RotationFilter.IsValid();
+							const bool bScaleFilterValid = Targets[TargetIndex].Filter.ScaleFilter.IsValid();
 
-							if (bRotationFilterValid)
+							if (bTranslationFilterValid && bRotationFilterValid && bScaleFilterValid)
 							{
-								AddConstraintData(Targets, ETransformConstraintType::Rotation, TargetIndex, SourceTransform, InputBaseTransform, ConstraintData, ConstraintDataToTargets);
+								AddConstraintData(Targets, ETransformConstraintType::Parent, TargetIndex, SourceTransform, InputBaseTransform, ConstraintData, ConstraintDataToTargets);
 							}
-
-							if (bScaleFilterValid)
+							else
 							{
-								AddConstraintData(Targets, ETransformConstraintType::Scale, TargetIndex, SourceTransform, InputBaseTransform, ConstraintData, ConstraintDataToTargets);
+								if (bTranslationFilterValid)
+								{
+									AddConstraintData(Targets, ETransformConstraintType::Translation, TargetIndex, SourceTransform, InputBaseTransform, ConstraintData, ConstraintDataToTargets);
+								}
+
+								if (bRotationFilterValid)
+								{
+									AddConstraintData(Targets, ETransformConstraintType::Rotation, TargetIndex, SourceTransform, InputBaseTransform, ConstraintData, ConstraintDataToTargets);
+								}
+
+								if (bScaleFilterValid)
+								{
+									AddConstraintData(Targets, ETransformConstraintType::Scale, TargetIndex, SourceTransform, InputBaseTransform, ConstraintData, ConstraintDataToTargets);
+								}
 							}
 						}
 					}
 				}
 			}
-		}
-	}
-	else if (Context.State == EControlRigState::Update)
-	{
-		FRigBoneHierarchy* Hierarchy = ExecuteContext.GetBones();
-		if (Hierarchy)
-		{
+
 			int32 BoneIndex = Hierarchy->GetIndex(Bone);
 			if (BoneIndex != INDEX_NONE)
 			{
