Control Rig: FRotator pins have incorrect display names

diff --git a/Engine/Plugins/Experimental/ControlRig/Source/ControlRigDeveloper/Private/ControlRigModel.cpp b/Engine/Plugins/Experimental/ControlRig/Source/ControlRigDeveloper/Private/ControlRigModel.cpp
--- a/Engine/Plugins/Experimental/ControlRig/Source/ControlRigDeveloper/Private/ControlRigModel.cpp
+++ b/Engine/Plugins/Experimental/ControlRig/Source/ControlRigDeveloper/Private/ControlRigModel.cpp
@@ -2594,6 +2594,12 @@ void UControlRigModel::ConfigurePinFromField(FControlRigModelPin& Pin, FProperty
 		Pin.Direction = PinDirection;
 		ConfigurePinFromField(Pin, *It, Node);
 
+		// override the Y, X, Z names with Roll, Yaw and Pitch
+		if (Struct == TBaseStructure<FRotator>::Get())
+		{
+			Pin.DisplayNameText = FText::FromName(Pin.Name);
+		}
+
 		if (Pin.Index == Node.Pins.Num())
 		{
 			Node.Pins.Add(Pin);
