Cook Errors with GizmoLibrary

diff --git a/Engine/Plugins/Experimental/ControlRig/Source/ControlRig/Private/Settings/ControlRigSettings.cpp b/Engine/Plugins/Experimental/ControlRig/Source/ControlRig/Private/Settings/ControlRigSettings.cpp
--- a/Engine/Plugins/Experimental/ControlRig/Source/ControlRig/Private/Settings/ControlRigSettings.cpp
+++ b/Engine/Plugins/Experimental/ControlRig/Source/ControlRig/Private/Settings/ControlRigSettings.cpp
@@ -5,7 +5,9 @@
 UControlRigSettings::UControlRigSettings(const FObjectInitializer& ObjectInitializer)
 	: Super(ObjectInitializer)
 {
+#if WITH_EDITORONLY_DATA
 	DefaultGizmoLibrary = LoadObject<UControlRigGizmoLibrary>(nullptr, TEXT("/ControlRig/Controls/DefaultGizmoLibrary.DefaultGizmoLibrary"));
+#endif
 }
 
 
diff --git a/Engine/Plugins/Experimental/ControlRig/Source/ControlRig/Public/Settings/ControlRigSettings.h b/Engine/Plugins/Experimental/ControlRig/Source/ControlRig/Public/Settings/ControlRigSettings.h
--- a/Engine/Plugins/Experimental/ControlRig/Source/ControlRig/Public/Settings/ControlRigSettings.h
+++ b/Engine/Plugins/Experimental/ControlRig/Source/ControlRig/Public/Settings/ControlRigSettings.h
@@ -23,8 +23,10 @@ class CONTROLRIG_API UControlRigSettings : public UDeveloperSettings
 
 public:
 
+#if WITH_EDITORONLY_DATA
 	UPROPERTY(EditAnywhere, config, Category = DefaultGizmo)
 	TAssetPtr<UControlRigGizmoLibrary> DefaultGizmoLibrary;
+#endif
 
 	static UControlRigSettings * Get() { return CastChecked<UControlRigSettings>(UControlRigSettings::StaticClass()->GetDefaultObject()); }
 };
