Control Rig: Manipulation layer - switch manipulatable object storage to weak ptr. [THIS CHANGE IS ALREAY IN DEV-ANIM as well]

diff --git a/Engine/Plugins/Experimental/ControlRig/Source/ControlRig/Private/Manipulatable/IControlRigManipulatable.cpp b/Engine/Plugins/Experimental/ControlRig/Source/ControlRig/Private/Manipulatable/IControlRigManipulatable.cpp
--- a/Engine/Plugins/Experimental/ControlRig/Source/ControlRig/Private/Manipulatable/IControlRigManipulatable.cpp
+++ b/Engine/Plugins/Experimental/ControlRig/Source/ControlRig/Private/Manipulatable/IControlRigManipulatable.cpp
@@ -2,18 +2,14 @@
 
 #include "Manipulatable/IControlRigManipulatable.h"
 
-IControlRigManipulatable::IControlRigManipulatable()
+UControlRigManipulatable::UControlRigManipulatable(const FObjectInitializer& ObjectInitializer)
+	: Super(ObjectInitializer)
 {
-	bManipulationEnabled = false;
 }
 
-IControlRigManipulatable::~IControlRigManipulatable()
+IControlRigManipulatable::IControlRigManipulatable()
 {
-	OnFilterControl.Clear();
-	OnControlModified.Clear();
-#if WITH_EDITOR
-	OnControlSelected.Clear();
-#endif
+	bManipulationEnabled = false;
 }
 
 bool IControlRigManipulatable::SetControlGlobalTransform(const FName& InControlName, const FTransform& InGlobalTransform)
diff --git a/Engine/Plugins/Experimental/ControlRig/Source/ControlRig/Public/Manipulatable/IControlRigManipulatable.h b/Engine/Plugins/Experimental/ControlRig/Source/ControlRig/Public/Manipulatable/IControlRigManipulatable.h
--- a/Engine/Plugins/Experimental/ControlRig/Source/ControlRig/Public/Manipulatable/IControlRigManipulatable.h
+++ b/Engine/Plugins/Experimental/ControlRig/Source/ControlRig/Public/Manipulatable/IControlRigManipulatable.h
@@ -7,6 +7,8 @@
 #include "Rigs/RigControlHierarchy.h"
 #include "Rigs/RigSpaceHierarchy.h"
 #include "ControlRigGizmoLibrary.h"
+#include "UObject/Interface.h"
+#include "IControlRigManipulatable.generated.h"
 
 /**
  *
@@ -19,11 +21,19 @@
 
 class IControlRigObjectBinding;
 
+UINTERFACE(meta = (CannotImplementInterfaceInBlueprint))
+class CONTROLRIG_API UControlRigManipulatable : public UInterface
+{
+	GENERATED_UINTERFACE_BODY()
+};
+
 class CONTROLRIG_API IControlRigManipulatable
 {
 public:
+
+	GENERATED_IINTERFACE_BODY()
+
 	IControlRigManipulatable();
-	virtual ~IControlRigManipulatable();
 
 	/** Bindable event for external objects to contribute to / filter a control value */
 	DECLARE_EVENT_ThreeParams(IControlRigManipulatable, FFilterControlEvent, IControlRigManipulatable*, const FRigControl&, FRigControlValue&);
diff --git a/Engine/Plugins/Experimental/ControlRig/Source/ControlRigManipulation/Private/DefaultControlRigManipulationLayer.cpp b/Engine/Plugins/Experimental/ControlRig/Source/ControlRigManipulation/Private/DefaultControlRigManipulationLayer.cpp
--- a/Engine/Plugins/Experimental/ControlRig/Source/ControlRigManipulation/Private/DefaultControlRigManipulationLayer.cpp
+++ b/Engine/Plugins/Experimental/ControlRig/Source/ControlRigManipulation/Private/DefaultControlRigManipulationLayer.cpp
@@ -172,19 +172,31 @@ void UDefaultControlRigManipulationLayer::GetGizmoCreationParams(TArray<FGizmoAc
 	OutCreationParams.Reset();
 
 	// for now we only support FTransform
-	for (IControlRigManipulatable* ManipObject : ManipulatableObjects)
+	for (TWeakObjectPtr<UObject> ManipObject : ManipulatableObjects)
 	{
-		const TArray<FRigControl>& Controls = ManipObject->AvailableControls();
-		UControlRigGizmoLibrary* GizmoLibrary = ManipObject->GetGizmoLibrary();
+		if (!ManipObject.IsValid())
+		{
+			continue;
+		}
+
+		IControlRigManipulatable* Manipulatable = Cast<IControlRigManipulatable>(ManipObject.Get());
+		if (Manipulatable == nullptr)
+		{
+			checkNoEntry();
+			continue;
+		}
+
+		const TArray<FRigControl>& Controls = Manipulatable->AvailableControls();
+		UControlRigGizmoLibrary* GizmoLibrary = Manipulatable->GetGizmoLibrary();
 
 		for (const FRigControl& Control : Controls)
 		{
 			if (IsSupportedControlType(Control.ControlType))
 			{
 				FGizmoActorCreationParam Param;
-				Param.ManipObj = ManipObject;
+				Param.ManipObj = Manipulatable;
 				Param.ControlName = Control.Name;
-				Param.SpawnTransform = ManipObject->GetControlGlobalTransform(Control.Name);
+				Param.SpawnTransform = Manipulatable->GetControlGlobalTransform(Control.Name);
 				Param.GizmoTransform = Control.GizmoTransform;
 
 				if (GizmoLibrary)
@@ -372,10 +384,9 @@ void UDefaultControlRigManipulationLayer::OnControlModified(IControlRigManipulat
 
 TSharedPtr<IControlRigObjectBinding> UDefaultControlRigManipulationLayer::GetObjectBinding() const
 {
-	for (IControlRigManipulatable* Manip : ManipulatableObjects)
+	for (TWeakObjectPtr<UObject> Manip : ManipulatableObjects)
 	{
-		UControlRig* ControlRig = static_cast<UControlRig*>(Manip);
-		if (ControlRig)
+		if (UControlRig* ControlRig = Cast<UControlRig>(Manip))
 		{
 			return ControlRig->GetObjectBinding();
 		}
@@ -386,10 +397,9 @@ TSharedPtr<IControlRigObjectBinding> UDefaultControlRigManipulationLayer::GetObj
 
 void UDefaultControlRigManipulationLayer::SetObjectBinding(TSharedPtr<IControlRigObjectBinding> InObjectBinding)
 {
-	for (IControlRigManipulatable* Manip : ManipulatableObjects)
+	for (TWeakObjectPtr<UObject> Manip : ManipulatableObjects)
 	{
-		UControlRig* ControlRig = static_cast<UControlRig*>(Manip);
-		if (ControlRig)
+		if (UControlRig* ControlRig = Cast<UControlRig>(Manip))
 		{
 			ControlRig->SetObjectBinding(InObjectBinding);
 		}
@@ -415,11 +425,13 @@ FTransform	UDefaultControlRigManipulationLayer::GetSkeletalMeshComponentTransfor
 
 void UDefaultControlRigManipulationLayer::BeginTransaction()
 {
-	for (IControlRigManipulatable* Manip : ManipulatableObjects)
+	for (TWeakObjectPtr<UObject> Manip : ManipulatableObjects)
 	{
-		UControlRig* ControlRig = static_cast<UControlRig*>(Manip);
-		ControlRig->SetFlags(RF_Transactional);
-		ControlRig->Modify();
+		if (UControlRig* ControlRig = Cast<UControlRig>(Manip))
+		{
+			ControlRig->SetFlags(RF_Transactional);
+			ControlRig->Modify();
+		}
 	}
 }
 
diff --git a/Engine/Plugins/Experimental/ControlRig/Source/ControlRigManipulation/Private/IControlRigManipulationLayer.cpp b/Engine/Plugins/Experimental/ControlRig/Source/ControlRigManipulation/Private/IControlRigManipulationLayer.cpp
--- a/Engine/Plugins/Experimental/ControlRig/Source/ControlRigManipulation/Private/IControlRigManipulationLayer.cpp
+++ b/Engine/Plugins/Experimental/ControlRig/Source/ControlRigManipulation/Private/IControlRigManipulationLayer.cpp
@@ -1,6 +1,12 @@
 // Copyright Epic Games, Inc. All Rights Reserved.
 
 #include "IControlRigManipulationLayer.h"
+#include "Manipulatable/IControlRigManipulatable.h"
+
+UControlRigManipulationLayer::UControlRigManipulationLayer(const FObjectInitializer& ObjectInitializer)
+	: Super(ObjectInitializer)
+{
+}
 
 IControlRigManipulationLayer::IControlRigManipulationLayer()
 	:bLayerCreated(false)
@@ -23,28 +29,54 @@ void IControlRigManipulationLayer::DestroyLayer()
 
 	DestroyGizmosActors();
 
-	for (int32 Index = ManipulatableObjects.Num() - 1; Index >= 0 ; --Index)
+	for (int32 Index = ManipulatableObjects.Num() - 1; Index >= 0; --Index)
 	{
-		IControlRigManipulatable* ManipObj = ManipulatableObjects[Index];
-		RemoveManipulatableObject(ManipObj);
-	}
+		if (!ManipulatableObjects[Index].IsValid())
+		{
+			continue;
+		}
 
-	ManipulatableObjects.Num();
+		if (IControlRigManipulatable* Manipulatable = Cast<IControlRigManipulatable>(ManipulatableObjects[Index].Get()))
+		{
+			RemoveManipulatableObject(Manipulatable);
+		}
+	}
 }
 
 void IControlRigManipulationLayer::AddManipulatableObject(IControlRigManipulatable* InObject)
 {
-	if (InObject && !ManipulatableObjects.Contains(InObject))
+	UObject* ManipulatableUObject = Cast<UObject>(InObject);
+	if (ManipulatableUObject == nullptr)
+	{
+		return;
+	}
+
+	for (int32 Index = ManipulatableObjects.Num() - 1; Index >= 0; --Index)
 	{
-		ManipulatableObjects.Add(InObject);
+		if (ManipulatableObjects[Index].Get() == ManipulatableUObject)
+		{
+			return;
+		}
 	}
+
+	ManipulatableObjects.Add(TWeakObjectPtr<UObject>(ManipulatableUObject));
 }
 
 void IControlRigManipulationLayer::RemoveManipulatableObject(IControlRigManipulatable* InObject)
 {
-	if (InObject && ManipulatableObjects.Contains(InObject))
+	UObject* ManipulatableUObject = Cast<UObject>(InObject);
+	if (ManipulatableUObject == nullptr)
+	{
+		return;
+	}
+
+	for (int32 Index = ManipulatableObjects.Num() - 1; Index >= 0; --Index)
 	{
-		ManipulatableObjects.Remove(InObject);
+		if (ManipulatableObjects[Index].Get() == ManipulatableUObject)
+		{
+			ManipulatableObjects.RemoveAt(Index);
+			return;
+		}
 	}
 }
 
diff --git a/Engine/Plugins/Experimental/ControlRig/Source/ControlRigManipulation/Public/IControlRigManipulationLayer.h b/Engine/Plugins/Experimental/ControlRig/Source/ControlRigManipulation/Public/IControlRigManipulationLayer.h
--- a/Engine/Plugins/Experimental/ControlRig/Source/ControlRigManipulation/Public/IControlRigManipulationLayer.h
+++ b/Engine/Plugins/Experimental/ControlRig/Source/ControlRigManipulation/Public/IControlRigManipulationLayer.h
@@ -4,11 +4,19 @@
 
 #include "CoreMinimal.h"
 #include "UnrealWidget.h"
+#include "UObject/Interface.h"
+#include "IControlRigManipulationLayer.generated.h"
 
 class ISequencer;
 class IControlRigManipulatable;
 class AControlRigGizmoActor;
 
+UINTERFACE(meta = (CannotImplementInterfaceInBlueprint))
+class CONTROLRIGMANIPULATION_API UControlRigManipulationLayer : public UInterface
+{
+	GENERATED_UINTERFACE_BODY()
+};
+
 /**
  * Manipulation Layer interface to create interactive interface
  * This work with IManipulatable object that provides options
@@ -16,8 +24,10 @@ class AControlRigGizmoActor;
 class CONTROLRIGMANIPULATION_API IControlRigManipulationLayer
 {
 public:
+
+	GENERATED_IINTERFACE_BODY()
+
 	IControlRigManipulationLayer();
-	virtual ~IControlRigManipulationLayer() {};
 
 	virtual void CreateLayer();
 	virtual void DestroyLayer();
@@ -40,8 +50,7 @@ public:
 protected:
 
 	// because manipulatable is created externally often. Manipulation Layer doesn't have any ownership on these object
-	// @fixme: We may want to make weak ptr? Can't be object ptr because IControlRiGManipulatable is not from UObject
-	TArray<IControlRigManipulatable*> ManipulatableObjects;
+	TArray<TWeakObjectPtr<UObject>> ManipulatableObjects;
 
 private:
 	// tracks whether this layer is created or not
