Control Rig: Preview mesh not deforming when debugging actors in PIE

diff --git a/Engine/Plugins/Experimental/ControlRig/Source/ControlRigEditor/Private/Editor/ControlRigEditor.cpp b/Engine/Plugins/Experimental/ControlRig/Source/ControlRigEditor/Private/Editor/ControlRigEditor.cpp
--- a/Engine/Plugins/Experimental/ControlRig/Source/ControlRigEditor/Private/Editor/ControlRigEditor.cpp
+++ b/Engine/Plugins/Experimental/ControlRig/Source/ControlRigEditor/Private/Editor/ControlRigEditor.cpp
@@ -504,8 +504,7 @@ void FControlRigEditor::HandleSetObjectBeingDebugged(UObject* InObject)
 	if (DebuggedControlRig)
 	{
 		bool bIsExternalControlRig = DebuggedControlRig != ControlRig;
-		// @fixme: I'll also change the "execute" button to be not toggleable during debugging external control rig
-		bool bShouldExecute = bIsExternalControlRig && bExecutionControlRig;
+		bool bShouldExecute = (!bIsExternalControlRig) && bExecutionControlRig;
 		DebuggedControlRig->ControlRigLog = &ControlRigLog;
 
 		UControlRigSkeletalMeshComponent* EditorSkelComp = Cast<UControlRigSkeletalMeshComponent>(GetPersonaToolkit()->GetPreviewScene()->GetPreviewMeshComponent());
@@ -516,7 +515,16 @@ void FControlRigEditor::HandleSetObjectBeingDebugged(UObject* InObject)
 			{
 				// we might want to move this into another method
 				FInputBlendPose Filter;
-				AnimInstance->UpdateControlRig(DebuggedControlRig, 0, false, false, Filter, 1.0f, bShouldExecute, bShouldExecute);
+				AnimInstance->UpdateControlRig(
+					DebuggedControlRig,
+					0,
+					false,
+					false,
+					Filter,
+					1.0f,
+					bShouldExecute, // update input
+					bShouldExecute // execute rig
+				);
 				AnimInstance->RecalcRequiredBones();
 
 				// since rig has changed, rebuild draw skeleton
