Control Rig: Remove warning in "quat from two vectors nodes"

The node will default to return an identity rotation.

diff --git a/Engine/Plugins/Experimental/ControlRig/Source/ControlRig/Private/Units/Math/RigUnit_MathQuaternion.cpp b/Engine/Plugins/Experimental/ControlRig/Source/ControlRig/Private/Units/Math/RigUnit_MathQuaternion.cpp
--- a/Engine/Plugins/Experimental/ControlRig/Source/ControlRig/Private/Units/Math/RigUnit_MathQuaternion.cpp
+++ b/Engine/Plugins/Experimental/ControlRig/Source/ControlRig/Private/Units/Math/RigUnit_MathQuaternion.cpp
@@ -32,14 +32,6 @@ FRigUnit_MathQuaternionFromTwoVectors_Execute()
     DECLARE_SCOPE_HIERARCHICAL_COUNTER_RIGUNIT()
 	if (A.IsNearlyZero() || B.IsNearlyZero())
 	{
-		if (A.IsNearlyZero())
-		{
-			UE_CONTROLRIG_RIGUNIT_REPORT_WARNING(TEXT("A is nearly zero"));
-		}
-		if (B.IsNearlyZero())
-		{
-			UE_CONTROLRIG_RIGUNIT_REPORT_WARNING(TEXT("B is nearly zero"));
-		}
 		Result = FQuat::Identity;
 		return;
 	}
