Rename RigHierarchy to RigBoneHierarchy - separate RigHierarchy and RigHierarchyContainer

diff --git a/Engine/Plugins/Experimental/ControlRig/Config/BaseControlRig.ini b/Engine/Plugins/Experimental/ControlRig/Config/BaseControlRig.ini
--- a/Engine/Plugins/Experimental/ControlRig/Config/BaseControlRig.ini
+++ b/Engine/Plugins/Experimental/ControlRig/Config/BaseControlRig.ini
@@ -8,4 +8,6 @@
 
 
 +StructRedirects=(OldName="RigJoint",NewName="/Script/ControlRig.RigBone")
-+PropertyRedirects=(OldName="RigHierarchy.Joints",NewName="RigHierarchy.Bones")
++StructRedirects=(OldName="RigHierarchy",NewName="/Script/ControlRig.RigBoneHierarchy")
++PropertyRedirects=(OldName="RigBoneHierarchy.Joints",NewName="RigBoneHierarchy.Bones")
++PropertyRedirects=(OldName="RigHierarchyContainer.BaseHierarchy",NewName="RigHierarchyContainer.BoneHierarchy")
diff --git a/Engine/Plugins/Experimental/ControlRig/Source/ControlRig/Private/ControlRig.cpp b/Engine/Plugins/Experimental/ControlRig/Source/ControlRig/Private/ControlRig.cpp
--- a/Engine/Plugins/Experimental/ControlRig/Source/ControlRig/Private/ControlRig.cpp
+++ b/Engine/Plugins/Experimental/ControlRig/Source/ControlRig/Private/ControlRig.cpp
@@ -199,8 +199,9 @@ void UControlRig::Initialize(bool bInitRigUnits)
 #endif // WITH_EDITORONLY_DATA
 
 	// should refresh mapping 
-	Hierarchy.BaseHierarchy.Initialize();
+	Hierarchy.Initialize();
 	CurveContainer.Initialize();
+
 	// resolve IO properties
 	ResolveInputOutputProperties();
 
@@ -317,7 +318,6 @@ void UControlRig::Execute(const EControlRigState InState)
 	Context.DeltaTime = DeltaTime;
 	Context.State = InState;
 	Context.HierarchyReference.Container = &Hierarchy;
-	Context.HierarchyReference.bUseBaseHierarchy = true;
 	Context.CurveReference = FRigCurveContainerRef(&CurveContainer);
 
 #if WITH_EDITOR
@@ -395,37 +395,37 @@ FTransform UControlRig::GetGlobalTransform(const FName& BoneName) const
 {
 	DECLARE_SCOPE_HIERARCHICAL_COUNTER_FUNC()
 
-	int32 Index = Hierarchy.BaseHierarchy.GetIndex(BoneName);
+	int32 Index = Hierarchy.BoneHierarchy.GetIndex(BoneName);
 	if (Index != INDEX_NONE)
 	{
-		return Hierarchy.BaseHierarchy.GetGlobalTransform(Index);
+		return Hierarchy.BoneHierarchy.GetGlobalTransform(Index);
 	}
 
 	return FTransform::Identity;
 
 }
 
-void UControlRig::SetGlobalTransform(const FName& BoneName, const FTransform& InTransform, bool bPropagateTransform) 
+void UControlRig::SetGlobalTransform(const FName& BoneName, const FTransform& InTransform, bool bPropagateTransform)
 {
 	DECLARE_SCOPE_HIERARCHICAL_COUNTER_FUNC()
 
-	int32 Index = Hierarchy.BaseHierarchy.GetIndex(BoneName);
+	int32 Index = Hierarchy.BoneHierarchy.GetIndex(BoneName);
 	if (Index != INDEX_NONE)
 	{
-		Hierarchy.BaseHierarchy.SetGlobalTransform(Index, InTransform, bPropagateTransform);
+		return Hierarchy.BoneHierarchy.SetGlobalTransform(Index, InTransform, bPropagateTransform);
 	}
 }
 
 FTransform UControlRig::GetGlobalTransform(const int32 BoneIndex) const
 {
 	DECLARE_SCOPE_HIERARCHICAL_COUNTER_FUNC()
-	return Hierarchy.BaseHierarchy.GetGlobalTransform(BoneIndex);
+	return Hierarchy.BoneHierarchy.GetGlobalTransform(BoneIndex);
 }
 
 void UControlRig::SetGlobalTransform(const int32 BoneIndex, const FTransform& InTransform, bool bPropagateTransform)
 {
 	DECLARE_SCOPE_HIERARCHICAL_COUNTER_FUNC()
-	Hierarchy.BaseHierarchy.SetGlobalTransform(BoneIndex, InTransform, bPropagateTransform);
+	Hierarchy.BoneHierarchy.SetGlobalTransform(BoneIndex, InTransform, bPropagateTransform);
 }
 
 float UControlRig::GetCurveValue(const FName& CurveName) const
@@ -461,13 +461,12 @@ void UControlRig::GetMappableNodeData(TArray<FName>& OutNames, TArray<FNodeItem>
 	OutNodeItems.Reset();
 
 	// now add all nodes
-	const FRigHierarchy& BaseHierarchy = Hierarchy.BaseHierarchy;
+	const FRigBoneHierarchy& BoneHierarchy = Hierarchy.BoneHierarchy;
 
-	const TArray<FRigBone>& Bones = BaseHierarchy.GetBones();
-	for (int32 Index = 0; Index < Bones.Num(); ++Index)
+	for (const FRigBone& Bone : BoneHierarchy)
 	{
-		OutNames.Add(Bones[Index].Name);
-		OutNodeItems.Add(FNodeItem(Bones[Index].ParentName, Bones[Index].InitialTransform));
+		OutNames.Add(Bone.Name);
+		OutNodeItems.Add(FNodeItem(Bone.ParentName, Bone.InitialTransform));
 	}
 
 	// have to handle curve separate because otherwise it will confuse with hierarchy
diff --git a/Engine/Plugins/Experimental/ControlRig/Source/ControlRig/Private/Drawing/ControlRigDrawInterface.cpp b/Engine/Plugins/Experimental/ControlRig/Source/ControlRig/Private/Drawing/ControlRigDrawInterface.cpp
--- a/Engine/Plugins/Experimental/ControlRig/Source/ControlRig/Private/Drawing/ControlRigDrawInterface.cpp
+++ b/Engine/Plugins/Experimental/ControlRig/Source/ControlRig/Private/Drawing/ControlRigDrawInterface.cpp
@@ -161,10 +161,8 @@ void FControlRigDrawInterface::DrawBezier(const FTransform& WorldOffset, const F
 	DrawInstructions.Add(Instruction);
 }
 
-void FControlRigDrawInterface::DrawHierarchy(const FTransform& WorldOffset, const FRigHierarchy& Hierarchy, EControlRigDrawHierarchyMode::Type Mode, float Scale, const FLinearColor& Color, float Thickness)
+void FControlRigDrawInterface::DrawHierarchy(const FTransform& WorldOffset, const FRigBoneHierarchy& Hierarchy, EControlRigDrawHierarchyMode::Type Mode, float Scale, const FLinearColor& Color, float Thickness)
 {
-	const TArray<FRigBone>& Bones = Hierarchy.GetBones();
-
 	switch (Mode)
 	{
 		case EControlRigDrawHierarchyMode::Axes:
@@ -173,12 +171,12 @@ void FControlRigDrawInterface::DrawHierarchy(const FTransform& WorldOffset, cons
 			FDrawIntruction InstructionY(EDrawType_Lines, FLinearColor::Green, Thickness);
 			FDrawIntruction InstructionZ(EDrawType_Lines, FLinearColor::Blue, Thickness);
 			FDrawIntruction InstructionParent(EDrawType_Lines, Color, Thickness);
-			InstructionX.Positions.Reserve(Bones.Num() * 2);
-			InstructionY.Positions.Reserve(Bones.Num() * 2);
-			InstructionZ.Positions.Reserve(Bones.Num() * 2);
-			InstructionParent.Positions.Reserve(Bones.Num() * 6);
+			InstructionX.Positions.Reserve(Hierarchy.Num() * 2);
+			InstructionY.Positions.Reserve(Hierarchy.Num() * 2);
+			InstructionZ.Positions.Reserve(Hierarchy.Num() * 2);
+			InstructionParent.Positions.Reserve(Hierarchy.Num() * 6);
 
-			for (const FRigBone& Bone : Bones)
+			for (const FRigBone& Bone : Hierarchy)
 			{
 				FTransform Transform = Bone.GlobalTransform * WorldOffset;
 				FVector P0 = Transform.GetLocation();
@@ -194,7 +192,7 @@ void FControlRigDrawInterface::DrawHierarchy(const FTransform& WorldOffset, cons
 
 				if (Bone.ParentIndex != INDEX_NONE)
 				{
-					FTransform ParentTransform = Bones[Bone.ParentIndex].GlobalTransform * WorldOffset;
+					FTransform ParentTransform = Hierarchy[Bone.ParentIndex].GlobalTransform * WorldOffset;
 					FVector P1 = ParentTransform.GetLocation();
 					InstructionParent.Positions.Add(P0);
 					InstructionParent.Positions.Add(P1);
diff --git a/Engine/Plugins/Experimental/ControlRig/Source/ControlRig/Private/Rigs/CurveContainer.cpp b/Engine/Plugins/Experimental/ControlRig/Source/ControlRig/Private/Rigs/CurveContainer.cpp
deleted file mode 100644
diff --git a/Engine/Plugins/Experimental/ControlRig/Source/ControlRig/Private/Rigs/Hierarchy.cpp b/Engine/Plugins/Experimental/ControlRig/Source/ControlRig/Private/Rigs/Hierarchy.cpp
deleted file mode 100644
diff --git a/Engine/Plugins/Experimental/ControlRig/Source/ControlRig/Private/Rigs/RigBoneHierarchy.cpp b/Engine/Plugins/Experimental/ControlRig/Source/ControlRig/Private/Rigs/RigBoneHierarchy.cpp
new file mode 100644
--- /dev/null
+++ b/Engine/Plugins/Experimental/ControlRig/Source/ControlRig/Private/Rigs/RigBoneHierarchy.cpp
@@ -0,0 +1,474 @@
+// Copyright 1998-2019 Epic Games, Inc. All Rights Reserved.
+
+#include "Rigs/RigBoneHierarchy.h"
+#include "ControlRig.h"
+#include "HelperUtil.h"
+
+////////////////////////////////////////////////////////////////////////////////
+// FRigBoneHierarchy
+////////////////////////////////////////////////////////////////////////////////
+
+FRigBoneHierarchy::FRigBoneHierarchy()
+#if WITH_EDITOR
+	:Container(nullptr)
+#endif
+{
+}
+
+FRigBoneHierarchy& FRigBoneHierarchy::operator= (const FRigBoneHierarchy &InOther)
+{
+#if WITH_EDITOR
+	for (int32 Index = Num() - 1; Index >= 0; Index--)
+	{
+		FRigBone BoneToRemove = Bones[Index];
+		OnBoneRemoved.Broadcast(Container, ERigHierarchyElementType::Bone, BoneToRemove.Name);
+	}
+#endif
+
+	Bones.Reset();
+	Bones.Append(InOther.Bones);
+	NameToIndexMapping.Reset();
+	RefreshMapping();
+
+#if WITH_EDITOR
+	for (const FRigBone& BoneAdded : Bones)
+	{
+		OnBoneAdded.Broadcast(Container, ERigHierarchyElementType::Bone, BoneAdded.Name);
+	}
+#endif
+
+	return *this;
+}
+
+FRigBone& FRigBoneHierarchy::Add(const FName& InNewName, const FName& InParentName, const FTransform& InInitTransform)
+{
+	int32 ParentIndex = GetIndex(InParentName);
+	bool bHasParent = (ParentIndex != INDEX_NONE);
+
+	FRigBone NewBone;
+	NewBone.Name = InNewName;
+	NewBone.ParentIndex = ParentIndex;
+	NewBone.ParentName = bHasParent? InParentName : NAME_None;
+	NewBone.InitialTransform = InInitTransform;
+	NewBone.GlobalTransform = InInitTransform;
+	RecalculateLocalTransform(NewBone);
+
+	Bones.Add(NewBone);
+	RefreshMapping();
+
+#if WITH_EDITOR
+	OnBoneAdded.Broadcast(Container, ERigHierarchyElementType::Bone, InNewName);
+#endif
+
+	int32 Index = GetIndex(InNewName);
+	return Bones[Index];
+}
+
+FRigBone& FRigBoneHierarchy::Add(const FName& InNewName, const FName& InParentName, const FTransform& InInitTransform, const FTransform& InLocalTransform, const FTransform& InGlobalTransform)
+{
+	FRigBone& Bone = Add(InNewName, InParentName, InInitTransform);
+	Bone.LocalTransform = InLocalTransform;
+	Bone.GlobalTransform = InGlobalTransform;
+	return Bone;
+}
+
+void FRigBoneHierarchy::Reparent(const FName& InName, const FName& InNewParentName)
+{
+	int32 Index = GetIndex(InName);
+	// can't parent to itself
+	if (Index != INDEX_NONE && InName != InNewParentName)
+	{
+		// should allow reparent to none (no parent)
+		// if invalid, we consider to be none
+		int32 ParentIndex = GetIndex(InNewParentName);
+		bool bHasParent = (ParentIndex != INDEX_NONE);
+		FRigBone& CurBone = Bones[Index];
+		CurBone.ParentIndex = ParentIndex;
+		FName OldParentName = CurBone.ParentName;
+		CurBone.ParentName = (bHasParent)? InNewParentName : NAME_None;
+		RecalculateLocalTransform(CurBone);
+
+		// we want to make sure parent is before the child
+		RefreshMapping();
+
+#if WITH_EDITOR
+		if (OldParentName != InNewParentName)
+		{
+			OnBoneReparented.Broadcast(Container, ERigHierarchyElementType::Bone, InName, OldParentName, InNewParentName);
+		}
+#endif
+	}
+}
+
+FRigBone FRigBoneHierarchy::Remove(const FName& InNameToRemove)
+{
+	TArray<int32> Children;
+#if WITH_EDITOR
+	TArray<FName> RemovedChildBones;
+#endif
+	if (GetChildren(InNameToRemove, Children, true) > 0)
+	{
+		// sort by child index
+		Children.Sort([](const int32& A, const int32& B) { return A < B; });
+
+		// want to delete from end to the first 
+		for (int32 ChildIndex = Children.Num() - 1; ChildIndex >= 0; --ChildIndex)
+		{
+#if WITH_EDITOR
+			RemovedChildBones.Add(Bones[Children[ChildIndex]].Name);
+#endif
+			Bones.RemoveAt(Children[ChildIndex]);
+		}
+	}
+
+	int32 IndexToDelete = GetIndex(InNameToRemove);
+	FRigBone RemovedBone = Bones[IndexToDelete];
+	Bones.RemoveAt(IndexToDelete);
+
+	RefreshMapping();
+
+#if WITH_EDITOR
+	for (const FName& RemovedChildBone : RemovedChildBones)
+	{
+		OnBoneRemoved.Broadcast(Container, ERigHierarchyElementType::Bone, RemovedChildBone);
+	}
+	OnBoneRemoved.Broadcast(Container, ERigHierarchyElementType::Bone, RemovedBone.Name);
+#endif
+
+	return RemovedBone;
+}
+
+// list of names of children - this is not cheap, and is supposed to be used only for one time set up
+int32 FRigBoneHierarchy::GetChildren(const FName& InName, TArray<int32>& OutChildren, bool bRecursively) const
+{
+	return GetChildren(GetIndex(InName), OutChildren, bRecursively);
+}
+
+int32 FRigBoneHierarchy::GetChildren(const int32 InIndex, TArray<int32>& OutChildren, bool bRecursively) const
+{
+	OutChildren.Reset();
+
+	if (InIndex != INDEX_NONE)
+	{
+		GetChildrenRecursive(InIndex, OutChildren, bRecursively);
+	}
+
+	return OutChildren.Num();
+}
+
+FName FRigBoneHierarchy::GetName(int32 InIndex) const
+{
+	if (Bones.IsValidIndex(InIndex))
+	{
+		return Bones[InIndex].Name;
+	}
+
+	return NAME_None;
+}
+
+int32 FRigBoneHierarchy::GetIndexSlow(const FName& InName) const
+{
+	for (int32 Index = 0; Index < Bones.Num(); ++Index)
+	{
+		if (Bones[Index].Name == InName)
+		{
+			return Index;
+		}
+	}
+
+	return INDEX_NONE;
+}
+
+void FRigBoneHierarchy::SetGlobalTransform(const FName& InName, const FTransform& InTransform, bool bPropagateTransform)
+{
+	SetGlobalTransform(GetIndex(InName), InTransform, bPropagateTransform);
+}
+
+void FRigBoneHierarchy::SetGlobalTransform(int32 InIndex, const FTransform& InTransform, bool bPropagateTransform)
+{
+	if (Bones.IsValidIndex(InIndex))
+	{
+		FRigBone& Bone = Bones[InIndex];
+		Bone.GlobalTransform = InTransform;
+		Bone.GlobalTransform.NormalizeRotation();
+		RecalculateLocalTransform(Bone);
+
+		if (bPropagateTransform)
+		{
+			PropagateTransform(InIndex);
+		}
+	}
+}
+
+FTransform FRigBoneHierarchy::GetGlobalTransform(const FName& InName) const
+{
+	return GetGlobalTransform(GetIndex(InName));
+}
+
+FTransform FRigBoneHierarchy::GetGlobalTransform(int32 InIndex) const
+{
+	if (Bones.IsValidIndex(InIndex))
+	{
+		return Bones[InIndex].GlobalTransform;
+	}
+
+	return FTransform::Identity;
+}
+
+void FRigBoneHierarchy::SetLocalTransform(const FName& InName, const FTransform& InTransform, bool bPropagateTransform)
+{
+	SetLocalTransform(GetIndex(InName), InTransform, bPropagateTransform);
+}
+
+void FRigBoneHierarchy::SetLocalTransform(int32 InIndex, const FTransform& InTransform, bool bPropagateTransform)
+{
+	if (Bones.IsValidIndex(InIndex))
+	{
+		FRigBone& Bone = Bones[InIndex];
+		Bone.LocalTransform = InTransform;
+		RecalculateGlobalTransform(Bone);
+
+		if (bPropagateTransform)
+		{
+			PropagateTransform(InIndex);
+		}
+	}
+}
+
+FTransform FRigBoneHierarchy::GetLocalTransform(const FName& InName) const
+{
+	return GetLocalTransform(GetIndex(InName));
+}
+
+FTransform FRigBoneHierarchy::GetLocalTransform(int32 InIndex) const
+{
+	if (Bones.IsValidIndex(InIndex))
+	{
+		return Bones[InIndex].LocalTransform;
+	}
+
+	return FTransform::Identity;
+}
+
+void FRigBoneHierarchy::SetInitialTransform(const FName& InName, const FTransform& InTransform)
+{
+	SetInitialTransform(GetIndex(InName), InTransform);
+}
+
+void FRigBoneHierarchy::SetInitialTransform(int32 InIndex, const FTransform& InTransform)
+{
+	if (Bones.IsValidIndex(InIndex))
+	{
+		FRigBone& Bone = Bones[InIndex];
+		Bone.InitialTransform = InTransform;
+		Bone.InitialTransform.NormalizeRotation();
+		RecalculateLocalTransform(Bone);
+	}
+}
+
+FTransform FRigBoneHierarchy::GetInitialTransform(const FName& InName) const
+{
+	return GetInitialTransform(GetIndex(InName));
+}
+
+FTransform FRigBoneHierarchy::GetInitialTransform(int32 InIndex) const
+{
+	if (Bones.IsValidIndex(InIndex))
+	{
+		return Bones[InIndex].InitialTransform;
+	}
+
+	return FTransform::Identity;
+}
+void FRigBoneHierarchy::RecalculateLocalTransform(FRigBone& InOutBone)
+{
+	bool bHasParent = InOutBone.ParentIndex != INDEX_NONE;
+	InOutBone.LocalTransform = (bHasParent) ? InOutBone.GlobalTransform.GetRelativeTransform(Bones[InOutBone.ParentIndex].GlobalTransform) : InOutBone.GlobalTransform;
+}
+
+void FRigBoneHierarchy::RecalculateGlobalTransform(FRigBone& InOutBone)
+{
+	bool bHasParent = InOutBone.ParentIndex != INDEX_NONE;
+	InOutBone.GlobalTransform = (bHasParent) ? InOutBone.LocalTransform * Bones[InOutBone.ParentIndex].GlobalTransform : InOutBone.LocalTransform;
+}
+
+void FRigBoneHierarchy::Rename(const FName& InOldName, const FName& InNewName)
+{
+	if (InOldName != InNewName)
+	{
+		const int32 Found = GetIndex(InOldName);
+		if (Found != INDEX_NONE)
+		{
+			Bones[Found].Name = InNewName;
+
+			// go through find all children and rename them
+#if WITH_EDITOR
+			TArray<FName> ReparentedBones;
+#endif
+			for (int32 Index = 0; Index < Bones.Num(); ++Index)
+			{
+				if (Bones[Index].ParentName == InOldName)
+				{
+					Bones[Index].ParentName = InNewName;
+					RecalculateLocalTransform(Bones[Index]);
+#if WITH_EDITOR
+					ReparentedBones.Add(Bones[Index].Name);
+#endif
+				}
+			}
+
+			RefreshMapping();
+
+#if WITH_EDITOR
+			OnBoneRenamed.Broadcast(Container, ERigHierarchyElementType::Bone, InOldName, InNewName);
+			for (const FName& ReparentedBone : ReparentedBones)
+			{
+				OnBoneReparented.Broadcast(Container, ERigHierarchyElementType::Bone, ReparentedBone, InOldName, InNewName);
+			}
+#endif
+		}
+	}
+}
+
+void FRigBoneHierarchy::Sort()
+{
+	TMap<int32, TArray<int32>> HierarchyTree;
+
+	TArray<int32> SortedArray;
+
+	// first figure out children and find roots
+	for (int32 Index = 0; Index < Bones.Num(); ++Index)
+	{
+		int32 ParentIndex = GetIndexSlow(Bones[Index].ParentName);
+		if (ParentIndex != INDEX_NONE)
+		{
+			TArray<int32>& ChildIndices = HierarchyTree.FindOrAdd(ParentIndex);
+			ChildIndices.Add(Index);
+		}
+		else
+		{
+			// as as a root
+			HierarchyTree.Add(Index);
+			// add them to the list first
+			SortedArray.Add(Index);
+		}
+	}
+
+	// now go through map and add to sorted array
+	for (int32 SortIndex = 0; SortIndex < SortedArray.Num(); ++SortIndex)
+	{
+		// add children of sorted array
+		TArray<int32>* ChildIndices = HierarchyTree.Find(SortedArray[SortIndex]);
+		if (ChildIndices)
+		{
+			// append children now
+			SortedArray.Append(*ChildIndices);
+			// now sorted array will grow
+			// this is same as BFS as it starts from all roots, and going down
+		}
+	}
+
+	check(SortedArray.Num() == Bones.Num());
+
+	// create new list with sorted
+	TArray<FRigBone> NewSortedList;
+	NewSortedList.AddDefaulted(Bones.Num());
+	for (int32 NewIndex = 0; NewIndex < SortedArray.Num(); ++NewIndex)
+	{
+		NewSortedList[NewIndex] = Bones[SortedArray[NewIndex]];
+	}
+
+	Bones = MoveTemp(NewSortedList);
+
+	// now fix up parent Index
+	for (int32 Index = 0; Index < Bones.Num(); ++Index)
+	{
+		Bones[Index].ParentIndex = GetIndexSlow(Bones[Index].ParentName);
+		// parent index always should be less than this index, even if invalid
+		check(Bones[Index].ParentIndex < Index);
+	}
+}
+
+void FRigBoneHierarchy::RefreshMapping()
+{
+	Sort();
+
+	NameToIndexMapping.Empty();
+	for (int32 Index = 0; Index < Bones.Num(); ++Index)
+	{
+		NameToIndexMapping.Add(Bones[Index].Name, Index);
+	}
+}
+
+void FRigBoneHierarchy::Initialize()
+{
+	RefreshMapping();
+
+	// update parent index
+	for (int32 Index = 0; Index < Bones.Num(); ++Index)
+	{
+		Bones[Index].ParentIndex = GetIndex(Bones[Index].ParentName);
+	}
+
+	// initialize transform
+	for (int32 Index = 0; Index < Bones.Num(); ++Index)
+	{
+		Bones[Index].GlobalTransform = Bones[Index].InitialTransform;
+		RecalculateLocalTransform(Bones[Index]);
+
+		// update children
+		GetChildren(Index, Bones[Index].Dependents, false);
+	}
+}
+
+void FRigBoneHierarchy::Reset()
+{
+	Bones.Reset();
+}
+
+void FRigBoneHierarchy::ResetTransforms()
+{
+	// initialize transform
+	for (int32 Index = 0; Index < Bones.Num(); ++Index)
+	{
+		Bones[Index].GlobalTransform = Bones[Index].InitialTransform;
+		RecalculateLocalTransform(Bones[Index]);
+	}
+}
+
+int32 FRigBoneHierarchy::GetChildrenRecursive(const int32 InIndex, TArray<int32>& OutChildren, bool bRecursively) const
+{
+	const int32 StartChildIndex = OutChildren.Num();
+
+	// all children should be later than parent
+	for (int32 ChildIndex = InIndex + 1; ChildIndex < Bones.Num(); ++ChildIndex)
+	{
+		if (Bones[ChildIndex].ParentIndex == InIndex)
+		{
+			OutChildren.AddUnique(ChildIndex);
+		}
+	}
+
+	if (bRecursively)
+	{
+		// since we keep appending inside of functions, we make sure not to go over original list
+		const int32 EndChildIndex = OutChildren.Num() - 1;
+		for (int32 ChildIndex = StartChildIndex; ChildIndex <= EndChildIndex; ++ChildIndex)
+		{
+			GetChildrenRecursive(OutChildren[ChildIndex], OutChildren, bRecursively);
+		}
+	}
+
+	return OutChildren.Num();
+}
+
+void FRigBoneHierarchy::PropagateTransform(int32 InIndex)
+{
+	const TArray<int32> Dependents = Bones[InIndex].Dependents;
+	for (int32 DependentIndex = 0; DependentIndex<Dependents.Num(); ++DependentIndex)
+	{
+		int32 Index = Dependents[DependentIndex];
+		RecalculateGlobalTransform(Bones[Index]);
+		PropagateTransform(Index);
+	}
+}
diff --git a/Engine/Plugins/Experimental/ControlRig/Source/ControlRig/Private/Rigs/RigCurveContainer.cpp b/Engine/Plugins/Experimental/ControlRig/Source/ControlRig/Private/Rigs/RigCurveContainer.cpp
new file mode 100644
--- /dev/null
+++ b/Engine/Plugins/Experimental/ControlRig/Source/ControlRig/Private/Rigs/RigCurveContainer.cpp
@@ -0,0 +1,47 @@
+// Copyright 1998-2019 Epic Games, Inc. All Rights Reserved.
+
+#include "Rigs/RigCurveContainer.h"
+#include "ControlRig.h"
+#include "HelperUtil.h"
+
+////////////////////////////////////////////////////////////////////////////////
+// FRigCurveContainer
+////////////////////////////////////////////////////////////////////////////////
+
+void FRigCurveContainer::RefreshMapping()
+{
+	DECLARE_SCOPE_HIERARCHICAL_COUNTER_FUNC()
+
+	NameToIndexMapping.Empty();
+	for (int32 Index = 0; Index < Curves.Num(); ++Index)
+	{
+		NameToIndexMapping.Add(Curves[Index].Name, Index);
+	}
+}
+
+void FRigCurveContainer::Initialize()
+{
+	DECLARE_SCOPE_HIERARCHICAL_COUNTER_FUNC()
+
+	RefreshMapping();
+	ResetValues();
+}
+
+void FRigCurveContainer::Reset()
+{
+	DECLARE_SCOPE_HIERARCHICAL_COUNTER_FUNC()
+
+	Curves.Reset();
+}
+
+void FRigCurveContainer::ResetValues()
+{
+	DECLARE_SCOPE_HIERARCHICAL_COUNTER_FUNC()
+
+	// update parent index
+	for (int32 Index = 0; Index < Curves.Num(); ++Index)
+	{
+		Curves[Index].Value = 0.f;
+	}
+}
+
diff --git a/Engine/Plugins/Experimental/ControlRig/Source/ControlRig/Private/Rigs/RigHierarchyContainer.cpp b/Engine/Plugins/Experimental/ControlRig/Source/ControlRig/Private/Rigs/RigHierarchyContainer.cpp
new file mode 100644
--- /dev/null
+++ b/Engine/Plugins/Experimental/ControlRig/Source/ControlRig/Private/Rigs/RigHierarchyContainer.cpp
@@ -0,0 +1,115 @@
+// Copyright 1998-2019 Epic Games, Inc. All Rights Reserved.
+
+#include "Rigs/RigHierarchyContainer.h"
+#include "ControlRig.h"
+#include "HelperUtil.h"
+
+////////////////////////////////////////////////////////////////////////////////
+// FRigHierarchyContainer
+////////////////////////////////////////////////////////////////////////////////
+
+FRigHierarchyContainer::FRigHierarchyContainer()
+{
+	Initialize();
+}
+
+FRigHierarchyContainer& FRigHierarchyContainer::operator= (const FRigHierarchyContainer &InOther)
+{
+	BoneHierarchy = InOther.BoneHierarchy;
+	return *this;
+}
+
+void FRigHierarchyContainer::Initialize()
+{
+#if WITH_EDITOR
+	BoneHierarchy.Container = this;
+
+	BoneHierarchy.OnBoneAdded.RemoveAll(this);
+	BoneHierarchy.OnBoneRemoved.RemoveAll(this);
+	BoneHierarchy.OnBoneRenamed.RemoveAll(this);
+	BoneHierarchy.OnBoneReparented.RemoveAll(this);
+
+	BoneHierarchy.OnBoneAdded.AddRaw(this, &FRigHierarchyContainer::HandleOnElementAdded);
+	BoneHierarchy.OnBoneRemoved.AddRaw(this, &FRigHierarchyContainer::HandleOnElementRemoved);
+	BoneHierarchy.OnBoneRenamed.AddRaw(this, &FRigHierarchyContainer::HandleOnElementRenamed);
+	BoneHierarchy.OnBoneReparented.AddRaw(this, &FRigHierarchyContainer::HandleOnElementReparented);
+#endif
+
+	BoneHierarchy.Initialize();
+
+	ResetTransforms();
+}
+
+void FRigHierarchyContainer::Reset()
+{
+	BoneHierarchy.Reset();
+	
+	Initialize();
+}
+
+void FRigHierarchyContainer::ResetTransforms()
+{
+	BoneHierarchy.ResetTransforms();
+}
+
+#if WITH_EDITOR
+
+void FRigHierarchyContainer::HandleOnElementAdded(FRigHierarchyContainer* InContainer, ERigHierarchyElementType InElementType, const FName& InName)
+{
+	// todo
+	OnElementAdded.Broadcast(InContainer, InElementType, InName);
+	OnElementChanged.Broadcast(InContainer, InElementType, InName);
+}
+
+void FRigHierarchyContainer::HandleOnElementRemoved(FRigHierarchyContainer* InContainer, ERigHierarchyElementType InElementType, const FName& InName)
+{
+	// todo
+	OnElementRemoved.Broadcast(InContainer, InElementType, InName);
+	OnElementChanged.Broadcast(InContainer, InElementType, InName);
+}
+
+void FRigHierarchyContainer::HandleOnElementRenamed(FRigHierarchyContainer* InContainer, ERigHierarchyElementType InElementType, const FName& InOldName, const FName& InNewName)
+{
+	// todo
+	OnElementRenamed.Broadcast(InContainer, InElementType, InOldName, InNewName);
+	OnElementChanged.Broadcast(InContainer, InElementType, InNewName);
+}
+
+void FRigHierarchyContainer::HandleOnElementReparented(FRigHierarchyContainer* InContainer, ERigHierarchyElementType InElementType, const FName& InName, const FName& InOldParentName, const FName& InNewParentName)
+{
+	// todo
+	OnElementReparented.Broadcast(InContainer, InElementType, InName, InOldParentName, InNewParentName);
+	OnElementChanged.Broadcast(InContainer, InElementType, InName);
+}
+
+#endif
+
+////////////////////////////////////////////////////////////////////////////////
+// FRigHierarchyRef
+////////////////////////////////////////////////////////////////////////////////
+
+FRigHierarchyRef::FRigHierarchyRef()
+: Container(nullptr)
+{
+
+}
+
+FRigBoneHierarchy* FRigHierarchyRef::GetBones()
+{
+	return GetBonesInternal();
+}
+
+const FRigBoneHierarchy* FRigHierarchyRef::GetBones() const
+{
+	return (const FRigBoneHierarchy*)GetBonesInternal();
+}
+
+FRigBoneHierarchy* FRigHierarchyRef::GetBonesInternal() const
+{
+	if (Container)
+	{
+		return &Container->BoneHierarchy;
+	}
+	return nullptr;
+}
+
diff --git a/Engine/Plugins/Experimental/ControlRig/Source/ControlRig/Private/Units/Debug/RigUnit_DebugBezier.cpp b/Engine/Plugins/Experimental/ControlRig/Source/ControlRig/Private/Units/Debug/RigUnit_DebugBezier.cpp
--- a/Engine/Plugins/Experimental/ControlRig/Source/ControlRig/Private/Units/Debug/RigUnit_DebugBezier.cpp
+++ b/Engine/Plugins/Experimental/ControlRig/Source/ControlRig/Private/Units/Debug/RigUnit_DebugBezier.cpp
@@ -17,9 +17,9 @@ void FRigUnit_DebugBezier::Execute(const FRigUnitContext& Context)
 	}
 
 	FTransform Transform = WorldOffset;
-	if (Space != NAME_None && Context.HierarchyReference.Get() != nullptr)
+	if (Space != NAME_None && Context.HierarchyReference.GetBones() != nullptr)
 	{
-		Transform = Transform * Context.HierarchyReference.Get()->GetGlobalTransform(Space);
+		Transform = Transform * Context.HierarchyReference.GetBones()->GetGlobalTransform(Space);
 	}
 
 
diff --git a/Engine/Plugins/Experimental/ControlRig/Source/ControlRig/Private/Units/Debug/RigUnit_DebugHierarchy.cpp b/Engine/Plugins/Experimental/ControlRig/Source/ControlRig/Private/Units/Debug/RigUnit_DebugHierarchy.cpp
--- a/Engine/Plugins/Experimental/ControlRig/Source/ControlRig/Private/Units/Debug/RigUnit_DebugHierarchy.cpp
+++ b/Engine/Plugins/Experimental/ControlRig/Source/ControlRig/Private/Units/Debug/RigUnit_DebugHierarchy.cpp
@@ -16,7 +16,7 @@ void FRigUnit_DebugHierarchy::Execute(const FRigUnitContext& Context)
 		return;
 	}
 
-	const FRigHierarchy* Hierarchy = Context.HierarchyReference.Get();
+	const FRigBoneHierarchy* Hierarchy = Context.HierarchyReference.GetBones();
 	if (Hierarchy)
 	{
 		Context.DrawInterface->DrawHierarchy(WorldOffset, *Hierarchy, Mode, Scale, Color, Thickness);
diff --git a/Engine/Plugins/Experimental/ControlRig/Source/ControlRig/Private/Units/Debug/RigUnit_DebugLine.cpp b/Engine/Plugins/Experimental/ControlRig/Source/ControlRig/Private/Units/Debug/RigUnit_DebugLine.cpp
--- a/Engine/Plugins/Experimental/ControlRig/Source/ControlRig/Private/Units/Debug/RigUnit_DebugLine.cpp
+++ b/Engine/Plugins/Experimental/ControlRig/Source/ControlRig/Private/Units/Debug/RigUnit_DebugLine.cpp
@@ -17,9 +17,9 @@ void FRigUnit_DebugLine::Execute(const FRigUnitContext& Context)
 	}
 
 	FVector DrawA = A, DrawB = B;
-	if (Space != NAME_None && Context.HierarchyReference.Get() != nullptr)
+	if (Space != NAME_None && Context.HierarchyReference.GetBones() != nullptr)
 	{
-		FTransform Transform = Context.HierarchyReference.Get()->GetGlobalTransform(Space);
+		FTransform Transform = Context.HierarchyReference.GetBones()->GetGlobalTransform(Space);
 		DrawA = Transform.TransformPosition(DrawA);
 		DrawB = Transform.TransformPosition(DrawB);
 	}
diff --git a/Engine/Plugins/Experimental/ControlRig/Source/ControlRig/Private/Units/Debug/RigUnit_DebugLineStrip.cpp b/Engine/Plugins/Experimental/ControlRig/Source/ControlRig/Private/Units/Debug/RigUnit_DebugLineStrip.cpp
--- a/Engine/Plugins/Experimental/ControlRig/Source/ControlRig/Private/Units/Debug/RigUnit_DebugLineStrip.cpp
+++ b/Engine/Plugins/Experimental/ControlRig/Source/ControlRig/Private/Units/Debug/RigUnit_DebugLineStrip.cpp
@@ -16,9 +16,9 @@ void FRigUnit_DebugLineStrip::Execute(const FRigUnitContext& Context)
 		return;
 	}
 
-	if (Space != NAME_None && Context.HierarchyReference.Get() != nullptr)
+	if (Space != NAME_None && Context.HierarchyReference.GetBones() != nullptr)
 	{
-		FTransform Transform = Context.HierarchyReference.Get()->GetGlobalTransform(Space);
+		FTransform Transform = Context.HierarchyReference.GetBones()->GetGlobalTransform(Space);
 		TArray<FVector> PointsTransformed;
 		PointsTransformed.Reserve(Points.Num());
 		for(const FVector& Point : Points)
diff --git a/Engine/Plugins/Experimental/ControlRig/Source/ControlRig/Private/Units/Debug/RigUnit_DebugPoint.cpp b/Engine/Plugins/Experimental/ControlRig/Source/ControlRig/Private/Units/Debug/RigUnit_DebugPoint.cpp
--- a/Engine/Plugins/Experimental/ControlRig/Source/ControlRig/Private/Units/Debug/RigUnit_DebugPoint.cpp
+++ b/Engine/Plugins/Experimental/ControlRig/Source/ControlRig/Private/Units/Debug/RigUnit_DebugPoint.cpp
@@ -18,9 +18,9 @@ void FRigUnit_DebugPoint::Execute(const FRigUnitContext& Context)
 
 	FVector Center = FVector::ZeroVector;
 	FVector DrawVector = Vector;
-	if (Space != NAME_None && Context.HierarchyReference.Get() != nullptr)
+	if (Space != NAME_None && Context.HierarchyReference.GetBones() != nullptr)
 	{
-		FTransform Transform = Context.HierarchyReference.Get()->GetGlobalTransform(Space);
+		FTransform Transform = Context.HierarchyReference.GetBones()->GetGlobalTransform(Space);
 		Center = Transform.GetLocation();
 		DrawVector = Transform.TransformPosition(DrawVector);
 	}
@@ -55,9 +55,9 @@ void FRigUnit_DebugPointMutable::Execute(const FRigUnitContext& Context)
 
 	FVector Center = FVector::ZeroVector;
 	FVector DrawVector = Vector;
-	if (Space != NAME_None && Context.HierarchyReference.Get() != nullptr)
+	if (Space != NAME_None && Context.HierarchyReference.GetBones() != nullptr)
 	{
-		FTransform Transform = Context.HierarchyReference.Get()->GetGlobalTransform(Space);
+		FTransform Transform = Context.HierarchyReference.GetBones()->GetGlobalTransform(Space);
 		Center = Transform.GetLocation();
 		DrawVector = Transform.TransformPosition(DrawVector);
 	}
diff --git a/Engine/Plugins/Experimental/ControlRig/Source/ControlRig/Private/Units/Debug/RigUnit_DebugPrimitives.cpp b/Engine/Plugins/Experimental/ControlRig/Source/ControlRig/Private/Units/Debug/RigUnit_DebugPrimitives.cpp
--- a/Engine/Plugins/Experimental/ControlRig/Source/ControlRig/Private/Units/Debug/RigUnit_DebugPrimitives.cpp
+++ b/Engine/Plugins/Experimental/ControlRig/Source/ControlRig/Private/Units/Debug/RigUnit_DebugPrimitives.cpp
@@ -17,9 +17,9 @@ void FRigUnit_DebugRectangle::Execute(const FRigUnitContext& Context)
 	}
 
 	FTransform DrawTransform = Transform;
-	if (Space != NAME_None && Context.HierarchyReference.Get() != nullptr)
+	if (Space != NAME_None && Context.HierarchyReference.GetBones() != nullptr)
 	{
-		DrawTransform = DrawTransform * Context.HierarchyReference.Get()->GetGlobalTransform(Space);
+		DrawTransform = DrawTransform * Context.HierarchyReference.GetBones()->GetGlobalTransform(Space);
 	}
 
 	Context.DrawInterface->DrawRectangle(WorldOffset, DrawTransform, Scale, Color, Thickness);
@@ -39,9 +39,9 @@ void FRigUnit_DebugArc::Execute(const FRigUnitContext& Context)
 	}
 
 	FTransform DrawTransform = Transform;
-	if (Space != NAME_None && Context.HierarchyReference.Get() != nullptr)
+	if (Space != NAME_None && Context.HierarchyReference.GetBones() != nullptr)
 	{
-		DrawTransform = DrawTransform * Context.HierarchyReference.Get()->GetGlobalTransform(Space);
+		DrawTransform = DrawTransform * Context.HierarchyReference.GetBones()->GetGlobalTransform(Space);
 	}
 
 	Context.DrawInterface->DrawArc(WorldOffset, DrawTransform, Radius, FMath::DegreesToRadians(MinimumDegrees), FMath::DegreesToRadians(MaximumDegrees), Color, Thickness, Detail);
diff --git a/Engine/Plugins/Experimental/ControlRig/Source/ControlRig/Private/Units/Debug/RigUnit_DebugTransform.cpp b/Engine/Plugins/Experimental/ControlRig/Source/ControlRig/Private/Units/Debug/RigUnit_DebugTransform.cpp
--- a/Engine/Plugins/Experimental/ControlRig/Source/ControlRig/Private/Units/Debug/RigUnit_DebugTransform.cpp
+++ b/Engine/Plugins/Experimental/ControlRig/Source/ControlRig/Private/Units/Debug/RigUnit_DebugTransform.cpp
@@ -17,9 +17,9 @@ void FRigUnit_DebugTransform::Execute(const FRigUnitContext& Context)
 	}
 
 	FTransform DrawTransform = Transform;
-	if (Space != NAME_None && Context.HierarchyReference.Get() != nullptr)
+	if (Space != NAME_None && Context.HierarchyReference.GetBones() != nullptr)
 	{
-		DrawTransform = Transform * Context.HierarchyReference.Get()->GetGlobalTransform(Space);
+		DrawTransform = Transform * Context.HierarchyReference.GetBones()->GetGlobalTransform(Space);
 	}
 
 	switch (Mode)
@@ -57,9 +57,9 @@ void FRigUnit_DebugTransformMutable::Execute(const FRigUnitContext& Context)
 	}
 
 	FTransform DrawTransform = Transform;
-	if (Space != NAME_None && Context.HierarchyReference.Get() != nullptr)
+	if (Space != NAME_None && Context.HierarchyReference.GetBones() != nullptr)
 	{
-		DrawTransform = Transform * Context.HierarchyReference.Get()->GetGlobalTransform(Space);
+		DrawTransform = Transform * Context.HierarchyReference.GetBones()->GetGlobalTransform(Space);
 	}
 
 	switch (Mode)
@@ -97,9 +97,9 @@ void FRigUnit_DebugTransformArrayMutable::Execute(const FRigUnitContext& Context
 	}
 
 	DrawTransforms.SetNumUninitialized(Transforms.Num());
-	if (Space != NAME_None && Context.HierarchyReference.Get() != nullptr)
+	if (Space != NAME_None && Context.HierarchyReference.GetBones() != nullptr)
 	{
-		FTransform SpaceTransform = Context.HierarchyReference.Get()->GetGlobalTransform(Space);
+		FTransform SpaceTransform = Context.HierarchyReference.GetBones()->GetGlobalTransform(Space);
 		for(int32 Index=0;Index<Transforms.Num();Index++)
 		{
 			DrawTransforms[Index] = Transforms[Index] * SpaceTransform;
diff --git a/Engine/Plugins/Experimental/ControlRig/Source/ControlRig/Private/Units/Deprecated/RigUnit_AimConstraint.cpp b/Engine/Plugins/Experimental/ControlRig/Source/ControlRig/Private/Units/Deprecated/RigUnit_AimConstraint.cpp
--- a/Engine/Plugins/Experimental/ControlRig/Source/ControlRig/Private/Units/Deprecated/RigUnit_AimConstraint.cpp
+++ b/Engine/Plugins/Experimental/ControlRig/Source/ControlRig/Private/Units/Deprecated/RigUnit_AimConstraint.cpp
@@ -14,7 +14,7 @@ void FRigUnit_AimConstraint::Execute(const FRigUnitContext& Context)
 	{
 		ConstraintData.Reset();
 
-		FRigHierarchy* Hierarchy = HierarchyRef.Get();
+		FRigBoneHierarchy* Hierarchy = HierarchyRef.GetBones();
 		if (Hierarchy)
 		{
 			int32 BoneIndex = Hierarchy->GetIndex(Joint);
@@ -45,7 +45,7 @@ void FRigUnit_AimConstraint::Execute(const FRigUnitContext& Context)
 	}
 	else if (Context.State == EControlRigState::Update)
 	{
-		FRigHierarchy* Hierarchy = HierarchyRef.Get();
+		FRigBoneHierarchy* Hierarchy = HierarchyRef.GetBones();
 		if (Hierarchy)
 		{
 			int32 BoneIndex = Hierarchy->GetIndex(Joint);
@@ -71,7 +71,7 @@ void FRigUnit_AimConstraint::Execute(const FRigUnitContext& Context)
 					}
 
 					FTransform BaseTransform = FTransform::Identity;
-					int32 ParentIndex = Hierarchy->GetParentIndex(BoneIndex);
+					int32 ParentIndex = (*Hierarchy)[BoneIndex].ParentIndex;
 					if (ParentIndex != INDEX_NONE)
 					{
 						BaseTransform = Hierarchy->GetGlobalTransform(ParentIndex);
diff --git a/Engine/Plugins/Experimental/ControlRig/Source/ControlRig/Private/Units/Deprecated/RigUnit_AimConstraint.h b/Engine/Plugins/Experimental/ControlRig/Source/ControlRig/Private/Units/Deprecated/RigUnit_AimConstraint.h
--- a/Engine/Plugins/Experimental/ControlRig/Source/ControlRig/Private/Units/Deprecated/RigUnit_AimConstraint.h
+++ b/Engine/Plugins/Experimental/ControlRig/Source/ControlRig/Private/Units/Deprecated/RigUnit_AimConstraint.h
@@ -3,7 +3,7 @@
 #pragma once
 
 #include "Units/RigUnit.h"
-#include "Rigs/Hierarchy.h"
+#include "Rigs/RigHierarchyContainer.h"
 #include "Constraint.h"
 #include "ControlRigDefines.h"
 #include "RigUnit_AimConstraint.generated.h"
diff --git a/Engine/Plugins/Experimental/ControlRig/Source/ControlRig/Private/Units/Deprecated/RigUnit_ApplyFK.cpp b/Engine/Plugins/Experimental/ControlRig/Source/ControlRig/Private/Units/Deprecated/RigUnit_ApplyFK.cpp
--- a/Engine/Plugins/Experimental/ControlRig/Source/ControlRig/Private/Units/Deprecated/RigUnit_ApplyFK.cpp
+++ b/Engine/Plugins/Experimental/ControlRig/Source/ControlRig/Private/Units/Deprecated/RigUnit_ApplyFK.cpp
@@ -11,7 +11,7 @@ void FRigUnit_ApplyFK::Execute(const FRigUnitContext& Context)
 	
 	if (Context.State == EControlRigState::Init)
 	{
-		FRigHierarchy* Hierarchy = HierarchyRef.Get();
+		FRigBoneHierarchy* Hierarchy = HierarchyRef.GetBones();
 		if (!Hierarchy)
 		{
 			UnitLogHelpers::PrintMissingHierarchy(RigUnitName);
@@ -19,7 +19,7 @@ void FRigUnit_ApplyFK::Execute(const FRigUnitContext& Context)
 	}
 	else if (Context.State == EControlRigState::Update)
 	{
-		FRigHierarchy* Hierarchy = HierarchyRef.Get();
+		FRigBoneHierarchy* Hierarchy = HierarchyRef.GetBones();
 		if (Hierarchy)
 		{
 			int32 Index = Hierarchy->GetIndex(Joint);
@@ -52,10 +52,10 @@ void FRigUnit_ApplyFK::Execute(const FRigUnitContext& Context)
 	}
 }
 
-FTransform FRigUnit_ApplyFK::GetBaseTransform(int32 BoneIndex, const FRigHierarchy* CurrentHierarchy) const
+FTransform FRigUnit_ApplyFK::GetBaseTransform(int32 BoneIndex, const FRigBoneHierarchy* CurrentHierarchy) const
 {
 	return UtilityHelpers::GetBaseTransformByMode(ApplyTransformSpace, [CurrentHierarchy](const FName& BoneName) { return CurrentHierarchy->GetGlobalTransform(BoneName); },
-		CurrentHierarchy->GetBones()[BoneIndex].ParentName, BaseJoint, BaseTransform);
+		(*CurrentHierarchy)[BoneIndex].ParentName, BaseJoint, BaseTransform);
 }
 
 #if WITH_DEV_AUTOMATION_TESTS
@@ -63,8 +63,8 @@ FTransform FRigUnit_ApplyFK::GetBaseTransform(int32 BoneIndex, const FRigHierarc
 
 IMPLEMENT_RIGUNIT_AUTOMATION_TEST(FRigUnit_ApplyFK)
 {
-	Hierarchy.AddBone(TEXT("Root"), NAME_None, FTransform(FVector(1.f, 0.f, 0.f)));
-	Hierarchy.AddBone(TEXT("BoneA"), TEXT("Root"), FTransform(FVector(1.f, 2.f, 3.f)));
+	Hierarchy.Add(TEXT("Root"), NAME_None, FTransform(FVector(1.f, 0.f, 0.f)));
+	Hierarchy.Add(TEXT("BoneA"), TEXT("Root"), FTransform(FVector(1.f, 2.f, 3.f)));
 
 	Unit.ExecuteContext = ExecuteContext;
 	Unit.Joint = TEXT("BoneA");
diff --git a/Engine/Plugins/Experimental/ControlRig/Source/ControlRig/Private/Units/Deprecated/RigUnit_ApplyFK.h b/Engine/Plugins/Experimental/ControlRig/Source/ControlRig/Private/Units/Deprecated/RigUnit_ApplyFK.h
--- a/Engine/Plugins/Experimental/ControlRig/Source/ControlRig/Private/Units/Deprecated/RigUnit_ApplyFK.h
+++ b/Engine/Plugins/Experimental/ControlRig/Source/ControlRig/Private/Units/Deprecated/RigUnit_ApplyFK.h
@@ -3,7 +3,7 @@
 #pragma once
 
 #include "Units/RigUnit.h"
-#include "Rigs/Hierarchy.h"
+#include "Rigs/RigHierarchyContainer.h"
 #include "Constraint.h"
 #include "ControlRigDefines.h"
 #include "RigUnit_ApplyFK.generated.h"
@@ -57,5 +57,5 @@ struct FRigUnit_ApplyFK : public FRigUnitMutable
 	FName BaseJoint;
 
 private:
-	FTransform GetBaseTransform(int32 JointIndex, const FRigHierarchy* CurrentHierarchy) const;
+	FTransform GetBaseTransform(int32 JointIndex, const FRigBoneHierarchy* CurrentHierarchy) const;
 };
diff --git a/Engine/Plugins/Experimental/ControlRig/Source/ControlRig/Private/Units/Deprecated/RigUnit_CreateHierarchy.cpp b/Engine/Plugins/Experimental/ControlRig/Source/ControlRig/Private/Units/Deprecated/RigUnit_CreateHierarchy.cpp
--- a/Engine/Plugins/Experimental/ControlRig/Source/ControlRig/Private/Units/Deprecated/RigUnit_CreateHierarchy.cpp
+++ b/Engine/Plugins/Experimental/ControlRig/Source/ControlRig/Private/Units/Deprecated/RigUnit_CreateHierarchy.cpp
@@ -9,7 +9,7 @@ void FRigUnit_CreateHierarchy::Execute(const FRigUnitContext& Context)
 	if (Context.State == EControlRigState::Init)
 	{
 		// create hierarchy
-		if (!NewHierarchy.CreateHierarchy(Root, SourceHierarchy))
+		//if (!NewHierarchy.CreateHierarchy(Root, SourceHierarchy))
 		{
 			// if failed, print
 		}
diff --git a/Engine/Plugins/Experimental/ControlRig/Source/ControlRig/Private/Units/Deprecated/RigUnit_CreateHierarchy.h b/Engine/Plugins/Experimental/ControlRig/Source/ControlRig/Private/Units/Deprecated/RigUnit_CreateHierarchy.h
--- a/Engine/Plugins/Experimental/ControlRig/Source/ControlRig/Private/Units/Deprecated/RigUnit_CreateHierarchy.h
+++ b/Engine/Plugins/Experimental/ControlRig/Source/ControlRig/Private/Units/Deprecated/RigUnit_CreateHierarchy.h
@@ -3,7 +3,7 @@
 #pragma once
 
 #include "Units/RigUnit.h"
-#include "Rigs/Hierarchy.h"
+#include "Rigs/RigBoneHierarchy.h"
 #include "RigUnit_CreateHierarchy.generated.h"
 
 USTRUCT(meta=(DisplayName="Create Hierarchy", Category="Hierarchy", Deprecated = "4.23.0"))
diff --git a/Engine/Plugins/Experimental/ControlRig/Source/ControlRig/Private/Units/Deprecated/RigUnit_GetJointTransform.cpp b/Engine/Plugins/Experimental/ControlRig/Source/ControlRig/Private/Units/Deprecated/RigUnit_GetJointTransform.cpp
--- a/Engine/Plugins/Experimental/ControlRig/Source/ControlRig/Private/Units/Deprecated/RigUnit_GetJointTransform.cpp
+++ b/Engine/Plugins/Experimental/ControlRig/Source/ControlRig/Private/Units/Deprecated/RigUnit_GetJointTransform.cpp
@@ -9,7 +9,7 @@ void FRigUnit_GetJointTransform::Execute(const FRigUnitContext& Context)
     DECLARE_SCOPE_HIERARCHICAL_COUNTER_RIGUNIT()
 	FRigHierarchyRef& HierarchyRef = ExecuteContext.HierarchyReference;
 
-	FRigHierarchy* Hierarchy = HierarchyRef.Get();
+	FRigBoneHierarchy* Hierarchy = HierarchyRef.GetBones();
 	if (Hierarchy)
 	{
 		int32 Index = Hierarchy->GetIndex(Joint);
@@ -41,17 +41,17 @@ void FRigUnit_GetJointTransform::Execute(const FRigUnitContext& Context)
 	}
 }
 
-FTransform FRigUnit_GetJointTransform::GetBaseTransform(int32 JointIndex, const FRigHierarchy* CurrentHierarchy, bool bUseInitial) const
+FTransform FRigUnit_GetJointTransform::GetBaseTransform(int32 JointIndex, const FRigBoneHierarchy* CurrentHierarchy, bool bUseInitial) const
 {
 	if (bUseInitial)
 	{
 		return UtilityHelpers::GetBaseTransformByMode(TransformSpace, [CurrentHierarchy](const FName& JointName) { return CurrentHierarchy->GetInitialTransform(JointName); },
-			CurrentHierarchy->GetBones()[JointIndex].ParentName, BaseJoint, BaseTransform);
+			(*CurrentHierarchy)[JointIndex].ParentName, BaseJoint, BaseTransform);
 	}
 	else
 	{
 		return UtilityHelpers::GetBaseTransformByMode(TransformSpace, [CurrentHierarchy](const FName& JointName) { return CurrentHierarchy->GetGlobalTransform(JointName); },
-			CurrentHierarchy->GetBones()[JointIndex].ParentName, BaseJoint, BaseTransform);
+			(*CurrentHierarchy)[JointIndex].ParentName, BaseJoint, BaseTransform);
 	}
 
 }
\ No newline at end of file
diff --git a/Engine/Plugins/Experimental/ControlRig/Source/ControlRig/Private/Units/Deprecated/RigUnit_GetJointTransform.h b/Engine/Plugins/Experimental/ControlRig/Source/ControlRig/Private/Units/Deprecated/RigUnit_GetJointTransform.h
--- a/Engine/Plugins/Experimental/ControlRig/Source/ControlRig/Private/Units/Deprecated/RigUnit_GetJointTransform.h
+++ b/Engine/Plugins/Experimental/ControlRig/Source/ControlRig/Private/Units/Deprecated/RigUnit_GetJointTransform.h
@@ -3,7 +3,7 @@
 #pragma once
 
 #include "Units/RigUnit.h"
-#include "Rigs/Hierarchy.h"
+#include "Rigs/RigHierarchyContainer.h"
 #include "Constraint.h"
 #include "ControlRigDefines.h"
 #include "RigUnit_GetJointTransform.generated.h"
@@ -42,5 +42,5 @@ struct FRigUnit_GetJointTransform : public FRigUnitMutable
 	FTransform Output;
 
 private:
-	FTransform GetBaseTransform(int32 JointIndex, const FRigHierarchy* CurrentHierarchy, bool bUseInitial) const;
+	FTransform GetBaseTransform(int32 JointIndex, const FRigBoneHierarchy* CurrentHierarchy, bool bUseInitial) const;
 };
diff --git a/Engine/Plugins/Experimental/ControlRig/Source/ControlRig/Private/Units/Deprecated/RigUnit_MergeHierarchy.cpp b/Engine/Plugins/Experimental/ControlRig/Source/ControlRig/Private/Units/Deprecated/RigUnit_MergeHierarchy.cpp
--- a/Engine/Plugins/Experimental/ControlRig/Source/ControlRig/Private/Units/Deprecated/RigUnit_MergeHierarchy.cpp
+++ b/Engine/Plugins/Experimental/ControlRig/Source/ControlRig/Private/Units/Deprecated/RigUnit_MergeHierarchy.cpp
@@ -7,5 +7,5 @@ void FRigUnit_MergeHierarchy::Execute(const FRigUnitContext& Context)
 {
     DECLARE_SCOPE_HIERARCHICAL_COUNTER_RIGUNIT()
 	// merge input hierarchy to base
-	TargetHierarchy.MergeHierarchy(SourceHierarchy);
+	//TargetHierarchy.MergeHierarchy(SourceHierarchy);
 }
diff --git a/Engine/Plugins/Experimental/ControlRig/Source/ControlRig/Private/Units/Deprecated/RigUnit_MergeHierarchy.h b/Engine/Plugins/Experimental/ControlRig/Source/ControlRig/Private/Units/Deprecated/RigUnit_MergeHierarchy.h
--- a/Engine/Plugins/Experimental/ControlRig/Source/ControlRig/Private/Units/Deprecated/RigUnit_MergeHierarchy.h
+++ b/Engine/Plugins/Experimental/ControlRig/Source/ControlRig/Private/Units/Deprecated/RigUnit_MergeHierarchy.h
@@ -3,7 +3,7 @@
 #pragma once
 
 #include "Units/RigUnit.h"
-#include "Rigs/Hierarchy.h"
+#include "Rigs/RigHierarchyContainer.h"
 #include "RigUnit_MergeHierarchy.generated.h"
 
 USTRUCT(meta=(DisplayName="Merge Hierarchy", Category="Hierarchy", Deprecated = "4.23.0"))
diff --git a/Engine/Plugins/Experimental/ControlRig/Source/ControlRig/Private/Units/Deprecated/RigUnit_TwoBoneIKFK.cpp b/Engine/Plugins/Experimental/ControlRig/Source/ControlRig/Private/Units/Deprecated/RigUnit_TwoBoneIKFK.cpp
--- a/Engine/Plugins/Experimental/ControlRig/Source/ControlRig/Private/Units/Deprecated/RigUnit_TwoBoneIKFK.cpp
+++ b/Engine/Plugins/Experimental/ControlRig/Source/ControlRig/Private/Units/Deprecated/RigUnit_TwoBoneIKFK.cpp
@@ -12,7 +12,7 @@ void FRigUnit_TwoBoneIKFK::Execute(const FRigUnitContext& Context)
 
 	if (Context.State == EControlRigState::Init)
 	{
-		const FRigHierarchy* Hierarchy = HierarchyRef.Get();
+		const FRigBoneHierarchy* Hierarchy = HierarchyRef.GetBones();
 		if (Hierarchy)
 		{
 			// reset
@@ -25,10 +25,10 @@ void FRigUnit_TwoBoneIKFK::Execute(const FRigUnitContext& Context)
 			if (StartIndex != INDEX_NONE && EndIndex != INDEX_NONE)
 			{
 				// ensure the chain
-				int32 EndParentIndex = Hierarchy->GetParentIndex(EndIndex);
+				int32 EndParentIndex = (*Hierarchy)[EndIndex].ParentIndex;
 				if (EndParentIndex != INDEX_NONE)
 				{
-					int32 MidParentIndex = Hierarchy->GetParentIndex(EndParentIndex);
+					int32 MidParentIndex = (*Hierarchy)[EndParentIndex].ParentIndex;
 					if (MidParentIndex == StartIndex)
 					{
 						StartJointIndex = StartIndex;
@@ -76,7 +76,7 @@ void FRigUnit_TwoBoneIKFK::Execute(const FRigUnitContext& Context)
 			else if (FMath::IsNearlyEqual(IKBlend, 1.f))
 			{
 				// update transform before going through IK
-				const FRigHierarchy* Hierarchy = HierarchyRef.Get();
+				const FRigBoneHierarchy* Hierarchy = HierarchyRef.GetBones();
 				check(Hierarchy);
 
 				StartJointIKTransform = Hierarchy->GetGlobalTransform(StartJointIndex);
@@ -92,7 +92,7 @@ void FRigUnit_TwoBoneIKFK::Execute(const FRigUnitContext& Context)
 			else
 			{
 				// update transform before going through IK
-				const FRigHierarchy* Hierarchy = HierarchyRef.Get();
+				const FRigBoneHierarchy* Hierarchy = HierarchyRef.GetBones();
 				check(Hierarchy);
 
 				StartJointIKTransform = Hierarchy->GetGlobalTransform(StartJointIndex);
@@ -106,7 +106,7 @@ void FRigUnit_TwoBoneIKFK::Execute(const FRigUnitContext& Context)
 				EndJointTransform.Blend(MidJointFKTransform, EndJointIKTransform, IKBlend);
 			}
 
-			FRigHierarchy* Hierarchy = HierarchyRef.Get();
+			FRigBoneHierarchy* Hierarchy = HierarchyRef.GetBones();
 			check(Hierarchy);
 			Hierarchy->SetGlobalTransform(StartJointIndex, StartJointTransform);
 			Hierarchy->SetGlobalTransform(MidJointIndex, MidJointTransform);
diff --git a/Engine/Plugins/Experimental/ControlRig/Source/ControlRig/Private/Units/Hierarchy/RigUnit_AddBoneTransform.cpp b/Engine/Plugins/Experimental/ControlRig/Source/ControlRig/Private/Units/Hierarchy/RigUnit_AddBoneTransform.cpp
--- a/Engine/Plugins/Experimental/ControlRig/Source/ControlRig/Private/Units/Hierarchy/RigUnit_AddBoneTransform.cpp
+++ b/Engine/Plugins/Experimental/ControlRig/Source/ControlRig/Private/Units/Hierarchy/RigUnit_AddBoneTransform.cpp
@@ -12,7 +12,7 @@ void FRigUnit_AddBoneTransform::Execute(const FRigUnitContext& Context)
 {
     DECLARE_SCOPE_HIERARCHICAL_COUNTER_RIGUNIT()
 	FRigHierarchyRef& HierarchyRef = ExecuteContext.HierarchyReference;
-	FRigHierarchy* Hierarchy = HierarchyRef.Get();
+	FRigBoneHierarchy* Hierarchy = HierarchyRef.GetBones();
 	if (Hierarchy)
 	{
 		switch (Context.State)
@@ -50,8 +50,8 @@ void FRigUnit_AddBoneTransform::Execute(const FRigUnitContext& Context)
 
 IMPLEMENT_RIGUNIT_AUTOMATION_TEST(FRigUnit_AddBoneTransform)
 {
-	Hierarchy.AddBone(TEXT("Root"), NAME_None, FTransform(FVector(1.f, 0.f, 0.f)));
-	Hierarchy.AddBone(TEXT("BoneA"), TEXT("Root"), FTransform(FVector(1.f, 2.f, 0.f)));
+	Hierarchy.Add(TEXT("Root"), NAME_None, FTransform(FVector(1.f, 0.f, 0.f)));
+	Hierarchy.Add(TEXT("BoneA"), TEXT("Root"), FTransform(FVector(1.f, 2.f, 0.f)));
 	Hierarchy.Initialize();
 	Unit.ExecuteContext.HierarchyReference = HierarchyRef;
 
diff --git a/Engine/Plugins/Experimental/ControlRig/Source/ControlRig/Private/Units/Hierarchy/RigUnit_GetBoneTransform.cpp b/Engine/Plugins/Experimental/ControlRig/Source/ControlRig/Private/Units/Hierarchy/RigUnit_GetBoneTransform.cpp
--- a/Engine/Plugins/Experimental/ControlRig/Source/ControlRig/Private/Units/Hierarchy/RigUnit_GetBoneTransform.cpp
+++ b/Engine/Plugins/Experimental/ControlRig/Source/ControlRig/Private/Units/Hierarchy/RigUnit_GetBoneTransform.cpp
@@ -12,7 +12,7 @@ void FRigUnit_GetBoneTransform::Execute(const FRigUnitContext& Context)
 {
     DECLARE_SCOPE_HIERARCHICAL_COUNTER_RIGUNIT()
 	const FRigHierarchyRef& HierarchyRef = Context.HierarchyReference;
-	const FRigHierarchy* Hierarchy = HierarchyRef.Get();
+	const FRigBoneHierarchy* Hierarchy = HierarchyRef.GetBones();
 	if (Hierarchy)
 	{
 		switch (Context.State)
@@ -57,8 +57,8 @@ void FRigUnit_GetBoneTransform::Execute(const FRigUnitContext& Context)
 
 IMPLEMENT_RIGUNIT_AUTOMATION_TEST(FRigUnit_GetBoneTransform)
 {
-	Hierarchy.AddBone(TEXT("Root"), NAME_None, FTransform(FVector(1.f, 0.f, 0.f)));
-	Hierarchy.AddBone(TEXT("BoneA"), TEXT("Root"), FTransform(FVector(1.f, 2.f, 3.f)));
+	Hierarchy.Add(TEXT("Root"), NAME_None, FTransform(FVector(1.f, 0.f, 0.f)));
+	Hierarchy.Add(TEXT("BoneA"), TEXT("Root"), FTransform(FVector(1.f, 2.f, 3.f)));
 	Hierarchy.Initialize();
 
 	Unit.Bone = TEXT("Unknown");
diff --git a/Engine/Plugins/Experimental/ControlRig/Source/ControlRig/Private/Units/Hierarchy/RigUnit_GetInitialBoneTransform.cpp b/Engine/Plugins/Experimental/ControlRig/Source/ControlRig/Private/Units/Hierarchy/RigUnit_GetInitialBoneTransform.cpp
--- a/Engine/Plugins/Experimental/ControlRig/Source/ControlRig/Private/Units/Hierarchy/RigUnit_GetInitialBoneTransform.cpp
+++ b/Engine/Plugins/Experimental/ControlRig/Source/ControlRig/Private/Units/Hierarchy/RigUnit_GetInitialBoneTransform.cpp
@@ -12,7 +12,7 @@ void FRigUnit_GetInitialBoneTransform::Execute(const FRigUnitContext& Context)
 {
     DECLARE_SCOPE_HIERARCHICAL_COUNTER_RIGUNIT()
 	const FRigHierarchyRef& HierarchyRef = Context.HierarchyReference;
-	const FRigHierarchy* Hierarchy = HierarchyRef.Get();
+	const FRigBoneHierarchy* Hierarchy = HierarchyRef.GetBones();
 	if (Hierarchy)
 	{
 		switch (Context.State)
@@ -35,7 +35,7 @@ void FRigUnit_GetInitialBoneTransform::Execute(const FRigUnitContext& Context)
 						case EBoneGetterSetterMode::LocalSpace:
 						{
 							Transform = Hierarchy->GetInitialTransform(CachedBoneIndex);
-							int32 ParentBoneIndex = Hierarchy->GetParentIndex(CachedBoneIndex);
+							int32 ParentBoneIndex = (*Hierarchy)[CachedBoneIndex].ParentIndex;
 							if (ParentBoneIndex != INDEX_NONE)
 							{
 								FTransform ParentTransform = Hierarchy->GetInitialTransform(ParentBoneIndex);
@@ -63,8 +63,8 @@ void FRigUnit_GetInitialBoneTransform::Execute(const FRigUnitContext& Context)
 
 IMPLEMENT_RIGUNIT_AUTOMATION_TEST(FRigUnit_GetInitialBoneTransform)
 {
-	Hierarchy.AddBone(TEXT("Root"), NAME_None, FTransform(FVector(1.f, 0.f, 0.f)));
-	Hierarchy.AddBone(TEXT("BoneA"), TEXT("Root"), FTransform(FVector(1.f, 2.f, 3.f)));
+	Hierarchy.Add(TEXT("Root"), NAME_None, FTransform(FVector(1.f, 0.f, 0.f)));
+	Hierarchy.Add(TEXT("BoneA"), TEXT("Root"), FTransform(FVector(1.f, 2.f, 3.f)));
 	Hierarchy.Initialize();
 
 	Unit.Bone = TEXT("Root");
diff --git a/Engine/Plugins/Experimental/ControlRig/Source/ControlRig/Private/Units/Hierarchy/RigUnit_GetRelativeBoneTransform.cpp b/Engine/Plugins/Experimental/ControlRig/Source/ControlRig/Private/Units/Hierarchy/RigUnit_GetRelativeBoneTransform.cpp
--- a/Engine/Plugins/Experimental/ControlRig/Source/ControlRig/Private/Units/Hierarchy/RigUnit_GetRelativeBoneTransform.cpp
+++ b/Engine/Plugins/Experimental/ControlRig/Source/ControlRig/Private/Units/Hierarchy/RigUnit_GetRelativeBoneTransform.cpp
@@ -12,7 +12,7 @@ void FRigUnit_GetRelativeBoneTransform::Execute(const FRigUnitContext& Context)
 {
     DECLARE_SCOPE_HIERARCHICAL_COUNTER_RIGUNIT()
 	const FRigHierarchyRef& HierarchyRef = Context.HierarchyReference;
-	const FRigHierarchy* Hierarchy = HierarchyRef.Get();
+	const FRigBoneHierarchy* Hierarchy = HierarchyRef.GetBones();
 	if (Hierarchy)
 	{
 		switch (Context.State)
@@ -44,9 +44,9 @@ void FRigUnit_GetRelativeBoneTransform::Execute(const FRigUnitContext& Context)
 
 IMPLEMENT_RIGUNIT_AUTOMATION_TEST(FRigUnit_GetRelativeBoneTransform)
 {
-	Hierarchy.AddBone(TEXT("Root"), NAME_None, FTransform(FVector(1.f, 0.f, 0.f)));
-	Hierarchy.AddBone(TEXT("BoneA"), TEXT("Root"), FTransform(FVector(1.f, 2.f, 3.f)));
-	Hierarchy.AddBone(TEXT("BoneB"), TEXT("Root"), FTransform(FVector(-4.f, 0.f, 0.f)));
+	Hierarchy.Add(TEXT("Root"), NAME_None, FTransform(FVector(1.f, 0.f, 0.f)));
+	Hierarchy.Add(TEXT("BoneA"), TEXT("Root"), FTransform(FVector(1.f, 2.f, 3.f)));
+	Hierarchy.Add(TEXT("BoneB"), TEXT("Root"), FTransform(FVector(-4.f, 0.f, 0.f)));
 	Hierarchy.Initialize();
 
 	Unit.Bone = TEXT("Unknown");
diff --git a/Engine/Plugins/Experimental/ControlRig/Source/ControlRig/Private/Units/Hierarchy/RigUnit_SetBoneRotation.cpp b/Engine/Plugins/Experimental/ControlRig/Source/ControlRig/Private/Units/Hierarchy/RigUnit_SetBoneRotation.cpp
--- a/Engine/Plugins/Experimental/ControlRig/Source/ControlRig/Private/Units/Hierarchy/RigUnit_SetBoneRotation.cpp
+++ b/Engine/Plugins/Experimental/ControlRig/Source/ControlRig/Private/Units/Hierarchy/RigUnit_SetBoneRotation.cpp
@@ -12,7 +12,7 @@ void FRigUnit_SetBoneRotation::Execute(const FRigUnitContext& Context)
 {
     DECLARE_SCOPE_HIERARCHICAL_COUNTER_RIGUNIT()
 	FRigHierarchyRef& HierarchyRef = ExecuteContext.HierarchyReference;
-	FRigHierarchy* Hierarchy = HierarchyRef.Get();
+	FRigBoneHierarchy* Hierarchy = HierarchyRef.GetBones();
 	if (Hierarchy)
 	{
 		switch (Context.State)
@@ -62,9 +62,9 @@ void FRigUnit_SetBoneRotation::Execute(const FRigUnitContext& Context)
 
 IMPLEMENT_RIGUNIT_AUTOMATION_TEST(FRigUnit_SetBoneRotation)
 {
-	Hierarchy.AddBone(TEXT("Root"), NAME_None, FTransform(FQuat(FVector(-1.f, 0.f, 0.f), 0.1f)));
-	Hierarchy.AddBone(TEXT("BoneA"), TEXT("Root"), FTransform(FQuat(FVector(-1.f, 0.f, 0.f), 0.5f)));
-	Hierarchy.AddBone(TEXT("BoneB"), TEXT("BoneA"), FTransform(FQuat(FVector(-1.f, 0.f, 0.f), 0.7f)));
+	Hierarchy.Add(TEXT("Root"), NAME_None, FTransform(FQuat(FVector(-1.f, 0.f, 0.f), 0.1f)));
+	Hierarchy.Add(TEXT("BoneA"), TEXT("Root"), FTransform(FQuat(FVector(-1.f, 0.f, 0.f), 0.5f)));
+	Hierarchy.Add(TEXT("BoneB"), TEXT("BoneA"), FTransform(FQuat(FVector(-1.f, 0.f, 0.f), 0.7f)));
 	Hierarchy.Initialize();
 
 	Unit.ExecuteContext.HierarchyReference = HierarchyRef;
diff --git a/Engine/Plugins/Experimental/ControlRig/Source/ControlRig/Private/Units/Hierarchy/RigUnit_SetBoneTransform.cpp b/Engine/Plugins/Experimental/ControlRig/Source/ControlRig/Private/Units/Hierarchy/RigUnit_SetBoneTransform.cpp
--- a/Engine/Plugins/Experimental/ControlRig/Source/ControlRig/Private/Units/Hierarchy/RigUnit_SetBoneTransform.cpp
+++ b/Engine/Plugins/Experimental/ControlRig/Source/ControlRig/Private/Units/Hierarchy/RigUnit_SetBoneTransform.cpp
@@ -12,7 +12,7 @@ void FRigUnit_SetBoneTransform::Execute(const FRigUnitContext& Context)
 {
     DECLARE_SCOPE_HIERARCHICAL_COUNTER_RIGUNIT()
 	FRigHierarchyRef& HierarchyRef = ExecuteContext.HierarchyReference;
-	FRigHierarchy* Hierarchy = HierarchyRef.Get();
+	FRigBoneHierarchy* Hierarchy = HierarchyRef.GetBones();
 	if (Hierarchy)
 	{
 		switch (Context.State)
@@ -58,9 +58,9 @@ void FRigUnit_SetBoneTransform::Execute(const FRigUnitContext& Context)
 
 IMPLEMENT_RIGUNIT_AUTOMATION_TEST(FRigUnit_SetBoneTransform)
 {
-	Hierarchy.AddBone(TEXT("Root"), NAME_None, FTransform(FVector(1.f, 0.f, 0.f)));
-	Hierarchy.AddBone(TEXT("BoneA"), TEXT("Root"), FTransform(FVector(1.f, 2.f, 3.f)));
-	Hierarchy.AddBone(TEXT("BoneB"), TEXT("BoneA"), FTransform(FVector(1.f, 5.f, 3.f)));
+	Hierarchy.Add(TEXT("Root"), NAME_None, FTransform(FVector(1.f, 0.f, 0.f)));
+	Hierarchy.Add(TEXT("BoneA"), TEXT("Root"), FTransform(FVector(1.f, 2.f, 3.f)));
+	Hierarchy.Add(TEXT("BoneB"), TEXT("BoneA"), FTransform(FVector(1.f, 5.f, 3.f)));
 	Hierarchy.Initialize();
 	Unit.ExecuteContext.HierarchyReference = HierarchyRef;
 
diff --git a/Engine/Plugins/Experimental/ControlRig/Source/ControlRig/Private/Units/Hierarchy/RigUnit_SetBoneTranslation.cpp b/Engine/Plugins/Experimental/ControlRig/Source/ControlRig/Private/Units/Hierarchy/RigUnit_SetBoneTranslation.cpp
--- a/Engine/Plugins/Experimental/ControlRig/Source/ControlRig/Private/Units/Hierarchy/RigUnit_SetBoneTranslation.cpp
+++ b/Engine/Plugins/Experimental/ControlRig/Source/ControlRig/Private/Units/Hierarchy/RigUnit_SetBoneTranslation.cpp
@@ -12,7 +12,7 @@ void FRigUnit_SetBoneTranslation::Execute(const FRigUnitContext& Context)
 {
     DECLARE_SCOPE_HIERARCHICAL_COUNTER_RIGUNIT()
 	FRigHierarchyRef& HierarchyRef = ExecuteContext.HierarchyReference;
-	FRigHierarchy* Hierarchy = HierarchyRef.Get();
+	FRigBoneHierarchy* Hierarchy = HierarchyRef.GetBones();
 	if (Hierarchy)
 	{
 		switch (Context.State)
diff --git a/Engine/Plugins/Experimental/ControlRig/Source/ControlRig/Private/Units/Hierarchy/RigUnit_SetRelativeBoneTransform.cpp b/Engine/Plugins/Experimental/ControlRig/Source/ControlRig/Private/Units/Hierarchy/RigUnit_SetRelativeBoneTransform.cpp
--- a/Engine/Plugins/Experimental/ControlRig/Source/ControlRig/Private/Units/Hierarchy/RigUnit_SetRelativeBoneTransform.cpp
+++ b/Engine/Plugins/Experimental/ControlRig/Source/ControlRig/Private/Units/Hierarchy/RigUnit_SetRelativeBoneTransform.cpp
@@ -12,7 +12,7 @@ void FRigUnit_SetRelativeBoneTransform::Execute(const FRigUnitContext& Context)
 {
     DECLARE_SCOPE_HIERARCHICAL_COUNTER_RIGUNIT()
 	FRigHierarchyRef& HierarchyRef = ExecuteContext.HierarchyReference;
-	FRigHierarchy* Hierarchy = HierarchyRef.Get();
+	FRigBoneHierarchy* Hierarchy = HierarchyRef.GetBones();
 	if (Hierarchy)
 	{
 		switch (Context.State)
@@ -45,10 +45,10 @@ void FRigUnit_SetRelativeBoneTransform::Execute(const FRigUnitContext& Context)
 
 IMPLEMENT_RIGUNIT_AUTOMATION_TEST(FRigUnit_SetRelativeBoneTransform)
 {
-	Hierarchy.AddBone(TEXT("Root"), NAME_None, FTransform(FVector(1.f, 0.f, 0.f)));
-	Hierarchy.AddBone(TEXT("BoneA"), TEXT("Root"), FTransform(FVector(1.f, 2.f, 3.f)));
-	Hierarchy.AddBone(TEXT("BoneB"), TEXT("BoneA"), FTransform(FVector(1.f, 5.f, 3.f)));
-	Hierarchy.AddBone(TEXT("BoneC"), TEXT("Root"), FTransform(FVector(-4.f, 0.f, 0.f)));
+	Hierarchy.Add(TEXT("Root"), NAME_None, FTransform(FVector(1.f, 0.f, 0.f)));
+	Hierarchy.Add(TEXT("BoneA"), TEXT("Root"), FTransform(FVector(1.f, 2.f, 3.f)));
+	Hierarchy.Add(TEXT("BoneB"), TEXT("BoneA"), FTransform(FVector(1.f, 5.f, 3.f)));
+	Hierarchy.Add(TEXT("BoneC"), TEXT("Root"), FTransform(FVector(-4.f, 0.f, 0.f)));
 	Hierarchy.Initialize();
 	Unit.ExecuteContext.HierarchyReference = HierarchyRef;
 
diff --git a/Engine/Plugins/Experimental/ControlRig/Source/ControlRig/Private/Units/Highlevel/Harmonics/RigUnit_BoneHarmonics.cpp b/Engine/Plugins/Experimental/ControlRig/Source/ControlRig/Private/Units/Highlevel/Harmonics/RigUnit_BoneHarmonics.cpp
--- a/Engine/Plugins/Experimental/ControlRig/Source/ControlRig/Private/Units/Highlevel/Harmonics/RigUnit_BoneHarmonics.cpp
+++ b/Engine/Plugins/Experimental/ControlRig/Source/ControlRig/Private/Units/Highlevel/Harmonics/RigUnit_BoneHarmonics.cpp
@@ -7,7 +7,7 @@ void FRigUnit_BoneHarmonics::Execute(const FRigUnitContext& Context)
 {
     DECLARE_SCOPE_HIERARCHICAL_COUNTER_RIGUNIT()
 	FRigHierarchyRef& HierarchyRef = ExecuteContext.HierarchyReference;
-	FRigHierarchy* Hierarchy = HierarchyRef.Get();
+	FRigBoneHierarchy* Hierarchy = HierarchyRef.GetBones();
 	if (Hierarchy == nullptr)
 	{
 		return;
diff --git a/Engine/Plugins/Experimental/ControlRig/Source/ControlRig/Private/Units/Highlevel/Harmonics/RigUnit_ChainHarmonics.cpp b/Engine/Plugins/Experimental/ControlRig/Source/ControlRig/Private/Units/Highlevel/Harmonics/RigUnit_ChainHarmonics.cpp
--- a/Engine/Plugins/Experimental/ControlRig/Source/ControlRig/Private/Units/Highlevel/Harmonics/RigUnit_ChainHarmonics.cpp
+++ b/Engine/Plugins/Experimental/ControlRig/Source/ControlRig/Private/Units/Highlevel/Harmonics/RigUnit_ChainHarmonics.cpp
@@ -6,7 +6,7 @@
 void FRigUnit_ChainHarmonics::Execute(const FRigUnitContext& Context)
 {
     DECLARE_SCOPE_HIERARCHICAL_COUNTER_RIGUNIT()
-	FRigHierarchy* Hierarchy = (FRigHierarchy*)(Context.HierarchyReference.Get());
+	FRigBoneHierarchy* Hierarchy = (FRigBoneHierarchy*)(Context.HierarchyReference.GetBones());
 	if (Hierarchy == nullptr)
 	{
 		return;
@@ -74,7 +74,7 @@ void FRigUnit_ChainHarmonics::Execute(const FRigUnitContext& Context)
 	}
 
 	FTransform ParentTransform = FTransform::Identity;
-	int32 ParentIndex = Hierarchy->GetParentIndex(Bones[0]);
+	int32 ParentIndex = (*Hierarchy)[Bones[0]].ParentIndex;
 	if (ParentIndex != INDEX_NONE)
 	{
 		ParentTransform = Hierarchy->GetGlobalTransform(ParentIndex);
diff --git a/Engine/Plugins/Experimental/ControlRig/Source/ControlRig/Private/Units/Highlevel/Hierarchy/RigUnit_AimBone.cpp b/Engine/Plugins/Experimental/ControlRig/Source/ControlRig/Private/Units/Highlevel/Hierarchy/RigUnit_AimBone.cpp
--- a/Engine/Plugins/Experimental/ControlRig/Source/ControlRig/Private/Units/Highlevel/Hierarchy/RigUnit_AimBone.cpp
+++ b/Engine/Plugins/Experimental/ControlRig/Source/ControlRig/Private/Units/Highlevel/Hierarchy/RigUnit_AimBone.cpp
@@ -6,7 +6,7 @@
 void FRigUnit_AimBone::Execute(const FRigUnitContext& Context)
 {
     DECLARE_SCOPE_HIERARCHICAL_COUNTER_RIGUNIT()
-	FRigHierarchy* Hierarchy = (FRigHierarchy*)(Context.HierarchyReference.Get());
+	FRigBoneHierarchy* Hierarchy = (FRigBoneHierarchy*)(Context.HierarchyReference.GetBones());
 	if (Hierarchy == nullptr)
 	{
 		return;
diff --git a/Engine/Plugins/Experimental/ControlRig/Source/ControlRig/Private/Units/Highlevel/Hierarchy/RigUnit_CCDIK.cpp b/Engine/Plugins/Experimental/ControlRig/Source/ControlRig/Private/Units/Highlevel/Hierarchy/RigUnit_CCDIK.cpp
--- a/Engine/Plugins/Experimental/ControlRig/Source/ControlRig/Private/Units/Highlevel/Hierarchy/RigUnit_CCDIK.cpp
+++ b/Engine/Plugins/Experimental/ControlRig/Source/ControlRig/Private/Units/Highlevel/Hierarchy/RigUnit_CCDIK.cpp
@@ -6,7 +6,7 @@
 void FRigUnit_CCDIK::Execute(const FRigUnitContext& Context)
 {
     DECLARE_SCOPE_HIERARCHICAL_COUNTER_RIGUNIT()
-	FRigHierarchy* Hierarchy = (FRigHierarchy*)(Context.HierarchyReference.Get());
+	FRigBoneHierarchy* Hierarchy = (FRigBoneHierarchy*)(Context.HierarchyReference.GetBones());
 	if (Hierarchy == nullptr)
 	{
 		return;
@@ -27,7 +27,7 @@ void FRigUnit_CCDIK::Execute(const FRigUnitContext& Context)
 			while (CurrentIndex != INDEX_NONE)
 			{
 				// ensure the chain
-				int32 ParentIndex = Hierarchy->GetParentIndex(CurrentIndex);
+				int32 ParentIndex = (*Hierarchy)[CurrentIndex].ParentIndex;
 				if (ParentIndex != INDEX_NONE)
 				{
 					BoneIndices.Add(CurrentIndex);
@@ -45,7 +45,7 @@ void FRigUnit_CCDIK::Execute(const FRigUnitContext& Context)
 			Chain.Reserve(BoneIndices.Num());
 		}
 
-		int32 RootParentIndex = Hierarchy->GetParentIndex(RootIndex);
+		int32 RootParentIndex = (*Hierarchy)[RootIndex].ParentIndex;
 		if (RootParentIndex != INDEX_NONE)
 		{
 			BoneIndices.Add(RootParentIndex);
diff --git a/Engine/Plugins/Experimental/ControlRig/Source/ControlRig/Private/Units/Highlevel/Hierarchy/RigUnit_DistributeRotation.cpp b/Engine/Plugins/Experimental/ControlRig/Source/ControlRig/Private/Units/Highlevel/Hierarchy/RigUnit_DistributeRotation.cpp
--- a/Engine/Plugins/Experimental/ControlRig/Source/ControlRig/Private/Units/Highlevel/Hierarchy/RigUnit_DistributeRotation.cpp
+++ b/Engine/Plugins/Experimental/ControlRig/Source/ControlRig/Private/Units/Highlevel/Hierarchy/RigUnit_DistributeRotation.cpp
@@ -6,7 +6,7 @@
 void FRigUnit_DistributeRotation::Execute(const FRigUnitContext& Context)
 {
     DECLARE_SCOPE_HIERARCHICAL_COUNTER_RIGUNIT()
-	FRigHierarchy* Hierarchy = (FRigHierarchy*)(Context.HierarchyReference.Get());
+	FRigBoneHierarchy* Hierarchy = (FRigBoneHierarchy*)(Context.HierarchyReference.GetBones());
 	if (Hierarchy == nullptr)
 	{
 		return;
@@ -36,7 +36,7 @@ void FRigUnit_DistributeRotation::Execute(const FRigUnitContext& Context)
 				{
 					break;
 				}
-				EndBoneIndex = Hierarchy->GetParentIndex(EndBoneIndex);
+				EndBoneIndex = (*Hierarchy)[EndBoneIndex].ParentIndex;
 			}
 		}
 
@@ -140,7 +140,7 @@ void FRigUnit_DistributeRotation::Execute(const FRigUnitContext& Context)
 	}
 
 	FTransform BaseTransform = FTransform::Identity;
-	int32 ParentIndex = Hierarchy->GetParentIndex(BoneIndices[0]);
+	int32 ParentIndex = (*Hierarchy)[BoneIndices[0]].ParentIndex;
 	if (ParentIndex != INDEX_NONE)
 	{
 		BaseTransform = Hierarchy->GetGlobalTransform(ParentIndex);
@@ -184,11 +184,11 @@ void FRigUnit_DistributeRotation::Execute(const FRigUnitContext& Context)
 
 IMPLEMENT_RIGUNIT_AUTOMATION_TEST(FRigUnit_DistributeRotation)
 {
-	Hierarchy.AddBone(TEXT("Root"), NAME_None, FTransform(FVector(1.f, 0.f, 0.f)));
-	Hierarchy.AddBone(TEXT("BoneA"), TEXT("Root"), FTransform(FVector(2.f, 0.f, 0.f)));
-	Hierarchy.AddBone(TEXT("BoneB"), TEXT("BoneA"), FTransform(FVector(2.f, 0.f, 0.f)));
-	Hierarchy.AddBone(TEXT("BoneC"), TEXT("BoneB"), FTransform(FVector(2.f, 0.f, 0.f)));
-	Hierarchy.AddBone(TEXT("BoneD"), TEXT("BoneC"), FTransform(FVector(2.f, 0.f, 0.f)));
+	Hierarchy.Add(TEXT("Root"), NAME_None, FTransform(FVector(1.f, 0.f, 0.f)));
+	Hierarchy.Add(TEXT("BoneA"), TEXT("Root"), FTransform(FVector(2.f, 0.f, 0.f)));
+	Hierarchy.Add(TEXT("BoneB"), TEXT("BoneA"), FTransform(FVector(2.f, 0.f, 0.f)));
+	Hierarchy.Add(TEXT("BoneC"), TEXT("BoneB"), FTransform(FVector(2.f, 0.f, 0.f)));
+	Hierarchy.Add(TEXT("BoneD"), TEXT("BoneC"), FTransform(FVector(2.f, 0.f, 0.f)));
 	Hierarchy.Initialize();
 
 	Unit.StartBone = TEXT("Root");
diff --git a/Engine/Plugins/Experimental/ControlRig/Source/ControlRig/Private/Units/Highlevel/Hierarchy/RigUnit_FABRIK.cpp b/Engine/Plugins/Experimental/ControlRig/Source/ControlRig/Private/Units/Highlevel/Hierarchy/RigUnit_FABRIK.cpp
--- a/Engine/Plugins/Experimental/ControlRig/Source/ControlRig/Private/Units/Highlevel/Hierarchy/RigUnit_FABRIK.cpp
+++ b/Engine/Plugins/Experimental/ControlRig/Source/ControlRig/Private/Units/Highlevel/Hierarchy/RigUnit_FABRIK.cpp
@@ -6,7 +6,7 @@
 void FRigUnit_FABRIK::Execute(const FRigUnitContext& Context)
 {
     DECLARE_SCOPE_HIERARCHICAL_COUNTER_RIGUNIT()
-	FRigHierarchy* Hierarchy = (FRigHierarchy*)(Context.HierarchyReference.Get());
+	FRigBoneHierarchy* Hierarchy = (FRigBoneHierarchy*)(Context.HierarchyReference.GetBones());
 	if (Hierarchy == nullptr)
 	{
 		return;
@@ -24,7 +24,7 @@ void FRigUnit_FABRIK::Execute(const FRigUnitContext& Context)
 			while (CurrentIndex != INDEX_NONE)
 			{
 				// ensure the chain
-				int32 ParentIndex = Hierarchy->GetParentIndex(CurrentIndex);
+				int32 ParentIndex = (*Hierarchy)[CurrentIndex].ParentIndex;
 				if (ParentIndex != INDEX_NONE)
 				{
 					BoneIndices.Add(CurrentIndex);
diff --git a/Engine/Plugins/Experimental/ControlRig/Source/ControlRig/Private/Units/Highlevel/Hierarchy/RigUnit_FitChainToCurve.cpp b/Engine/Plugins/Experimental/ControlRig/Source/ControlRig/Private/Units/Highlevel/Hierarchy/RigUnit_FitChainToCurve.cpp
--- a/Engine/Plugins/Experimental/ControlRig/Source/ControlRig/Private/Units/Highlevel/Hierarchy/RigUnit_FitChainToCurve.cpp
+++ b/Engine/Plugins/Experimental/ControlRig/Source/ControlRig/Private/Units/Highlevel/Hierarchy/RigUnit_FitChainToCurve.cpp
@@ -6,7 +6,7 @@
 void FRigUnit_FitChainToCurve::Execute(const FRigUnitContext& Context)
 {
     DECLARE_SCOPE_HIERARCHICAL_COUNTER_RIGUNIT()
-	FRigHierarchy* Hierarchy = (FRigHierarchy*)(Context.HierarchyReference.Get());
+	FRigBoneHierarchy* Hierarchy = (FRigBoneHierarchy*)(Context.HierarchyReference.GetBones());
 	if (Hierarchy == nullptr)
 	{
 		return;
@@ -42,7 +42,7 @@ void FRigUnit_FitChainToCurve::Execute(const FRigUnitContext& Context)
 				{
 					break;
 				}
-				EndBoneIndex = Hierarchy->GetParentIndex(EndBoneIndex);
+				EndBoneIndex = (*Hierarchy)[EndBoneIndex].ParentIndex;
 			}
 		}
 
@@ -324,7 +324,7 @@ void FRigUnit_FitChainToCurve::Execute(const FRigUnitContext& Context)
 	if (Rotations.Num() > 0)
 	{
 		FTransform BaseTransform = FTransform::Identity;
-		int32 ParentIndex = Hierarchy->GetParentIndex(BoneIndices[0]);
+		int32 ParentIndex = (*Hierarchy)[BoneIndices[0]].ParentIndex;
 		if (ParentIndex != INDEX_NONE)
 		{
 			BaseTransform = Hierarchy->GetGlobalTransform(ParentIndex);
diff --git a/Engine/Plugins/Experimental/ControlRig/Source/ControlRig/Private/Units/Highlevel/Hierarchy/RigUnit_ModifyBoneTransforms.cpp b/Engine/Plugins/Experimental/ControlRig/Source/ControlRig/Private/Units/Highlevel/Hierarchy/RigUnit_ModifyBoneTransforms.cpp
--- a/Engine/Plugins/Experimental/ControlRig/Source/ControlRig/Private/Units/Highlevel/Hierarchy/RigUnit_ModifyBoneTransforms.cpp
+++ b/Engine/Plugins/Experimental/ControlRig/Source/ControlRig/Private/Units/Highlevel/Hierarchy/RigUnit_ModifyBoneTransforms.cpp
@@ -7,7 +7,7 @@ void FRigUnit_ModifyBoneTransforms::Execute(const FRigUnitContext& Context)
 {
     DECLARE_SCOPE_HIERARCHICAL_COUNTER_RIGUNIT()
 	FRigHierarchyRef& HierarchyRef = ExecuteContext.HierarchyReference;
-	FRigHierarchy* Hierarchy = HierarchyRef.Get();
+	FRigBoneHierarchy* Hierarchy = HierarchyRef.GetBones();
 	if (Hierarchy)
 	{
 		switch (Context.State)
@@ -109,9 +109,9 @@ void FRigUnit_ModifyBoneTransforms::Execute(const FRigUnitContext& Context)
 
 IMPLEMENT_RIGUNIT_AUTOMATION_TEST(FRigUnit_ModifyBoneTransforms)
 {
-	Hierarchy.AddBone(TEXT("Root"), NAME_None, FTransform(FVector(1.f, 0.f, 0.f)));
-	Hierarchy.AddBone(TEXT("BoneA"), TEXT("Root"), FTransform(FVector(1.f, 2.f, 3.f)));
-	Hierarchy.AddBone(TEXT("BoneB"), TEXT("Root"), FTransform(FVector(5.f, 6.f, 7.f)));
+	Hierarchy.Add(TEXT("Root"), NAME_None, FTransform(FVector(1.f, 0.f, 0.f)));
+	Hierarchy.Add(TEXT("BoneA"), TEXT("Root"), FTransform(FVector(1.f, 2.f, 3.f)));
+	Hierarchy.Add(TEXT("BoneB"), TEXT("Root"), FTransform(FVector(5.f, 6.f, 7.f)));
 	Hierarchy.Initialize();
 	Unit.ExecuteContext.HierarchyReference = HierarchyRef;
 
diff --git a/Engine/Plugins/Experimental/ControlRig/Source/ControlRig/Private/Units/Highlevel/Hierarchy/RigUnit_SpringIK.cpp b/Engine/Plugins/Experimental/ControlRig/Source/ControlRig/Private/Units/Highlevel/Hierarchy/RigUnit_SpringIK.cpp
--- a/Engine/Plugins/Experimental/ControlRig/Source/ControlRig/Private/Units/Highlevel/Hierarchy/RigUnit_SpringIK.cpp
+++ b/Engine/Plugins/Experimental/ControlRig/Source/ControlRig/Private/Units/Highlevel/Hierarchy/RigUnit_SpringIK.cpp
@@ -6,7 +6,7 @@
 void FRigUnit_SpringIK::Execute(const FRigUnitContext& Context)
 {
     DECLARE_SCOPE_HIERARCHICAL_COUNTER_RIGUNIT()
-	FRigHierarchy* Hierarchy = (FRigHierarchy*)(Context.HierarchyReference.Get());
+	FRigBoneHierarchy* Hierarchy = (FRigBoneHierarchy*)(Context.HierarchyReference.GetBones());
 	if (Hierarchy == nullptr)
 	{
 		return;
@@ -34,7 +34,7 @@ void FRigUnit_SpringIK::Execute(const FRigUnitContext& Context)
 				{
 					break;
 				}
-				EndBoneIndex = Hierarchy->GetParentIndex(EndBoneIndex);
+				EndBoneIndex = (*Hierarchy)[EndBoneIndex].ParentIndex;
 			}
 		}
 
@@ -252,7 +252,7 @@ void FRigUnit_SpringIK::Execute(const FRigUnitContext& Context)
 
 		if(bLimitLocalPosition)
 		{
-			int32 ParentIndex = Hierarchy->GetParentIndex(BoneIndices[PointIndex]);
+			int32 ParentIndex = (*Hierarchy)[BoneIndices[PointIndex]].ParentIndex;
 			if (ParentIndex != INDEX_NONE)
 			{
 				FTransform InitialTransform = Hierarchy->GetInitialTransform(BoneIndices[PointIndex]);
diff --git a/Engine/Plugins/Experimental/ControlRig/Source/ControlRig/Private/Units/Highlevel/Hierarchy/RigUnit_TransformConstraint.cpp b/Engine/Plugins/Experimental/ControlRig/Source/ControlRig/Private/Units/Highlevel/Hierarchy/RigUnit_TransformConstraint.cpp
--- a/Engine/Plugins/Experimental/ControlRig/Source/ControlRig/Private/Units/Highlevel/Hierarchy/RigUnit_TransformConstraint.cpp
+++ b/Engine/Plugins/Experimental/ControlRig/Source/ControlRig/Private/Units/Highlevel/Hierarchy/RigUnit_TransformConstraint.cpp
@@ -15,7 +15,7 @@ void FRigUnit_TransformConstraint::Execute(const FRigUnitContext& Context)
 		ConstraintData.Reset();
 		ConstraintDataToTargets.Reset();
 
-		FRigHierarchy* Hierarchy = HierarchyRef.Get();
+		FRigBoneHierarchy* Hierarchy = HierarchyRef.GetBones();
 		if (Hierarchy)
 		{
 			int32 BoneIndex = Hierarchy->GetIndex(Bone);
@@ -26,7 +26,7 @@ void FRigUnit_TransformConstraint::Execute(const FRigUnitContext& Context)
 				{
 					const FTransform SourceTransform = Hierarchy->GetGlobalTransform(BoneIndex);
 					FTransform InputBaseTransform = UtilityHelpers::GetBaseTransformByMode(BaseTransformSpace, [Hierarchy](const FName& JointName) { return Hierarchy->GetGlobalTransform(JointName); },
-						Hierarchy->GetBones()[BoneIndex].ParentName, BaseBone, BaseTransform);
+						(*Hierarchy)[BoneIndex].ParentName, BaseBone, BaseTransform);
 					for (int32 TargetIndex = 0; TargetIndex < TargetNum; ++TargetIndex)
 					{
 						// talk to Rob about the implication of support both of this
@@ -62,7 +62,7 @@ void FRigUnit_TransformConstraint::Execute(const FRigUnitContext& Context)
 	}
 	else if (Context.State == EControlRigState::Update)
 	{
-		FRigHierarchy* Hierarchy = HierarchyRef.Get();
+		FRigBoneHierarchy* Hierarchy = HierarchyRef.GetBones();
 		if (Hierarchy)
 		{
 			int32 BoneIndex = Hierarchy->GetIndex(Bone);
@@ -80,7 +80,7 @@ void FRigUnit_TransformConstraint::Execute(const FRigUnitContext& Context)
 					}
 
 					FTransform InputBaseTransform = UtilityHelpers::GetBaseTransformByMode(BaseTransformSpace, [Hierarchy](const FName& JointName) { return Hierarchy->GetGlobalTransform(JointName); },
-							Hierarchy->GetBones()[BoneIndex].ParentName, BaseBone, BaseTransform);
+							(*Hierarchy)[BoneIndex].ParentName, BaseBone, BaseTransform);
 
 					FTransform SourceTransform = Hierarchy->GetGlobalTransform(BoneIndex);
 
diff --git a/Engine/Plugins/Experimental/ControlRig/Source/ControlRig/Private/Units/Highlevel/Hierarchy/RigUnit_TransformConstraint.h b/Engine/Plugins/Experimental/ControlRig/Source/ControlRig/Private/Units/Highlevel/Hierarchy/RigUnit_TransformConstraint.h
--- a/Engine/Plugins/Experimental/ControlRig/Source/ControlRig/Private/Units/Highlevel/Hierarchy/RigUnit_TransformConstraint.h
+++ b/Engine/Plugins/Experimental/ControlRig/Source/ControlRig/Private/Units/Highlevel/Hierarchy/RigUnit_TransformConstraint.h
@@ -4,7 +4,7 @@
 
 #include "Units/RigUnit.h"
 #include "Units/Highlevel/RigUnit_HighlevelBase.h"
-#include "Rigs/Hierarchy.h"
+#include "Rigs/RigHierarchyContainer.h"
 #include "Constraint.h"
 #include "ControlRigDefines.h"
 #include "RigUnit_TransformConstraint.generated.h"
diff --git a/Engine/Plugins/Experimental/ControlRig/Source/ControlRig/Private/Units/Highlevel/Hierarchy/RigUnit_TwoBoneIKSimple.cpp b/Engine/Plugins/Experimental/ControlRig/Source/ControlRig/Private/Units/Highlevel/Hierarchy/RigUnit_TwoBoneIKSimple.cpp
--- a/Engine/Plugins/Experimental/ControlRig/Source/ControlRig/Private/Units/Highlevel/Hierarchy/RigUnit_TwoBoneIKSimple.cpp
+++ b/Engine/Plugins/Experimental/ControlRig/Source/ControlRig/Private/Units/Highlevel/Hierarchy/RigUnit_TwoBoneIKSimple.cpp
@@ -7,7 +7,7 @@
 void FRigUnit_TwoBoneIKSimple::Execute(const FRigUnitContext& Context)
 {
     DECLARE_SCOPE_HIERARCHICAL_COUNTER_RIGUNIT()
-	FRigHierarchy* Hierarchy = (FRigHierarchy*)(Context.HierarchyReference.Get());
+	FRigBoneHierarchy* Hierarchy = (FRigBoneHierarchy*)(Context.HierarchyReference.GetBones());
 	if (Hierarchy == nullptr)
 	{
 		return;
diff --git a/Engine/Plugins/Experimental/ControlRig/Source/ControlRig/Private/Units/Simulation/RigUnit_PointSimulation.cpp b/Engine/Plugins/Experimental/ControlRig/Source/ControlRig/Private/Units/Simulation/RigUnit_PointSimulation.cpp
--- a/Engine/Plugins/Experimental/ControlRig/Source/ControlRig/Private/Units/Simulation/RigUnit_PointSimulation.cpp
+++ b/Engine/Plugins/Experimental/ControlRig/Source/ControlRig/Private/Units/Simulation/RigUnit_PointSimulation.cpp
@@ -7,7 +7,7 @@ void FRigUnit_PointSimulation::Execute(const FRigUnitContext& Context)
 {
     DECLARE_SCOPE_HIERARCHICAL_COUNTER_RIGUNIT()
 	FRigHierarchyRef& HierarchyRef = ExecuteContext.HierarchyReference;
-	FRigHierarchy* Hierarchy = HierarchyRef.Get();
+	FRigBoneHierarchy* Hierarchy = HierarchyRef.GetBones();
 
 	float DeltaTime = Context.DeltaTime;
 
@@ -163,7 +163,7 @@ void FRigUnit_PointSimulation::Execute(const FRigUnitContext& Context)
 
 				if (bLimitLocalPosition)
 				{
-					int32 ParentIndex = Hierarchy->GetParentIndex(BoneIndices[TargetIndex]);
+					int32 ParentIndex = (*Hierarchy)[BoneIndices[TargetIndex]].ParentIndex;
 					if (ParentIndex != INDEX_NONE)
 					{
 						FTransform InitialTransform = Hierarchy->GetInitialTransform(BoneIndices[TargetIndex]);
diff --git a/Engine/Plugins/Experimental/ControlRig/Source/ControlRig/Public/ControlRig.h b/Engine/Plugins/Experimental/ControlRig/Source/ControlRig/Public/ControlRig.h
--- a/Engine/Plugins/Experimental/ControlRig/Source/ControlRig/Public/ControlRig.h
+++ b/Engine/Plugins/Experimental/ControlRig/Source/ControlRig/Public/ControlRig.h
@@ -9,8 +9,8 @@
 #include "Engine/EngineBaseTypes.h"
 #include "Templates/SubclassOf.h"
 #include "ControlRigDefines.h"
-#include "CurveContainer.h"
-#include "Rigs/Hierarchy.h"
+#include "Rigs/RigCurveContainer.h"
+#include "Rigs/RigHierarchyContainer.h"
 #include "Units/RigUnitContext.h"
 #include "Animation/NodeMappingProviderInterface.h"
 #include "Units/RigUnit.h"
@@ -118,9 +118,9 @@ public:
 	void SetGlobalTransform(const int32 BoneIndex, const FTransform& InTransform, bool bPropagateTransform = true) ;
 
 	/** Returns base hierarchy */
-	const FRigHierarchy& GetBaseHierarchy() const
+	const FRigBoneHierarchy& GetBoneHierarchy() const
 	{
-		return Hierarchy.BaseHierarchy;
+		return Hierarchy.BoneHierarchy;
 	}
 
 	/** Evaluate another animation ControlRig */
@@ -307,8 +307,8 @@ private:
 	friend class UControlRigEditorLibrary;
 	friend class URigUnitEditor_Base;
 	friend class FControlRigEditor;
-	friend class SRigHierarchy;
 	friend class SRigCurveContainer;
+	friend class SRigBoneHierarchy;
 	friend class UEngineTestControlRig;
  	friend class FControlRigEditMode;
 	friend class FControlRigIOHelper;
diff --git a/Engine/Plugins/Experimental/ControlRig/Source/ControlRig/Public/ControlRigDefines.h b/Engine/Plugins/Experimental/ControlRig/Source/ControlRig/Public/ControlRigDefines.h
--- a/Engine/Plugins/Experimental/ControlRig/Source/ControlRig/Public/ControlRigDefines.h
+++ b/Engine/Plugins/Experimental/ControlRig/Source/ControlRig/Public/ControlRigDefines.h
@@ -3,8 +3,8 @@
 #pragma once
 
 #include "PropertyPathHelpers.h"
-#include "Rigs/Hierarchy.h"
-#include "CurveContainer.h"
+#include "Rigs/RigHierarchyContainer.h"
+#include "Rigs/RigCurveContainer.h"
 #include "Stats/StatsHierarchical.h"
 #include "ControlRigDefines.generated.h"
 
diff --git a/Engine/Plugins/Experimental/ControlRig/Source/ControlRig/Public/CurveContainer.h b/Engine/Plugins/Experimental/ControlRig/Source/ControlRig/Public/CurveContainer.h
deleted file mode 100644
diff --git a/Engine/Plugins/Experimental/ControlRig/Source/ControlRig/Public/Drawing/ControlRigDrawInterface.h b/Engine/Plugins/Experimental/ControlRig/Source/ControlRig/Public/Drawing/ControlRigDrawInterface.h
--- a/Engine/Plugins/Experimental/ControlRig/Source/ControlRig/Public/Drawing/ControlRigDrawInterface.h
+++ b/Engine/Plugins/Experimental/ControlRig/Source/ControlRig/Public/Drawing/ControlRigDrawInterface.h
@@ -3,7 +3,7 @@
 #pragma once
 
 #include "CoreMinimal.h"
-#include "Rigs/Hierarchy.h"
+#include "Rigs/RigHierarchyContainer.h"
 #include "Math/Simulation/CRSimPointContainer.h"
 #include "ControlRigDrawInterface.generated.h"
 
@@ -38,7 +38,7 @@ public:
 	void DrawRectangle(const FTransform& WorldOffset, const FTransform& Transform, float Size, const FLinearColor& Color, float Thickness);
 	void DrawArc(const FTransform& WorldOffset, const FTransform& Transform, float Radius, float MinimumAngle, float MaximumAngle, const FLinearColor& Color, float Thickness, int32 Detail);
 	void DrawBezier(const FTransform& WorldOffset, const FCRFourPointBezier& InBezier, float MinimumU, float MaximumU, const FLinearColor& Color, float Thickness, int32 Detail);
-	void DrawHierarchy(const FTransform& WorldOffset, const FRigHierarchy& Hierarchy, EControlRigDrawHierarchyMode::Type Mode, float Scale, const FLinearColor& Color, float Thickness);
+	void DrawHierarchy(const FTransform& WorldOffset, const FRigBoneHierarchy& Hierarchy, EControlRigDrawHierarchyMode::Type Mode, float Scale, const FLinearColor& Color, float Thickness);
 	void DrawPointSimulation(const FTransform& WorldOffset, const FCRSimPointContainer& Simulation, const FLinearColor& Color, float Thickness, float PrimitiveSize = 0.f, bool bDrawPointsAsSphere = false);
 
 private:
diff --git a/Engine/Plugins/Experimental/ControlRig/Source/ControlRig/Public/Rigs/Hierarchy.h b/Engine/Plugins/Experimental/ControlRig/Source/ControlRig/Public/Rigs/Hierarchy.h
deleted file mode 100644
diff --git a/Engine/Plugins/Experimental/ControlRig/Source/ControlRig/Public/Rigs/RigBoneHierarchy.h b/Engine/Plugins/Experimental/ControlRig/Source/ControlRig/Public/Rigs/RigBoneHierarchy.h
new file mode 100644
--- /dev/null
+++ b/Engine/Plugins/Experimental/ControlRig/Source/ControlRig/Public/Rigs/RigBoneHierarchy.h
@@ -0,0 +1,159 @@
+// Copyright 1998-2019 Epic Games, Inc. All Rights Reserved.
+
+#pragma once
+
+#include "CoreMinimal.h"
+#include "RigHierarchyDefines.h"
+#include "RigBoneHierarchy.generated.h"
+
+class UControlRig;
+
+USTRUCT(BlueprintType)
+struct CONTROLRIG_API FRigBone
+{
+	GENERATED_BODY()
+
+	UPROPERTY(VisibleAnywhere, Category = FRigBoneHierarchy)
+	FName Name;
+
+	UPROPERTY(VisibleAnywhere, Category = FRigBoneHierarchy)
+	FName ParentName;
+
+	UPROPERTY(transient)
+	int32 ParentIndex;
+
+	/* Initial global transform that is saved in this rig */
+	UPROPERTY(EditAnywhere, Category = FRigBoneHierarchy)
+	FTransform InitialTransform;
+
+	UPROPERTY(transient, VisibleAnywhere, Category = FRigBoneHierarchy)
+	FTransform GlobalTransform;
+
+	UPROPERTY(transient, VisibleAnywhere, Category = FRigBoneHierarchy)
+	FTransform LocalTransform;
+
+	/** dependent list - direct dependent for child or anything that needs to update due to this */
+	UPROPERTY(transient)
+	TArray<int32> Dependents;
+};
+
+USTRUCT()
+struct CONTROLRIG_API FRigBoneHierarchy
+{
+	GENERATED_BODY()
+
+	FRigBoneHierarchy();
+	FRigBoneHierarchy& operator= (const FRigBoneHierarchy &InOther);
+
+	FORCEINLINE int32 Num() const { return Bones.Num(); }
+	FORCEINLINE const FRigBone& operator[](int32 InIndex) const { return Bones[InIndex]; }
+	FORCEINLINE FRigBone& operator[](int32 InIndex) { return Bones[InIndex]; }
+	FORCEINLINE const FRigBone& operator[](const FName& InName) const { return Bones[GetIndex(InName)]; }
+	FORCEINLINE FRigBone& operator[](const FName& InName) { return Bones[GetIndex(InName)]; }
+
+	FORCEINLINE TArray<FRigBone>::RangedForIteratorType      begin()       { return Bones.begin(); }
+	FORCEINLINE TArray<FRigBone>::RangedForConstIteratorType begin() const { return Bones.begin(); }
+	FORCEINLINE TArray<FRigBone>::RangedForIteratorType      end()         { return Bones.end();   }
+	FORCEINLINE TArray<FRigBone>::RangedForConstIteratorType end() const   { return Bones.end();   }
+
+	FRigBone& Add(const FName& InNewName, const FName& InParentName, const FTransform& InInitTransform);
+
+	FRigBone& Add(const FName& InNewName, const FName& InParentName, const FTransform& InInitTransform, const FTransform& InLocalTransform, const FTransform& InGlobalTransform);
+
+	FRigBone Remove(const FName& InNameToRemove);
+
+	void Rename(const FName& InOldName, const FName& InNewName);
+
+	void Reparent(const FName& InName, const FName& InNewParentName);
+
+	FName GetName(int32 InIndex) const;
+
+	FORCEINLINE int32 GetIndex(const FName& InName) const
+	{
+		if(NameToIndexMapping.Num() != Bones.Num())
+		{
+			return GetIndexSlow(InName);
+		}
+
+		const int32* Index = NameToIndexMapping.Find(InName);
+		if (Index)
+		{
+			return *Index;
+		}
+
+		return INDEX_NONE;
+	}
+
+	// list of names of children - this is not cheap, and is supposed to be used only for one time set up
+	int32 GetChildren(const FName& InName, TArray<int32>& OutChildren, bool bRecursively) const;
+
+	int32 GetChildren(const int32 InIndex, TArray<int32>& OutChildren, bool bRecursively) const;
+
+	void SetGlobalTransform(const FName& InName, const FTransform& InTransform, bool bPropagateTransform = true);
+
+	void SetGlobalTransform(int32 InIndex, const FTransform& InTransform, bool bPropagateTransform = true);
+
+	FTransform GetGlobalTransform(const FName& InName) const;
+
+	FTransform GetGlobalTransform(int32 InIndex) const;
+
+	void SetLocalTransform(const FName& InName, const FTransform& InTransform, bool bPropagateTransform = true);
+
+	void SetLocalTransform(int32 InIndex, const FTransform& InTransform, bool bPropagateTransform = true);
+
+	FTransform GetLocalTransform(const FName& InName) const;
+
+	FTransform GetLocalTransform(int32 InIndex) const;
+
+	void SetInitialTransform(const FName& InName, const FTransform& InTransform);
+
+	void SetInitialTransform(int32 InIndex, const FTransform& InTransform);
+
+	FTransform GetInitialTransform(const FName& InName) const;
+
+	FTransform GetInitialTransform(int32 InIndex) const;
+
+	// updates all of the internal caches
+	void Initialize();
+
+	// clears the hierarchy and removes all content
+	void Reset();
+
+	// resets all of the transforms back to the initial transform
+	void ResetTransforms();
+
+private:
+
+	// disable copy constructor
+	FRigBoneHierarchy(const FRigBoneHierarchy& InOther) {}
+
+#if WITH_EDITOR
+	FRigHierarchyContainer* Container;
+	FRigHierarchyElementAdded OnBoneAdded;
+	FRigHierarchyElementRemoved OnBoneRemoved;
+	FRigHierarchyElementRenamed OnBoneRenamed;
+	FRigHierarchyElementReparented OnBoneReparented;
+#endif
+
+	UPROPERTY(EditAnywhere, Category = FRigBoneHierarchy)
+	TArray<FRigBone> Bones;
+
+	// can serialize fine? 
+	UPROPERTY()
+	TMap<FName, int32> NameToIndexMapping;
+
+	int32 GetIndexSlow(const FName& InName) const;
+
+	void RecalculateLocalTransform(FRigBone& InOutBone);
+	void RecalculateGlobalTransform(FRigBone& InOutBone);
+
+	void RefreshMapping();
+	void Sort();
+
+	// list of names of children - this is not cheap, and is supposed to be used only for one time set up
+	int32 GetChildrenRecursive(const int32 InIndex, TArray<int32>& OutChildren, bool bRecursively) const;
+
+	void PropagateTransform(int32 InIndex);
+	
+	friend struct FRigHierarchyContainer;
+};
diff --git a/Engine/Plugins/Experimental/ControlRig/Source/ControlRig/Public/Rigs/RigCurveContainer.h b/Engine/Plugins/Experimental/ControlRig/Source/ControlRig/Public/Rigs/RigCurveContainer.h
new file mode 100644
--- /dev/null
+++ b/Engine/Plugins/Experimental/ControlRig/Source/ControlRig/Public/Rigs/RigCurveContainer.h
@@ -0,0 +1,220 @@
+// Copyright 1998-2019 Epic Games, Inc. All Rights Reserved.
+
+#pragma once
+
+#include "CoreMinimal.h"
+#include "Stats/StatsHierarchical.h"
+#include "RigCurveContainer.generated.h"
+
+class UControlRig;
+
+USTRUCT(BlueprintType)
+struct FRigCurve
+{
+	GENERATED_BODY()
+
+	UPROPERTY(VisibleAnywhere, Category = FRigCurveContainer)
+	FName Name;
+
+	UPROPERTY(VisibleAnywhere, Category = FRigCurveContainer)
+	float Value;
+};
+
+USTRUCT()
+struct CONTROLRIG_API FRigCurveContainer
+{
+	GENERATED_BODY()
+
+private:
+	UPROPERTY(EditAnywhere, Category = FRigCurveContainer)
+	TArray<FRigCurve> Curves;
+
+	// can serialize fine? 
+	UPROPERTY()
+	TMap<FName, int32> NameToIndexMapping;
+
+public:
+
+	const TArray<FRigCurve>& GetCurves() const
+	{
+		return Curves;
+	}
+
+	void AddCurve(const FName& NewCurveName)
+	{
+		DECLARE_SCOPE_HIERARCHICAL_COUNTER_FUNC()
+
+		int32 Found = GetIndex(NewCurveName);
+		if (Found == INDEX_NONE)
+		{
+			FRigCurve NewCurve;
+			NewCurve.Name = NewCurveName;
+			NewCurve.Value = 0.f;
+			Curves.Add(NewCurve);
+			RefreshMapping();
+		}
+	}
+
+	void DeleteCurve(const FName& CurveToDelete)
+	{
+		DECLARE_SCOPE_HIERARCHICAL_COUNTER_FUNC()
+
+		int32 IndexToDelete = GetIndex(CurveToDelete);
+		if (IndexToDelete != INDEX_NONE)
+		{
+			Curves.RemoveAt(IndexToDelete);
+			RefreshMapping();
+		}
+	}
+
+	FName GetName(int32 Index) const
+	{
+		DECLARE_SCOPE_HIERARCHICAL_COUNTER_FUNC()
+
+		if (Curves.IsValidIndex(Index))
+		{
+			return Curves[Index].Name;
+		}
+
+		return NAME_None;
+	}
+
+	int32 GetIndex(const FName& Curve) const
+	{
+		DECLARE_SCOPE_HIERARCHICAL_COUNTER_FUNC()
+
+		const int32* Index = NameToIndexMapping.Find(Curve);
+		if (Index)
+		{
+			return *Index;
+		}
+
+		return INDEX_NONE;
+	}
+
+//#if WITH_EDITOR
+	// @FIXMELINA: figure out how to remove this outside of editor
+	// ignore mapping, run slow search
+	// this is useful in editor while editing
+	// we don't want to build mapping data every time
+	int32 GetIndexSlow(const FName& Curve) const
+	{
+		DECLARE_SCOPE_HIERARCHICAL_COUNTER_FUNC()
+
+		for (int32 CurveId = 0; CurveId < Curves.Num(); ++CurveId)
+		{
+			if (Curves[CurveId].Name == Curve)
+			{
+				return CurveId;
+			}
+		}
+
+		return INDEX_NONE;
+	}
+//#endif // WITH_EDITOR
+	void SetValue(const FName& Curve, const float Value)
+	{
+		DECLARE_SCOPE_HIERARCHICAL_COUNTER_FUNC()
+
+		SetValue(GetIndex(Curve), Value);
+	}
+
+	void SetValue(int32 Index, const float Value)
+	{
+		DECLARE_SCOPE_HIERARCHICAL_COUNTER_FUNC()
+
+		if (Curves.IsValidIndex(Index))
+		{
+			FRigCurve& Curve = Curves[Index];
+			Curve.Value = Value;
+		}
+	}
+
+	float GetValue(const FName& Curve) const
+	{
+		DECLARE_SCOPE_HIERARCHICAL_COUNTER_FUNC()
+
+		return GetValue(GetIndex(Curve));
+	}
+
+	float GetValue(int32 Index) const
+	{
+		DECLARE_SCOPE_HIERARCHICAL_COUNTER_FUNC()
+
+		if (Curves.IsValidIndex(Index))
+		{
+			return Curves[Index].Value;
+		}
+
+		return 0.f;
+	}
+
+	void Rename(const FName& OldName, const FName& NewName)
+	{
+		DECLARE_SCOPE_HIERARCHICAL_COUNTER_FUNC()
+
+		if (OldName != NewName)
+		{
+			const int32 Found = GetIndex(OldName);
+			const int32 NewNameFound = GetIndex(NewName);
+			// if I have new name, and didn't find new name
+			if (Found != INDEX_NONE && NewNameFound == INDEX_NONE)
+			{
+				Curves[Found].Name = NewName;
+				RefreshMapping();
+			}
+		}
+	}
+
+	// updates all of the internal caches
+	void Initialize();
+
+	// clears the CurveContainer and removes all content
+	void Reset();
+
+	// resets all of the transforms back to the initial transform
+	void ResetValues();
+
+	int32 GetNum() const
+	{
+		return Curves.Num();
+	}
+
+private:
+	void RefreshMapping();
+};
+
+
+// @todo: do we need this container now?
+USTRUCT()
+struct FRigCurveContainerRef
+{
+	GENERATED_BODY()
+
+	FRigCurveContainerRef(FRigCurveContainer* InContainer = nullptr)
+		: Container(InContainer)
+	{
+
+	}
+
+	FRigCurveContainer* Get() 
+	{
+		return GetInternal();
+	}
+
+	const FRigCurveContainer* Get() const
+	{
+		return (const FRigCurveContainer*)GetInternal();
+	}
+
+private:
+	struct FRigCurveContainer* Container;
+
+	FRigCurveContainer* GetInternal() const
+	{
+		return Container;
+	}
+
+	friend class UControlRig;
+	friend class FControlRigUnitTestBase;
+};
\ No newline at end of file
diff --git a/Engine/Plugins/Experimental/ControlRig/Source/ControlRig/Public/Rigs/RigHierarchyContainer.h b/Engine/Plugins/Experimental/ControlRig/Source/ControlRig/Public/Rigs/RigHierarchyContainer.h
new file mode 100644
--- /dev/null
+++ b/Engine/Plugins/Experimental/ControlRig/Source/ControlRig/Public/Rigs/RigHierarchyContainer.h
@@ -0,0 +1,68 @@
+// Copyright 1998-2019 Epic Games, Inc. All Rights Reserved.
+
+#pragma once
+
+#include "CoreMinimal.h"
+#include "RigHierarchyDefines.h"
+#include "RigBoneHierarchy.h"
+#include "RigHierarchyContainer.generated.h"
+
+class UControlRig;
+struct FRigHierarchyContainer;
+
+USTRUCT()
+struct CONTROLRIG_API FRigHierarchyContainer
+{
+public:
+	GENERATED_BODY()
+
+	FRigHierarchyContainer();
+	FRigHierarchyContainer& operator= (const FRigHierarchyContainer &InOther);
+
+	UPROPERTY()
+	FRigBoneHierarchy BoneHierarchy;
+
+	void Initialize();
+	void Reset();
+	void ResetTransforms();
+
+#if WITH_EDITOR
+	FRigHierarchyElementChanged OnElementChanged;
+	FRigHierarchyElementAdded OnElementAdded;
+	FRigHierarchyElementRemoved OnElementRemoved;
+	FRigHierarchyElementRenamed OnElementRenamed;
+	FRigHierarchyElementReparented OnElementReparented;
+#endif
+
+protected:
+
+	// disable copy constructor
+	FRigHierarchyContainer(const FRigHierarchyContainer& InContainer) {}
+
+#if WITH_EDITOR
+	void HandleOnElementAdded(FRigHierarchyContainer* InContainer, ERigHierarchyElementType InElementType, const FName& InName);
+	void HandleOnElementRemoved(FRigHierarchyContainer* InContainer, ERigHierarchyElementType InElementType, const FName& InName);
+	void HandleOnElementRenamed(FRigHierarchyContainer* InContainer, ERigHierarchyElementType InElementType, const FName& InOldName, const FName& InNewName);
+	void HandleOnElementReparented(FRigHierarchyContainer* InContainer, ERigHierarchyElementType InElementType, const FName& InName, const FName& InOldParentName, const FName& InNewParentName);
+#endif
+};
+
+USTRUCT()
+struct CONTROLRIG_API FRigHierarchyRef
+{
+	GENERATED_BODY()
+
+	FRigHierarchyRef();
+
+	FRigBoneHierarchy* GetBones();
+	const FRigBoneHierarchy* GetBones() const;
+
+private:
+
+	struct FRigHierarchyContainer* Container;
+
+	FRigBoneHierarchy* GetBonesInternal() const;
+
+	friend class UControlRig;
+	friend class FControlRigUnitTestBase;
+};
\ No newline at end of file
diff --git a/Engine/Plugins/Experimental/ControlRig/Source/ControlRig/Public/Rigs/RigHierarchyDefines.h b/Engine/Plugins/Experimental/ControlRig/Source/ControlRig/Public/Rigs/RigHierarchyDefines.h
new file mode 100644
--- /dev/null
+++ b/Engine/Plugins/Experimental/ControlRig/Source/ControlRig/Public/Rigs/RigHierarchyDefines.h
@@ -0,0 +1,24 @@
+// Copyright 1998-2019 Epic Games, Inc. All Rights Reserved.
+
+#pragma once
+
+#include "CoreMinimal.h"
+
+struct FRigHierarchyContainer;
+
+UENUM()
+enum class ERigHierarchyElementType : uint8
+{
+	Bone,
+	Space,
+	Control,
+
+	/** MAX - invalid */
+	Max UMETA(Hidden),
+};
+
+DECLARE_MULTICAST_DELEGATE_ThreeParams(FRigHierarchyElementChanged, FRigHierarchyContainer*, ERigHierarchyElementType, const FName&);
+DECLARE_MULTICAST_DELEGATE_ThreeParams(FRigHierarchyElementAdded, FRigHierarchyContainer*, ERigHierarchyElementType, const FName&);
+DECLARE_MULTICAST_DELEGATE_ThreeParams(FRigHierarchyElementRemoved, FRigHierarchyContainer*, ERigHierarchyElementType, const FName&);
+DECLARE_MULTICAST_DELEGATE_FourParams(FRigHierarchyElementRenamed, FRigHierarchyContainer*, ERigHierarchyElementType, const FName&, const FName&);
+DECLARE_MULTICAST_DELEGATE_FiveParams(FRigHierarchyElementReparented, FRigHierarchyContainer*, ERigHierarchyElementType, const FName&, const FName&, const FName&);
diff --git a/Engine/Plugins/Experimental/ControlRig/Source/ControlRig/Public/Units/Deprecated/RigUnit_TwoBoneIKFK.h b/Engine/Plugins/Experimental/ControlRig/Source/ControlRig/Public/Units/Deprecated/RigUnit_TwoBoneIKFK.h
--- a/Engine/Plugins/Experimental/ControlRig/Source/ControlRig/Public/Units/Deprecated/RigUnit_TwoBoneIKFK.h
+++ b/Engine/Plugins/Experimental/ControlRig/Source/ControlRig/Public/Units/Deprecated/RigUnit_TwoBoneIKFK.h
@@ -3,7 +3,7 @@
 #pragma once
 
 #include "Units/RigUnit.h"
-#include "Rigs/Hierarchy.h"
+#include "Rigs/RigHierarchyContainer.h"
 #include "ControlRigDefines.h"
 #include "RigUnit_TwoBoneIKFK.generated.h"
 
diff --git a/Engine/Plugins/Experimental/ControlRig/Source/ControlRig/Public/Units/RigUnitContext.h b/Engine/Plugins/Experimental/ControlRig/Source/ControlRig/Public/Units/RigUnitContext.h
--- a/Engine/Plugins/Experimental/ControlRig/Source/ControlRig/Public/Units/RigUnitContext.h
+++ b/Engine/Plugins/Experimental/ControlRig/Source/ControlRig/Public/Units/RigUnitContext.h
@@ -3,8 +3,8 @@
 #pragma once
 
 #include "CoreMinimal.h"
-#include "Rigs/Hierarchy.h"
-#include "CurveContainer.h"
+#include "Rigs/RigHierarchyContainer.h"
+#include "Rigs/RigCurveContainer.h"
 #include "ControlRigLog.h"
 #include "AnimationDataSource.h"
 #include "Drawing/ControlRigDrawInterface.h"
diff --git a/Engine/Plugins/Experimental/ControlRig/Source/ControlRig/Public/Units/RigUnitTest.h b/Engine/Plugins/Experimental/ControlRig/Source/ControlRig/Public/Units/RigUnitTest.h
--- a/Engine/Plugins/Experimental/ControlRig/Source/ControlRig/Public/Units/RigUnitTest.h
+++ b/Engine/Plugins/Experimental/ControlRig/Source/ControlRig/Public/Units/RigUnitTest.h
@@ -5,22 +5,21 @@
 #include "Misc/AutomationTest.h"
 #include "RigUnit.h"
 #include "RigUnitContext.h"
-#include "Rigs/Hierarchy.h"
-#include "CurveContainer.h"
+#include "Rigs/RigHierarchyContainer.h"
+#include "Rigs/RigCurveContainer.h"
 
 class FControlRigUnitTestBase : public FAutomationTestBase
 {
 public:
 	FControlRigUnitTestBase(const FString& InName, bool bIsComplex)
 		: FAutomationTestBase(InName, bIsComplex)
-		, HierarchyContainer(FRigHierarchyContainer())
+		, HierarchyContainer()
 		, HierarchyRef(FRigHierarchyRef())
-		, Hierarchy(HierarchyContainer.BaseHierarchy)
+		, Hierarchy(HierarchyContainer.BoneHierarchy)
 		, CurveContainer(FRigCurveContainer())
 		, CurveContainerRef(&CurveContainer)
 	{
 		HierarchyRef.Container = &HierarchyContainer;
-		HierarchyRef.bUseBaseHierarchy = true;
 		Context.HierarchyReference = HierarchyRef;
 		ExecuteContext.HierarchyReference = HierarchyRef;
 		ExecuteContext.CurveReference = CurveContainerRef;
@@ -28,7 +27,7 @@ public:
 
 	FRigHierarchyContainer HierarchyContainer;
 	FRigHierarchyRef HierarchyRef;
-	FRigHierarchy& Hierarchy;
+	FRigBoneHierarchy& Hierarchy;
 	FRigCurveContainer CurveContainer;
 	FRigCurveContainerRef CurveContainerRef;
 	FControlRigExecuteContext ExecuteContext;
diff --git a/Engine/Plugins/Experimental/ControlRig/Source/ControlRigDeveloper/Private/ControlRigBlueprint.cpp b/Engine/Plugins/Experimental/ControlRig/Source/ControlRigDeveloper/Private/ControlRigBlueprint.cpp
--- a/Engine/Plugins/Experimental/ControlRig/Source/ControlRigDeveloper/Private/ControlRigBlueprint.cpp
+++ b/Engine/Plugins/Experimental/ControlRig/Source/ControlRigDeveloper/Private/ControlRigBlueprint.cpp
@@ -598,5 +598,16 @@ bool UControlRigBlueprint::PerformArrayOperation(const FString& InPropertyPath,
 	return false;
 }
 
+void UControlRigBlueprint::CleanupBoneHierarchyDeprecated()
+{
+	if (Hierarchy_DEPRECATED.Num() == 0)
+	{
+		return;
+	}
+
+	HierarchyContainer.BoneHierarchy = Hierarchy_DEPRECATED;
+	Hierarchy_DEPRECATED.Reset();
+}
+
 #undef LOCTEXT_NAMESPACE
 
diff --git a/Engine/Plugins/Experimental/ControlRig/Source/ControlRigDeveloper/Private/ControlRigBlueprintCompiler.cpp b/Engine/Plugins/Experimental/ControlRig/Source/ControlRigDeveloper/Private/ControlRigBlueprintCompiler.cpp
--- a/Engine/Plugins/Experimental/ControlRig/Source/ControlRigDeveloper/Private/ControlRigBlueprintCompiler.cpp
+++ b/Engine/Plugins/Experimental/ControlRig/Source/ControlRigDeveloper/Private/ControlRigBlueprintCompiler.cpp
@@ -484,8 +484,10 @@ void FControlRigBlueprintCompilerContext::CopyTermDefaultsToDefaultObject(UObjec
 	UControlRigBlueprint* ControlRigBlueprint = Cast<UControlRigBlueprint>(Blueprint);
 	if (ControlRigBlueprint)
 	{
+		ControlRigBlueprint->CleanupBoneHierarchyDeprecated();
+
 		UControlRig* ControlRig = CastChecked<UControlRig>(DefaultObject);
-		ControlRig->Hierarchy.BaseHierarchy = ControlRigBlueprint->Hierarchy;
+		ControlRig->Hierarchy = ControlRigBlueprint->HierarchyContainer;
 		ControlRig->CurveContainer = ControlRigBlueprint->CurveContainer;
 		// copy available rig units info, so that control rig can do things with it
 		ControlRig->AllowSourceAccessProperties = ControlRigBlueprint->AllowSourceAccessProperties;
diff --git a/Engine/Plugins/Experimental/ControlRig/Source/ControlRigDeveloper/Private/Graph/ControlRigGraph.cpp b/Engine/Plugins/Experimental/ControlRig/Source/ControlRigDeveloper/Private/Graph/ControlRigGraph.cpp
--- a/Engine/Plugins/Experimental/ControlRig/Source/ControlRigDeveloper/Private/Graph/ControlRigGraph.cpp
+++ b/Engine/Plugins/Experimental/ControlRig/Source/ControlRigDeveloper/Private/Graph/ControlRigGraph.cpp
@@ -272,13 +272,12 @@ void UControlRigGraph::OnBlueprintCompiledPostLoad(UBlueprint* InCompiledBluepri
 	}
 }
 
-void UControlRigGraph::CacheBoneNameList(const FRigHierarchy& Hierarchy)
+void UControlRigGraph::CacheBoneNameList(const FRigBoneHierarchy& Hierarchy)
 {
 	DECLARE_SCOPE_HIERARCHICAL_COUNTER_FUNC()
 
 	TArray<FString> Names;
-	const TArray<FRigBone>& Bones = Hierarchy.GetBones();
-	for (const FRigBone& Bone : Bones )
+	for (const FRigBone& Bone : Hierarchy)
 	{
 		Names.Add(Bone.Name.ToString());
 	}
diff --git a/Engine/Plugins/Experimental/ControlRig/Source/ControlRigDeveloper/Public/ControlRigBlueprint.h b/Engine/Plugins/Experimental/ControlRig/Source/ControlRigDeveloper/Public/ControlRigBlueprint.h
--- a/Engine/Plugins/Experimental/ControlRig/Source/ControlRigDeveloper/Public/ControlRigBlueprint.h
+++ b/Engine/Plugins/Experimental/ControlRig/Source/ControlRigDeveloper/Public/ControlRigBlueprint.h
@@ -8,8 +8,8 @@
 #include "Engine/Blueprint.h"
 #include "Misc/Crc.h"
 #include "ControlRigDefines.h"
-#include "Rigs/Hierarchy.h"
-#include "CurveContainer.h"
+#include "Rigs/RigHierarchyContainer.h"
+#include "Rigs/RigCurveContainer.h"
 #include "Interfaces/Interface_PreviewMeshProvider.h"
 #include "ControlRigController.h"
 #include "ControlRigBlueprint.generated.h"
@@ -163,7 +163,10 @@ private:
 	TMap<FName, FString> AllowSourceAccessProperties;
 
 	UPROPERTY()
-	FRigHierarchy Hierarchy;
+	FRigHierarchyContainer HierarchyContainer;
+
+	UPROPERTY()
+	FRigBoneHierarchy Hierarchy_DEPRECATED;
 
 	UPROPERTY()
 	FRigCurveContainer CurveContainer;
@@ -183,9 +186,10 @@ private:
 	void HandleModelModified(const UControlRigModel* InModel, EControlRigModelNotifType InType, const void* InPayload);
 	bool UpdateParametersOnControlRig(UControlRig* InRig = nullptr);
 	bool PerformArrayOperation(const FString& InPropertyPath, TFunctionRef<bool(FScriptArrayHelper&, int32)> InOperation, bool bCallModify, bool bPropagateToInstances);
+	void CleanupBoneHierarchyDeprecated();
 
 	friend class FControlRigBlueprintCompilerContext;
-	friend class SRigHierarchy;
+	friend class SRigBoneHierarchy;
 	friend class SRigCurveContainer;
 	friend class FControlRigEditor;
 	friend class UEngineTestControlRig;
diff --git a/Engine/Plugins/Experimental/ControlRig/Source/ControlRigDeveloper/Public/Graph/ControlRigGraph.h b/Engine/Plugins/Experimental/ControlRig/Source/ControlRigDeveloper/Public/Graph/ControlRigGraph.h
--- a/Engine/Plugins/Experimental/ControlRig/Source/ControlRigDeveloper/Public/Graph/ControlRigGraph.h
+++ b/Engine/Plugins/Experimental/ControlRig/Source/ControlRigDeveloper/Public/Graph/ControlRigGraph.h
@@ -4,7 +4,7 @@
 
 #include "EdGraph/EdGraph.h"
 #include "Graph/ControlRigGraphNode.h"
-#include "Rigs/Hierarchy.h"
+#include "Rigs/RigHierarchyContainer.h"
 #include "ControlRigModel.h"
 #include "ControlRigGraph.generated.h"
 
@@ -36,7 +36,7 @@ public:
 	virtual void PostLoad() override;
 	void OnBlueprintCompiledPostLoad(UBlueprint*);
 	FDelegateHandle BlueprintOnCompiledHandle;
-	void CacheBoneNameList(const FRigHierarchy& Hierarchy);
+	void CacheBoneNameList(const FRigBoneHierarchy& Hierarchy);
 	const TArray<TSharedPtr<FString>>& GetBoneNameList() const;
 	void CacheCurveNameList(const FRigCurveContainer& Container);
 	const TArray<TSharedPtr<FString>>& GetCurveNameList() const;
diff --git a/Engine/Plugins/Experimental/ControlRig/Source/ControlRigEditor/Private/EditMode/ControlRigEditMode.cpp b/Engine/Plugins/Experimental/ControlRig/Source/ControlRigEditor/Private/EditMode/ControlRigEditMode.cpp
--- a/Engine/Plugins/Experimental/ControlRig/Source/ControlRigEditor/Private/EditMode/ControlRigEditMode.cpp
+++ b/Engine/Plugins/Experimental/ControlRig/Source/ControlRigEditor/Private/EditMode/ControlRigEditMode.cpp
@@ -380,11 +380,10 @@ void FControlRigEditMode::Render(const FSceneView* View, FViewport* Viewport, FP
 				}
 
 				// each base hierarchy Bone
-				const FRigHierarchy& BaseHierarchy = ControlRig->GetBaseHierarchy();
-				const TArray<FRigBone>& Bones = BaseHierarchy.GetBones();
-				for (int32 BoneIndex = 0; BoneIndex < Bones.Num(); ++BoneIndex)
+				const FRigBoneHierarchy& BaseHierarchy = ControlRig->GetBoneHierarchy();
+				for (int32 BoneIndex = 0; BoneIndex < BaseHierarchy.Num(); ++BoneIndex)
 				{
-					const FRigBone& CurrentBone = Bones[BoneIndex];
+					const FRigBone& CurrentBone = BaseHierarchy[BoneIndex];
 					const FTransform Transform = BaseHierarchy.GetGlobalTransform(BoneIndex);
 
 					if (CurrentBone.ParentIndex != INDEX_NONE)
diff --git a/Engine/Plugins/Experimental/ControlRig/Source/ControlRigEditor/Private/Editor/ControlRigEditor.cpp b/Engine/Plugins/Experimental/ControlRig/Source/ControlRigEditor/Private/Editor/ControlRigEditor.cpp
--- a/Engine/Plugins/Experimental/ControlRig/Source/ControlRigEditor/Private/Editor/ControlRigEditor.cpp
+++ b/Engine/Plugins/Experimental/ControlRig/Source/ControlRigEditor/Private/Editor/ControlRigEditor.cpp
@@ -38,7 +38,7 @@
 #include "Widgets/Input/SButton.h"
 #include "Widgets/Notifications/SNotificationList.h"
 #include "Framework/Notifications/NotificationManager.h"
-#include "SRigHierarchy.h"
+#include "SRigBoneHierarchy.h"
 #include "Framework/Application/MenuStack.h"
 #include "Framework/Application/SlateApplication.h"
 #include "ControlRig/Private/Units/Hierarchy/RigUnit_BoneName.h"
@@ -86,6 +86,15 @@ FControlRigEditor::FControlRigEditor()
 
 FControlRigEditor::~FControlRigEditor()
 {
+	UControlRigBlueprint* RigBlueprint = GetControlRigBlueprint();
+	if (RigBlueprint)
+	{
+		RigBlueprint->HierarchyContainer.OnElementChanged.RemoveAll(this);
+		RigBlueprint->HierarchyContainer.OnElementAdded.RemoveAll(this);
+		RigBlueprint->HierarchyContainer.OnElementRemoved.RemoveAll(this);
+		RigBlueprint->HierarchyContainer.OnElementRenamed.RemoveAll(this);
+		RigBlueprint->HierarchyContainer.OnElementReparented.RemoveAll(this);
+	}
 }
 
 UControlRigBlueprint* FControlRigEditor::GetControlRigBlueprint() const
@@ -176,6 +185,10 @@ void FControlRigEditor::InitControlRigEditor(const EToolkitMode::Type Mode, cons
 	}
 
 	InControlRigBlueprint->OnModified().AddSP(this, &FControlRigEditor::HandleModelModified);
+	InControlRigBlueprint->HierarchyContainer.OnElementAdded.AddSP(this, &FControlRigEditor::OnHierarchyElementAdded);
+	InControlRigBlueprint->HierarchyContainer.OnElementRemoved.AddSP(this, &FControlRigEditor::OnHierarchyElementRemoved);
+	InControlRigBlueprint->HierarchyContainer.OnElementRenamed.AddSP(this, &FControlRigEditor::OnHierarchyElementRenamed);
+	InControlRigBlueprint->HierarchyContainer.OnElementReparented.AddSP(this, &FControlRigEditor::OnHierarchyElementReparented);
 
 	BindCommands();
 
@@ -229,7 +242,20 @@ void FControlRigEditor::InitControlRigEditor(const EToolkitMode::Type Mode, cons
 			}
 			else
 			{
+				bool bWasDirty = false;
+				UPackage* Package = InControlRigBlueprint->GetOutermost();
+				if (Package)
+				{
+					bWasDirty = Package->IsDirty();
+				}
 				InControlRigBlueprint->RebuildGraphFromModel();
+				if (Package)
+				{
+					if (Package->IsDirty() != bWasDirty)
+					{
+						Package->SetDirtyFlag(bWasDirty);
+					}
+				}
 			}
 		}
 	}
@@ -1344,7 +1370,7 @@ void FControlRigEditor::CacheBoneNameList()
 				continue;
 			}
 
-			RigGraph->CacheBoneNameList(ControlRig->GetBaseHierarchy());
+			RigGraph->CacheBoneNameList(ControlRig->GetBoneHierarchy());
 		}
 	}
 }
@@ -1439,7 +1465,7 @@ void FControlRigEditor::SelectBone(const FName& InBone)
 	{
 		EditorSkelComp->BonesOfInterest.Reset();
 
-		int32 Index = ControlRig->Hierarchy.BaseHierarchy.GetIndex(InBone);
+		int32 Index = ControlRig->Hierarchy.BoneHierarchy.GetIndex(InBone);
 		if (Index != INDEX_NONE)
 		{
 			EditorSkelComp->BonesOfInterest.Add(Index);
@@ -1454,10 +1480,10 @@ FTransform FControlRigEditor::GetBoneTransform(const FName& InBone, bool bLocal)
 	// @todo: think about transform mode
 	if (bLocal)
 	{
-		return ControlRig->Hierarchy.BaseHierarchy.GetLocalTransform(InBone);
+		return ControlRig->Hierarchy.BoneHierarchy.GetLocalTransform(InBone);
 	}
 
-	return ControlRig->Hierarchy.BaseHierarchy.GetGlobalTransform(InBone);
+	return ControlRig->Hierarchy.BoneHierarchy.GetGlobalTransform(InBone);
 }
 
 void FControlRigEditor::SetBoneTransform(const FName& InBone, const FTransform& InTransform)
@@ -1478,10 +1504,10 @@ void FControlRigEditor::SetBoneTransform(const FName& InBone, const FTransform&
 	// get local transform of current
 	// apply init based on parent init * current local 
 
-	ControlRig->Hierarchy.BaseHierarchy.SetInitialTransform(InBone, InTransform);
-	ControlRig->Hierarchy.BaseHierarchy.SetGlobalTransform(InBone, InTransform);
+	ControlRig->Hierarchy.BoneHierarchy.SetInitialTransform(InBone, InTransform);
+	ControlRig->Hierarchy.BoneHierarchy.SetGlobalTransform(InBone, InTransform);
 
-	ControlRigBP->Hierarchy.SetInitialTransform(InBone, InTransform);
+	ControlRigBP->HierarchyContainer.BoneHierarchy.SetInitialTransform(InBone, InTransform);
 	
 	UControlRigSkeletalMeshComponent* EditorSkelComp = Cast<UControlRigSkeletalMeshComponent>(GetPersonaToolkit()->GetPreviewScene()->GetPreviewMeshComponent());
 	if (EditorSkelComp)
@@ -1516,20 +1542,20 @@ void FControlRigEditor::OnFinishedChangingProperties(const FPropertyChangedEvent
 		{
 			if (SelectedBone != NAME_None)
 			{
-				const int32 BoneIndex = ControlRig->Hierarchy.BaseHierarchy.GetIndex(SelectedBone);
+				const int32 BoneIndex = ControlRig->Hierarchy.BoneHierarchy.GetIndex(SelectedBone);
 				if (BoneIndex != INDEX_NONE)
 				{
-					FTransform InitialTransform = ControlRig->Hierarchy.BaseHierarchy.GetInitialTransform(BoneIndex);
+					FTransform InitialTransform = ControlRig->Hierarchy.BoneHierarchy.GetInitialTransform(BoneIndex);
 					// update CDO  @todo - re-think about how we wrap around this nicer
 					// copy currently selected Bone to base hierarchy			
-					ControlRigBP->Hierarchy.SetInitialTransform(BoneIndex, InitialTransform);
+					ControlRigBP->HierarchyContainer.BoneHierarchy.SetInitialTransform(BoneIndex, InitialTransform);
 				}
 			}
 		}
 	}
 }
 
-void FControlRigEditor::OnHierarchyChanged()
+void FControlRigEditor::OnBoneHierarchyChanged()
 {
 	DECLARE_SCOPE_HIERARCHICAL_COUNTER_FUNC()
 
@@ -1545,64 +1571,82 @@ void FControlRigEditor::OnHierarchyChanged()
 		UpdateControlRig();
 	}
 	CacheBoneNameList();
+}
 
-	// notification
-	FNotificationInfo Info(LOCTEXT("HierarchyChangeHelpMessage", "Hierarchy has been successfully modified. If you want to move the Bone, compile and turn off execution mode."));
-	Info.bFireAndForget = true;
-	Info.FadeOutDuration = 10.0f;
-	Info.ExpireDuration = 0.0f;
+void FControlRigEditor::OnHierarchyElementAdded(FRigHierarchyContainer* Container, ERigHierarchyElementType ElementType, const FName& InName)
+{
+	if (ElementType == ERigHierarchyElementType::Bone)
+	{
+		OnBoneHierarchyChanged();
+	}
+}
 
-	TSharedPtr<SNotificationItem> NotificationPtr = FSlateNotificationManager::Get().AddNotification(Info);
-	NotificationPtr->SetCompletionState(SNotificationItem::CS_Success);
+void FControlRigEditor::OnHierarchyElementRemoved(FRigHierarchyContainer* Container, ERigHierarchyElementType ElementType, const FName& InName)
+{
+	if (ElementType == ERigHierarchyElementType::Bone)
+	{
+		OnBoneHierarchyChanged();
+	}
 }
 
-void FControlRigEditor::OnBoneRenamed(const FName& OldName, const FName& NewName)
+void FControlRigEditor::OnHierarchyElementRenamed(FRigHierarchyContainer* Container, ERigHierarchyElementType ElementType, const FName& InOldName, const FName& InNewName)
 {
 	DECLARE_SCOPE_HIERARCHICAL_COUNTER_FUNC()
 
-	UControlRigBlueprint* Blueprint = GetControlRigBlueprint();
-	for (UEdGraph* Graph : Blueprint->UbergraphPages)
+	if (ElementType == ERigHierarchyElementType::Bone)
 	{
-		UControlRigGraph* RigGraph = Cast<UControlRigGraph>(Graph);
-		if (RigGraph == nullptr)
-		{
-			continue;
-		}
-
-		for (UEdGraphNode* Node : RigGraph->Nodes)
+		UControlRigBlueprint* Blueprint = GetControlRigBlueprint();
+		for (UEdGraph* Graph : Blueprint->UbergraphPages)
 		{
-			UControlRigGraphNode* RigNode = Cast<UControlRigGraphNode>(Node);
-			if (RigNode == nullptr)
+			UControlRigGraph* RigGraph = Cast<UControlRigGraph>(Graph);
+			if (RigGraph == nullptr)
 			{
 				continue;
 			}
 
-			UStructProperty* UnitProperty = RigNode->GetUnitProperty();
-			UStruct* UnitStruct = RigNode->GetUnitScriptStruct();
-			if (UnitProperty && UnitStruct)
+			for (UEdGraphNode* Node : RigGraph->Nodes)
 			{
-				for (TFieldIterator<UNameProperty> It(UnitStruct); It; ++It)
+				UControlRigGraphNode* RigNode = Cast<UControlRigGraphNode>(Node);
+				if (RigNode == nullptr)
+				{
+					continue;
+				}
+
+				UStructProperty* UnitProperty = RigNode->GetUnitProperty();
+				UStruct* UnitStruct = RigNode->GetUnitScriptStruct();
+				if (UnitProperty && UnitStruct)
 				{
-					if (It->HasMetaData(UControlRig::BoneNameMetaName))
+					for (TFieldIterator<UNameProperty> It(UnitStruct); It; ++It)
 					{
-						FString PinName = FString::Printf(TEXT("%s.%s"), *UnitProperty->GetName(), *It->GetName());
-						UEdGraphPin* Pin = Node->FindPin(PinName, EEdGraphPinDirection::EGPD_Input);
-						if (Pin)
+						if (It->HasMetaData(UControlRig::BoneNameMetaName))
 						{
-							FName CurrentBone = FName(*Pin->GetDefaultAsString());
-							if (CurrentBone == OldName)
+							FString PinName = FString::Printf(TEXT("%s.%s"), *UnitProperty->GetName(), *It->GetName());
+							UEdGraphPin* Pin = Node->FindPin(PinName, EEdGraphPinDirection::EGPD_Input);
+							if (Pin)
 							{
-								const FScopedTransaction Transaction(NSLOCTEXT("ControlRigEditor", "ChangeBoneNamePinValue", "Change Bone Name Pin Value"));
-								Pin->Modify();
-								Pin->GetSchema()->TrySetDefaultValue(*Pin, NewName.ToString());
+								FName CurrentBone = FName(*Pin->GetDefaultAsString());
+								if (CurrentBone == InOldName)
+								{
+									const FScopedTransaction Transaction(NSLOCTEXT("ControlRigEditor", "ChangeBoneNamePinValue", "Change Bone Name Pin Value"));
+									Pin->Modify();
+									Pin->GetSchema()->TrySetDefaultValue(*Pin, InNewName.ToString());
+								}
 							}
 						}
 					}
 				}
 			}
+
+			CacheBoneNameList();
 		}
+	}
+}
 
-		CacheBoneNameList();
+void FControlRigEditor::OnHierarchyElementReparented(FRigHierarchyContainer* Container, ERigHierarchyElementType ElementType, const FName& InName, const FName& InOldParentName, const FName& InNewParentName)
+{
+	if (ElementType == ERigHierarchyElementType::Bone)
+	{
+		OnBoneHierarchyChanged();
 	}
 }
 
@@ -1687,9 +1731,9 @@ void FControlRigEditor::OnGraphNodeDropToPerform(TSharedPtr<FGraphNodeDragDropOp
 {
 	DECLARE_SCOPE_HIERARCHICAL_COUNTER_FUNC()
 
-	if (DragDropOp->IsOfType<FRigHierarchyDragDropOp>())
+	if (DragDropOp->IsOfType<FRigBoneHierarchyDragDropOp>())
 	{
-		TSharedPtr<FRigHierarchyDragDropOp> RigHierarchyOp = StaticCastSharedPtr<FRigHierarchyDragDropOp>(DragDropOp);
+		TSharedPtr<FRigBoneHierarchyDragDropOp> RigHierarchyOp = StaticCastSharedPtr<FRigBoneHierarchyDragDropOp>(DragDropOp);
 		TArray<FName> BoneNames = RigHierarchyOp->GetBoneNames();
 		if (BoneNames.Num() > 0 && FocusedGraphEdPtr.IsValid())
 		{
@@ -2060,10 +2104,10 @@ void FControlRigEditor::DumpUnitTestCode()
 		// dump the hierarchy
 		if (ControlRig)
 		{
-			const FRigHierarchy& Hierarchy = ControlRig->GetBaseHierarchy();
+			const FRigBoneHierarchy& Hierarchy = ControlRig->GetBoneHierarchy();
 			if (Hierarchy.Bones.Num() > 0)
 			{
-				Code.Add(TEXT("FRigHierarchy& Hierarchy = Rig->GetBaseHierarchy();"));
+				Code.Add(TEXT("FRigBoneHierarchy& Hierarchy = Rig->GetBoneHierarchy();"));
 			}
 			for (const FRigBone& Bone : Hierarchy.Bones)
 			{
diff --git a/Engine/Plugins/Experimental/ControlRig/Source/ControlRigEditor/Private/Editor/ControlRigEditor.h b/Engine/Plugins/Experimental/ControlRig/Source/ControlRigEditor/Private/Editor/ControlRigEditor.h
--- a/Engine/Plugins/Experimental/ControlRig/Source/ControlRigEditor/Private/Editor/ControlRigEditor.h
+++ b/Engine/Plugins/Experimental/ControlRig/Source/ControlRigEditor/Private/Editor/ControlRigEditor.h
@@ -107,15 +107,21 @@ public:
 	void SelectBone(const FName& InBone);
 	// this changes everytime you compile, so don't cache it expecting it will last. 
 	UControlRig* GetInstanceRig() const { return ControlRig;  }
-	// restart animation 
-	void OnHierarchyChanged();
-	void OnBoneRenamed(const FName& OldName, const FName& NewName);
+
 	void OnCurveContainerChanged();
 	void OnCurveRenamed(const FName& OldName, const FName& NewName);
 
+	void OnHierarchyElementAdded(FRigHierarchyContainer* Container, ERigHierarchyElementType ElementType, const FName& InName);
+	void OnHierarchyElementRemoved(FRigHierarchyContainer* Container, ERigHierarchyElementType ElementType, const FName& InName);
+	void OnHierarchyElementRenamed(FRigHierarchyContainer* Container, ERigHierarchyElementType ElementType, const FName& InOldName, const FName& InNewName);
+	void OnHierarchyElementReparented(FRigHierarchyContainer* Container, ERigHierarchyElementType ElementType, const FName& InName, const FName& InOldParentName, const FName& InNewParentName);
+
 	void OnGraphNodeDropToPerform(TSharedPtr<FGraphNodeDragDropOp> DragDropOp, UEdGraph* Graph, const FVector2D& NodePosition, const FVector2D& ScreenPosition);
 
 protected:
+
+	void OnBoneHierarchyChanged();
+
 	// FBlueprintEditor Interface
 	virtual void CreateDefaultCommands() override;
 	virtual void OnCreateGraphEditorCommands(TSharedPtr<FUICommandList> GraphEditorCommandsList);
@@ -261,5 +267,5 @@ protected:
 
 	friend class FControlRigEditorMode;
 	friend class SControlRigStackView;
-	friend class SRigHierarchy;
+	friend class SRigBoneHierarchy;
 };
\ No newline at end of file
diff --git a/Engine/Plugins/Experimental/ControlRig/Source/ControlRigEditor/Private/Editor/ControlRigEditorMode.cpp b/Engine/Plugins/Experimental/ControlRig/Source/ControlRigEditor/Private/Editor/ControlRigEditorMode.cpp
--- a/Engine/Plugins/Experimental/ControlRig/Source/ControlRigEditor/Private/Editor/ControlRigEditorMode.cpp
+++ b/Engine/Plugins/Experimental/ControlRig/Source/ControlRigEditor/Private/Editor/ControlRigEditorMode.cpp
@@ -15,7 +15,7 @@ FControlRigEditorMode::FControlRigEditorMode(const TSharedRef<FControlRigEditor>
 {
 	ControlRigBlueprintPtr = CastChecked<UControlRigBlueprint>(InControlRigEditor->GetBlueprintObj());
 
-	TabFactories.RegisterFactory(MakeShared<FRigHierarchyTabSummoner>(InControlRigEditor));
+	TabFactories.RegisterFactory(MakeShared<FRigBoneHierarchyTabSummoner>(InControlRigEditor));
 	TabFactories.RegisterFactory(MakeShared<FRigStackTabSummoner>(InControlRigEditor));
 	TabFactories.RegisterFactory(MakeShared<FRigCurveContainerTabSummoner>(InControlRigEditor));
 
@@ -70,7 +70,7 @@ FControlRigEditorMode::FControlRigEditorMode(const TSharedRef<FControlRigEditor>
 						//	Left bottom - rig/hierarchy
 						FTabManager::NewStack()
 						->SetSizeCoefficient(0.5f)
-						->AddTab(FRigHierarchyTabSummoner::TabID, ETabState::OpenedTab)
+						->AddTab(FRigBoneHierarchyTabSummoner::TabID, ETabState::OpenedTab)
 						->AddTab(FRigStackTabSummoner::TabID, ETabState::OpenedTab)
 					)
 				)
diff --git a/Engine/Plugins/Experimental/ControlRig/Source/ControlRigEditor/Private/Editor/ControlRigSkeletalMeshComponent.cpp b/Engine/Plugins/Experimental/ControlRig/Source/ControlRigEditor/Private/Editor/ControlRigSkeletalMeshComponent.cpp
--- a/Engine/Plugins/Experimental/ControlRig/Source/ControlRigEditor/Private/Editor/ControlRigSkeletalMeshComponent.cpp
+++ b/Engine/Plugins/Experimental/ControlRig/Source/ControlRigEditor/Private/Editor/ControlRigSkeletalMeshComponent.cpp
@@ -58,7 +58,7 @@ void UControlRigSkeletalMeshComponent::RebuildDebugDrawSkeleton()
 	{
 		UControlRig* ControlRig = ControlRigInstance->CachedControlRig.Get();
 		// just copy it because this is not thread safe
-		const FRigHierarchy BaseHiearchy = ControlRig->GetBaseHierarchy();
+		const FRigBoneHierarchy& BaseHiearchy = ControlRig->GetBoneHierarchy();
 
 		DebugDrawSkeleton.Empty();
 		DebugDrawBones.Reset();
@@ -66,11 +66,11 @@ void UControlRigSkeletalMeshComponent::RebuildDebugDrawSkeleton()
 		// create ref modifier
  		FReferenceSkeletonModifier RefSkelModifier(DebugDrawSkeleton, nullptr);
  
- 		for (int32 Index = 0; Index < BaseHiearchy.GetNum(); Index++)
+ 		for (int32 Index = 0; Index < BaseHiearchy.Num(); Index++)
  		{
 			FMeshBoneInfo NewMeshBoneInfo;
 			NewMeshBoneInfo.Name = BaseHiearchy.GetName(Index);
-			NewMeshBoneInfo.ParentIndex = BaseHiearchy.GetParentIndex(Index);
+			NewMeshBoneInfo.ParentIndex = BaseHiearchy[Index].ParentIndex;
 			// give ref pose here
 			RefSkelModifier.Add(NewMeshBoneInfo, BaseHiearchy.GetInitialTransform(Index));
 
@@ -87,7 +87,7 @@ FTransform UControlRigSkeletalMeshComponent::GetDrawTransform(int32 BoneIndex) c
 	{
 		UControlRig* ControlRig = ControlRigInstance->CachedControlRig.Get();
 		// just copy it because this is not thread safe
-		const FRigHierarchy& BaseHiearchy = ControlRig->GetBaseHierarchy();
+		const FRigBoneHierarchy& BaseHiearchy = ControlRig->GetBoneHierarchy();
 		return BaseHiearchy.GetGlobalTransform(BoneIndex);
 	}
 
diff --git a/Engine/Plugins/Experimental/ControlRig/Source/ControlRigEditor/Private/Editor/RigHierarchyTabSummoner.cpp b/Engine/Plugins/Experimental/ControlRig/Source/ControlRigEditor/Private/Editor/RigHierarchyTabSummoner.cpp
--- a/Engine/Plugins/Experimental/ControlRig/Source/ControlRigEditor/Private/Editor/RigHierarchyTabSummoner.cpp
+++ b/Engine/Plugins/Experimental/ControlRig/Source/ControlRigEditor/Private/Editor/RigHierarchyTabSummoner.cpp
@@ -1,30 +1,30 @@
 // Copyright 1998-2019 Epic Games, Inc. All Rights Reserved.
 
 #include "RigHierarchyTabSummoner.h"
-#include "SRigHierarchy.h"
+#include "SRigBoneHierarchy.h"
 #include "ControlRigEditorStyle.h"
 #include "ControlRigEditor.h"
 
-#define LOCTEXT_NAMESPACE "RigHierarchyTabSummoner"
+#define LOCTEXT_NAMESPACE "RigBoneHierarchyTabSummoner"
 
-const FName FRigHierarchyTabSummoner::TabID(TEXT("RigHierarchy"));
+const FName FRigBoneHierarchyTabSummoner::TabID(TEXT("RigHierarchy"));
 
-FRigHierarchyTabSummoner::FRigHierarchyTabSummoner(const TSharedRef<FControlRigEditor>& InControlRigEditor)
+FRigBoneHierarchyTabSummoner::FRigBoneHierarchyTabSummoner(const TSharedRef<FControlRigEditor>& InControlRigEditor)
 	: FWorkflowTabFactory(TabID, InControlRigEditor)
 	, ControlRigEditor(InControlRigEditor)
 {
-	TabLabel = LOCTEXT("RigHierarchyTabLabel", "Hierarchy");
+	TabLabel = LOCTEXT("RigBoneHierarchyTabLabel", "Bone Hierarchy");
 	TabIcon = FSlateIcon(FControlRigEditorStyle::Get().GetStyleSetName(), "ControlRig.TabIcon");
 
 	bIsSingleton = true;
 
-	ViewMenuDescription = LOCTEXT("RigHierarchy_ViewMenu_Desc", "Hierarchy");
-	ViewMenuTooltip = LOCTEXT("RigHierarchy_ViewMenu_ToolTip", "Show the Rig Hierarchy tab");
+	ViewMenuDescription = LOCTEXT("RigBoneHierarchy_ViewMenu_Desc", "Bone Hierarchy");
+	ViewMenuTooltip = LOCTEXT("RigBoneHierarchy_ViewMenu_ToolTip", "Show the Rig Bone Hierarchy tab");
 }
 
-TSharedRef<SWidget> FRigHierarchyTabSummoner::CreateTabBody(const FWorkflowTabSpawnInfo& Info) const
+TSharedRef<SWidget> FRigBoneHierarchyTabSummoner::CreateTabBody(const FWorkflowTabSpawnInfo& Info) const
 {
-	return SNew(SRigHierarchy, ControlRigEditor.Pin().ToSharedRef());
+	return SNew(SRigBoneHierarchy, ControlRigEditor.Pin().ToSharedRef());
 }
 
 #undef LOCTEXT_NAMESPACE 
diff --git a/Engine/Plugins/Experimental/ControlRig/Source/ControlRigEditor/Private/Editor/RigHierarchyTabSummoner.h b/Engine/Plugins/Experimental/ControlRig/Source/ControlRigEditor/Private/Editor/RigHierarchyTabSummoner.h
--- a/Engine/Plugins/Experimental/ControlRig/Source/ControlRigEditor/Private/Editor/RigHierarchyTabSummoner.h
+++ b/Engine/Plugins/Experimental/ControlRig/Source/ControlRigEditor/Private/Editor/RigHierarchyTabSummoner.h
@@ -7,13 +7,13 @@
 
 class FControlRigEditor;
 
-struct FRigHierarchyTabSummoner : public FWorkflowTabFactory
+struct FRigBoneHierarchyTabSummoner : public FWorkflowTabFactory
 {
 public:
 	static const FName TabID;
 	
 public:
-	FRigHierarchyTabSummoner(const TSharedRef<FControlRigEditor>& InControlRigEditor);
+	FRigBoneHierarchyTabSummoner(const TSharedRef<FControlRigEditor>& InControlRigEditor);
 	
 	virtual TSharedRef<SWidget> CreateTabBody(const FWorkflowTabSpawnInfo& Info) const override;
 	
diff --git a/Engine/Plugins/Experimental/ControlRig/Source/ControlRigEditor/Private/Editor/SRigBoneHierarchy.cpp b/Engine/Plugins/Experimental/ControlRig/Source/ControlRigEditor/Private/Editor/SRigBoneHierarchy.cpp
new file mode 100644
--- /dev/null
+++ b/Engine/Plugins/Experimental/ControlRig/Source/ControlRigEditor/Private/Editor/SRigBoneHierarchy.cpp
@@ -0,0 +1,812 @@
+// Copyright 1998-2019 Epic Games, Inc. All Rights Reserved.
+
+#include "SRigBoneHierarchy.h"
+#include "Widgets/Input/SComboButton.h"
+#include "EditorStyleSet.h"
+#include "Widgets/Input/SSearchBox.h"
+#include "Kismet2/BlueprintEditorUtils.h"
+#include "ScopedTransaction.h"
+#include "ControlRigEditor.h"
+#include "BlueprintActionDatabase.h"
+#include "BlueprintVariableNodeSpawner.h"
+#include "Widgets/Layout/SSpacer.h"
+#include "Widgets/Images/SImage.h"
+#include "Widgets/Text/STextBlock.h"
+#include "K2Node_VariableGet.h"
+#include "ControlRigBlueprintUtils.h"
+#include "NodeSpawners/ControlRigPropertyNodeSpawner.h"
+#include "ControlRigHierarchyCommands.h"
+#include "ControlRigBlueprint.h"
+#include "Graph/ControlRigGraph.h"
+#include "Graph/ControlRigGraphNode.h"
+#include "Graph/ControlRigGraphSchema.h"
+#include "GraphEditorModule.h"
+#include "Framework/MultiBox/MultiBoxBuilder.h"
+#include "AnimationRuntime.h"
+#include "PropertyCustomizationHelpers.h"
+#include "Framework/Application/SlateApplication.h"
+#include "Editor/EditorEngine.h"
+#include "HelperUtil.h"
+#include "Widgets/Text/SInlineEditableTextBlock.h"
+#include "ControlRig.h"
+
+#define LOCTEXT_NAMESPACE "SRigBoneHierarchy"
+
+//////////////////////////////////////////////////////////////
+/// FRigTreeBone
+///////////////////////////////////////////////////////////
+FRigTreeBone::FRigTreeBone(const FName& InBone, TWeakPtr<SRigBoneHierarchy> InHierarchyHandler)
+{
+	CachedBone = InBone;
+}
+
+TSharedRef<ITableRow> FRigTreeBone::MakeTreeRowWidget(TSharedPtr<FControlRigEditor> InControlRigEditor, const TSharedRef<STableViewBase>& InOwnerTable, TSharedRef<FRigTreeBone> InRigTreeBone, TSharedRef<FUICommandList> InCommandList, TSharedPtr<SRigBoneHierarchy> InHierarchy)
+{
+	return SNew(SRigBoneHierarchyItem, InControlRigEditor, InOwnerTable, InRigTreeBone, InCommandList, InHierarchy)
+		.OnRenameBone(InHierarchy.Get(), &SRigBoneHierarchy::RenameBone)
+		.OnVerifyBoneNameChanged(InHierarchy.Get(), &SRigBoneHierarchy::OnVerifyNameChanged);
+}
+
+void FRigTreeBone::RequestRename()
+{
+	OnRenameRequested.ExecuteIfBound();
+}
+
+//////////////////////////////////////////////////////////////
+/// FRigBoneHierarchyDragDropOp
+///////////////////////////////////////////////////////////
+TSharedRef<FRigBoneHierarchyDragDropOp> FRigBoneHierarchyDragDropOp::New(TArray<FName> InBoneNames)
+{
+	TSharedRef<FRigBoneHierarchyDragDropOp> Operation = MakeShared<FRigBoneHierarchyDragDropOp>();
+	Operation->BoneNames = InBoneNames;
+	Operation->Construct();
+	return Operation;
+}
+
+TSharedPtr<SWidget> FRigBoneHierarchyDragDropOp::GetDefaultDecorator() const
+{
+	return SNew(SBorder)
+		.Visibility(EVisibility::Visible)
+		.BorderImage(FEditorStyle::GetBrush("Menu.Background"))
+		[
+			SNew(STextBlock)
+			.Text(FText::FromString(GetJoinedBoneNames()))
+			.Font(FEditorStyle::Get().GetFontStyle("FontAwesome.10"))
+		];
+}
+
+FString FRigBoneHierarchyDragDropOp::GetJoinedBoneNames() const
+{
+	TArray<FString> BoneNameStrings;
+	for (const FName& BoneName : BoneNames)
+	{
+		BoneNameStrings.Add(BoneName.ToString());
+	}
+	return FString::Join(BoneNameStrings, TEXT(","));
+}
+
+
+//////////////////////////////////////////////////////////////
+/// SRigBoneHierarchyItem
+///////////////////////////////////////////////////////////
+void SRigBoneHierarchyItem::Construct(const FArguments& InArgs, TSharedPtr<FControlRigEditor> InControlRigEditor, const TSharedRef<STableViewBase>& OwnerTable, TSharedRef<FRigTreeBone> InRigTreeBone, TSharedRef<FUICommandList> InCommandList, TSharedPtr<SRigBoneHierarchy> InHierarchy)
+{
+	WeakRigTreeBone = InRigTreeBone;
+	WeakCommandList = InCommandList;
+	ControlRigEditor = InControlRigEditor;
+
+	OnVerifyBoneNameChanged = InArgs._OnVerifyBoneNameChanged;
+	OnRenameBone = InArgs._OnRenameBone;
+
+	TSharedPtr< SInlineEditableTextBlock > InlineWidget;
+
+	STableRow<TSharedPtr<FRigTreeBone>>::Construct(
+		STableRow<TSharedPtr<FRigTreeBone>>::FArguments()
+		.OnDragDetected(InHierarchy.Get(), &SRigBoneHierarchy::OnDragDetected)
+		.Content()
+		[
+			SNew(SHorizontalBox)
+			+ SHorizontalBox::Slot()
+			.AutoWidth()
+			.VAlign(VAlign_Center)
+			[
+				SAssignNew(InlineWidget, SInlineEditableTextBlock)
+				.Text(this, &SRigBoneHierarchyItem::GetName)
+				//.HighlightText(FilterText)
+				.OnVerifyTextChanged(this, &SRigBoneHierarchyItem::OnVerifyNameChanged)
+				.OnTextCommitted(this, &SRigBoneHierarchyItem::OnNameCommitted)
+				.MultiLine(false)
+			]
+		], OwnerTable);
+
+	InRigTreeBone->OnRenameRequested.BindSP(InlineWidget.Get(), &SInlineEditableTextBlock::EnterEditingMode);
+}
+
+FText SRigBoneHierarchyItem::GetName() const
+{
+	return (FText::FromName(WeakRigTreeBone.Pin()->CachedBone));
+}
+
+bool SRigBoneHierarchyItem::OnVerifyNameChanged(const FText& InText, FText& OutErrorMessage)
+{
+	const FName NewName = FName(*InText.ToString());
+	if (OnVerifyBoneNameChanged.IsBound())
+	{
+		return OnVerifyBoneNameChanged.Execute(WeakRigTreeBone.Pin()->CachedBone, NewName, OutErrorMessage);
+	}
+
+	// if not bound, just allow
+	return true;
+}
+
+void SRigBoneHierarchyItem::OnNameCommitted(const FText& InText, ETextCommit::Type InCommitType) const
+{
+	// for now only allow enter
+	// because it is important to keep the unique names per pose
+	if (InCommitType == ETextCommit::OnEnter)
+	{
+		FName NewName = FName(*InText.ToString());
+		FName OldName = WeakRigTreeBone.Pin()->CachedBone;
+
+		if (!OnRenameBone.IsBound() || OnRenameBone.Execute(OldName, NewName))
+		{
+			if (WeakRigTreeBone.IsValid())
+			{
+				WeakRigTreeBone.Pin()->CachedBone = NewName;
+			}
+		}
+	}
+}
+
+///////////////////////////////////////////////////////////
+
+SRigBoneHierarchy::~SRigBoneHierarchy()
+{
+	if (ControlRigEditor.IsValid())
+	{
+		ControlRigBlueprint = ControlRigEditor.Pin()->GetControlRigBlueprint();
+		if (ControlRigBlueprint.IsValid())
+		{
+			ControlRigBlueprint->HierarchyContainer.OnElementChanged.RemoveAll(this);
+		}
+	}
+}
+
+void SRigBoneHierarchy::Construct(const FArguments& InArgs, TSharedRef<FControlRigEditor> InControlRigEditor)
+{
+	ControlRigEditor = InControlRigEditor;
+
+	ControlRigBlueprint = ControlRigEditor.Pin()->GetControlRigBlueprint();
+	// @todo: find a better place to do it
+	ControlRigBlueprint->HierarchyContainer.BoneHierarchy.Initialize();
+
+	ControlRigBlueprint->HierarchyContainer.OnElementChanged.AddRaw(this, &SRigBoneHierarchy::OnBoneHierarchyChanged);
+
+	// for deleting, renaming, dragging
+	CommandList = MakeShared<FUICommandList>();
+
+	UEditorEngine* Editor = Cast<UEditorEngine>(GEngine);
+	if (Editor != nullptr)
+	{
+		Editor->RegisterForUndo(this);
+	}
+
+	BindCommands();
+
+	ChildSlot
+	[
+		SNew(SVerticalBox)
+		+SVerticalBox::Slot()
+		.AutoHeight()
+		.VAlign(VAlign_Top)
+		.Padding(0.0f)
+		[
+			SNew(SBorder)
+			.Padding(0.0f)
+			.BorderImage(FEditorStyle::GetBrush("DetailsView.CategoryTop"))
+			.BorderBackgroundColor(FLinearColor(0.6f, 0.6f, 0.6f, 1.0f))
+			[
+				SNew(SVerticalBox)
+				+SVerticalBox::Slot()
+				.AutoHeight()
+				.VAlign(VAlign_Top)
+				[
+					SNew(SHorizontalBox)
+					+SHorizontalBox::Slot()
+					.VAlign(VAlign_Center)
+					.Padding(3.0f, 1.0f)
+					[
+						SAssignNew(FilterBox, SSearchBox)
+						.OnTextChanged(this, &SRigBoneHierarchy::OnFilterTextChanged)
+					]
+				]
+			]
+		]
+		+SVerticalBox::Slot()
+		.Padding(0.0f, 0.0f)
+		[
+			SNew(SBorder)
+			.Padding(2.0f)
+			.BorderImage(FEditorStyle::GetBrush("SCSEditor.TreePanel"))
+			[
+				SAssignNew(TreeView, STreeView<TSharedPtr<FRigTreeBone>>)
+				.TreeItemsSource(&RootBones)
+				.SelectionMode(ESelectionMode::Multi)
+				.OnGenerateRow(this, &SRigBoneHierarchy::MakeTableRowWidget)
+				.OnGetChildren(this, &SRigBoneHierarchy::HandleGetChildrenForTree)
+				.OnSelectionChanged(this, &SRigBoneHierarchy::OnSelectionChanged)
+				.OnContextMenuOpening(this, &SRigBoneHierarchy::CreateContextMenu)
+				.ItemHeight(24)
+			]
+		]
+	];
+
+	RefreshTreeView();
+}
+
+void SRigBoneHierarchy::BindCommands()
+{
+	// create new command
+	const FControlRigHierarchyCommands& Commands = FControlRigHierarchyCommands::Get();
+	CommandList->MapAction(Commands.AddItem,
+		FExecuteAction::CreateSP(this, &SRigBoneHierarchy::HandleNewItem));
+
+	CommandList->MapAction(Commands.DuplicateItem,
+		FExecuteAction::CreateSP(this, &SRigBoneHierarchy::HandleDuplicateItem),
+		FCanExecuteAction::CreateSP(this, &SRigBoneHierarchy::CanDuplicateItem));
+
+	CommandList->MapAction(Commands.DeleteItem,
+		FExecuteAction::CreateSP(this, &SRigBoneHierarchy::HandleDeleteItem),
+		FCanExecuteAction::CreateSP(this, &SRigBoneHierarchy::CanDeleteItem));
+
+	CommandList->MapAction(Commands.RenameItem,
+		FExecuteAction::CreateSP(this, &SRigBoneHierarchy::HandleRenameItem),
+		FCanExecuteAction::CreateSP(this, &SRigBoneHierarchy::CanRenameItem));
+}
+
+void SRigBoneHierarchy::OnFilterTextChanged(const FText& SearchText)
+{
+	FilterText = SearchText;
+
+	RefreshTreeView();
+}
+
+void SRigBoneHierarchy::RefreshTreeView()
+{
+	RootBones.Reset();
+	FilteredRootBones.Reset();
+
+	if (ControlRigBlueprint.IsValid())
+	{
+		FRigBoneHierarchy& Hierarchy = ControlRigBlueprint->HierarchyContainer.BoneHierarchy;
+
+		TMap<FName, TSharedPtr<FRigTreeBone>> SearchTable;
+
+		FString FilteredString = FilterText.ToString();
+		const bool bSearchOff = FilteredString.IsEmpty();
+		for (int32 BoneIndex = 0; BoneIndex < Hierarchy.Num(); ++BoneIndex)
+		{
+			const FRigBone& Bone = Hierarchy[BoneIndex];
+
+			// create new item
+			if (bSearchOff)
+			{
+				TSharedPtr<FRigTreeBone> NewItem = MakeShared<FRigTreeBone>(Bone.Name, SharedThis(this));
+				SearchTable.Add(Bone.Name, NewItem);
+
+				if (Bone.ParentName == NAME_None)
+				{
+					RootBones.Add(NewItem);
+				}
+				else
+				{
+					// you have to find one
+					TSharedPtr<FRigTreeBone>* FoundItem = SearchTable.Find(Bone.ParentName);
+					check(FoundItem);
+					// add to children list
+					FoundItem->Get()->Children.Add(NewItem);
+				}
+			}
+			else if (Bone.Name.ToString().Contains(FilteredString))
+			{
+				// if contains, just list out everything to root
+				TSharedPtr<FRigTreeBone> NewItem = MakeShared<FRigTreeBone>(Bone.Name, SharedThis(this));
+				// during search, everything is on root
+				RootBones.Add(NewItem);
+			}
+		}
+
+		if (bSearchOff)
+		{
+			for (int32 RootIndex = 0; RootIndex < RootBones.Num(); ++RootIndex)
+			{
+				SetExpansionRecursive(RootBones[RootIndex]);
+			}
+		}
+	}
+
+	TreeView->RequestTreeRefresh();
+}
+
+void SRigBoneHierarchy::SetExpansionRecursive(TSharedPtr<FRigTreeBone> InBone)
+{
+	TreeView->SetItemExpansion(InBone, true);
+
+	for (int32 ChildIndex = 0; ChildIndex < InBone->Children.Num(); ++ChildIndex)
+	{
+		SetExpansionRecursive(InBone->Children[ChildIndex]);
+	}
+}
+TSharedRef<ITableRow> SRigBoneHierarchy::MakeTableRowWidget(TSharedPtr<FRigTreeBone> InItem, const TSharedRef<STableViewBase>& OwnerTable)
+{
+	return InItem->MakeTreeRowWidget(ControlRigEditor.Pin(), OwnerTable, InItem.ToSharedRef(), CommandList.ToSharedRef(), SharedThis(this));
+}
+
+void SRigBoneHierarchy::HandleGetChildrenForTree(TSharedPtr<FRigTreeBone> InItem, TArray<TSharedPtr<FRigTreeBone>>& OutChildren)
+{
+	OutChildren = InItem.Get()->Children;
+}
+
+void SRigBoneHierarchy::OnSelectionChanged(TSharedPtr<FRigTreeBone> Selection, ESelectInfo::Type SelectInfo)
+{
+	// need dummy object
+	if (Selection.IsValid())
+	{
+		FRigBoneHierarchy* RigHierarchy = GetInstanceHierarchy();
+
+		if (RigHierarchy)
+		{
+			const int32 BoneIndex = RigHierarchy->GetIndex(Selection->CachedBone);
+			if (BoneIndex != INDEX_NONE)
+			{
+				ControlRigEditor.Pin()->SetDetailStruct(MakeShareable(new FStructOnScope(FRigBone::StaticStruct(), (uint8*)&(*RigHierarchy)[BoneIndex])));
+				ControlRigEditor.Pin()->SelectBone(Selection->CachedBone);
+				return;
+			}
+			else
+			{
+				// clear the current selection
+				ControlRigEditor.Pin()->ClearDetailObject();
+				ControlRigEditor.Pin()->SelectBone(NAME_None);
+			}
+		}
+
+		// if failed, try BP hierarhcy? Todo:
+	}
+}
+
+TSharedPtr<FRigTreeBone> FindBone(const FName& InBoneName, TSharedPtr<FRigTreeBone> CurrentItem)
+{
+	if (CurrentItem->CachedBone == InBoneName)
+	{
+		return CurrentItem;
+	}
+
+	for (int32 ChildIndex = 0; ChildIndex < CurrentItem->Children.Num(); ++ChildIndex)
+	{
+		TSharedPtr<FRigTreeBone> Found = FindBone(InBoneName, CurrentItem->Children[ChildIndex]);
+		if (Found.IsValid())
+		{
+			return Found;
+		}
+	}
+
+	return TSharedPtr<FRigTreeBone>();
+}
+
+void SRigBoneHierarchy::SelectBone(const FName& BoneName) const
+{
+	for (int32 RootIndex = 0; RootIndex < RootBones.Num(); ++RootIndex)
+	{
+		TSharedPtr<FRigTreeBone> Found = FindBone(BoneName, RootBones[RootIndex]);
+		if (Found.IsValid())
+		{
+			TreeView->SetSelection(Found);
+		}
+	}
+}
+
+void SRigBoneHierarchy::OnBoneHierarchyChanged(FRigHierarchyContainer* Container, ERigHierarchyElementType ElementType, const FName& InName)
+{
+	if (ElementType != ERigHierarchyElementType::Bone)
+	{
+		return;
+	}
+
+	RefreshTreeView();
+}
+
+void SRigBoneHierarchy::ClearDetailPanel() const
+{
+	ControlRigEditor.Pin()->ClearDetailObject();
+}
+
+TSharedPtr< SWidget > SRigBoneHierarchy::CreateContextMenu()
+{
+	const FControlRigHierarchyCommands& Actions = FControlRigHierarchyCommands::Get();
+
+	const bool CloseAfterSelection = true;
+	FMenuBuilder MenuBuilder(CloseAfterSelection, CommandList);
+	{
+		MenuBuilder.BeginSection("HierarchyEditAction", LOCTEXT("EditAction", "Edit"));
+		MenuBuilder.AddMenuEntry(Actions.AddItem);
+		MenuBuilder.AddMenuEntry(Actions.DeleteItem);
+		MenuBuilder.AddMenuEntry(Actions.DuplicateItem);
+		MenuBuilder.AddMenuEntry(Actions.RenameItem);
+
+		MenuBuilder.AddMenuSeparator();
+		MenuBuilder.AddSubMenu(
+			LOCTEXT("ImportSubMenu", "Import"),
+			LOCTEXT("ImportSubMenu_ToolTip", "Import hierarchy to the current rig. This only imports non-existing node. For example, if there is hand_r, it won't import hand_r. \
+				If you want to reimport whole new hiearchy, delete all nodes, and use import hierarchy."),
+			FNewMenuDelegate::CreateSP(this, &SRigBoneHierarchy::CreateImportMenu)
+		);
+
+		MenuBuilder.AddMenuSeparator();
+		MenuBuilder.AddSubMenu(
+			LOCTEXT("RefreshSubMenu", "Refresh"),
+			LOCTEXT("RefreshSubMenu_ToolTip", "Refresh the existing initial transform from the selected mesh. This only updates if the node is found."),
+			FNewMenuDelegate::CreateSP(this, &SRigBoneHierarchy::CreateRefreshMenu)
+		);
+	
+		MenuBuilder.EndSection();
+	}
+
+	return MenuBuilder.MakeWidget();
+}
+
+void SRigBoneHierarchy::CreateRefreshMenu(FMenuBuilder& MenuBuilder)
+{
+	MenuBuilder.AddWidget(
+		SNew(SVerticalBox)
+
+		+ SVerticalBox::Slot()
+		.AutoHeight()
+		.Padding(3)
+		[
+			SNew(STextBlock)
+			.Font(FEditorStyle::GetFontStyle("ControlRig.Hierarchy.Menu"))
+			.Text(LOCTEXT("RefreshMesh_Title", "Select Mesh"))
+			.ToolTipText(LOCTEXT("RefreshMesh_Tooltip", "Select Mesh to refresh transform from... It will refresh init transform from selected mesh. This doesn't change hierarchy. \
+				If you want to reimport hierarchy, please delete all nodes, and use import hierarchy."))
+		]
+		+ SVerticalBox::Slot()
+		.AutoHeight()
+		.Padding(3)
+		[
+			SNew(SObjectPropertyEntryBox)
+			.AllowedClass(USkeletalMesh::StaticClass())
+			.OnObjectChanged(this, &SRigBoneHierarchy::RefreshHierarchy)
+		]
+		,
+		FText()
+	);
+}
+
+void SRigBoneHierarchy::RefreshHierarchy(const FAssetData& InAssetData)
+{
+	FRigBoneHierarchy* Hier = GetHierarchy();
+	USkeletalMesh* Mesh = Cast<USkeletalMesh>(InAssetData.GetAsset());
+	if (Mesh && Hier)
+	{
+		FScopedTransaction Transaction(LOCTEXT("HierarchyRefresh", "Refresh Transform"));
+		ControlRigBlueprint->Modify();
+
+		const FReferenceSkeleton& RefSkeleton = Mesh->RefSkeleton;
+		const TArray<FMeshBoneInfo>& BoneInfos = RefSkeleton.GetRefBoneInfo();
+		const TArray<FTransform>& BonePoses = RefSkeleton.GetRefBonePose();
+
+		for (int32 BoneIndex = 0; BoneIndex < RefSkeleton.GetNum(); ++BoneIndex)
+		{
+			// only add if you don't have it. This may change in the future
+			int32 RigIndex = Hier->GetIndex(BoneInfos[BoneIndex].Name);
+			if (RigIndex != INDEX_NONE)
+			{
+				// @todo: add optimized version without sorting, but if no sort, we should make sure not to use find index function
+				Hier->SetInitialTransform(RigIndex, FAnimationRuntime::GetComponentSpaceTransform(RefSkeleton, BonePoses, BoneIndex));
+			}
+		}
+	}
+}
+void SRigBoneHierarchy::CreateImportMenu(FMenuBuilder& MenuBuilder)
+{
+	MenuBuilder.AddWidget(
+		SNew(SVerticalBox)
+
+		+ SVerticalBox::Slot()
+		.AutoHeight()
+		.Padding(3)
+		[
+			SNew(STextBlock)
+			.Font(FEditorStyle::GetFontStyle("ControlRig.Hierarchy.Menu"))
+			.Text(LOCTEXT("ImportMesh_Title", "Select Mesh"))
+			.ToolTipText(LOCTEXT("ImportMesh_Tooltip", "Select Mesh to import hierarchy from... It will only import if the node doens't exists in the current hierarchy."))
+		]
+		+ SVerticalBox::Slot()
+		.AutoHeight()
+		.Padding(3)
+		[
+			SNew(SObjectPropertyEntryBox)
+			.AllowedClass(USkeletalMesh::StaticClass())
+			.OnObjectChanged(this, &SRigBoneHierarchy::ImportHierarchy)
+		]
+		,
+		FText()
+	);
+}
+
+void SRigBoneHierarchy::ImportHierarchy(const FAssetData& InAssetData)
+{
+	FRigBoneHierarchy* Hier = GetHierarchy();
+	USkeletalMesh* Mesh = Cast<USkeletalMesh> (InAssetData.GetAsset());
+	if (Mesh && Hier)
+	{
+		FScopedTransaction Transaction(LOCTEXT("HierarchyImport", "Import Hierarchy"));
+		ControlRigBlueprint->Modify();
+
+		const FReferenceSkeleton& RefSkeleton = Mesh->RefSkeleton;
+		const TArray<FMeshBoneInfo>& BoneInfos = RefSkeleton.GetRefBoneInfo();
+		const TArray<FTransform>& BonePoses = RefSkeleton.GetRefBonePose();
+
+		for (int32 BoneIndex = 0; BoneIndex < RefSkeleton.GetNum(); ++BoneIndex)
+		{
+			// only add if you don't have it. This may change in the future
+			if (Hier->GetIndex(BoneInfos[BoneIndex].Name) == INDEX_NONE)
+			{
+				// @todo: add optimized version without sorting, but if no sort, we should make sure not to use find index function
+				FName ParentName = (BoneInfos[BoneIndex].ParentIndex != INDEX_NONE) ? BoneInfos[BoneInfos[BoneIndex].ParentIndex].Name : NAME_None;
+				Hier->Add(BoneInfos[BoneIndex].Name, ParentName, FAnimationRuntime::GetComponentSpaceTransform(RefSkeleton, BonePoses, BoneIndex));
+			}
+		}
+	}
+}
+
+bool SRigBoneHierarchy::IsMultiSelected() const
+{
+	return TreeView->GetNumItemsSelected() > 0;
+}
+
+bool SRigBoneHierarchy::IsSingleSelected() const
+{
+	return TreeView->GetNumItemsSelected() == 1;
+}
+
+void SRigBoneHierarchy::HandleDeleteItem()
+{
+ 	FRigBoneHierarchy* Hierarchy = GetHierarchy();
+ 	if (Hierarchy)
+ 	{
+		ClearDetailPanel();
+		FScopedTransaction Transaction(LOCTEXT("HierarchyTreeDeleteSelected", "Delete selected items from hierarchy"));
+		ControlRigBlueprint->Modify();
+
+		// clear detail view display
+		ControlRigEditor.Pin()->ClearDetailObject();
+
+		TArray<TSharedPtr<FRigTreeBone>> SelectedItems = TreeView->GetSelectedItems();
+
+		for (int32 ItemIndex = 0; ItemIndex < SelectedItems.Num(); ++ItemIndex)
+		{
+			// when you select whole Bones, you might not have them anymore
+			if (Hierarchy->GetIndex(SelectedItems[ItemIndex]->CachedBone) != INDEX_NONE)
+			{
+				Hierarchy->Remove(SelectedItems[ItemIndex]->CachedBone);
+			}
+		}
+ 	}
+}
+
+bool SRigBoneHierarchy::CanDeleteItem() const
+{
+	return IsMultiSelected();
+}
+
+/** Delete Item */
+void SRigBoneHierarchy::HandleNewItem()
+{
+	FRigBoneHierarchy* Hierarchy = GetHierarchy();
+	if (Hierarchy)
+	{
+		// unselect current selected item
+		ClearDetailPanel();
+
+		FScopedTransaction Transaction(LOCTEXT("HierarchyTreeAdded", "Add new item to hierarchy"));
+		ControlRigBlueprint->Modify();
+
+		FName ParentName = NAME_None;
+		FTransform ParentTransform = FTransform::Identity;
+
+		TArray<TSharedPtr<FRigTreeBone>> SelectedItems = TreeView->GetSelectedItems();
+		if (SelectedItems.Num() > 0)
+		{
+			ParentName = SelectedItems[0]->CachedBone;
+			ParentTransform = Hierarchy->GetGlobalTransform(ParentName);
+		}
+
+		const FName NewBoneName = CreateUniqueName(TEXT("NewBone"));
+		Hierarchy->Add(NewBoneName, ParentName, ParentTransform);
+
+		// reselect current selected item
+		SelectBone(NewBoneName);
+	}
+}
+
+/** Check whether we can deleting the selected item(s) */
+bool SRigBoneHierarchy::CanDuplicateItem() const
+{
+	return IsMultiSelected();
+}
+
+/** Duplicate Item */
+void SRigBoneHierarchy::HandleDuplicateItem()
+{
+	FRigBoneHierarchy* Hierarchy = GetHierarchy();
+	if (Hierarchy)
+	{
+		ClearDetailPanel();
+
+		FScopedTransaction Transaction(LOCTEXT("HierarchyTreeDuplicateSelected", "Duplicate selected items from hierarchy"));
+		ControlRigBlueprint->Modify();
+
+		TArray<TSharedPtr<FRigTreeBone>> SelectedItems = TreeView->GetSelectedItems();
+		TArray<FName> NewNames;
+		for (int32 Index = 0; Index < SelectedItems.Num(); ++Index)
+		{
+			FName Name = SelectedItems[Index]->CachedBone;
+			FTransform Transform = Hierarchy->GetGlobalTransform(Name);
+
+			FName ParentName = (*Hierarchy)[Name].ParentName;
+
+			const FName NewName = CreateUniqueName(Name);
+			Hierarchy->Add(NewName, ParentName, Transform);
+			NewNames.Add(NewName);
+		}
+
+		for (int32 Index = 0; Index < NewNames.Num(); ++Index)
+		{
+			SelectBone(NewNames[Index]);
+		}
+	}
+}
+
+/** Check whether we can deleting the selected item(s) */
+bool SRigBoneHierarchy::CanRenameItem() const
+{
+	return IsSingleSelected();
+}
+
+/** Delete Item */
+void SRigBoneHierarchy::HandleRenameItem()
+{
+	FRigBoneHierarchy* Hierarchy = GetHierarchy();
+	if (Hierarchy)
+	{
+		ClearDetailPanel();
+
+		FScopedTransaction Transaction(LOCTEXT("HierarchyTreeRenameSelected", "Rename selected item from hierarchy"));
+		ControlRigBlueprint->Modify();
+
+		TArray<TSharedPtr<FRigTreeBone>> SelectedItems = TreeView->GetSelectedItems();
+		if (SelectedItems.Num() > 0)
+		{
+			SelectedItems[0]->RequestRename();
+		}
+	}
+}
+
+FRigBoneHierarchy* SRigBoneHierarchy::GetHierarchy() const
+{
+	if (ControlRigBlueprint.IsValid())
+	{
+		return &ControlRigBlueprint->HierarchyContainer.BoneHierarchy;
+	}
+
+	return nullptr;
+}
+
+FRigBoneHierarchy* SRigBoneHierarchy::GetInstanceHierarchy() const
+{
+	if (ControlRigEditor.IsValid())
+	{
+		UControlRig* ControlRig = ControlRigEditor.Pin()->GetInstanceRig();
+		if (ControlRig)
+		{
+			return &ControlRig->Hierarchy.BoneHierarchy;
+		}
+	}
+
+	return nullptr;
+}
+FName SRigBoneHierarchy::CreateUniqueName(const FName& InBaseName) const
+{
+	return UtilityHelpers::CreateUniqueName(InBaseName, [this](const FName& CurName) { return GetHierarchy()->GetIndex(CurName) == INDEX_NONE; });
+}
+
+void SRigBoneHierarchy::PostRedo(bool bSuccess) 
+{
+	if (bSuccess)
+	{
+		RefreshTreeView();
+	}
+}
+
+void SRigBoneHierarchy::PostUndo(bool bSuccess) 
+{
+	if (bSuccess)
+	{
+		RefreshTreeView();
+	}
+}
+
+FReply SRigBoneHierarchy::OnDragDetected(const FGeometry& MyGeometry, const FPointerEvent& MouseEvent)
+{
+	TArray<FName> DraggedBoneNames;
+	TArray<TSharedPtr<FRigTreeBone>> SelectedItems =TreeView->GetSelectedItems();
+	for (const TSharedPtr<FRigTreeBone>& SelectedItem : SelectedItems)
+	{
+		DraggedBoneNames.Add(SelectedItem->CachedBone);
+	}
+
+	if (MouseEvent.IsMouseButtonDown(EKeys::LeftMouseButton))
+	{
+		if (ControlRigEditor.IsValid())
+		{
+			TSharedRef<FRigBoneHierarchyDragDropOp> DragDropOp = FRigBoneHierarchyDragDropOp::New(MoveTemp(DraggedBoneNames));
+			DragDropOp->OnPerformDropToGraph.BindSP(ControlRigEditor.Pin().Get(), &FControlRigEditor::OnGraphNodeDropToPerform);
+			return FReply::Handled().BeginDragDrop(DragDropOp);
+		}
+	}
+
+	return FReply::Unhandled();
+}
+
+
+bool SRigBoneHierarchy::RenameBone(const FName& OldName, const FName& NewName)
+{
+	ClearDetailPanel();
+
+	if (OldName == NewName)
+	{
+		return true;
+	}
+
+	// make sure there is no duplicate
+	FRigBoneHierarchy* Hierarchy = GetHierarchy();
+	if (Hierarchy)
+	{
+		Hierarchy->Rename(OldName, NewName);
+		SelectBone(NewName);
+
+		return true;
+	}
+
+	return false;
+}
+
+bool SRigBoneHierarchy::OnVerifyNameChanged(const FName& OldName, const FName& NewName, FText& OutErrorMessage)
+{
+	if (OldName == NewName)
+	{
+		return true;
+	}
+
+	// make sure there is no duplicate
+	FRigBoneHierarchy* Hierarchy = GetHierarchy();
+	if (Hierarchy)
+	{
+		const int32 Found = Hierarchy->GetIndex(OldName);
+		if (Found != INDEX_NONE)
+		{
+			const int32 Duplicate = Hierarchy->GetIndex(NewName);
+			if (Duplicate != INDEX_NONE)
+			{
+				OutErrorMessage = FText::FromString(TEXT("Duplicate name exists"));
+
+				return false;
+			}
+		}
+	}
+
+	return true;
+}
+#undef LOCTEXT_NAMESPACE
\ No newline at end of file
diff --git a/Engine/Plugins/Experimental/ControlRig/Source/ControlRigEditor/Private/Editor/SRigBoneHierarchy.h b/Engine/Plugins/Experimental/ControlRig/Source/ControlRigEditor/Private/Editor/SRigBoneHierarchy.h
new file mode 100644
--- /dev/null
+++ b/Engine/Plugins/Experimental/ControlRig/Source/ControlRigEditor/Private/Editor/SRigBoneHierarchy.h
@@ -0,0 +1,201 @@
+// Copyright 1998-2019 Epic Games, Inc. All Rights Reserved.
+
+#pragma once
+
+#include "CoreMinimal.h"
+#include "Widgets/Views/STreeView.h"
+#include "Rigs/RigHierarchyContainer.h"
+#include "EditorUndoClient.h"
+#include "DragAndDrop/GraphNodeDragDropOp.h"
+
+class SRigBoneHierarchy;
+class FControlRigEditor;
+class SSearchBox;
+class FUICommandList;
+class UControlRigBlueprint;
+struct FAssetData;
+class FMenuBuilder;
+class SRigBoneHierarchyItem;
+
+DECLARE_DELEGATE_RetVal_TwoParams(bool, FOnRenameBone, const FName& /*OldName*/, const FName& /*NewName*/);
+DECLARE_DELEGATE_RetVal_ThreeParams(bool, FOnVerifyBoneNameChanged, const FName& /*OldName*/, const FName& /*NewName*/, FText& /*OutErrorMessage*/);
+
+/** An item in the tree */
+class FRigTreeBone : public TSharedFromThis<FRigTreeBone>
+{
+public:
+	FRigTreeBone(const FName& InBone, TWeakPtr<SRigBoneHierarchy> InHierarchyHandler);
+public:
+	/** Bone Data to display */
+	FName CachedBone;
+	TArray<TSharedPtr<FRigTreeBone>> Children;
+
+	TSharedRef<ITableRow> MakeTreeRowWidget(TSharedPtr<FControlRigEditor> InControlRigEditor, const TSharedRef<STableViewBase>& InOwnerTable, TSharedRef<FRigTreeBone> InRigTreeBone, TSharedRef<FUICommandList> InCommandList, TSharedPtr<SRigBoneHierarchy> InHierarchy);
+
+	void RequestRename();
+
+	/** Delegate for when the context menu requests a rename */
+	DECLARE_DELEGATE(FOnRenameRequested);
+	FOnRenameRequested OnRenameRequested;
+};
+
+class FRigBoneHierarchyDragDropOp : public FGraphNodeDragDropOp
+{
+public:
+	DRAG_DROP_OPERATOR_TYPE(FRigBoneHierarchyDragDropOp, FGraphNodeDragDropOp)
+
+	static TSharedRef<FRigBoneHierarchyDragDropOp> New(TArray<FName> InBoneNames);
+
+	virtual TSharedPtr<SWidget> GetDefaultDecorator() const override;
+
+	/** @return true if this drag operation contains property paths */
+	bool HasBoneNames() const
+	{
+		return BoneNames.Num() > 0;
+	}
+
+	/** @return The property paths from this drag operation */
+	const TArray<FName>& GetBoneNames() const
+	{
+		return BoneNames;
+	}
+
+	FString GetJoinedBoneNames() const;
+
+private:
+
+	/** Data for the property paths this item represents */
+	TArray<FName> BoneNames;
+};
+
+ class SRigBoneHierarchyItem : public STableRow<TSharedPtr<FRigTreeBone>>
+{
+	SLATE_BEGIN_ARGS(SRigBoneHierarchyItem) {}
+		/** Callback when the text is committed. */
+		SLATE_EVENT(FOnRenameBone, OnRenameBone)
+		/** Called whenever the text is changed interactively by the user */
+		SLATE_EVENT(FOnVerifyBoneNameChanged, OnVerifyBoneNameChanged)
+	SLATE_END_ARGS()
+
+	void Construct(const FArguments& InArgs, TSharedPtr<FControlRigEditor> InControlRigEditor, const TSharedRef<STableViewBase>& OwnerTable, TSharedRef<FRigTreeBone> InControlRigTreeBone, TSharedRef<FUICommandList> InCommandList, TSharedPtr<SRigBoneHierarchy> InHierarchy);
+	void OnNameCommitted(const FText& InText, ETextCommit::Type InCommitType) const;
+	bool OnVerifyNameChanged(const FText& InText, FText& OutErrorMessage);
+
+private:
+	TWeakPtr<FRigTreeBone> WeakRigTreeBone;
+	TWeakPtr<FUICommandList> WeakCommandList;
+	TWeakPtr<FControlRigEditor> ControlRigEditor;
+
+	FOnRenameBone OnRenameBone;
+	FOnVerifyBoneNameChanged OnVerifyBoneNameChanged;
+
+	FText GetName() const;
+};
+
+/** Widget allowing editing of a control rig's structure */
+class SRigBoneHierarchy : public SCompoundWidget, public FEditorUndoClient
+{
+public:
+	SLATE_BEGIN_ARGS(SRigBoneHierarchy) {}
+	SLATE_END_ARGS()
+
+	~SRigBoneHierarchy();
+
+	void Construct(const FArguments& InArgs, TSharedRef<FControlRigEditor> InControlRigEditor);
+
+private:
+	/** Bind commands that this widget handles */
+	void BindCommands();
+
+	/** Rebuild the tree view */
+	void RefreshTreeView();
+
+	/** Make a row widget for the table */
+	TSharedRef<ITableRow> MakeTableRowWidget(TSharedPtr<FRigTreeBone> InItem, const TSharedRef<STableViewBase>& OwnerTable);
+
+	/** Get children for the tree */
+	void HandleGetChildrenForTree(TSharedPtr<FRigTreeBone> InItem, TArray<TSharedPtr<FRigTreeBone>>& OutChildren);
+
+	/** Check whether we can deleting the selected item(s) */
+	bool CanDeleteItem() const;
+
+	/** Delete Item */
+	void HandleDeleteItem();
+
+	/** Delete Item */
+	void HandleNewItem();
+
+	/** Check whether we can deleting the selected item(s) */
+	bool CanDuplicateItem() const;
+
+	/** Delete Item */
+	void HandleDuplicateItem();
+
+	/** Check whether we can deleting the selected item(s) */
+	bool CanRenameItem() const;
+
+	/** Delete Item */
+	void HandleRenameItem();
+
+	/** Set Selection Changed */
+	void OnSelectionChanged(TSharedPtr<FRigTreeBone> Selection, ESelectInfo::Type SelectInfo);
+
+	TSharedPtr< SWidget > CreateContextMenu();
+
+	// FEditorUndoClient
+	virtual void PostUndo(bool bSuccess) override;
+	virtual void PostRedo(bool bSuccess) override;
+
+	// reply to a drag operation
+	FReply OnDragDetected(const FGeometry& MyGeometry, const FPointerEvent& MouseEvent);
+
+private:
+	/** Our owning control rig editor */
+	TWeakPtr<FControlRigEditor> ControlRigEditor;
+
+	/** Search box widget */
+	TSharedPtr<SSearchBox> FilterBox;
+	FText FilterText;
+
+	void OnFilterTextChanged(const FText& SearchText);
+
+	/** Tree view widget */
+	TSharedPtr<STreeView<TSharedPtr<FRigTreeBone>>> TreeView;
+
+	/** Backing array for tree view */
+	TArray<TSharedPtr<FRigTreeBone>> RootBones;
+
+	/** Backing array for tree view (filtered, displayed) */
+	TArray<TSharedPtr<FRigTreeBone>> FilteredRootBones;
+
+	TWeakObjectPtr<UControlRigBlueprint> ControlRigBlueprint;
+	
+	/** Command list we bind to */
+	TSharedPtr<FUICommandList> CommandList;
+
+	bool IsMultiSelected() const;
+	bool IsSingleSelected() const;
+
+	FRigBoneHierarchy* GetHierarchy() const;
+	FRigBoneHierarchy* GetInstanceHierarchy() const;
+
+	void ImportHierarchy(const FAssetData& InAssetData);
+	void CreateImportMenu(FMenuBuilder& MenuBuilder);
+	void CreateRefreshMenu(FMenuBuilder& MenuBuilder);
+	bool ShouldFilterOnImport(const FAssetData& AssetData) const;
+	void RefreshHierarchy(const FAssetData& InAssetData);
+
+	FName CreateUniqueName(const FName& InBaseName) const;
+
+	void SetExpansionRecursive(TSharedPtr<FRigTreeBone> InBones);
+
+	void ClearDetailPanel() const;
+	void SelectBone(const FName& BoneName) const;
+	void OnBoneHierarchyChanged(FRigHierarchyContainer* Container, ERigHierarchyElementType ElementType, const FName& InName);
+
+public:
+	bool RenameBone(const FName& OldName, const FName& NewName);
+	bool OnVerifyNameChanged(const FName& OldName, const FName& NewName, FText& OutErrorMessage);
+
+	friend class SRigBoneHierarchyItem;
+};
\ No newline at end of file
diff --git a/Engine/Plugins/Experimental/ControlRig/Source/ControlRigEditor/Private/Editor/SRigHierarchy.cpp b/Engine/Plugins/Experimental/ControlRig/Source/ControlRigEditor/Private/Editor/SRigHierarchy.cpp
deleted file mode 100644
diff --git a/Engine/Plugins/Experimental/ControlRig/Source/ControlRigEditor/Private/Editor/SRigHierarchy.h b/Engine/Plugins/Experimental/ControlRig/Source/ControlRigEditor/Private/Editor/SRigHierarchy.h
deleted file mode 100644
