ControlRig: Restructuring CurveContainer, delegates outstanding. Renaming Hierarchy to BoneHierarchy in tests.

diff --git a/Engine/Plugins/Experimental/ControlRig/Source/ControlRig/Private/Rigs/RigCurveContainer.cpp b/Engine/Plugins/Experimental/ControlRig/Source/ControlRig/Private/Rigs/RigCurveContainer.cpp
--- a/Engine/Plugins/Experimental/ControlRig/Source/ControlRig/Private/Rigs/RigCurveContainer.cpp
+++ b/Engine/Plugins/Experimental/ControlRig/Source/ControlRig/Private/Rigs/RigCurveContainer.cpp
@@ -8,6 +8,135 @@
 // FRigCurveContainer
 ////////////////////////////////////////////////////////////////////////////////
 
+FRigCurveContainer::FRigCurveContainer()
+{
+}
+
+FRigCurveContainer& FRigCurveContainer::operator= (const FRigCurveContainer &InOther)
+{
+	/*
+#if WITH_EDITOR
+	for (int32 Index = Num() - 1; Index >= 0; Index--)
+	{
+		FRigSpace SpaceToRemove = Spaces[Index];
+		OnSpaceRemoved.Broadcast(Container, ERigHierarchyElementType::Space, SpaceToRemove.Name);
+	}
+#endif
+	*/
+
+	Curves.Reset();
+	Curves.Append(InOther.Curves);
+	NameToIndexMapping.Reset();
+	RefreshMapping();
+
+	/*
+#if WITH_EDITOR
+	for (const FRigSpace& SpaceAdded : Spaces)
+	{
+		OnSpaceAdded.Broadcast(Container, ERigHierarchyElementType::Space, SpaceAdded.Name);
+	}
+#endif
+	*/
+
+	return *this;
+}
+
+void FRigCurveContainer::Add(const FName& InNewName)
+{
+	DECLARE_SCOPE_HIERARCHICAL_COUNTER_FUNC()
+
+	int32 Found = GetIndex(InNewName);
+	if (Found == INDEX_NONE)
+	{
+		FRigCurve NewCurve;
+		NewCurve.Name = InNewName;
+		NewCurve.Value = 0.f;
+		Curves.Add(NewCurve);
+		RefreshMapping();
+	}
+}
+
+void FRigCurveContainer::Remove(const FName& InName)
+{
+	DECLARE_SCOPE_HIERARCHICAL_COUNTER_FUNC()
+
+	int32 IndexToDelete = GetIndex(InName);
+	if (IndexToDelete != INDEX_NONE)
+	{
+		Curves.RemoveAt(IndexToDelete);
+		RefreshMapping();
+	}
+}
+
+FName FRigCurveContainer::GetName(int32 InIndex) const
+{
+	if (Curves.IsValidIndex(InIndex))
+	{
+		return Curves[InIndex].Name;
+	}
+
+	return NAME_None;
+}
+
+int32 FRigCurveContainer::GetIndexSlow(const FName& InName) const
+{
+	for (int32 Index = 0; Index < Curves.Num(); ++Index)
+	{
+		if (Curves[Index].Name == InName)
+		{
+			return Index;
+		}
+	}
+
+	return INDEX_NONE;
+}
+
+void FRigCurveContainer::SetValue(const FName& InName, const float InValue)
+{
+	SetValue(GetIndex(InName), InValue);
+}
+
+void FRigCurveContainer::SetValue(int32 InIndex, const float InValue)
+{
+	if (Curves.IsValidIndex(InIndex))
+	{
+		FRigCurve& Curve = Curves[InIndex];
+		Curve.Value = InValue;
+	}
+}
+
+float FRigCurveContainer::GetValue(const FName& InName) const
+{
+	return GetValue(GetIndex(InName));
+}
+
+float FRigCurveContainer::GetValue(int32 InIndex) const
+{
+	if (Curves.IsValidIndex(InIndex))
+	{
+		return Curves[InIndex].Value;
+	}
+
+	return 0.f;
+}
+
+void FRigCurveContainer::Rename(const FName& InOldName, const FName& InNewName)
+{
+	DECLARE_SCOPE_HIERARCHICAL_COUNTER_FUNC()
+
+	if (InOldName != InNewName)
+	{
+		const int32 Found = GetIndex(InOldName);
+		const int32 NewNameFound = GetIndex(InNewName);
+		// if I have new name, and didn't find new name
+		if (Found != INDEX_NONE && NewNameFound == INDEX_NONE)
+		{
+			Curves[Found].Name = InNewName;
+			RefreshMapping();
+		}
+	}
+}
+
 void FRigCurveContainer::RefreshMapping()
 {
 	DECLARE_SCOPE_HIERARCHICAL_COUNTER_FUNC()
diff --git a/Engine/Plugins/Experimental/ControlRig/Source/ControlRig/Private/Units/Deprecated/RigUnit_ApplyFK.cpp b/Engine/Plugins/Experimental/ControlRig/Source/ControlRig/Private/Units/Deprecated/RigUnit_ApplyFK.cpp
--- a/Engine/Plugins/Experimental/ControlRig/Source/ControlRig/Private/Units/Deprecated/RigUnit_ApplyFK.cpp
+++ b/Engine/Plugins/Experimental/ControlRig/Source/ControlRig/Private/Units/Deprecated/RigUnit_ApplyFK.cpp
@@ -63,8 +63,8 @@ FTransform FRigUnit_ApplyFK::GetBaseTransform(int32 BoneIndex, const FRigBoneHie
 
 IMPLEMENT_RIGUNIT_AUTOMATION_TEST(FRigUnit_ApplyFK)
 {
-	Hierarchy.Add(TEXT("Root"), NAME_None, FTransform(FVector(1.f, 0.f, 0.f)));
-	Hierarchy.Add(TEXT("BoneA"), TEXT("Root"), FTransform(FVector(1.f, 2.f, 3.f)));
+	BoneHierarchy.Add(TEXT("Root"), NAME_None, FTransform(FVector(1.f, 0.f, 0.f)));
+	BoneHierarchy.Add(TEXT("BoneA"), TEXT("Root"), FTransform(FVector(1.f, 2.f, 3.f)));
 
 	Unit.ExecuteContext = ExecuteContext;
 	Unit.Joint = TEXT("BoneA");
@@ -72,25 +72,25 @@ IMPLEMENT_RIGUNIT_AUTOMATION_TEST(FRigUnit_ApplyFK)
 	Unit.ApplyTransformSpace = ETransformSpaceMode::GlobalSpace;
 	Unit.Transform = FTransform(FVector(0.f, 5.f, 0.f));
 
-	Hierarchy.Initialize();
+	BoneHierarchy.Initialize();
 	Execute();
-	AddErrorIfFalse(Hierarchy.GetGlobalTransform(1).GetTranslation().Equals(FVector(0.f, 5.f, 0.f)), TEXT("unexpected global transform"));
-	AddErrorIfFalse(Hierarchy.GetLocalTransform(1).GetTranslation().Equals(FVector(-1.f, 5.f, 0.f)), TEXT("unexpected local transform"));
+	AddErrorIfFalse(BoneHierarchy.GetGlobalTransform(1).GetTranslation().Equals(FVector(0.f, 5.f, 0.f)), TEXT("unexpected global transform"));
+	AddErrorIfFalse(BoneHierarchy.GetLocalTransform(1).GetTranslation().Equals(FVector(-1.f, 5.f, 0.f)), TEXT("unexpected local transform"));
 
 	Unit.ApplyTransformMode = EApplyTransformMode::Override;
 	Unit.ApplyTransformSpace = ETransformSpaceMode::LocalSpace;
 
-	Hierarchy.Initialize();
+	BoneHierarchy.Initialize();
 	Execute();
-	AddErrorIfFalse(Hierarchy.GetGlobalTransform(1).GetTranslation().Equals(FVector(1.f, 5.f, 0.f)), TEXT("unexpected global transform"));
-	AddErrorIfFalse(Hierarchy.GetLocalTransform(1).GetTranslation().Equals(FVector(0.f, 5.f, 0.f)), TEXT("unexpected local transform"));
+	AddErrorIfFalse(BoneHierarchy.GetGlobalTransform(1).GetTranslation().Equals(FVector(1.f, 5.f, 0.f)), TEXT("unexpected global transform"));
+	AddErrorIfFalse(BoneHierarchy.GetLocalTransform(1).GetTranslation().Equals(FVector(0.f, 5.f, 0.f)), TEXT("unexpected local transform"));
 
 	Unit.ApplyTransformMode = EApplyTransformMode::Additive;
 
-	Hierarchy.Initialize();
+	BoneHierarchy.Initialize();
 	Execute();
-	AddErrorIfFalse(Hierarchy.GetGlobalTransform(1).GetTranslation().Equals(FVector(1.f, 7.f, 3.f)), TEXT("unexpected global transform"));
-	AddErrorIfFalse(Hierarchy.GetLocalTransform(1).GetTranslation().Equals(FVector(0.f, 7.f, 3.f)), TEXT("unexpected local transform"));
+	AddErrorIfFalse(BoneHierarchy.GetGlobalTransform(1).GetTranslation().Equals(FVector(1.f, 7.f, 3.f)), TEXT("unexpected global transform"));
+	AddErrorIfFalse(BoneHierarchy.GetLocalTransform(1).GetTranslation().Equals(FVector(0.f, 7.f, 3.f)), TEXT("unexpected local transform"));
 	return true;
 }
 #endif
\ No newline at end of file
diff --git a/Engine/Plugins/Experimental/ControlRig/Source/ControlRig/Private/Units/Hierarchy/RigUnit_AddBoneTransform.cpp b/Engine/Plugins/Experimental/ControlRig/Source/ControlRig/Private/Units/Hierarchy/RigUnit_AddBoneTransform.cpp
--- a/Engine/Plugins/Experimental/ControlRig/Source/ControlRig/Private/Units/Hierarchy/RigUnit_AddBoneTransform.cpp
+++ b/Engine/Plugins/Experimental/ControlRig/Source/ControlRig/Private/Units/Hierarchy/RigUnit_AddBoneTransform.cpp
@@ -50,30 +50,30 @@ void FRigUnit_AddBoneTransform::Execute(const FRigUnitContext& Context)
 
 IMPLEMENT_RIGUNIT_AUTOMATION_TEST(FRigUnit_AddBoneTransform)
 {
-	Hierarchy.Add(TEXT("Root"), NAME_None, FTransform(FVector(1.f, 0.f, 0.f)));
-	Hierarchy.Add(TEXT("BoneA"), TEXT("Root"), FTransform(FVector(1.f, 2.f, 0.f)));
-	Hierarchy.Initialize();
+	BoneHierarchy.Add(TEXT("Root"), NAME_None, FTransform(FVector(1.f, 0.f, 0.f)));
+	BoneHierarchy.Add(TEXT("BoneA"), TEXT("Root"), FTransform(FVector(1.f, 2.f, 0.f)));
+	BoneHierarchy.Initialize();
 	Unit.ExecuteContext.HierarchyReference = HierarchyRef;
 
-	Hierarchy.ResetTransforms();
+	BoneHierarchy.ResetTransforms();
 	Unit.Bone = TEXT("BoneA");
 	Unit.Transform = FTransform(FVector(0.f, 0.f, 7.f));
 	Unit.bPropagateToChildren = false;
 	InitAndExecute();
-	AddErrorIfFalse(Hierarchy.GetGlobalTransform(0).GetTranslation().Equals(FVector(1.f, 0.f, 0.f)), TEXT("unexpected transform"));
-	AddErrorIfFalse(Hierarchy.GetGlobalTransform(1).GetTranslation().Equals(FVector(1.f, 2.f, 7.f)), TEXT("unexpected transform"));
+	AddErrorIfFalse(BoneHierarchy.GetGlobalTransform(0).GetTranslation().Equals(FVector(1.f, 0.f, 0.f)), TEXT("unexpected transform"));
+	AddErrorIfFalse(BoneHierarchy.GetGlobalTransform(1).GetTranslation().Equals(FVector(1.f, 2.f, 7.f)), TEXT("unexpected transform"));
 
 	Unit.Bone = TEXT("Root");
-	Hierarchy.ResetTransforms();
+	BoneHierarchy.ResetTransforms();
 	InitAndExecute();
-	AddErrorIfFalse(Hierarchy.GetGlobalTransform(0).GetTranslation().Equals(FVector(1.f, 0.f, 7.f)), TEXT("unexpected transform"));
-	AddErrorIfFalse(Hierarchy.GetGlobalTransform(1).GetTranslation().Equals(FVector(1.f, 2.f, 0.f)), TEXT("unexpected transform"));
+	AddErrorIfFalse(BoneHierarchy.GetGlobalTransform(0).GetTranslation().Equals(FVector(1.f, 0.f, 7.f)), TEXT("unexpected transform"));
+	AddErrorIfFalse(BoneHierarchy.GetGlobalTransform(1).GetTranslation().Equals(FVector(1.f, 2.f, 0.f)), TEXT("unexpected transform"));
 
 	Unit.bPropagateToChildren = true;
-	Hierarchy.ResetTransforms();
+	BoneHierarchy.ResetTransforms();
 	InitAndExecute();
-	AddErrorIfFalse(Hierarchy.GetGlobalTransform(0).GetTranslation().Equals(FVector(1.f, 0.f, 7.f)), TEXT("unexpected transform"));
-	AddErrorIfFalse(Hierarchy.GetGlobalTransform(1).GetTranslation().Equals(FVector(1.f, 2.f, 7.f)), TEXT("unexpected transform"));
+	AddErrorIfFalse(BoneHierarchy.GetGlobalTransform(0).GetTranslation().Equals(FVector(1.f, 0.f, 7.f)), TEXT("unexpected transform"));
+	AddErrorIfFalse(BoneHierarchy.GetGlobalTransform(1).GetTranslation().Equals(FVector(1.f, 2.f, 7.f)), TEXT("unexpected transform"));
 
 	return true;
 }
diff --git a/Engine/Plugins/Experimental/ControlRig/Source/ControlRig/Private/Units/Hierarchy/RigUnit_GetBoneTransform.cpp b/Engine/Plugins/Experimental/ControlRig/Source/ControlRig/Private/Units/Hierarchy/RigUnit_GetBoneTransform.cpp
--- a/Engine/Plugins/Experimental/ControlRig/Source/ControlRig/Private/Units/Hierarchy/RigUnit_GetBoneTransform.cpp
+++ b/Engine/Plugins/Experimental/ControlRig/Source/ControlRig/Private/Units/Hierarchy/RigUnit_GetBoneTransform.cpp
@@ -57,9 +57,9 @@ void FRigUnit_GetBoneTransform::Execute(const FRigUnitContext& Context)
 
 IMPLEMENT_RIGUNIT_AUTOMATION_TEST(FRigUnit_GetBoneTransform)
 {
-	Hierarchy.Add(TEXT("Root"), NAME_None, FTransform(FVector(1.f, 0.f, 0.f)));
-	Hierarchy.Add(TEXT("BoneA"), TEXT("Root"), FTransform(FVector(1.f, 2.f, 3.f)));
-	Hierarchy.Initialize();
+	BoneHierarchy.Add(TEXT("Root"), NAME_None, FTransform(FVector(1.f, 0.f, 0.f)));
+	BoneHierarchy.Add(TEXT("BoneA"), TEXT("Root"), FTransform(FVector(1.f, 2.f, 3.f)));
+	BoneHierarchy.Initialize();
 
 	Unit.Bone = TEXT("Unknown");
 	Unit.Space = EBoneGetterSetterMode::GlobalSpace;
diff --git a/Engine/Plugins/Experimental/ControlRig/Source/ControlRig/Private/Units/Hierarchy/RigUnit_GetInitialBoneTransform.cpp b/Engine/Plugins/Experimental/ControlRig/Source/ControlRig/Private/Units/Hierarchy/RigUnit_GetInitialBoneTransform.cpp
--- a/Engine/Plugins/Experimental/ControlRig/Source/ControlRig/Private/Units/Hierarchy/RigUnit_GetInitialBoneTransform.cpp
+++ b/Engine/Plugins/Experimental/ControlRig/Source/ControlRig/Private/Units/Hierarchy/RigUnit_GetInitialBoneTransform.cpp
@@ -63,9 +63,9 @@ void FRigUnit_GetInitialBoneTransform::Execute(const FRigUnitContext& Context)
 
 IMPLEMENT_RIGUNIT_AUTOMATION_TEST(FRigUnit_GetInitialBoneTransform)
 {
-	Hierarchy.Add(TEXT("Root"), NAME_None, FTransform(FVector(1.f, 0.f, 0.f)));
-	Hierarchy.Add(TEXT("BoneA"), TEXT("Root"), FTransform(FVector(1.f, 2.f, 3.f)));
-	Hierarchy.Initialize();
+	BoneHierarchy.Add(TEXT("Root"), NAME_None, FTransform(FVector(1.f, 0.f, 0.f)));
+	BoneHierarchy.Add(TEXT("BoneA"), TEXT("Root"), FTransform(FVector(1.f, 2.f, 3.f)));
+	BoneHierarchy.Initialize();
 
 	Unit.Bone = TEXT("Root");
 	Unit.Space = EBoneGetterSetterMode::GlobalSpace;
diff --git a/Engine/Plugins/Experimental/ControlRig/Source/ControlRig/Private/Units/Hierarchy/RigUnit_GetRelativeBoneTransform.cpp b/Engine/Plugins/Experimental/ControlRig/Source/ControlRig/Private/Units/Hierarchy/RigUnit_GetRelativeBoneTransform.cpp
--- a/Engine/Plugins/Experimental/ControlRig/Source/ControlRig/Private/Units/Hierarchy/RigUnit_GetRelativeBoneTransform.cpp
+++ b/Engine/Plugins/Experimental/ControlRig/Source/ControlRig/Private/Units/Hierarchy/RigUnit_GetRelativeBoneTransform.cpp
@@ -44,10 +44,10 @@ void FRigUnit_GetRelativeBoneTransform::Execute(const FRigUnitContext& Context)
 
 IMPLEMENT_RIGUNIT_AUTOMATION_TEST(FRigUnit_GetRelativeBoneTransform)
 {
-	Hierarchy.Add(TEXT("Root"), NAME_None, FTransform(FVector(1.f, 0.f, 0.f)));
-	Hierarchy.Add(TEXT("BoneA"), TEXT("Root"), FTransform(FVector(1.f, 2.f, 3.f)));
-	Hierarchy.Add(TEXT("BoneB"), TEXT("Root"), FTransform(FVector(-4.f, 0.f, 0.f)));
-	Hierarchy.Initialize();
+	BoneHierarchy.Add(TEXT("Root"), NAME_None, FTransform(FVector(1.f, 0.f, 0.f)));
+	BoneHierarchy.Add(TEXT("BoneA"), TEXT("Root"), FTransform(FVector(1.f, 2.f, 3.f)));
+	BoneHierarchy.Add(TEXT("BoneB"), TEXT("Root"), FTransform(FVector(-4.f, 0.f, 0.f)));
+	BoneHierarchy.Initialize();
 
 	Unit.Bone = TEXT("Unknown");
 	Unit.Space = TEXT("Root");
diff --git a/Engine/Plugins/Experimental/ControlRig/Source/ControlRig/Private/Units/Hierarchy/RigUnit_SetBoneRotation.cpp b/Engine/Plugins/Experimental/ControlRig/Source/ControlRig/Private/Units/Hierarchy/RigUnit_SetBoneRotation.cpp
--- a/Engine/Plugins/Experimental/ControlRig/Source/ControlRig/Private/Units/Hierarchy/RigUnit_SetBoneRotation.cpp
+++ b/Engine/Plugins/Experimental/ControlRig/Source/ControlRig/Private/Units/Hierarchy/RigUnit_SetBoneRotation.cpp
@@ -62,59 +62,59 @@ void FRigUnit_SetBoneRotation::Execute(const FRigUnitContext& Context)
 
 IMPLEMENT_RIGUNIT_AUTOMATION_TEST(FRigUnit_SetBoneRotation)
 {
-	Hierarchy.Add(TEXT("Root"), NAME_None, FTransform(FQuat(FVector(-1.f, 0.f, 0.f), 0.1f)));
-	Hierarchy.Add(TEXT("BoneA"), TEXT("Root"), FTransform(FQuat(FVector(-1.f, 0.f, 0.f), 0.5f)));
-	Hierarchy.Add(TEXT("BoneB"), TEXT("BoneA"), FTransform(FQuat(FVector(-1.f, 0.f, 0.f), 0.7f)));
-	Hierarchy.Initialize();
+	BoneHierarchy.Add(TEXT("Root"), NAME_None, FTransform(FQuat(FVector(-1.f, 0.f, 0.f), 0.1f)));
+	BoneHierarchy.Add(TEXT("BoneA"), TEXT("Root"), FTransform(FQuat(FVector(-1.f, 0.f, 0.f), 0.5f)));
+	BoneHierarchy.Add(TEXT("BoneB"), TEXT("BoneA"), FTransform(FQuat(FVector(-1.f, 0.f, 0.f), 0.7f)));
+	BoneHierarchy.Initialize();
 
 	Unit.ExecuteContext.HierarchyReference = HierarchyRef;
 
-	Hierarchy.ResetTransforms();
+	BoneHierarchy.ResetTransforms();
 	Unit.Bone = TEXT("Root");
 	Unit.Space = EBoneGetterSetterMode::GlobalSpace;
 	Unit.Rotation = FQuat(FVector(-1.f, 0.f, 0.f), 0.25f);
 	Unit.bPropagateToChildren = false;
 	InitAndExecute();
-	AddErrorIfFalse(FMath::IsNearlyEqual(Hierarchy.GetGlobalTransform(0).GetRotation().GetAngle(), 0.25f, 0.001f), FString::Printf(TEXT("unexpected angle %.04f"), Hierarchy.GetGlobalTransform(0).GetRotation().GetAngle()));
-	AddErrorIfFalse(FMath::IsNearlyEqual(Hierarchy.GetGlobalTransform(1).GetRotation().GetAngle(), 0.5f, 0.001f), FString::Printf(TEXT("unexpected angle %.04f"), Hierarchy.GetGlobalTransform(1).GetRotation().GetAngle()));
-	AddErrorIfFalse(FMath::IsNearlyEqual(Hierarchy.GetGlobalTransform(2).GetRotation().GetAngle(), 0.7f, 0.001f), FString::Printf(TEXT("unexpected angle %.04f"), Hierarchy.GetGlobalTransform(2).GetRotation().GetAngle()));
+	AddErrorIfFalse(FMath::IsNearlyEqual(BoneHierarchy.GetGlobalTransform(0).GetRotation().GetAngle(), 0.25f, 0.001f), FString::Printf(TEXT("unexpected angle %.04f"), BoneHierarchy.GetGlobalTransform(0).GetRotation().GetAngle()));
+	AddErrorIfFalse(FMath::IsNearlyEqual(BoneHierarchy.GetGlobalTransform(1).GetRotation().GetAngle(), 0.5f, 0.001f), FString::Printf(TEXT("unexpected angle %.04f"), BoneHierarchy.GetGlobalTransform(1).GetRotation().GetAngle()));
+	AddErrorIfFalse(FMath::IsNearlyEqual(BoneHierarchy.GetGlobalTransform(2).GetRotation().GetAngle(), 0.7f, 0.001f), FString::Printf(TEXT("unexpected angle %.04f"), BoneHierarchy.GetGlobalTransform(2).GetRotation().GetAngle()));
 
-	Hierarchy.ResetTransforms();
+	BoneHierarchy.ResetTransforms();
 	Unit.Space = EBoneGetterSetterMode::LocalSpace;
 	InitAndExecute();
-	AddErrorIfFalse(FMath::IsNearlyEqual(Hierarchy.GetGlobalTransform(0).GetRotation().GetAngle(), 0.25f, 0.001f), FString::Printf(TEXT("unexpected angle %.04f"), Hierarchy.GetGlobalTransform(0).GetRotation().GetAngle()));
-	AddErrorIfFalse(FMath::IsNearlyEqual(Hierarchy.GetGlobalTransform(1).GetRotation().GetAngle(), 0.5f, 0.001f), FString::Printf(TEXT("unexpected angle %.04f"), Hierarchy.GetGlobalTransform(1).GetRotation().GetAngle()));
-	AddErrorIfFalse(FMath::IsNearlyEqual(Hierarchy.GetGlobalTransform(2).GetRotation().GetAngle(), 0.7f, 0.001f), FString::Printf(TEXT("unexpected angle %.04f"), Hierarchy.GetGlobalTransform(2).GetRotation().GetAngle()));
+	AddErrorIfFalse(FMath::IsNearlyEqual(BoneHierarchy.GetGlobalTransform(0).GetRotation().GetAngle(), 0.25f, 0.001f), FString::Printf(TEXT("unexpected angle %.04f"), BoneHierarchy.GetGlobalTransform(0).GetRotation().GetAngle()));
+	AddErrorIfFalse(FMath::IsNearlyEqual(BoneHierarchy.GetGlobalTransform(1).GetRotation().GetAngle(), 0.5f, 0.001f), FString::Printf(TEXT("unexpected angle %.04f"), BoneHierarchy.GetGlobalTransform(1).GetRotation().GetAngle()));
+	AddErrorIfFalse(FMath::IsNearlyEqual(BoneHierarchy.GetGlobalTransform(2).GetRotation().GetAngle(), 0.7f, 0.001f), FString::Printf(TEXT("unexpected angle %.04f"), BoneHierarchy.GetGlobalTransform(2).GetRotation().GetAngle()));
 
-	Hierarchy.ResetTransforms();
+	BoneHierarchy.ResetTransforms();
 	Unit.bPropagateToChildren = true;
 	InitAndExecute();
-	AddErrorIfFalse(FMath::IsNearlyEqual(Hierarchy.GetGlobalTransform(0).GetRotation().GetAngle(), 0.25f, 0.001f), FString::Printf(TEXT("unexpected angle %.04f"), Hierarchy.GetGlobalTransform(0).GetRotation().GetAngle()));
-	AddErrorIfFalse(FMath::IsNearlyEqual(Hierarchy.GetGlobalTransform(1).GetRotation().GetAngle(), 0.65f, 0.001f), FString::Printf(TEXT("unexpected angle %.04f"), Hierarchy.GetGlobalTransform(1).GetRotation().GetAngle()));
-	AddErrorIfFalse(FMath::IsNearlyEqual(Hierarchy.GetGlobalTransform(2).GetRotation().GetAngle(), 0.85f, 0.001f), FString::Printf(TEXT("unexpected angle %.04f"), Hierarchy.GetGlobalTransform(2).GetRotation().GetAngle()));
+	AddErrorIfFalse(FMath::IsNearlyEqual(BoneHierarchy.GetGlobalTransform(0).GetRotation().GetAngle(), 0.25f, 0.001f), FString::Printf(TEXT("unexpected angle %.04f"), BoneHierarchy.GetGlobalTransform(0).GetRotation().GetAngle()));
+	AddErrorIfFalse(FMath::IsNearlyEqual(BoneHierarchy.GetGlobalTransform(1).GetRotation().GetAngle(), 0.65f, 0.001f), FString::Printf(TEXT("unexpected angle %.04f"), BoneHierarchy.GetGlobalTransform(1).GetRotation().GetAngle()));
+	AddErrorIfFalse(FMath::IsNearlyEqual(BoneHierarchy.GetGlobalTransform(2).GetRotation().GetAngle(), 0.85f, 0.001f), FString::Printf(TEXT("unexpected angle %.04f"), BoneHierarchy.GetGlobalTransform(2).GetRotation().GetAngle()));
 
-	Hierarchy.ResetTransforms();
+	BoneHierarchy.ResetTransforms();
 	Unit.Bone = TEXT("BoneA");
 	Unit.Space = EBoneGetterSetterMode::GlobalSpace;
 	Unit.bPropagateToChildren = false;
 	InitAndExecute();
-	AddErrorIfFalse(FMath::IsNearlyEqual(Hierarchy.GetGlobalTransform(0).GetRotation().GetAngle(), 0.1f, 0.001f), FString::Printf(TEXT("unexpected angle %.04f"), Hierarchy.GetGlobalTransform(0).GetRotation().GetAngle()));
-	AddErrorIfFalse(FMath::IsNearlyEqual(Hierarchy.GetGlobalTransform(1).GetRotation().GetAngle(), 0.25f, 0.001f), FString::Printf(TEXT("unexpected angle %.04f"), Hierarchy.GetGlobalTransform(1).GetRotation().GetAngle()));
-	AddErrorIfFalse(FMath::IsNearlyEqual(Hierarchy.GetGlobalTransform(2).GetRotation().GetAngle(), 0.7f, 0.001f), FString::Printf(TEXT("unexpected angle %.04f"), Hierarchy.GetGlobalTransform(2).GetRotation().GetAngle()));
+	AddErrorIfFalse(FMath::IsNearlyEqual(BoneHierarchy.GetGlobalTransform(0).GetRotation().GetAngle(), 0.1f, 0.001f), FString::Printf(TEXT("unexpected angle %.04f"), BoneHierarchy.GetGlobalTransform(0).GetRotation().GetAngle()));
+	AddErrorIfFalse(FMath::IsNearlyEqual(BoneHierarchy.GetGlobalTransform(1).GetRotation().GetAngle(), 0.25f, 0.001f), FString::Printf(TEXT("unexpected angle %.04f"), BoneHierarchy.GetGlobalTransform(1).GetRotation().GetAngle()));
+	AddErrorIfFalse(FMath::IsNearlyEqual(BoneHierarchy.GetGlobalTransform(2).GetRotation().GetAngle(), 0.7f, 0.001f), FString::Printf(TEXT("unexpected angle %.04f"), BoneHierarchy.GetGlobalTransform(2).GetRotation().GetAngle()));
 
-	Hierarchy.ResetTransforms();
+	BoneHierarchy.ResetTransforms();
 	Unit.Space = EBoneGetterSetterMode::LocalSpace;
 	InitAndExecute();
-	AddErrorIfFalse(FMath::IsNearlyEqual(Hierarchy.GetGlobalTransform(0).GetRotation().GetAngle(), 0.1f, 0.001f), FString::Printf(TEXT("unexpected angle %.04f"), Hierarchy.GetGlobalTransform(0).GetRotation().GetAngle()));
-	AddErrorIfFalse(FMath::IsNearlyEqual(Hierarchy.GetGlobalTransform(1).GetRotation().GetAngle(), 0.35f, 0.001f), FString::Printf(TEXT("unexpected angle %.04f"), Hierarchy.GetGlobalTransform(1).GetRotation().GetAngle()));
-	AddErrorIfFalse(FMath::IsNearlyEqual(Hierarchy.GetGlobalTransform(2).GetRotation().GetAngle(), 0.7f, 0.001f), FString::Printf(TEXT("unexpected angle %.04f"), Hierarchy.GetGlobalTransform(2).GetRotation().GetAngle()));
+	AddErrorIfFalse(FMath::IsNearlyEqual(BoneHierarchy.GetGlobalTransform(0).GetRotation().GetAngle(), 0.1f, 0.001f), FString::Printf(TEXT("unexpected angle %.04f"), BoneHierarchy.GetGlobalTransform(0).GetRotation().GetAngle()));
+	AddErrorIfFalse(FMath::IsNearlyEqual(BoneHierarchy.GetGlobalTransform(1).GetRotation().GetAngle(), 0.35f, 0.001f), FString::Printf(TEXT("unexpected angle %.04f"), BoneHierarchy.GetGlobalTransform(1).GetRotation().GetAngle()));
+	AddErrorIfFalse(FMath::IsNearlyEqual(BoneHierarchy.GetGlobalTransform(2).GetRotation().GetAngle(), 0.7f, 0.001f), FString::Printf(TEXT("unexpected angle %.04f"), BoneHierarchy.GetGlobalTransform(2).GetRotation().GetAngle()));
 
-	Hierarchy.ResetTransforms();
+	BoneHierarchy.ResetTransforms();
 	Unit.bPropagateToChildren = true;
 	InitAndExecute();
-	AddErrorIfFalse(FMath::IsNearlyEqual(Hierarchy.GetGlobalTransform(0).GetRotation().GetAngle(), 0.1f, 0.001f), FString::Printf(TEXT("unexpected angle %.04f"), Hierarchy.GetGlobalTransform(0).GetRotation().GetAngle()));
-	AddErrorIfFalse(FMath::IsNearlyEqual(Hierarchy.GetGlobalTransform(1).GetRotation().GetAngle(), 0.35f, 0.001f), FString::Printf(TEXT("unexpected angle %.04f"), Hierarchy.GetGlobalTransform(1).GetRotation().GetAngle()));
-	AddErrorIfFalse(FMath::IsNearlyEqual(Hierarchy.GetGlobalTransform(2).GetRotation().GetAngle(), 0.55f, 0.001f), FString::Printf(TEXT("unexpected angle %.04f"), Hierarchy.GetGlobalTransform(2).GetRotation().GetAngle()));
+	AddErrorIfFalse(FMath::IsNearlyEqual(BoneHierarchy.GetGlobalTransform(0).GetRotation().GetAngle(), 0.1f, 0.001f), FString::Printf(TEXT("unexpected angle %.04f"), BoneHierarchy.GetGlobalTransform(0).GetRotation().GetAngle()));
+	AddErrorIfFalse(FMath::IsNearlyEqual(BoneHierarchy.GetGlobalTransform(1).GetRotation().GetAngle(), 0.35f, 0.001f), FString::Printf(TEXT("unexpected angle %.04f"), BoneHierarchy.GetGlobalTransform(1).GetRotation().GetAngle()));
+	AddErrorIfFalse(FMath::IsNearlyEqual(BoneHierarchy.GetGlobalTransform(2).GetRotation().GetAngle(), 0.55f, 0.001f), FString::Printf(TEXT("unexpected angle %.04f"), BoneHierarchy.GetGlobalTransform(2).GetRotation().GetAngle()));
 
 	return true;
 }
diff --git a/Engine/Plugins/Experimental/ControlRig/Source/ControlRig/Private/Units/Hierarchy/RigUnit_SetBoneTransform.cpp b/Engine/Plugins/Experimental/ControlRig/Source/ControlRig/Private/Units/Hierarchy/RigUnit_SetBoneTransform.cpp
--- a/Engine/Plugins/Experimental/ControlRig/Source/ControlRig/Private/Units/Hierarchy/RigUnit_SetBoneTransform.cpp
+++ b/Engine/Plugins/Experimental/ControlRig/Source/ControlRig/Private/Units/Hierarchy/RigUnit_SetBoneTransform.cpp
@@ -58,58 +58,58 @@ void FRigUnit_SetBoneTransform::Execute(const FRigUnitContext& Context)
 
 IMPLEMENT_RIGUNIT_AUTOMATION_TEST(FRigUnit_SetBoneTransform)
 {
-	Hierarchy.Add(TEXT("Root"), NAME_None, FTransform(FVector(1.f, 0.f, 0.f)));
-	Hierarchy.Add(TEXT("BoneA"), TEXT("Root"), FTransform(FVector(1.f, 2.f, 3.f)));
-	Hierarchy.Add(TEXT("BoneB"), TEXT("BoneA"), FTransform(FVector(1.f, 5.f, 3.f)));
-	Hierarchy.Initialize();
+	BoneHierarchy.Add(TEXT("Root"), NAME_None, FTransform(FVector(1.f, 0.f, 0.f)));
+	BoneHierarchy.Add(TEXT("BoneA"), TEXT("Root"), FTransform(FVector(1.f, 2.f, 3.f)));
+	BoneHierarchy.Add(TEXT("BoneB"), TEXT("BoneA"), FTransform(FVector(1.f, 5.f, 3.f)));
+	BoneHierarchy.Initialize();
 	Unit.ExecuteContext.HierarchyReference = HierarchyRef;
 
-	Hierarchy.ResetTransforms();
+	BoneHierarchy.ResetTransforms();
 	Unit.Bone = TEXT("Root");
 	Unit.Space = EBoneGetterSetterMode::GlobalSpace;
 	Unit.Transform = FTransform(FVector(0.f, 0.f, 7.f));
 	Unit.bPropagateToChildren = false;
 	InitAndExecute();
-	AddErrorIfFalse(Hierarchy.GetGlobalTransform(0).GetTranslation().Equals(FVector(0.f, 0.f, 7.f)), TEXT("unexpected transform"));
-	AddErrorIfFalse(Hierarchy.GetGlobalTransform(1).GetTranslation().Equals(FVector(1.f, 2.f, 3.f)), TEXT("unexpected transform"));
-	AddErrorIfFalse(Hierarchy.GetGlobalTransform(2).GetTranslation().Equals(FVector(1.f, 5.f, 3.f)), TEXT("unexpected transform"));
+	AddErrorIfFalse(BoneHierarchy.GetGlobalTransform(0).GetTranslation().Equals(FVector(0.f, 0.f, 7.f)), TEXT("unexpected transform"));
+	AddErrorIfFalse(BoneHierarchy.GetGlobalTransform(1).GetTranslation().Equals(FVector(1.f, 2.f, 3.f)), TEXT("unexpected transform"));
+	AddErrorIfFalse(BoneHierarchy.GetGlobalTransform(2).GetTranslation().Equals(FVector(1.f, 5.f, 3.f)), TEXT("unexpected transform"));
 
-	Hierarchy.ResetTransforms();
+	BoneHierarchy.ResetTransforms();
 	Unit.Space = EBoneGetterSetterMode::LocalSpace;
 	InitAndExecute();
-	AddErrorIfFalse(Hierarchy.GetGlobalTransform(0).GetTranslation().Equals(FVector(0.f, 0.f, 7.f)), TEXT("unexpected transform"));
-	AddErrorIfFalse(Hierarchy.GetGlobalTransform(1).GetTranslation().Equals(FVector(1.f, 2.f, 3.f)), TEXT("unexpected transform"));
-	AddErrorIfFalse(Hierarchy.GetGlobalTransform(2).GetTranslation().Equals(FVector(1.f, 5.f, 3.f)), TEXT("unexpected transform"));
+	AddErrorIfFalse(BoneHierarchy.GetGlobalTransform(0).GetTranslation().Equals(FVector(0.f, 0.f, 7.f)), TEXT("unexpected transform"));
+	AddErrorIfFalse(BoneHierarchy.GetGlobalTransform(1).GetTranslation().Equals(FVector(1.f, 2.f, 3.f)), TEXT("unexpected transform"));
+	AddErrorIfFalse(BoneHierarchy.GetGlobalTransform(2).GetTranslation().Equals(FVector(1.f, 5.f, 3.f)), TEXT("unexpected transform"));
 
-	Hierarchy.ResetTransforms();
+	BoneHierarchy.ResetTransforms();
 	Unit.bPropagateToChildren = true;
 	InitAndExecute();
-	AddErrorIfFalse(Hierarchy.GetGlobalTransform(0).GetTranslation().Equals(FVector(0.f, 0.f, 7.f)), TEXT("unexpected transform"));
-	AddErrorIfFalse(Hierarchy.GetGlobalTransform(1).GetTranslation().Equals(FVector(0.f, 2.f, 10.f)), TEXT("unexpected transform"));
-	AddErrorIfFalse(Hierarchy.GetGlobalTransform(2).GetTranslation().Equals(FVector(0.f, 5.f, 10.f)), TEXT("unexpected transform"));
+	AddErrorIfFalse(BoneHierarchy.GetGlobalTransform(0).GetTranslation().Equals(FVector(0.f, 0.f, 7.f)), TEXT("unexpected transform"));
+	AddErrorIfFalse(BoneHierarchy.GetGlobalTransform(1).GetTranslation().Equals(FVector(0.f, 2.f, 10.f)), TEXT("unexpected transform"));
+	AddErrorIfFalse(BoneHierarchy.GetGlobalTransform(2).GetTranslation().Equals(FVector(0.f, 5.f, 10.f)), TEXT("unexpected transform"));
 
-	Hierarchy.ResetTransforms();
+	BoneHierarchy.ResetTransforms();
 	Unit.Bone = TEXT("BoneA");
 	Unit.Space = EBoneGetterSetterMode::GlobalSpace;
 	Unit.bPropagateToChildren = false;
 	InitAndExecute();
-	AddErrorIfFalse(Hierarchy.GetGlobalTransform(0).GetTranslation().Equals(FVector(1.f, 0.f, 0.f)), TEXT("unexpected transform"));
-	AddErrorIfFalse(Hierarchy.GetGlobalTransform(1).GetTranslation().Equals(FVector(0.f, 0.f, 7.f)), TEXT("unexpected transform"));
-	AddErrorIfFalse(Hierarchy.GetGlobalTransform(2).GetTranslation().Equals(FVector(1.f, 5.f, 3.f)), TEXT("unexpected transform"));
+	AddErrorIfFalse(BoneHierarchy.GetGlobalTransform(0).GetTranslation().Equals(FVector(1.f, 0.f, 0.f)), TEXT("unexpected transform"));
+	AddErrorIfFalse(BoneHierarchy.GetGlobalTransform(1).GetTranslation().Equals(FVector(0.f, 0.f, 7.f)), TEXT("unexpected transform"));
+	AddErrorIfFalse(BoneHierarchy.GetGlobalTransform(2).GetTranslation().Equals(FVector(1.f, 5.f, 3.f)), TEXT("unexpected transform"));
 
-	Hierarchy.ResetTransforms();
+	BoneHierarchy.ResetTransforms();
 	Unit.Space = EBoneGetterSetterMode::LocalSpace;
 	InitAndExecute();
-	AddErrorIfFalse(Hierarchy.GetGlobalTransform(0).GetTranslation().Equals(FVector(1.f, 0.f, 0.f)), TEXT("unexpected transform"));
-	AddErrorIfFalse(Hierarchy.GetGlobalTransform(1).GetTranslation().Equals(FVector(1.f, 0.f, 7.f)), TEXT("unexpected transform"));
-	AddErrorIfFalse(Hierarchy.GetGlobalTransform(2).GetTranslation().Equals(FVector(1.f, 5.f, 3.f)), TEXT("unexpected transform"));
+	AddErrorIfFalse(BoneHierarchy.GetGlobalTransform(0).GetTranslation().Equals(FVector(1.f, 0.f, 0.f)), TEXT("unexpected transform"));
+	AddErrorIfFalse(BoneHierarchy.GetGlobalTransform(1).GetTranslation().Equals(FVector(1.f, 0.f, 7.f)), TEXT("unexpected transform"));
+	AddErrorIfFalse(BoneHierarchy.GetGlobalTransform(2).GetTranslation().Equals(FVector(1.f, 5.f, 3.f)), TEXT("unexpected transform"));
 
-	Hierarchy.ResetTransforms();
+	BoneHierarchy.ResetTransforms();
 	Unit.bPropagateToChildren = true;
 	InitAndExecute();
-	AddErrorIfFalse(Hierarchy.GetGlobalTransform(0).GetTranslation().Equals(FVector(1.f, 0.f, 0.f)), TEXT("unexpected transform"));
-	AddErrorIfFalse(Hierarchy.GetGlobalTransform(1).GetTranslation().Equals(FVector(1.f, 0.f, 7.f)), TEXT("unexpected transform"));
-	AddErrorIfFalse(Hierarchy.GetGlobalTransform(2).GetTranslation().Equals(FVector(1.f, 3.f, 7.f)), TEXT("unexpected transform"));
+	AddErrorIfFalse(BoneHierarchy.GetGlobalTransform(0).GetTranslation().Equals(FVector(1.f, 0.f, 0.f)), TEXT("unexpected transform"));
+	AddErrorIfFalse(BoneHierarchy.GetGlobalTransform(1).GetTranslation().Equals(FVector(1.f, 0.f, 7.f)), TEXT("unexpected transform"));
+	AddErrorIfFalse(BoneHierarchy.GetGlobalTransform(2).GetTranslation().Equals(FVector(1.f, 3.f, 7.f)), TEXT("unexpected transform"));
 
 	return true;
 }
diff --git a/Engine/Plugins/Experimental/ControlRig/Source/ControlRig/Private/Units/Hierarchy/RigUnit_SetCurveValue.cpp b/Engine/Plugins/Experimental/ControlRig/Source/ControlRig/Private/Units/Hierarchy/RigUnit_SetCurveValue.cpp
--- a/Engine/Plugins/Experimental/ControlRig/Source/ControlRig/Private/Units/Hierarchy/RigUnit_SetCurveValue.cpp
+++ b/Engine/Plugins/Experimental/ControlRig/Source/ControlRig/Private/Units/Hierarchy/RigUnit_SetCurveValue.cpp
@@ -41,8 +41,8 @@ void FRigUnit_SetCurveValue::Execute(const FRigUnitContext& Context)
 
 IMPLEMENT_RIGUNIT_AUTOMATION_TEST(FRigUnit_SetCurveValue)
 {
-	CurveContainer.AddCurve(TEXT("CurveA"));
-	CurveContainer.AddCurve(TEXT("CurveB"));
+	CurveContainer.Add(TEXT("CurveA"));
+	CurveContainer.Add(TEXT("CurveB"));
 	CurveContainer.Initialize();
 	Unit.ExecuteContext.CurveReference = CurveContainerRef;
 	
diff --git a/Engine/Plugins/Experimental/ControlRig/Source/ControlRig/Private/Units/Hierarchy/RigUnit_SetRelativeBoneTransform.cpp b/Engine/Plugins/Experimental/ControlRig/Source/ControlRig/Private/Units/Hierarchy/RigUnit_SetRelativeBoneTransform.cpp
--- a/Engine/Plugins/Experimental/ControlRig/Source/ControlRig/Private/Units/Hierarchy/RigUnit_SetRelativeBoneTransform.cpp
+++ b/Engine/Plugins/Experimental/ControlRig/Source/ControlRig/Private/Units/Hierarchy/RigUnit_SetRelativeBoneTransform.cpp
@@ -45,48 +45,48 @@ void FRigUnit_SetRelativeBoneTransform::Execute(const FRigUnitContext& Context)
 
 IMPLEMENT_RIGUNIT_AUTOMATION_TEST(FRigUnit_SetRelativeBoneTransform)
 {
-	Hierarchy.Add(TEXT("Root"), NAME_None, FTransform(FVector(1.f, 0.f, 0.f)));
-	Hierarchy.Add(TEXT("BoneA"), TEXT("Root"), FTransform(FVector(1.f, 2.f, 3.f)));
-	Hierarchy.Add(TEXT("BoneB"), TEXT("BoneA"), FTransform(FVector(1.f, 5.f, 3.f)));
-	Hierarchy.Add(TEXT("BoneC"), TEXT("Root"), FTransform(FVector(-4.f, 0.f, 0.f)));
-	Hierarchy.Initialize();
+	BoneHierarchy.Add(TEXT("Root"), NAME_None, FTransform(FVector(1.f, 0.f, 0.f)));
+	BoneHierarchy.Add(TEXT("BoneA"), TEXT("Root"), FTransform(FVector(1.f, 2.f, 3.f)));
+	BoneHierarchy.Add(TEXT("BoneB"), TEXT("BoneA"), FTransform(FVector(1.f, 5.f, 3.f)));
+	BoneHierarchy.Add(TEXT("BoneC"), TEXT("Root"), FTransform(FVector(-4.f, 0.f, 0.f)));
+	BoneHierarchy.Initialize();
 	Unit.ExecuteContext.HierarchyReference = HierarchyRef;
 
-	Hierarchy.ResetTransforms();
+	BoneHierarchy.ResetTransforms();
 	Unit.Bone = TEXT("BoneA");
 	Unit.Space = TEXT("Root");
 	Unit.Transform = FTransform(FVector(0.f, 0.f, 7.f));
 	Unit.bPropagateToChildren = false;
 	InitAndExecute();
-	AddErrorIfFalse(Hierarchy.GetGlobalTransform(0).GetTranslation().Equals(FVector(1.f, 0.f, 0.f)), TEXT("unexpected transform"));
-	AddErrorIfFalse(Hierarchy.GetGlobalTransform(1).GetTranslation().Equals(FVector(1.f, 0.f, 7.f)), TEXT("unexpected transform"));
-	AddErrorIfFalse(Hierarchy.GetGlobalTransform(3).GetTranslation().Equals(FVector(1.f, 5.f, 3.f)), TEXT("unexpected transform"));
-	AddErrorIfFalse(Hierarchy.GetGlobalTransform(2).GetTranslation().Equals(FVector(-4.f, 0.f, 0.f)), TEXT("unexpected transform"));
+	AddErrorIfFalse(BoneHierarchy.GetGlobalTransform(0).GetTranslation().Equals(FVector(1.f, 0.f, 0.f)), TEXT("unexpected transform"));
+	AddErrorIfFalse(BoneHierarchy.GetGlobalTransform(1).GetTranslation().Equals(FVector(1.f, 0.f, 7.f)), TEXT("unexpected transform"));
+	AddErrorIfFalse(BoneHierarchy.GetGlobalTransform(3).GetTranslation().Equals(FVector(1.f, 5.f, 3.f)), TEXT("unexpected transform"));
+	AddErrorIfFalse(BoneHierarchy.GetGlobalTransform(2).GetTranslation().Equals(FVector(-4.f, 0.f, 0.f)), TEXT("unexpected transform"));
 
-	Hierarchy.ResetTransforms();
+	BoneHierarchy.ResetTransforms();
 	Unit.bPropagateToChildren = true;
 	InitAndExecute();
-	AddErrorIfFalse(Hierarchy.GetGlobalTransform(0).GetTranslation().Equals(FVector(1.f, 0.f, 0.f)), TEXT("unexpected transform"));
-	AddErrorIfFalse(Hierarchy.GetGlobalTransform(1).GetTranslation().Equals(FVector(1.f, 0.f, 7.f)), TEXT("unexpected transform"));
-	AddErrorIfFalse(Hierarchy.GetGlobalTransform(3).GetTranslation().Equals(FVector(1.f, 3.f, 7.f)), TEXT("unexpected transform"));
-	AddErrorIfFalse(Hierarchy.GetGlobalTransform(2).GetTranslation().Equals(FVector(-4.f, 0.f, 0.f)), TEXT("unexpected transform"));
+	AddErrorIfFalse(BoneHierarchy.GetGlobalTransform(0).GetTranslation().Equals(FVector(1.f, 0.f, 0.f)), TEXT("unexpected transform"));
+	AddErrorIfFalse(BoneHierarchy.GetGlobalTransform(1).GetTranslation().Equals(FVector(1.f, 0.f, 7.f)), TEXT("unexpected transform"));
+	AddErrorIfFalse(BoneHierarchy.GetGlobalTransform(3).GetTranslation().Equals(FVector(1.f, 3.f, 7.f)), TEXT("unexpected transform"));
+	AddErrorIfFalse(BoneHierarchy.GetGlobalTransform(2).GetTranslation().Equals(FVector(-4.f, 0.f, 0.f)), TEXT("unexpected transform"));
 
-	Hierarchy.ResetTransforms();
+	BoneHierarchy.ResetTransforms();
 	Unit.Space = TEXT("BoneC");
 	Unit.bPropagateToChildren = false;
 	InitAndExecute();
-	AddErrorIfFalse(Hierarchy.GetGlobalTransform(0).GetTranslation().Equals(FVector(1.f, 0.f, 0.f)), TEXT("unexpected transform"));
-	AddErrorIfFalse(Hierarchy.GetGlobalTransform(1).GetTranslation().Equals(FVector(-4.f, 0.f, 7.f)), TEXT("unexpected transform"));
-	AddErrorIfFalse(Hierarchy.GetGlobalTransform(3).GetTranslation().Equals(FVector(1.f, 5.f, 3.f)), TEXT("unexpected transform"));
-	AddErrorIfFalse(Hierarchy.GetGlobalTransform(2).GetTranslation().Equals(FVector(-4.f, 0.f, 0.f)), TEXT("unexpected transform"));
+	AddErrorIfFalse(BoneHierarchy.GetGlobalTransform(0).GetTranslation().Equals(FVector(1.f, 0.f, 0.f)), TEXT("unexpected transform"));
+	AddErrorIfFalse(BoneHierarchy.GetGlobalTransform(1).GetTranslation().Equals(FVector(-4.f, 0.f, 7.f)), TEXT("unexpected transform"));
+	AddErrorIfFalse(BoneHierarchy.GetGlobalTransform(3).GetTranslation().Equals(FVector(1.f, 5.f, 3.f)), TEXT("unexpected transform"));
+	AddErrorIfFalse(BoneHierarchy.GetGlobalTransform(2).GetTranslation().Equals(FVector(-4.f, 0.f, 0.f)), TEXT("unexpected transform"));
 
-	Hierarchy.ResetTransforms();
+	BoneHierarchy.ResetTransforms();
 	Unit.bPropagateToChildren = true;
 	InitAndExecute();
-	AddErrorIfFalse(Hierarchy.GetGlobalTransform(0).GetTranslation().Equals(FVector(1.f, 0.f, 0.f)), TEXT("unexpected transform"));
-	AddErrorIfFalse(Hierarchy.GetGlobalTransform(1).GetTranslation().Equals(FVector(-4.f, 0.f, 7.f)), TEXT("unexpected transform"));
-	AddErrorIfFalse(Hierarchy.GetGlobalTransform(3).GetTranslation().Equals(FVector(-4.f, 3.f, 7.f)), TEXT("unexpected transform"));
-	AddErrorIfFalse(Hierarchy.GetGlobalTransform(2).GetTranslation().Equals(FVector(-4.f, 0.f, 0.f)), TEXT("unexpected transform"));
+	AddErrorIfFalse(BoneHierarchy.GetGlobalTransform(0).GetTranslation().Equals(FVector(1.f, 0.f, 0.f)), TEXT("unexpected transform"));
+	AddErrorIfFalse(BoneHierarchy.GetGlobalTransform(1).GetTranslation().Equals(FVector(-4.f, 0.f, 7.f)), TEXT("unexpected transform"));
+	AddErrorIfFalse(BoneHierarchy.GetGlobalTransform(3).GetTranslation().Equals(FVector(-4.f, 3.f, 7.f)), TEXT("unexpected transform"));
+	AddErrorIfFalse(BoneHierarchy.GetGlobalTransform(2).GetTranslation().Equals(FVector(-4.f, 0.f, 0.f)), TEXT("unexpected transform"));
 
 	return true;
 }
diff --git a/Engine/Plugins/Experimental/ControlRig/Source/ControlRig/Private/Units/Highlevel/Hierarchy/RigUnit_DistributeRotation.cpp b/Engine/Plugins/Experimental/ControlRig/Source/ControlRig/Private/Units/Highlevel/Hierarchy/RigUnit_DistributeRotation.cpp
--- a/Engine/Plugins/Experimental/ControlRig/Source/ControlRig/Private/Units/Highlevel/Hierarchy/RigUnit_DistributeRotation.cpp
+++ b/Engine/Plugins/Experimental/ControlRig/Source/ControlRig/Private/Units/Highlevel/Hierarchy/RigUnit_DistributeRotation.cpp
@@ -184,12 +184,12 @@ void FRigUnit_DistributeRotation::Execute(const FRigUnitContext& Context)
 
 IMPLEMENT_RIGUNIT_AUTOMATION_TEST(FRigUnit_DistributeRotation)
 {
-	Hierarchy.Add(TEXT("Root"), NAME_None, FTransform(FVector(1.f, 0.f, 0.f)));
-	Hierarchy.Add(TEXT("BoneA"), TEXT("Root"), FTransform(FVector(2.f, 0.f, 0.f)));
-	Hierarchy.Add(TEXT("BoneB"), TEXT("BoneA"), FTransform(FVector(2.f, 0.f, 0.f)));
-	Hierarchy.Add(TEXT("BoneC"), TEXT("BoneB"), FTransform(FVector(2.f, 0.f, 0.f)));
-	Hierarchy.Add(TEXT("BoneD"), TEXT("BoneC"), FTransform(FVector(2.f, 0.f, 0.f)));
-	Hierarchy.Initialize();
+	BoneHierarchy.Add(TEXT("Root"), NAME_None, FTransform(FVector(1.f, 0.f, 0.f)));
+	BoneHierarchy.Add(TEXT("BoneA"), TEXT("Root"), FTransform(FVector(2.f, 0.f, 0.f)));
+	BoneHierarchy.Add(TEXT("BoneB"), TEXT("BoneA"), FTransform(FVector(2.f, 0.f, 0.f)));
+	BoneHierarchy.Add(TEXT("BoneC"), TEXT("BoneB"), FTransform(FVector(2.f, 0.f, 0.f)));
+	BoneHierarchy.Add(TEXT("BoneD"), TEXT("BoneC"), FTransform(FVector(2.f, 0.f, 0.f)));
+	BoneHierarchy.Initialize();
 
 	Unit.StartBone = TEXT("Root");
 	Unit.EndBone = TEXT("BoneD");
@@ -226,15 +226,15 @@ IMPLEMENT_RIGUNIT_AUTOMATION_TEST(FRigUnit_DistributeRotation)
 	Execute();
 
 	FVector Euler = FVector::ZeroVector;
-	Euler = FControlRigMathLibrary::EulerFromQuat(Hierarchy.GetLocalTransform(0).GetRotation(), EControlRigRotationOrder::XYZ);
+	Euler = FControlRigMathLibrary::EulerFromQuat(BoneHierarchy.GetLocalTransform(0).GetRotation(), EControlRigRotationOrder::XYZ);
 	AddErrorIfFalse(FMath::IsNearlyEqual(Euler.Y, 0.f, 0.1f), TEXT("unexpected rotation Y"));
-	Euler = FControlRigMathLibrary::EulerFromQuat(Hierarchy.GetLocalTransform(1).GetRotation(), EControlRigRotationOrder::XYZ);
+	Euler = FControlRigMathLibrary::EulerFromQuat(BoneHierarchy.GetLocalTransform(1).GetRotation(), EControlRigRotationOrder::XYZ);
 	AddErrorIfFalse(FMath::IsNearlyEqual(Euler.Y, 45.f, 0.1f), TEXT("unexpected rotation Y"));
-	Euler = FControlRigMathLibrary::EulerFromQuat(Hierarchy.GetLocalTransform(2).GetRotation(), EControlRigRotationOrder::XYZ);
+	Euler = FControlRigMathLibrary::EulerFromQuat(BoneHierarchy.GetLocalTransform(2).GetRotation(), EControlRigRotationOrder::XYZ);
 	AddErrorIfFalse(FMath::IsNearlyEqual(Euler.Y, 90.f, 0.1f), TEXT("unexpected rotation Y"));
-	Euler = FControlRigMathLibrary::EulerFromQuat(Hierarchy.GetLocalTransform(3).GetRotation(), EControlRigRotationOrder::XYZ);
+	Euler = FControlRigMathLibrary::EulerFromQuat(BoneHierarchy.GetLocalTransform(3).GetRotation(), EControlRigRotationOrder::XYZ);
 	AddErrorIfFalse(FMath::IsNearlyEqual(Euler.Y, 45.f, 0.1f), TEXT("unexpected rotation Y"));
-	Euler = FControlRigMathLibrary::EulerFromQuat(Hierarchy.GetLocalTransform(4).GetRotation(), EControlRigRotationOrder::XYZ);
+	Euler = FControlRigMathLibrary::EulerFromQuat(BoneHierarchy.GetLocalTransform(4).GetRotation(), EControlRigRotationOrder::XYZ);
 	AddErrorIfFalse(FMath::IsNearlyEqual(Euler.Y, 0.f, 0.1f), TEXT("unexpected rotation Y"));
 
 	return true;
diff --git a/Engine/Plugins/Experimental/ControlRig/Source/ControlRig/Private/Units/Highlevel/Hierarchy/RigUnit_ModifyBoneTransforms.cpp b/Engine/Plugins/Experimental/ControlRig/Source/ControlRig/Private/Units/Highlevel/Hierarchy/RigUnit_ModifyBoneTransforms.cpp
--- a/Engine/Plugins/Experimental/ControlRig/Source/ControlRig/Private/Units/Highlevel/Hierarchy/RigUnit_ModifyBoneTransforms.cpp
+++ b/Engine/Plugins/Experimental/ControlRig/Source/ControlRig/Private/Units/Highlevel/Hierarchy/RigUnit_ModifyBoneTransforms.cpp
@@ -109,10 +109,10 @@ void FRigUnit_ModifyBoneTransforms::Execute(const FRigUnitContext& Context)
 
 IMPLEMENT_RIGUNIT_AUTOMATION_TEST(FRigUnit_ModifyBoneTransforms)
 {
-	Hierarchy.Add(TEXT("Root"), NAME_None, FTransform(FVector(1.f, 0.f, 0.f)));
-	Hierarchy.Add(TEXT("BoneA"), TEXT("Root"), FTransform(FVector(1.f, 2.f, 3.f)));
-	Hierarchy.Add(TEXT("BoneB"), TEXT("Root"), FTransform(FVector(5.f, 6.f, 7.f)));
-	Hierarchy.Initialize();
+	BoneHierarchy.Add(TEXT("Root"), NAME_None, FTransform(FVector(1.f, 0.f, 0.f)));
+	BoneHierarchy.Add(TEXT("BoneA"), TEXT("Root"), FTransform(FVector(1.f, 2.f, 3.f)));
+	BoneHierarchy.Add(TEXT("BoneB"), TEXT("Root"), FTransform(FVector(5.f, 6.f, 7.f)));
+	BoneHierarchy.Initialize();
 	Unit.ExecuteContext.HierarchyReference = HierarchyRef;
 
 	Unit.BoneToModify.SetNumZeroed(2);
@@ -120,41 +120,41 @@ IMPLEMENT_RIGUNIT_AUTOMATION_TEST(FRigUnit_ModifyBoneTransforms)
 	Unit.BoneToModify[1].Bone = TEXT("BoneB");
 	Unit.BoneToModify[0].Transform = Unit.BoneToModify[1].Transform = FTransform(FVector(10.f, 11.f, 12.f));
 
-	Hierarchy.ResetTransforms();
+	BoneHierarchy.ResetTransforms();
 	Unit.Mode = EControlRigModifyBoneMode::AdditiveLocal;
 	InitAndExecute();
-	AddErrorIfFalse((Hierarchy.GetGlobalTransform(0).GetTranslation() - FVector(1.f, 0.f, 0.f)).IsNearlyZero(), TEXT("unexpected transform"));
-	AddErrorIfFalse((Hierarchy.GetGlobalTransform(1).GetTranslation() - FVector(11.f, 13.f, 15.f)).IsNearlyZero(), TEXT("unexpected transform"));
-	AddErrorIfFalse((Hierarchy.GetGlobalTransform(2).GetTranslation() - FVector(15.f, 17.f, 19.f)).IsNearlyZero(), TEXT("unexpected transform"));
+	AddErrorIfFalse((BoneHierarchy.GetGlobalTransform(0).GetTranslation() - FVector(1.f, 0.f, 0.f)).IsNearlyZero(), TEXT("unexpected transform"));
+	AddErrorIfFalse((BoneHierarchy.GetGlobalTransform(1).GetTranslation() - FVector(11.f, 13.f, 15.f)).IsNearlyZero(), TEXT("unexpected transform"));
+	AddErrorIfFalse((BoneHierarchy.GetGlobalTransform(2).GetTranslation() - FVector(15.f, 17.f, 19.f)).IsNearlyZero(), TEXT("unexpected transform"));
 
-	Hierarchy.ResetTransforms();
+	BoneHierarchy.ResetTransforms();
 	Unit.Mode = EControlRigModifyBoneMode::AdditiveGlobal;
 	InitAndExecute();
-	AddErrorIfFalse((Hierarchy.GetGlobalTransform(0).GetTranslation() - FVector(1.f, 0.f, 0.f)).IsNearlyZero(), TEXT("unexpected transform"));
-	AddErrorIfFalse((Hierarchy.GetGlobalTransform(1).GetTranslation() - FVector(11.f, 13.f, 15.f)).IsNearlyZero(), TEXT("unexpected transform"));
-	AddErrorIfFalse((Hierarchy.GetGlobalTransform(2).GetTranslation() - FVector(15.f, 17.f, 19.f)).IsNearlyZero(), TEXT("unexpected transform"));
+	AddErrorIfFalse((BoneHierarchy.GetGlobalTransform(0).GetTranslation() - FVector(1.f, 0.f, 0.f)).IsNearlyZero(), TEXT("unexpected transform"));
+	AddErrorIfFalse((BoneHierarchy.GetGlobalTransform(1).GetTranslation() - FVector(11.f, 13.f, 15.f)).IsNearlyZero(), TEXT("unexpected transform"));
+	AddErrorIfFalse((BoneHierarchy.GetGlobalTransform(2).GetTranslation() - FVector(15.f, 17.f, 19.f)).IsNearlyZero(), TEXT("unexpected transform"));
 
-	Hierarchy.ResetTransforms();
+	BoneHierarchy.ResetTransforms();
 	Unit.Mode = EControlRigModifyBoneMode::OverrideLocal;
 	InitAndExecute();
-	AddErrorIfFalse((Hierarchy.GetGlobalTransform(0).GetTranslation() - FVector(1.f, 0.f, 0.f)).IsNearlyZero(), TEXT("unexpected transform"));
-	AddErrorIfFalse((Hierarchy.GetGlobalTransform(1).GetTranslation() - FVector(11.f, 11.f, 12.f)).IsNearlyZero(), TEXT("unexpected transform"));
-	AddErrorIfFalse((Hierarchy.GetGlobalTransform(2).GetTranslation() - FVector(11.f, 11.f, 12.f)).IsNearlyZero(), TEXT("unexpected transform"));
+	AddErrorIfFalse((BoneHierarchy.GetGlobalTransform(0).GetTranslation() - FVector(1.f, 0.f, 0.f)).IsNearlyZero(), TEXT("unexpected transform"));
+	AddErrorIfFalse((BoneHierarchy.GetGlobalTransform(1).GetTranslation() - FVector(11.f, 11.f, 12.f)).IsNearlyZero(), TEXT("unexpected transform"));
+	AddErrorIfFalse((BoneHierarchy.GetGlobalTransform(2).GetTranslation() - FVector(11.f, 11.f, 12.f)).IsNearlyZero(), TEXT("unexpected transform"));
 
-	Hierarchy.ResetTransforms();
+	BoneHierarchy.ResetTransforms();
 	Unit.Mode = EControlRigModifyBoneMode::OverrideGlobal;
 	InitAndExecute();
-	AddErrorIfFalse((Hierarchy.GetGlobalTransform(0).GetTranslation() - FVector(1.f, 0.f, 0.f)).IsNearlyZero(), TEXT("unexpected transform"));
-	AddErrorIfFalse((Hierarchy.GetGlobalTransform(1).GetTranslation() - FVector(10.f, 11.f, 12.f)).IsNearlyZero(), TEXT("unexpected transform"));
-	AddErrorIfFalse((Hierarchy.GetGlobalTransform(2).GetTranslation() - FVector(10.f, 11.f, 12.f)).IsNearlyZero(), TEXT("unexpected transform"));
+	AddErrorIfFalse((BoneHierarchy.GetGlobalTransform(0).GetTranslation() - FVector(1.f, 0.f, 0.f)).IsNearlyZero(), TEXT("unexpected transform"));
+	AddErrorIfFalse((BoneHierarchy.GetGlobalTransform(1).GetTranslation() - FVector(10.f, 11.f, 12.f)).IsNearlyZero(), TEXT("unexpected transform"));
+	AddErrorIfFalse((BoneHierarchy.GetGlobalTransform(2).GetTranslation() - FVector(10.f, 11.f, 12.f)).IsNearlyZero(), TEXT("unexpected transform"));
 
-	Hierarchy.ResetTransforms();
+	BoneHierarchy.ResetTransforms();
 	Unit.Mode = EControlRigModifyBoneMode::AdditiveLocal;
 	Unit.Weight = 0.5f;
 	InitAndExecute();
-	AddErrorIfFalse((Hierarchy.GetGlobalTransform(0).GetTranslation() - FVector(1.f, 0.f, 0.f)).IsNearlyZero(), TEXT("unexpected transform"));
-	AddErrorIfFalse((Hierarchy.GetGlobalTransform(1).GetTranslation() - FVector(6.f, 7.5f, 9.f)).IsNearlyZero(), TEXT("unexpected transform"));
-	AddErrorIfFalse((Hierarchy.GetGlobalTransform(2).GetTranslation() - FVector(10.f, 11.5f, 13.f)).IsNearlyZero(), TEXT("unexpected transform"));
+	AddErrorIfFalse((BoneHierarchy.GetGlobalTransform(0).GetTranslation() - FVector(1.f, 0.f, 0.f)).IsNearlyZero(), TEXT("unexpected transform"));
+	AddErrorIfFalse((BoneHierarchy.GetGlobalTransform(1).GetTranslation() - FVector(6.f, 7.5f, 9.f)).IsNearlyZero(), TEXT("unexpected transform"));
+	AddErrorIfFalse((BoneHierarchy.GetGlobalTransform(2).GetTranslation() - FVector(10.f, 11.5f, 13.f)).IsNearlyZero(), TEXT("unexpected transform"));
 
 
 	return true;
diff --git a/Engine/Plugins/Experimental/ControlRig/Source/ControlRig/Public/Rigs/RigCurveContainer.h b/Engine/Plugins/Experimental/ControlRig/Source/ControlRig/Public/Rigs/RigCurveContainer.h
--- a/Engine/Plugins/Experimental/ControlRig/Source/ControlRig/Public/Rigs/RigCurveContainer.h
+++ b/Engine/Plugins/Experimental/ControlRig/Source/ControlRig/Public/Rigs/RigCurveContainer.h
@@ -25,65 +25,36 @@ struct CONTROLRIG_API FRigCurveContainer
 {
 	GENERATED_BODY()
 
-private:
-	UPROPERTY(EditAnywhere, Category = FRigCurveContainer)
-	TArray<FRigCurve> Curves;
-
-	// can serialize fine? 
-	UPROPERTY()
-	TMap<FName, int32> NameToIndexMapping;
-
 public:
 
-	const TArray<FRigCurve>& GetCurves() const
-	{
-		return Curves;
-	}
+	FRigCurveContainer();
+	FRigCurveContainer& operator= (const FRigCurveContainer &InOther);
 
-	void AddCurve(const FName& NewCurveName)
-	{
-		DECLARE_SCOPE_HIERARCHICAL_COUNTER_FUNC()
+	FORCEINLINE int32 Num() const { return Curves.Num(); }
+	FORCEINLINE const FRigCurve& operator[](int32 InIndex) const { return Curves[InIndex]; }
+	FORCEINLINE FRigCurve& operator[](int32 InIndex) { return Curves[InIndex]; }
+	FORCEINLINE const FRigCurve& operator[](const FName& InName) const { return Curves[GetIndex(InName)]; }
+	FORCEINLINE FRigCurve& operator[](const FName& InName) { return Curves[GetIndex(InName)]; }
 
-		int32 Found = GetIndex(NewCurveName);
-		if (Found == INDEX_NONE)
-		{
-			FRigCurve NewCurve;
-			NewCurve.Name = NewCurveName;
-			NewCurve.Value = 0.f;
-			Curves.Add(NewCurve);
-			RefreshMapping();
-		}
-	}
+	FORCEINLINE TArray<FRigCurve>::RangedForIteratorType      begin()       { return Curves.begin(); }
+	FORCEINLINE TArray<FRigCurve>::RangedForConstIteratorType begin() const { return Curves.begin(); }
+	FORCEINLINE TArray<FRigCurve>::RangedForIteratorType      end()         { return Curves.end();   }
+	FORCEINLINE TArray<FRigCurve>::RangedForConstIteratorType end() const   { return Curves.end();   }
 
-	void DeleteCurve(const FName& CurveToDelete)
-	{
-		DECLARE_SCOPE_HIERARCHICAL_COUNTER_FUNC()
+	void Add(const FName& InNewName);
 
-		int32 IndexToDelete = GetIndex(CurveToDelete);
-		if (IndexToDelete != INDEX_NONE)
-		{
-			Curves.RemoveAt(IndexToDelete);
-			RefreshMapping();
-		}
-	}
+	void Remove(const FName& InName);
 
-	FName GetName(int32 Index) const
-	{
-		DECLARE_SCOPE_HIERARCHICAL_COUNTER_FUNC()
+	FName GetName(int32 InIndex) const;
 
-		if (Curves.IsValidIndex(Index))
+	FORCEINLINE int32 GetIndex(const FName& InName) const
+	{
+		if(NameToIndexMapping.Num() != Curves.Num())
 		{
-			return Curves[Index].Name;
+			return GetIndexSlow(InName);
 		}
 
-		return NAME_None;
-	}
-
-	int32 GetIndex(const FName& Curve) const
-	{
-		DECLARE_SCOPE_HIERARCHICAL_COUNTER_FUNC()
-
-		const int32* Index = NameToIndexMapping.Find(Curve);
+		const int32* Index = NameToIndexMapping.Find(InName);
 		if (Index)
 		{
 			return *Index;
@@ -92,79 +63,15 @@ public:
 		return INDEX_NONE;
 	}
 
-//#if WITH_EDITOR
-	// @FIXMELINA: figure out how to remove this outside of editor
-	// ignore mapping, run slow search
-	// this is useful in editor while editing
-	// we don't want to build mapping data every time
-	int32 GetIndexSlow(const FName& Curve) const
-	{
-		DECLARE_SCOPE_HIERARCHICAL_COUNTER_FUNC()
-
-		for (int32 CurveId = 0; CurveId < Curves.Num(); ++CurveId)
-		{
-			if (Curves[CurveId].Name == Curve)
-			{
-				return CurveId;
-			}
-		}
-
-		return INDEX_NONE;
-	}
-//#endif // WITH_EDITOR
-	void SetValue(const FName& Curve, const float Value)
-	{
-		DECLARE_SCOPE_HIERARCHICAL_COUNTER_FUNC()
-
-		SetValue(GetIndex(Curve), Value);
-	}
-
-	void SetValue(int32 Index, const float Value)
-	{
-		DECLARE_SCOPE_HIERARCHICAL_COUNTER_FUNC()
-
-		if (Curves.IsValidIndex(Index))
-		{
-			FRigCurve& Curve = Curves[Index];
-			Curve.Value = Value;
-		}
-	}
-
-	float GetValue(const FName& Curve) const
-	{
-		DECLARE_SCOPE_HIERARCHICAL_COUNTER_FUNC()
-
-		return GetValue(GetIndex(Curve));
-	}
-
-	float GetValue(int32 Index) const
-	{
-		DECLARE_SCOPE_HIERARCHICAL_COUNTER_FUNC()
+	void SetValue(const FName& InName, const float InValue);
 
-		if (Curves.IsValidIndex(Index))
-		{
-			return Curves[Index].Value;
-		}
+	void SetValue(int32 InIndex, const float InValue);
 
-		return 0.f;
-	}
+	float GetValue(const FName& InName) const;
 
-	void Rename(const FName& OldName, const FName& NewName)
-	{
-		DECLARE_SCOPE_HIERARCHICAL_COUNTER_FUNC()
+	float GetValue(int32 InIndex) const;
 
-		if (OldName != NewName)
-		{
-			const int32 Found = GetIndex(OldName);
-			const int32 NewNameFound = GetIndex(NewName);
-			// if I have new name, and didn't find new name
-			if (Found != INDEX_NONE && NewNameFound == INDEX_NONE)
-			{
-				Curves[Found].Name = NewName;
-				RefreshMapping();
-			}
-		}
-	}
+	void Rename(const FName& InOldName, const FName& InNewName);
 
 	// updates all of the internal caches
 	void Initialize();
@@ -175,12 +82,19 @@ public:
 	// resets all of the transforms back to the initial transform
 	void ResetValues();
 
-	int32 GetNum() const
-	{
-		return Curves.Num();
-	}
-
 private:
+
+	// disable copy constructor
+	FRigCurveContainer(const FRigCurveContainer& InOther) {}
+
+	UPROPERTY(EditAnywhere, Category = FRigCurveContainer)
+	TArray<FRigCurve> Curves;
+
+	UPROPERTY()
+	TMap<FName, int32> NameToIndexMapping;
+
+	int32 GetIndexSlow(const FName& InName) const;
+
 	void RefreshMapping();
 };
 
diff --git a/Engine/Plugins/Experimental/ControlRig/Source/ControlRig/Public/Units/RigUnitTest.h b/Engine/Plugins/Experimental/ControlRig/Source/ControlRig/Public/Units/RigUnitTest.h
--- a/Engine/Plugins/Experimental/ControlRig/Source/ControlRig/Public/Units/RigUnitTest.h
+++ b/Engine/Plugins/Experimental/ControlRig/Source/ControlRig/Public/Units/RigUnitTest.h
@@ -15,8 +15,8 @@ public:
 		: FAutomationTestBase(InName, bIsComplex)
 		, HierarchyContainer()
 		, HierarchyRef(FRigHierarchyRef())
-		, Hierarchy(HierarchyContainer.BoneHierarchy)
-		, CurveContainer(FRigCurveContainer())
+		, BoneHierarchy(HierarchyContainer.BoneHierarchy)
+		, CurveContainer()
 		, CurveContainerRef(&CurveContainer)
 	{
 		HierarchyRef.Container = &HierarchyContainer;
@@ -27,7 +27,7 @@ public:
 
 	FRigHierarchyContainer HierarchyContainer;
 	FRigHierarchyRef HierarchyRef;
-	FRigBoneHierarchy& Hierarchy;
+	FRigBoneHierarchy& BoneHierarchy;
 	FRigCurveContainer CurveContainer;
 	FRigCurveContainerRef CurveContainerRef;
 	FControlRigExecuteContext ExecuteContext;
@@ -56,7 +56,7 @@ public:
 		TUnitStruct Unit; \
 		virtual bool RunTest(const FString& Parameters) override \
 		{ \
-			Hierarchy.Reset(); \
+			HierarchyContainer.Reset(); \
 			CurveContainer.Reset(); \
 			Unit = TUnitStruct(); \
 			return RunControlRigUnitTest(Parameters); \
diff --git a/Engine/Plugins/Experimental/ControlRig/Source/ControlRigDeveloper/Private/Graph/ControlRigGraph.cpp b/Engine/Plugins/Experimental/ControlRig/Source/ControlRigDeveloper/Private/Graph/ControlRigGraph.cpp
--- a/Engine/Plugins/Experimental/ControlRig/Source/ControlRigDeveloper/Private/Graph/ControlRigGraph.cpp
+++ b/Engine/Plugins/Experimental/ControlRig/Source/ControlRigDeveloper/Private/Graph/ControlRigGraph.cpp
@@ -578,8 +578,7 @@ void UControlRigGraph::CacheCurveNameList(const FRigCurveContainer& Container)
 	DECLARE_SCOPE_HIERARCHICAL_COUNTER_FUNC()
 
 	TArray<FString> Names;
-	const TArray<FRigCurve>& Curves = Container.GetCurves();
-	for (const FRigCurve& Curve : Curves)
+	for (const FRigCurve& Curve : Container)
 	{
 		Names.Add(Curve.Name.ToString());
 	}
diff --git a/Engine/Plugins/Experimental/ControlRig/Source/ControlRigEditor/Private/Editor/SRigCurveContainer.cpp b/Engine/Plugins/Experimental/ControlRig/Source/ControlRigEditor/Private/Editor/SRigCurveContainer.cpp
--- a/Engine/Plugins/Experimental/ControlRig/Source/ControlRigEditor/Private/Editor/SRigCurveContainer.cpp
+++ b/Engine/Plugins/Experimental/ControlRig/Source/ControlRigEditor/Private/Editor/SRigCurveContainer.cpp
@@ -337,7 +337,7 @@ void SRigCurveContainer::CreateNewNameEntry(const FText& CommittedText, ETextCom
 		FRigCurveContainer* Container = GetCurveContainer();
 		if (Container)
 		{
-			Container->AddCurve(FName(*CommittedText.ToString()));
+			Container->Add(FName(*CommittedText.ToString()));
 			RefreshCurveList();
 		}
 	}
@@ -351,8 +351,7 @@ void SRigCurveContainer::CreateRigCurveList( const FString& SearchText )
 		RigCurveList.Reset();
 
 		// Iterate through all curves..
-		const TArray<FRigCurve>& Curves = Container->GetCurves();
-		for (const FRigCurve& Curve : Curves)
+		for (const FRigCurve& Curve : (*Container))
 		{
 			FString CurveString = Curve.Name.ToString();
 
@@ -408,7 +407,7 @@ void SRigCurveContainer::OnDeleteNameClicked()
 		TArray< TSharedPtr<FDisplayedRigCurveInfo> > SelectedItems = RigCurveListView->GetSelectedItems();
 		for (auto Item : SelectedItems)
 		{
-			Container->DeleteCurve(Item->CurveName);
+			Container->Remove(Item->CurveName);
 		}
 
 		RefreshCurveList();
@@ -510,7 +509,7 @@ void SRigCurveContainer::ImportCurve(const FAssetData& InAssetData)
 			SmartNameMapping->FillNameArray(NameArray);
 			for (int32 Index = 0; Index < NameArray.Num() ; ++Index)
 			{
-				Container->AddCurve(NameArray[Index]);
+				Container->Add(NameArray[Index]);
 			}
 
 			ControlRigEditor.Pin()->OnCurveContainerChanged();
