ControlRig: Fix build break due to missing Generated include.

diff --git a/Engine/Plugins/Experimental/ControlRig/Source/ControlRig/Public/Rigs/RigHierarchyDefines.h b/Engine/Plugins/Experimental/ControlRig/Source/ControlRig/Public/Rigs/RigHierarchyDefines.h
--- a/Engine/Plugins/Experimental/ControlRig/Source/ControlRig/Public/Rigs/RigHierarchyDefines.h
+++ b/Engine/Plugins/Experimental/ControlRig/Source/ControlRig/Public/Rigs/RigHierarchyDefines.h
@@ -3,6 +3,7 @@
 #pragma once
 
 #include "CoreMinimal.h"
+#include "RigHierarchyDefines.generated.h"
 
 struct FRigHierarchyContainer;
 
