Control Rig: Refactoring of hierarchy data structures

* introduces MVC model / controller / and views
* refactored UI widgets to rely on OnRigElementChanged delegates
* improved naming stability on all data structures
* removed HierarchyReference / CurveContainerReference

diff --git a/Engine/Plugins/Experimental/ControlRig/Source/ControlRig/Private/ControlRig.cpp b/Engine/Plugins/Experimental/ControlRig/Source/ControlRig/Private/ControlRig.cpp
--- a/Engine/Plugins/Experimental/ControlRig/Source/ControlRig/Private/ControlRig.cpp
+++ b/Engine/Plugins/Experimental/ControlRig/Source/ControlRig/Private/ControlRig.cpp
@@ -200,7 +200,6 @@ void UControlRig::Initialize(bool bInitRigUnits)
 
 	// should refresh mapping 
 	Hierarchy.Initialize();
-	CurveContainer.Initialize();
 
 	// resolve IO properties
 	ResolveInputOutputProperties();
@@ -223,7 +222,6 @@ void UControlRig::InitializeFromCDO()
 	UControlRig* CDO = GetClass()->GetDefaultObject<UControlRig>();
 	// copy operation
 	Hierarchy = CDO->Hierarchy;
-	CurveContainer = CDO->CurveContainer;
 }
 
 void UControlRig::PreEvaluate_GameThread()
@@ -317,8 +315,7 @@ void UControlRig::Execute(const EControlRigState InState)
 
 	Context.DeltaTime = DeltaTime;
 	Context.State = InState;
-	Context.HierarchyReference.Container = &Hierarchy;
-	Context.CurveReference = FRigCurveContainerRef(&CurveContainer);
+	Context.Hierarchy = &Hierarchy;
 
 #if WITH_EDITOR
 	Context.Log = ControlRigLog;
@@ -431,26 +428,26 @@ void UControlRig::SetGlobalTransform(const int32 BoneIndex, const FTransform& In
 float UControlRig::GetCurveValue(const FName& CurveName) const
 {
 	DECLARE_SCOPE_HIERARCHICAL_COUNTER_FUNC()
-	return CurveContainer.GetValue(CurveName);
+	return Hierarchy.CurveContainer.GetValue(CurveName);
 }
 
 void UControlRig::SetCurveValue(const FName& CurveName, const float CurveValue)
 {
 	DECLARE_SCOPE_HIERARCHICAL_COUNTER_FUNC()
 
-	CurveContainer.SetValue(CurveName, CurveValue);
+	Hierarchy.CurveContainer.SetValue(CurveName, CurveValue);
 }
 
 float UControlRig::GetCurveValue(const int32 CurveIndex) const
 {
 	DECLARE_SCOPE_HIERARCHICAL_COUNTER_FUNC()
-	return CurveContainer.GetValue(CurveIndex);
+	return Hierarchy.CurveContainer.GetValue(CurveIndex);
 }
 
 void UControlRig::SetCurveValue(const int32 CurveIndex, const float CurveValue)
 {
 	DECLARE_SCOPE_HIERARCHICAL_COUNTER_FUNC()
-	CurveContainer.SetValue(CurveIndex, CurveValue);
+	Hierarchy.CurveContainer.SetValue(CurveIndex, CurveValue);
 }
 
 void UControlRig::GetMappableNodeData(TArray<FName>& OutNames, TArray<FNodeItem>& OutNodeItems) const
diff --git a/Engine/Plugins/Experimental/ControlRig/Source/ControlRig/Private/Rigs/RigBoneHierarchy.cpp b/Engine/Plugins/Experimental/ControlRig/Source/ControlRig/Private/Rigs/RigBoneHierarchy.cpp
--- a/Engine/Plugins/Experimental/ControlRig/Source/ControlRig/Private/Rigs/RigBoneHierarchy.cpp
+++ b/Engine/Plugins/Experimental/ControlRig/Source/ControlRig/Private/Rigs/RigBoneHierarchy.cpp
@@ -21,7 +21,7 @@ FRigBoneHierarchy& FRigBoneHierarchy::operator= (const FRigBoneHierarchy &InOthe
 	for (int32 Index = Num() - 1; Index >= 0; Index--)
 	{
 		FRigBone BoneToRemove = Bones[Index];
-		OnBoneRemoved.Broadcast(Container, ERigHierarchyElementType::Bone, BoneToRemove.Name);
+		OnBoneRemoved.Broadcast(Container, RigElementType(), BoneToRemove.Name);
 	}
 #endif
 
@@ -33,20 +33,31 @@ FRigBoneHierarchy& FRigBoneHierarchy::operator= (const FRigBoneHierarchy &InOthe
 #if WITH_EDITOR
 	for (const FRigBone& BoneAdded : Bones)
 	{
-		OnBoneAdded.Broadcast(Container, ERigHierarchyElementType::Bone, BoneAdded.Name);
+		OnBoneAdded.Broadcast(Container, RigElementType(), BoneAdded.Name);
 	}
 #endif
 
 	return *this;
 }
 
+FName FRigBoneHierarchy::GetSafeNewName(const FName& InPotentialNewName) const
+{
+	FName Name = InPotentialNewName;
+	int32 Suffix = 1;
+	while(!IsNameAvailable(Name))
+	{
+		Name = *FString::Printf(TEXT("%s_%d"), *InPotentialNewName.ToString(), ++Suffix);
+	}
+	return Name;
+}
+
 FRigBone& FRigBoneHierarchy::Add(const FName& InNewName, const FName& InParentName, const FTransform& InInitTransform)
 {
 	int32 ParentIndex = GetIndex(InParentName);
 	bool bHasParent = (ParentIndex != INDEX_NONE);
 
 	FRigBone NewBone;
-	NewBone.Name = InNewName;
+	NewBone.Name = GetSafeNewName(InNewName);
 	NewBone.ParentIndex = ParentIndex;
 	NewBone.ParentName = bHasParent? InParentName : NAME_None;
 	NewBone.InitialTransform = InInitTransform;
@@ -57,10 +68,10 @@ FRigBone& FRigBoneHierarchy::Add(const FName& InNewName, const FName& InParentNa
 	RefreshMapping();
 
 #if WITH_EDITOR
-	OnBoneAdded.Broadcast(Container, ERigHierarchyElementType::Bone, InNewName);
+	OnBoneAdded.Broadcast(Container, RigElementType(), NewBone.Name);
 #endif
 
-	int32 Index = GetIndex(InNewName);
+	int32 Index = GetIndex(NewBone.Name);
 	return Bones[Index];
 }
 
@@ -94,7 +105,7 @@ void FRigBoneHierarchy::Reparent(const FName& InName, const FName& InNewParentNa
 #if WITH_EDITOR
 		if (OldParentName != InNewParentName)
 		{
-			OnBoneReparented.Broadcast(Container, ERigHierarchyElementType::Bone, InName, OldParentName, InNewParentName);
+			OnBoneReparented.Broadcast(Container, RigElementType(), InName, OldParentName, InNewParentName);
 		}
 #endif
 	}
@@ -130,9 +141,9 @@ FRigBone FRigBoneHierarchy::Remove(const FName& InNameToRemove)
 #if WITH_EDITOR
 	for (const FName& RemovedChildBone : RemovedChildBones)
 	{
-		OnBoneRemoved.Broadcast(Container, ERigHierarchyElementType::Bone, RemovedChildBone);
+		OnBoneRemoved.Broadcast(Container, RigElementType(), RemovedChildBone);
 	}
-	OnBoneRemoved.Broadcast(Container, ERigHierarchyElementType::Bone, RemovedBone.Name);
+	OnBoneRemoved.Broadcast(Container, RigElementType(), RemovedBone.Name);
 #endif
 
 	return RemovedBone;
@@ -292,14 +303,15 @@ void FRigBoneHierarchy::RecalculateGlobalTransform(FRigBone& InOutBone)
 	InOutBone.GlobalTransform = (bHasParent) ? InOutBone.LocalTransform * Bones[InOutBone.ParentIndex].GlobalTransform : InOutBone.LocalTransform;
 }
 
-void FRigBoneHierarchy::Rename(const FName& InOldName, const FName& InNewName)
+FName FRigBoneHierarchy::Rename(const FName& InOldName, const FName& InNewName)
 {
 	if (InOldName != InNewName)
 	{
 		const int32 Found = GetIndex(InOldName);
 		if (Found != INDEX_NONE)
 		{
-			Bones[Found].Name = InNewName;
+			FName NewName = GetSafeNewName(InNewName);
+			Bones[Found].Name = NewName;
 
 			// go through find all children and rename them
 #if WITH_EDITOR
@@ -309,7 +321,7 @@ void FRigBoneHierarchy::Rename(const FName& InOldName, const FName& InNewName)
 			{
 				if (Bones[Index].ParentName == InOldName)
 				{
-					Bones[Index].ParentName = InNewName;
+					Bones[Index].ParentName = NewName;
 					RecalculateLocalTransform(Bones[Index]);
 #if WITH_EDITOR
 					ReparentedBones.Add(Bones[Index].Name);
@@ -320,14 +332,17 @@ void FRigBoneHierarchy::Rename(const FName& InOldName, const FName& InNewName)
 			RefreshMapping();
 
 #if WITH_EDITOR
-			OnBoneRenamed.Broadcast(Container, ERigHierarchyElementType::Bone, InOldName, InNewName);
+			OnBoneRenamed.Broadcast(Container, RigElementType(), InOldName, NewName);
 			for (const FName& ReparentedBone : ReparentedBones)
 			{
-				OnBoneReparented.Broadcast(Container, ERigHierarchyElementType::Bone, ReparentedBone, InOldName, InNewName);
+				OnBoneReparented.Broadcast(Container, RigElementType(), ReparentedBone, InOldName, NewName);
 			}
 #endif
+			return NewName;
 		}
 	}
+
+	return NAME_None;
 }
 
 void FRigBoneHierarchy::Sort()
@@ -396,6 +411,7 @@ void FRigBoneHierarchy::RefreshMapping()
 	NameToIndexMapping.Empty();
 	for (int32 Index = 0; Index < Bones.Num(); ++Index)
 	{
+		Bones[Index].Index = Index;
 		NameToIndexMapping.Add(Bones[Index].Name, Index);
 	}
 }
diff --git a/Engine/Plugins/Experimental/ControlRig/Source/ControlRig/Private/Rigs/RigControlHierarchy.cpp b/Engine/Plugins/Experimental/ControlRig/Source/ControlRig/Private/Rigs/RigControlHierarchy.cpp
new file mode 100644
--- /dev/null
+++ b/Engine/Plugins/Experimental/ControlRig/Source/ControlRig/Private/Rigs/RigControlHierarchy.cpp
@@ -0,0 +1,461 @@
+// Copyright 1998-2019 Epic Games, Inc. All Rights Reserved.
+
+#include "Rigs/RigControlHierarchy.h"
+#include "ControlRig.h"
+#include "HelperUtil.h"
+
+////////////////////////////////////////////////////////////////////////////////
+// FRigControlHierarchy
+////////////////////////////////////////////////////////////////////////////////
+
+FRigControlHierarchy::FRigControlHierarchy()
+#if WITH_EDITOR
+	:Container(nullptr)
+#endif
+{
+}
+
+FRigControlHierarchy& FRigControlHierarchy::operator= (const FRigControlHierarchy &InOther)
+{
+#if WITH_EDITOR
+	for (int32 Index = Num() - 1; Index >= 0; Index--)
+	{
+		FRigControl ControlToRemove = Controls[Index];
+		OnControlRemoved.Broadcast(Container, RigElementType(), ControlToRemove.Name);
+	}
+#endif
+
+	Controls.Reset();
+	Controls.Append(InOther.Controls);
+	NameToIndexMapping.Reset();
+	RefreshMapping();
+
+#if WITH_EDITOR
+	for (const FRigControl& ControlAdded : Controls)
+	{
+		OnControlAdded.Broadcast(Container, RigElementType(), ControlAdded.Name);
+	}
+#endif
+
+	return *this;
+}
+
+FName FRigControlHierarchy::GetSafeNewName(const FName& InPotentialNewName) const
+{
+	FName Name = InPotentialNewName;
+	int32 Suffix = 1;
+	while(!IsNameAvailable(Name))
+	{
+		Name = *FString::Printf(TEXT("%s_%d"), *InPotentialNewName.ToString(), ++Suffix);
+	}
+	return Name;
+}
+
+FRigControl& FRigControlHierarchy::Add(const FName& InNewName, const FName& InParentName, const FName& InSpaceName, const FTransform& InTransform)
+{
+	int32 ParentIndex = GetIndex(InParentName);
+	bool bHasParent = (ParentIndex != INDEX_NONE);
+
+	FRigControl NewControl;
+	NewControl.Name = GetSafeNewName(InNewName);
+	NewControl.ParentIndex = ParentIndex;
+	NewControl.ParentName = bHasParent ? InParentName : NAME_None;;
+	NewControl.SpaceIndex = INDEX_NONE;
+	NewControl.SpaceName = NAME_None;
+	NewControl.InitialTransform = InTransform;
+	NewControl.LocalTransform = InTransform;
+
+	// todo: set control type, etc
+
+	Controls.Add(NewControl);
+	RefreshMapping();
+
+#if WITH_EDITOR
+	OnControlAdded.Broadcast(Container, RigElementType(), NewControl.Name);
+#endif
+
+	SetSpace(NewControl.Name, InSpaceName);
+
+	int32 Index = GetIndex(NewControl.Name);
+	return Controls[Index];
+}
+
+void FRigControlHierarchy::Reparent(const FName& InName, const FName& InNewParentName)
+{
+	int32 Index = GetIndex(InName);
+	// can't parent to itself
+	if (Index != INDEX_NONE && InName != InNewParentName)
+	{
+		// should allow reparent to none (no parent)
+		// if invalid, we consider to be none
+		int32 ParentIndex = GetIndex(InNewParentName);
+		bool bHasParent = (ParentIndex != INDEX_NONE);
+		FRigControl& CurControl = Controls[Index];
+		FName OldParentName = CurControl.ParentName;
+		CurControl.ParentName = (bHasParent)? InNewParentName : NAME_None;
+		CurControl.ParentIndex = ParentIndex;
+
+		// todo: we need to update the space if this had been parented by space as well before
+
+		// we want to make sure parent is before the child
+		RefreshMapping();
+
+#if WITH_EDITOR
+		if (OldParentName != InNewParentName)
+		{
+			OnControlReparented.Broadcast(Container, RigElementType(), InName, OldParentName, InNewParentName);
+		}
+#endif
+	}
+}
+
+void FRigControlHierarchy::SetSpace(const FName& InName, const FName& InNewSpaceName)
+{
+	// todo find the space, set the space index...
+}
+
+FRigControl FRigControlHierarchy::Remove(const FName& InNameToRemove)
+{
+	TArray<int32> Children;
+#if WITH_EDITOR
+	TArray<FName> RemovedChildControls;
+#endif
+	if (GetChildren(InNameToRemove, Children, true) > 0)
+	{
+		// sort by child index
+		Children.Sort([](const int32& A, const int32& B) { return A < B; });
+
+		// want to delete from end to the first 
+		for (int32 ChildIndex = Children.Num() - 1; ChildIndex >= 0; --ChildIndex)
+		{
+#if WITH_EDITOR
+			RemovedChildControls.Add(Controls[Children[ChildIndex]].Name);
+#endif
+			Controls.RemoveAt(Children[ChildIndex]);
+		}
+	}
+
+	int32 IndexToDelete = GetIndex(InNameToRemove);
+	FRigControl RemovedControl = Controls[IndexToDelete];
+	Controls.RemoveAt(IndexToDelete);
+
+	RefreshMapping();
+
+#if WITH_EDITOR
+	for (const FName& RemovedChildControl : RemovedChildControls)
+	{
+		OnControlRemoved.Broadcast(Container, RigElementType(), RemovedChildControl);
+	}
+	OnControlRemoved.Broadcast(Container, RigElementType(), RemovedControl.Name);
+#endif
+
+	return RemovedControl;
+}
+
+// list of names of children - this is not cheap, and is supposed to be used only for one time set up
+int32 FRigControlHierarchy::GetChildren(const FName& InName, TArray<int32>& OutChildren, bool bRecursively) const
+{
+	return GetChildren(GetIndex(InName), OutChildren, bRecursively);
+}
+
+int32 FRigControlHierarchy::GetChildren(const int32 InIndex, TArray<int32>& OutChildren, bool bRecursively) const
+{
+	OutChildren.Reset();
+
+	if (InIndex != INDEX_NONE)
+	{
+		GetChildrenRecursive(InIndex, OutChildren, bRecursively);
+	}
+
+	return OutChildren.Num();
+}
+
+FName FRigControlHierarchy::GetName(int32 InIndex) const
+{
+	if (Controls.IsValidIndex(InIndex))
+	{
+		return Controls[InIndex].Name;
+	}
+
+	return NAME_None;
+}
+
+int32 FRigControlHierarchy::GetIndexSlow(const FName& InName) const
+{
+	for (int32 Index = 0; Index < Controls.Num(); ++Index)
+	{
+		if (Controls[Index].Name == InName)
+		{
+			return Index;
+		}
+	}
+
+	return INDEX_NONE;
+}
+
+void FRigControlHierarchy::SetGlobalTransform(const FName& InName, const FTransform& InTransform)
+{
+	SetGlobalTransform(GetIndex(InName), InTransform);
+}
+
+void FRigControlHierarchy::SetGlobalTransform(int32 InIndex, const FTransform& InTransform)
+{
+	if (Controls.IsValidIndex(InIndex))
+	{
+		FRigControl& Control = Controls[InIndex];
+		// todo: get the space's transform + inverse
+		//Control.GlobalTransform = InTransform;
+		//Control.GlobalTransform.NormalizeRotation();
+	}
+}
+
+FTransform FRigControlHierarchy::GetGlobalTransform(const FName& InName) const
+{
+	return GetGlobalTransform(GetIndex(InName));
+}
+
+FTransform FRigControlHierarchy::GetGlobalTransform(int32 InIndex) const
+{
+	if (Controls.IsValidIndex(InIndex))
+	{
+		// todo: get the space's transform and compute
+		//return Controls[InIndex].GlobalTransform;
+		return FTransform::Identity;
+	}
+
+	return FTransform::Identity;
+}
+
+void FRigControlHierarchy::SetLocalTransform(const FName& InName, const FTransform& InTransform)
+{
+	SetLocalTransform(GetIndex(InName), InTransform);
+}
+
+void FRigControlHierarchy::SetLocalTransform(int32 InIndex, const FTransform& InTransform)
+{
+	if (Controls.IsValidIndex(InIndex))
+	{
+		FRigControl& Control = Controls[InIndex];
+		Control.LocalTransform = InTransform;
+		Control.LocalTransform.NormalizeRotation();
+	}
+}
+
+FTransform FRigControlHierarchy::GetLocalTransform(const FName& InName) const
+{
+	return GetLocalTransform(GetIndex(InName));
+}
+
+FTransform FRigControlHierarchy::GetLocalTransform(int32 InIndex) const
+{
+	if (Controls.IsValidIndex(InIndex))
+	{
+		return Controls[InIndex].LocalTransform;
+	}
+
+	return FTransform::Identity;
+}
+
+void FRigControlHierarchy::SetInitialTransform(const FName& InName, const FTransform& InTransform)
+{
+	SetInitialTransform(GetIndex(InName), InTransform);
+}
+
+void FRigControlHierarchy::SetInitialTransform(int32 InIndex, const FTransform& InTransform)
+{
+	if (Controls.IsValidIndex(InIndex))
+	{
+		FRigControl& Control = Controls[InIndex];
+		Control.InitialTransform = InTransform;
+		Control.InitialTransform.NormalizeRotation();
+	}
+}
+
+FTransform FRigControlHierarchy::GetInitialTransform(const FName& InName) const
+{
+	return GetInitialTransform(GetIndex(InName));
+}
+
+FTransform FRigControlHierarchy::GetInitialTransform(int32 InIndex) const
+{
+	if (Controls.IsValidIndex(InIndex))
+	{
+		return Controls[InIndex].InitialTransform;
+	}
+
+	return FTransform::Identity;
+}
+
+FName FRigControlHierarchy::Rename(const FName& InOldName, const FName& InNewName)
+{
+	if (InOldName != InNewName)
+	{
+		const int32 Found = GetIndex(InOldName);
+		if (Found != INDEX_NONE)
+		{
+			FName NewName = GetSafeNewName(InNewName);
+			Controls[Found].Name = NewName;
+
+			// go through find all children and rename them
+#if WITH_EDITOR
+			TArray<FName> ReparentedControls;
+#endif
+			for (int32 Index = 0; Index < Controls.Num(); ++Index)
+			{
+				if (Controls[Index].ParentName == InOldName)
+				{
+					Controls[Index].ParentName = NewName;
+#if WITH_EDITOR
+					ReparentedControls.Add(Controls[Index].Name);
+#endif
+				}
+			}
+
+			RefreshMapping();
+
+#if WITH_EDITOR
+			OnControlRenamed.Broadcast(Container, RigElementType(), InOldName, NewName);
+			for (const FName& ReparentedControl : ReparentedControls)
+			{
+				OnControlReparented.Broadcast(Container, RigElementType(), ReparentedControl, InOldName, NewName);
+			}
+#endif
+			return NewName;
+		}
+	}
+
+	return NAME_None;
+}
+
+void FRigControlHierarchy::Sort()
+{
+	TMap<int32, TArray<int32>> HierarchyTree;
+
+	TArray<int32> SortedArray;
+
+	// first figure out children and find roots
+	for (int32 Index = 0; Index < Controls.Num(); ++Index)
+	{
+		int32 ParentIndex = GetIndexSlow(Controls[Index].ParentName);
+		if (ParentIndex != INDEX_NONE)
+		{
+			TArray<int32>& ChildIndices = HierarchyTree.FindOrAdd(ParentIndex);
+			ChildIndices.Add(Index);
+		}
+		else
+		{
+			// as as a root
+			HierarchyTree.Add(Index);
+			// add them to the list first
+			SortedArray.Add(Index);
+		}
+	}
+
+	// now go through map and add to sorted array
+	for (int32 SortIndex = 0; SortIndex < SortedArray.Num(); ++SortIndex)
+	{
+		// add children of sorted array
+		TArray<int32>* ChildIndices = HierarchyTree.Find(SortedArray[SortIndex]);
+		if (ChildIndices)
+		{
+			// append children now
+			SortedArray.Append(*ChildIndices);
+			// now sorted array will grow
+			// this is same as BFS as it starts from all roots, and going down
+		}
+	}
+
+	check(SortedArray.Num() == Controls.Num());
+
+	// create new list with sorted
+	TArray<FRigControl> NewSortedList;
+	NewSortedList.AddDefaulted(Controls.Num());
+	for (int32 NewIndex = 0; NewIndex < SortedArray.Num(); ++NewIndex)
+	{
+		NewSortedList[NewIndex] = Controls[SortedArray[NewIndex]];
+	}
+
+	Controls = MoveTemp(NewSortedList);
+
+	// now fix up parent Index
+	for (int32 Index = 0; Index < Controls.Num(); ++Index)
+	{
+		Controls[Index].ParentIndex = GetIndexSlow(Controls[Index].ParentName);
+		// parent index always should be less than this index, even if invalid
+		check(Controls[Index].ParentIndex < Index);
+	}
+}
+
+void FRigControlHierarchy::RefreshMapping()
+{
+	Sort();
+
+	NameToIndexMapping.Empty();
+	for (int32 Index = 0; Index < Controls.Num(); ++Index)
+	{
+		Controls[Index].Index = Index;
+		NameToIndexMapping.Add(Controls[Index].Name, Index);
+	}
+}
+
+void FRigControlHierarchy::Initialize()
+{
+	RefreshMapping();
+
+	// update parent index
+	for (int32 Index = 0; Index < Controls.Num(); ++Index)
+	{
+		Controls[Index].ParentIndex = GetIndex(Controls[Index].ParentName);
+		// todo
+		//Controls[Index].SpaceIndex = GetIndex(Controls[Index].ParentName);
+	}
+
+	// initialize transform
+	for (int32 Index = 0; Index < Controls.Num(); ++Index)
+	{
+		Controls[Index].LocalTransform = Controls[Index].InitialTransform;
+
+		// update children
+		GetChildren(Index, Controls[Index].Dependents, false);
+	}
+}
+
+void FRigControlHierarchy::Reset()
+{
+	Controls.Reset();
+}
+
+void FRigControlHierarchy::ResetTransforms()
+{
+	// initialize transform
+	for (int32 Index = 0; Index < Controls.Num(); ++Index)
+	{
+		Controls[Index].LocalTransform = Controls[Index].InitialTransform;
+	}
+}
+
+int32 FRigControlHierarchy::GetChildrenRecursive(const int32 InIndex, TArray<int32>& OutChildren, bool bRecursively) const
+{
+	const int32 StartChildIndex = OutChildren.Num();
+
+	// all children should be later than parent
+	for (int32 ChildIndex = InIndex + 1; ChildIndex < Controls.Num(); ++ChildIndex)
+	{
+		if (Controls[ChildIndex].ParentIndex == InIndex)
+		{
+			OutChildren.AddUnique(ChildIndex);
+		}
+	}
+
+	if (bRecursively)
+	{
+		// since we keep appending inside of functions, we make sure not to go over original list
+		const int32 EndChildIndex = OutChildren.Num() - 1;
+		for (int32 ChildIndex = StartChildIndex; ChildIndex <= EndChildIndex; ++ChildIndex)
+		{
+			GetChildrenRecursive(OutChildren[ChildIndex], OutChildren, bRecursively);
+		}
+	}
+
+	return OutChildren.Num();
+}
diff --git a/Engine/Plugins/Experimental/ControlRig/Source/ControlRig/Private/Rigs/RigCurveContainer.cpp b/Engine/Plugins/Experimental/ControlRig/Source/ControlRig/Private/Rigs/RigCurveContainer.cpp
--- a/Engine/Plugins/Experimental/ControlRig/Source/ControlRig/Private/Rigs/RigCurveContainer.cpp
+++ b/Engine/Plugins/Experimental/ControlRig/Source/ControlRig/Private/Rigs/RigCurveContainer.cpp
@@ -9,63 +9,83 @@
 ////////////////////////////////////////////////////////////////////////////////
 
 FRigCurveContainer::FRigCurveContainer()
+#if WITH_EDITOR
+	:Container(nullptr)
+#endif
 {
 }
 
 FRigCurveContainer& FRigCurveContainer::operator= (const FRigCurveContainer &InOther)
 {
-	/*
 #if WITH_EDITOR
 	for (int32 Index = Num() - 1; Index >= 0; Index--)
 	{
-		FRigSpace SpaceToRemove = Spaces[Index];
-		OnSpaceRemoved.Broadcast(Container, ERigHierarchyElementType::Space, SpaceToRemove.Name);
+		FRigCurve CurveToRemove = Curves[Index];
+		OnCurveRemoved.Broadcast(Container, RigElementType(), CurveToRemove.Name);
 	}
 #endif
-	*/
 
 	Curves.Reset();
 	Curves.Append(InOther.Curves);
 	NameToIndexMapping.Reset();
 	RefreshMapping();
 
-	/*
 #if WITH_EDITOR
-	for (const FRigSpace& SpaceAdded : Spaces)
+	for (const FRigCurve& CurveAdded : Curves)
 	{
-		OnSpaceAdded.Broadcast(Container, ERigHierarchyElementType::Space, SpaceAdded.Name);
+		OnCurveAdded.Broadcast(Container, RigElementType(), CurveAdded.Name);
 	}
 #endif
-	*/
 
 	return *this;
 }
 
-void FRigCurveContainer::Add(const FName& InNewName)
+FName FRigCurveContainer::GetSafeNewName(const FName& InPotentialNewName) const
 {
-	DECLARE_SCOPE_HIERARCHICAL_COUNTER_FUNC()
-
-	int32 Found = GetIndex(InNewName);
-	if (Found == INDEX_NONE)
+	FName Name = InPotentialNewName;
+	int32 Suffix = 1;
+	while(!IsNameAvailable(Name))
 	{
-		FRigCurve NewCurve;
-		NewCurve.Name = InNewName;
-		NewCurve.Value = 0.f;
-		Curves.Add(NewCurve);
-		RefreshMapping();
+		Name = *FString::Printf(TEXT("%s_%d"), *InPotentialNewName.ToString(), ++Suffix);
 	}
+	return Name;
+}
+
+FRigCurve& FRigCurveContainer::Add(const FName& InNewName)
+{
+	DECLARE_SCOPE_HIERARCHICAL_COUNTER_FUNC()
+
+	FRigCurve NewCurve;
+	NewCurve.Name = GetSafeNewName(InNewName);
+	NewCurve.Value = 0.f;
+	Curves.Add(NewCurve);
+	RefreshMapping();
+
+#if WITH_EDITOR
+	OnCurveAdded.Broadcast(Container, RigElementType(), NewCurve.Name);
+#endif
+
+	int32 Index = GetIndex(NewCurve.Name);
+	return Curves[Index];
 }
 
-void FRigCurveContainer::Remove(const FName& InName)
+FRigCurve FRigCurveContainer::Remove(const FName& InName)
 {
 	DECLARE_SCOPE_HIERARCHICAL_COUNTER_FUNC()
 
+	FRigCurve RemovedCurve;
+
 	int32 IndexToDelete = GetIndex(InName);
-	if (IndexToDelete != INDEX_NONE)
-	{
-		Curves.RemoveAt(IndexToDelete);
-		RefreshMapping();
-	}
+	ensure(IndexToDelete != INDEX_NONE);
+	RemovedCurve = Curves[IndexToDelete];
+	Curves.RemoveAt(IndexToDelete);
+	RefreshMapping();
+
+#if WITH_EDITOR
+	OnCurveRemoved.Broadcast(Container, RigElementType(), RemovedCurve.Name);
+#endif
+
+	return RemovedCurve;
 }
 
 FName FRigCurveContainer::GetName(int32 InIndex) const
@@ -120,21 +140,27 @@ float FRigCurveContainer::GetValue(int32 InIndex) const
 	return 0.f;
 }
 
-void FRigCurveContainer::Rename(const FName& InOldName, const FName& InNewName)
+FName FRigCurveContainer::Rename(const FName& InOldName, const FName& InNewName)
 {
 	DECLARE_SCOPE_HIERARCHICAL_COUNTER_FUNC()
 
 	if (InOldName != InNewName)
 	{
 		const int32 Found = GetIndex(InOldName);
-		const int32 NewNameFound = GetIndex(InNewName);
-		// if I have new name, and didn't find new name
-		if (Found != INDEX_NONE && NewNameFound == INDEX_NONE)
+		if (Found != INDEX_NONE)
 		{
-			Curves[Found].Name = InNewName;
+			FName NewName = GetSafeNewName(InNewName);
+			Curves[Found].Name = NewName;
 			RefreshMapping();
+
+#if WITH_EDITOR
+			OnCurveRenamed.Broadcast(Container, RigElementType(), InOldName, NewName);
+#endif
+			return NewName;
 		}
 	}
+	
+	return NAME_None;
 }
 
 void FRigCurveContainer::RefreshMapping()
@@ -144,6 +170,7 @@ void FRigCurveContainer::RefreshMapping()
 	NameToIndexMapping.Empty();
 	for (int32 Index = 0; Index < Curves.Num(); ++Index)
 	{
+		Curves[Index].Index = Index;
 		NameToIndexMapping.Add(Curves[Index].Name, Index);
 	}
 }
diff --git a/Engine/Plugins/Experimental/ControlRig/Source/ControlRig/Private/Rigs/RigHierarchyContainer.cpp b/Engine/Plugins/Experimental/ControlRig/Source/ControlRig/Private/Rigs/RigHierarchyContainer.cpp
--- a/Engine/Plugins/Experimental/ControlRig/Source/ControlRig/Private/Rigs/RigHierarchyContainer.cpp
+++ b/Engine/Plugins/Experimental/ControlRig/Source/ControlRig/Private/Rigs/RigHierarchyContainer.cpp
@@ -16,6 +16,9 @@ FRigHierarchyContainer::FRigHierarchyContainer()
 FRigHierarchyContainer& FRigHierarchyContainer::operator= (const FRigHierarchyContainer &InOther)
 {
 	BoneHierarchy = InOther.BoneHierarchy;
+	SpaceHierarchy = InOther.SpaceHierarchy;
+	ControlHierarchy = InOther.ControlHierarchy;
+	CurveContainer = InOther.CurveContainer;
 	return *this;
 }
 
@@ -23,6 +26,9 @@ void FRigHierarchyContainer::Initialize()
 {
 #if WITH_EDITOR
 	BoneHierarchy.Container = this;
+	SpaceHierarchy.Container = this;
+	ControlHierarchy.Container = this;
+	CurveContainer.Container = this;
 
 	BoneHierarchy.OnBoneAdded.RemoveAll(this);
 	BoneHierarchy.OnBoneRemoved.RemoveAll(this);
@@ -33,9 +39,38 @@ void FRigHierarchyContainer::Initialize()
 	BoneHierarchy.OnBoneRemoved.AddRaw(this, &FRigHierarchyContainer::HandleOnElementRemoved);
 	BoneHierarchy.OnBoneRenamed.AddRaw(this, &FRigHierarchyContainer::HandleOnElementRenamed);
 	BoneHierarchy.OnBoneReparented.AddRaw(this, &FRigHierarchyContainer::HandleOnElementReparented);
+
+	SpaceHierarchy.OnSpaceAdded.RemoveAll(this);
+	SpaceHierarchy.OnSpaceRemoved.RemoveAll(this);
+	SpaceHierarchy.OnSpaceRenamed.RemoveAll(this);
+
+	SpaceHierarchy.OnSpaceAdded.AddRaw(this, &FRigHierarchyContainer::HandleOnElementAdded);
+	SpaceHierarchy.OnSpaceRemoved.AddRaw(this, &FRigHierarchyContainer::HandleOnElementRemoved);
+	SpaceHierarchy.OnSpaceRenamed.AddRaw(this, &FRigHierarchyContainer::HandleOnElementRenamed);
+
+	ControlHierarchy.OnControlAdded.RemoveAll(this);
+	ControlHierarchy.OnControlRemoved.RemoveAll(this);
+	ControlHierarchy.OnControlRenamed.RemoveAll(this);
+	ControlHierarchy.OnControlReparented.RemoveAll(this);
+
+	ControlHierarchy.OnControlAdded.AddRaw(this, &FRigHierarchyContainer::HandleOnElementAdded);
+	ControlHierarchy.OnControlRemoved.AddRaw(this, &FRigHierarchyContainer::HandleOnElementRemoved);
+	ControlHierarchy.OnControlRenamed.AddRaw(this, &FRigHierarchyContainer::HandleOnElementRenamed);
+	ControlHierarchy.OnControlReparented.AddRaw(this, &FRigHierarchyContainer::HandleOnElementReparented);
+
+	CurveContainer.OnCurveAdded.RemoveAll(this);
+	CurveContainer.OnCurveRemoved.RemoveAll(this);
+	CurveContainer.OnCurveRenamed.RemoveAll(this);
+
+	CurveContainer.OnCurveAdded.AddRaw(this, &FRigHierarchyContainer::HandleOnElementAdded);
+	CurveContainer.OnCurveRemoved.AddRaw(this, &FRigHierarchyContainer::HandleOnElementRemoved);
+	CurveContainer.OnCurveRenamed.AddRaw(this, &FRigHierarchyContainer::HandleOnElementRenamed);
 #endif
 
 	BoneHierarchy.Initialize();
+	SpaceHierarchy.Initialize();
+	ControlHierarchy.Initialize();
+	CurveContainer.Initialize();
 
 	ResetTransforms();
 }
@@ -43,39 +78,45 @@ void FRigHierarchyContainer::Initialize()
 void FRigHierarchyContainer::Reset()
 {
 	BoneHierarchy.Reset();
-	
+	SpaceHierarchy.Reset();
+	ControlHierarchy.Reset();
+	CurveContainer.Reset();
+
 	Initialize();
 }
 
 void FRigHierarchyContainer::ResetTransforms()
 {
 	BoneHierarchy.ResetTransforms();
+	SpaceHierarchy.ResetTransforms();
+	ControlHierarchy.ResetTransforms();
+	CurveContainer.ResetValues();
 }
 
 #if WITH_EDITOR
 
-void FRigHierarchyContainer::HandleOnElementAdded(FRigHierarchyContainer* InContainer, ERigHierarchyElementType InElementType, const FName& InName)
+void FRigHierarchyContainer::HandleOnElementAdded(FRigHierarchyContainer* InContainer, ERigElementType InElementType, const FName& InName)
 {
 	// todo
 	OnElementAdded.Broadcast(InContainer, InElementType, InName);
 	OnElementChanged.Broadcast(InContainer, InElementType, InName);
 }
 
-void FRigHierarchyContainer::HandleOnElementRemoved(FRigHierarchyContainer* InContainer, ERigHierarchyElementType InElementType, const FName& InName)
+void FRigHierarchyContainer::HandleOnElementRemoved(FRigHierarchyContainer* InContainer, ERigElementType InElementType, const FName& InName)
 {
 	// todo
 	OnElementRemoved.Broadcast(InContainer, InElementType, InName);
 	OnElementChanged.Broadcast(InContainer, InElementType, InName);
 }
 
-void FRigHierarchyContainer::HandleOnElementRenamed(FRigHierarchyContainer* InContainer, ERigHierarchyElementType InElementType, const FName& InOldName, const FName& InNewName)
+void FRigHierarchyContainer::HandleOnElementRenamed(FRigHierarchyContainer* InContainer, ERigElementType InElementType, const FName& InOldName, const FName& InNewName)
 {
 	// todo
 	OnElementRenamed.Broadcast(InContainer, InElementType, InOldName, InNewName);
 	OnElementChanged.Broadcast(InContainer, InElementType, InNewName);
 }
 
-void FRigHierarchyContainer::HandleOnElementReparented(FRigHierarchyContainer* InContainer, ERigHierarchyElementType InElementType, const FName& InName, const FName& InOldParentName, const FName& InNewParentName)
+void FRigHierarchyContainer::HandleOnElementReparented(FRigHierarchyContainer* InContainer, ERigElementType InElementType, const FName& InName, const FName& InOldParentName, const FName& InNewParentName)
 {
 	// todo
 	OnElementReparented.Broadcast(InContainer, InElementType, InName, InOldParentName, InNewParentName);
@@ -83,33 +124,3 @@ void FRigHierarchyContainer::HandleOnElementReparented(FRigHierarchyContainer* I
 }
 
 #endif
-
-////////////////////////////////////////////////////////////////////////////////
-// FRigHierarchyRef
-////////////////////////////////////////////////////////////////////////////////
-
-FRigHierarchyRef::FRigHierarchyRef()
-: Container(nullptr)
-{
-
-}
-
-FRigBoneHierarchy* FRigHierarchyRef::GetBones()
-{
-	return GetBonesInternal();
-}
-
-const FRigBoneHierarchy* FRigHierarchyRef::GetBones() const
-{
-	return (const FRigBoneHierarchy*)GetBonesInternal();
-}
-
-FRigBoneHierarchy* FRigHierarchyRef::GetBonesInternal() const
-{
-	if (Container)
-	{
-		return &Container->BoneHierarchy;
-	}
-	return nullptr;
-}
-
diff --git a/Engine/Plugins/Experimental/ControlRig/Source/ControlRig/Private/Rigs/RigSpaceHierarchy.cpp b/Engine/Plugins/Experimental/ControlRig/Source/ControlRig/Private/Rigs/RigSpaceHierarchy.cpp
new file mode 100644
--- /dev/null
+++ b/Engine/Plugins/Experimental/ControlRig/Source/ControlRig/Private/Rigs/RigSpaceHierarchy.cpp
@@ -0,0 +1,286 @@
+// Copyright 1998-2019 Epic Games, Inc. All Rights Reserved.
+
+#include "Rigs/RigSpaceHierarchy.h"
+#include "ControlRig.h"
+#include "HelperUtil.h"
+
+////////////////////////////////////////////////////////////////////////////////
+// FRigSpaceHierarchy
+////////////////////////////////////////////////////////////////////////////////
+
+FRigSpaceHierarchy::FRigSpaceHierarchy()
+#if WITH_EDITOR
+	:Container(nullptr)
+#endif
+{
+}
+
+FRigSpaceHierarchy& FRigSpaceHierarchy::operator= (const FRigSpaceHierarchy &InOther)
+{
+#if WITH_EDITOR
+	for (int32 Index = Num() - 1; Index >= 0; Index--)
+	{
+		FRigSpace SpaceToRemove = Spaces[Index];
+		OnSpaceRemoved.Broadcast(Container, RigElementType(), SpaceToRemove.Name);
+	}
+#endif
+
+	Spaces.Reset();
+	Spaces.Append(InOther.Spaces);
+	NameToIndexMapping.Reset();
+	RefreshMapping();
+
+#if WITH_EDITOR
+	for (const FRigSpace& SpaceAdded : Spaces)
+	{
+		OnSpaceAdded.Broadcast(Container, RigElementType(), SpaceAdded.Name);
+	}
+#endif
+
+	return *this;
+}
+
+FName FRigSpaceHierarchy::GetSafeNewName(const FName& InPotentialNewName) const
+{
+	FName Name = InPotentialNewName;
+	int32 Suffix = 1;
+	while(!IsNameAvailable(Name))
+	{
+		Name = *FString::Printf(TEXT("%s_%d"), *InPotentialNewName.ToString(), ++Suffix);
+	}
+	return Name;
+}
+
+FRigSpace& FRigSpaceHierarchy::Add(const FName& InNewName, ERigSpaceType InSpaceType, const FName& InParentName, const FTransform& InTransform)
+{
+	FRigSpace NewSpace;
+	NewSpace.Name = GetSafeNewName(InNewName);
+	NewSpace.SpaceType = InSpaceType;
+	NewSpace.ParentName = InParentName;
+	NewSpace.ParentIndex = INDEX_NONE;
+	NewSpace.InitialTransform = InTransform;
+	NewSpace.LocalTransform = InTransform;
+	Spaces.Add(NewSpace);
+	RefreshMapping();
+
+#if WITH_EDITOR
+	OnSpaceAdded.Broadcast(Container, RigElementType(), NewSpace.Name);
+#endif
+
+	int32 Index = GetIndex(NewSpace.Name);
+	return Spaces[Index];
+}
+
+void FRigSpaceHierarchy::Reparent(const FName& InName, const FName& InNewParentName)
+{
+	int32 Index = GetIndex(InName);
+	if (Index != INDEX_NONE)
+	{
+		Spaces[Index].ParentName = InNewParentName;
+		Spaces[Index].ParentIndex = INDEX_NONE;
+	}
+}
+
+FRigSpace FRigSpaceHierarchy::Remove(const FName& InNameToRemove)
+{
+	int32 IndexToDelete = GetIndex(InNameToRemove);
+	FRigSpace RemovedSpace = Spaces[IndexToDelete];
+	Spaces.RemoveAt(IndexToDelete);
+
+	RefreshMapping();
+
+#if WITH_EDITOR
+	OnSpaceRemoved.Broadcast(Container, RigElementType(), RemovedSpace.Name);
+#endif
+
+	return RemovedSpace;
+}
+
+FName FRigSpaceHierarchy::GetName(int32 InIndex) const
+{
+	if (Spaces.IsValidIndex(InIndex))
+	{
+		return Spaces[InIndex].Name;
+	}
+
+	return NAME_None;
+}
+
+int32 FRigSpaceHierarchy::GetIndexSlow(const FName& InName) const
+{
+	for (int32 Index = 0; Index < Spaces.Num(); ++Index)
+	{
+		if (Spaces[Index].Name == InName)
+		{
+			return Index;
+		}
+	}
+
+	return INDEX_NONE;
+}
+
+void FRigSpaceHierarchy::SetGlobalTransform(const FName& InName, const FTransform& InTransform)
+{
+	SetGlobalTransform(GetIndex(InName), InTransform);
+}
+
+void FRigSpaceHierarchy::SetGlobalTransform(int32 InIndex, const FTransform& InTransform)
+{
+	if (Spaces.IsValidIndex(InIndex))
+	{
+		FRigSpace& Space = Spaces[InIndex];
+		// todo: get the space's transform + inverse
+		//Space.GlobalTransform = InTransform;
+		//Space.GlobalTransform.NormalizeRotation();
+	}
+}
+
+FTransform FRigSpaceHierarchy::GetGlobalTransform(const FName& InName) const
+{
+	return GetGlobalTransform(GetIndex(InName));
+}
+
+FTransform FRigSpaceHierarchy::GetGlobalTransform(int32 InIndex) const
+{
+	if (Spaces.IsValidIndex(InIndex))
+	{
+		// todo: get the space's transform and compute
+		//return Spaces[InIndex].GlobalTransform;
+		return FTransform::Identity;
+	}
+
+	return FTransform::Identity;
+}
+
+void FRigSpaceHierarchy::SetLocalTransform(const FName& InName, const FTransform& InTransform)
+{
+	SetLocalTransform(GetIndex(InName), InTransform);
+}
+
+void FRigSpaceHierarchy::SetLocalTransform(int32 InIndex, const FTransform& InTransform)
+{
+	if (Spaces.IsValidIndex(InIndex))
+	{
+		FRigSpace& Space = Spaces[InIndex];
+		Space.LocalTransform = InTransform;
+		Space.LocalTransform.NormalizeRotation();
+	}
+}
+
+FTransform FRigSpaceHierarchy::GetLocalTransform(const FName& InName) const
+{
+	return GetLocalTransform(GetIndex(InName));
+}
+
+FTransform FRigSpaceHierarchy::GetLocalTransform(int32 InIndex) const
+{
+	if (Spaces.IsValidIndex(InIndex))
+	{
+		return Spaces[InIndex].LocalTransform;
+	}
+
+	return FTransform::Identity;
+}
+
+void FRigSpaceHierarchy::SetInitialTransform(const FName& InName, const FTransform& InTransform)
+{
+	SetInitialTransform(GetIndex(InName), InTransform);
+}
+
+void FRigSpaceHierarchy::SetInitialTransform(int32 InIndex, const FTransform& InTransform)
+{
+	if (Spaces.IsValidIndex(InIndex))
+	{
+		FRigSpace& Space = Spaces[InIndex];
+		Space.InitialTransform = InTransform;
+		Space.InitialTransform.NormalizeRotation();
+	}
+}
+
+FTransform FRigSpaceHierarchy::GetInitialTransform(const FName& InName) const
+{
+	return GetInitialTransform(GetIndex(InName));
+}
+
+FTransform FRigSpaceHierarchy::GetInitialTransform(int32 InIndex) const
+{
+	if (Spaces.IsValidIndex(InIndex))
+	{
+		return Spaces[InIndex].InitialTransform;
+	}
+
+	return FTransform::Identity;
+}
+
+FName FRigSpaceHierarchy::Rename(const FName& InOldName, const FName& InNewName)
+{
+	if (InOldName != InNewName)
+	{
+		const int32 Found = GetIndex(InOldName);
+		if (Found != INDEX_NONE)
+		{
+			FName NewName = GetSafeNewName(InNewName);
+			Spaces[Found].Name = NewName;
+
+			// go through find all children and rename them
+			for (int32 Index = 0; Index < Spaces.Num(); ++Index)
+			{
+				if (Spaces[Index].ParentName == InOldName)
+				{
+					Spaces[Index].ParentName = NewName;
+				}
+			}
+
+			RefreshMapping();
+
+#if WITH_EDITOR
+			OnSpaceRenamed.Broadcast(Container, RigElementType(), InOldName, NewName);
+#endif
+			return NewName;
+		}
+	}
+
+	return NAME_None;
+}
+
+void FRigSpaceHierarchy::RefreshMapping()
+{
+	NameToIndexMapping.Empty();
+	for (int32 Index = 0; Index < Spaces.Num(); ++Index)
+	{
+		Spaces[Index].Index = Index;
+		NameToIndexMapping.Add(Spaces[Index].Name, Index);
+	}
+}
+
+void FRigSpaceHierarchy::Initialize()
+{
+	RefreshMapping();
+
+	// update parent index
+	for (int32 Index = 0; Index < Spaces.Num(); ++Index)
+	{
+		// todo
+		//Spaces[Index].ParentIndex = GetIndex(Spaces[Index].ParentName);
+	}
+
+	// initialize transform
+	for (int32 Index = 0; Index < Spaces.Num(); ++Index)
+	{
+		Spaces[Index].LocalTransform = Spaces[Index].InitialTransform;
+	}
+}
+
+void FRigSpaceHierarchy::Reset()
+{
+	Spaces.Reset();
+}
+
+void FRigSpaceHierarchy::ResetTransforms()
+{
+	// initialize transform
+	for (int32 Index = 0; Index < Spaces.Num(); ++Index)
+	{
+		Spaces[Index].LocalTransform = Spaces[Index].InitialTransform;
+	}
+}
diff --git a/Engine/Plugins/Experimental/ControlRig/Source/ControlRig/Private/Units/Debug/RigUnit_DebugBezier.cpp b/Engine/Plugins/Experimental/ControlRig/Source/ControlRig/Private/Units/Debug/RigUnit_DebugBezier.cpp
--- a/Engine/Plugins/Experimental/ControlRig/Source/ControlRig/Private/Units/Debug/RigUnit_DebugBezier.cpp
+++ b/Engine/Plugins/Experimental/ControlRig/Source/ControlRig/Private/Units/Debug/RigUnit_DebugBezier.cpp
@@ -17,11 +17,10 @@ void FRigUnit_DebugBezier::Execute(const FRigUnitContext& Context)
 	}
 
 	FTransform Transform = WorldOffset;
-	if (Space != NAME_None && Context.HierarchyReference.GetBones() != nullptr)
+	if (Space != NAME_None && Context.GetBones() != nullptr)
 	{
-		Transform = Transform * Context.HierarchyReference.GetBones()->GetGlobalTransform(Space);
+		Transform = Transform * Context.GetBones()->GetGlobalTransform(Space);
 	}
 
-
 	Context.DrawInterface->DrawBezier(Transform, Bezier, MinimumU, MaximumU, Color, Thickness, Detail);
 }
\ No newline at end of file
diff --git a/Engine/Plugins/Experimental/ControlRig/Source/ControlRig/Private/Units/Debug/RigUnit_DebugHierarchy.cpp b/Engine/Plugins/Experimental/ControlRig/Source/ControlRig/Private/Units/Debug/RigUnit_DebugHierarchy.cpp
--- a/Engine/Plugins/Experimental/ControlRig/Source/ControlRig/Private/Units/Debug/RigUnit_DebugHierarchy.cpp
+++ b/Engine/Plugins/Experimental/ControlRig/Source/ControlRig/Private/Units/Debug/RigUnit_DebugHierarchy.cpp
@@ -16,7 +16,7 @@ void FRigUnit_DebugHierarchy::Execute(const FRigUnitContext& Context)
 		return;
 	}
 
-	const FRigBoneHierarchy* Hierarchy = Context.HierarchyReference.GetBones();
+	const FRigBoneHierarchy* Hierarchy = Context.GetBones();
 	if (Hierarchy)
 	{
 		Context.DrawInterface->DrawHierarchy(WorldOffset, *Hierarchy, Mode, Scale, Color, Thickness);
diff --git a/Engine/Plugins/Experimental/ControlRig/Source/ControlRig/Private/Units/Debug/RigUnit_DebugLine.cpp b/Engine/Plugins/Experimental/ControlRig/Source/ControlRig/Private/Units/Debug/RigUnit_DebugLine.cpp
--- a/Engine/Plugins/Experimental/ControlRig/Source/ControlRig/Private/Units/Debug/RigUnit_DebugLine.cpp
+++ b/Engine/Plugins/Experimental/ControlRig/Source/ControlRig/Private/Units/Debug/RigUnit_DebugLine.cpp
@@ -17,9 +17,9 @@ void FRigUnit_DebugLine::Execute(const FRigUnitContext& Context)
 	}
 
 	FVector DrawA = A, DrawB = B;
-	if (Space != NAME_None && Context.HierarchyReference.GetBones() != nullptr)
+	if (Space != NAME_None && Context.GetBones() != nullptr)
 	{
-		FTransform Transform = Context.HierarchyReference.GetBones()->GetGlobalTransform(Space);
+		FTransform Transform = Context.GetBones()->GetGlobalTransform(Space);
 		DrawA = Transform.TransformPosition(DrawA);
 		DrawB = Transform.TransformPosition(DrawB);
 	}
diff --git a/Engine/Plugins/Experimental/ControlRig/Source/ControlRig/Private/Units/Debug/RigUnit_DebugLineStrip.cpp b/Engine/Plugins/Experimental/ControlRig/Source/ControlRig/Private/Units/Debug/RigUnit_DebugLineStrip.cpp
--- a/Engine/Plugins/Experimental/ControlRig/Source/ControlRig/Private/Units/Debug/RigUnit_DebugLineStrip.cpp
+++ b/Engine/Plugins/Experimental/ControlRig/Source/ControlRig/Private/Units/Debug/RigUnit_DebugLineStrip.cpp
@@ -16,9 +16,9 @@ void FRigUnit_DebugLineStrip::Execute(const FRigUnitContext& Context)
 		return;
 	}
 
-	if (Space != NAME_None && Context.HierarchyReference.GetBones() != nullptr)
+	if (Space != NAME_None && Context.GetBones() != nullptr)
 	{
-		FTransform Transform = Context.HierarchyReference.GetBones()->GetGlobalTransform(Space);
+		FTransform Transform = Context.GetBones()->GetGlobalTransform(Space);
 		TArray<FVector> PointsTransformed;
 		PointsTransformed.Reserve(Points.Num());
 		for(const FVector& Point : Points)
diff --git a/Engine/Plugins/Experimental/ControlRig/Source/ControlRig/Private/Units/Debug/RigUnit_DebugPoint.cpp b/Engine/Plugins/Experimental/ControlRig/Source/ControlRig/Private/Units/Debug/RigUnit_DebugPoint.cpp
--- a/Engine/Plugins/Experimental/ControlRig/Source/ControlRig/Private/Units/Debug/RigUnit_DebugPoint.cpp
+++ b/Engine/Plugins/Experimental/ControlRig/Source/ControlRig/Private/Units/Debug/RigUnit_DebugPoint.cpp
@@ -18,9 +18,9 @@ void FRigUnit_DebugPoint::Execute(const FRigUnitContext& Context)
 
 	FVector Center = FVector::ZeroVector;
 	FVector DrawVector = Vector;
-	if (Space != NAME_None && Context.HierarchyReference.GetBones() != nullptr)
+	if (Space != NAME_None && Context.GetBones() != nullptr)
 	{
-		FTransform Transform = Context.HierarchyReference.GetBones()->GetGlobalTransform(Space);
+		FTransform Transform = Context.GetBones()->GetGlobalTransform(Space);
 		Center = Transform.GetLocation();
 		DrawVector = Transform.TransformPosition(DrawVector);
 	}
@@ -55,9 +55,9 @@ void FRigUnit_DebugPointMutable::Execute(const FRigUnitContext& Context)
 
 	FVector Center = FVector::ZeroVector;
 	FVector DrawVector = Vector;
-	if (Space != NAME_None && Context.HierarchyReference.GetBones() != nullptr)
+	if (Space != NAME_None && Context.GetBones() != nullptr)
 	{
-		FTransform Transform = Context.HierarchyReference.GetBones()->GetGlobalTransform(Space);
+		FTransform Transform = Context.GetBones()->GetGlobalTransform(Space);
 		Center = Transform.GetLocation();
 		DrawVector = Transform.TransformPosition(DrawVector);
 	}
diff --git a/Engine/Plugins/Experimental/ControlRig/Source/ControlRig/Private/Units/Debug/RigUnit_DebugPrimitives.cpp b/Engine/Plugins/Experimental/ControlRig/Source/ControlRig/Private/Units/Debug/RigUnit_DebugPrimitives.cpp
--- a/Engine/Plugins/Experimental/ControlRig/Source/ControlRig/Private/Units/Debug/RigUnit_DebugPrimitives.cpp
+++ b/Engine/Plugins/Experimental/ControlRig/Source/ControlRig/Private/Units/Debug/RigUnit_DebugPrimitives.cpp
@@ -17,9 +17,9 @@ void FRigUnit_DebugRectangle::Execute(const FRigUnitContext& Context)
 	}
 
 	FTransform DrawTransform = Transform;
-	if (Space != NAME_None && Context.HierarchyReference.GetBones() != nullptr)
+	if (Space != NAME_None && Context.GetBones() != nullptr)
 	{
-		DrawTransform = DrawTransform * Context.HierarchyReference.GetBones()->GetGlobalTransform(Space);
+		DrawTransform = DrawTransform * Context.GetBones()->GetGlobalTransform(Space);
 	}
 
 	Context.DrawInterface->DrawRectangle(WorldOffset, DrawTransform, Scale, Color, Thickness);
@@ -39,9 +39,9 @@ void FRigUnit_DebugArc::Execute(const FRigUnitContext& Context)
 	}
 
 	FTransform DrawTransform = Transform;
-	if (Space != NAME_None && Context.HierarchyReference.GetBones() != nullptr)
+	if (Space != NAME_None && Context.GetBones() != nullptr)
 	{
-		DrawTransform = DrawTransform * Context.HierarchyReference.GetBones()->GetGlobalTransform(Space);
+		DrawTransform = DrawTransform * Context.GetBones()->GetGlobalTransform(Space);
 	}
 
 	Context.DrawInterface->DrawArc(WorldOffset, DrawTransform, Radius, FMath::DegreesToRadians(MinimumDegrees), FMath::DegreesToRadians(MaximumDegrees), Color, Thickness, Detail);
diff --git a/Engine/Plugins/Experimental/ControlRig/Source/ControlRig/Private/Units/Debug/RigUnit_DebugTransform.cpp b/Engine/Plugins/Experimental/ControlRig/Source/ControlRig/Private/Units/Debug/RigUnit_DebugTransform.cpp
--- a/Engine/Plugins/Experimental/ControlRig/Source/ControlRig/Private/Units/Debug/RigUnit_DebugTransform.cpp
+++ b/Engine/Plugins/Experimental/ControlRig/Source/ControlRig/Private/Units/Debug/RigUnit_DebugTransform.cpp
@@ -17,9 +17,9 @@ void FRigUnit_DebugTransform::Execute(const FRigUnitContext& Context)
 	}
 
 	FTransform DrawTransform = Transform;
-	if (Space != NAME_None && Context.HierarchyReference.GetBones() != nullptr)
+	if (Space != NAME_None && Context.GetBones() != nullptr)
 	{
-		DrawTransform = Transform * Context.HierarchyReference.GetBones()->GetGlobalTransform(Space);
+		DrawTransform = Transform * Context.GetBones()->GetGlobalTransform(Space);
 	}
 
 	switch (Mode)
@@ -57,9 +57,9 @@ void FRigUnit_DebugTransformMutable::Execute(const FRigUnitContext& Context)
 	}
 
 	FTransform DrawTransform = Transform;
-	if (Space != NAME_None && Context.HierarchyReference.GetBones() != nullptr)
+	if (Space != NAME_None && Context.GetBones() != nullptr)
 	{
-		DrawTransform = Transform * Context.HierarchyReference.GetBones()->GetGlobalTransform(Space);
+		DrawTransform = Transform * Context.GetBones()->GetGlobalTransform(Space);
 	}
 
 	switch (Mode)
@@ -97,9 +97,9 @@ void FRigUnit_DebugTransformArrayMutable::Execute(const FRigUnitContext& Context
 	}
 
 	DrawTransforms.SetNumUninitialized(Transforms.Num());
-	if (Space != NAME_None && Context.HierarchyReference.GetBones() != nullptr)
+	if (Space != NAME_None && Context.GetBones() != nullptr)
 	{
-		FTransform SpaceTransform = Context.HierarchyReference.GetBones()->GetGlobalTransform(Space);
+		FTransform SpaceTransform = Context.GetBones()->GetGlobalTransform(Space);
 		for(int32 Index=0;Index<Transforms.Num();Index++)
 		{
 			DrawTransforms[Index] = Transforms[Index] * SpaceTransform;
diff --git a/Engine/Plugins/Experimental/ControlRig/Source/ControlRig/Private/Units/Deprecated/RigUnit_AimConstraint.cpp b/Engine/Plugins/Experimental/ControlRig/Source/ControlRig/Private/Units/Deprecated/RigUnit_AimConstraint.cpp
--- a/Engine/Plugins/Experimental/ControlRig/Source/ControlRig/Private/Units/Deprecated/RigUnit_AimConstraint.cpp
+++ b/Engine/Plugins/Experimental/ControlRig/Source/ControlRig/Private/Units/Deprecated/RigUnit_AimConstraint.cpp
@@ -8,13 +8,12 @@
 void FRigUnit_AimConstraint::Execute(const FRigUnitContext& Context)
 {
     DECLARE_SCOPE_HIERARCHICAL_COUNTER_RIGUNIT()
-	FRigHierarchyRef& HierarchyRef = ExecuteContext.HierarchyReference;
 
 	if (Context.State == EControlRigState::Init)
 	{
 		ConstraintData.Reset();
 
-		FRigBoneHierarchy* Hierarchy = HierarchyRef.GetBones();
+		FRigBoneHierarchy* Hierarchy = ExecuteContext.GetBones();
 		if (Hierarchy)
 		{
 			int32 BoneIndex = Hierarchy->GetIndex(Joint);
@@ -45,7 +44,7 @@ void FRigUnit_AimConstraint::Execute(const FRigUnitContext& Context)
 	}
 	else if (Context.State == EControlRigState::Update)
 	{
-		FRigBoneHierarchy* Hierarchy = HierarchyRef.GetBones();
+		FRigBoneHierarchy* Hierarchy = ExecuteContext.GetBones();
 		if (Hierarchy)
 		{
 			int32 BoneIndex = Hierarchy->GetIndex(Joint);
diff --git a/Engine/Plugins/Experimental/ControlRig/Source/ControlRig/Private/Units/Deprecated/RigUnit_ApplyFK.cpp b/Engine/Plugins/Experimental/ControlRig/Source/ControlRig/Private/Units/Deprecated/RigUnit_ApplyFK.cpp
--- a/Engine/Plugins/Experimental/ControlRig/Source/ControlRig/Private/Units/Deprecated/RigUnit_ApplyFK.cpp
+++ b/Engine/Plugins/Experimental/ControlRig/Source/ControlRig/Private/Units/Deprecated/RigUnit_ApplyFK.cpp
@@ -7,11 +7,10 @@
 void FRigUnit_ApplyFK::Execute(const FRigUnitContext& Context)
 {
     DECLARE_SCOPE_HIERARCHICAL_COUNTER_RIGUNIT()
-	FRigHierarchyRef& HierarchyRef = ExecuteContext.HierarchyReference;
 	
 	if (Context.State == EControlRigState::Init)
 	{
-		FRigBoneHierarchy* Hierarchy = HierarchyRef.GetBones();
+		FRigBoneHierarchy* Hierarchy = ExecuteContext.GetBones();
 		if (!Hierarchy)
 		{
 			UnitLogHelpers::PrintMissingHierarchy(RigUnitName);
@@ -19,7 +18,7 @@ void FRigUnit_ApplyFK::Execute(const FRigUnitContext& Context)
 	}
 	else if (Context.State == EControlRigState::Update)
 	{
-		FRigBoneHierarchy* Hierarchy = HierarchyRef.GetBones();
+		FRigBoneHierarchy* Hierarchy = ExecuteContext.GetBones();
 		if (Hierarchy)
 		{
 			int32 Index = Hierarchy->GetIndex(Joint);
diff --git a/Engine/Plugins/Experimental/ControlRig/Source/ControlRig/Private/Units/Deprecated/RigUnit_CreateHierarchy.cpp b/Engine/Plugins/Experimental/ControlRig/Source/ControlRig/Private/Units/Deprecated/RigUnit_CreateHierarchy.cpp
deleted file mode 100644
diff --git a/Engine/Plugins/Experimental/ControlRig/Source/ControlRig/Private/Units/Deprecated/RigUnit_CreateHierarchy.h b/Engine/Plugins/Experimental/ControlRig/Source/ControlRig/Private/Units/Deprecated/RigUnit_CreateHierarchy.h
deleted file mode 100644
diff --git a/Engine/Plugins/Experimental/ControlRig/Source/ControlRig/Private/Units/Deprecated/RigUnit_GetJointTransform.cpp b/Engine/Plugins/Experimental/ControlRig/Source/ControlRig/Private/Units/Deprecated/RigUnit_GetJointTransform.cpp
--- a/Engine/Plugins/Experimental/ControlRig/Source/ControlRig/Private/Units/Deprecated/RigUnit_GetJointTransform.cpp
+++ b/Engine/Plugins/Experimental/ControlRig/Source/ControlRig/Private/Units/Deprecated/RigUnit_GetJointTransform.cpp
@@ -7,9 +7,8 @@
 void FRigUnit_GetJointTransform::Execute(const FRigUnitContext& Context)
 {
     DECLARE_SCOPE_HIERARCHICAL_COUNTER_RIGUNIT()
-	FRigHierarchyRef& HierarchyRef = ExecuteContext.HierarchyReference;
 
-	FRigBoneHierarchy* Hierarchy = HierarchyRef.GetBones();
+	FRigBoneHierarchy* Hierarchy = ExecuteContext.GetBones();
 	if (Hierarchy)
 	{
 		int32 Index = Hierarchy->GetIndex(Joint);
diff --git a/Engine/Plugins/Experimental/ControlRig/Source/ControlRig/Private/Units/Deprecated/RigUnit_MergeHierarchy.cpp b/Engine/Plugins/Experimental/ControlRig/Source/ControlRig/Private/Units/Deprecated/RigUnit_MergeHierarchy.cpp
deleted file mode 100644
diff --git a/Engine/Plugins/Experimental/ControlRig/Source/ControlRig/Private/Units/Deprecated/RigUnit_MergeHierarchy.h b/Engine/Plugins/Experimental/ControlRig/Source/ControlRig/Private/Units/Deprecated/RigUnit_MergeHierarchy.h
deleted file mode 100644
diff --git a/Engine/Plugins/Experimental/ControlRig/Source/ControlRig/Private/Units/Deprecated/RigUnit_TwoBoneIKFK.cpp b/Engine/Plugins/Experimental/ControlRig/Source/ControlRig/Private/Units/Deprecated/RigUnit_TwoBoneIKFK.cpp
--- a/Engine/Plugins/Experimental/ControlRig/Source/ControlRig/Private/Units/Deprecated/RigUnit_TwoBoneIKFK.cpp
+++ b/Engine/Plugins/Experimental/ControlRig/Source/ControlRig/Private/Units/Deprecated/RigUnit_TwoBoneIKFK.cpp
@@ -8,11 +8,10 @@
 void FRigUnit_TwoBoneIKFK::Execute(const FRigUnitContext& Context)
 {
     DECLARE_SCOPE_HIERARCHICAL_COUNTER_RIGUNIT()
-	FRigHierarchyRef& HierarchyRef = ExecuteContext.HierarchyReference;
 
 	if (Context.State == EControlRigState::Init)
 	{
-		const FRigBoneHierarchy* Hierarchy = HierarchyRef.GetBones();
+		const FRigBoneHierarchy* Hierarchy = Context.GetBones();
 		if (Hierarchy)
 		{
 			// reset
@@ -76,7 +75,7 @@ void FRigUnit_TwoBoneIKFK::Execute(const FRigUnitContext& Context)
 			else if (FMath::IsNearlyEqual(IKBlend, 1.f))
 			{
 				// update transform before going through IK
-				const FRigBoneHierarchy* Hierarchy = HierarchyRef.GetBones();
+				const FRigBoneHierarchy* Hierarchy = Context.GetBones();
 				check(Hierarchy);
 
 				StartJointIKTransform = Hierarchy->GetGlobalTransform(StartJointIndex);
@@ -92,7 +91,7 @@ void FRigUnit_TwoBoneIKFK::Execute(const FRigUnitContext& Context)
 			else
 			{
 				// update transform before going through IK
-				const FRigBoneHierarchy* Hierarchy = HierarchyRef.GetBones();
+				const FRigBoneHierarchy* Hierarchy = Context.GetBones();
 				check(Hierarchy);
 
 				StartJointIKTransform = Hierarchy->GetGlobalTransform(StartJointIndex);
@@ -106,7 +105,7 @@ void FRigUnit_TwoBoneIKFK::Execute(const FRigUnitContext& Context)
 				EndJointTransform.Blend(MidJointFKTransform, EndJointIKTransform, IKBlend);
 			}
 
-			FRigBoneHierarchy* Hierarchy = HierarchyRef.GetBones();
+			FRigBoneHierarchy* Hierarchy = ExecuteContext.GetBones();
 			check(Hierarchy);
 			Hierarchy->SetGlobalTransform(StartJointIndex, StartJointTransform);
 			Hierarchy->SetGlobalTransform(MidJointIndex, MidJointTransform);
diff --git a/Engine/Plugins/Experimental/ControlRig/Source/ControlRig/Private/Units/Execution/RigUnit_BeginExecution.cpp b/Engine/Plugins/Experimental/ControlRig/Source/ControlRig/Private/Units/Execution/RigUnit_BeginExecution.cpp
--- a/Engine/Plugins/Experimental/ControlRig/Source/ControlRig/Private/Units/Execution/RigUnit_BeginExecution.cpp
+++ b/Engine/Plugins/Experimental/ControlRig/Source/ControlRig/Private/Units/Execution/RigUnit_BeginExecution.cpp
@@ -6,6 +6,5 @@
 void FRigUnit_BeginExecution::Execute(const FRigUnitContext& Context)
 {
     DECLARE_SCOPE_HIERARCHICAL_COUNTER_RIGUNIT()
-	ExecuteContext.HierarchyReference = Context.HierarchyReference;
-	ExecuteContext.CurveReference = Context.CurveReference;
+	ExecuteContext.Hierarchy = (FRigHierarchyContainer*)Context.Hierarchy;
 }
diff --git a/Engine/Plugins/Experimental/ControlRig/Source/ControlRig/Private/Units/Hierarchy/RigUnit_AddBoneTransform.cpp b/Engine/Plugins/Experimental/ControlRig/Source/ControlRig/Private/Units/Hierarchy/RigUnit_AddBoneTransform.cpp
--- a/Engine/Plugins/Experimental/ControlRig/Source/ControlRig/Private/Units/Hierarchy/RigUnit_AddBoneTransform.cpp
+++ b/Engine/Plugins/Experimental/ControlRig/Source/ControlRig/Private/Units/Hierarchy/RigUnit_AddBoneTransform.cpp
@@ -11,8 +11,7 @@ FString FRigUnit_AddBoneTransform::GetUnitLabel() const
 void FRigUnit_AddBoneTransform::Execute(const FRigUnitContext& Context)
 {
     DECLARE_SCOPE_HIERARCHICAL_COUNTER_RIGUNIT()
-	FRigHierarchyRef& HierarchyRef = ExecuteContext.HierarchyReference;
-	FRigBoneHierarchy* Hierarchy = HierarchyRef.GetBones();
+	FRigBoneHierarchy* Hierarchy = ExecuteContext.GetBones();
 	if (Hierarchy)
 	{
 		switch (Context.State)
@@ -53,7 +52,7 @@ IMPLEMENT_RIGUNIT_AUTOMATION_TEST(FRigUnit_AddBoneTransform)
 	BoneHierarchy.Add(TEXT("Root"), NAME_None, FTransform(FVector(1.f, 0.f, 0.f)));
 	BoneHierarchy.Add(TEXT("BoneA"), TEXT("Root"), FTransform(FVector(1.f, 2.f, 0.f)));
 	BoneHierarchy.Initialize();
-	Unit.ExecuteContext.HierarchyReference = HierarchyRef;
+	Unit.ExecuteContext.Hierarchy = &HierarchyContainer;
 
 	BoneHierarchy.ResetTransforms();
 	Unit.Bone = TEXT("BoneA");
diff --git a/Engine/Plugins/Experimental/ControlRig/Source/ControlRig/Private/Units/Hierarchy/RigUnit_GetBoneTransform.cpp b/Engine/Plugins/Experimental/ControlRig/Source/ControlRig/Private/Units/Hierarchy/RigUnit_GetBoneTransform.cpp
--- a/Engine/Plugins/Experimental/ControlRig/Source/ControlRig/Private/Units/Hierarchy/RigUnit_GetBoneTransform.cpp
+++ b/Engine/Plugins/Experimental/ControlRig/Source/ControlRig/Private/Units/Hierarchy/RigUnit_GetBoneTransform.cpp
@@ -11,8 +11,7 @@ FString FRigUnit_GetBoneTransform::GetUnitLabel() const
 void FRigUnit_GetBoneTransform::Execute(const FRigUnitContext& Context)
 {
     DECLARE_SCOPE_HIERARCHICAL_COUNTER_RIGUNIT()
-	const FRigHierarchyRef& HierarchyRef = Context.HierarchyReference;
-	const FRigBoneHierarchy* Hierarchy = HierarchyRef.GetBones();
+	const FRigBoneHierarchy* Hierarchy = Context.GetBones();
 	if (Hierarchy)
 	{
 		switch (Context.State)
diff --git a/Engine/Plugins/Experimental/ControlRig/Source/ControlRig/Private/Units/Hierarchy/RigUnit_GetCurveValue.cpp b/Engine/Plugins/Experimental/ControlRig/Source/ControlRig/Private/Units/Hierarchy/RigUnit_GetCurveValue.cpp
--- a/Engine/Plugins/Experimental/ControlRig/Source/ControlRig/Private/Units/Hierarchy/RigUnit_GetCurveValue.cpp
+++ b/Engine/Plugins/Experimental/ControlRig/Source/ControlRig/Private/Units/Hierarchy/RigUnit_GetCurveValue.cpp
@@ -11,7 +11,7 @@ FString FRigUnit_GetCurveValue::GetUnitLabel() const
 void FRigUnit_GetCurveValue::Execute(const FRigUnitContext& Context)
 {
     DECLARE_SCOPE_HIERARCHICAL_COUNTER_RIGUNIT()
-	const FRigCurveContainer* CurveContainer = Context.CurveReference.Get();
+	const FRigCurveContainer* CurveContainer = Context.GetCurves();
 	if (CurveContainer)
 	{
 		switch (Context.State)
diff --git a/Engine/Plugins/Experimental/ControlRig/Source/ControlRig/Private/Units/Hierarchy/RigUnit_GetInitialBoneTransform.cpp b/Engine/Plugins/Experimental/ControlRig/Source/ControlRig/Private/Units/Hierarchy/RigUnit_GetInitialBoneTransform.cpp
--- a/Engine/Plugins/Experimental/ControlRig/Source/ControlRig/Private/Units/Hierarchy/RigUnit_GetInitialBoneTransform.cpp
+++ b/Engine/Plugins/Experimental/ControlRig/Source/ControlRig/Private/Units/Hierarchy/RigUnit_GetInitialBoneTransform.cpp
@@ -11,8 +11,7 @@ FString FRigUnit_GetInitialBoneTransform::GetUnitLabel() const
 void FRigUnit_GetInitialBoneTransform::Execute(const FRigUnitContext& Context)
 {
     DECLARE_SCOPE_HIERARCHICAL_COUNTER_RIGUNIT()
-	const FRigHierarchyRef& HierarchyRef = Context.HierarchyReference;
-	const FRigBoneHierarchy* Hierarchy = HierarchyRef.GetBones();
+	const FRigBoneHierarchy* Hierarchy = Context.GetBones();
 	if (Hierarchy)
 	{
 		switch (Context.State)
diff --git a/Engine/Plugins/Experimental/ControlRig/Source/ControlRig/Private/Units/Hierarchy/RigUnit_GetRelativeBoneTransform.cpp b/Engine/Plugins/Experimental/ControlRig/Source/ControlRig/Private/Units/Hierarchy/RigUnit_GetRelativeBoneTransform.cpp
--- a/Engine/Plugins/Experimental/ControlRig/Source/ControlRig/Private/Units/Hierarchy/RigUnit_GetRelativeBoneTransform.cpp
+++ b/Engine/Plugins/Experimental/ControlRig/Source/ControlRig/Private/Units/Hierarchy/RigUnit_GetRelativeBoneTransform.cpp
@@ -11,8 +11,7 @@ FString FRigUnit_GetRelativeBoneTransform::GetUnitLabel() const
 void FRigUnit_GetRelativeBoneTransform::Execute(const FRigUnitContext& Context)
 {
     DECLARE_SCOPE_HIERARCHICAL_COUNTER_RIGUNIT()
-	const FRigHierarchyRef& HierarchyRef = Context.HierarchyReference;
-	const FRigBoneHierarchy* Hierarchy = HierarchyRef.GetBones();
+	const FRigBoneHierarchy* Hierarchy = Context.GetBones();
 	if (Hierarchy)
 	{
 		switch (Context.State)
diff --git a/Engine/Plugins/Experimental/ControlRig/Source/ControlRig/Private/Units/Hierarchy/RigUnit_SetBoneRotation.cpp b/Engine/Plugins/Experimental/ControlRig/Source/ControlRig/Private/Units/Hierarchy/RigUnit_SetBoneRotation.cpp
--- a/Engine/Plugins/Experimental/ControlRig/Source/ControlRig/Private/Units/Hierarchy/RigUnit_SetBoneRotation.cpp
+++ b/Engine/Plugins/Experimental/ControlRig/Source/ControlRig/Private/Units/Hierarchy/RigUnit_SetBoneRotation.cpp
@@ -11,8 +11,7 @@ FString FRigUnit_SetBoneRotation::GetUnitLabel() const
 void FRigUnit_SetBoneRotation::Execute(const FRigUnitContext& Context)
 {
     DECLARE_SCOPE_HIERARCHICAL_COUNTER_RIGUNIT()
-	FRigHierarchyRef& HierarchyRef = ExecuteContext.HierarchyReference;
-	FRigBoneHierarchy* Hierarchy = HierarchyRef.GetBones();
+	FRigBoneHierarchy* Hierarchy = ExecuteContext.GetBones();
 	if (Hierarchy)
 	{
 		switch (Context.State)
@@ -67,7 +66,7 @@ IMPLEMENT_RIGUNIT_AUTOMATION_TEST(FRigUnit_SetBoneRotation)
 	BoneHierarchy.Add(TEXT("BoneB"), TEXT("BoneA"), FTransform(FQuat(FVector(-1.f, 0.f, 0.f), 0.7f)));
 	BoneHierarchy.Initialize();
 
-	Unit.ExecuteContext.HierarchyReference = HierarchyRef;
+	Unit.ExecuteContext.Hierarchy = &HierarchyContainer;
 
 	BoneHierarchy.ResetTransforms();
 	Unit.Bone = TEXT("Root");
diff --git a/Engine/Plugins/Experimental/ControlRig/Source/ControlRig/Private/Units/Hierarchy/RigUnit_SetBoneTransform.cpp b/Engine/Plugins/Experimental/ControlRig/Source/ControlRig/Private/Units/Hierarchy/RigUnit_SetBoneTransform.cpp
--- a/Engine/Plugins/Experimental/ControlRig/Source/ControlRig/Private/Units/Hierarchy/RigUnit_SetBoneTransform.cpp
+++ b/Engine/Plugins/Experimental/ControlRig/Source/ControlRig/Private/Units/Hierarchy/RigUnit_SetBoneTransform.cpp
@@ -11,8 +11,7 @@ FString FRigUnit_SetBoneTransform::GetUnitLabel() const
 void FRigUnit_SetBoneTransform::Execute(const FRigUnitContext& Context)
 {
     DECLARE_SCOPE_HIERARCHICAL_COUNTER_RIGUNIT()
-	FRigHierarchyRef& HierarchyRef = ExecuteContext.HierarchyReference;
-	FRigBoneHierarchy* Hierarchy = HierarchyRef.GetBones();
+	FRigBoneHierarchy* Hierarchy = ExecuteContext.GetBones();
 	if (Hierarchy)
 	{
 		switch (Context.State)
@@ -62,7 +61,7 @@ IMPLEMENT_RIGUNIT_AUTOMATION_TEST(FRigUnit_SetBoneTransform)
 	BoneHierarchy.Add(TEXT("BoneA"), TEXT("Root"), FTransform(FVector(1.f, 2.f, 3.f)));
 	BoneHierarchy.Add(TEXT("BoneB"), TEXT("BoneA"), FTransform(FVector(1.f, 5.f, 3.f)));
 	BoneHierarchy.Initialize();
-	Unit.ExecuteContext.HierarchyReference = HierarchyRef;
+	Unit.ExecuteContext.Hierarchy = &HierarchyContainer;
 
 	BoneHierarchy.ResetTransforms();
 	Unit.Bone = TEXT("Root");
diff --git a/Engine/Plugins/Experimental/ControlRig/Source/ControlRig/Private/Units/Hierarchy/RigUnit_SetBoneTranslation.cpp b/Engine/Plugins/Experimental/ControlRig/Source/ControlRig/Private/Units/Hierarchy/RigUnit_SetBoneTranslation.cpp
--- a/Engine/Plugins/Experimental/ControlRig/Source/ControlRig/Private/Units/Hierarchy/RigUnit_SetBoneTranslation.cpp
+++ b/Engine/Plugins/Experimental/ControlRig/Source/ControlRig/Private/Units/Hierarchy/RigUnit_SetBoneTranslation.cpp
@@ -11,8 +11,7 @@ FString FRigUnit_SetBoneTranslation::GetUnitLabel() const
 void FRigUnit_SetBoneTranslation::Execute(const FRigUnitContext& Context)
 {
     DECLARE_SCOPE_HIERARCHICAL_COUNTER_RIGUNIT()
-	FRigHierarchyRef& HierarchyRef = ExecuteContext.HierarchyReference;
-	FRigBoneHierarchy* Hierarchy = HierarchyRef.GetBones();
+	FRigBoneHierarchy* Hierarchy = ExecuteContext.GetBones();
 	if (Hierarchy)
 	{
 		switch (Context.State)
diff --git a/Engine/Plugins/Experimental/ControlRig/Source/ControlRig/Private/Units/Hierarchy/RigUnit_SetCurveValue.cpp b/Engine/Plugins/Experimental/ControlRig/Source/ControlRig/Private/Units/Hierarchy/RigUnit_SetCurveValue.cpp
--- a/Engine/Plugins/Experimental/ControlRig/Source/ControlRig/Private/Units/Hierarchy/RigUnit_SetCurveValue.cpp
+++ b/Engine/Plugins/Experimental/ControlRig/Source/ControlRig/Private/Units/Hierarchy/RigUnit_SetCurveValue.cpp
@@ -11,7 +11,7 @@ FString FRigUnit_SetCurveValue::GetUnitLabel() const
 void FRigUnit_SetCurveValue::Execute(const FRigUnitContext& Context)
 {
     DECLARE_SCOPE_HIERARCHICAL_COUNTER_RIGUNIT()
-	FRigCurveContainer* CurveContainer = ExecuteContext.CurveReference.Get();
+	FRigCurveContainer* CurveContainer = ExecuteContext.GetCurves();
 	if (CurveContainer)
 	{
 		switch (Context.State)
@@ -44,7 +44,7 @@ IMPLEMENT_RIGUNIT_AUTOMATION_TEST(FRigUnit_SetCurveValue)
 	CurveContainer.Add(TEXT("CurveA"));
 	CurveContainer.Add(TEXT("CurveB"));
 	CurveContainer.Initialize();
-	Unit.ExecuteContext.CurveReference = CurveContainerRef;
+	Unit.ExecuteContext.Hierarchy = &HierarchyContainer;
 	
 	CurveContainer.ResetValues();
 	Unit.Curve = TEXT("CurveA");
diff --git a/Engine/Plugins/Experimental/ControlRig/Source/ControlRig/Private/Units/Hierarchy/RigUnit_SetRelativeBoneTransform.cpp b/Engine/Plugins/Experimental/ControlRig/Source/ControlRig/Private/Units/Hierarchy/RigUnit_SetRelativeBoneTransform.cpp
--- a/Engine/Plugins/Experimental/ControlRig/Source/ControlRig/Private/Units/Hierarchy/RigUnit_SetRelativeBoneTransform.cpp
+++ b/Engine/Plugins/Experimental/ControlRig/Source/ControlRig/Private/Units/Hierarchy/RigUnit_SetRelativeBoneTransform.cpp
@@ -11,8 +11,7 @@ FString FRigUnit_SetRelativeBoneTransform::GetUnitLabel() const
 void FRigUnit_SetRelativeBoneTransform::Execute(const FRigUnitContext& Context)
 {
     DECLARE_SCOPE_HIERARCHICAL_COUNTER_RIGUNIT()
-	FRigHierarchyRef& HierarchyRef = ExecuteContext.HierarchyReference;
-	FRigBoneHierarchy* Hierarchy = HierarchyRef.GetBones();
+	FRigBoneHierarchy* Hierarchy = ExecuteContext.GetBones();
 	if (Hierarchy)
 	{
 		switch (Context.State)
@@ -50,7 +49,7 @@ IMPLEMENT_RIGUNIT_AUTOMATION_TEST(FRigUnit_SetRelativeBoneTransform)
 	BoneHierarchy.Add(TEXT("BoneB"), TEXT("BoneA"), FTransform(FVector(1.f, 5.f, 3.f)));
 	BoneHierarchy.Add(TEXT("BoneC"), TEXT("Root"), FTransform(FVector(-4.f, 0.f, 0.f)));
 	BoneHierarchy.Initialize();
-	Unit.ExecuteContext.HierarchyReference = HierarchyRef;
+	Unit.ExecuteContext.Hierarchy = &HierarchyContainer;
 
 	BoneHierarchy.ResetTransforms();
 	Unit.Bone = TEXT("BoneA");
diff --git a/Engine/Plugins/Experimental/ControlRig/Source/ControlRig/Private/Units/Highlevel/Harmonics/RigUnit_BoneHarmonics.cpp b/Engine/Plugins/Experimental/ControlRig/Source/ControlRig/Private/Units/Highlevel/Harmonics/RigUnit_BoneHarmonics.cpp
--- a/Engine/Plugins/Experimental/ControlRig/Source/ControlRig/Private/Units/Highlevel/Harmonics/RigUnit_BoneHarmonics.cpp
+++ b/Engine/Plugins/Experimental/ControlRig/Source/ControlRig/Private/Units/Highlevel/Harmonics/RigUnit_BoneHarmonics.cpp
@@ -6,8 +6,7 @@
 void FRigUnit_BoneHarmonics::Execute(const FRigUnitContext& Context)
 {
     DECLARE_SCOPE_HIERARCHICAL_COUNTER_RIGUNIT()
-	FRigHierarchyRef& HierarchyRef = ExecuteContext.HierarchyReference;
-	FRigBoneHierarchy* Hierarchy = HierarchyRef.GetBones();
+	FRigBoneHierarchy* Hierarchy = ExecuteContext.GetBones();
 	if (Hierarchy == nullptr)
 	{
 		return;
diff --git a/Engine/Plugins/Experimental/ControlRig/Source/ControlRig/Private/Units/Highlevel/Harmonics/RigUnit_ChainHarmonics.cpp b/Engine/Plugins/Experimental/ControlRig/Source/ControlRig/Private/Units/Highlevel/Harmonics/RigUnit_ChainHarmonics.cpp
--- a/Engine/Plugins/Experimental/ControlRig/Source/ControlRig/Private/Units/Highlevel/Harmonics/RigUnit_ChainHarmonics.cpp
+++ b/Engine/Plugins/Experimental/ControlRig/Source/ControlRig/Private/Units/Highlevel/Harmonics/RigUnit_ChainHarmonics.cpp
@@ -6,7 +6,7 @@
 void FRigUnit_ChainHarmonics::Execute(const FRigUnitContext& Context)
 {
     DECLARE_SCOPE_HIERARCHICAL_COUNTER_RIGUNIT()
-	FRigBoneHierarchy* Hierarchy = (FRigBoneHierarchy*)(Context.HierarchyReference.GetBones());
+	FRigBoneHierarchy* Hierarchy = ExecuteContext.GetBones();
 	if (Hierarchy == nullptr)
 	{
 		return;
diff --git a/Engine/Plugins/Experimental/ControlRig/Source/ControlRig/Private/Units/Highlevel/Hierarchy/RigUnit_AimBone.cpp b/Engine/Plugins/Experimental/ControlRig/Source/ControlRig/Private/Units/Highlevel/Hierarchy/RigUnit_AimBone.cpp
--- a/Engine/Plugins/Experimental/ControlRig/Source/ControlRig/Private/Units/Highlevel/Hierarchy/RigUnit_AimBone.cpp
+++ b/Engine/Plugins/Experimental/ControlRig/Source/ControlRig/Private/Units/Highlevel/Hierarchy/RigUnit_AimBone.cpp
@@ -6,7 +6,7 @@
 void FRigUnit_AimBone::Execute(const FRigUnitContext& Context)
 {
     DECLARE_SCOPE_HIERARCHICAL_COUNTER_RIGUNIT()
-	FRigBoneHierarchy* Hierarchy = (FRigBoneHierarchy*)(Context.HierarchyReference.GetBones());
+	FRigBoneHierarchy* Hierarchy = ExecuteContext.GetBones();
 	if (Hierarchy == nullptr)
 	{
 		return;
diff --git a/Engine/Plugins/Experimental/ControlRig/Source/ControlRig/Private/Units/Highlevel/Hierarchy/RigUnit_CCDIK.cpp b/Engine/Plugins/Experimental/ControlRig/Source/ControlRig/Private/Units/Highlevel/Hierarchy/RigUnit_CCDIK.cpp
--- a/Engine/Plugins/Experimental/ControlRig/Source/ControlRig/Private/Units/Highlevel/Hierarchy/RigUnit_CCDIK.cpp
+++ b/Engine/Plugins/Experimental/ControlRig/Source/ControlRig/Private/Units/Highlevel/Hierarchy/RigUnit_CCDIK.cpp
@@ -6,7 +6,7 @@
 void FRigUnit_CCDIK::Execute(const FRigUnitContext& Context)
 {
     DECLARE_SCOPE_HIERARCHICAL_COUNTER_RIGUNIT()
-	FRigBoneHierarchy* Hierarchy = (FRigBoneHierarchy*)(Context.HierarchyReference.GetBones());
+	FRigBoneHierarchy* Hierarchy = ExecuteContext.GetBones();
 	if (Hierarchy == nullptr)
 	{
 		return;
diff --git a/Engine/Plugins/Experimental/ControlRig/Source/ControlRig/Private/Units/Highlevel/Hierarchy/RigUnit_DistributeRotation.cpp b/Engine/Plugins/Experimental/ControlRig/Source/ControlRig/Private/Units/Highlevel/Hierarchy/RigUnit_DistributeRotation.cpp
--- a/Engine/Plugins/Experimental/ControlRig/Source/ControlRig/Private/Units/Highlevel/Hierarchy/RigUnit_DistributeRotation.cpp
+++ b/Engine/Plugins/Experimental/ControlRig/Source/ControlRig/Private/Units/Highlevel/Hierarchy/RigUnit_DistributeRotation.cpp
@@ -6,7 +6,7 @@
 void FRigUnit_DistributeRotation::Execute(const FRigUnitContext& Context)
 {
     DECLARE_SCOPE_HIERARCHICAL_COUNTER_RIGUNIT()
-	FRigBoneHierarchy* Hierarchy = (FRigBoneHierarchy*)(Context.HierarchyReference.GetBones());
+	FRigBoneHierarchy* Hierarchy = ExecuteContext.GetBones();
 	if (Hierarchy == nullptr)
 	{
 		return;
@@ -190,6 +190,7 @@ IMPLEMENT_RIGUNIT_AUTOMATION_TEST(FRigUnit_DistributeRotation)
 	BoneHierarchy.Add(TEXT("BoneC"), TEXT("BoneB"), FTransform(FVector(2.f, 0.f, 0.f)));
 	BoneHierarchy.Add(TEXT("BoneD"), TEXT("BoneC"), FTransform(FVector(2.f, 0.f, 0.f)));
 	BoneHierarchy.Initialize();
+	Unit.ExecuteContext.Hierarchy = &HierarchyContainer;
 
 	Unit.StartBone = TEXT("Root");
 	Unit.EndBone = TEXT("BoneD");
diff --git a/Engine/Plugins/Experimental/ControlRig/Source/ControlRig/Private/Units/Highlevel/Hierarchy/RigUnit_FABRIK.cpp b/Engine/Plugins/Experimental/ControlRig/Source/ControlRig/Private/Units/Highlevel/Hierarchy/RigUnit_FABRIK.cpp
--- a/Engine/Plugins/Experimental/ControlRig/Source/ControlRig/Private/Units/Highlevel/Hierarchy/RigUnit_FABRIK.cpp
+++ b/Engine/Plugins/Experimental/ControlRig/Source/ControlRig/Private/Units/Highlevel/Hierarchy/RigUnit_FABRIK.cpp
@@ -6,7 +6,7 @@
 void FRigUnit_FABRIK::Execute(const FRigUnitContext& Context)
 {
     DECLARE_SCOPE_HIERARCHICAL_COUNTER_RIGUNIT()
-	FRigBoneHierarchy* Hierarchy = (FRigBoneHierarchy*)(Context.HierarchyReference.GetBones());
+	FRigBoneHierarchy* Hierarchy = ExecuteContext.GetBones();
 	if (Hierarchy == nullptr)
 	{
 		return;
diff --git a/Engine/Plugins/Experimental/ControlRig/Source/ControlRig/Private/Units/Highlevel/Hierarchy/RigUnit_FitChainToCurve.cpp b/Engine/Plugins/Experimental/ControlRig/Source/ControlRig/Private/Units/Highlevel/Hierarchy/RigUnit_FitChainToCurve.cpp
--- a/Engine/Plugins/Experimental/ControlRig/Source/ControlRig/Private/Units/Highlevel/Hierarchy/RigUnit_FitChainToCurve.cpp
+++ b/Engine/Plugins/Experimental/ControlRig/Source/ControlRig/Private/Units/Highlevel/Hierarchy/RigUnit_FitChainToCurve.cpp
@@ -6,7 +6,7 @@
 void FRigUnit_FitChainToCurve::Execute(const FRigUnitContext& Context)
 {
     DECLARE_SCOPE_HIERARCHICAL_COUNTER_RIGUNIT()
-	FRigBoneHierarchy* Hierarchy = (FRigBoneHierarchy*)(Context.HierarchyReference.GetBones());
+	FRigBoneHierarchy* Hierarchy = ExecuteContext.GetBones();
 	if (Hierarchy == nullptr)
 	{
 		return;
diff --git a/Engine/Plugins/Experimental/ControlRig/Source/ControlRig/Private/Units/Highlevel/Hierarchy/RigUnit_ModifyBoneTransforms.cpp b/Engine/Plugins/Experimental/ControlRig/Source/ControlRig/Private/Units/Highlevel/Hierarchy/RigUnit_ModifyBoneTransforms.cpp
--- a/Engine/Plugins/Experimental/ControlRig/Source/ControlRig/Private/Units/Highlevel/Hierarchy/RigUnit_ModifyBoneTransforms.cpp
+++ b/Engine/Plugins/Experimental/ControlRig/Source/ControlRig/Private/Units/Highlevel/Hierarchy/RigUnit_ModifyBoneTransforms.cpp
@@ -6,8 +6,7 @@
 void FRigUnit_ModifyBoneTransforms::Execute(const FRigUnitContext& Context)
 {
     DECLARE_SCOPE_HIERARCHICAL_COUNTER_RIGUNIT()
-	FRigHierarchyRef& HierarchyRef = ExecuteContext.HierarchyReference;
-	FRigBoneHierarchy* Hierarchy = HierarchyRef.GetBones();
+	FRigBoneHierarchy* Hierarchy = ExecuteContext.GetBones();
 	if (Hierarchy)
 	{
 		switch (Context.State)
@@ -113,7 +112,7 @@ IMPLEMENT_RIGUNIT_AUTOMATION_TEST(FRigUnit_ModifyBoneTransforms)
 	BoneHierarchy.Add(TEXT("BoneA"), TEXT("Root"), FTransform(FVector(1.f, 2.f, 3.f)));
 	BoneHierarchy.Add(TEXT("BoneB"), TEXT("Root"), FTransform(FVector(5.f, 6.f, 7.f)));
 	BoneHierarchy.Initialize();
-	Unit.ExecuteContext.HierarchyReference = HierarchyRef;
+	Unit.ExecuteContext.Hierarchy = &HierarchyContainer;
 
 	Unit.BoneToModify.SetNumZeroed(2);
 	Unit.BoneToModify[0].Bone = TEXT("BoneA");
diff --git a/Engine/Plugins/Experimental/ControlRig/Source/ControlRig/Private/Units/Highlevel/Hierarchy/RigUnit_SpringIK.cpp b/Engine/Plugins/Experimental/ControlRig/Source/ControlRig/Private/Units/Highlevel/Hierarchy/RigUnit_SpringIK.cpp
--- a/Engine/Plugins/Experimental/ControlRig/Source/ControlRig/Private/Units/Highlevel/Hierarchy/RigUnit_SpringIK.cpp
+++ b/Engine/Plugins/Experimental/ControlRig/Source/ControlRig/Private/Units/Highlevel/Hierarchy/RigUnit_SpringIK.cpp
@@ -6,7 +6,7 @@
 void FRigUnit_SpringIK::Execute(const FRigUnitContext& Context)
 {
     DECLARE_SCOPE_HIERARCHICAL_COUNTER_RIGUNIT()
-	FRigBoneHierarchy* Hierarchy = (FRigBoneHierarchy*)(Context.HierarchyReference.GetBones());
+	FRigBoneHierarchy* Hierarchy = ExecuteContext.GetBones();
 	if (Hierarchy == nullptr)
 	{
 		return;
diff --git a/Engine/Plugins/Experimental/ControlRig/Source/ControlRig/Private/Units/Highlevel/Hierarchy/RigUnit_TransformConstraint.cpp b/Engine/Plugins/Experimental/ControlRig/Source/ControlRig/Private/Units/Highlevel/Hierarchy/RigUnit_TransformConstraint.cpp
--- a/Engine/Plugins/Experimental/ControlRig/Source/ControlRig/Private/Units/Highlevel/Hierarchy/RigUnit_TransformConstraint.cpp
+++ b/Engine/Plugins/Experimental/ControlRig/Source/ControlRig/Private/Units/Highlevel/Hierarchy/RigUnit_TransformConstraint.cpp
@@ -8,14 +8,13 @@
 void FRigUnit_TransformConstraint::Execute(const FRigUnitContext& Context)
 {
     DECLARE_SCOPE_HIERARCHICAL_COUNTER_RIGUNIT()
-	FRigHierarchyRef& HierarchyRef = ExecuteContext.HierarchyReference;
 
 	if (Context.State == EControlRigState::Init)
 	{
 		ConstraintData.Reset();
 		ConstraintDataToTargets.Reset();
 
-		FRigBoneHierarchy* Hierarchy = HierarchyRef.GetBones();
+		FRigBoneHierarchy* Hierarchy = ExecuteContext.GetBones();
 		if (Hierarchy)
 		{
 			int32 BoneIndex = Hierarchy->GetIndex(Bone);
@@ -62,7 +61,7 @@ void FRigUnit_TransformConstraint::Execute(const FRigUnitContext& Context)
 	}
 	else if (Context.State == EControlRigState::Update)
 	{
-		FRigBoneHierarchy* Hierarchy = HierarchyRef.GetBones();
+		FRigBoneHierarchy* Hierarchy = ExecuteContext.GetBones();
 		if (Hierarchy)
 		{
 			int32 BoneIndex = Hierarchy->GetIndex(Bone);
diff --git a/Engine/Plugins/Experimental/ControlRig/Source/ControlRig/Private/Units/Highlevel/Hierarchy/RigUnit_TwoBoneIKSimple.cpp b/Engine/Plugins/Experimental/ControlRig/Source/ControlRig/Private/Units/Highlevel/Hierarchy/RigUnit_TwoBoneIKSimple.cpp
--- a/Engine/Plugins/Experimental/ControlRig/Source/ControlRig/Private/Units/Highlevel/Hierarchy/RigUnit_TwoBoneIKSimple.cpp
+++ b/Engine/Plugins/Experimental/ControlRig/Source/ControlRig/Private/Units/Highlevel/Hierarchy/RigUnit_TwoBoneIKSimple.cpp
@@ -7,7 +7,7 @@
 void FRigUnit_TwoBoneIKSimple::Execute(const FRigUnitContext& Context)
 {
     DECLARE_SCOPE_HIERARCHICAL_COUNTER_RIGUNIT()
-	FRigBoneHierarchy* Hierarchy = (FRigBoneHierarchy*)(Context.HierarchyReference.GetBones());
+	FRigBoneHierarchy* Hierarchy = ExecuteContext.GetBones();
 	if (Hierarchy == nullptr)
 	{
 		return;
diff --git a/Engine/Plugins/Experimental/ControlRig/Source/ControlRig/Private/Units/Simulation/RigUnit_PointSimulation.cpp b/Engine/Plugins/Experimental/ControlRig/Source/ControlRig/Private/Units/Simulation/RigUnit_PointSimulation.cpp
--- a/Engine/Plugins/Experimental/ControlRig/Source/ControlRig/Private/Units/Simulation/RigUnit_PointSimulation.cpp
+++ b/Engine/Plugins/Experimental/ControlRig/Source/ControlRig/Private/Units/Simulation/RigUnit_PointSimulation.cpp
@@ -6,8 +6,7 @@
 void FRigUnit_PointSimulation::Execute(const FRigUnitContext& Context)
 {
     DECLARE_SCOPE_HIERARCHICAL_COUNTER_RIGUNIT()
-	FRigHierarchyRef& HierarchyRef = ExecuteContext.HierarchyReference;
-	FRigBoneHierarchy* Hierarchy = HierarchyRef.GetBones();
+	FRigBoneHierarchy* Hierarchy = ExecuteContext.GetBones();
 
 	float DeltaTime = Context.DeltaTime;
 
diff --git a/Engine/Plugins/Experimental/ControlRig/Source/ControlRig/Public/ControlRig.h b/Engine/Plugins/Experimental/ControlRig/Source/ControlRig/Public/ControlRig.h
--- a/Engine/Plugins/Experimental/ControlRig/Source/ControlRig/Public/ControlRig.h
+++ b/Engine/Plugins/Experimental/ControlRig/Source/ControlRig/Public/ControlRig.h
@@ -9,7 +9,6 @@
 #include "Engine/EngineBaseTypes.h"
 #include "Templates/SubclassOf.h"
 #include "ControlRigDefines.h"
-#include "Rigs/RigCurveContainer.h"
 #include "Rigs/RigHierarchyContainer.h"
 #include "Units/RigUnitContext.h"
 #include "Animation/NodeMappingProviderInterface.h"
@@ -105,6 +104,26 @@ public:
 		return ObjectBinding;
 	}
 
+	const FRigBoneHierarchy& GetBoneHierarchy() const
+	{
+		return Hierarchy.BoneHierarchy;
+	}
+
+	const FRigSpaceHierarchy& GetSpaceHierarchy() const
+	{
+		return Hierarchy.SpaceHierarchy;
+	}
+
+	const FRigControlHierarchy& GetControlHierarchy() const
+	{
+		return Hierarchy.ControlHierarchy;
+	}
+
+	const FRigCurveContainer& GetCurveContainer() const
+	{
+		return Hierarchy.CurveContainer;
+	}
+
 	/** Evaluate another animation ControlRig */
 	FTransform GetGlobalTransform(const FName& BoneName) const;
 
@@ -117,12 +136,6 @@ public:
 	/** Evaluate another animation ControlRig */
 	void SetGlobalTransform(const int32 BoneIndex, const FTransform& InTransform, bool bPropagateTransform = true) ;
 
-	/** Returns base hierarchy */
-	const FRigBoneHierarchy& GetBoneHierarchy() const
-	{
-		return Hierarchy.BoneHierarchy;
-	}
-
 	/** Evaluate another animation ControlRig */
 	float GetCurveValue(const FName& CurveName) const;
 
@@ -135,12 +148,6 @@ public:
 	/** Evaluate another animation ControlRig */
 	void SetCurveValue(const int32 CurveIndex, const float CurveValue);
 
-	/** Returns base hierarchy */
-	const FRigCurveContainer& GetCurveContainer() const
-	{
-		return CurveContainer;
-	}
-
 	void SetPreEvaluateGatherInputDelegate(const FPreEvaluateGatherInput& Delegate)
 	{
 		OnPreEvaluateGatherInput = Delegate;
@@ -221,9 +228,6 @@ private:
 	UPROPERTY(VisibleDefaultsOnly, Category = "Hierarchy")
 	FRigHierarchyContainer Hierarchy;
 
-	UPROPERTY(VisibleDefaultsOnly, Category = "Curve")
-	FRigCurveContainer CurveContainer;
-
 #if WITH_EDITORONLY_DATA
 	/** The properties of source accessible <target, source local path> when source -> target
 	 * For example, if you have property RigUnitA.B->RigUnitB.C, this will save as <RigUnitB.C, RigUnitA.B> */
diff --git a/Engine/Plugins/Experimental/ControlRig/Source/ControlRig/Public/ControlRigDefines.h b/Engine/Plugins/Experimental/ControlRig/Source/ControlRig/Public/ControlRigDefines.h
--- a/Engine/Plugins/Experimental/ControlRig/Source/ControlRig/Public/ControlRigDefines.h
+++ b/Engine/Plugins/Experimental/ControlRig/Source/ControlRig/Public/ControlRigDefines.h
@@ -15,9 +15,49 @@ USTRUCT()
 struct FControlRigExecuteContext
 {
 	GENERATED_BODY()
+
+	FControlRigExecuteContext()
+		: Hierarchy(nullptr)
+	{
+	}
 		
-	FRigHierarchyRef HierarchyReference;
-	FRigCurveContainerRef CurveReference;
+	FRigHierarchyContainer* Hierarchy;
+
+	FRigBoneHierarchy* GetBones()
+	{
+		if (Hierarchy != nullptr)
+		{
+			return &Hierarchy->BoneHierarchy;
+		}
+		return nullptr;
+	}
+
+	FRigSpaceHierarchy* GetSpaces()
+	{
+		if (Hierarchy != nullptr)
+		{
+			return &Hierarchy->SpaceHierarchy;
+		}
+		return nullptr;
+	}
+
+	FRigControlHierarchy* GetControls()
+	{
+		if (Hierarchy != nullptr)
+		{
+			return &Hierarchy->ControlHierarchy;
+		}
+		return nullptr;
+	}
+
+	FRigCurveContainer* GetCurves()
+	{
+		if (Hierarchy != nullptr)
+		{
+			return &Hierarchy->CurveContainer;
+		}
+		return nullptr;
+	}
 };
 
 UENUM()
diff --git a/Engine/Plugins/Experimental/ControlRig/Source/ControlRig/Public/Rigs/RigBoneHierarchy.h b/Engine/Plugins/Experimental/ControlRig/Source/ControlRig/Public/Rigs/RigBoneHierarchy.h
--- a/Engine/Plugins/Experimental/ControlRig/Source/ControlRig/Public/Rigs/RigBoneHierarchy.h
+++ b/Engine/Plugins/Experimental/ControlRig/Source/ControlRig/Public/Rigs/RigBoneHierarchy.h
@@ -13,9 +13,24 @@ struct CONTROLRIG_API FRigBone
 {
 	GENERATED_BODY()
 
+	FRigBone()
+		: Name(NAME_None)
+		, Index(INDEX_NONE)
+		, ParentName(NAME_None)
+		, ParentIndex(INDEX_NONE)
+		, InitialTransform(FTransform::Identity)
+		, GlobalTransform(FTransform::Identity)
+		, LocalTransform(FTransform::Identity)
+		, Dependents()
+	{
+	}
+
 	UPROPERTY(VisibleAnywhere, Category = FRigBoneHierarchy)
 	FName Name;
 
+	UPROPERTY(VisibleAnywhere, Category = FRigCurveContainer)
+	int32 Index;
+
 	UPROPERTY(VisibleAnywhere, Category = FRigBoneHierarchy)
 	FName ParentName;
 
@@ -45,6 +60,8 @@ struct CONTROLRIG_API FRigBoneHierarchy
 	FRigBoneHierarchy();
 	FRigBoneHierarchy& operator= (const FRigBoneHierarchy &InOther);
 
+	FORCEINLINE ERigElementType RigElementType() const { return ERigElementType::Bone; }
+
 	FORCEINLINE int32 Num() const { return Bones.Num(); }
 	FORCEINLINE const FRigBone& operator[](int32 InIndex) const { return Bones[InIndex]; }
 	FORCEINLINE FRigBone& operator[](int32 InIndex) { return Bones[InIndex]; }
@@ -56,13 +73,17 @@ struct CONTROLRIG_API FRigBoneHierarchy
 	FORCEINLINE TArray<FRigBone>::RangedForIteratorType      end()         { return Bones.end();   }
 	FORCEINLINE TArray<FRigBone>::RangedForConstIteratorType end() const   { return Bones.end();   }
 
+	FORCEINLINE bool IsNameAvailable(const FName& InPotentialNewName) const { return GetIndex(InPotentialNewName) == INDEX_NONE; }
+
+	FName GetSafeNewName(const FName& InPotentialNewName) const;
+
 	FRigBone& Add(const FName& InNewName, const FName& InParentName, const FTransform& InInitTransform);
 
 	FRigBone& Add(const FName& InNewName, const FName& InParentName, const FTransform& InInitTransform, const FTransform& InLocalTransform, const FTransform& InGlobalTransform);
 
 	FRigBone Remove(const FName& InNameToRemove);
 
-	void Rename(const FName& InOldName, const FName& InNewName);
+	FName Rename(const FName& InOldName, const FName& InNewName);
 
 	void Reparent(const FName& InName, const FName& InNewParentName);
 
@@ -129,10 +150,10 @@ private:
 
 #if WITH_EDITOR
 	FRigHierarchyContainer* Container;
-	FRigHierarchyElementAdded OnBoneAdded;
-	FRigHierarchyElementRemoved OnBoneRemoved;
-	FRigHierarchyElementRenamed OnBoneRenamed;
-	FRigHierarchyElementReparented OnBoneReparented;
+	FRigElementAdded OnBoneAdded;
+	FRigElementRemoved OnBoneRemoved;
+	FRigElementRenamed OnBoneRenamed;
+	FRigElementReparented OnBoneReparented;
 #endif
 
 	UPROPERTY(EditAnywhere, Category = FRigBoneHierarchy)
diff --git a/Engine/Plugins/Experimental/ControlRig/Source/ControlRig/Public/Rigs/RigControlHierarchy.h b/Engine/Plugins/Experimental/ControlRig/Source/ControlRig/Public/Rigs/RigControlHierarchy.h
new file mode 100644
--- /dev/null
+++ b/Engine/Plugins/Experimental/ControlRig/Source/ControlRig/Public/Rigs/RigControlHierarchy.h
@@ -0,0 +1,178 @@
+// Copyright 1998-2019 Epic Games, Inc. All Rights Reserved.
+
+#pragma once
+
+#include "CoreMinimal.h"
+#include "RigHierarchyDefines.h"
+#include "RigControlHierarchy.generated.h"
+
+class UControlRig;
+
+USTRUCT(BlueprintType)
+struct CONTROLRIG_API FRigControl
+{
+	GENERATED_BODY()
+
+	FRigControl()
+		: Name(NAME_None)
+		, Index(INDEX_NONE)
+		, ParentName(NAME_None)
+		, ParentIndex(INDEX_NONE)
+		, SpaceName(NAME_None)
+		, SpaceIndex(INDEX_NONE)
+		, InitialTransform(FTransform::Identity)
+		, LocalTransform(FTransform::Identity)
+		, Dependents()
+	{
+	}
+
+	UPROPERTY(VisibleAnywhere, Category = FRigControlHierarchy)
+	FName Name;
+
+	UPROPERTY(VisibleAnywhere, Category = FRigCurveContainer)
+	int32 Index;
+
+	UPROPERTY(VisibleAnywhere, Category = FRigControlHierarchy)
+	FName ParentName;
+
+	UPROPERTY(transient)
+	int32 ParentIndex;
+
+	UPROPERTY(VisibleAnywhere, Category = FRigControlHierarchy)
+	FName SpaceName;
+
+	UPROPERTY(transient)
+	int32 SpaceIndex;
+
+	UPROPERTY(VisibleAnywhere, Category = FRigControlHierarchy)
+	FTransform InitialTransform;
+
+	UPROPERTY(transient, VisibleAnywhere, Category = FRigControlHierarchy)
+	FTransform LocalTransform;
+
+	/** dependent list - direct dependent for child or anything that needs to update due to this */
+	UPROPERTY(transient)
+	TArray<int32> Dependents;
+};
+
+USTRUCT()
+struct CONTROLRIG_API FRigControlHierarchy
+{
+	GENERATED_BODY()
+
+	FRigControlHierarchy();
+	FRigControlHierarchy& operator= (const FRigControlHierarchy &InOther);
+
+	FORCEINLINE ERigElementType RigElementType() const { return ERigElementType::Control; }
+
+	FORCEINLINE int32 Num() const { return Controls.Num(); }
+	FORCEINLINE const FRigControl& operator[](int32 InIndex) const { return Controls[InIndex]; }
+	FORCEINLINE FRigControl& operator[](int32 InIndex) { return Controls[InIndex]; }
+	FORCEINLINE const FRigControl& operator[](const FName& InName) const { return Controls[GetIndex(InName)]; }
+	FORCEINLINE FRigControl& operator[](const FName& InName) { return Controls[GetIndex(InName)]; }
+
+	FORCEINLINE TArray<FRigControl>::RangedForIteratorType      begin()       { return Controls.begin(); }
+	FORCEINLINE TArray<FRigControl>::RangedForConstIteratorType begin() const { return Controls.begin(); }
+	FORCEINLINE TArray<FRigControl>::RangedForIteratorType      end()         { return Controls.end();   }
+	FORCEINLINE TArray<FRigControl>::RangedForConstIteratorType end() const   { return Controls.end();   }
+
+	FORCEINLINE bool IsNameAvailable(const FName& InPotentialNewName) const { return GetIndex(InPotentialNewName) == INDEX_NONE; }
+
+	FName GetSafeNewName(const FName& InPotentialNewName) const;
+
+	FRigControl& Add(const FName& InNewName, const FName& InParentName, const FName& InSpaceName, const FTransform& InTransform);
+
+	FRigControl Remove(const FName& InNameToRemove);
+
+	FName Rename(const FName& InOldName, const FName& InNewName);
+
+	void Reparent(const FName& InName, const FName& InNewParentName);
+
+	void SetSpace(const FName& InName, const FName& InNewSpaceName);
+
+	FName GetName(int32 InIndex) const;
+
+	FORCEINLINE int32 GetIndex(const FName& InName) const
+	{
+		if(NameToIndexMapping.Num() != Controls.Num())
+		{
+			return GetIndexSlow(InName);
+		}
+
+		const int32* Index = NameToIndexMapping.Find(InName);
+		if (Index)
+		{
+			return *Index;
+		}
+
+		return INDEX_NONE;
+	}
+
+	// list of names of children - this is not cheap, and is supposed to be used only for one time set up
+	int32 GetChildren(const FName& InName, TArray<int32>& OutChildren, bool bRecursively) const;
+
+	int32 GetChildren(const int32 InIndex, TArray<int32>& OutChildren, bool bRecursively) const;
+
+	void SetGlobalTransform(const FName& InName, const FTransform& InTransform);
+
+	void SetGlobalTransform(int32 InIndex, const FTransform& InTransform);
+
+	FTransform GetGlobalTransform(const FName& InName) const;
+
+	FTransform GetGlobalTransform(int32 InIndex) const;
+
+	void SetLocalTransform(const FName& InName, const FTransform& InTransform);
+
+	void SetLocalTransform(int32 InIndex, const FTransform& InTransform);
+
+	FTransform GetLocalTransform(const FName& InName) const;
+
+	FTransform GetLocalTransform(int32 InIndex) const;
+
+	void SetInitialTransform(const FName& InName, const FTransform& InTransform);
+
+	void SetInitialTransform(int32 InIndex, const FTransform& InTransform);
+
+	FTransform GetInitialTransform(const FName& InName) const;
+
+	FTransform GetInitialTransform(int32 InIndex) const;
+
+	// updates all of the internal caches
+	void Initialize();
+
+	// clears the hierarchy and removes all content
+	void Reset();
+
+	// resets all of the transforms back to the initial transform
+	void ResetTransforms();
+
+private:
+
+	// disable copy constructor
+	FRigControlHierarchy(const FRigControlHierarchy& InOther) {}
+
+#if WITH_EDITOR
+	FRigHierarchyContainer* Container;
+	FRigElementAdded OnControlAdded;
+	FRigElementRemoved OnControlRemoved;
+	FRigElementRenamed OnControlRenamed;
+	FRigElementReparented OnControlReparented;
+#endif
+
+	UPROPERTY(EditAnywhere, Category = FRigControlHierarchy)
+	TArray<FRigControl> Controls;
+
+	// can serialize fine? 
+	UPROPERTY()
+	TMap<FName, int32> NameToIndexMapping;
+
+	int32 GetIndexSlow(const FName& InName) const;
+
+	void RefreshMapping();
+	void Sort();
+
+	// list of names of children - this is not cheap, and is supposed to be used only for one time set up
+	int32 GetChildrenRecursive(const int32 InIndex, TArray<int32>& OutChildren, bool bRecursively) const;
+
+	friend struct FRigHierarchyContainer;
+};
diff --git a/Engine/Plugins/Experimental/ControlRig/Source/ControlRig/Public/Rigs/RigCurveContainer.h b/Engine/Plugins/Experimental/ControlRig/Source/ControlRig/Public/Rigs/RigCurveContainer.h
--- a/Engine/Plugins/Experimental/ControlRig/Source/ControlRig/Public/Rigs/RigCurveContainer.h
+++ b/Engine/Plugins/Experimental/ControlRig/Source/ControlRig/Public/Rigs/RigCurveContainer.h
@@ -4,6 +4,7 @@
 
 #include "CoreMinimal.h"
 #include "Stats/StatsHierarchical.h"
+#include "RigHierarchyDefines.h"
 #include "RigCurveContainer.generated.h"
 
 class UControlRig;
@@ -13,9 +14,19 @@ struct FRigCurve
 {
 	GENERATED_BODY()
 
+	FRigCurve()
+		: Name(NAME_None)
+		, Index(INDEX_NONE)
+		, Value(0.f)
+	{
+	}
+
 	UPROPERTY(VisibleAnywhere, Category = FRigCurveContainer)
 	FName Name;
 
+	UPROPERTY(VisibleAnywhere, Category = FRigCurveContainer)
+	int32 Index;
+
 	UPROPERTY(VisibleAnywhere, Category = FRigCurveContainer)
 	float Value;
 };
@@ -30,6 +41,8 @@ public:
 	FRigCurveContainer();
 	FRigCurveContainer& operator= (const FRigCurveContainer &InOther);
 
+	FORCEINLINE ERigElementType RigElementType() const { return ERigElementType::Curve; }
+
 	FORCEINLINE int32 Num() const { return Curves.Num(); }
 	FORCEINLINE const FRigCurve& operator[](int32 InIndex) const { return Curves[InIndex]; }
 	FORCEINLINE FRigCurve& operator[](int32 InIndex) { return Curves[InIndex]; }
@@ -41,9 +54,13 @@ public:
 	FORCEINLINE TArray<FRigCurve>::RangedForIteratorType      end()         { return Curves.end();   }
 	FORCEINLINE TArray<FRigCurve>::RangedForConstIteratorType end() const   { return Curves.end();   }
 
-	void Add(const FName& InNewName);
+	FORCEINLINE bool IsNameAvailable(const FName& InPotentialNewName) const { return GetIndex(InPotentialNewName) == INDEX_NONE; }
+
+	FName GetSafeNewName(const FName& InPotentialNewName) const;
+
+	FRigCurve& Add(const FName& InNewName);
 
-	void Remove(const FName& InName);
+	FRigCurve Remove(const FName& InName);
 
 	FName GetName(int32 InIndex) const;
 
@@ -71,7 +88,7 @@ public:
 
 	float GetValue(int32 InIndex) const;
 
-	void Rename(const FName& InOldName, const FName& InNewName);
+	FName Rename(const FName& InOldName, const FName& InNewName);
 
 	// updates all of the internal caches
 	void Initialize();
@@ -82,6 +99,13 @@ public:
 	// resets all of the transforms back to the initial transform
 	void ResetValues();
 
+#if WITH_EDITOR
+	FRigHierarchyContainer* Container;
+	FRigElementAdded OnCurveAdded;
+	FRigElementRemoved OnCurveRemoved;
+	FRigElementRenamed OnCurveRenamed;
+#endif
+
 private:
 
 	// disable copy constructor
@@ -98,37 +122,9 @@ private:
 	void RefreshMapping();
 };
 
-
-// @todo: do we need this container now?
-USTRUCT()
-struct FRigCurveContainerRef
-{
-	GENERATED_BODY()
-
-	FRigCurveContainerRef(FRigCurveContainer* InContainer = nullptr)
-		: Container(InContainer)
-	{
-
-	}
-
-	FRigCurveContainer* Get() 
-	{
-		return GetInternal();
-	}
-
-	const FRigCurveContainer* Get() const
-	{
-		return (const FRigCurveContainer*)GetInternal();
-	}
-
-private:
-	struct FRigCurveContainer* Container;
-
-	FRigCurveContainer* GetInternal() const
-	{
-		return Container;
-	}
-
-	friend class UControlRig;
-	friend class FControlRigUnitTestBase;
-};
\ No newline at end of file
+// todo
+//USTRUCT()
+//struct CONTROLRIG_API FRigHierarchyRef
+//{
+//	GENERATED_BODY()
+//};
\ No newline at end of file
diff --git a/Engine/Plugins/Experimental/ControlRig/Source/ControlRig/Public/Rigs/RigHierarchyContainer.h b/Engine/Plugins/Experimental/ControlRig/Source/ControlRig/Public/Rigs/RigHierarchyContainer.h
--- a/Engine/Plugins/Experimental/ControlRig/Source/ControlRig/Public/Rigs/RigHierarchyContainer.h
+++ b/Engine/Plugins/Experimental/ControlRig/Source/ControlRig/Public/Rigs/RigHierarchyContainer.h
@@ -5,6 +5,9 @@
 #include "CoreMinimal.h"
 #include "RigHierarchyDefines.h"
 #include "RigBoneHierarchy.h"
+#include "RigSpaceHierarchy.h"
+#include "RigControlHierarchy.h"
+#include "RigCurveContainer.h"
 #include "RigHierarchyContainer.generated.h"
 
 class UControlRig;
@@ -22,16 +25,25 @@ public:
 	UPROPERTY()
 	FRigBoneHierarchy BoneHierarchy;
 
+	UPROPERTY()
+	FRigSpaceHierarchy SpaceHierarchy;
+
+	UPROPERTY()
+	FRigControlHierarchy ControlHierarchy;
+
+	UPROPERTY()
+	FRigCurveContainer CurveContainer;
+
 	void Initialize();
 	void Reset();
 	void ResetTransforms();
 
 #if WITH_EDITOR
-	FRigHierarchyElementChanged OnElementChanged;
-	FRigHierarchyElementAdded OnElementAdded;
-	FRigHierarchyElementRemoved OnElementRemoved;
-	FRigHierarchyElementRenamed OnElementRenamed;
-	FRigHierarchyElementReparented OnElementReparented;
+	FRigElementChanged OnElementChanged;
+	FRigElementAdded OnElementAdded;
+	FRigElementRemoved OnElementRemoved;
+	FRigElementRenamed OnElementRenamed;
+	FRigElementReparented OnElementReparented;
 #endif
 
 protected:
@@ -40,29 +52,16 @@ protected:
 	FRigHierarchyContainer(const FRigHierarchyContainer& InContainer) {}
 
 #if WITH_EDITOR
-	void HandleOnElementAdded(FRigHierarchyContainer* InContainer, ERigHierarchyElementType InElementType, const FName& InName);
-	void HandleOnElementRemoved(FRigHierarchyContainer* InContainer, ERigHierarchyElementType InElementType, const FName& InName);
-	void HandleOnElementRenamed(FRigHierarchyContainer* InContainer, ERigHierarchyElementType InElementType, const FName& InOldName, const FName& InNewName);
-	void HandleOnElementReparented(FRigHierarchyContainer* InContainer, ERigHierarchyElementType InElementType, const FName& InName, const FName& InOldParentName, const FName& InNewParentName);
+	void HandleOnElementAdded(FRigHierarchyContainer* InContainer, ERigElementType InElementType, const FName& InName);
+	void HandleOnElementRemoved(FRigHierarchyContainer* InContainer, ERigElementType InElementType, const FName& InName);
+	void HandleOnElementRenamed(FRigHierarchyContainer* InContainer, ERigElementType InElementType, const FName& InOldName, const FName& InNewName);
+	void HandleOnElementReparented(FRigHierarchyContainer* InContainer, ERigElementType InElementType, const FName& InName, const FName& InOldParentName, const FName& InNewParentName);
 #endif
 };
 
+// this struct is still here for backwards compatibility - but not used anywhere
 USTRUCT()
 struct CONTROLRIG_API FRigHierarchyRef
 {
 	GENERATED_BODY()
-
-	FRigHierarchyRef();
-
-	FRigBoneHierarchy* GetBones();
-	const FRigBoneHierarchy* GetBones() const;
-
-private:
-
-	struct FRigHierarchyContainer* Container;
-
-	FRigBoneHierarchy* GetBonesInternal() const;
-
-	friend class UControlRig;
-	friend class FControlRigUnitTestBase;
 };
\ No newline at end of file
diff --git a/Engine/Plugins/Experimental/ControlRig/Source/ControlRig/Public/Rigs/RigHierarchyDefines.h b/Engine/Plugins/Experimental/ControlRig/Source/ControlRig/Public/Rigs/RigHierarchyDefines.h
--- a/Engine/Plugins/Experimental/ControlRig/Source/ControlRig/Public/Rigs/RigHierarchyDefines.h
+++ b/Engine/Plugins/Experimental/ControlRig/Source/ControlRig/Public/Rigs/RigHierarchyDefines.h
@@ -8,18 +8,19 @@
 struct FRigHierarchyContainer;
 
 UENUM()
-enum class ERigHierarchyElementType : uint8
+enum class ERigElementType : uint8
 {
 	Bone,
 	Space,
 	Control,
+	Curve,
 
 	/** MAX - invalid */
 	Max UMETA(Hidden),
 };
 
-DECLARE_MULTICAST_DELEGATE_ThreeParams(FRigHierarchyElementChanged, FRigHierarchyContainer*, ERigHierarchyElementType, const FName&);
-DECLARE_MULTICAST_DELEGATE_ThreeParams(FRigHierarchyElementAdded, FRigHierarchyContainer*, ERigHierarchyElementType, const FName&);
-DECLARE_MULTICAST_DELEGATE_ThreeParams(FRigHierarchyElementRemoved, FRigHierarchyContainer*, ERigHierarchyElementType, const FName&);
-DECLARE_MULTICAST_DELEGATE_FourParams(FRigHierarchyElementRenamed, FRigHierarchyContainer*, ERigHierarchyElementType, const FName&, const FName&);
-DECLARE_MULTICAST_DELEGATE_FiveParams(FRigHierarchyElementReparented, FRigHierarchyContainer*, ERigHierarchyElementType, const FName&, const FName&, const FName&);
+DECLARE_MULTICAST_DELEGATE_ThreeParams(FRigElementChanged, FRigHierarchyContainer*, ERigElementType, const FName&);
+DECLARE_MULTICAST_DELEGATE_ThreeParams(FRigElementAdded, FRigHierarchyContainer*, ERigElementType, const FName&);
+DECLARE_MULTICAST_DELEGATE_ThreeParams(FRigElementRemoved, FRigHierarchyContainer*, ERigElementType, const FName&);
+DECLARE_MULTICAST_DELEGATE_FourParams(FRigElementRenamed, FRigHierarchyContainer*, ERigElementType, const FName&, const FName&);
+DECLARE_MULTICAST_DELEGATE_FiveParams(FRigElementReparented, FRigHierarchyContainer*, ERigElementType, const FName&, const FName&, const FName&);
diff --git a/Engine/Plugins/Experimental/ControlRig/Source/ControlRig/Public/Rigs/RigSpaceHierarchy.h b/Engine/Plugins/Experimental/ControlRig/Source/ControlRig/Public/Rigs/RigSpaceHierarchy.h
new file mode 100644
--- /dev/null
+++ b/Engine/Plugins/Experimental/ControlRig/Source/ControlRig/Public/Rigs/RigSpaceHierarchy.h
@@ -0,0 +1,170 @@
+// Copyright 1998-2019 Epic Games, Inc. All Rights Reserved.
+
+#pragma once
+
+#include "CoreMinimal.h"
+#include "RigHierarchyDefines.h"
+#include "RigSpaceHierarchy.generated.h"
+
+class UControlRig;
+
+UENUM(BlueprintType)
+enum class ERigSpaceType : uint8
+{
+	/** Not attached to anything */
+	Global,
+
+	/** Attached to a bone */
+	Bone,
+
+	/** MAX - invalid */
+	Max UMETA(Hidden),
+};
+
+USTRUCT(BlueprintType)
+struct CONTROLRIG_API FRigSpace
+{
+	GENERATED_BODY()
+
+	FRigSpace()
+		: Name(NAME_None)
+		, Index(INDEX_NONE)
+		, SpaceType(ERigSpaceType::Global)
+		, ParentName(NAME_None)
+		, ParentIndex(INDEX_NONE)
+		, InitialTransform(FTransform::Identity)
+		, LocalTransform(FTransform::Identity)
+	{
+	}
+
+	UPROPERTY(VisibleAnywhere, Category = FRigSpaceHierarchy)
+	FName Name;
+
+	UPROPERTY(VisibleAnywhere, Category = FRigCurveContainer)
+	int32 Index;
+
+	UPROPERTY(VisibleAnywhere, Category = FRigSpaceHierarchy)
+	ERigSpaceType SpaceType;
+
+	UPROPERTY(VisibleAnywhere, Category = FRigSpaceHierarchy)
+	FName ParentName;
+
+	UPROPERTY(transient)
+	int32 ParentIndex;
+
+	UPROPERTY(VisibleAnywhere, Category = FRigSpaceHierarchy)
+	FTransform InitialTransform;
+
+	UPROPERTY(transient, VisibleAnywhere, Category = FRigSpaceHierarchy)
+	FTransform LocalTransform;
+};
+
+USTRUCT()
+struct CONTROLRIG_API FRigSpaceHierarchy
+{
+	GENERATED_BODY()
+
+	FRigSpaceHierarchy();
+	FRigSpaceHierarchy& operator= (const FRigSpaceHierarchy &InOther);
+
+	FORCEINLINE ERigElementType RigElementType() const { return ERigElementType::Space; }
+
+	FORCEINLINE int32 Num() const { return Spaces.Num(); }
+	FORCEINLINE const FRigSpace& operator[](int32 InIndex) const { return Spaces[InIndex]; }
+	FORCEINLINE FRigSpace& operator[](int32 InIndex) { return Spaces[InIndex]; }
+	FORCEINLINE const FRigSpace& operator[](const FName& InName) const { return Spaces[GetIndex(InName)]; }
+	FORCEINLINE FRigSpace& operator[](const FName& InName) { return Spaces[GetIndex(InName)]; }
+
+	FORCEINLINE TArray<FRigSpace>::RangedForIteratorType      begin()       { return Spaces.begin(); }
+	FORCEINLINE TArray<FRigSpace>::RangedForConstIteratorType begin() const { return Spaces.begin(); }
+	FORCEINLINE TArray<FRigSpace>::RangedForIteratorType      end()         { return Spaces.end();   }
+	FORCEINLINE TArray<FRigSpace>::RangedForConstIteratorType end() const   { return Spaces.end();   }
+
+	FORCEINLINE bool IsNameAvailable(const FName& InPotentialNewName) const { return GetIndex(InPotentialNewName) == INDEX_NONE; }
+
+	FName GetSafeNewName(const FName& InPotentialNewName) const;
+
+	FRigSpace& Add(const FName& InNewName, ERigSpaceType InSpaceType, const FName& InParentName, const FTransform& InTransform);
+
+	FRigSpace Remove(const FName& InNameToRemove);
+
+	FName Rename(const FName& InOldName, const FName& InNewName);
+
+	void Reparent(const FName& InName, const FName& InNewParentName);
+
+	FName GetName(int32 InIndex) const;
+
+	FORCEINLINE int32 GetIndex(const FName& InName) const
+	{
+		if(NameToIndexMapping.Num() != Spaces.Num())
+		{
+			return GetIndexSlow(InName);
+		}
+
+		const int32* Index = NameToIndexMapping.Find(InName);
+		if (Index)
+		{
+			return *Index;
+		}
+
+		return INDEX_NONE;
+	}
+
+	void SetGlobalTransform(const FName& InName, const FTransform& InTransform);
+
+	void SetGlobalTransform(int32 InIndex, const FTransform& InTransform);
+
+	FTransform GetGlobalTransform(const FName& InName) const;
+
+	FTransform GetGlobalTransform(int32 InIndex) const;
+
+	void SetLocalTransform(const FName& InName, const FTransform& InTransform);
+
+	void SetLocalTransform(int32 InIndex, const FTransform& InTransform);
+
+	FTransform GetLocalTransform(const FName& InName) const;
+
+	FTransform GetLocalTransform(int32 InIndex) const;
+
+	void SetInitialTransform(const FName& InName, const FTransform& InTransform);
+
+	void SetInitialTransform(int32 InIndex, const FTransform& InTransform);
+
+	FTransform GetInitialTransform(const FName& InName) const;
+
+	FTransform GetInitialTransform(int32 InIndex) const;
+
+	// updates all of the internal caches
+	void Initialize();
+
+	// clears the hierarchy and removes all content
+	void Reset();
+
+	// resets all of the transforms back to the initial transform
+	void ResetTransforms();
+
+private:
+
+	// disable copy constructor
+	FRigSpaceHierarchy(const FRigSpaceHierarchy& InOther) {}
+
+#if WITH_EDITOR
+	FRigHierarchyContainer* Container;
+	FRigElementAdded OnSpaceAdded;
+	FRigElementRemoved OnSpaceRemoved;
+	FRigElementRenamed OnSpaceRenamed;
+#endif
+
+	UPROPERTY(EditAnywhere, Category = FRigSpaceHierarchy)
+	TArray<FRigSpace> Spaces;
+
+	// can serialize fine? 
+	UPROPERTY()
+	TMap<FName, int32> NameToIndexMapping;
+
+	int32 GetIndexSlow(const FName& InName) const;
+
+	void RefreshMapping();
+
+	friend struct FRigHierarchyContainer;
+};
diff --git a/Engine/Plugins/Experimental/ControlRig/Source/ControlRig/Public/Units/RigUnitContext.h b/Engine/Plugins/Experimental/ControlRig/Source/ControlRig/Public/Units/RigUnitContext.h
--- a/Engine/Plugins/Experimental/ControlRig/Source/ControlRig/Public/Units/RigUnitContext.h
+++ b/Engine/Plugins/Experimental/ControlRig/Source/ControlRig/Public/Units/RigUnitContext.h
@@ -27,6 +27,9 @@ struct FRigUnitContext
 	FRigUnitContext()
 		: DrawInterface(nullptr)
 		, DataSourceRegistry(nullptr)
+		, DeltaTime(0.f)
+		, State(EControlRigState::Invalid)
+		, Hierarchy(nullptr)
 #if WITH_EDITOR
 		, Log(nullptr)
 #endif
@@ -47,16 +50,49 @@ struct FRigUnitContext
 	EControlRigState State;
 
 	/** The current hierarchy being executed */
-	FRigHierarchyRef HierarchyReference;
-
-	/** The current curve being executed */
-	FRigCurveContainerRef CurveReference;
+	const FRigHierarchyContainer* Hierarchy;
 
 #if WITH_EDITOR
 	/** A handle to the compiler log */
 	FControlRigLog* Log;
 #endif
 
+	const FRigBoneHierarchy* GetBones() const
+	{
+		if (Hierarchy != nullptr)
+		{
+			return &Hierarchy->BoneHierarchy;
+		}
+		return nullptr;
+	}
+
+	const FRigSpaceHierarchy* GetSpaces() const
+	{
+		if (Hierarchy != nullptr)
+		{
+			return &Hierarchy->SpaceHierarchy;
+		}
+		return nullptr;
+	}
+
+	const FRigControlHierarchy* GetControls() const
+	{
+		if (Hierarchy != nullptr)
+		{
+			return &Hierarchy->ControlHierarchy;
+		}
+		return nullptr;
+	}
+
+	const FRigCurveContainer* GetCurves() const
+	{
+		if (Hierarchy != nullptr)
+		{
+			return &Hierarchy->CurveContainer;
+		}
+		return nullptr;
+	}
+
 	/**
 	 * Returns a given data source and cast it to the expected class.
 	 *
diff --git a/Engine/Plugins/Experimental/ControlRig/Source/ControlRig/Public/Units/RigUnitTest.h b/Engine/Plugins/Experimental/ControlRig/Source/ControlRig/Public/Units/RigUnitTest.h
--- a/Engine/Plugins/Experimental/ControlRig/Source/ControlRig/Public/Units/RigUnitTest.h
+++ b/Engine/Plugins/Experimental/ControlRig/Source/ControlRig/Public/Units/RigUnitTest.h
@@ -14,22 +14,20 @@ public:
 	FControlRigUnitTestBase(const FString& InName, bool bIsComplex)
 		: FAutomationTestBase(InName, bIsComplex)
 		, HierarchyContainer()
-		, HierarchyRef(FRigHierarchyRef())
 		, BoneHierarchy(HierarchyContainer.BoneHierarchy)
-		, CurveContainer()
-		, CurveContainerRef(&CurveContainer)
+		, SpaceHierarchy(HierarchyContainer.SpaceHierarchy)
+		, ControlHierarchy(HierarchyContainer.ControlHierarchy)
+		, CurveContainer(HierarchyContainer.CurveContainer)
 	{
-		HierarchyRef.Container = &HierarchyContainer;
-		Context.HierarchyReference = HierarchyRef;
-		ExecuteContext.HierarchyReference = HierarchyRef;
-		ExecuteContext.CurveReference = CurveContainerRef;
+		Context.Hierarchy = &HierarchyContainer;
+		ExecuteContext.Hierarchy= &HierarchyContainer;
 	}
 
 	FRigHierarchyContainer HierarchyContainer;
-	FRigHierarchyRef HierarchyRef;
 	FRigBoneHierarchy& BoneHierarchy;
-	FRigCurveContainer CurveContainer;
-	FRigCurveContainerRef CurveContainerRef;
+	FRigSpaceHierarchy& SpaceHierarchy;
+	FRigControlHierarchy& ControlHierarchy;
+	FRigCurveContainer& CurveContainer;
 	FControlRigExecuteContext ExecuteContext;
 	FRigUnitContext Context;
 };
diff --git a/Engine/Plugins/Experimental/ControlRig/Source/ControlRigDeveloper/Private/ControlRigBlueprint.cpp b/Engine/Plugins/Experimental/ControlRig/Source/ControlRigDeveloper/Private/ControlRigBlueprint.cpp
--- a/Engine/Plugins/Experimental/ControlRig/Source/ControlRigDeveloper/Private/ControlRigBlueprint.cpp
+++ b/Engine/Plugins/Experimental/ControlRig/Source/ControlRigDeveloper/Private/ControlRigBlueprint.cpp
@@ -600,13 +600,18 @@ bool UControlRigBlueprint::PerformArrayOperation(const FString& InPropertyPath,
 
 void UControlRigBlueprint::CleanupBoneHierarchyDeprecated()
 {
-	if (Hierarchy_DEPRECATED.Num() == 0)
+	if (Hierarchy_DEPRECATED.Num() > 0)
 	{
-		return;
+		HierarchyContainer.BoneHierarchy = Hierarchy_DEPRECATED;
+		Hierarchy_DEPRECATED.Reset();
+	}
+
+	if (CurveContainer_DEPRECATED.Num() > 0)
+	{
+		HierarchyContainer.CurveContainer = CurveContainer_DEPRECATED;
+		CurveContainer_DEPRECATED.Reset();
 	}
 
-	HierarchyContainer.BoneHierarchy = Hierarchy_DEPRECATED;
-	Hierarchy_DEPRECATED.Reset();
 }
 
 #undef LOCTEXT_NAMESPACE
diff --git a/Engine/Plugins/Experimental/ControlRig/Source/ControlRigDeveloper/Private/ControlRigBlueprintCompiler.cpp b/Engine/Plugins/Experimental/ControlRig/Source/ControlRigDeveloper/Private/ControlRigBlueprintCompiler.cpp
--- a/Engine/Plugins/Experimental/ControlRig/Source/ControlRigDeveloper/Private/ControlRigBlueprintCompiler.cpp
+++ b/Engine/Plugins/Experimental/ControlRig/Source/ControlRigDeveloper/Private/ControlRigBlueprintCompiler.cpp
@@ -488,7 +488,6 @@ void FControlRigBlueprintCompilerContext::CopyTermDefaultsToDefaultObject(UObjec
 
 		UControlRig* ControlRig = CastChecked<UControlRig>(DefaultObject);
 		ControlRig->Hierarchy = ControlRigBlueprint->HierarchyContainer;
-		ControlRig->CurveContainer = ControlRigBlueprint->CurveContainer;
 		// copy available rig units info, so that control rig can do things with it
 		ControlRig->AllowSourceAccessProperties = ControlRigBlueprint->AllowSourceAccessProperties;
 		ControlRigBlueprint->UpdateParametersOnControlRig(ControlRig);
diff --git a/Engine/Plugins/Experimental/ControlRig/Source/ControlRigDeveloper/Public/ControlRigBlueprint.h b/Engine/Plugins/Experimental/ControlRig/Source/ControlRigDeveloper/Public/ControlRigBlueprint.h
--- a/Engine/Plugins/Experimental/ControlRig/Source/ControlRigDeveloper/Public/ControlRigBlueprint.h
+++ b/Engine/Plugins/Experimental/ControlRig/Source/ControlRigDeveloper/Public/ControlRigBlueprint.h
@@ -169,7 +169,7 @@ private:
 	FRigBoneHierarchy Hierarchy_DEPRECATED;
 
 	UPROPERTY()
-	FRigCurveContainer CurveContainer;
+	FRigCurveContainer CurveContainer_DEPRECATED;
 
 	/** The default skeletal mesh to use when previewing this asset */
 	UPROPERTY(DuplicateTransient, AssetRegistrySearchable)
diff --git a/Engine/Plugins/Experimental/ControlRig/Source/ControlRigEditor/Private/Editor/ControlRigEditor.cpp b/Engine/Plugins/Experimental/ControlRig/Source/ControlRigEditor/Private/Editor/ControlRigEditor.cpp
--- a/Engine/Plugins/Experimental/ControlRig/Source/ControlRigEditor/Private/Editor/ControlRigEditor.cpp
+++ b/Engine/Plugins/Experimental/ControlRig/Source/ControlRigEditor/Private/Editor/ControlRigEditor.cpp
@@ -185,10 +185,10 @@ void FControlRigEditor::InitControlRigEditor(const EToolkitMode::Type Mode, cons
 	}
 
 	InControlRigBlueprint->OnModified().AddSP(this, &FControlRigEditor::HandleModelModified);
-	InControlRigBlueprint->HierarchyContainer.OnElementAdded.AddSP(this, &FControlRigEditor::OnHierarchyElementAdded);
-	InControlRigBlueprint->HierarchyContainer.OnElementRemoved.AddSP(this, &FControlRigEditor::OnHierarchyElementRemoved);
-	InControlRigBlueprint->HierarchyContainer.OnElementRenamed.AddSP(this, &FControlRigEditor::OnHierarchyElementRenamed);
-	InControlRigBlueprint->HierarchyContainer.OnElementReparented.AddSP(this, &FControlRigEditor::OnHierarchyElementReparented);
+	InControlRigBlueprint->HierarchyContainer.OnElementAdded.AddSP(this, &FControlRigEditor::OnRigElementAdded);
+	InControlRigBlueprint->HierarchyContainer.OnElementRemoved.AddSP(this, &FControlRigEditor::OnRigElementRemoved);
+	InControlRigBlueprint->HierarchyContainer.OnElementRenamed.AddSP(this, &FControlRigEditor::OnRigElementRenamed);
+	InControlRigBlueprint->HierarchyContainer.OnElementReparented.AddSP(this, &FControlRigEditor::OnRigElementReparented);
 
 	BindCommands();
 
@@ -1573,27 +1573,27 @@ void FControlRigEditor::OnBoneHierarchyChanged()
 	CacheBoneNameList();
 }
 
-void FControlRigEditor::OnHierarchyElementAdded(FRigHierarchyContainer* Container, ERigHierarchyElementType ElementType, const FName& InName)
+void FControlRigEditor::OnRigElementAdded(FRigHierarchyContainer* Container, ERigElementType ElementType, const FName& InName)
 {
-	if (ElementType == ERigHierarchyElementType::Bone)
+	if (ElementType == ERigElementType::Bone)
 	{
 		OnBoneHierarchyChanged();
 	}
 }
 
-void FControlRigEditor::OnHierarchyElementRemoved(FRigHierarchyContainer* Container, ERigHierarchyElementType ElementType, const FName& InName)
+void FControlRigEditor::OnRigElementRemoved(FRigHierarchyContainer* Container, ERigElementType ElementType, const FName& InName)
 {
-	if (ElementType == ERigHierarchyElementType::Bone)
+	if (ElementType == ERigElementType::Bone)
 	{
 		OnBoneHierarchyChanged();
 	}
 }
 
-void FControlRigEditor::OnHierarchyElementRenamed(FRigHierarchyContainer* Container, ERigHierarchyElementType ElementType, const FName& InOldName, const FName& InNewName)
+void FControlRigEditor::OnRigElementRenamed(FRigHierarchyContainer* Container, ERigElementType ElementType, const FName& InOldName, const FName& InNewName)
 {
 	DECLARE_SCOPE_HIERARCHICAL_COUNTER_FUNC()
 
-	if (ElementType == ERigHierarchyElementType::Bone)
+	if (ElementType == ERigElementType::Bone)
 	{
 		UControlRigBlueprint* Blueprint = GetControlRigBlueprint();
 		for (UEdGraph* Graph : Blueprint->UbergraphPages)
@@ -1640,11 +1640,58 @@ void FControlRigEditor::OnHierarchyElementRenamed(FRigHierarchyContainer* Contai
 			CacheBoneNameList();
 		}
 	}
+	else if(ElementType == ERigElementType::Curve)
+	{
+		UControlRigBlueprint* Blueprint = GetControlRigBlueprint();
+		for (UEdGraph* Graph : Blueprint->UbergraphPages)
+		{
+			UControlRigGraph* RigGraph = Cast<UControlRigGraph>(Graph);
+			if (RigGraph == nullptr)
+			{
+				continue;
+			}
+
+			for (UEdGraphNode* Node : RigGraph->Nodes)
+			{
+				UControlRigGraphNode* RigNode = Cast<UControlRigGraphNode>(Node);
+				if (RigNode == nullptr)
+				{
+					continue;
+				}
+
+				UStructProperty* UnitProperty = RigNode->GetUnitProperty();
+				UStruct* UnitStruct = RigNode->GetUnitScriptStruct();
+				if (UnitProperty && UnitStruct)
+				{
+					for (TFieldIterator<UNameProperty> It(UnitStruct); It; ++It)
+					{
+						if (It->HasMetaData(UControlRig::CurveNameMetaName))
+						{
+							FString PinName = FString::Printf(TEXT("%s.%s"), *UnitProperty->GetName(), *It->GetName());
+							UEdGraphPin* Pin = Node->FindPin(PinName, EEdGraphPinDirection::EGPD_Input);
+							if (Pin)
+							{
+								FName CurrentCurve = FName(*Pin->GetDefaultAsString());
+								if (CurrentCurve == InOldName)
+								{
+									const FScopedTransaction Transaction(NSLOCTEXT("ControlRigEditor", "ChangeCurveNamePinValue", "Change Curve Name Pin Value"));
+									Pin->Modify();
+									Pin->GetSchema()->TrySetDefaultValue(*Pin, InNewName.ToString());
+								}
+							}
+						}
+					}
+				}
+			}
+
+			CacheCurveNameList();
+		}
+	}
 }
 
-void FControlRigEditor::OnHierarchyElementReparented(FRigHierarchyContainer* Container, ERigHierarchyElementType ElementType, const FName& InName, const FName& InOldParentName, const FName& InNewParentName)
+void FControlRigEditor::OnRigElementReparented(FRigHierarchyContainer* Container, ERigElementType ElementType, const FName& InName, const FName& InOldParentName, const FName& InNewParentName)
 {
-	if (ElementType == ERigHierarchyElementType::Bone)
+	if (ElementType == ERigElementType::Bone)
 	{
 		OnBoneHierarchyChanged();
 	}
@@ -1677,56 +1724,6 @@ void FControlRigEditor::OnCurveContainerChanged()
 	NotificationPtr->SetCompletionState(SNotificationItem::CS_Success);
 }
 
-void FControlRigEditor::OnCurveRenamed(const FName& OldName, const FName& NewName)
-{
-	DECLARE_SCOPE_HIERARCHICAL_COUNTER_FUNC()
-
-	UControlRigBlueprint* Blueprint = GetControlRigBlueprint();
-	for (UEdGraph* Graph : Blueprint->UbergraphPages)
-	{
-		UControlRigGraph* RigGraph = Cast<UControlRigGraph>(Graph);
-		if (RigGraph == nullptr)
-		{
-			continue;
-		}
-
-		for (UEdGraphNode* Node : RigGraph->Nodes)
-		{
-			UControlRigGraphNode* RigNode = Cast<UControlRigGraphNode>(Node);
-			if (RigNode == nullptr)
-			{
-				continue;
-			}
-
-			UStructProperty* UnitProperty = RigNode->GetUnitProperty();
-			UStruct* UnitStruct = RigNode->GetUnitScriptStruct();
-			if (UnitProperty && UnitStruct)
-			{
-				for (TFieldIterator<UNameProperty> It(UnitStruct); It; ++It)
-				{
-					if (It->HasMetaData(UControlRig::CurveNameMetaName))
-					{
-						FString PinName = FString::Printf(TEXT("%s.%s"), *UnitProperty->GetName(), *It->GetName());
-						UEdGraphPin* Pin = Node->FindPin(PinName, EEdGraphPinDirection::EGPD_Input);
-						if (Pin)
-						{
-							FName CurrentCurve = FName(*Pin->GetDefaultAsString());
-							if (CurrentCurve == OldName)
-							{
-								const FScopedTransaction Transaction(NSLOCTEXT("ControlRigEditor", "ChangeCurveNamePinValue", "Change Curve Name Pin Value"));
-								Pin->Modify();
-								Pin->GetSchema()->TrySetDefaultValue(*Pin, NewName.ToString());
-							}
-						}
-					}
-				}
-			}
-		}
-
-		CacheCurveNameList();
-	}
-}
-
 void FControlRigEditor::OnGraphNodeDropToPerform(TSharedPtr<FGraphNodeDragDropOp> DragDropOp, UEdGraph* Graph, const FVector2D& NodePosition, const FVector2D& ScreenPosition)
 {
 	DECLARE_SCOPE_HIERARCHICAL_COUNTER_FUNC()
diff --git a/Engine/Plugins/Experimental/ControlRig/Source/ControlRigEditor/Private/Editor/ControlRigEditor.h b/Engine/Plugins/Experimental/ControlRig/Source/ControlRigEditor/Private/Editor/ControlRigEditor.h
--- a/Engine/Plugins/Experimental/ControlRig/Source/ControlRigEditor/Private/Editor/ControlRigEditor.h
+++ b/Engine/Plugins/Experimental/ControlRig/Source/ControlRigEditor/Private/Editor/ControlRigEditor.h
@@ -109,12 +109,11 @@ public:
 	UControlRig* GetInstanceRig() const { return ControlRig;  }
 
 	void OnCurveContainerChanged();
-	void OnCurveRenamed(const FName& OldName, const FName& NewName);
 
-	void OnHierarchyElementAdded(FRigHierarchyContainer* Container, ERigHierarchyElementType ElementType, const FName& InName);
-	void OnHierarchyElementRemoved(FRigHierarchyContainer* Container, ERigHierarchyElementType ElementType, const FName& InName);
-	void OnHierarchyElementRenamed(FRigHierarchyContainer* Container, ERigHierarchyElementType ElementType, const FName& InOldName, const FName& InNewName);
-	void OnHierarchyElementReparented(FRigHierarchyContainer* Container, ERigHierarchyElementType ElementType, const FName& InName, const FName& InOldParentName, const FName& InNewParentName);
+	void OnRigElementAdded(FRigHierarchyContainer* Container, ERigElementType ElementType, const FName& InName);
+	void OnRigElementRemoved(FRigHierarchyContainer* Container, ERigElementType ElementType, const FName& InName);
+	void OnRigElementRenamed(FRigHierarchyContainer* Container, ERigElementType ElementType, const FName& InOldName, const FName& InNewName);
+	void OnRigElementReparented(FRigHierarchyContainer* Container, ERigElementType ElementType, const FName& InName, const FName& InOldParentName, const FName& InNewParentName);
 
 	void OnGraphNodeDropToPerform(TSharedPtr<FGraphNodeDragDropOp> DragDropOp, UEdGraph* Graph, const FVector2D& NodePosition, const FVector2D& ScreenPosition);
 
diff --git a/Engine/Plugins/Experimental/ControlRig/Source/ControlRigEditor/Private/Editor/SRigBoneHierarchy.cpp b/Engine/Plugins/Experimental/ControlRig/Source/ControlRigEditor/Private/Editor/SRigBoneHierarchy.cpp
--- a/Engine/Plugins/Experimental/ControlRig/Source/ControlRigEditor/Private/Editor/SRigBoneHierarchy.cpp
+++ b/Engine/Plugins/Experimental/ControlRig/Source/ControlRigEditor/Private/Editor/SRigBoneHierarchy.cpp
@@ -180,7 +180,7 @@ void SRigBoneHierarchy::Construct(const FArguments& InArgs, TSharedRef<FControlR
 	// @todo: find a better place to do it
 	ControlRigBlueprint->HierarchyContainer.BoneHierarchy.Initialize();
 
-	ControlRigBlueprint->HierarchyContainer.OnElementChanged.AddRaw(this, &SRigBoneHierarchy::OnBoneHierarchyChanged);
+	ControlRigBlueprint->HierarchyContainer.OnElementChanged.AddRaw(this, &SRigBoneHierarchy::OnRigElementChanged);
 
 	// for deleting, renaming, dragging
 	CommandList = MakeShared<FUICommandList>();
@@ -406,9 +406,9 @@ void SRigBoneHierarchy::SelectBone(const FName& BoneName) const
 	}
 }
 
-void SRigBoneHierarchy::OnBoneHierarchyChanged(FRigHierarchyContainer* Container, ERigHierarchyElementType ElementType, const FName& InName)
+void SRigBoneHierarchy::OnRigElementChanged(FRigHierarchyContainer* Container, ERigElementType ElementType, const FName& InName)
 {
-	if (ElementType != ERigHierarchyElementType::Bone)
+	if (ElementType != ERigElementType::Bone)
 	{
 		return;
 	}
@@ -558,6 +558,8 @@ void SRigBoneHierarchy::ImportHierarchy(const FAssetData& InAssetData)
 				Hier->Add(BoneInfos[BoneIndex].Name, ParentName, FAnimationRuntime::GetComponentSpaceTransform(RefSkeleton, BonePoses, BoneIndex));
 			}
 		}
+
+		FSlateApplication::Get().DismissAllMenus();
 	}
 }
 
@@ -794,16 +796,10 @@ bool SRigBoneHierarchy::OnVerifyNameChanged(const FName& OldName, const FName& N
 	FRigBoneHierarchy* Hierarchy = GetHierarchy();
 	if (Hierarchy)
 	{
-		const int32 Found = Hierarchy->GetIndex(OldName);
-		if (Found != INDEX_NONE)
+		if (!Hierarchy->IsNameAvailable(NewName))
 		{
-			const int32 Duplicate = Hierarchy->GetIndex(NewName);
-			if (Duplicate != INDEX_NONE)
-			{
-				OutErrorMessage = FText::FromString(TEXT("Duplicate name exists"));
-
-				return false;
-			}
+			OutErrorMessage = FText::FromString(TEXT("Duplicate name exists"));
+			return false;
 		}
 	}
 
diff --git a/Engine/Plugins/Experimental/ControlRig/Source/ControlRigEditor/Private/Editor/SRigBoneHierarchy.h b/Engine/Plugins/Experimental/ControlRig/Source/ControlRigEditor/Private/Editor/SRigBoneHierarchy.h
--- a/Engine/Plugins/Experimental/ControlRig/Source/ControlRigEditor/Private/Editor/SRigBoneHierarchy.h
+++ b/Engine/Plugins/Experimental/ControlRig/Source/ControlRigEditor/Private/Editor/SRigBoneHierarchy.h
@@ -191,7 +191,7 @@ private:
 
 	void ClearDetailPanel() const;
 	void SelectBone(const FName& BoneName) const;
-	void OnBoneHierarchyChanged(FRigHierarchyContainer* Container, ERigHierarchyElementType ElementType, const FName& InName);
+	void OnRigElementChanged(FRigHierarchyContainer* Container, ERigElementType ElementType, const FName& InName);
 
 public:
 	bool RenameBone(const FName& OldName, const FName& NewName);
diff --git a/Engine/Plugins/Experimental/ControlRig/Source/ControlRigEditor/Private/Editor/SRigCurveContainer.cpp b/Engine/Plugins/Experimental/ControlRig/Source/ControlRigEditor/Private/Editor/SRigCurveContainer.cpp
--- a/Engine/Plugins/Experimental/ControlRig/Source/ControlRigEditor/Private/Editor/SRigCurveContainer.cpp
+++ b/Engine/Plugins/Experimental/ControlRig/Source/ControlRigEditor/Private/Editor/SRigCurveContainer.cpp
@@ -143,6 +143,7 @@ void SRigCurveContainer::Construct(const FArguments& InArgs, TSharedRef<FControl
 {
 	ControlRigEditor = InControlRigEditor;
 	ControlRigBlueprint = InControlRigEditor.Get().GetControlRigBlueprint();
+	ControlRigBlueprint->HierarchyContainer.OnElementChanged.AddRaw(this, &SRigCurveContainer::OnRigElementChanged);
 
 	// Register and bind all our menu commands
 	FCurveContainerCommands::Register();
@@ -197,6 +198,14 @@ void SRigCurveContainer::Construct(const FArguments& InArgs, TSharedRef<FControl
 
 SRigCurveContainer::~SRigCurveContainer()
 {
+	if (ControlRigEditor.IsValid())
+	{
+		ControlRigBlueprint = ControlRigEditor.Pin()->GetControlRigBlueprint();
+		if (ControlRigBlueprint.IsValid())
+		{
+			ControlRigBlueprint->HierarchyContainer.OnElementChanged.RemoveAll(this);
+		}
+	}
 }
 
 FReply SRigCurveContainer::OnKeyDown(const FGeometry& MyGeometry, const FKeyEvent& InKeyEvent)
@@ -250,11 +259,6 @@ void SRigCurveContainer::OnFilterTextChanged( const FText& SearchText )
 	RefreshCurveList();
 }
 
-void SRigCurveContainer::OnCurvesChanged()
-{
-	RefreshCurveList();
-}
-
 void SRigCurveContainer::OnFilterTextCommitted( const FText& SearchText, ETextCommit::Type CommitInfo )
 {
 	// Just do the same as if the user typed in the box
@@ -338,7 +342,6 @@ void SRigCurveContainer::CreateNewNameEntry(const FText& CommittedText, ETextCom
 		if (Container)
 		{
 			Container->Add(FName(*CommittedText.ToString()));
-			RefreshCurveList();
 		}
 	}
 }
@@ -394,7 +397,6 @@ void SRigCurveContainer::OnNameCommitted(const FText& InNewName, ETextCommit::Ty
 			FName NewName = FName(*InNewName.ToString());
 			FName OldName = Item->CurveName;
 			Container->Rename(OldName, NewName);
-			RefreshCurveList();
 		}
 	}
 }
@@ -409,8 +411,6 @@ void SRigCurveContainer::OnDeleteNameClicked()
 		{
 			Container->Remove(Item->CurveName);
 		}
-
-		RefreshCurveList();
 	}
 }
 
@@ -448,6 +448,16 @@ void SRigCurveContainer::ChangeCurveName(const FName& OldName, const FName& NewN
 	}
 }
 
+void SRigCurveContainer::OnRigElementChanged(FRigHierarchyContainer* Container, ERigElementType ElementType, const FName& InName)
+{
+	if (ElementType != ERigElementType::Curve)
+	{
+		return;
+	}
+
+	RefreshCurveList();
+}
+
 void SRigCurveContainer::CreateImportMenu(FMenuBuilder& MenuBuilder)
 {
 	MenuBuilder.AddWidget(
@@ -512,8 +522,6 @@ void SRigCurveContainer::ImportCurve(const FAssetData& InAssetData)
 				Container->Add(NameArray[Index]);
 			}
 
-			ControlRigEditor.Pin()->OnCurveContainerChanged();
-			RefreshCurveList();
 			FSlateApplication::Get().DismissAllMenus();
 		}
 	}
@@ -526,7 +534,7 @@ FRigCurveContainer* SRigCurveContainer::GetInstanceCurveContainer() const
 		UControlRig* ControlRig = ControlRigEditor.Pin()->GetInstanceRig();
 		if (ControlRig)
 		{
-			return &ControlRig->CurveContainer;
+			return &ControlRig->Hierarchy.CurveContainer;
 		}
 	}
 
@@ -537,7 +545,7 @@ FRigCurveContainer* SRigCurveContainer::GetCurveContainer() const
 {
 	if (ControlRigBlueprint.IsValid())
 	{
-		return &ControlRigBlueprint->CurveContainer;
+		return &ControlRigBlueprint->HierarchyContainer.CurveContainer;
 	}
 
 	return nullptr;
diff --git a/Engine/Plugins/Experimental/ControlRig/Source/ControlRigEditor/Private/Editor/SRigCurveContainer.h b/Engine/Plugins/Experimental/ControlRig/Source/ControlRigEditor/Private/Editor/SRigCurveContainer.h
--- a/Engine/Plugins/Experimental/ControlRig/Source/ControlRigEditor/Private/Editor/SRigCurveContainer.h
+++ b/Engine/Plugins/Experimental/ControlRig/Source/ControlRigEditor/Private/Editor/SRigCurveContainer.h
@@ -7,6 +7,7 @@
 #include "Widgets/Views/STableViewBase.h"
 #include "Widgets/Views/STableRow.h"
 #include "Widgets/Views/SListView.h"
+#include "Rigs/RigHierarchyContainer.h"
 
 //////////////////////////////////////////////////////////////////////////
 // FDisplayedRigCurveInfo
@@ -158,9 +159,6 @@ public:
 	*/
 	void OnPreviewAssetChanged(class UAnimationAsset* NewPreviewAsset);
 
-	/** Is registered with ControlRig Editor to handle when curves change. */
-	void OnCurvesChanged();
-
 	/**
 	* Filters the SListView when the user changes the search text box (NameFilterBox)
 	*
@@ -231,6 +229,7 @@ private:
 	void SetCurveValue(const FName& CurveName, float CurveValue);
 	float GetCurveValue(const FName& CurveName);
 	void ChangeCurveName(const FName& OldName, const FName& NewName);
+	void OnRigElementChanged(FRigHierarchyContainer* Container, ERigElementType ElementType, const FName& InName);
 
 	// import curve part
 	void ImportCurve(const FAssetData& InAssetData);
