ControlRig: Further progress on ControlRig Hierarchy Datastructures + unit tests, cycle checking

diff --git a/Engine/Plugins/Experimental/ControlRig/Source/ControlRig/Private/Rigs/RigBoneHierarchy.cpp b/Engine/Plugins/Experimental/ControlRig/Source/ControlRig/Private/Rigs/RigBoneHierarchy.cpp
--- a/Engine/Plugins/Experimental/ControlRig/Source/ControlRig/Private/Rigs/RigBoneHierarchy.cpp
+++ b/Engine/Plugins/Experimental/ControlRig/Source/ControlRig/Private/Rigs/RigBoneHierarchy.cpp
@@ -9,9 +9,7 @@
 ////////////////////////////////////////////////////////////////////////////////
 
 FRigBoneHierarchy::FRigBoneHierarchy()
-#if WITH_EDITOR
 	:Container(nullptr)
-#endif
 {
 }
 
@@ -53,25 +51,23 @@ FName FRigBoneHierarchy::GetSafeNewName(const FName& InPotentialNewName) const
 
 FRigBone& FRigBoneHierarchy::Add(const FName& InNewName, const FName& InParentName, const FTransform& InInitTransform)
 {
-	int32 ParentIndex = GetIndex(InParentName);
-	bool bHasParent = (ParentIndex != INDEX_NONE);
-
 	FRigBone NewBone;
 	NewBone.Name = GetSafeNewName(InNewName);
-	NewBone.ParentIndex = ParentIndex;
-	NewBone.ParentName = bHasParent? InParentName : NAME_None;
+	NewBone.ParentIndex = GetIndex(InParentName);
+	NewBone.ParentName = NewBone.ParentIndex == INDEX_NONE ? NAME_None : InParentName;
 	NewBone.InitialTransform = InInitTransform;
 	NewBone.GlobalTransform = InInitTransform;
 	RecalculateLocalTransform(NewBone);
 
+	FName NewBoneName = NewBone.Name;
 	Bones.Add(NewBone);
 	RefreshMapping();
 
 #if WITH_EDITOR
-	OnBoneAdded.Broadcast(Container, RigElementType(), NewBone.Name);
+	OnBoneAdded.Broadcast(Container, RigElementType(), NewBoneName);
 #endif
 
-	int32 Index = GetIndex(NewBone.Name);
+	int32 Index = GetIndex(NewBoneName);
 	return Bones[Index];
 }
 
@@ -83,32 +79,48 @@ FRigBone& FRigBoneHierarchy::Add(const FName& InNewName, const FName& InParentNa
 	return Bone;
 }
 
-void FRigBoneHierarchy::Reparent(const FName& InName, const FName& InNewParentName)
+bool FRigBoneHierarchy::Reparent(const FName& InName, const FName& InNewParentName)
 {
 	int32 Index = GetIndex(InName);
-	// can't parent to itself
+
 	if (Index != INDEX_NONE && InName != InNewParentName)
 	{
-		// should allow reparent to none (no parent)
-		// if invalid, we consider to be none
+		FRigBone& Bone= Bones[Index];
+
+#if WITH_EDITOR
+		FName OldParentName = Bone.ParentName;
+#endif
+
 		int32 ParentIndex = GetIndex(InNewParentName);
-		bool bHasParent = (ParentIndex != INDEX_NONE);
-		FRigBone& CurBone = Bones[Index];
-		CurBone.ParentIndex = ParentIndex;
-		FName OldParentName = CurBone.ParentName;
-		CurBone.ParentName = (bHasParent)? InNewParentName : NAME_None;
-		RecalculateLocalTransform(CurBone);
+		if (Container != nullptr)
+		{
+			if (Container->IsParentedTo(ERigElementType::Bone, ParentIndex, ERigElementType::Bone, Index))
+			{
+				ParentIndex = INDEX_NONE;
+			}
+		}
+
+		Bone.ParentIndex = ParentIndex;
+		Bone.ParentName = Bone.ParentIndex == INDEX_NONE ? NAME_None : InNewParentName;
+		RecalculateLocalTransform(Bone);
+
+#if WITH_EDITOR
+		FName NewParentName = Bone.ParentName;
+#endif
 
 		// we want to make sure parent is before the child
 		RefreshMapping();
 
 #if WITH_EDITOR
-		if (OldParentName != InNewParentName)
+		if (OldParentName != NewParentName)
 		{
-			OnBoneReparented.Broadcast(Container, RigElementType(), InName, OldParentName, InNewParentName);
+			OnBoneReparented.Broadcast(Container, RigElementType(), InName, OldParentName, NewParentName);
 		}
 #endif
+		return Bones[GetIndex(InName)].ParentName == InNewParentName;
 	}
+
+	return false;
 }
 
 FRigBone FRigBoneHierarchy::Remove(const FName& InNameToRemove)
@@ -488,3 +500,4 @@ void FRigBoneHierarchy::PropagateTransform(int32 InIndex)
 		PropagateTransform(Index);
 	}
 }
+
diff --git a/Engine/Plugins/Experimental/ControlRig/Source/ControlRig/Private/Rigs/RigControlHierarchy.cpp b/Engine/Plugins/Experimental/ControlRig/Source/ControlRig/Private/Rigs/RigControlHierarchy.cpp
--- a/Engine/Plugins/Experimental/ControlRig/Source/ControlRig/Private/Rigs/RigControlHierarchy.cpp
+++ b/Engine/Plugins/Experimental/ControlRig/Source/ControlRig/Private/Rigs/RigControlHierarchy.cpp
@@ -9,9 +9,7 @@
 ////////////////////////////////////////////////////////////////////////////////
 
 FRigControlHierarchy::FRigControlHierarchy()
-#if WITH_EDITOR
 	:Container(nullptr)
-#endif
 {
 }
 
@@ -51,67 +49,118 @@ FName FRigControlHierarchy::GetSafeNewName(const FName& InPotentialNewName) cons
 	return Name;
 }
 
-FRigControl& FRigControlHierarchy::Add(const FName& InNewName, const FName& InParentName, const FName& InSpaceName, const FTransform& InTransform)
+FRigControl& FRigControlHierarchy::Add(const FName& InNewName, ERigControlType InControlType, const FName& InParentName, const FName& InSpaceName, const FTransform& InTransform)
 {
-	int32 ParentIndex = GetIndex(InParentName);
-	bool bHasParent = (ParentIndex != INDEX_NONE);
-
 	FRigControl NewControl;
 	NewControl.Name = GetSafeNewName(InNewName);
-	NewControl.ParentIndex = ParentIndex;
-	NewControl.ParentName = bHasParent ? InParentName : NAME_None;;
+	NewControl.ControlType = InControlType;
+	NewControl.ParentIndex = GetIndex(InParentName);
+	NewControl.ParentName = NewControl.ParentIndex == INDEX_NONE ? NAME_None : InParentName;
 	NewControl.SpaceIndex = INDEX_NONE;
 	NewControl.SpaceName = NAME_None;
 	NewControl.InitialTransform = InTransform;
 	NewControl.LocalTransform = InTransform;
 
-	// todo: set control type, etc
-
+	FName NewControlName = NewControl.Name;
 	Controls.Add(NewControl);
 	RefreshMapping();
 
 #if WITH_EDITOR
-	OnControlAdded.Broadcast(Container, RigElementType(), NewControl.Name);
+	OnControlAdded.Broadcast(Container, RigElementType(), NewControlName);
 #endif
 
-	SetSpace(NewControl.Name, InSpaceName);
+	SetSpace(NewControlName, InSpaceName);
 
-	int32 Index = GetIndex(NewControl.Name);
+	int32 Index = GetIndex(NewControlName);
 	return Controls[Index];
 }
 
-void FRigControlHierarchy::Reparent(const FName& InName, const FName& InNewParentName)
+bool FRigControlHierarchy::Reparent(const FName& InName, const FName& InNewParentName)
 {
 	int32 Index = GetIndex(InName);
 	// can't parent to itself
 	if (Index != INDEX_NONE && InName != InNewParentName)
 	{
-		// should allow reparent to none (no parent)
-		// if invalid, we consider to be none
+		FRigControl& Control = Controls[Index];
+
+#if WITH_EDITOR
+		FName OldParentName = Control.ParentName;
+#endif
+
+		struct Local
+		{
+			static bool IsParentedTo(int32 Child, int32 Parent, const TArray<FRigControl>& Controls)
+			{
+				if (Parent == INDEX_NONE || Child == INDEX_NONE)
+				{
+					return false;
+				}
+
+				if (Child == Parent)
+				{
+					return true;
+				}
+
+				if (Controls[Child].ParentIndex == Parent)
+				{
+					return true;
+				}
+
+				return IsParentedTo(Controls[Child].ParentIndex, Parent, Controls);
+			}
+		};
+
 		int32 ParentIndex = GetIndex(InNewParentName);
-		bool bHasParent = (ParentIndex != INDEX_NONE);
-		FRigControl& CurControl = Controls[Index];
-		FName OldParentName = CurControl.ParentName;
-		CurControl.ParentName = (bHasParent)? InNewParentName : NAME_None;
-		CurControl.ParentIndex = ParentIndex;
+		if (Local::IsParentedTo(ParentIndex, Index, Controls))
+		{
+			ParentIndex = INDEX_NONE;
+		}
 
-		// todo: we need to update the space if this had been parented by space as well before
+		Control.ParentIndex = ParentIndex;
+		Control.ParentName = Control.ParentIndex == INDEX_NONE ? NAME_None : InNewParentName;
+
+#if WITH_EDITOR
+		FName NewParentName = Control.ParentName;
+#endif
 
-		// we want to make sure parent is before the child
 		RefreshMapping();
 
 #if WITH_EDITOR
-		if (OldParentName != InNewParentName)
+		if (OldParentName != NewParentName)
 		{
-			OnControlReparented.Broadcast(Container, RigElementType(), InName, OldParentName, InNewParentName);
+			OnControlReparented.Broadcast(Container, RigElementType(), InName, OldParentName, NewParentName);
 		}
 #endif
+		return Controls[GetIndex(InName)].ParentName == InNewParentName;
 	}
+	return false;
 }
 
 void FRigControlHierarchy::SetSpace(const FName& InName, const FName& InNewSpaceName)
 {
-	// todo find the space, set the space index...
+	int32 Index = GetIndex(InName);
+	if (Index != INDEX_NONE)
+	{
+		int32 SpaceIndex = GetSpaceIndex(InNewSpaceName);
+
+		if (SpaceIndex != INDEX_NONE)
+		{
+			if (Container != nullptr)
+			{
+				if (Container->IsParentedTo(ERigElementType::Space, SpaceIndex, ERigElementType::Control, Index))
+				{
+					SpaceIndex = INDEX_NONE;
+				}
+			}
+		}
+
+		Controls[Index].SpaceIndex = SpaceIndex;
+		Controls[Index].SpaceName = Controls[Index].SpaceIndex == INDEX_NONE ? NAME_None : InNewSpaceName;
+
+#if WITH_EDITOR
+		OnControlReparented.Broadcast(Container, RigElementType(), InName, Controls[Index].ParentName, Controls[Index].ParentName);
+#endif
+	}
 }
 
 FRigControl FRigControlHierarchy::Remove(const FName& InNameToRemove)
@@ -200,12 +249,25 @@ void FRigControlHierarchy::SetGlobalTransform(const FName& InName, const FTransf
 
 void FRigControlHierarchy::SetGlobalTransform(int32 InIndex, const FTransform& InTransform)
 {
+	if (Container == nullptr)
+	{
+		SetLocalTransform(InIndex, InTransform);
+		return;
+	}
+
 	if (Controls.IsValidIndex(InIndex))
 	{
-		FRigControl& Control = Controls[InIndex];
-		// todo: get the space's transform + inverse
-		//Control.GlobalTransform = InTransform;
-		//Control.GlobalTransform.NormalizeRotation();
+		const FRigControl& Control = Controls[InIndex];
+		if (Control.SpaceIndex == INDEX_NONE)
+		{
+			SetLocalTransform(InIndex, InTransform);
+			return;
+		}
+		else
+		{
+			FTransform ParentTransform = Container->GetGlobalTransform(ERigElementType::Space, Control.SpaceIndex);
+			SetLocalTransform(InIndex, InTransform.GetRelativeTransform(ParentTransform));
+		}
 	}
 }
 
@@ -216,11 +278,23 @@ FTransform FRigControlHierarchy::GetGlobalTransform(const FName& InName) const
 
 FTransform FRigControlHierarchy::GetGlobalTransform(int32 InIndex) const
 {
+	if (Container == nullptr)
+	{
+		return GetLocalTransform(InIndex);
+	}
+
 	if (Controls.IsValidIndex(InIndex))
 	{
-		// todo: get the space's transform and compute
-		//return Controls[InIndex].GlobalTransform;
-		return FTransform::Identity;
+		const FRigControl& Control = Controls[InIndex];
+		if (Control.SpaceIndex == INDEX_NONE)
+		{
+			return Control.LocalTransform;
+		}
+		else
+		{
+			FTransform ParentTransform = Container->GetGlobalTransform(ERigElementType::Space, Control.SpaceIndex);
+			return Control.LocalTransform * ParentTransform;
+		}
 	}
 
 	return FTransform::Identity;
@@ -459,3 +533,12 @@ int32 FRigControlHierarchy::GetChildrenRecursive(const int32 InIndex, TArray<int
 
 	return OutChildren.Num();
 }
+
+int32 FRigControlHierarchy::GetSpaceIndex(const FName& InName) const
+{
+	if (Container == nullptr || InName == NAME_None)
+	{
+		return INDEX_NONE;
+	}
+	return Container->GetIndex(ERigElementType::Space, InName);
+}
diff --git a/Engine/Plugins/Experimental/ControlRig/Source/ControlRig/Private/Rigs/RigCurveContainer.cpp b/Engine/Plugins/Experimental/ControlRig/Source/ControlRig/Private/Rigs/RigCurveContainer.cpp
--- a/Engine/Plugins/Experimental/ControlRig/Source/ControlRig/Private/Rigs/RigCurveContainer.cpp
+++ b/Engine/Plugins/Experimental/ControlRig/Source/ControlRig/Private/Rigs/RigCurveContainer.cpp
@@ -9,9 +9,7 @@
 ////////////////////////////////////////////////////////////////////////////////
 
 FRigCurveContainer::FRigCurveContainer()
-#if WITH_EDITOR
 	:Container(nullptr)
-#endif
 {
 }
 
@@ -58,14 +56,15 @@ FRigCurve& FRigCurveContainer::Add(const FName& InNewName)
 	FRigCurve NewCurve;
 	NewCurve.Name = GetSafeNewName(InNewName);
 	NewCurve.Value = 0.f;
+	FName NewCurveName = NewCurve.Name;
 	Curves.Add(NewCurve);
 	RefreshMapping();
 
 #if WITH_EDITOR
-	OnCurveAdded.Broadcast(Container, RigElementType(), NewCurve.Name);
+	OnCurveAdded.Broadcast(Container, RigElementType(), NewCurveName);
 #endif
 
-	int32 Index = GetIndex(NewCurve.Name);
+	int32 Index = GetIndex(NewCurveName);
 	return Curves[Index];
 }
 
diff --git a/Engine/Plugins/Experimental/ControlRig/Source/ControlRig/Private/Rigs/RigHierarchyContainer.cpp b/Engine/Plugins/Experimental/ControlRig/Source/ControlRig/Private/Rigs/RigHierarchyContainer.cpp
--- a/Engine/Plugins/Experimental/ControlRig/Source/ControlRig/Private/Rigs/RigHierarchyContainer.cpp
+++ b/Engine/Plugins/Experimental/ControlRig/Source/ControlRig/Private/Rigs/RigHierarchyContainer.cpp
@@ -24,12 +24,12 @@ FRigHierarchyContainer& FRigHierarchyContainer::operator= (const FRigHierarchyCo
 
 void FRigHierarchyContainer::Initialize()
 {
-#if WITH_EDITOR
 	BoneHierarchy.Container = this;
 	SpaceHierarchy.Container = this;
 	ControlHierarchy.Container = this;
 	CurveContainer.Container = this;
 
+#if WITH_EDITOR
 	BoneHierarchy.OnBoneAdded.RemoveAll(this);
 	BoneHierarchy.OnBoneRemoved.RemoveAll(this);
 	BoneHierarchy.OnBoneRenamed.RemoveAll(this);
@@ -93,6 +93,193 @@ void FRigHierarchyContainer::ResetTransforms()
 	CurveContainer.ResetValues();
 }
 
+FTransform FRigHierarchyContainer::GetInitialTransform(ERigElementType InElementType, int32 InIndex) const
+{
+	if(InIndex == INDEX_NONE)
+	{
+		return FTransform::Identity;
+	}
+
+	switch(InElementType)
+	{
+		case ERigElementType::Bone:
+		{
+			return BoneHierarchy.GetInitialTransform(InIndex);
+		}
+		case ERigElementType::Space:
+		{
+			return SpaceHierarchy.GetInitialTransform(InIndex);
+		}
+		case ERigElementType::Control:
+		{
+			return ControlHierarchy.GetInitialTransform(InIndex);
+		}
+		case ERigElementType::Curve:
+		{
+			break;
+		}
+	}
+
+	return FTransform::Identity;
+}
+
+#if WITH_EDITOR
+
+void FRigHierarchyContainer::SetInitialTransform(ERigElementType InElementType, int32 InIndex, const FTransform& InTransform)
+{
+	if (InIndex == INDEX_NONE)
+	{
+		return;
+	}
+
+	switch(InElementType)
+	{
+		case ERigElementType::Bone:
+		{
+			BoneHierarchy.SetInitialTransform(InIndex, InTransform);
+			break;
+		}
+		case ERigElementType::Space:
+		{
+			SpaceHierarchy.SetInitialTransform(InIndex, InTransform);
+			break;
+		}
+		case ERigElementType::Control:
+		{
+			ControlHierarchy.SetInitialTransform(InIndex, InTransform);
+			break;
+		}
+		case ERigElementType::Curve:
+		{
+			break;
+		}
+	}
+}
+
+#endif
+
+FTransform FRigHierarchyContainer::GetLocalTransform(ERigElementType InElementType, int32 InIndex) const
+{
+	if(InIndex == INDEX_NONE)
+	{
+		return FTransform::Identity;
+	}
+
+	switch(InElementType)
+	{
+		case ERigElementType::Bone:
+		{
+			return BoneHierarchy.GetLocalTransform(InIndex);
+		}
+		case ERigElementType::Space:
+		{
+			return SpaceHierarchy.GetLocalTransform(InIndex);
+		}
+		case ERigElementType::Control:
+		{
+			return ControlHierarchy.GetLocalTransform(InIndex);
+		}
+		case ERigElementType::Curve:
+		{
+			break;
+		}
+	}
+
+	return FTransform::Identity;
+}
+
+void FRigHierarchyContainer::SetLocalTransform(ERigElementType InElementType, int32 InIndex, const FTransform& InTransform)
+{
+	if (InIndex == INDEX_NONE)
+	{
+		return;
+	}
+
+	switch(InElementType)
+	{
+		case ERigElementType::Bone:
+		{
+			BoneHierarchy.SetLocalTransform(InIndex, InTransform);
+			break;
+		}
+		case ERigElementType::Space:
+		{
+			SpaceHierarchy.SetLocalTransform(InIndex, InTransform);
+			break;
+		}
+		case ERigElementType::Control:
+		{
+			ControlHierarchy.SetLocalTransform(InIndex, InTransform);
+			break;
+		}
+		case ERigElementType::Curve:
+		{
+			break;
+		}
+	}
+}
+
+FTransform FRigHierarchyContainer::GetGlobalTransform(ERigElementType InElementType, int32 InIndex) const
+{
+	if(InIndex == INDEX_NONE)
+	{
+		return FTransform::Identity;
+	}
+
+	switch(InElementType)
+	{
+		case ERigElementType::Bone:
+		{
+			return BoneHierarchy.GetGlobalTransform(InIndex);
+		}
+		case ERigElementType::Space:
+		{
+			return SpaceHierarchy.GetGlobalTransform(InIndex);
+		}
+		case ERigElementType::Control:
+		{
+			return ControlHierarchy.GetGlobalTransform(InIndex);
+		}
+		case ERigElementType::Curve:
+		{
+			break;
+		}
+	}
+
+	return FTransform::Identity;
+}
+
+void FRigHierarchyContainer::SetGlobalTransform(ERigElementType InElementType, int32 InIndex, const FTransform& InTransform)
+{
+	if (InIndex == INDEX_NONE)
+	{
+		return;
+	}
+
+	switch(InElementType)
+	{
+		case ERigElementType::Bone:
+		{
+			BoneHierarchy.SetGlobalTransform(InIndex, InTransform);
+			break;
+		}
+		case ERigElementType::Space:
+		{
+			SpaceHierarchy.SetGlobalTransform(InIndex, InTransform);
+			break;
+		}
+		case ERigElementType::Control:
+		{
+			ControlHierarchy.SetGlobalTransform(InIndex, InTransform);
+			break;
+		}
+		case ERigElementType::Curve:
+		{
+			break;
+		}
+	}
+}
+
 #if WITH_EDITOR
 
 void FRigHierarchyContainer::HandleOnElementAdded(FRigHierarchyContainer* InContainer, ERigElementType InElementType, const FName& InName)
@@ -124,3 +311,106 @@ void FRigHierarchyContainer::HandleOnElementReparented(FRigHierarchyContainer* I
 }
 
 #endif
+
+bool FRigHierarchyContainer::IsParentedTo(ERigElementType InChildType, int32 InChildIndex, ERigElementType InParentType, int32 InParentIndex) const
+{
+	ensure(InChildIndex != INDEX_NONE);
+
+	if (InParentIndex == INDEX_NONE)
+	{
+		return false;
+	}
+
+	switch (InChildType)
+	{
+		case ERigElementType::Curve:
+		{
+			return false;
+		}
+		case ERigElementType::Bone:
+		{
+			switch (InParentType)
+			{
+				case ERigElementType::Bone:
+				{
+					if (BoneHierarchy[InChildIndex].ParentIndex != INDEX_NONE)
+					{
+						if (BoneHierarchy[InChildIndex].ParentIndex == InParentIndex)
+						{
+							return true;
+						}
+						return IsParentedTo(ERigElementType::Bone, BoneHierarchy[InChildIndex].ParentIndex, InParentType, InParentIndex);
+					}
+					// no break - fall through to next case
+				}
+				case ERigElementType::Space:
+				case ERigElementType::Control:
+				case ERigElementType::Curve:
+				{
+					return false;
+				}
+			}
+		}
+		case ERigElementType::Space:
+		{
+			const FRigSpace& ChildSpace = SpaceHierarchy[InChildIndex];
+			switch (ChildSpace.SpaceType)
+			{
+				case ERigSpaceType::Global:
+				{
+					return false;
+				}
+				case ERigSpaceType::Bone:
+				{
+					if (ChildSpace.ParentIndex == InParentIndex && InParentType == ERigElementType::Bone)
+					{
+						return true;
+					}
+					return IsParentedTo(ERigElementType::Bone, ChildSpace.ParentIndex, InParentType, InParentIndex);
+				}
+				case ERigSpaceType::Space:
+				{
+					if (ChildSpace.ParentIndex == InParentIndex && InParentType == ERigElementType::Space)
+					{
+						return true;
+					}
+					return IsParentedTo(ERigElementType::Space, ChildSpace.ParentIndex, InParentType, InParentIndex);
+				}
+				case ERigSpaceType::Control:
+				{
+					if (ChildSpace.ParentIndex == InParentIndex && InParentType == ERigElementType::Control)
+					{
+						return true;
+					}
+					return IsParentedTo(ERigElementType::Control, ChildSpace.ParentIndex, InParentType, InParentIndex);
+				}
+			}
+		}
+		case ERigElementType::Control:
+		{
+			const FRigControl& ChildControl = ControlHierarchy[InChildIndex];
+			switch (InParentType)
+			{
+				case ERigElementType::Space:
+				{
+					if (ChildControl.SpaceIndex == InParentIndex)
+					{
+						return true;
+					}
+					// no break - fall through to next cases
+				}
+				case ERigElementType::Control:
+				case ERigElementType::Bone:
+				{
+					return IsParentedTo(ERigElementType::Space, ChildControl.SpaceIndex, InParentType, InParentIndex);
+				}
+				case ERigElementType::Curve:
+				{
+					return false;
+				}
+			}
+		}
+	}
+
+	return false;
+}
diff --git a/Engine/Plugins/Experimental/ControlRig/Source/ControlRig/Private/Rigs/RigSpaceHierarchy.cpp b/Engine/Plugins/Experimental/ControlRig/Source/ControlRig/Private/Rigs/RigSpaceHierarchy.cpp
--- a/Engine/Plugins/Experimental/ControlRig/Source/ControlRig/Private/Rigs/RigSpaceHierarchy.cpp
+++ b/Engine/Plugins/Experimental/ControlRig/Source/ControlRig/Private/Rigs/RigSpaceHierarchy.cpp
@@ -9,9 +9,7 @@
 ////////////////////////////////////////////////////////////////////////////////
 
 FRigSpaceHierarchy::FRigSpaceHierarchy()
-#if WITH_EDITOR
 	:Container(nullptr)
-#endif
 {
 }
 
@@ -55,30 +53,89 @@ FRigSpace& FRigSpaceHierarchy::Add(const FName& InNewName, ERigSpaceType InSpace
 {
 	FRigSpace NewSpace;
 	NewSpace.Name = GetSafeNewName(InNewName);
-	NewSpace.SpaceType = InSpaceType;
-	NewSpace.ParentName = InParentName;
-	NewSpace.ParentIndex = INDEX_NONE;
+	NewSpace.ParentIndex = GetParentIndex(InSpaceType, InParentName);
+	NewSpace.SpaceType = NewSpace.ParentIndex == INDEX_NONE ? ERigSpaceType::Global : InSpaceType;
+	NewSpace.ParentName = NewSpace.ParentIndex == INDEX_NONE ? NAME_None : InParentName;
 	NewSpace.InitialTransform = InTransform;
 	NewSpace.LocalTransform = InTransform;
+	FName NewSpaceName = NewSpace.Name;
 	Spaces.Add(NewSpace);
 	RefreshMapping();
 
 #if WITH_EDITOR
-	OnSpaceAdded.Broadcast(Container, RigElementType(), NewSpace.Name);
+	OnSpaceAdded.Broadcast(Container, RigElementType(), NewSpaceName);
 #endif
 
-	int32 Index = GetIndex(NewSpace.Name);
+	int32 Index = GetIndex(NewSpaceName);
 	return Spaces[Index];
 }
 
-void FRigSpaceHierarchy::Reparent(const FName& InName, const FName& InNewParentName)
+bool FRigSpaceHierarchy::Reparent(const FName& InName, ERigSpaceType InSpaceType, const FName& InNewParentName)
 {
 	int32 Index = GetIndex(InName);
 	if (Index != INDEX_NONE)
 	{
-		Spaces[Index].ParentName = InNewParentName;
-		Spaces[Index].ParentIndex = INDEX_NONE;
+		FRigSpace& Space = Spaces[Index];
+
+#if WITH_EDITOR
+		FName OldParentName = Space.ParentName;
+#endif
+
+		int32 ParentIndex = GetParentIndex(InSpaceType, InNewParentName);
+		if (ParentIndex != INDEX_NONE)
+		{
+			if (Container != nullptr)
+			{
+				switch (InSpaceType)
+				{
+					case ERigSpaceType::Global:
+					{
+						ParentIndex = INDEX_NONE;
+						break;
+					}
+					case ERigSpaceType::Bone:
+					{
+						break;
+					}
+					case ERigSpaceType::Space:
+					{
+						if (Container->IsParentedTo(ERigElementType::Space, ParentIndex, ERigElementType::Space, Index))
+						{
+							ParentIndex = INDEX_NONE;
+						}
+						break;
+					}
+					case ERigSpaceType::Control:
+					{
+						if (Container->IsParentedTo(ERigElementType::Control, ParentIndex, ERigElementType::Space, Index))
+						{
+							ParentIndex = INDEX_NONE;
+						}
+						break;
+					}
+				}
+			}
+		}
+
+		Space.ParentIndex = ParentIndex;
+		Space.SpaceType = Space.ParentIndex == INDEX_NONE ? ERigSpaceType::Global : InSpaceType;
+		Space.ParentName = Space.ParentIndex == INDEX_NONE ? NAME_None : InNewParentName;
+
+#if WITH_EDITOR
+		FName NewParentName = Space.ParentName;
+#endif
+
+		RefreshMapping();
+
+#if WITH_EDITOR
+		if (OldParentName != NewParentName)
+		{
+			OnSpaceReparented.Broadcast(Container, RigElementType(), InName, OldParentName, NewParentName);
+		}
+#endif		
+		return Spaces[GetIndex(InName)].ParentName == InNewParentName;
 	}
+	return false;
 }
 
 FRigSpace FRigSpaceHierarchy::Remove(const FName& InNameToRemove)
@@ -126,12 +183,42 @@ void FRigSpaceHierarchy::SetGlobalTransform(const FName& InName, const FTransfor
 
 void FRigSpaceHierarchy::SetGlobalTransform(int32 InIndex, const FTransform& InTransform)
 {
+	if (Container == nullptr)
+	{
+		SetLocalTransform(InIndex, InTransform);
+		return;
+	}
+
 	if (Spaces.IsValidIndex(InIndex))
 	{
 		FRigSpace& Space = Spaces[InIndex];
-		// todo: get the space's transform + inverse
-		//Space.GlobalTransform = InTransform;
-		//Space.GlobalTransform.NormalizeRotation();
+		FTransform ParentTransform = FTransform::Identity;
+
+		switch (Space.SpaceType)
+		{
+			case ERigSpaceType::Global:
+			{
+				SetLocalTransform(InIndex, InTransform);
+				return;
+			}
+			case ERigSpaceType::Bone:
+			{
+				ParentTransform = Container->GetGlobalTransform(ERigElementType::Bone, Space.ParentIndex);
+				break;
+			}
+			case ERigSpaceType::Space:
+			{
+				ParentTransform = Container->GetGlobalTransform(ERigElementType::Space, Space.ParentIndex);
+				break;
+			}
+			case ERigSpaceType::Control:
+			{
+				ParentTransform = Container->GetGlobalTransform(ERigElementType::Control, Space.ParentIndex);
+				break;
+			}
+		}
+
+		SetLocalTransform(InIndex, InTransform.GetRelativeTransform(ParentTransform));
 	}
 }
 
@@ -142,11 +229,33 @@ FTransform FRigSpaceHierarchy::GetGlobalTransform(const FName& InName) const
 
 FTransform FRigSpaceHierarchy::GetGlobalTransform(int32 InIndex) const
 {
+	if(Container == nullptr)
+	{
+		return GetLocalTransform(InIndex);
+	}
+
 	if (Spaces.IsValidIndex(InIndex))
 	{
-		// todo: get the space's transform and compute
-		//return Spaces[InIndex].GlobalTransform;
-		return FTransform::Identity;
+		const FRigSpace& Space = Spaces[InIndex];
+		switch (Space.SpaceType)
+		{
+			case ERigSpaceType::Global:
+			{
+				return Space.LocalTransform;
+			}
+			case ERigSpaceType::Bone:
+			{
+				return Space.LocalTransform * Container->GetGlobalTransform(ERigElementType::Bone, Space.ParentIndex);
+			}
+			case ERigSpaceType::Space:
+			{
+				return Space.LocalTransform * Container->GetGlobalTransform(ERigElementType::Space, Space.ParentIndex);
+			}
+			case ERigSpaceType::Control:
+			{
+				return Space.LocalTransform * Container->GetGlobalTransform(ERigElementType::Control, Space.ParentIndex);
+			}
+		}
 	}
 
 	return FTransform::Identity;
@@ -257,13 +366,6 @@ void FRigSpaceHierarchy::Initialize()
 {
 	RefreshMapping();
 
-	// update parent index
-	for (int32 Index = 0; Index < Spaces.Num(); ++Index)
-	{
-		// todo
-		//Spaces[Index].ParentIndex = GetIndex(Spaces[Index].ParentName);
-	}
-
 	// initialize transform
 	for (int32 Index = 0; Index < Spaces.Num(); ++Index)
 	{
@@ -284,3 +386,31 @@ void FRigSpaceHierarchy::ResetTransforms()
 		Spaces[Index].LocalTransform = Spaces[Index].InitialTransform;
 	}
 }
+
+int32 FRigSpaceHierarchy::GetParentIndex(ERigSpaceType InSpaceType, const FName& InName) const
+{
+	if (Container != nullptr)
+	{
+		switch (InSpaceType)
+		{
+			case ERigSpaceType::Global:
+			{
+				return INDEX_NONE;
+			}
+			case ERigSpaceType::Bone:
+			{
+				return Container->GetIndex(ERigElementType::Bone, InName);
+			}
+			case ERigSpaceType::Space:
+			{
+				return Container->GetIndex(ERigElementType::Space, InName);
+			}
+			case ERigSpaceType::Control:
+			{
+				return Container->GetIndex(ERigElementType::Control, InName);
+			}
+		}
+	}
+
+	return INDEX_NONE;
+}
diff --git a/Engine/Plugins/Experimental/ControlRig/Source/ControlRig/Public/Rigs/RigBoneHierarchy.h b/Engine/Plugins/Experimental/ControlRig/Source/ControlRig/Public/Rigs/RigBoneHierarchy.h
--- a/Engine/Plugins/Experimental/ControlRig/Source/ControlRig/Public/Rigs/RigBoneHierarchy.h
+++ b/Engine/Plugins/Experimental/ControlRig/Source/ControlRig/Public/Rigs/RigBoneHierarchy.h
@@ -77,7 +77,7 @@ struct CONTROLRIG_API FRigBoneHierarchy
 
 	FName GetSafeNewName(const FName& InPotentialNewName) const;
 
-	FRigBone& Add(const FName& InNewName, const FName& InParentName, const FTransform& InInitTransform);
+	FRigBone& Add(const FName& InNewName, const FName& InParentName = NAME_None, const FTransform& InInitTransform = FTransform::Identity);
 
 	FRigBone& Add(const FName& InNewName, const FName& InParentName, const FTransform& InInitTransform, const FTransform& InLocalTransform, const FTransform& InGlobalTransform);
 
@@ -85,7 +85,7 @@ struct CONTROLRIG_API FRigBoneHierarchy
 
 	FName Rename(const FName& InOldName, const FName& InNewName);
 
-	void Reparent(const FName& InName, const FName& InNewParentName);
+	bool Reparent(const FName& InName, const FName& InNewParentName);
 
 	FName GetName(int32 InIndex) const;
 
@@ -148,8 +148,8 @@ private:
 	// disable copy constructor
 	FRigBoneHierarchy(const FRigBoneHierarchy& InOther) {}
 
-#if WITH_EDITOR
 	FRigHierarchyContainer* Container;
+#if WITH_EDITOR
 	FRigElementAdded OnBoneAdded;
 	FRigElementRemoved OnBoneRemoved;
 	FRigElementRenamed OnBoneRenamed;
diff --git a/Engine/Plugins/Experimental/ControlRig/Source/ControlRig/Public/Rigs/RigControlHierarchy.h b/Engine/Plugins/Experimental/ControlRig/Source/ControlRig/Public/Rigs/RigControlHierarchy.h
--- a/Engine/Plugins/Experimental/ControlRig/Source/ControlRig/Public/Rigs/RigControlHierarchy.h
+++ b/Engine/Plugins/Experimental/ControlRig/Source/ControlRig/Public/Rigs/RigControlHierarchy.h
@@ -8,6 +8,18 @@
 
 class UControlRig;
 
+UENUM(BlueprintType)
+enum class ERigControlType : uint8
+{
+	Bool,
+	Float,
+	Vector2D,
+	Vector,
+	Quat,
+	Rotator,
+	Transform
+};
+
 USTRUCT(BlueprintType)
 struct CONTROLRIG_API FRigControl
 {
@@ -16,12 +28,14 @@ struct CONTROLRIG_API FRigControl
 	FRigControl()
 		: Name(NAME_None)
 		, Index(INDEX_NONE)
+		, ControlType(ERigControlType::Transform)
 		, ParentName(NAME_None)
 		, ParentIndex(INDEX_NONE)
 		, SpaceName(NAME_None)
 		, SpaceIndex(INDEX_NONE)
 		, InitialTransform(FTransform::Identity)
 		, LocalTransform(FTransform::Identity)
+		, Color(FLinearColor::Black)
 		, Dependents()
 	{
 	}
@@ -32,6 +46,9 @@ struct CONTROLRIG_API FRigControl
 	UPROPERTY(VisibleAnywhere, Category = FRigCurveContainer)
 	int32 Index;
 
+	UPROPERTY(VisibleAnywhere, Category = FRigCurveContainer)
+	ERigControlType ControlType;
+
 	UPROPERTY(VisibleAnywhere, Category = FRigControlHierarchy)
 	FName ParentName;
 
@@ -50,6 +67,9 @@ struct CONTROLRIG_API FRigControl
 	UPROPERTY(transient, VisibleAnywhere, Category = FRigControlHierarchy)
 	FTransform LocalTransform;
 
+	UPROPERTY(VisibleAnywhere, Category = FRigControlHierarchy)
+	FLinearColor Color;
+
 	/** dependent list - direct dependent for child or anything that needs to update due to this */
 	UPROPERTY(transient)
 	TArray<int32> Dependents;
@@ -80,13 +100,13 @@ struct CONTROLRIG_API FRigControlHierarchy
 
 	FName GetSafeNewName(const FName& InPotentialNewName) const;
 
-	FRigControl& Add(const FName& InNewName, const FName& InParentName, const FName& InSpaceName, const FTransform& InTransform);
+	FRigControl& Add(const FName& InNewName, ERigControlType InControlType = ERigControlType::Transform, const FName& InParentName = NAME_None, const FName& InSpaceName = NAME_None, const FTransform& InTransform = FTransform::Identity);
 
 	FRigControl Remove(const FName& InNameToRemove);
 
 	FName Rename(const FName& InOldName, const FName& InNewName);
 
-	void Reparent(const FName& InName, const FName& InNewParentName);
+	bool Reparent(const FName& InName, const FName& InNewParentName);
 
 	void SetSpace(const FName& InName, const FName& InNewSpaceName);
 
@@ -151,8 +171,8 @@ private:
 	// disable copy constructor
 	FRigControlHierarchy(const FRigControlHierarchy& InOther) {}
 
-#if WITH_EDITOR
 	FRigHierarchyContainer* Container;
+#if WITH_EDITOR
 	FRigElementAdded OnControlAdded;
 	FRigElementRemoved OnControlRemoved;
 	FRigElementRenamed OnControlRenamed;
@@ -166,6 +186,8 @@ private:
 	UPROPERTY()
 	TMap<FName, int32> NameToIndexMapping;
 
+	int32 GetSpaceIndex(const FName& InName) const;
+
 	int32 GetIndexSlow(const FName& InName) const;
 
 	void RefreshMapping();
diff --git a/Engine/Plugins/Experimental/ControlRig/Source/ControlRig/Public/Rigs/RigHierarchyContainer.h b/Engine/Plugins/Experimental/ControlRig/Source/ControlRig/Public/Rigs/RigHierarchyContainer.h
--- a/Engine/Plugins/Experimental/ControlRig/Source/ControlRig/Public/Rigs/RigHierarchyContainer.h
+++ b/Engine/Plugins/Experimental/ControlRig/Source/ControlRig/Public/Rigs/RigHierarchyContainer.h
@@ -38,6 +38,70 @@ public:
 	void Reset();
 	void ResetTransforms();
 
+	FORCEINLINE int32 GetIndex(ERigElementType InElementType, const FName& InName) const
+	{
+		switch(InElementType)
+		{
+			case ERigElementType::Bone:
+			{
+				return BoneHierarchy.GetIndex(InName);
+			}
+			case ERigElementType::Space:
+			{
+				return SpaceHierarchy.GetIndex(InName);
+			}
+			case ERigElementType::Control:
+			{
+				return ControlHierarchy.GetIndex(InName);
+			}
+			case ERigElementType::Curve:
+			{
+				return CurveContainer.GetIndex(InName);
+			}
+		}
+		return INDEX_NONE;
+	}
+
+	FORCEINLINE FTransform GetInitialTransform(ERigElementType InElementType, const FName& InName) const
+	{
+		return GetInitialTransform(InElementType, GetIndex(InElementType, InName));
+	}
+	FTransform GetInitialTransform(ERigElementType InElementType, int32 InIndex) const;
+
+#if WITH_EDITOR
+	FORCEINLINE void SetInitialTransform(ERigElementType InElementType, const FName& InName, const FTransform& InTransform)
+	{
+		return SetInitialTransform(InElementType, GetIndex(InElementType, InName), InTransform);
+	}
+	void SetInitialTransform(ERigElementType InElementType, int32 InIndex, const FTransform& InTransform);
+#endif
+
+	FORCEINLINE FTransform GetLocalTransform(ERigElementType InElementType, const FName& InName) const
+	{
+		return GetLocalTransform(InElementType, GetIndex(InElementType, InName));
+	}
+	FTransform GetLocalTransform(ERigElementType InElementType, int32 InIndex) const;
+
+	FORCEINLINE void SetLocalTransform(ERigElementType InElementType, const FName& InName, const FTransform& InTransform)
+	{
+		return SetLocalTransform(InElementType, GetIndex(InElementType, InName), InTransform);
+	}
+	void SetLocalTransform(ERigElementType InElementType, int32 InIndex, const FTransform& InTransform);
+
+	FORCEINLINE FTransform GetGlobalTransform(ERigElementType InElementType, const FName& InName) const
+	{
+		return GetGlobalTransform(InElementType, GetIndex(InElementType, InName));
+	}
+	FTransform GetGlobalTransform(ERigElementType InElementType, int32 InIndex) const;
+
+	FORCEINLINE void SetGlobalTransform(ERigElementType InElementType, const FName& InName, const FTransform& InTransform)
+	{
+		return SetGlobalTransform(InElementType, GetIndex(InElementType, InName), InTransform);
+	}
+	void SetGlobalTransform(ERigElementType InElementType, int32 InIndex, const FTransform& InTransform);
+
+	bool IsParentedTo(ERigElementType InChildType, int32 InChildIndex, ERigElementType InParentType, int32 InParentIndex) const;
+
 #if WITH_EDITOR
 	FRigElementChanged OnElementChanged;
 	FRigElementAdded OnElementAdded;
diff --git a/Engine/Plugins/Experimental/ControlRig/Source/ControlRig/Public/Rigs/RigHierarchyDefines.h b/Engine/Plugins/Experimental/ControlRig/Source/ControlRig/Public/Rigs/RigHierarchyDefines.h
--- a/Engine/Plugins/Experimental/ControlRig/Source/ControlRig/Public/Rigs/RigHierarchyDefines.h
+++ b/Engine/Plugins/Experimental/ControlRig/Source/ControlRig/Public/Rigs/RigHierarchyDefines.h
@@ -13,10 +13,7 @@ enum class ERigElementType : uint8
 	Bone,
 	Space,
 	Control,
-	Curve,
-
-	/** MAX - invalid */
-	Max UMETA(Hidden),
+	Curve
 };
 
 DECLARE_MULTICAST_DELEGATE_ThreeParams(FRigElementChanged, FRigHierarchyContainer*, ERigElementType, const FName&);
diff --git a/Engine/Plugins/Experimental/ControlRig/Source/ControlRig/Public/Rigs/RigSpaceHierarchy.h b/Engine/Plugins/Experimental/ControlRig/Source/ControlRig/Public/Rigs/RigSpaceHierarchy.h
--- a/Engine/Plugins/Experimental/ControlRig/Source/ControlRig/Public/Rigs/RigSpaceHierarchy.h
+++ b/Engine/Plugins/Experimental/ControlRig/Source/ControlRig/Public/Rigs/RigSpaceHierarchy.h
@@ -17,8 +17,11 @@ enum class ERigSpaceType : uint8
 	/** Attached to a bone */
 	Bone,
 
-	/** MAX - invalid */
-	Max UMETA(Hidden),
+	/** Attached to a control */
+	Control,
+
+	/** Attached to a space*/
+	Space
 };
 
 USTRUCT(BlueprintType)
@@ -34,6 +37,7 @@ struct CONTROLRIG_API FRigSpace
 		, ParentIndex(INDEX_NONE)
 		, InitialTransform(FTransform::Identity)
 		, LocalTransform(FTransform::Identity)
+		, Color(FLinearColor::Black)
 	{
 	}
 
@@ -57,6 +61,9 @@ struct CONTROLRIG_API FRigSpace
 
 	UPROPERTY(transient, VisibleAnywhere, Category = FRigSpaceHierarchy)
 	FTransform LocalTransform;
+
+	UPROPERTY(VisibleAnywhere, Category = FRigControlHierarchy)
+	FLinearColor Color;
 };
 
 USTRUCT()
@@ -84,13 +91,13 @@ struct CONTROLRIG_API FRigSpaceHierarchy
 
 	FName GetSafeNewName(const FName& InPotentialNewName) const;
 
-	FRigSpace& Add(const FName& InNewName, ERigSpaceType InSpaceType, const FName& InParentName, const FTransform& InTransform);
+	FRigSpace& Add(const FName& InNewName, ERigSpaceType InSpaceType = ERigSpaceType::Global, const FName& InParentName = NAME_None, const FTransform& InTransform = FTransform::Identity);
 
 	FRigSpace Remove(const FName& InNameToRemove);
 
 	FName Rename(const FName& InOldName, const FName& InNewName);
 
-	void Reparent(const FName& InName, const FName& InNewParentName);
+	bool Reparent(const FName& InName, ERigSpaceType InSpaceType, const FName& InNewParentName);
 
 	FName GetName(int32 InIndex) const;
 
@@ -148,13 +155,16 @@ private:
 	// disable copy constructor
 	FRigSpaceHierarchy(const FRigSpaceHierarchy& InOther) {}
 
-#if WITH_EDITOR
 	FRigHierarchyContainer* Container;
+#if WITH_EDITOR
 	FRigElementAdded OnSpaceAdded;
 	FRigElementRemoved OnSpaceRemoved;
 	FRigElementRenamed OnSpaceRenamed;
+	FRigElementReparented OnSpaceReparented;
 #endif
 
+	int32 GetParentIndex(ERigSpaceType InSpaceType, const FName& InName) const;
+
 	UPROPERTY(EditAnywhere, Category = FRigSpaceHierarchy)
 	TArray<FRigSpace> Spaces;
 
