ControlRig: Implemented MVC for Curves, Bones, Controls + Spaces

Also implemented the notion of selection to the model layer of the HierarchyContainer - so that all views can synchronize

diff --git a/Engine/Plugins/Experimental/ControlRig/Source/ControlRig/Private/Rigs/RigBoneHierarchy.cpp b/Engine/Plugins/Experimental/ControlRig/Source/ControlRig/Private/Rigs/RigBoneHierarchy.cpp
--- a/Engine/Plugins/Experimental/ControlRig/Source/ControlRig/Private/Rigs/RigBoneHierarchy.cpp
+++ b/Engine/Plugins/Experimental/ControlRig/Source/ControlRig/Private/Rigs/RigBoneHierarchy.cpp
@@ -145,6 +145,9 @@ FRigBone FRigBoneHierarchy::Remove(const FName& InNameToRemove)
 	}
 
 	int32 IndexToDelete = GetIndex(InNameToRemove);
+#if WITH_EDITOR
+	Select(InNameToRemove, false);
+#endif
 	FRigBone RemovedBone = Bones[IndexToDelete];
 	Bones.RemoveAt(IndexToDelete);
 
@@ -323,6 +326,15 @@ FName FRigBoneHierarchy::Rename(const FName& InOldName, const FName& InNewName)
 		if (Found != INDEX_NONE)
 		{
 			FName NewName = GetSafeNewName(InNewName);
+
+#if WITH_EDITOR
+			bool bWasSelected = IsSelected(InOldName);
+			if(bWasSelected)
+			{
+				Select(InOldName, false);
+			}
+#endif
+
 			Bones[Found].Name = NewName;
 
 			// go through find all children and rename them
@@ -349,6 +361,10 @@ FName FRigBoneHierarchy::Rename(const FName& InOldName, const FName& InNewName)
 			{
 				OnBoneReparented.Broadcast(Container, RigElementType(), ReparentedBone, InOldName, NewName);
 			}
+			if (bWasSelected)
+			{
+				Select(NewName, true);
+			}
 #endif
 			return NewName;
 		}
@@ -501,3 +517,62 @@ void FRigBoneHierarchy::PropagateTransform(int32 InIndex)
 	}
 }
 
+#if WITH_EDITOR
+
+bool FRigBoneHierarchy::Select(const FName& InName, bool bSelect)
+{
+	if(GetIndex(InName) == INDEX_NONE)
+	{
+		return false;
+	}
+
+	if(bSelect == IsSelected(InName))
+	{
+		return false;
+	}
+
+	if(bSelect)
+	{
+		if (Container)
+		{
+			Container->SpaceHierarchy.ClearSelection();
+			Container->ControlHierarchy.ClearSelection();
+			Container->CurveContainer.ClearSelection();
+		}
+
+		Selection.Add(InName);
+	}
+	else
+	{
+		Selection.Remove(InName);
+	}
+
+	OnBoneSelected.Broadcast(Container, RigElementType(), InName, bSelect);
+
+	return true;
+}
+
+bool FRigBoneHierarchy::ClearSelection()
+{
+	TArray<FName> TempSelection;
+	TempSelection.Append(Selection);
+	for(const FName& SelectedName : TempSelection)
+	{
+		Select(SelectedName, false);
+	}
+	return TempSelection.Num() > 0;
+}
+
+TArray<FName> FRigBoneHierarchy::CurrentSelection() const
+{
+	TArray<FName> TempSelection;
+	TempSelection.Append(Selection);
+	return TempSelection;
+}
+
+bool FRigBoneHierarchy::IsSelected(const FName& InName) const
+{
+	return Selection.Contains(InName);
+}
+
+#endif
\ No newline at end of file
diff --git a/Engine/Plugins/Experimental/ControlRig/Source/ControlRig/Private/Rigs/RigControlHierarchy.cpp b/Engine/Plugins/Experimental/ControlRig/Source/ControlRig/Private/Rigs/RigControlHierarchy.cpp
--- a/Engine/Plugins/Experimental/ControlRig/Source/ControlRig/Private/Rigs/RigControlHierarchy.cpp
+++ b/Engine/Plugins/Experimental/ControlRig/Source/ControlRig/Private/Rigs/RigControlHierarchy.cpp
@@ -185,6 +185,9 @@ FRigControl FRigControlHierarchy::Remove(const FName& InNameToRemove)
 	}
 
 	int32 IndexToDelete = GetIndex(InNameToRemove);
+#if WITH_EDITOR
+	Select(InNameToRemove, false);
+#endif
 	FRigControl RemovedControl = Controls[IndexToDelete];
 	Controls.RemoveAt(IndexToDelete);
 
@@ -368,6 +371,15 @@ FName FRigControlHierarchy::Rename(const FName& InOldName, const FName& InNewNam
 		if (Found != INDEX_NONE)
 		{
 			FName NewName = GetSafeNewName(InNewName);
+
+#if WITH_EDITOR
+			bool bWasSelected = IsSelected(InOldName);
+			if(bWasSelected)
+			{
+				Select(InOldName, false);
+			}
+#endif
+
 			Controls[Found].Name = NewName;
 
 			// go through find all children and rename them
@@ -393,6 +405,10 @@ FName FRigControlHierarchy::Rename(const FName& InOldName, const FName& InNewNam
 			{
 				OnControlReparented.Broadcast(Container, RigElementType(), ReparentedControl, InOldName, NewName);
 			}
+			if(bWasSelected)
+			{
+				Select(NewName, true);
+			}
 #endif
 			return NewName;
 		}
@@ -542,3 +558,128 @@ int32 FRigControlHierarchy::GetSpaceIndex(const FName& InName) const
 	}
 	return Container->GetIndex(ERigElementType::Space, InName);
 }
+
+#if WITH_EDITOR
+
+bool FRigControlHierarchy::Select(const FName& InName, bool bSelect)
+{
+	if(GetIndex(InName) == INDEX_NONE)
+	{
+		return false;
+	}
+
+	if(bSelect == IsSelected(InName))
+	{
+		return false;
+	}
+
+	if(bSelect)
+	{
+		if (Container)
+		{
+			Container->BoneHierarchy.ClearSelection();
+			Container->SpaceHierarchy.ClearSelection();
+			Container->CurveContainer.ClearSelection();
+		}
+
+		Selection.Add(InName);
+	}
+	else
+	{
+		Selection.Remove(InName);
+	}
+
+	OnControlSelected.Broadcast(Container, RigElementType(), InName, bSelect);
+
+	return true;
+}
+
+bool FRigControlHierarchy::ClearSelection()
+{
+	TArray<FName> TempSelection;
+	TempSelection.Append(Selection);
+	for(const FName& SelectedName : TempSelection)
+	{
+		Select(SelectedName, false);
+	}
+	return TempSelection.Num() > 0;
+}
+
+TArray<FName> FRigControlHierarchy::CurrentSelection() const
+{
+	TArray<FName> TempSelection;
+	TempSelection.Append(Selection);
+	return TempSelection;
+}
+
+bool FRigControlHierarchy::IsSelected(const FName& InName) const
+{
+	return Selection.Contains(InName);
+}
+
+void FRigControlHierarchy::HandleOnElementRemoved(FRigHierarchyContainer* InContainer, ERigElementType InElementType, const FName& InName)
+{
+	if (Container == nullptr)
+	{
+		return;
+	}
+
+	switch (InElementType)
+	{
+	case ERigElementType::Space:
+	{
+		for (FRigControl& Control : Controls)
+		{
+			if (Control.SpaceName == InName)
+			{
+				Control.SpaceIndex = INDEX_NONE;
+				Control.SpaceName = NAME_None;
+#if WITH_EDITOR
+				OnControlReparented.Broadcast(Container, RigElementType(), Control.Name, Control.ParentName, Control.ParentName);
+#endif
+			}
+		}
+		break;
+	}
+	case ERigElementType::Bone:
+	case ERigElementType::Control:
+	case ERigElementType::Curve:
+	{
+		break;
+	}
+	}
+}
+
+void FRigControlHierarchy::HandleOnElementRenamed(FRigHierarchyContainer* InContainer, ERigElementType InElementType, const FName& InOldName, const FName& InNewName)
+{
+	if (Container == nullptr)
+	{
+		return;
+	}
+
+	switch (InElementType)
+	{
+		case ERigElementType::Space:
+		{
+			for (FRigControl& Control : Controls)
+			{
+				if (Control.SpaceName == InOldName)
+				{
+					Control.SpaceIndex = Container->SpaceHierarchy.GetIndex(InNewName);
+					Control.SpaceName = Control.SpaceIndex == INDEX_NONE ? NAME_None : InNewName;
+#if WITH_EDITOR
+					OnControlReparented.Broadcast(Container, RigElementType(), Control.Name, Control.ParentName, Control.ParentName);
+#endif
+				}
+			}
+			break;
+		}
+		case ERigElementType::Bone:
+		case ERigElementType::Control:
+		case ERigElementType::Curve:
+		{
+			break;
+		}
+	}
+}
+#endif
\ No newline at end of file
diff --git a/Engine/Plugins/Experimental/ControlRig/Source/ControlRig/Private/Rigs/RigCurveContainer.cpp b/Engine/Plugins/Experimental/ControlRig/Source/ControlRig/Private/Rigs/RigCurveContainer.cpp
--- a/Engine/Plugins/Experimental/ControlRig/Source/ControlRig/Private/Rigs/RigCurveContainer.cpp
+++ b/Engine/Plugins/Experimental/ControlRig/Source/ControlRig/Private/Rigs/RigCurveContainer.cpp
@@ -76,6 +76,9 @@ FRigCurve FRigCurveContainer::Remove(const FName& InName)
 
 	int32 IndexToDelete = GetIndex(InName);
 	ensure(IndexToDelete != INDEX_NONE);
+#if WITH_EDITOR
+	Select(InName, false);
+#endif
 	RemovedCurve = Curves[IndexToDelete];
 	Curves.RemoveAt(IndexToDelete);
 	RefreshMapping();
@@ -149,11 +152,24 @@ FName FRigCurveContainer::Rename(const FName& InOldName, const FName& InNewName)
 		if (Found != INDEX_NONE)
 		{
 			FName NewName = GetSafeNewName(InNewName);
+
+#if WITH_EDITOR
+			bool bWasSelected = IsSelected(InOldName);
+			if(bWasSelected)
+			{
+				Select(InOldName, false);
+			}
+#endif
+
 			Curves[Found].Name = NewName;
 			RefreshMapping();
 
 #if WITH_EDITOR
 			OnCurveRenamed.Broadcast(Container, RigElementType(), InOldName, NewName);
+			if(bWasSelected)
+			{
+				Select(NewName, true);
+			}
 #endif
 			return NewName;
 		}
@@ -200,3 +216,62 @@ void FRigCurveContainer::ResetValues()
 	}
 }
 
+#if WITH_EDITOR
+
+bool FRigCurveContainer::Select(const FName& InName, bool bSelect)
+{
+	if(GetIndex(InName) == INDEX_NONE)
+	{
+		return false;
+	}
+
+	if(bSelect == IsSelected(InName))
+	{
+		return false;
+	}
+
+	if(bSelect)
+	{
+		if (Container)
+		{
+			Container->BoneHierarchy.ClearSelection();
+			Container->SpaceHierarchy.ClearSelection();
+			Container->ControlHierarchy.ClearSelection();
+		}
+
+		Selection.Add(InName);
+	}
+	else
+	{
+		Selection.Remove(InName);
+	}
+
+	OnCurveSelected.Broadcast(Container, RigElementType(), InName, bSelect);
+
+	return true;
+}
+
+bool FRigCurveContainer::ClearSelection()
+{
+	TArray<FName> TempSelection;
+	TempSelection.Append(Selection);
+	for(const FName& SelectedName : TempSelection)
+	{
+		Select(SelectedName, false);
+	}
+	return TempSelection.Num() > 0;
+}
+
+TArray<FName> FRigCurveContainer::CurrentSelection() const
+{
+	TArray<FName> TempSelection;
+	TempSelection.Append(Selection);
+	return TempSelection;
+}
+
+bool FRigCurveContainer::IsSelected(const FName& InName) const
+{
+	return Selection.Contains(InName);
+}
+
+#endif
\ No newline at end of file
diff --git a/Engine/Plugins/Experimental/ControlRig/Source/ControlRig/Private/Rigs/RigHierarchyContainer.cpp b/Engine/Plugins/Experimental/ControlRig/Source/ControlRig/Private/Rigs/RigHierarchyContainer.cpp
--- a/Engine/Plugins/Experimental/ControlRig/Source/ControlRig/Private/Rigs/RigHierarchyContainer.cpp
+++ b/Engine/Plugins/Experimental/ControlRig/Source/ControlRig/Private/Rigs/RigHierarchyContainer.cpp
@@ -34,37 +34,56 @@ void FRigHierarchyContainer::Initialize()
 	BoneHierarchy.OnBoneRemoved.RemoveAll(this);
 	BoneHierarchy.OnBoneRenamed.RemoveAll(this);
 	BoneHierarchy.OnBoneReparented.RemoveAll(this);
+	BoneHierarchy.OnBoneSelected.RemoveAll(this);
 
 	BoneHierarchy.OnBoneAdded.AddRaw(this, &FRigHierarchyContainer::HandleOnElementAdded);
 	BoneHierarchy.OnBoneRemoved.AddRaw(this, &FRigHierarchyContainer::HandleOnElementRemoved);
 	BoneHierarchy.OnBoneRenamed.AddRaw(this, &FRigHierarchyContainer::HandleOnElementRenamed);
 	BoneHierarchy.OnBoneReparented.AddRaw(this, &FRigHierarchyContainer::HandleOnElementReparented);
+	BoneHierarchy.OnBoneSelected.AddRaw(this, &FRigHierarchyContainer::HandleOnElementSelected);
 
 	SpaceHierarchy.OnSpaceAdded.RemoveAll(this);
 	SpaceHierarchy.OnSpaceRemoved.RemoveAll(this);
 	SpaceHierarchy.OnSpaceRenamed.RemoveAll(this);
+	SpaceHierarchy.OnSpaceReparented.RemoveAll(this);
+	SpaceHierarchy.OnSpaceSelected.RemoveAll(this);
 
 	SpaceHierarchy.OnSpaceAdded.AddRaw(this, &FRigHierarchyContainer::HandleOnElementAdded);
 	SpaceHierarchy.OnSpaceRemoved.AddRaw(this, &FRigHierarchyContainer::HandleOnElementRemoved);
 	SpaceHierarchy.OnSpaceRenamed.AddRaw(this, &FRigHierarchyContainer::HandleOnElementRenamed);
+	SpaceHierarchy.OnSpaceReparented.AddRaw(this, &FRigHierarchyContainer::HandleOnElementReparented);
+	SpaceHierarchy.OnSpaceSelected.AddRaw(this, &FRigHierarchyContainer::HandleOnElementSelected);
 
 	ControlHierarchy.OnControlAdded.RemoveAll(this);
 	ControlHierarchy.OnControlRemoved.RemoveAll(this);
 	ControlHierarchy.OnControlRenamed.RemoveAll(this);
 	ControlHierarchy.OnControlReparented.RemoveAll(this);
+	ControlHierarchy.OnControlSelected.RemoveAll(this);
 
 	ControlHierarchy.OnControlAdded.AddRaw(this, &FRigHierarchyContainer::HandleOnElementAdded);
 	ControlHierarchy.OnControlRemoved.AddRaw(this, &FRigHierarchyContainer::HandleOnElementRemoved);
 	ControlHierarchy.OnControlRenamed.AddRaw(this, &FRigHierarchyContainer::HandleOnElementRenamed);
 	ControlHierarchy.OnControlReparented.AddRaw(this, &FRigHierarchyContainer::HandleOnElementReparented);
+	ControlHierarchy.OnControlSelected.AddRaw(this, &FRigHierarchyContainer::HandleOnElementSelected);
 
 	CurveContainer.OnCurveAdded.RemoveAll(this);
 	CurveContainer.OnCurveRemoved.RemoveAll(this);
 	CurveContainer.OnCurveRenamed.RemoveAll(this);
+	CurveContainer.OnCurveSelected.RemoveAll(this);
 
 	CurveContainer.OnCurveAdded.AddRaw(this, &FRigHierarchyContainer::HandleOnElementAdded);
 	CurveContainer.OnCurveRemoved.AddRaw(this, &FRigHierarchyContainer::HandleOnElementRemoved);
 	CurveContainer.OnCurveRenamed.AddRaw(this, &FRigHierarchyContainer::HandleOnElementRenamed);
+	CurveContainer.OnCurveSelected.AddRaw(this, &FRigHierarchyContainer::HandleOnElementSelected);
+
+	// wire them between each other
+	BoneHierarchy.OnBoneRenamed.RemoveAll(&SpaceHierarchy);
+	SpaceHierarchy.OnSpaceRenamed.RemoveAll(&ControlHierarchy);
+	ControlHierarchy.OnControlRenamed.RemoveAll(&SpaceHierarchy);
+	BoneHierarchy.OnBoneRenamed.AddRaw(&SpaceHierarchy, &FRigSpaceHierarchy::HandleOnElementRenamed);
+	SpaceHierarchy.OnSpaceRenamed.AddRaw(&ControlHierarchy, &FRigControlHierarchy::HandleOnElementRenamed);
+	ControlHierarchy.OnControlRenamed.AddRaw(&SpaceHierarchy, &FRigSpaceHierarchy::HandleOnElementRenamed);
+
 #endif
 
 	BoneHierarchy.Initialize();
@@ -310,6 +329,12 @@ void FRigHierarchyContainer::HandleOnElementReparented(FRigHierarchyContainer* I
 	OnElementChanged.Broadcast(InContainer, InElementType, InName);
 }
 
+void FRigHierarchyContainer::HandleOnElementSelected(FRigHierarchyContainer* InContainer, ERigElementType InElementType, const FName& InName, bool bSelected)
+{
+	OnElementSelected.Broadcast(InContainer, InElementType, InName, bSelected);
+	OnElementChanged.Broadcast(InContainer, InElementType, InName);
+}
+
 #endif
 
 bool FRigHierarchyContainer::IsParentedTo(ERigElementType InChildType, int32 InChildIndex, ERigElementType InParentType, int32 InParentIndex) const
@@ -414,3 +439,79 @@ bool FRigHierarchyContainer::IsParentedTo(ERigElementType InChildType, int32 InC
 
 	return false;
 }
+
+#if WITH_EDITOR
+
+bool FRigHierarchyContainer::Select(const FName& InName, ERigElementType InElementType, bool bSelect)
+{
+	switch(InElementType)
+	{
+		case ERigElementType::Bone:
+		{
+			return BoneHierarchy.Select(InName, bSelect);
+		}
+		case ERigElementType::Space:
+		{
+			return SpaceHierarchy.Select(InName, bSelect);
+		}
+		case ERigElementType::Control:
+		{
+			return ControlHierarchy.Select(InName, bSelect);
+		}
+		case ERigElementType::Curve:
+		{
+			return CurveContainer.Select(InName, bSelect);
+		}
+	}
+	return false;
+}
+
+bool FRigHierarchyContainer::ClearSelection(ERigElementType InElementType)
+{
+	switch(InElementType)
+	{
+		case ERigElementType::Bone:
+		{
+			return BoneHierarchy.ClearSelection();
+		}
+		case ERigElementType::Space:
+		{
+			return SpaceHierarchy.ClearSelection();
+		}
+		case ERigElementType::Control:
+		{
+			return ControlHierarchy.ClearSelection();
+		}
+		case ERigElementType::Curve:
+		{
+			return CurveContainer.ClearSelection();
+		}
+	}
+	return false;
+}
+
+bool FRigHierarchyContainer::IsSelected(const FName& InName, ERigElementType InElementType) const
+{
+	switch(InElementType)
+	{
+		case ERigElementType::Bone:
+		{
+			return BoneHierarchy.IsSelected(InName);
+		}
+		case ERigElementType::Space:
+		{
+			return SpaceHierarchy.IsSelected(InName);
+		}
+		case ERigElementType::Control:
+		{
+			return ControlHierarchy.IsSelected(InName);
+		}
+		case ERigElementType::Curve:
+		{
+			return CurveContainer.IsSelected(InName);
+		}
+	}
+	return false;
+}
+
+#endif
\ No newline at end of file
diff --git a/Engine/Plugins/Experimental/ControlRig/Source/ControlRig/Private/Rigs/RigSpaceHierarchy.cpp b/Engine/Plugins/Experimental/ControlRig/Source/ControlRig/Private/Rigs/RigSpaceHierarchy.cpp
--- a/Engine/Plugins/Experimental/ControlRig/Source/ControlRig/Private/Rigs/RigSpaceHierarchy.cpp
+++ b/Engine/Plugins/Experimental/ControlRig/Source/ControlRig/Private/Rigs/RigSpaceHierarchy.cpp
@@ -141,12 +141,35 @@ bool FRigSpaceHierarchy::Reparent(const FName& InName, ERigSpaceType InSpaceType
 FRigSpace FRigSpaceHierarchy::Remove(const FName& InNameToRemove)
 {
 	int32 IndexToDelete = GetIndex(InNameToRemove);
+#if WITH_EDITOR
+	Select(InNameToRemove, false);
+#endif
 	FRigSpace RemovedSpace = Spaces[IndexToDelete];
 	Spaces.RemoveAt(IndexToDelete);
 
+#if WITH_EDITOR
+	TArray<FName> SpacesReparented;
+#endif
+	for (FRigSpace& ChildSpace : Spaces)
+	{
+		if (ChildSpace.SpaceType == ERigSpaceType::Space && ChildSpace.ParentName == InNameToRemove)
+		{
+			ChildSpace.SpaceType = ERigSpaceType::Global;
+			ChildSpace.ParentIndex = INDEX_NONE;
+			ChildSpace.ParentName = NAME_None;
+#if WITH_EDITOR
+			SpacesReparented.Add(ChildSpace.Name);
+#endif
+		}
+	}
+
 	RefreshMapping();
 
 #if WITH_EDITOR
+	for (const FName& ReparentedSpace : SpacesReparented)
+	{
+		OnSpaceReparented.Broadcast(Container, RigElementType(), ReparentedSpace, InNameToRemove, NAME_None);
+	}
 	OnSpaceRemoved.Broadcast(Container, RigElementType(), RemovedSpace.Name);
 #endif
 
@@ -329,12 +352,21 @@ FName FRigSpaceHierarchy::Rename(const FName& InOldName, const FName& InNewName)
 		if (Found != INDEX_NONE)
 		{
 			FName NewName = GetSafeNewName(InNewName);
+
+#if WITH_EDITOR
+			bool bWasSelected = IsSelected(InOldName);
+			if(bWasSelected)
+			{
+				Select(InOldName, false);
+			}
+#endif
+
 			Spaces[Found].Name = NewName;
 
 			// go through find all children and rename them
 			for (int32 Index = 0; Index < Spaces.Num(); ++Index)
 			{
-				if (Spaces[Index].ParentName == InOldName)
+				if (Spaces[Index].SpaceType == ERigSpaceType::Space && Spaces[Index].ParentName == InOldName)
 				{
 					Spaces[Index].ParentName = NewName;
 				}
@@ -344,6 +376,10 @@ FName FRigSpaceHierarchy::Rename(const FName& InOldName, const FName& InNewName)
 
 #if WITH_EDITOR
 			OnSpaceRenamed.Broadcast(Container, RigElementType(), InOldName, NewName);
+			if(bWasSelected)
+			{
+				Select(NewName, true);
+			}
 #endif
 			return NewName;
 		}
@@ -414,3 +450,161 @@ int32 FRigSpaceHierarchy::GetParentIndex(ERigSpaceType InSpaceType, const FName&
 
 	return INDEX_NONE;
 }
+
+#if WITH_EDITOR
+
+bool FRigSpaceHierarchy::Select(const FName& InName, bool bSelect)
+{
+	if(GetIndex(InName) == INDEX_NONE)
+	{
+		return false;
+	}
+
+	if(bSelect == IsSelected(InName))
+	{
+		return false;
+	}
+
+	if(bSelect)
+	{
+		if (Container)
+		{
+			Container->BoneHierarchy.ClearSelection();
+			Container->ControlHierarchy.ClearSelection();
+			Container->CurveContainer.ClearSelection();
+		}
+
+		Selection.Add(InName);
+	}
+	else
+	{
+		Selection.Remove(InName);
+	}
+
+	OnSpaceSelected.Broadcast(Container, RigElementType(), InName, bSelect);
+
+	return true;
+}
+
+bool FRigSpaceHierarchy::ClearSelection()
+{
+	TArray<FName> TempSelection;
+	TempSelection.Append(Selection);
+	for(const FName& SelectedName : TempSelection)
+	{
+		Select(SelectedName, false);
+	}
+	return TempSelection.Num() > 0;
+}
+
+TArray<FName> FRigSpaceHierarchy::CurrentSelection() const
+{
+	TArray<FName> TempSelection;
+	TempSelection.Append(Selection);
+	return TempSelection;
+}
+
+bool FRigSpaceHierarchy::IsSelected(const FName& InName) const
+{
+	return Selection.Contains(InName);
+}
+
+void FRigSpaceHierarchy::HandleOnElementRemoved(FRigHierarchyContainer* InContainer, ERigElementType InElementType, const FName& InName)
+{
+	if (Container == nullptr)
+	{
+		return;
+	}
+
+	switch (InElementType)
+	{
+		case ERigElementType::Bone:
+		{
+			for (FRigSpace& Space : Spaces)
+			{
+				if (Space.SpaceType == ERigSpaceType::Bone && Space.ParentName == InName)
+				{
+					Space.ParentIndex = Container->BoneHierarchy.GetIndex(InName);
+					Space.ParentName = Space.ParentIndex == INDEX_NONE ? NAME_None : InName;
+					Space.SpaceType = Space.ParentIndex == INDEX_NONE ? ERigSpaceType::Global : Space.SpaceType;
+#if WITH_EDITOR
+					OnSpaceReparented.Broadcast(Container, RigElementType(), Space.Name, InName, Space.ParentName);
+#endif
+				}
+			}
+			break;
+		}
+		case ERigElementType::Control:
+		{
+			for (FRigSpace& Space : Spaces)
+			{
+				if (Space.SpaceType == ERigSpaceType::Control && Space.ParentName == InName)
+				{
+					Space.ParentIndex = Container->ControlHierarchy.GetIndex(InName);
+					Space.ParentName = Space.ParentIndex == INDEX_NONE ? NAME_None : InName;
+					Space.SpaceType = Space.ParentIndex == INDEX_NONE ? ERigSpaceType::Global : Space.SpaceType;
+#if WITH_EDITOR
+					OnSpaceReparented.Broadcast(Container, RigElementType(), Space.Name, InName, Space.ParentName);
+#endif
+				}
+			}
+			break;
+		}
+		case ERigElementType::Space:
+		case ERigElementType::Curve:
+		{
+			break;
+		}
+	}
+}
+
+void FRigSpaceHierarchy::HandleOnElementRenamed(FRigHierarchyContainer* InContainer, ERigElementType InElementType, const FName& InOldName, const FName& InNewName)
+{
+	if (Container == nullptr)
+	{
+		return;
+	}
+
+	switch (InElementType)
+	{
+		case ERigElementType::Bone:
+		{
+			for (FRigSpace& Space : Spaces)
+			{
+				if (Space.SpaceType == ERigSpaceType::Bone && Space.ParentName == InOldName)
+				{
+					Space.ParentIndex = Container->BoneHierarchy.GetIndex(InNewName);
+					Space.ParentName = Space.ParentIndex == INDEX_NONE ? NAME_None : InNewName;
+					Space.SpaceType = Space.ParentIndex == INDEX_NONE ? ERigSpaceType::Global : Space.SpaceType;
+#if WITH_EDITOR
+					OnSpaceReparented.Broadcast(Container, RigElementType(), Space.Name, InOldName, Space.ParentName);
+#endif
+				}
+			}
+			break;
+		}
+		case ERigElementType::Control:
+		{
+			for (FRigSpace& Space : Spaces)
+			{
+				if (Space.SpaceType == ERigSpaceType::Control && Space.ParentName == InOldName)
+				{
+					Space.ParentIndex = Container->ControlHierarchy.GetIndex(InNewName);
+					Space.ParentName = Space.ParentIndex == INDEX_NONE ? NAME_None : InNewName;
+					Space.SpaceType = Space.ParentIndex == INDEX_NONE ? ERigSpaceType::Global : Space.SpaceType;
+#if WITH_EDITOR
+					OnSpaceReparented.Broadcast(Container, RigElementType(), Space.Name, InOldName, Space.ParentName);
+#endif
+				}
+			}
+			break;
+		}
+		case ERigElementType::Space:
+		case ERigElementType::Curve:
+		{
+			break;
+		}
+	}
+}
+
+#endif
\ No newline at end of file
diff --git a/Engine/Plugins/Experimental/ControlRig/Source/ControlRig/Public/Rigs/RigBoneHierarchy.h b/Engine/Plugins/Experimental/ControlRig/Source/ControlRig/Public/Rigs/RigBoneHierarchy.h
--- a/Engine/Plugins/Experimental/ControlRig/Source/ControlRig/Public/Rigs/RigBoneHierarchy.h
+++ b/Engine/Plugins/Experimental/ControlRig/Source/ControlRig/Public/Rigs/RigBoneHierarchy.h
@@ -28,7 +28,7 @@ struct CONTROLRIG_API FRigBone
 	UPROPERTY(VisibleAnywhere, Category = FRigBoneHierarchy)
 	FName Name;
 
-	UPROPERTY(VisibleAnywhere, Category = FRigCurveContainer)
+	UPROPERTY(VisibleAnywhere, Category = FRigBoneHierarchy)
 	int32 Index;
 
 	UPROPERTY(VisibleAnywhere, Category = FRigBoneHierarchy)
@@ -143,19 +143,28 @@ struct CONTROLRIG_API FRigBoneHierarchy
 	// resets all of the transforms back to the initial transform
 	void ResetTransforms();
 
-private:
+#if WITH_EDITOR
 
-	// disable copy constructor
-	FRigBoneHierarchy(const FRigBoneHierarchy& InOther) {}
+	bool Select(const FName& InName, bool bSelect = true);
+	bool ClearSelection();
+	TArray<FName> CurrentSelection() const;
+	bool IsSelected(const FName& InName) const;
 
-	FRigHierarchyContainer* Container;
-#if WITH_EDITOR
 	FRigElementAdded OnBoneAdded;
 	FRigElementRemoved OnBoneRemoved;
 	FRigElementRenamed OnBoneRenamed;
 	FRigElementReparented OnBoneReparented;
+	FRigElementSelected OnBoneSelected;
+
 #endif
 
+private:
+
+	// disable copy constructor
+	FRigBoneHierarchy(const FRigBoneHierarchy& InOther) {}
+
+	FRigHierarchyContainer* Container;
+
 	UPROPERTY(EditAnywhere, Category = FRigBoneHierarchy)
 	TArray<FRigBone> Bones;
 
@@ -163,6 +172,11 @@ private:
 	UPROPERTY()
 	TMap<FName, int32> NameToIndexMapping;
 
+#if WITH_EDITOR
+	UPROPERTY(transient)
+	TArray<FName> Selection;
+#endif
+
 	int32 GetIndexSlow(const FName& InName) const;
 
 	void RecalculateLocalTransform(FRigBone& InOutBone);
diff --git a/Engine/Plugins/Experimental/ControlRig/Source/ControlRig/Public/Rigs/RigControlHierarchy.h b/Engine/Plugins/Experimental/ControlRig/Source/ControlRig/Public/Rigs/RigControlHierarchy.h
--- a/Engine/Plugins/Experimental/ControlRig/Source/ControlRig/Public/Rigs/RigControlHierarchy.h
+++ b/Engine/Plugins/Experimental/ControlRig/Source/ControlRig/Public/Rigs/RigControlHierarchy.h
@@ -43,10 +43,10 @@ struct CONTROLRIG_API FRigControl
 	UPROPERTY(VisibleAnywhere, Category = FRigControlHierarchy)
 	FName Name;
 
-	UPROPERTY(VisibleAnywhere, Category = FRigCurveContainer)
+	UPROPERTY(VisibleAnywhere, Category = FRigControlHierarchy)
 	int32 Index;
 
-	UPROPERTY(VisibleAnywhere, Category = FRigCurveContainer)
+	UPROPERTY(VisibleAnywhere, Category = FRigControlHierarchy)
 	ERigControlType ControlType;
 
 	UPROPERTY(VisibleAnywhere, Category = FRigControlHierarchy)
@@ -166,19 +166,31 @@ struct CONTROLRIG_API FRigControlHierarchy
 	// resets all of the transforms back to the initial transform
 	void ResetTransforms();
 
-private:
+#if WITH_EDITOR
 
-	// disable copy constructor
-	FRigControlHierarchy(const FRigControlHierarchy& InOther) {}
+	bool Select(const FName& InName, bool bSelect = true);
+	bool ClearSelection();
+	TArray<FName> CurrentSelection() const;
+	bool IsSelected(const FName& InName) const;
 
-	FRigHierarchyContainer* Container;
-#if WITH_EDITOR
 	FRigElementAdded OnControlAdded;
 	FRigElementRemoved OnControlRemoved;
 	FRigElementRenamed OnControlRenamed;
 	FRigElementReparented OnControlReparented;
+	FRigElementSelected OnControlSelected;
+
+	void HandleOnElementRemoved(FRigHierarchyContainer* InContainer, ERigElementType InElementType, const FName& InName);
+	void HandleOnElementRenamed(FRigHierarchyContainer* InContainer, ERigElementType InElementType, const FName& InOldName, const FName& InNewName);
+
 #endif
 
+private:
+
+	// disable copy constructor
+	FRigControlHierarchy(const FRigControlHierarchy& InOther) {}
+
+	FRigHierarchyContainer* Container;
+
 	UPROPERTY(EditAnywhere, Category = FRigControlHierarchy)
 	TArray<FRigControl> Controls;
 
@@ -186,6 +198,11 @@ private:
 	UPROPERTY()
 	TMap<FName, int32> NameToIndexMapping;
 
+#if WITH_EDITOR
+	UPROPERTY(transient)
+	TArray<FName> Selection;
+#endif
+
 	int32 GetSpaceIndex(const FName& InName) const;
 
 	int32 GetIndexSlow(const FName& InName) const;
diff --git a/Engine/Plugins/Experimental/ControlRig/Source/ControlRig/Public/Rigs/RigCurveContainer.h b/Engine/Plugins/Experimental/ControlRig/Source/ControlRig/Public/Rigs/RigCurveContainer.h
--- a/Engine/Plugins/Experimental/ControlRig/Source/ControlRig/Public/Rigs/RigCurveContainer.h
+++ b/Engine/Plugins/Experimental/ControlRig/Source/ControlRig/Public/Rigs/RigCurveContainer.h
@@ -100,10 +100,19 @@ public:
 	void ResetValues();
 
 	FRigHierarchyContainer* Container;
+
 #if WITH_EDITOR
+
+	bool Select(const FName& InName, bool bSelect = true);
+	bool ClearSelection();
+	TArray<FName> CurrentSelection() const;
+	bool IsSelected(const FName& InName) const;
+
 	FRigElementAdded OnCurveAdded;
 	FRigElementRemoved OnCurveRemoved;
 	FRigElementRenamed OnCurveRenamed;
+	FRigElementSelected OnCurveSelected;
+
 #endif
 
 private:
@@ -117,14 +126,15 @@ private:
 	UPROPERTY()
 	TMap<FName, int32> NameToIndexMapping;
 
+#if WITH_EDITOR
+	UPROPERTY(transient)
+	TArray<FName> Selection;
+#endif
+
 	int32 GetIndexSlow(const FName& InName) const;
 
 	void RefreshMapping();
+
+	friend struct FRigHierarchyContainer;
 };
 
-// todo
-//USTRUCT()
-//struct CONTROLRIG_API FRigHierarchyRef
-//{
-//	GENERATED_BODY()
-//};
\ No newline at end of file
diff --git a/Engine/Plugins/Experimental/ControlRig/Source/ControlRig/Public/Rigs/RigHierarchyContainer.h b/Engine/Plugins/Experimental/ControlRig/Source/ControlRig/Public/Rigs/RigHierarchyContainer.h
--- a/Engine/Plugins/Experimental/ControlRig/Source/ControlRig/Public/Rigs/RigHierarchyContainer.h
+++ b/Engine/Plugins/Experimental/ControlRig/Source/ControlRig/Public/Rigs/RigHierarchyContainer.h
@@ -103,11 +103,18 @@ public:
 	bool IsParentedTo(ERigElementType InChildType, int32 InChildIndex, ERigElementType InParentType, int32 InParentIndex) const;
 
 #if WITH_EDITOR
+
+	bool Select(const FName& InName, ERigElementType InElementType, bool bSelect = true);
+	bool ClearSelection(ERigElementType InElementType);
+	bool IsSelected(const FName& InName, ERigElementType InElementType) const;
+
 	FRigElementChanged OnElementChanged;
 	FRigElementAdded OnElementAdded;
 	FRigElementRemoved OnElementRemoved;
 	FRigElementRenamed OnElementRenamed;
 	FRigElementReparented OnElementReparented;
+	FRigElementSelected OnElementSelected;
+
 #endif
 
 protected:
@@ -120,6 +127,7 @@ protected:
 	void HandleOnElementRemoved(FRigHierarchyContainer* InContainer, ERigElementType InElementType, const FName& InName);
 	void HandleOnElementRenamed(FRigHierarchyContainer* InContainer, ERigElementType InElementType, const FName& InOldName, const FName& InNewName);
 	void HandleOnElementReparented(FRigHierarchyContainer* InContainer, ERigElementType InElementType, const FName& InName, const FName& InOldParentName, const FName& InNewParentName);
+	void HandleOnElementSelected(FRigHierarchyContainer* InContainer, ERigElementType InElementType, const FName& InName, bool bSelected);
 #endif
 };
 
diff --git a/Engine/Plugins/Experimental/ControlRig/Source/ControlRig/Public/Rigs/RigHierarchyDefines.h b/Engine/Plugins/Experimental/ControlRig/Source/ControlRig/Public/Rigs/RigHierarchyDefines.h
--- a/Engine/Plugins/Experimental/ControlRig/Source/ControlRig/Public/Rigs/RigHierarchyDefines.h
+++ b/Engine/Plugins/Experimental/ControlRig/Source/ControlRig/Public/Rigs/RigHierarchyDefines.h
@@ -21,3 +21,4 @@ DECLARE_MULTICAST_DELEGATE_ThreeParams(FRigElementAdded, FRigHierarchyContainer*
 DECLARE_MULTICAST_DELEGATE_ThreeParams(FRigElementRemoved, FRigHierarchyContainer*, ERigElementType, const FName&);
 DECLARE_MULTICAST_DELEGATE_FourParams(FRigElementRenamed, FRigHierarchyContainer*, ERigElementType, const FName&, const FName&);
 DECLARE_MULTICAST_DELEGATE_FiveParams(FRigElementReparented, FRigHierarchyContainer*, ERigElementType, const FName&, const FName&, const FName&);
+DECLARE_MULTICAST_DELEGATE_FourParams(FRigElementSelected, FRigHierarchyContainer*, ERigElementType, const FName&, bool);
diff --git a/Engine/Plugins/Experimental/ControlRig/Source/ControlRig/Public/Rigs/RigSpaceHierarchy.h b/Engine/Plugins/Experimental/ControlRig/Source/ControlRig/Public/Rigs/RigSpaceHierarchy.h
--- a/Engine/Plugins/Experimental/ControlRig/Source/ControlRig/Public/Rigs/RigSpaceHierarchy.h
+++ b/Engine/Plugins/Experimental/ControlRig/Source/ControlRig/Public/Rigs/RigSpaceHierarchy.h
@@ -44,7 +44,7 @@ struct CONTROLRIG_API FRigSpace
 	UPROPERTY(VisibleAnywhere, Category = FRigSpaceHierarchy)
 	FName Name;
 
-	UPROPERTY(VisibleAnywhere, Category = FRigCurveContainer)
+	UPROPERTY(VisibleAnywhere, Category = FRigSpaceHierarchy)
 	int32 Index;
 
 	UPROPERTY(VisibleAnywhere, Category = FRigSpaceHierarchy)
@@ -150,19 +150,31 @@ struct CONTROLRIG_API FRigSpaceHierarchy
 	// resets all of the transforms back to the initial transform
 	void ResetTransforms();
 
-private:
+#if WITH_EDITOR
 
-	// disable copy constructor
-	FRigSpaceHierarchy(const FRigSpaceHierarchy& InOther) {}
+	bool Select(const FName& InName, bool bSelect = true);
+	bool ClearSelection();
+	TArray<FName> CurrentSelection() const;
+	bool IsSelected(const FName& InName) const;
 
-	FRigHierarchyContainer* Container;
-#if WITH_EDITOR
 	FRigElementAdded OnSpaceAdded;
 	FRigElementRemoved OnSpaceRemoved;
 	FRigElementRenamed OnSpaceRenamed;
 	FRigElementReparented OnSpaceReparented;
+	FRigElementSelected OnSpaceSelected;
+
+	void HandleOnElementRemoved(FRigHierarchyContainer* InContainer, ERigElementType InElementType, const FName& InName);
+	void HandleOnElementRenamed(FRigHierarchyContainer* InContainer, ERigElementType InElementType, const FName& InOldName, const FName& InNewName);
+
 #endif
 
+private:
+
+	// disable copy constructor
+	FRigSpaceHierarchy(const FRigSpaceHierarchy& InOther) {}
+
+	FRigHierarchyContainer* Container;
+
 	int32 GetParentIndex(ERigSpaceType InSpaceType, const FName& InName) const;
 
 	UPROPERTY(EditAnywhere, Category = FRigSpaceHierarchy)
@@ -172,6 +184,11 @@ private:
 	UPROPERTY()
 	TMap<FName, int32> NameToIndexMapping;
 
+#if WITH_EDITOR
+	UPROPERTY(transient)
+	TArray<FName> Selection;
+#endif
+
 	int32 GetIndexSlow(const FName& InName) const;
 
 	void RefreshMapping();
diff --git a/Engine/Plugins/Experimental/ControlRig/Source/ControlRigEditor/Private/EditMode/ControlRigEditMode.cpp b/Engine/Plugins/Experimental/ControlRig/Source/ControlRigEditor/Private/EditMode/ControlRigEditMode.cpp
--- a/Engine/Plugins/Experimental/ControlRig/Source/ControlRigEditor/Private/EditMode/ControlRigEditMode.cpp
+++ b/Engine/Plugins/Experimental/ControlRig/Source/ControlRigEditor/Private/EditMode/ControlRigEditMode.cpp
@@ -1523,14 +1523,20 @@ bool FControlRigEditMode::AreBoneSelected() const
 	return (SelectedBones.Num() > 0);
 }
 
-void FControlRigEditMode::SelectBone(const FName& InBone)
+void FControlRigEditMode::SelectBone(const FName& InBone, bool bSelect)
 {
 	ClearControlSelection();
 
-	SelectedBones.Reset();
 	if (InBone != NAME_None)
 	{
-		SelectedBones.Add(InBone);
+		if (bSelect)
+		{
+			SelectedBones.Add(InBone);
+		}
+		else
+		{
+			SelectedBones.Remove(InBone);
+		}
 	}
 }
 
@@ -1565,6 +1571,34 @@ void FControlRigEditMode::HandleModelModified(const UControlRigModel* InModel, E
 	_ModifiedEvent.Broadcast(InModel, InType, InPayload);
 }
 
+void FControlRigEditMode::OnRigElementSelected(FRigHierarchyContainer* Container, ERigElementType ElementType, const FName& InName, bool bSelected)
+{
+	if (UControlRig* ControlRig = WeakControlRig.Get())
+	{
+		if (UControlRigBlueprintGeneratedClass* Class = Cast<UControlRigBlueprintGeneratedClass>(ControlRig->GetClass()))
+		{
+			if (UControlRigBlueprint* RigBlueprint = Cast<UControlRigBlueprint>(Class->ClassGeneratedBy))
+			{
+				switch (ElementType)
+				{
+					case ERigElementType::Bone:
+					{
+						//SelectBone(InName, bSelected);
+						break;
+					}
+					case ERigElementType::Space:
+					case ERigElementType::Control:
+					case ERigElementType::Curve:
+					{
+						// todo
+						break;
+					}
+				}
+			}
+		}
+	}
+}
+
 UControlRigModel::FModifiedEvent& FControlRigEditMode::OnModified()
 {
 	return _ModifiedEvent;
diff --git a/Engine/Plugins/Experimental/ControlRig/Source/ControlRigEditor/Private/EditMode/ControlRigEditMode.h b/Engine/Plugins/Experimental/ControlRig/Source/ControlRigEditor/Private/EditMode/ControlRigEditMode.h
--- a/Engine/Plugins/Experimental/ControlRig/Source/ControlRigEditor/Private/EditMode/ControlRigEditMode.h
+++ b/Engine/Plugins/Experimental/ControlRig/Source/ControlRigEditor/Private/EditMode/ControlRigEditMode.h
@@ -8,6 +8,7 @@
 #include "ControlRigTrajectoryCache.h"
 #include "ControlUnitProxy.h"
 #include "ControlRigModel.h"
+#include "Rigs/RigHierarchyContainer.h"
 
 class FEditorViewportClient;
 class FViewport;
@@ -137,13 +138,15 @@ public:
 	static FRigUnit_Control* GetRigUnit(const FControlUnitProxy& InProxy, UControlRig* InControlRig, UScriptStruct** OutControlStructPtr = nullptr);
 
 	/** Select Bone */
-	void SelectBone(const FName& InBone);
+	void SelectBone(const FName& InBone, bool bSelect = true);
 	FOnGetBoneTransform& OnGetBoneTransform() { return OnGetBoneTransformDelegate; }
 	FOnSetBoneTransform& OnSetBoneTransform() { return OnSetBoneTransformDelegate; }
 
 	UControlRigModel::FModifiedEvent& OnModified();
 	void HandleModelModified(const UControlRigModel* InModel, EControlRigModelNotifType InType, const void* InPayload);
 
+	void OnRigElementSelected(FRigHierarchyContainer* Container, ERigElementType ElementType, const FName& InName, bool bSelected);
+
 protected:
 	/** Helper function: set ControlRigs array to the details panel */
 	void SetObjects_Internal();
diff --git a/Engine/Plugins/Experimental/ControlRig/Source/ControlRigEditor/Private/Editor/ControlRigEditor.cpp b/Engine/Plugins/Experimental/ControlRig/Source/ControlRigEditor/Private/Editor/ControlRigEditor.cpp
--- a/Engine/Plugins/Experimental/ControlRig/Source/ControlRigEditor/Private/Editor/ControlRigEditor.cpp
+++ b/Engine/Plugins/Experimental/ControlRig/Source/ControlRigEditor/Private/Editor/ControlRigEditor.cpp
@@ -94,6 +94,8 @@ FControlRigEditor::~FControlRigEditor()
 		RigBlueprint->HierarchyContainer.OnElementRemoved.RemoveAll(this);
 		RigBlueprint->HierarchyContainer.OnElementRenamed.RemoveAll(this);
 		RigBlueprint->HierarchyContainer.OnElementReparented.RemoveAll(this);
+		RigBlueprint->HierarchyContainer.OnElementSelected.RemoveAll(this);
+		RigBlueprint->HierarchyContainer.OnElementSelected.RemoveAll(&GetEditMode());
 	}
 }
 
@@ -185,11 +187,6 @@ void FControlRigEditor::InitControlRigEditor(const EToolkitMode::Type Mode, cons
 	}
 
 	InControlRigBlueprint->OnModified().AddSP(this, &FControlRigEditor::HandleModelModified);
-	InControlRigBlueprint->HierarchyContainer.OnElementAdded.AddSP(this, &FControlRigEditor::OnRigElementAdded);
-	InControlRigBlueprint->HierarchyContainer.OnElementRemoved.AddSP(this, &FControlRigEditor::OnRigElementRemoved);
-	InControlRigBlueprint->HierarchyContainer.OnElementRenamed.AddSP(this, &FControlRigEditor::OnRigElementRenamed);
-	InControlRigBlueprint->HierarchyContainer.OnElementReparented.AddSP(this, &FControlRigEditor::OnRigElementReparented);
-
 	BindCommands();
 
 	AddApplicationMode(
@@ -209,6 +206,7 @@ void FControlRigEditor::InitControlRigEditor(const EToolkitMode::Type Mode, cons
 	GetAssetEditorModeManager()->ActivateMode(FControlRigEditorEditMode::ModeName);
 	GetEditMode().OnGetBoneTransform() = FOnGetBoneTransform::CreateSP(this, &FControlRigEditor::GetBoneTransform);
 	GetEditMode().OnSetBoneTransform() = FOnSetBoneTransform::CreateSP(this, &FControlRigEditor::SetBoneTransform);
+
 	InControlRigBlueprint->OnModified().AddSP(&GetEditMode(), &FControlRigEditMode::HandleModelModified);
 
 	UpdateControlRig();
@@ -258,6 +256,13 @@ void FControlRigEditor::InitControlRigEditor(const EToolkitMode::Type Mode, cons
 				}
 			}
 		}
+
+		InControlRigBlueprint->HierarchyContainer.OnElementAdded.AddSP(this, &FControlRigEditor::OnRigElementAdded);
+		InControlRigBlueprint->HierarchyContainer.OnElementRemoved.AddSP(this, &FControlRigEditor::OnRigElementRemoved);
+		InControlRigBlueprint->HierarchyContainer.OnElementRenamed.AddSP(this, &FControlRigEditor::OnRigElementRenamed);
+		InControlRigBlueprint->HierarchyContainer.OnElementReparented.AddSP(this, &FControlRigEditor::OnRigElementReparented);
+		InControlRigBlueprint->HierarchyContainer.OnElementSelected.AddSP(this, &FControlRigEditor::OnRigElementSelected);
+		InControlRigBlueprint->HierarchyContainer.OnElementSelected.AddSP(&GetEditMode(), &FControlRigEditMode::OnRigElementSelected);
 	}
 
 	bControlRigEditorInitialized = true;
@@ -525,11 +530,13 @@ UBlueprint* FControlRigEditor::GetBlueprintObj() const
 
 void FControlRigEditor::SetDetailObjects(const TArray<UObject*>& InObjects)
 {
+	RigElementNameInDetailPanel = NAME_None;
 	Inspector->ShowDetailsForObjects(InObjects);
 }
 
 void FControlRigEditor::SetDetailObject(UObject* Obj)
 {
+	RigElementNameInDetailPanel = NAME_None;
 	TArray<UObject*> Objects;
 	if (Obj)
 	{
@@ -538,13 +545,15 @@ void FControlRigEditor::SetDetailObject(UObject* Obj)
 	SetDetailObjects(Objects);
 }
 
-void FControlRigEditor::SetDetailStruct(TSharedPtr<FStructOnScope> StructToDisplay)
+void FControlRigEditor::SetDetailStruct(const FName& InName, TSharedPtr<FStructOnScope> StructToDisplay)
 {
+	RigElementNameInDetailPanel = InName;
 	Inspector->ShowSingleStruct(StructToDisplay);
 }
 
 void FControlRigEditor::ClearDetailObject()
 {
+	RigElementNameInDetailPanel = NAME_None;
 	Inspector->ShowDetailsForObjects(TArray<UObject*>());
 	Inspector->ShowSingleStruct(TSharedPtr<FStructOnScope>());
 }
@@ -1450,29 +1459,6 @@ void FControlRigEditor::OnNodeTitleCommitted(const FText& NewText, ETextCommit::
 	}
 }
 
-
-void FControlRigEditor::SelectBone(const FName& InBone)
-{
-	DECLARE_SCOPE_HIERARCHICAL_COUNTER_FUNC()
-
-	// edit mode has to know
-	GetEditMode().SelectBone(InBone);
-	// copy locally, we use this for copying back to template when modified
-
-	SelectedBone = InBone;
-	UControlRigSkeletalMeshComponent* EditorSkelComp = Cast<UControlRigSkeletalMeshComponent>(GetPersonaToolkit()->GetPreviewScene()->GetPreviewMeshComponent());
-	if (EditorSkelComp)
-	{
-		EditorSkelComp->BonesOfInterest.Reset();
-
-		int32 Index = ControlRig->Hierarchy.BoneHierarchy.GetIndex(InBone);
-		if (Index != INDEX_NONE)
-		{
-			EditorSkelComp->BonesOfInterest.Add(Index);
-		}
-	}
-}
-
 FTransform FControlRigEditor::GetBoneTransform(const FName& InBone, bool bLocal) const
 {
 	DECLARE_SCOPE_HIERARCHICAL_COUNTER_FUNC()
@@ -1531,24 +1517,44 @@ void FControlRigEditor::SetBoneTransform(const FName& InBone, const FTransform&
 void FControlRigEditor::OnFinishedChangingProperties(const FPropertyChangedEvent& PropertyChangedEvent)
 {
 	DECLARE_SCOPE_HIERARCHICAL_COUNTER_FUNC()
-	//	UE_LOG(LogControlRigEditor, Warning, TEXT("Current Property being modified : %s"), *GetNameSafe(PropertyChangedEvent.Property));
 
-	FName PropertyName = (PropertyChangedEvent.Property != NULL) ? PropertyChangedEvent.Property->GetFName() : NAME_None;
-	if (PropertyName == GET_MEMBER_NAME_CHECKED(FRigBone, InitialTransform))
+	UControlRigBlueprint* ControlRigBP = GetControlRigBlueprint();
+	if (ControlRig && ControlRigBP && RigElementNameInDetailPanel != NAME_None)
 	{
-		// if init transform changes, it updates to the base
-		UControlRigBlueprint* ControlRigBP = GetControlRigBlueprint();
-		if (ControlRig && ControlRigBP)
+		UControlRig* DebuggedControlRig = Cast<UControlRig>(GetBlueprintObj()->GetObjectBeingDebugged());
+		UScriptStruct* ScriptStruct = Cast< UScriptStruct>(PropertyChangedEvent.Property->GetOuter());
+		if (ScriptStruct)
 		{
-			if (SelectedBone != NAME_None)
+			if (ScriptStruct == FRigBone::StaticStruct())
+			{
+				ControlRig->Hierarchy.BoneHierarchy[RigElementNameInDetailPanel] = ControlRigBP->HierarchyContainer.BoneHierarchy[RigElementNameInDetailPanel];
+				if (DebuggedControlRig && DebuggedControlRig != ControlRig)
+				{
+					DebuggedControlRig->Hierarchy.BoneHierarchy[RigElementNameInDetailPanel] = ControlRigBP->HierarchyContainer.BoneHierarchy[RigElementNameInDetailPanel];
+				}
+			}
+			else if (ScriptStruct == FRigSpace::StaticStruct())
+			{
+				ControlRig->Hierarchy.SpaceHierarchy[RigElementNameInDetailPanel] = ControlRigBP->HierarchyContainer.SpaceHierarchy[RigElementNameInDetailPanel];
+				if (DebuggedControlRig && DebuggedControlRig != ControlRig)
+				{
+					DebuggedControlRig->Hierarchy.SpaceHierarchy[RigElementNameInDetailPanel] = ControlRigBP->HierarchyContainer.SpaceHierarchy[RigElementNameInDetailPanel];
+				}
+			}
+			else if (ScriptStruct == FRigControl::StaticStruct())
 			{
-				const int32 BoneIndex = ControlRig->Hierarchy.BoneHierarchy.GetIndex(SelectedBone);
-				if (BoneIndex != INDEX_NONE)
+				ControlRig->Hierarchy.ControlHierarchy[RigElementNameInDetailPanel] = ControlRigBP->HierarchyContainer.ControlHierarchy[RigElementNameInDetailPanel];
+				if (DebuggedControlRig && DebuggedControlRig != ControlRig)
 				{
-					FTransform InitialTransform = ControlRig->Hierarchy.BoneHierarchy.GetInitialTransform(BoneIndex);
-					// update CDO  @todo - re-think about how we wrap around this nicer
-					// copy currently selected Bone to base hierarchy			
-					ControlRigBP->HierarchyContainer.BoneHierarchy.SetInitialTransform(BoneIndex, InitialTransform);
+					DebuggedControlRig->Hierarchy.ControlHierarchy[RigElementNameInDetailPanel] = ControlRigBP->HierarchyContainer.ControlHierarchy[RigElementNameInDetailPanel];
+				}
+			}
+			else if (ScriptStruct == FRigCurve::StaticStruct())
+			{
+				ControlRig->Hierarchy.CurveContainer[RigElementNameInDetailPanel] = ControlRigBP->HierarchyContainer.CurveContainer[RigElementNameInDetailPanel];
+				if (DebuggedControlRig && DebuggedControlRig != ControlRig)
+				{
+					DebuggedControlRig->Hierarchy.CurveContainer[RigElementNameInDetailPanel] = ControlRigBP->HierarchyContainer.CurveContainer[RigElementNameInDetailPanel];
 				}
 			}
 		}
@@ -1697,6 +1703,89 @@ void FControlRigEditor::OnRigElementReparented(FRigHierarchyContainer* Container
 	}
 }
 
+void FControlRigEditor::OnRigElementSelected(FRigHierarchyContainer* Container, ERigElementType ElementType, const FName& InName, bool bSelected)
+{
+	UControlRigBlueprint* RigBlueprint = GetControlRigBlueprint();
+	if (RigBlueprint == nullptr)
+	{
+		return;
+	}
+
+	switch (ElementType)
+	{
+		case ERigElementType::Bone:
+		{
+			FRigBoneHierarchy& BoneHierarchy = RigBlueprint->HierarchyContainer.BoneHierarchy;
+			if (bSelected)
+			{
+				SetDetailStruct(InName, MakeShareable(new FStructOnScope(FRigBone::StaticStruct(), (uint8*)&(BoneHierarchy[InName]))));
+			}
+			else
+			{
+				ClearDetailObject();
+			}
+
+			UControlRigSkeletalMeshComponent* EditorSkelComp = Cast<UControlRigSkeletalMeshComponent>(GetPersonaToolkit()->GetPreviewScene()->GetPreviewMeshComponent());
+			if (EditorSkelComp)
+			{
+				int32 Index = ControlRig->Hierarchy.BoneHierarchy.GetIndex(InName);
+				if (Index != INDEX_NONE)
+				{
+					if (bSelected)
+					{
+						EditorSkelComp->BonesOfInterest.Add(Index);
+					}
+					else
+					{
+						EditorSkelComp->BonesOfInterest.Remove(Index);
+					}
+				}
+			}
+
+			break;
+		}
+		case ERigElementType::Space:
+		{
+			if (bSelected)
+			{
+				FRigSpaceHierarchy& SpaceHierarchy = RigBlueprint->HierarchyContainer.SpaceHierarchy;
+				SetDetailStruct(InName, MakeShareable(new FStructOnScope(FRigSpace::StaticStruct(), (uint8*)&(SpaceHierarchy[InName]))));
+			}
+			else
+			{
+				ClearDetailObject();
+			}
+			break;
+		}
+		case ERigElementType::Control:
+		{
+			if (bSelected)
+			{
+				FRigControlHierarchy& ControlHierarchy = RigBlueprint->HierarchyContainer.ControlHierarchy;
+				SetDetailStruct(InName, MakeShareable(new FStructOnScope(FRigControl::StaticStruct(), (uint8*)&(ControlHierarchy[InName]))));
+			}
+			else
+			{
+				ClearDetailObject();
+			}
+			break;
+		}
+		case ERigElementType::Curve:
+		{
+			if (bSelected)
+			{
+				FRigCurveContainer& CurveContainer = RigBlueprint->HierarchyContainer.CurveContainer;
+				SetDetailStruct(InName, MakeShareable(new FStructOnScope(FRigCurve::StaticStruct(), (uint8*)&(CurveContainer[InName]))));
+			}
+			else
+			{
+				ClearDetailObject();
+			}
+			break;
+		}
+	}
+}
+
 void FControlRigEditor::OnCurveContainerChanged()
 {
 	DECLARE_SCOPE_HIERARCHICAL_COUNTER_FUNC()
diff --git a/Engine/Plugins/Experimental/ControlRig/Source/ControlRigEditor/Private/Editor/ControlRigEditor.h b/Engine/Plugins/Experimental/ControlRig/Source/ControlRigEditor/Private/Editor/ControlRigEditor.h
--- a/Engine/Plugins/Experimental/ControlRig/Source/ControlRigEditor/Private/Editor/ControlRigEditor.h
+++ b/Engine/Plugins/Experimental/ControlRig/Source/ControlRigEditor/Private/Editor/ControlRigEditor.h
@@ -91,7 +91,7 @@ public:
 
 	void SetDetailObject(UObject* Obj);
 
-	void SetDetailStruct(TSharedPtr<FStructOnScope> StructToDisplay);
+	void SetDetailStruct(const FName& InName, TSharedPtr<FStructOnScope> StructToDisplay);
 
 	void ClearDetailObject();
 
@@ -104,7 +104,6 @@ public:
 	/** Get the edit mode */
 	FControlRigEditorEditMode& GetEditMode() { return *static_cast<FControlRigEditorEditMode*>(GetAssetEditorModeManager()->GetActiveMode(FControlRigEditorEditMode::ModeName)); }
 
-	void SelectBone(const FName& InBone);
 	// this changes everytime you compile, so don't cache it expecting it will last. 
 	UControlRig* GetInstanceRig() const { return ControlRig;  }
 
@@ -114,6 +113,7 @@ public:
 	void OnRigElementRemoved(FRigHierarchyContainer* Container, ERigElementType ElementType, const FName& InName);
 	void OnRigElementRenamed(FRigHierarchyContainer* Container, ERigElementType ElementType, const FName& InOldName, const FName& InNewName);
 	void OnRigElementReparented(FRigHierarchyContainer* Container, ERigElementType ElementType, const FName& InName, const FName& InOldParentName, const FName& InNewParentName);
+	void OnRigElementSelected(FRigHierarchyContainer* Container, ERigElementType ElementType, const FName& InName, bool bSelected);
 
 	void OnGraphNodeDropToPerform(TSharedPtr<FGraphNodeDragDropOp> DragDropOp, UEdGraph* Graph, const FVector2D& NodePosition, const FVector2D& ScreenPosition);
 
@@ -246,12 +246,10 @@ protected:
 	/** delegate for changing property */
 	void OnFinishedChangingProperties(const FPropertyChangedEvent& PropertyChangedEvent);
 
-	/** Selected Bone from hierarchy tree */
-	FName SelectedBone;
-
 	bool bControlRigEditorInitialized;
 	bool bIsSelecting;
 	bool bIsSettingObjectBeingDebugged;
+	FName RigElementNameInDetailPanel;
 
 	/** The log to use for errors resulting from the init phase of the units */
 	FControlRigLog ControlRigLog;
diff --git a/Engine/Plugins/Experimental/ControlRig/Source/ControlRigEditor/Private/Editor/SRigBoneHierarchy.cpp b/Engine/Plugins/Experimental/ControlRig/Source/ControlRigEditor/Private/Editor/SRigBoneHierarchy.cpp
--- a/Engine/Plugins/Experimental/ControlRig/Source/ControlRigEditor/Private/Editor/SRigBoneHierarchy.cpp
+++ b/Engine/Plugins/Experimental/ControlRig/Source/ControlRigEditor/Private/Editor/SRigBoneHierarchy.cpp
@@ -167,7 +167,11 @@ SRigBoneHierarchy::~SRigBoneHierarchy()
 		ControlRigBlueprint = ControlRigEditor.Pin()->GetControlRigBlueprint();
 		if (ControlRigBlueprint.IsValid())
 		{
-			ControlRigBlueprint->HierarchyContainer.OnElementChanged.RemoveAll(this);
+			ControlRigBlueprint->HierarchyContainer.OnElementAdded.RemoveAll(this);
+			ControlRigBlueprint->HierarchyContainer.OnElementRemoved.RemoveAll(this);
+			ControlRigBlueprint->HierarchyContainer.OnElementRenamed.RemoveAll(this);
+			ControlRigBlueprint->HierarchyContainer.OnElementReparented.RemoveAll(this);
+			ControlRigBlueprint->HierarchyContainer.OnElementSelected.RemoveAll(this);
 		}
 	}
 }
@@ -180,7 +184,11 @@ void SRigBoneHierarchy::Construct(const FArguments& InArgs, TSharedRef<FControlR
 	// @todo: find a better place to do it
 	ControlRigBlueprint->HierarchyContainer.BoneHierarchy.Initialize();
 
-	ControlRigBlueprint->HierarchyContainer.OnElementChanged.AddRaw(this, &SRigBoneHierarchy::OnRigElementChanged);
+	ControlRigBlueprint->HierarchyContainer.OnElementAdded.AddRaw(this, &SRigBoneHierarchy::OnRigElementAdded);
+	ControlRigBlueprint->HierarchyContainer.OnElementRemoved.AddRaw(this, &SRigBoneHierarchy::OnRigElementRemoved);
+	ControlRigBlueprint->HierarchyContainer.OnElementRenamed.AddRaw(this, &SRigBoneHierarchy::OnRigElementRenamed);
+	ControlRigBlueprint->HierarchyContainer.OnElementReparented.AddRaw(this, &SRigBoneHierarchy::OnRigElementReparented);
+	ControlRigBlueprint->HierarchyContainer.OnElementSelected.AddRaw(this, &SRigBoneHierarchy::OnRigElementSelected);
 
 	// for deleting, renaming, dragging
 	CommandList = MakeShared<FUICommandList>();
@@ -236,11 +244,13 @@ void SRigBoneHierarchy::Construct(const FArguments& InArgs, TSharedRef<FControlR
 				.OnGetChildren(this, &SRigBoneHierarchy::HandleGetChildrenForTree)
 				.OnSelectionChanged(this, &SRigBoneHierarchy::OnSelectionChanged)
 				.OnContextMenuOpening(this, &SRigBoneHierarchy::CreateContextMenu)
+				.HighlightParentNodesForSelection(true)
 				.ItemHeight(24)
 			]
 		]
 	];
 
+	bIsChangingRigHierarchy = false;
 	RefreshTreeView();
 }
 
@@ -326,6 +336,15 @@ void SRigBoneHierarchy::RefreshTreeView()
 	}
 
 	TreeView->RequestTreeRefresh();
+
+	if (ControlRigBlueprint.IsValid())
+	{
+		FRigBoneHierarchy& Hierarchy = ControlRigBlueprint->HierarchyContainer.BoneHierarchy;
+		for (const FName& SelectedBone : Hierarchy.CurrentSelection())
+		{
+			OnRigElementSelected(&ControlRigBlueprint->HierarchyContainer, ERigElementType::Bone, SelectedBone, true);
+		}
+	}
 }
 
 void SRigBoneHierarchy::SetExpansionRecursive(TSharedPtr<FRigTreeBone> InBone)
@@ -349,29 +368,39 @@ void SRigBoneHierarchy::HandleGetChildrenForTree(TSharedPtr<FRigTreeBone> InItem
 
 void SRigBoneHierarchy::OnSelectionChanged(TSharedPtr<FRigTreeBone> Selection, ESelectInfo::Type SelectInfo)
 {
-	// need dummy object
-	if (Selection.IsValid())
+	if (bIsChangingRigHierarchy)
+	{
+		return;
+	}
+
+	FRigBoneHierarchy* RigHierarchy = GetHierarchy();
+
+	if (RigHierarchy)
 	{
-		FRigBoneHierarchy* RigHierarchy = GetInstanceHierarchy();
+		TGuardValue<bool> GuardRigHierarchyChanges(bIsChangingRigHierarchy, true);
 
-		if (RigHierarchy)
+		TArray<FName> OldSelection = RigHierarchy->CurrentSelection();
+		TArray<FName> NewSelection;
+
+		TArray<TSharedPtr<FRigTreeBone>> SelectedItems = TreeView->GetSelectedItems();
+		for (const TSharedPtr<FRigTreeBone>& SelectedItem : SelectedItems)
 		{
-			const int32 BoneIndex = RigHierarchy->GetIndex(Selection->CachedBone);
-			if (BoneIndex != INDEX_NONE)
-			{
-				ControlRigEditor.Pin()->SetDetailStruct(MakeShareable(new FStructOnScope(FRigBone::StaticStruct(), (uint8*)&(*RigHierarchy)[BoneIndex])));
-				ControlRigEditor.Pin()->SelectBone(Selection->CachedBone);
-				return;
-			}
-			else
+			NewSelection.Add(SelectedItem->CachedBone);
+		}
+
+		for (const FName& PreviouslySelected : OldSelection)
+		{
+			if (NewSelection.Contains(PreviouslySelected))
 			{
-				// clear the current selection
-				ControlRigEditor.Pin()->ClearDetailObject();
-				ControlRigEditor.Pin()->SelectBone(NAME_None);
+				continue;
 			}
+			RigHierarchy->Select(PreviouslySelected, false);
 		}
 
-		// if failed, try BP hierarhcy? Todo:
+		for (const FName& NewlySelected : NewSelection)
+		{
+			RigHierarchy->Select(NewlySelected, true);
+		}
 	}
 }
 
@@ -394,28 +423,59 @@ TSharedPtr<FRigTreeBone> FindBone(const FName& InBoneName, TSharedPtr<FRigTreeBo
 	return TSharedPtr<FRigTreeBone>();
 }
 
-void SRigBoneHierarchy::SelectBone(const FName& BoneName) const
+void SRigBoneHierarchy::OnRigElementAdded(FRigHierarchyContainer* Container, ERigElementType ElementType, const FName& InName)
 {
-	for (int32 RootIndex = 0; RootIndex < RootBones.Num(); ++RootIndex)
+	if (bIsChangingRigHierarchy || ElementType != ERigElementType::Bone)
 	{
-		TSharedPtr<FRigTreeBone> Found = FindBone(BoneName, RootBones[RootIndex]);
-		if (Found.IsValid())
-		{
-			TreeView->SetSelection(Found);
-		}
+		return;
 	}
+	RefreshTreeView();
 }
 
-void SRigBoneHierarchy::OnRigElementChanged(FRigHierarchyContainer* Container, ERigElementType ElementType, const FName& InName)
+void SRigBoneHierarchy::OnRigElementRemoved(FRigHierarchyContainer* Container, ERigElementType ElementType, const FName& InName)
 {
-	if (ElementType != ERigElementType::Bone)
+	if (bIsChangingRigHierarchy || ElementType != ERigElementType::Bone)
 	{
 		return;
 	}
+	RefreshTreeView();
+}
 
+void SRigBoneHierarchy::OnRigElementRenamed(FRigHierarchyContainer* Container, ERigElementType ElementType, const FName& InOldName, const FName& InNewName)
+{
+	if (bIsChangingRigHierarchy || ElementType != ERigElementType::Bone)
+	{
+		return;
+	}
 	RefreshTreeView();
 }
 
+void SRigBoneHierarchy::OnRigElementReparented(FRigHierarchyContainer* Container, ERigElementType ElementType, const FName& InName, const FName& InOldParentName, const FName& InNewParentName)
+{
+	if (bIsChangingRigHierarchy || ElementType != ERigElementType::Bone)
+	{
+		return;
+	}
+	RefreshTreeView();
+}
+
+void SRigBoneHierarchy::OnRigElementSelected(FRigHierarchyContainer* Container, ERigElementType ElementType, const FName& InName, bool bSelected)
+{
+	if (bIsChangingRigHierarchy || ElementType != ERigElementType::Bone)
+	{
+		return;
+	}
+
+	for (int32 RootIndex = 0; RootIndex < RootBones.Num(); ++RootIndex)
+	{
+		TSharedPtr<FRigTreeBone> Found = FindBone(InName, RootBones[RootIndex]);
+		if (Found.IsValid())
+		{
+			TreeView->SetItemSelection(Found, bSelected);
+		}
+	}
+}
+
 void SRigBoneHierarchy::ClearDetailPanel() const
 {
 	ControlRigEditor.Pin()->ClearDetailObject();
@@ -548,6 +608,9 @@ void SRigBoneHierarchy::ImportHierarchy(const FAssetData& InAssetData)
 		const TArray<FMeshBoneInfo>& BoneInfos = RefSkeleton.GetRefBoneInfo();
 		const TArray<FTransform>& BonePoses = RefSkeleton.GetRefBonePose();
 
+		Hier->ClearSelection();
+
+		TArray<FName> AddedBones;
 		for (int32 BoneIndex = 0; BoneIndex < RefSkeleton.GetNum(); ++BoneIndex)
 		{
 			// only add if you don't have it. This may change in the future
@@ -556,10 +619,17 @@ void SRigBoneHierarchy::ImportHierarchy(const FAssetData& InAssetData)
 				// @todo: add optimized version without sorting, but if no sort, we should make sure not to use find index function
 				FName ParentName = (BoneInfos[BoneIndex].ParentIndex != INDEX_NONE) ? BoneInfos[BoneInfos[BoneIndex].ParentIndex].Name : NAME_None;
 				Hier->Add(BoneInfos[BoneIndex].Name, ParentName, FAnimationRuntime::GetComponentSpaceTransform(RefSkeleton, BonePoses, BoneIndex));
+				AddedBones.Add(BoneInfos[BoneIndex].Name);
 			}
 		}
 
+		for (const FName& AddedBone : AddedBones)
+		{
+			Hier->Select(AddedBone);
+		}
+
 		FSlateApplication::Get().DismissAllMenus();
+		RefreshTreeView();
 	}
 }
 
@@ -589,12 +659,17 @@ void SRigBoneHierarchy::HandleDeleteItem()
 
 		for (int32 ItemIndex = 0; ItemIndex < SelectedItems.Num(); ++ItemIndex)
 		{
+			TGuardValue<bool> GuardRigHierarchyChanges(bIsChangingRigHierarchy, true);
+
 			// when you select whole Bones, you might not have them anymore
 			if (Hierarchy->GetIndex(SelectedItems[ItemIndex]->CachedBone) != INDEX_NONE)
 			{
 				Hierarchy->Remove(SelectedItems[ItemIndex]->CachedBone);
 			}
 		}
+
+		RefreshTreeView();
+		FSlateApplication::Get().DismissAllMenus();
  	}
 }
 
@@ -626,10 +701,15 @@ void SRigBoneHierarchy::HandleNewItem()
 		}
 
 		const FName NewBoneName = CreateUniqueName(TEXT("NewBone"));
-		Hierarchy->Add(NewBoneName, ParentName, ParentTransform);
+		{
+			TGuardValue<bool> GuardRigHierarchyChanges(bIsChangingRigHierarchy, true);
+			Hierarchy->Add(NewBoneName, ParentName, ParentTransform);
+		}
+		Hierarchy->ClearSelection();
+		Hierarchy->Select(NewBoneName);
 
-		// reselect current selected item
-		SelectBone(NewBoneName);
+		FSlateApplication::Get().DismissAllMenus();
+		RefreshTreeView();
 	}
 }
 
@@ -654,6 +734,8 @@ void SRigBoneHierarchy::HandleDuplicateItem()
 		TArray<FName> NewNames;
 		for (int32 Index = 0; Index < SelectedItems.Num(); ++Index)
 		{
+			TGuardValue<bool> GuardRigHierarchyChanges(bIsChangingRigHierarchy, true);
+
 			FName Name = SelectedItems[Index]->CachedBone;
 			FTransform Transform = Hierarchy->GetGlobalTransform(Name);
 
@@ -664,10 +746,14 @@ void SRigBoneHierarchy::HandleDuplicateItem()
 			NewNames.Add(NewName);
 		}
 
+		Hierarchy->ClearSelection();
 		for (int32 Index = 0; Index < NewNames.Num(); ++Index)
 		{
-			SelectBone(NewNames[Index]);
+			Hierarchy->Select(NewNames[Index]);
 		}
+
+		FSlateApplication::Get().DismissAllMenus();
+		RefreshTreeView();
 	}
 }
 
@@ -743,7 +829,7 @@ void SRigBoneHierarchy::PostUndo(bool bSuccess)
 FReply SRigBoneHierarchy::OnDragDetected(const FGeometry& MyGeometry, const FPointerEvent& MouseEvent)
 {
 	TArray<FName> DraggedBoneNames;
-	TArray<TSharedPtr<FRigTreeBone>> SelectedItems =TreeView->GetSelectedItems();
+	TArray<TSharedPtr<FRigTreeBone>> SelectedItems = TreeView->GetSelectedItems();
 	for (const TSharedPtr<FRigTreeBone>& SelectedItem : SelectedItems)
 	{
 		DraggedBoneNames.Add(SelectedItem->CachedBone);
@@ -777,8 +863,8 @@ bool SRigBoneHierarchy::RenameBone(const FName& OldName, const FName& NewName)
 	if (Hierarchy)
 	{
 		Hierarchy->Rename(OldName, NewName);
-		SelectBone(NewName);
-
+		Hierarchy->ClearSelection();
+		Hierarchy->Select(NewName);
 		return true;
 	}
 
diff --git a/Engine/Plugins/Experimental/ControlRig/Source/ControlRigEditor/Private/Editor/SRigBoneHierarchy.h b/Engine/Plugins/Experimental/ControlRig/Source/ControlRigEditor/Private/Editor/SRigBoneHierarchy.h
--- a/Engine/Plugins/Experimental/ControlRig/Source/ControlRigEditor/Private/Editor/SRigBoneHierarchy.h
+++ b/Engine/Plugins/Experimental/ControlRig/Source/ControlRigEditor/Private/Editor/SRigBoneHierarchy.h
@@ -190,8 +190,13 @@ private:
 	void SetExpansionRecursive(TSharedPtr<FRigTreeBone> InBones);
 
 	void ClearDetailPanel() const;
-	void SelectBone(const FName& BoneName) const;
-	void OnRigElementChanged(FRigHierarchyContainer* Container, ERigElementType ElementType, const FName& InName);
+
+	bool bIsChangingRigHierarchy;
+	void OnRigElementAdded(FRigHierarchyContainer* Container, ERigElementType ElementType, const FName& InName);
+	void OnRigElementRemoved(FRigHierarchyContainer* Container, ERigElementType ElementType, const FName& InName);
+	void OnRigElementRenamed(FRigHierarchyContainer* Container, ERigElementType ElementType, const FName& InOldName, const FName& InNewName);
+	void OnRigElementReparented(FRigHierarchyContainer* Container, ERigElementType ElementType, const FName& InName, const FName& InOldParentName, const FName& InNewParentName);
+	void OnRigElementSelected(FRigHierarchyContainer* Container, ERigElementType ElementType, const FName& InName, bool bSelected);
 
 public:
 	bool RenameBone(const FName& OldName, const FName& NewName);
diff --git a/Engine/Plugins/Experimental/ControlRig/Source/ControlRigEditor/Private/Editor/SRigCurveContainer.cpp b/Engine/Plugins/Experimental/ControlRig/Source/ControlRigEditor/Private/Editor/SRigCurveContainer.cpp
--- a/Engine/Plugins/Experimental/ControlRig/Source/ControlRigEditor/Private/Editor/SRigCurveContainer.cpp
+++ b/Engine/Plugins/Experimental/ControlRig/Source/ControlRigEditor/Private/Editor/SRigCurveContainer.cpp
@@ -26,8 +26,6 @@ static const FName ColumnID_RigCurveValueLabel( "Value" );
 //////////////////////////////////////////////////////////////////////////
 // SRigCurveListRow
 
-typedef TSharedPtr< FDisplayedRigCurveInfo > FDisplayedRigCurveInfoPtr;
-
 void SRigCurveListRow::Construct( const FArguments& InArgs, const TSharedRef<STableViewBase>& InOwnerTableView)
 {
 	Item = InArgs._Item;
@@ -38,7 +36,7 @@ void SRigCurveListRow::Construct( const FArguments& InArgs, const TSharedRef<STa
 
 	check( Item.IsValid() );
 
-	SMultiColumnTableRow< TSharedPtr<FDisplayedRigCurveInfo> >::Construct( FSuperRowType::FArguments(), InOwnerTableView );
+	SMultiColumnTableRow< FDisplayedRigCurveInfoPtr >::Construct( FSuperRowType::FArguments(), InOwnerTableView );
 }
 
 TSharedRef< SWidget > SRigCurveListRow::GenerateWidgetForColumn( const FName& ColumnName )
@@ -98,6 +96,7 @@ void SRigCurveListRow::OnRigCurveValueValueCommitted( float NewValue, ETextCommi
 	}
 }
 
+
 FText SRigCurveListRow::GetItemName() const
 {
 	return FText::FromName(Item->CurveName);
@@ -143,7 +142,11 @@ void SRigCurveContainer::Construct(const FArguments& InArgs, TSharedRef<FControl
 {
 	ControlRigEditor = InControlRigEditor;
 	ControlRigBlueprint = InControlRigEditor.Get().GetControlRigBlueprint();
-	ControlRigBlueprint->HierarchyContainer.OnElementChanged.AddRaw(this, &SRigCurveContainer::OnRigElementChanged);
+	bIsChangingRigHierarchy = false;
+	ControlRigBlueprint->HierarchyContainer.OnElementAdded.AddRaw(this, &SRigCurveContainer::OnRigElementAdded);
+	ControlRigBlueprint->HierarchyContainer.OnElementRemoved.AddRaw(this, &SRigCurveContainer::OnRigElementRemoved);
+	ControlRigBlueprint->HierarchyContainer.OnElementRenamed.AddRaw(this, &SRigCurveContainer::OnRigElementRenamed);
+	ControlRigBlueprint->HierarchyContainer.OnElementSelected.AddRaw(this, &SRigCurveContainer::OnRigElementSelected);
 
 	// Register and bind all our menu commands
 	FCurveContainerCommands::Register();
@@ -178,7 +181,7 @@ void SRigCurveContainer::Construct(const FArguments& InArgs, TSharedRef<FControl
 			.OnContextMenuOpening( this, &SRigCurveContainer::OnGetContextMenuContent )
 			.ItemHeight( 22.0f )
 			.SelectionMode(ESelectionMode::Multi)
-//			.OnSelectionChanged( this, &SRigCurveContainer::OnSelectionChanged )
+			.OnSelectionChanged( this, &SRigCurveContainer::OnSelectionChanged )
 			.HeaderRow
 			(
 				SNew( SHeaderRow )
@@ -203,7 +206,10 @@ SRigCurveContainer::~SRigCurveContainer()
 		ControlRigBlueprint = ControlRigEditor.Pin()->GetControlRigBlueprint();
 		if (ControlRigBlueprint.IsValid())
 		{
-			ControlRigBlueprint->HierarchyContainer.OnElementChanged.RemoveAll(this);
+			ControlRigBlueprint->HierarchyContainer.OnElementAdded.RemoveAll(this);
+			ControlRigBlueprint->HierarchyContainer.OnElementRemoved.RemoveAll(this);
+			ControlRigBlueprint->HierarchyContainer.OnElementRenamed.RemoveAll(this);
+			ControlRigBlueprint->HierarchyContainer.OnElementSelected.RemoveAll(this);
 		}
 	}
 }
@@ -265,7 +271,7 @@ void SRigCurveContainer::OnFilterTextCommitted( const FText& SearchText, ETextCo
 	OnFilterTextChanged( SearchText );
 }
 
-TSharedRef<ITableRow> SRigCurveContainer::GenerateRigCurveRow(TSharedPtr<FDisplayedRigCurveInfo> InInfo, const TSharedRef<STableViewBase>& OwnerTable)
+TSharedRef<ITableRow> SRigCurveContainer::GenerateRigCurveRow(FDisplayedRigCurveInfoPtr InInfo, const TSharedRef<STableViewBase>& OwnerTable)
 {
 	check( InInfo.IsValid() );
 
@@ -305,7 +311,7 @@ TSharedPtr<SWidget> SRigCurveContainer::OnGetContextMenuContent() const
 
 void SRigCurveContainer::OnRenameClicked()
 {
-	TArray< TSharedPtr<FDisplayedRigCurveInfo> > SelectedItems = RigCurveListView->GetSelectedItems();
+	TArray< FDisplayedRigCurveInfoPtr > SelectedItems = RigCurveListView->GetSelectedItems();
 
 	SelectedItems[0]->EditableText->EnterEditingMode();
 }
@@ -341,8 +347,16 @@ void SRigCurveContainer::CreateNewNameEntry(const FText& CommittedText, ETextCom
 		FRigCurveContainer* Container = GetCurveContainer();
 		if (Container)
 		{
-			Container->Add(FName(*CommittedText.ToString()));
+			TGuardValue<bool> GuardReentry(bIsChangingRigHierarchy, true);
+
+			FName NewName = FName(*CommittedText.ToString());
+			Container->Add(NewName);
+			Container->ClearSelection();
+			Container->Select(NewName);
 		}
+
+		FSlateApplication::Get().DismissAllMenus();
+		RefreshCurveList();
 	}
 }
 
@@ -371,7 +385,7 @@ void SRigCurveContainer::CreateRigCurveList( const FString& SearchText )
 		// Sort final list
 		struct FSortNamesAlphabetically
 		{
-			bool operator()(const TSharedPtr<FDisplayedRigCurveInfo>& A, const TSharedPtr<FDisplayedRigCurveInfo>& B) const
+			bool operator()(const FDisplayedRigCurveInfoPtr& A, const FDisplayedRigCurveInfoPtr& B) const
 			{
 				return (A.Get()->CurveName.Compare(B.Get()->CurveName) < 0);
 			}
@@ -380,6 +394,15 @@ void SRigCurveContainer::CreateRigCurveList( const FString& SearchText )
 		RigCurveList.Sort(FSortNamesAlphabetically());
 	}
 	RigCurveListView->RequestListRefresh();
+
+	if (Container)
+	{
+		for (const FName& SelectedCurve : Container->CurrentSelection())
+		{
+			OnRigElementSelected(&ControlRigBlueprint->HierarchyContainer, ERigElementType::Curve, SelectedCurve, true);
+		}
+	}
+
 }
 
 void SRigCurveContainer::RefreshCurveList()
@@ -387,7 +410,7 @@ void SRigCurveContainer::RefreshCurveList()
 	CreateRigCurveList(FilterText.ToString());
 }
 
-void SRigCurveContainer::OnNameCommitted(const FText& InNewName, ETextCommit::Type CommitType, TSharedPtr<FDisplayedRigCurveInfo> Item)
+void SRigCurveContainer::OnNameCommitted(const FText& InNewName, ETextCommit::Type CommitType, FDisplayedRigCurveInfoPtr Item)
 {
 	FRigCurveContainer* Container = GetCurveContainer();
 	if (Container)
@@ -406,7 +429,7 @@ void SRigCurveContainer::OnDeleteNameClicked()
 	FRigCurveContainer* Container = GetCurveContainer();
 	if (Container)
 	{
-		TArray< TSharedPtr<FDisplayedRigCurveInfo> > SelectedItems = RigCurveListView->GetSelectedItems();
+		TArray< FDisplayedRigCurveInfoPtr > SelectedItems = RigCurveListView->GetSelectedItems();
 		for (auto Item : SelectedItems)
 		{
 			Container->Remove(Item->CurveName);
@@ -441,6 +464,8 @@ float SRigCurveContainer::GetCurveValue(const FName& CurveName)
 
 void SRigCurveContainer::ChangeCurveName(const FName& OldName, const FName& NewName)
 {
+	TGuardValue<bool> GuardReentry(bIsChangingRigHierarchy, true);
+
 	FRigCurveContainer* Container = GetCurveContainer();
 	if (Container)
 	{
@@ -448,16 +473,88 @@ void SRigCurveContainer::ChangeCurveName(const FName& OldName, const FName& NewN
 	}
 }
 
-void SRigCurveContainer::OnRigElementChanged(FRigHierarchyContainer* Container, ERigElementType ElementType, const FName& InName)
+void SRigCurveContainer::OnSelectionChanged(FDisplayedRigCurveInfoPtr Selection, ESelectInfo::Type SelectInfo)
 {
-	if (ElementType != ERigElementType::Curve)
+	if (bIsChangingRigHierarchy)
 	{
 		return;
 	}
 
+	FRigCurveContainer* Container = GetCurveContainer();
+
+	if (Container)
+	{
+		TGuardValue<bool> GuardRigHierarchyChanges(bIsChangingRigHierarchy, true);
+
+		TArray<FName> OldSelection = Container->CurrentSelection();
+		TArray<FName> NewSelection;
+
+		TArray<FDisplayedRigCurveInfoPtr> SelectedItems = RigCurveListView->GetSelectedItems();
+		for (const FDisplayedRigCurveInfoPtr& SelectedItem : SelectedItems)
+		{
+			NewSelection.Add(SelectedItem->CurveName);
+		}
+
+		for (const FName& PreviouslySelected : OldSelection)
+		{
+			if (NewSelection.Contains(PreviouslySelected))
+			{
+				continue;
+			}
+			Container->Select(PreviouslySelected, false);
+		}
+
+		for (const FName& NewlySelected : NewSelection)
+		{
+			Container->Select(NewlySelected, true);
+		}
+	}
+}
+
+void SRigCurveContainer::OnRigElementAdded(FRigHierarchyContainer* Container, ERigElementType ElementType, const FName& InName)
+{
+	if (bIsChangingRigHierarchy || ElementType != ERigElementType::Curve)
+	{
+		return;
+	}
+	RefreshCurveList();
+}
+
+void SRigCurveContainer::OnRigElementRemoved(FRigHierarchyContainer* Container, ERigElementType ElementType, const FName& InName)
+{
+	if (bIsChangingRigHierarchy || ElementType != ERigElementType::Curve)
+	{
+		return;
+	}
+	RefreshCurveList();
+}
+
+void SRigCurveContainer::OnRigElementRenamed(FRigHierarchyContainer* Container, ERigElementType ElementType, const FName& InOldName, const FName& InNewName)
+{
+	if (bIsChangingRigHierarchy || ElementType != ERigElementType::Curve)
+	{
+		return;
+	}
 	RefreshCurveList();
 }
 
+void SRigCurveContainer::OnRigElementSelected(FRigHierarchyContainer* Container, ERigElementType ElementType, const FName& InName, bool bSelected)
+{
+	if (bIsChangingRigHierarchy || ElementType != ERigElementType::Curve)
+	{
+		return;
+	}
+
+	for(const FDisplayedRigCurveInfoPtr& Item : RigCurveList)
+	{
+		if (Item->CurveName == InName)
+		{
+			RigCurveListView->SetItemSelection(Item, bSelected);
+			break;
+		}
+	}
+}
+
 void SRigCurveContainer::CreateImportMenu(FMenuBuilder& MenuBuilder)
 {
 	MenuBuilder.AddWidget(
@@ -515,14 +612,23 @@ void SRigCurveContainer::ImportCurve(const FAssetData& InAssetData)
 				
 			const FSmartNameMapping* SmartNameMapping = Skeleton->GetSmartNameContainer(USkeleton::AnimCurveMappingName);
 
+			Container->ClearSelection();
+
 			TArray<FName> NameArray;
 			SmartNameMapping->FillNameArray(NameArray);
 			for (int32 Index = 0; Index < NameArray.Num() ; ++Index)
 			{
+				TGuardValue<bool> GuardReentry(bIsChangingRigHierarchy, true);
 				Container->Add(NameArray[Index]);
 			}
 
+			for (int32 Index = 0; Index < NameArray.Num(); ++Index)
+			{
+				Container->Select(NameArray[Index]);
+			}
+
 			FSlateApplication::Get().DismissAllMenus();
+			RefreshCurveList();
 		}
 	}
 }
diff --git a/Engine/Plugins/Experimental/ControlRig/Source/ControlRigEditor/Private/Editor/SRigCurveContainer.h b/Engine/Plugins/Experimental/ControlRig/Source/ControlRigEditor/Private/Editor/SRigCurveContainer.h
--- a/Engine/Plugins/Experimental/ControlRig/Source/ControlRigEditor/Private/Editor/SRigCurveContainer.h
+++ b/Engine/Plugins/Experimental/ControlRig/Source/ControlRigEditor/Private/Editor/SRigCurveContainer.h
@@ -41,7 +41,8 @@ protected:
 	{}
 };
 
-typedef SListView< TSharedPtr<FDisplayedRigCurveInfo> > SRigCurveListType;
+typedef TSharedPtr< FDisplayedRigCurveInfo > FDisplayedRigCurveInfoPtr;
+typedef SListView< FDisplayedRigCurveInfoPtr > SRigCurveListType;
 
 //////////////////////////////////////////////////////////////////////////
 // SRigCurveListRow
@@ -50,14 +51,14 @@ DECLARE_DELEGATE_TwoParams(FSetRigCurveValue, const FName&, float);
 DECLARE_DELEGATE_RetVal_OneParam(float, FGetRigCurveValue, const FName&);
 DECLARE_DELEGATE_RetVal(FText&, FGetFilterText);
 
-class SRigCurveListRow : public SMultiColumnTableRow< TSharedPtr<FDisplayedRigCurveInfo> >
+class SRigCurveListRow : public SMultiColumnTableRow< FDisplayedRigCurveInfoPtr >
 {
 public:
 
 	SLATE_BEGIN_ARGS(SRigCurveListRow) {}
 
 		/** The item for this row **/
-		SLATE_ARGUMENT(TSharedPtr<FDisplayedRigCurveInfo>, Item)
+		SLATE_ARGUMENT(FDisplayedRigCurveInfoPtr, Item)
 
 		/** Callback when the text is committed. */
 		SLATE_EVENT(FOnTextCommitted, OnTextCommitted)
@@ -108,7 +109,7 @@ private:
 	FSlateColor GetItemTextColor() const;
 
 	/** The name and Value of the morph target */
-	TSharedPtr<FDisplayedRigCurveInfo>	Item;
+	FDisplayedRigCurveInfoPtr	Item;
 
 	FOnTextCommitted OnTextCommitted;
 	FSetRigCurveValue OnSetRigCurveValue;
@@ -186,7 +187,7 @@ public:
 	*
 	* @return A new Slate widget, containing the UI for this row
 	*/
-	TSharedRef<ITableRow> GenerateRigCurveRow(TSharedPtr<FDisplayedRigCurveInfo> InInfo, const TSharedRef<STableViewBase>& OwnerTable);
+	TSharedRef<ITableRow> GenerateRigCurveRow(FDisplayedRigCurveInfoPtr InInfo, const TSharedRef<STableViewBase>& OwnerTable);
 
 	/**
 	* Accessor so our rows can grab the filtertext for highlighting
@@ -197,7 +198,7 @@ public:
 	void RefreshCurveList();
 
 	// When a name is committed after being edited in the list
-	virtual void OnNameCommitted(const FText& NewName, ETextCommit::Type CommitType, TSharedPtr<FDisplayedRigCurveInfo> Item);
+	virtual void OnNameCommitted(const FText& NewName, ETextCommit::Type CommitType, FDisplayedRigCurveInfoPtr Item);
 
 private:
 
@@ -229,7 +230,14 @@ private:
 	void SetCurveValue(const FName& CurveName, float CurveValue);
 	float GetCurveValue(const FName& CurveName);
 	void ChangeCurveName(const FName& OldName, const FName& NewName);
-	void OnRigElementChanged(FRigHierarchyContainer* Container, ERigElementType ElementType, const FName& InName);
+
+	void OnSelectionChanged(FDisplayedRigCurveInfoPtr Selection, ESelectInfo::Type SelectInfo);
+
+	bool bIsChangingRigHierarchy;
+	void OnRigElementAdded(FRigHierarchyContainer* Container, ERigElementType ElementType, const FName& InName);
+	void OnRigElementRemoved(FRigHierarchyContainer* Container, ERigElementType ElementType, const FName& InName);
+	void OnRigElementRenamed(FRigHierarchyContainer* Container, ERigElementType ElementType, const FName& InOldName, const FName& InNewName);
+	void OnRigElementSelected(FRigHierarchyContainer* Container, ERigElementType ElementType, const FName& InName, bool bSelected);
 
 	// import curve part
 	void ImportCurve(const FAssetData& InAssetData);
@@ -240,7 +248,7 @@ private:
 	TSharedPtr<SSearchBox>	NameFilterBox;
 
 	/** A list of animation curve. Used by the RigCurveListView. */
-	TArray< TSharedPtr<FDisplayedRigCurveInfo> > RigCurveList;
+	TArray< FDisplayedRigCurveInfoPtr > RigCurveList;
 
 	/** Widget used to display the list of animation curve */
 	TSharedPtr<SRigCurveListType> RigCurveListView;
