ControlRig: Fixed used of WITH_EDITOR vs WITH_EDITORONLY_DATA.

diff --git a/Engine/Plugins/Experimental/ControlRig/Source/ControlRig/Public/Rigs/RigBoneHierarchy.h b/Engine/Plugins/Experimental/ControlRig/Source/ControlRig/Public/Rigs/RigBoneHierarchy.h
--- a/Engine/Plugins/Experimental/ControlRig/Source/ControlRig/Public/Rigs/RigBoneHierarchy.h
+++ b/Engine/Plugins/Experimental/ControlRig/Source/ControlRig/Public/Rigs/RigBoneHierarchy.h
@@ -172,7 +172,7 @@ private:
 	UPROPERTY()
 	TMap<FName, int32> NameToIndexMapping;
 
-#if WITH_EDITOR
+#if WITH_EDITORONLY_DATA
 	UPROPERTY(transient)
 	TArray<FName> Selection;
 #endif
diff --git a/Engine/Plugins/Experimental/ControlRig/Source/ControlRig/Public/Rigs/RigControlHierarchy.h b/Engine/Plugins/Experimental/ControlRig/Source/ControlRig/Public/Rigs/RigControlHierarchy.h
--- a/Engine/Plugins/Experimental/ControlRig/Source/ControlRig/Public/Rigs/RigControlHierarchy.h
+++ b/Engine/Plugins/Experimental/ControlRig/Source/ControlRig/Public/Rigs/RigControlHierarchy.h
@@ -198,7 +198,7 @@ private:
 	UPROPERTY()
 	TMap<FName, int32> NameToIndexMapping;
 
-#if WITH_EDITOR
+#if WITH_EDITORONLY_DATA
 	UPROPERTY(transient)
 	TArray<FName> Selection;
 #endif
diff --git a/Engine/Plugins/Experimental/ControlRig/Source/ControlRig/Public/Rigs/RigCurveContainer.h b/Engine/Plugins/Experimental/ControlRig/Source/ControlRig/Public/Rigs/RigCurveContainer.h
--- a/Engine/Plugins/Experimental/ControlRig/Source/ControlRig/Public/Rigs/RigCurveContainer.h
+++ b/Engine/Plugins/Experimental/ControlRig/Source/ControlRig/Public/Rigs/RigCurveContainer.h
@@ -126,7 +126,7 @@ private:
 	UPROPERTY()
 	TMap<FName, int32> NameToIndexMapping;
 
-#if WITH_EDITOR
+#if WITH_EDITORONLY_DATA
 	UPROPERTY(transient)
 	TArray<FName> Selection;
 #endif
diff --git a/Engine/Plugins/Experimental/ControlRig/Source/ControlRig/Public/Rigs/RigSpaceHierarchy.h b/Engine/Plugins/Experimental/ControlRig/Source/ControlRig/Public/Rigs/RigSpaceHierarchy.h
--- a/Engine/Plugins/Experimental/ControlRig/Source/ControlRig/Public/Rigs/RigSpaceHierarchy.h
+++ b/Engine/Plugins/Experimental/ControlRig/Source/ControlRig/Public/Rigs/RigSpaceHierarchy.h
@@ -184,7 +184,7 @@ private:
 	UPROPERTY()
 	TMap<FName, int32> NameToIndexMapping;
 
-#if WITH_EDITOR
+#if WITH_EDITORONLY_DATA
 	UPROPERTY(transient)
 	TArray<FName> Selection;
 #endif
