Fix for build break from Control Rig API change, auto-merge.

diff --git a/Engine/Plugins/Experimental/ControlRig/Source/ControlRig/Private/Units/Highlevel/Hierarchy/RigUnit_SlideChain.cpp b/Engine/Plugins/Experimental/ControlRig/Source/ControlRig/Private/Units/Highlevel/Hierarchy/RigUnit_SlideChain.cpp
--- a/Engine/Plugins/Experimental/ControlRig/Source/ControlRig/Private/Units/Highlevel/Hierarchy/RigUnit_SlideChain.cpp
+++ b/Engine/Plugins/Experimental/ControlRig/Source/ControlRig/Private/Units/Highlevel/Hierarchy/RigUnit_SlideChain.cpp
@@ -6,7 +6,8 @@
 
 void FRigUnit_SlideChain::Execute(const FRigUnitContext& Context)
 {
-	FRigHierarchy* Hierarchy = (FRigHierarchy*)(Context.HierarchyReference.Get());
+
+	FRigBoneHierarchy* Hierarchy = ExecuteContext.GetBones();
 	if (Hierarchy == nullptr)
 	{
 		return;
@@ -37,7 +38,7 @@ void FRigUnit_SlideChain::Execute(const FRigUnitContext& Context)
 				{
 					break;
 				}
-				EndBoneIndex = Hierarchy->GetParentIndex(EndBoneIndex);
+				EndBoneIndex = (*Hierarchy)[EndBoneIndex].ParentIndex;
 			}
 		}
 
