Put Control Rig  Debug Drawing in Editor mode.
Just moved the interface over from the BP Ed Mode drawing to the Editor Ed Mode drawing so it's in both editor and Persona.
Then used ComponentTransform to make sure the debugged points are drawn relative to the actor.(Will need to support parented actors/actros in sequencer).

diff --git a/Engine/Plugins/Experimental/ControlRig/Source/ControlRigEditor/Private/EditMode/ControlRigEditMode.cpp b/Engine/Plugins/Experimental/ControlRig/Source/ControlRigEditor/Private/EditMode/ControlRigEditMode.cpp
--- a/Engine/Plugins/Experimental/ControlRig/Source/ControlRigEditor/Private/EditMode/ControlRigEditMode.cpp
+++ b/Engine/Plugins/Experimental/ControlRig/Source/ControlRigEditor/Private/EditMode/ControlRigEditMode.cpp
@@ -123,6 +123,7 @@ void FControlRigEditMode::SetObjects_Internal()
 	if(WeakControlRig.IsValid())
 	{
 		SelectedObjects.Add(WeakControlRig);
+		WeakControlRig.Get()->DrawInterface = &DrawInterface;
 	}
 
 	if(UsesToolkits())
@@ -368,16 +369,16 @@ void FControlRigEditMode::Render(const FSceneView* View, FViewport* Viewport, FP
 			bRender = false;
 		}
 
+		FTransform ComponentTransform = FTransform::Identity;
+		if (TSharedPtr<IControlRigObjectBinding> ObjectBinding = ControlRig->GetObjectBinding())
+		{
+			USceneComponent* Component = Cast<USceneComponent>(ObjectBinding->GetBoundObject());
+			ComponentTransform = Component ? Component->GetComponentTransform() : FTransform::Identity;
+		}
 		if (bRender)
 		{
 			if (Settings->bDisplayHierarchy)
 			{
-				FTransform ComponentTransform = FTransform::Identity;
-				if (TSharedPtr<IControlRigObjectBinding> ObjectBinding = ControlRig->GetObjectBinding())
-				{
-					USceneComponent* Component = Cast<USceneComponent>(ObjectBinding->GetBoundObject());
-					ComponentTransform = Component ? Component->GetComponentTransform() : FTransform::Identity;
-				}
 
 				// each base hierarchy Bone
 				const FRigBoneHierarchy& BaseHierarchy = ControlRig->GetBoneHierarchy();
@@ -390,10 +391,10 @@ void FControlRigEditMode::Render(const FSceneView* View, FViewport* Viewport, FP
 					{
 						const FTransform ParentTransform = BaseHierarchy.GetGlobalTransform(CurrentBone.ParentIndex);
 
-						PDI->DrawLine(Transform.GetLocation(), ParentTransform.GetLocation(), FLinearColor::White, SDPG_Foreground);
+						PDI->DrawLine(ComponentTransform.TransformPosition(Transform.GetLocation()),ComponentTransform.TransformPosition(ParentTransform.GetLocation()), FLinearColor::White, SDPG_Foreground);
 					}
 
-					PDI->DrawPoint(Transform.GetLocation(), FLinearColor::White, 5.0f, SDPG_Foreground);
+					PDI->DrawPoint(ComponentTransform.TransformPosition(Transform.GetLocation()), FLinearColor::White, 5.0f, SDPG_Foreground);
 				}
 			}
 
@@ -420,7 +421,7 @@ void FControlRigEditMode::Render(const FSceneView* View, FViewport* Viewport, FP
 					{
 						for (const FVector& Point : Instruction.Positions)
 						{
-							PDI->DrawPoint(Point, Instruction.Color, Instruction.Thickness, SDPG_Foreground);
+							PDI->DrawPoint(ComponentTransform.TransformPosition(Point), Instruction.Color, Instruction.Thickness, SDPG_Foreground);
 						}
 						break;
 					}
@@ -429,7 +430,7 @@ void FControlRigEditMode::Render(const FSceneView* View, FViewport* Viewport, FP
 						const TArray<FVector>& Points = Instruction.Positions;
 						for (int32 PointIndex = 0; PointIndex < Points.Num() - 1; PointIndex += 2)
 						{
-							PDI->DrawLine(Points[PointIndex], Points[PointIndex+1], Instruction.Color, SDPG_Foreground, Instruction.Thickness);
+							PDI->DrawLine(ComponentTransform.TransformPosition(Points[PointIndex]), ComponentTransform.TransformPosition(Points[PointIndex+1]), Instruction.Color, SDPG_Foreground, Instruction.Thickness);
 						}
 						break;
 					}
@@ -438,7 +439,7 @@ void FControlRigEditMode::Render(const FSceneView* View, FViewport* Viewport, FP
 						const TArray<FVector>& Points = Instruction.Positions;
 						for (int32 PointIndex = 0; PointIndex < Points.Num() - 1; PointIndex++)
 						{
-							PDI->DrawLine(Points[PointIndex], Points[PointIndex + 1], Instruction.Color, SDPG_Foreground, Instruction.Thickness);
+							PDI->DrawLine(ComponentTransform.TransformPosition(Points[PointIndex]), ComponentTransform.TransformPosition(Points[PointIndex + 1]), Instruction.Color, SDPG_Foreground, Instruction.Thickness);
 						}
 						break;
 					}
diff --git a/Engine/Plugins/Experimental/ControlRig/Source/ControlRigEditor/Private/EditMode/ControlRigEditMode.h b/Engine/Plugins/Experimental/ControlRig/Source/ControlRigEditor/Private/EditMode/ControlRigEditMode.h
--- a/Engine/Plugins/Experimental/ControlRig/Source/ControlRigEditor/Private/EditMode/ControlRigEditMode.h
+++ b/Engine/Plugins/Experimental/ControlRig/Source/ControlRigEditor/Private/EditMode/ControlRigEditMode.h
@@ -9,6 +9,7 @@
 #include "ControlUnitProxy.h"
 #include "ControlRigModel.h"
 #include "Rigs/RigHierarchyContainer.h"
+#include "Drawing/ControlRigDrawInterface.h"
 
 class FEditorViewportClient;
 class FViewport;
@@ -200,6 +201,9 @@ protected:
 	/** The sequencer GUID of the object we are animating */
 	FGuid ControlRigGuid;
 
+	/** The draw interface to use for the control rig */
+	FControlRigDrawInterface DrawInterface;
+
 	/** Sequencer we are currently bound to */
 	TWeakPtr<ISequencer> WeakSequencer;
 
diff --git a/Engine/Plugins/Experimental/ControlRig/Source/ControlRigEditor/Private/Editor/ControlRigEditor.cpp b/Engine/Plugins/Experimental/ControlRig/Source/ControlRigEditor/Private/Editor/ControlRigEditor.cpp
--- a/Engine/Plugins/Experimental/ControlRig/Source/ControlRigEditor/Private/Editor/ControlRigEditor.cpp
+++ b/Engine/Plugins/Experimental/ControlRig/Source/ControlRigEditor/Private/Editor/ControlRigEditor.cpp
@@ -469,7 +469,6 @@ void FControlRigEditor::HandleSetObjectBeingDebugged(UObject* InObject)
 	if (DebuggedControlRig)
 	{
 		bool bIsExternalControlRig = DebuggedControlRig != ControlRig;
-		DebuggedControlRig->DrawInterface = &DrawInterface;
 		DebuggedControlRig->ControlRigLog = &ControlRigLog;
 
 		UControlRigSkeletalMeshComponent* EditorSkelComp = Cast<UControlRigSkeletalMeshComponent>(GetPersonaToolkit()->GetPreviewScene()->GetPreviewMeshComponent());
@@ -601,7 +600,6 @@ void FControlRigEditor::Compile()
 		if (ControlRig)
 		{
 			ControlRig->ControlRigLog = &ControlRigLog;
-			ControlRig->DrawInterface = &DrawInterface;
 
 			UControlRigBlueprintGeneratedClass* GeneratedClass = Cast<UControlRigBlueprintGeneratedClass>(ControlRig->GetClass());
 			if (GeneratedClass)
@@ -1334,7 +1332,6 @@ void FControlRigEditor::UpdateControlRig()
 				ControlRig->ExecutionType = ERigExecutionType::Editing;
 
 				ControlRig->ControlRigLog = &ControlRigLog;
-				ControlRig->DrawInterface = &DrawInterface;
  			}
 
 			CacheBoneNameList();
diff --git a/Engine/Plugins/Experimental/ControlRig/Source/ControlRigEditor/Private/Editor/ControlRigEditor.h b/Engine/Plugins/Experimental/ControlRig/Source/ControlRigEditor/Private/Editor/ControlRigEditor.h
--- a/Engine/Plugins/Experimental/ControlRig/Source/ControlRigEditor/Private/Editor/ControlRigEditor.h
+++ b/Engine/Plugins/Experimental/ControlRig/Source/ControlRigEditor/Private/Editor/ControlRigEditor.h
@@ -8,7 +8,6 @@
 #include "DragAndDrop//GraphNodeDragDropOp.h"
 #include "ControlRigDefines.h"
 #include "ControlRigLog.h"
-#include "Drawing/ControlRigDrawInterface.h"
 #include "ControlRigModel.h"
 
 class UControlRigBlueprint;
@@ -256,9 +255,6 @@ protected:
 	/** Once the log is collected update the graph */
 	void UpdateGraphCompilerErrors();
 
-	/** The draw interface to use for the control rig */
-	FControlRigDrawInterface DrawInterface;
-
 	/** This can be used to enable dumping of a unit test */
 	void DumpUnitTestCode();
 
