CIS Fix
#trivial

diff --git a/Engine/Plugins/Experimental/ControlRig/Source/ControlRigEditor/Private/EditMode/ControlRigEditMode.cpp b/Engine/Plugins/Experimental/ControlRig/Source/ControlRigEditor/Private/EditMode/ControlRigEditMode.cpp
--- a/Engine/Plugins/Experimental/ControlRig/Source/ControlRigEditor/Private/EditMode/ControlRigEditMode.cpp
+++ b/Engine/Plugins/Experimental/ControlRig/Source/ControlRigEditor/Private/EditMode/ControlRigEditMode.cpp
@@ -406,48 +406,45 @@ void FControlRigEditMode::Render(const FSceneView* View, FViewport* Viewport, FP
 // 			}
 		}
 
-		FControlRigDrawInterface* DrawInterface = ControlRig->DrawInterface;
-		if (DrawInterface)
+
+		for (const FControlRigDrawInterface::FDrawIntruction& Instruction : DrawInterface.DrawInstructions)
 		{
-			for (const FControlRigDrawInterface::FDrawIntruction& Instruction : DrawInterface->DrawInstructions)
+			if (Instruction.Positions.Num() == 0)
+			{
+				continue;
+			}
+			switch (Instruction.DrawType)
 			{
-				if (Instruction.Positions.Num() == 0)
+				case FControlRigDrawInterface::EDrawType_Point:
 				{
-					continue;
+					for (const FVector& Point : Instruction.Positions)
+					{
+						PDI->DrawPoint(ComponentTransform.TransformPosition(Point), Instruction.Color, Instruction.Thickness, SDPG_Foreground);
+					}
+					break;
 				}
-				switch (Instruction.DrawType)
+				case FControlRigDrawInterface::EDrawType_Lines:
 				{
-					case FControlRigDrawInterface::EDrawType_Point:
+					const TArray<FVector>& Points = Instruction.Positions;
+					for (int32 PointIndex = 0; PointIndex < Points.Num() - 1; PointIndex += 2)
 					{
-						for (const FVector& Point : Instruction.Positions)
-						{
-							PDI->DrawPoint(ComponentTransform.TransformPosition(Point), Instruction.Color, Instruction.Thickness, SDPG_Foreground);
-						}
-						break;
+						PDI->DrawLine(ComponentTransform.TransformPosition(Points[PointIndex]), ComponentTransform.TransformPosition(Points[PointIndex+1]), Instruction.Color, SDPG_Foreground, Instruction.Thickness);
 					}
-					case FControlRigDrawInterface::EDrawType_Lines:
-					{
-						const TArray<FVector>& Points = Instruction.Positions;
-						for (int32 PointIndex = 0; PointIndex < Points.Num() - 1; PointIndex += 2)
-						{
-							PDI->DrawLine(ComponentTransform.TransformPosition(Points[PointIndex]), ComponentTransform.TransformPosition(Points[PointIndex+1]), Instruction.Color, SDPG_Foreground, Instruction.Thickness);
-						}
-						break;
-					}
-					case FControlRigDrawInterface::EDrawType_LineStrip:
+					break;
+				}
+				case FControlRigDrawInterface::EDrawType_LineStrip:
+				{
+					const TArray<FVector>& Points = Instruction.Positions;
+					for (int32 PointIndex = 0; PointIndex < Points.Num() - 1; PointIndex++)
 					{
-						const TArray<FVector>& Points = Instruction.Positions;
-						for (int32 PointIndex = 0; PointIndex < Points.Num() - 1; PointIndex++)
-						{
-							PDI->DrawLine(ComponentTransform.TransformPosition(Points[PointIndex]), ComponentTransform.TransformPosition(Points[PointIndex + 1]), Instruction.Color, SDPG_Foreground, Instruction.Thickness);
-						}
-						break;
+						PDI->DrawLine(ComponentTransform.TransformPosition(Points[PointIndex]), ComponentTransform.TransformPosition(Points[PointIndex + 1]), Instruction.Color, SDPG_Foreground, Instruction.Thickness);
 					}
+					break;
 				}
 			}
-
-			DrawInterface->DrawInstructions.Reset();
 		}
+		DrawInterface.DrawInstructions.Reset();
+
 	}
 }
 
