Refactoring all Control Rig Units to use Static Virtual Methods

diff --git a/Engine/Plugins/Experimental/ControlRig/Source/ControlRig/Private/Units/Animation/RigUnit_AnimEasing.cpp b/Engine/Plugins/Experimental/ControlRig/Source/ControlRig/Private/Units/Animation/RigUnit_AnimEasing.cpp
--- a/Engine/Plugins/Experimental/ControlRig/Source/ControlRig/Private/Units/Animation/RigUnit_AnimEasing.cpp
+++ b/Engine/Plugins/Experimental/ControlRig/Source/ControlRig/Private/Units/Animation/RigUnit_AnimEasing.cpp
@@ -3,7 +3,7 @@
 #include "Units/Animation/RigUnit_AnimEasing.h"
 #include "Units/RigUnitContext.h"
 
-void FRigUnit_AnimEasing::Execute(const FRigUnitContext& Context)
+UE_RigUnit_AnimEasing_IMPLEMENT_STATIC_VIRTUAL_METHOD(void, Execute, const FRigUnitContext& Context)
 {
     DECLARE_SCOPE_HIERARCHICAL_COUNTER_RIGUNIT()
 	if (FMath::IsNearlyEqual(SourceMinimum, SourceMaximum))
diff --git a/Engine/Plugins/Experimental/ControlRig/Source/ControlRig/Private/Units/Animation/RigUnit_AnimEasing.h b/Engine/Plugins/Experimental/ControlRig/Source/ControlRig/Private/Units/Animation/RigUnit_AnimEasing.h
--- a/Engine/Plugins/Experimental/ControlRig/Source/ControlRig/Private/Units/Animation/RigUnit_AnimEasing.h
+++ b/Engine/Plugins/Experimental/ControlRig/Source/ControlRig/Private/Units/Animation/RigUnit_AnimEasing.h
@@ -39,6 +39,7 @@ struct FRigUnit_AnimEasing : public FRigUnit_AnimBase
 		SourceMaximum = TargetMaximum = 1.f;
 	}
 
+	STATIC_VIRTUAL_METHOD()
 	virtual void Execute(const FRigUnitContext& Context) override;
 
 	UPROPERTY(meta=(Input))
diff --git a/Engine/Plugins/Experimental/ControlRig/Source/ControlRig/Private/Units/Animation/RigUnit_AnimEvalRichCurve.cpp b/Engine/Plugins/Experimental/ControlRig/Source/ControlRig/Private/Units/Animation/RigUnit_AnimEvalRichCurve.cpp
--- a/Engine/Plugins/Experimental/ControlRig/Source/ControlRig/Private/Units/Animation/RigUnit_AnimEvalRichCurve.cpp
+++ b/Engine/Plugins/Experimental/ControlRig/Source/ControlRig/Private/Units/Animation/RigUnit_AnimEvalRichCurve.cpp
@@ -3,7 +3,7 @@
 #include "Units/Animation/RigUnit_AnimEvalRichCurve.h"
 #include "Units/RigUnitContext.h"
 
-void FRigUnit_AnimEvalRichCurve::Execute(const FRigUnitContext& Context)
+UE_RigUnit_AnimEvalRichCurve_IMPLEMENT_STATIC_VIRTUAL_METHOD(void, Execute, const FRigUnitContext& Context)
 {
     DECLARE_SCOPE_HIERARCHICAL_COUNTER_RIGUNIT()
 	if (FMath::IsNearlyEqual(SourceMinimum, SourceMaximum))
diff --git a/Engine/Plugins/Experimental/ControlRig/Source/ControlRig/Private/Units/Animation/RigUnit_AnimEvalRichCurve.h b/Engine/Plugins/Experimental/ControlRig/Source/ControlRig/Private/Units/Animation/RigUnit_AnimEvalRichCurve.h
--- a/Engine/Plugins/Experimental/ControlRig/Source/ControlRig/Private/Units/Animation/RigUnit_AnimEvalRichCurve.h
+++ b/Engine/Plugins/Experimental/ControlRig/Source/ControlRig/Private/Units/Animation/RigUnit_AnimEvalRichCurve.h
@@ -24,6 +24,7 @@ struct FRigUnit_AnimEvalRichCurve : public FRigUnit_AnimBase
 		SourceMaximum = TargetMaximum = 1.f;
 	}
 
+	STATIC_VIRTUAL_METHOD()
 	virtual void Execute(const FRigUnitContext& Context) override;
 
 	UPROPERTY(meta=(Input))
diff --git a/Engine/Plugins/Experimental/ControlRig/Source/ControlRig/Private/Units/Animation/RigUnit_AnimRichCurve.cpp b/Engine/Plugins/Experimental/ControlRig/Source/ControlRig/Private/Units/Animation/RigUnit_AnimRichCurve.cpp
--- a/Engine/Plugins/Experimental/ControlRig/Source/ControlRig/Private/Units/Animation/RigUnit_AnimRichCurve.cpp
+++ b/Engine/Plugins/Experimental/ControlRig/Source/ControlRig/Private/Units/Animation/RigUnit_AnimRichCurve.cpp
@@ -3,7 +3,7 @@
 #include "Units/Animation/RigUnit_AnimRichCurve.h"
 #include "Units/RigUnitContext.h"
 
-void FRigUnit_AnimRichCurve::Execute(const FRigUnitContext& Context)
+UE_RigUnit_AnimRichCurve_IMPLEMENT_STATIC_VIRTUAL_METHOD(void, Execute, const FRigUnitContext& Context)
 {
     DECLARE_SCOPE_HIERARCHICAL_COUNTER_RIGUNIT()
 }
diff --git a/Engine/Plugins/Experimental/ControlRig/Source/ControlRig/Private/Units/Animation/RigUnit_AnimRichCurve.h b/Engine/Plugins/Experimental/ControlRig/Source/ControlRig/Private/Units/Animation/RigUnit_AnimRichCurve.h
--- a/Engine/Plugins/Experimental/ControlRig/Source/ControlRig/Private/Units/Animation/RigUnit_AnimRichCurve.h
+++ b/Engine/Plugins/Experimental/ControlRig/Source/ControlRig/Private/Units/Animation/RigUnit_AnimRichCurve.h
@@ -21,6 +21,7 @@ struct FRigUnit_AnimRichCurve : public FRigUnit_AnimBase
 		Curve.GetRichCurve()->AddKey(1.f, 1.f);
 	}
 
+	STATIC_VIRTUAL_METHOD()
 	virtual void Execute(const FRigUnitContext& Context) override;
 
 	UPROPERTY(meta=(Input, Output, Constant))
diff --git a/Engine/Plugins/Experimental/ControlRig/Source/ControlRig/Private/Units/Animation/RigUnit_GetDeltaTime.cpp b/Engine/Plugins/Experimental/ControlRig/Source/ControlRig/Private/Units/Animation/RigUnit_GetDeltaTime.cpp
--- a/Engine/Plugins/Experimental/ControlRig/Source/ControlRig/Private/Units/Animation/RigUnit_GetDeltaTime.cpp
+++ b/Engine/Plugins/Experimental/ControlRig/Source/ControlRig/Private/Units/Animation/RigUnit_GetDeltaTime.cpp
@@ -3,7 +3,7 @@
 #include "Units/Animation/RigUnit_GetDeltaTime.h"
 #include "Units/RigUnitContext.h"
 
-void FRigUnit_GetDeltaTime::Execute(const FRigUnitContext& Context)
+UE_RigUnit_GetDeltaTime_IMPLEMENT_STATIC_VIRTUAL_METHOD(void, Execute, const FRigUnitContext& Context)
 {
     DECLARE_SCOPE_HIERARCHICAL_COUNTER_RIGUNIT()
 	Result = Context.DeltaTime;
diff --git a/Engine/Plugins/Experimental/ControlRig/Source/ControlRig/Private/Units/Animation/RigUnit_GetDeltaTime.h b/Engine/Plugins/Experimental/ControlRig/Source/ControlRig/Private/Units/Animation/RigUnit_GetDeltaTime.h
--- a/Engine/Plugins/Experimental/ControlRig/Source/ControlRig/Private/Units/Animation/RigUnit_GetDeltaTime.h
+++ b/Engine/Plugins/Experimental/ControlRig/Source/ControlRig/Private/Units/Animation/RigUnit_GetDeltaTime.h
@@ -18,6 +18,7 @@ struct FRigUnit_GetDeltaTime : public FRigUnit_AnimBase
 		Result = 0.f;
 	}
 
+	STATIC_VIRTUAL_METHOD()
 	virtual void Execute(const FRigUnitContext& Context) override;
 
 	UPROPERTY(meta=(Output))
diff --git a/Engine/Plugins/Experimental/ControlRig/Source/ControlRig/Private/Units/Animation/RigUnit_GetWorldTime.cpp b/Engine/Plugins/Experimental/ControlRig/Source/ControlRig/Private/Units/Animation/RigUnit_GetWorldTime.cpp
--- a/Engine/Plugins/Experimental/ControlRig/Source/ControlRig/Private/Units/Animation/RigUnit_GetWorldTime.cpp
+++ b/Engine/Plugins/Experimental/ControlRig/Source/ControlRig/Private/Units/Animation/RigUnit_GetWorldTime.cpp
@@ -3,7 +3,7 @@
 #include "Units/Animation/RigUnit_GetWorldTime.h"
 #include "Units/RigUnitContext.h"
 
-void FRigUnit_GetWorldTime::Execute(const FRigUnitContext& Context)
+UE_RigUnit_GetWorldTime_IMPLEMENT_STATIC_VIRTUAL_METHOD(void, Execute, const FRigUnitContext& Context)
 {
     DECLARE_SCOPE_HIERARCHICAL_COUNTER_RIGUNIT()
 	FDateTime Now = FDateTime::Now();
diff --git a/Engine/Plugins/Experimental/ControlRig/Source/ControlRig/Private/Units/Animation/RigUnit_GetWorldTime.h b/Engine/Plugins/Experimental/ControlRig/Source/ControlRig/Private/Units/Animation/RigUnit_GetWorldTime.h
--- a/Engine/Plugins/Experimental/ControlRig/Source/ControlRig/Private/Units/Animation/RigUnit_GetWorldTime.h
+++ b/Engine/Plugins/Experimental/ControlRig/Source/ControlRig/Private/Units/Animation/RigUnit_GetWorldTime.h
@@ -18,6 +18,7 @@ struct FRigUnit_GetWorldTime : public FRigUnit_AnimBase
 		Year = Month = Day = WeekDay = Hours = Minutes = Seconds = OverallSeconds = 0.f;
 	}
 
+	STATIC_VIRTUAL_METHOD()
 	virtual void Execute(const FRigUnitContext& Context) override;
 
 	UPROPERTY(meta = (Output))
diff --git a/Engine/Plugins/Experimental/ControlRig/Source/ControlRig/Private/Units/Control/RigUnit_Control.cpp b/Engine/Plugins/Experimental/ControlRig/Source/ControlRig/Private/Units/Control/RigUnit_Control.cpp
--- a/Engine/Plugins/Experimental/ControlRig/Source/ControlRig/Private/Units/Control/RigUnit_Control.cpp
+++ b/Engine/Plugins/Experimental/ControlRig/Source/ControlRig/Private/Units/Control/RigUnit_Control.cpp
@@ -3,7 +3,7 @@
 #include "Units/Control/RigUnit_Control.h"
 #include "Units/RigUnitContext.h"
 
-void FRigUnit_Control::Execute(const FRigUnitContext& Context)
+UE_RigUnit_Control_IMPLEMENT_STATIC_VIRTUAL_METHOD(void, Execute, const FRigUnitContext& Context)
 {
     DECLARE_SCOPE_HIERARCHICAL_COUNTER_RIGUNIT()
 	if (Context.State == EControlRigState::Init)
@@ -11,37 +11,62 @@ void FRigUnit_Control::Execute(const FRigUnitContext& Context)
 		Transform.FromFTransform(InitTransform);
 	}
 
-	Result = GetResultantTransform();
+	Result = StaticGetResultantTransform(Transform, Base, Filter);
 }
 
 FTransform FRigUnit_Control::GetResultantTransform() const
 {
-	return GetFilteredTransform().ToFTransform() * Base;
+	return StaticGetResultantTransform(Transform, Base, Filter);
+}
+
+FTransform FRigUnit_Control::StaticGetResultantTransform(const FEulerTransform& InTransform, const FTransform& InBase, const FTransformFilter& InFilter)
+{
+	return StaticGetFilteredTransform(InTransform, InFilter).ToFTransform() * InBase;
 }
 
 FMatrix FRigUnit_Control::GetResultantMatrix() const
 {
-	const FEulerTransform FilteredTransform = GetFilteredTransform();
-	return FScaleRotationTranslationMatrix(FilteredTransform.Scale, FilteredTransform.Rotation, FilteredTransform.Location) * Base.ToMatrixWithScale();
+	return StaticGetResultantMatrix(Transform, Base, Filter);
+}
+
+FMatrix FRigUnit_Control::StaticGetResultantMatrix(const FEulerTransform& InTransform, const FTransform& InBase, const FTransformFilter& InFilter)
+{
+	const FEulerTransform FilteredTransform = StaticGetFilteredTransform(InTransform, InFilter);
+	return FScaleRotationTranslationMatrix(FilteredTransform.Scale, FilteredTransform.Rotation, FilteredTransform.Location) * InBase.ToMatrixWithScale();
 }
 
 void FRigUnit_Control::SetResultantTransform(const FTransform& InResultantTransform)
 {
-	Transform.FromFTransform(InResultantTransform.GetRelativeTransform(Base));
+	StaticSetResultantTransform(InResultantTransform, Base, Transform);
+}
+
+void FRigUnit_Control::StaticSetResultantTransform(const FTransform& InResultantTransform, const FTransform& InBase, FEulerTransform& OutTransform)
+{
+	OutTransform.FromFTransform(InResultantTransform.GetRelativeTransform(InBase));
 }
 
 void FRigUnit_Control::SetResultantMatrix(const FMatrix& InResultantMatrix)
 {
-	const FMatrix RelativeTransform = InResultantMatrix * Base.ToMatrixWithScale().Inverse();
+	StaticSetResultantMatrix(InResultantMatrix, Base, Transform);
+}
+
+void FRigUnit_Control::StaticSetResultantMatrix(const FMatrix& InResultantMatrix, const FTransform& InBase, FEulerTransform& OutTransform)
+{
+	const FMatrix RelativeTransform = InResultantMatrix * InBase.ToMatrixWithScale().Inverse();
 
-	Transform.Location = RelativeTransform.GetOrigin();
-	Transform.Rotation = RelativeTransform.Rotator();
-	Transform.Scale = RelativeTransform.GetScaleVector();
+	OutTransform.Location = RelativeTransform.GetOrigin();
+	OutTransform.Rotation = RelativeTransform.Rotator();
+	OutTransform.Scale = RelativeTransform.GetScaleVector();
 }
 
 FEulerTransform FRigUnit_Control::GetFilteredTransform() const
 {
-	FEulerTransform FilteredTransform = Transform;
-	Filter.FilterTransform(FilteredTransform);
+	return StaticGetFilteredTransform(Transform, Filter);
+}
+
+FEulerTransform FRigUnit_Control::StaticGetFilteredTransform(const FEulerTransform& InTransform, const FTransformFilter& InFilter)
+{
+	FEulerTransform FilteredTransform = InTransform;
+	InFilter.FilterTransform(FilteredTransform);
 	return FilteredTransform;
-}
\ No newline at end of file
+}
diff --git a/Engine/Plugins/Experimental/ControlRig/Source/ControlRig/Private/Units/Debug/RigUnit_DebugBezier.cpp b/Engine/Plugins/Experimental/ControlRig/Source/ControlRig/Private/Units/Debug/RigUnit_DebugBezier.cpp
--- a/Engine/Plugins/Experimental/ControlRig/Source/ControlRig/Private/Units/Debug/RigUnit_DebugBezier.cpp
+++ b/Engine/Plugins/Experimental/ControlRig/Source/ControlRig/Private/Units/Debug/RigUnit_DebugBezier.cpp
@@ -3,7 +3,7 @@
 #include "Units/Debug/RigUnit_DebugBezier.h"
 #include "Units/RigUnitContext.h"
 
-void FRigUnit_DebugBezier::Execute(const FRigUnitContext& Context)
+UE_RigUnit_DebugBezier_IMPLEMENT_STATIC_VIRTUAL_METHOD(void, Execute, const FRigUnitContext& Context)
 {
     DECLARE_SCOPE_HIERARCHICAL_COUNTER_RIGUNIT()
 	if (Context.State == EControlRigState::Init)
diff --git a/Engine/Plugins/Experimental/ControlRig/Source/ControlRig/Private/Units/Debug/RigUnit_DebugBezier.h b/Engine/Plugins/Experimental/ControlRig/Source/ControlRig/Private/Units/Debug/RigUnit_DebugBezier.h
--- a/Engine/Plugins/Experimental/ControlRig/Source/ControlRig/Private/Units/Debug/RigUnit_DebugBezier.h
+++ b/Engine/Plugins/Experimental/ControlRig/Source/ControlRig/Private/Units/Debug/RigUnit_DebugBezier.h
@@ -23,6 +23,7 @@ struct FRigUnit_DebugBezier : public FRigUnit_DebugBaseMutable
 		bEnabled = true;
 	}
 
+	STATIC_VIRTUAL_METHOD()
 	virtual void Execute(const FRigUnitContext& Context) override;
 
 	UPROPERTY(meta = (Input))
diff --git a/Engine/Plugins/Experimental/ControlRig/Source/ControlRig/Private/Units/Debug/RigUnit_DebugHierarchy.cpp b/Engine/Plugins/Experimental/ControlRig/Source/ControlRig/Private/Units/Debug/RigUnit_DebugHierarchy.cpp
--- a/Engine/Plugins/Experimental/ControlRig/Source/ControlRig/Private/Units/Debug/RigUnit_DebugHierarchy.cpp
+++ b/Engine/Plugins/Experimental/ControlRig/Source/ControlRig/Private/Units/Debug/RigUnit_DebugHierarchy.cpp
@@ -3,7 +3,7 @@
 #include "Units/Debug/RigUnit_DebugHierarchy.h"
 #include "Units/RigUnitContext.h"
 
-void FRigUnit_DebugHierarchy::Execute(const FRigUnitContext& Context)
+UE_RigUnit_DebugHierarchy_IMPLEMENT_STATIC_VIRTUAL_METHOD(void, Execute, const FRigUnitContext& Context)
 {
     DECLARE_SCOPE_HIERARCHICAL_COUNTER_RIGUNIT()
 	if (Context.State == EControlRigState::Init)
diff --git a/Engine/Plugins/Experimental/ControlRig/Source/ControlRig/Private/Units/Debug/RigUnit_DebugHierarchy.h b/Engine/Plugins/Experimental/ControlRig/Source/ControlRig/Private/Units/Debug/RigUnit_DebugHierarchy.h
--- a/Engine/Plugins/Experimental/ControlRig/Source/ControlRig/Private/Units/Debug/RigUnit_DebugHierarchy.h
+++ b/Engine/Plugins/Experimental/ControlRig/Source/ControlRig/Private/Units/Debug/RigUnit_DebugHierarchy.h
@@ -21,6 +21,7 @@ struct FRigUnit_DebugHierarchy : public FRigUnit_DebugBaseMutable
 		bEnabled = true;
 	}
 
+	STATIC_VIRTUAL_METHOD()
 	virtual void Execute(const FRigUnitContext& Context) override;
 
 	UPROPERTY(meta = (Input))
diff --git a/Engine/Plugins/Experimental/ControlRig/Source/ControlRig/Private/Units/Debug/RigUnit_DebugLine.cpp b/Engine/Plugins/Experimental/ControlRig/Source/ControlRig/Private/Units/Debug/RigUnit_DebugLine.cpp
--- a/Engine/Plugins/Experimental/ControlRig/Source/ControlRig/Private/Units/Debug/RigUnit_DebugLine.cpp
+++ b/Engine/Plugins/Experimental/ControlRig/Source/ControlRig/Private/Units/Debug/RigUnit_DebugLine.cpp
@@ -3,7 +3,7 @@
 #include "Units/Debug/RigUnit_DebugLine.h"
 #include "Units/RigUnitContext.h"
 
-void FRigUnit_DebugLine::Execute(const FRigUnitContext& Context)
+UE_RigUnit_DebugLine_IMPLEMENT_STATIC_VIRTUAL_METHOD(void, Execute, const FRigUnitContext& Context)
 {
     DECLARE_SCOPE_HIERARCHICAL_COUNTER_RIGUNIT()
 	if (Context.State == EControlRigState::Init)
diff --git a/Engine/Plugins/Experimental/ControlRig/Source/ControlRig/Private/Units/Debug/RigUnit_DebugLine.h b/Engine/Plugins/Experimental/ControlRig/Source/ControlRig/Private/Units/Debug/RigUnit_DebugLine.h
--- a/Engine/Plugins/Experimental/ControlRig/Source/ControlRig/Private/Units/Debug/RigUnit_DebugLine.h
+++ b/Engine/Plugins/Experimental/ControlRig/Source/ControlRig/Private/Units/Debug/RigUnit_DebugLine.h
@@ -19,6 +19,7 @@ struct FRigUnit_DebugLine : public FRigUnit_DebugBaseMutable
 		bEnabled = true;
 	}
 
+	STATIC_VIRTUAL_METHOD()
 	virtual void Execute(const FRigUnitContext& Context) override;
 
 	UPROPERTY(meta = (Input))
diff --git a/Engine/Plugins/Experimental/ControlRig/Source/ControlRig/Private/Units/Debug/RigUnit_DebugLineStrip.cpp b/Engine/Plugins/Experimental/ControlRig/Source/ControlRig/Private/Units/Debug/RigUnit_DebugLineStrip.cpp
--- a/Engine/Plugins/Experimental/ControlRig/Source/ControlRig/Private/Units/Debug/RigUnit_DebugLineStrip.cpp
+++ b/Engine/Plugins/Experimental/ControlRig/Source/ControlRig/Private/Units/Debug/RigUnit_DebugLineStrip.cpp
@@ -3,7 +3,7 @@
 #include "Units/Debug/RigUnit_DebugLineStrip.h"
 #include "Units/RigUnitContext.h"
 
-void FRigUnit_DebugLineStrip::Execute(const FRigUnitContext& Context)
+UE_RigUnit_DebugLineStrip_IMPLEMENT_STATIC_VIRTUAL_METHOD(void, Execute, const FRigUnitContext& Context)
 {
     DECLARE_SCOPE_HIERARCHICAL_COUNTER_RIGUNIT()
 	if (Context.State == EControlRigState::Init)
diff --git a/Engine/Plugins/Experimental/ControlRig/Source/ControlRig/Private/Units/Debug/RigUnit_DebugLineStrip.h b/Engine/Plugins/Experimental/ControlRig/Source/ControlRig/Private/Units/Debug/RigUnit_DebugLineStrip.h
--- a/Engine/Plugins/Experimental/ControlRig/Source/ControlRig/Private/Units/Debug/RigUnit_DebugLineStrip.h
+++ b/Engine/Plugins/Experimental/ControlRig/Source/ControlRig/Private/Units/Debug/RigUnit_DebugLineStrip.h
@@ -18,6 +18,7 @@ struct FRigUnit_DebugLineStrip : public FRigUnit_DebugBaseMutable
 		bEnabled = true;
 	}
 
+	STATIC_VIRTUAL_METHOD()
 	virtual void Execute(const FRigUnitContext& Context) override;
 
 	UPROPERTY(meta = (Input))
diff --git a/Engine/Plugins/Experimental/ControlRig/Source/ControlRig/Private/Units/Debug/RigUnit_DebugPoint.cpp b/Engine/Plugins/Experimental/ControlRig/Source/ControlRig/Private/Units/Debug/RigUnit_DebugPoint.cpp
--- a/Engine/Plugins/Experimental/ControlRig/Source/ControlRig/Private/Units/Debug/RigUnit_DebugPoint.cpp
+++ b/Engine/Plugins/Experimental/ControlRig/Source/ControlRig/Private/Units/Debug/RigUnit_DebugPoint.cpp
@@ -3,7 +3,7 @@
 #include "Units/Debug/RigUnit_DebugPoint.h"
 #include "Units/RigUnitContext.h"
 
-void FRigUnit_DebugPoint::Execute(const FRigUnitContext& Context)
+UE_RigUnit_DebugPoint_IMPLEMENT_STATIC_VIRTUAL_METHOD(void, Execute, const FRigUnitContext& Context)
 {
     DECLARE_SCOPE_HIERARCHICAL_COUNTER_RIGUNIT()
 	if (Context.State == EControlRigState::Init)
@@ -40,7 +40,7 @@ void FRigUnit_DebugPoint::Execute(const FRigUnitContext& Context)
 	}
 }
 
-void FRigUnit_DebugPointMutable::Execute(const FRigUnitContext& Context)
+UE_RigUnit_DebugPointMutable_IMPLEMENT_STATIC_VIRTUAL_METHOD(void, Execute, const FRigUnitContext& Context)
 {
     DECLARE_SCOPE_HIERARCHICAL_COUNTER_RIGUNIT()
 	if (Context.State == EControlRigState::Init)
diff --git a/Engine/Plugins/Experimental/ControlRig/Source/ControlRig/Private/Units/Debug/RigUnit_DebugPoint.h b/Engine/Plugins/Experimental/ControlRig/Source/ControlRig/Private/Units/Debug/RigUnit_DebugPoint.h
--- a/Engine/Plugins/Experimental/ControlRig/Source/ControlRig/Private/Units/Debug/RigUnit_DebugPoint.h
+++ b/Engine/Plugins/Experimental/ControlRig/Source/ControlRig/Private/Units/Debug/RigUnit_DebugPoint.h
@@ -34,6 +34,7 @@ struct FRigUnit_DebugPoint : public FRigUnit_DebugBase
 		bEnabled = true;
 	}
 
+	STATIC_VIRTUAL_METHOD()
 	virtual void Execute(const FRigUnitContext& Context) override;
 
 	UPROPERTY(meta = (Input, Output))
@@ -77,6 +78,7 @@ struct FRigUnit_DebugPointMutable : public FRigUnit_DebugBaseMutable
 		bEnabled = true;
 	}
 
+	STATIC_VIRTUAL_METHOD()
 	virtual void Execute(const FRigUnitContext& Context) override;
 
 	UPROPERTY(meta = (Input))
diff --git a/Engine/Plugins/Experimental/ControlRig/Source/ControlRig/Private/Units/Debug/RigUnit_DebugPrimitives.cpp b/Engine/Plugins/Experimental/ControlRig/Source/ControlRig/Private/Units/Debug/RigUnit_DebugPrimitives.cpp
--- a/Engine/Plugins/Experimental/ControlRig/Source/ControlRig/Private/Units/Debug/RigUnit_DebugPrimitives.cpp
+++ b/Engine/Plugins/Experimental/ControlRig/Source/ControlRig/Private/Units/Debug/RigUnit_DebugPrimitives.cpp
@@ -3,7 +3,7 @@
 #include "Units/Debug/RigUnit_DebugPrimitives.h"
 #include "Units/RigUnitContext.h"
 
-void FRigUnit_DebugRectangle::Execute(const FRigUnitContext& Context)
+UE_RigUnit_DebugRectangle_IMPLEMENT_STATIC_VIRTUAL_METHOD(void, Execute, const FRigUnitContext& Context)
 {
     DECLARE_SCOPE_HIERARCHICAL_COUNTER_RIGUNIT()
 	if (Context.State == EControlRigState::Init)
@@ -25,7 +25,7 @@ void FRigUnit_DebugRectangle::Execute(const FRigUnitContext& Context)
 	Context.DrawInterface->DrawRectangle(WorldOffset, DrawTransform, Scale, Color, Thickness);
 }
 
-void FRigUnit_DebugArc::Execute(const FRigUnitContext& Context)
+UE_RigUnit_DebugArc_IMPLEMENT_STATIC_VIRTUAL_METHOD(void, Execute, const FRigUnitContext& Context)
 {
     DECLARE_SCOPE_HIERARCHICAL_COUNTER_RIGUNIT()
 	if (Context.State == EControlRigState::Init)
diff --git a/Engine/Plugins/Experimental/ControlRig/Source/ControlRig/Private/Units/Debug/RigUnit_DebugPrimitives.h b/Engine/Plugins/Experimental/ControlRig/Source/ControlRig/Private/Units/Debug/RigUnit_DebugPrimitives.h
--- a/Engine/Plugins/Experimental/ControlRig/Source/ControlRig/Private/Units/Debug/RigUnit_DebugPrimitives.h
+++ b/Engine/Plugins/Experimental/ControlRig/Source/ControlRig/Private/Units/Debug/RigUnit_DebugPrimitives.h
@@ -20,6 +20,7 @@ struct FRigUnit_DebugRectangle : public FRigUnit_DebugBaseMutable
 		bEnabled = true;
 	}
 
+	STATIC_VIRTUAL_METHOD()
 	virtual void Execute(const FRigUnitContext& Context) override;
 
 	UPROPERTY(meta = (Input))
@@ -61,6 +62,7 @@ struct FRigUnit_DebugArc : public FRigUnit_DebugBaseMutable
 		bEnabled = true;
 	}
 
+	STATIC_VIRTUAL_METHOD()
 	virtual void Execute(const FRigUnitContext& Context) override;
 
 	UPROPERTY(meta = (Input))
diff --git a/Engine/Plugins/Experimental/ControlRig/Source/ControlRig/Private/Units/Debug/RigUnit_DebugTransform.cpp b/Engine/Plugins/Experimental/ControlRig/Source/ControlRig/Private/Units/Debug/RigUnit_DebugTransform.cpp
--- a/Engine/Plugins/Experimental/ControlRig/Source/ControlRig/Private/Units/Debug/RigUnit_DebugTransform.cpp
+++ b/Engine/Plugins/Experimental/ControlRig/Source/ControlRig/Private/Units/Debug/RigUnit_DebugTransform.cpp
@@ -3,7 +3,7 @@
 #include "Units/Debug/RigUnit_DebugTransform.h"
 #include "Units/RigUnitContext.h"
 
-void FRigUnit_DebugTransform::Execute(const FRigUnitContext& Context)
+UE_RigUnit_DebugTransform_IMPLEMENT_STATIC_VIRTUAL_METHOD(void, Execute, const FRigUnitContext& Context)
 {
     DECLARE_SCOPE_HIERARCHICAL_COUNTER_RIGUNIT()
 	if (Context.State == EControlRigState::Init)
@@ -43,7 +43,7 @@ void FRigUnit_DebugTransform::Execute(const FRigUnitContext& Context)
 	}
 }
 
-void FRigUnit_DebugTransformMutable::Execute(const FRigUnitContext& Context)
+UE_RigUnit_DebugTransformMutable_IMPLEMENT_STATIC_VIRTUAL_METHOD(void, Execute, const FRigUnitContext& Context)
 {
     DECLARE_SCOPE_HIERARCHICAL_COUNTER_RIGUNIT()
 	if (Context.State == EControlRigState::Init)
@@ -83,7 +83,7 @@ void FRigUnit_DebugTransformMutable::Execute(const FRigUnitContext& Context)
 	}
 }
 
-void FRigUnit_DebugTransformArrayMutable::Execute(const FRigUnitContext& Context)
+UE_RigUnit_DebugTransformArrayMutable_IMPLEMENT_STATIC_VIRTUAL_METHOD(void, Execute, const FRigUnitContext& Context)
 {
     DECLARE_SCOPE_HIERARCHICAL_COUNTER_RIGUNIT()
 	if (Context.State == EControlRigState::Init)
diff --git a/Engine/Plugins/Experimental/ControlRig/Source/ControlRig/Private/Units/Debug/RigUnit_DebugTransform.h b/Engine/Plugins/Experimental/ControlRig/Source/ControlRig/Private/Units/Debug/RigUnit_DebugTransform.h
--- a/Engine/Plugins/Experimental/ControlRig/Source/ControlRig/Private/Units/Debug/RigUnit_DebugTransform.h
+++ b/Engine/Plugins/Experimental/ControlRig/Source/ControlRig/Private/Units/Debug/RigUnit_DebugTransform.h
@@ -36,6 +36,7 @@ struct FRigUnit_DebugTransform : public FRigUnit_DebugBase
 		bEnabled = true;
 	}
 
+	STATIC_VIRTUAL_METHOD()
 	virtual void Execute(const FRigUnitContext& Context) override;
 
 	UPROPERTY(meta = (Input, Output))
@@ -78,6 +79,7 @@ struct FRigUnit_DebugTransformMutable : public FRigUnit_DebugBaseMutable
 		bEnabled = true;
 	}
 
+	STATIC_VIRTUAL_METHOD()
 	virtual void Execute(const FRigUnitContext& Context) override;
 
 	UPROPERTY(meta = (Input))
@@ -120,6 +122,7 @@ struct FRigUnit_DebugTransformArrayMutable : public FRigUnit_DebugBaseMutable
 		bEnabled = true;
 	}
 
+	STATIC_VIRTUAL_METHOD()
 	virtual void Execute(const FRigUnitContext& Context) override;
 
 	UPROPERTY(meta = (Input))
diff --git a/Engine/Plugins/Experimental/ControlRig/Source/ControlRig/Private/Units/Debug/RigUnit_ProfilingBracket.cpp b/Engine/Plugins/Experimental/ControlRig/Source/ControlRig/Private/Units/Debug/RigUnit_ProfilingBracket.cpp
--- a/Engine/Plugins/Experimental/ControlRig/Source/ControlRig/Private/Units/Debug/RigUnit_ProfilingBracket.cpp
+++ b/Engine/Plugins/Experimental/ControlRig/Source/ControlRig/Private/Units/Debug/RigUnit_ProfilingBracket.cpp
@@ -4,7 +4,7 @@
 #include "Units/RigUnitContext.h"
 #include "KismetAnimationLibrary.h"
 
-void FRigUnit_StartProfilingTimer::Execute(const FRigUnitContext& Context)
+UE_RigUnit_StartProfilingTimer_IMPLEMENT_STATIC_VIRTUAL_METHOD(void, Execute, const FRigUnitContext& Context)
 {
     DECLARE_SCOPE_HIERARCHICAL_COUNTER_RIGUNIT()
 	if (Context.State == EControlRigState::Update)
@@ -13,7 +13,7 @@ void FRigUnit_StartProfilingTimer::Execute(const FRigUnitContext& Context)
 	}
 }
 
-void FRigUnit_EndProfilingTimer::Execute(const FRigUnitContext& Context)
+UE_RigUnit_EndProfilingTimer_IMPLEMENT_STATIC_VIRTUAL_METHOD(void, Execute, const FRigUnitContext& Context)
 {
     DECLARE_SCOPE_HIERARCHICAL_COUNTER_RIGUNIT()
 	if (Context.State == EControlRigState::Init)
diff --git a/Engine/Plugins/Experimental/ControlRig/Source/ControlRig/Private/Units/Debug/RigUnit_ProfilingBracket.h b/Engine/Plugins/Experimental/ControlRig/Source/ControlRig/Private/Units/Debug/RigUnit_ProfilingBracket.h
--- a/Engine/Plugins/Experimental/ControlRig/Source/ControlRig/Private/Units/Debug/RigUnit_ProfilingBracket.h
+++ b/Engine/Plugins/Experimental/ControlRig/Source/ControlRig/Private/Units/Debug/RigUnit_ProfilingBracket.h
@@ -14,6 +14,7 @@ struct FRigUnit_StartProfilingTimer : public FRigUnit_DebugBaseMutable
 	{
 	}
 
+	STATIC_VIRTUAL_METHOD()
 	virtual void Execute(const FRigUnitContext& Context) override;
 };
 
@@ -30,6 +31,7 @@ struct FRigUnit_EndProfilingTimer : public FRigUnit_DebugBaseMutable
 		Prefix = TEXT("Timer");
 	}
 
+	STATIC_VIRTUAL_METHOD()
 	virtual void Execute(const FRigUnitContext& Context) override;
 
 	UPROPERTY(meta = (Input, Constant))
diff --git a/Engine/Plugins/Experimental/ControlRig/Source/ControlRig/Private/Units/Execution/RigUnit_BeginExecution.cpp b/Engine/Plugins/Experimental/ControlRig/Source/ControlRig/Private/Units/Execution/RigUnit_BeginExecution.cpp
--- a/Engine/Plugins/Experimental/ControlRig/Source/ControlRig/Private/Units/Execution/RigUnit_BeginExecution.cpp
+++ b/Engine/Plugins/Experimental/ControlRig/Source/ControlRig/Private/Units/Execution/RigUnit_BeginExecution.cpp
@@ -3,7 +3,7 @@
 #include "RigUnit_BeginExecution.h"
 #include "Units/RigUnitContext.h"
 
-void FRigUnit_BeginExecution::Execute(const FRigUnitContext& Context)
+UE_RigUnit_BeginExecution_IMPLEMENT_STATIC_VIRTUAL_METHOD(void, Execute, const FRigUnitContext& Context)
 {
     DECLARE_SCOPE_HIERARCHICAL_COUNTER_RIGUNIT()
 	ExecuteContext.Hierarchy = (FRigHierarchyContainer*)Context.Hierarchy;
diff --git a/Engine/Plugins/Experimental/ControlRig/Source/ControlRig/Private/Units/Execution/RigUnit_BeginExecution.h b/Engine/Plugins/Experimental/ControlRig/Source/ControlRig/Private/Units/Execution/RigUnit_BeginExecution.h
--- a/Engine/Plugins/Experimental/ControlRig/Source/ControlRig/Private/Units/Execution/RigUnit_BeginExecution.h
+++ b/Engine/Plugins/Experimental/ControlRig/Source/ControlRig/Private/Units/Execution/RigUnit_BeginExecution.h
@@ -11,6 +11,7 @@ struct FRigUnit_BeginExecution : public FRigUnit
 {
 	GENERATED_BODY()
 
+	STATIC_VIRTUAL_METHOD()
 	virtual void Execute(const FRigUnitContext& Context) override;
 
 	// The execution result
diff --git a/Engine/Plugins/Experimental/ControlRig/Source/ControlRig/Private/Units/Hierarchy/RigUnit_AddBoneTransform.cpp b/Engine/Plugins/Experimental/ControlRig/Source/ControlRig/Private/Units/Hierarchy/RigUnit_AddBoneTransform.cpp
--- a/Engine/Plugins/Experimental/ControlRig/Source/ControlRig/Private/Units/Hierarchy/RigUnit_AddBoneTransform.cpp
+++ b/Engine/Plugins/Experimental/ControlRig/Source/ControlRig/Private/Units/Hierarchy/RigUnit_AddBoneTransform.cpp
@@ -8,7 +8,7 @@ FString FRigUnit_AddBoneTransform::GetUnitLabel() const
 	return FString::Printf(TEXT("Offset Transform %s"), *Bone.ToString());
 }
 
-void FRigUnit_AddBoneTransform::Execute(const FRigUnitContext& Context)
+UE_RigUnit_AddBoneTransform_IMPLEMENT_STATIC_VIRTUAL_METHOD(void, Execute, const FRigUnitContext& Context)
 {
     DECLARE_SCOPE_HIERARCHICAL_COUNTER_RIGUNIT()
 	FRigBoneHierarchy* Hierarchy = ExecuteContext.GetBones();
diff --git a/Engine/Plugins/Experimental/ControlRig/Source/ControlRig/Private/Units/Hierarchy/RigUnit_AddBoneTransform.h b/Engine/Plugins/Experimental/ControlRig/Source/ControlRig/Private/Units/Hierarchy/RigUnit_AddBoneTransform.h
--- a/Engine/Plugins/Experimental/ControlRig/Source/ControlRig/Private/Units/Hierarchy/RigUnit_AddBoneTransform.h
+++ b/Engine/Plugins/Experimental/ControlRig/Source/ControlRig/Private/Units/Hierarchy/RigUnit_AddBoneTransform.h
@@ -21,6 +21,7 @@ struct FRigUnit_AddBoneTransform : public FRigUnitMutable
 	{}
 
 	virtual FString GetUnitLabel() const override;
+	STATIC_VIRTUAL_METHOD()
 	virtual void Execute(const FRigUnitContext& Context) override;
 
 	/**
diff --git a/Engine/Plugins/Experimental/ControlRig/Source/ControlRig/Private/Units/Hierarchy/RigUnit_GetBoneTransform.cpp b/Engine/Plugins/Experimental/ControlRig/Source/ControlRig/Private/Units/Hierarchy/RigUnit_GetBoneTransform.cpp
--- a/Engine/Plugins/Experimental/ControlRig/Source/ControlRig/Private/Units/Hierarchy/RigUnit_GetBoneTransform.cpp
+++ b/Engine/Plugins/Experimental/ControlRig/Source/ControlRig/Private/Units/Hierarchy/RigUnit_GetBoneTransform.cpp
@@ -8,7 +8,7 @@ FString FRigUnit_GetBoneTransform::GetUnitLabel() const
 	return FString::Printf(TEXT("Get Transform %s"), *Bone.ToString());
 }
 
-void FRigUnit_GetBoneTransform::Execute(const FRigUnitContext& Context)
+UE_RigUnit_GetBoneTransform_IMPLEMENT_STATIC_VIRTUAL_METHOD(void, Execute, const FRigUnitContext& Context)
 {
     DECLARE_SCOPE_HIERARCHICAL_COUNTER_RIGUNIT()
 	const FRigBoneHierarchy* Hierarchy = Context.GetBones();
diff --git a/Engine/Plugins/Experimental/ControlRig/Source/ControlRig/Private/Units/Hierarchy/RigUnit_GetBoneTransform.h b/Engine/Plugins/Experimental/ControlRig/Source/ControlRig/Private/Units/Hierarchy/RigUnit_GetBoneTransform.h
--- a/Engine/Plugins/Experimental/ControlRig/Source/ControlRig/Private/Units/Hierarchy/RigUnit_GetBoneTransform.h
+++ b/Engine/Plugins/Experimental/ControlRig/Source/ControlRig/Private/Units/Hierarchy/RigUnit_GetBoneTransform.h
@@ -19,6 +19,7 @@ struct FRigUnit_GetBoneTransform : public FRigUnit
 	{}
 
 	virtual FString GetUnitLabel() const override;
+	STATIC_VIRTUAL_METHOD()
 	virtual void Execute(const FRigUnitContext& Context) override;
 
 	/**
diff --git a/Engine/Plugins/Experimental/ControlRig/Source/ControlRig/Private/Units/Hierarchy/RigUnit_GetCurveValue.cpp b/Engine/Plugins/Experimental/ControlRig/Source/ControlRig/Private/Units/Hierarchy/RigUnit_GetCurveValue.cpp
--- a/Engine/Plugins/Experimental/ControlRig/Source/ControlRig/Private/Units/Hierarchy/RigUnit_GetCurveValue.cpp
+++ b/Engine/Plugins/Experimental/ControlRig/Source/ControlRig/Private/Units/Hierarchy/RigUnit_GetCurveValue.cpp
@@ -8,7 +8,7 @@ FString FRigUnit_GetCurveValue::GetUnitLabel() const
 	return FString::Printf(TEXT("Get Curve %s"), *Curve.ToString());
 }
 
-void FRigUnit_GetCurveValue::Execute(const FRigUnitContext& Context)
+UE_RigUnit_GetCurveValue_IMPLEMENT_STATIC_VIRTUAL_METHOD(void, Execute, const FRigUnitContext& Context)
 {
     DECLARE_SCOPE_HIERARCHICAL_COUNTER_RIGUNIT()
 	const FRigCurveContainer* CurveContainer = Context.GetCurves();
diff --git a/Engine/Plugins/Experimental/ControlRig/Source/ControlRig/Private/Units/Hierarchy/RigUnit_GetCurveValue.h b/Engine/Plugins/Experimental/ControlRig/Source/ControlRig/Private/Units/Hierarchy/RigUnit_GetCurveValue.h
--- a/Engine/Plugins/Experimental/ControlRig/Source/ControlRig/Private/Units/Hierarchy/RigUnit_GetCurveValue.h
+++ b/Engine/Plugins/Experimental/ControlRig/Source/ControlRig/Private/Units/Hierarchy/RigUnit_GetCurveValue.h
@@ -20,6 +20,7 @@ struct FRigUnit_GetCurveValue : public FRigUnit
 	{}
 
 	virtual FString GetUnitLabel() const override;
+	STATIC_VIRTUAL_METHOD()
 	virtual void Execute(const FRigUnitContext& Context) override;
 
 	/**
diff --git a/Engine/Plugins/Experimental/ControlRig/Source/ControlRig/Private/Units/Hierarchy/RigUnit_GetInitialBoneTransform.cpp b/Engine/Plugins/Experimental/ControlRig/Source/ControlRig/Private/Units/Hierarchy/RigUnit_GetInitialBoneTransform.cpp
--- a/Engine/Plugins/Experimental/ControlRig/Source/ControlRig/Private/Units/Hierarchy/RigUnit_GetInitialBoneTransform.cpp
+++ b/Engine/Plugins/Experimental/ControlRig/Source/ControlRig/Private/Units/Hierarchy/RigUnit_GetInitialBoneTransform.cpp
@@ -8,7 +8,7 @@ FString FRigUnit_GetInitialBoneTransform::GetUnitLabel() const
 	return FString::Printf(TEXT("Get Initial Transform %s"), *Bone.ToString());
 }
 
-void FRigUnit_GetInitialBoneTransform::Execute(const FRigUnitContext& Context)
+UE_RigUnit_GetInitialBoneTransform_IMPLEMENT_STATIC_VIRTUAL_METHOD(void, Execute, const FRigUnitContext& Context)
 {
     DECLARE_SCOPE_HIERARCHICAL_COUNTER_RIGUNIT()
 	const FRigBoneHierarchy* Hierarchy = Context.GetBones();
diff --git a/Engine/Plugins/Experimental/ControlRig/Source/ControlRig/Private/Units/Hierarchy/RigUnit_GetInitialBoneTransform.h b/Engine/Plugins/Experimental/ControlRig/Source/ControlRig/Private/Units/Hierarchy/RigUnit_GetInitialBoneTransform.h
--- a/Engine/Plugins/Experimental/ControlRig/Source/ControlRig/Private/Units/Hierarchy/RigUnit_GetInitialBoneTransform.h
+++ b/Engine/Plugins/Experimental/ControlRig/Source/ControlRig/Private/Units/Hierarchy/RigUnit_GetInitialBoneTransform.h
@@ -19,6 +19,7 @@ struct FRigUnit_GetInitialBoneTransform : public FRigUnit
 	{}
 
 	virtual FString GetUnitLabel() const override;
+	STATIC_VIRTUAL_METHOD()
 	virtual void Execute(const FRigUnitContext& Context) override;
 
 	/**
diff --git a/Engine/Plugins/Experimental/ControlRig/Source/ControlRig/Private/Units/Hierarchy/RigUnit_GetRelativeBoneTransform.cpp b/Engine/Plugins/Experimental/ControlRig/Source/ControlRig/Private/Units/Hierarchy/RigUnit_GetRelativeBoneTransform.cpp
--- a/Engine/Plugins/Experimental/ControlRig/Source/ControlRig/Private/Units/Hierarchy/RigUnit_GetRelativeBoneTransform.cpp
+++ b/Engine/Plugins/Experimental/ControlRig/Source/ControlRig/Private/Units/Hierarchy/RigUnit_GetRelativeBoneTransform.cpp
@@ -8,7 +8,7 @@ FString FRigUnit_GetRelativeBoneTransform::GetUnitLabel() const
 	return FString::Printf(TEXT("Get Relative Transform %s"), *Bone.ToString());
 }
 
-void FRigUnit_GetRelativeBoneTransform::Execute(const FRigUnitContext& Context)
+UE_RigUnit_GetRelativeBoneTransform_IMPLEMENT_STATIC_VIRTUAL_METHOD(void, Execute, const FRigUnitContext& Context)
 {
     DECLARE_SCOPE_HIERARCHICAL_COUNTER_RIGUNIT()
 	const FRigBoneHierarchy* Hierarchy = Context.GetBones();
diff --git a/Engine/Plugins/Experimental/ControlRig/Source/ControlRig/Private/Units/Hierarchy/RigUnit_GetRelativeBoneTransform.h b/Engine/Plugins/Experimental/ControlRig/Source/ControlRig/Private/Units/Hierarchy/RigUnit_GetRelativeBoneTransform.h
--- a/Engine/Plugins/Experimental/ControlRig/Source/ControlRig/Private/Units/Hierarchy/RigUnit_GetRelativeBoneTransform.h
+++ b/Engine/Plugins/Experimental/ControlRig/Source/ControlRig/Private/Units/Hierarchy/RigUnit_GetRelativeBoneTransform.h
@@ -19,6 +19,7 @@ struct FRigUnit_GetRelativeBoneTransform : public FRigUnit
 	{}
 
 	virtual FString GetUnitLabel() const override;
+	STATIC_VIRTUAL_METHOD()
 	virtual void Execute(const FRigUnitContext& Context) override;
 
 	/**
diff --git a/Engine/Plugins/Experimental/ControlRig/Source/ControlRig/Private/Units/Hierarchy/RigUnit_SetBoneRotation.cpp b/Engine/Plugins/Experimental/ControlRig/Source/ControlRig/Private/Units/Hierarchy/RigUnit_SetBoneRotation.cpp
--- a/Engine/Plugins/Experimental/ControlRig/Source/ControlRig/Private/Units/Hierarchy/RigUnit_SetBoneRotation.cpp
+++ b/Engine/Plugins/Experimental/ControlRig/Source/ControlRig/Private/Units/Hierarchy/RigUnit_SetBoneRotation.cpp
@@ -8,7 +8,7 @@ FString FRigUnit_SetBoneRotation::GetUnitLabel() const
 	return FString::Printf(TEXT("Set Rotation %s"), *Bone.ToString());
 }
 
-void FRigUnit_SetBoneRotation::Execute(const FRigUnitContext& Context)
+UE_RigUnit_SetBoneRotation_IMPLEMENT_STATIC_VIRTUAL_METHOD(void, Execute, const FRigUnitContext& Context)
 {
     DECLARE_SCOPE_HIERARCHICAL_COUNTER_RIGUNIT()
 	FRigBoneHierarchy* Hierarchy = ExecuteContext.GetBones();
diff --git a/Engine/Plugins/Experimental/ControlRig/Source/ControlRig/Private/Units/Hierarchy/RigUnit_SetBoneRotation.h b/Engine/Plugins/Experimental/ControlRig/Source/ControlRig/Private/Units/Hierarchy/RigUnit_SetBoneRotation.h
--- a/Engine/Plugins/Experimental/ControlRig/Source/ControlRig/Private/Units/Hierarchy/RigUnit_SetBoneRotation.h
+++ b/Engine/Plugins/Experimental/ControlRig/Source/ControlRig/Private/Units/Hierarchy/RigUnit_SetBoneRotation.h
@@ -22,6 +22,7 @@ struct FRigUnit_SetBoneRotation : public FRigUnitMutable
 	{}
 
 	virtual FString GetUnitLabel() const override;
+	STATIC_VIRTUAL_METHOD()
 	virtual void Execute(const FRigUnitContext& Context) override;
 
 	/**
diff --git a/Engine/Plugins/Experimental/ControlRig/Source/ControlRig/Private/Units/Hierarchy/RigUnit_SetBoneTransform.cpp b/Engine/Plugins/Experimental/ControlRig/Source/ControlRig/Private/Units/Hierarchy/RigUnit_SetBoneTransform.cpp
--- a/Engine/Plugins/Experimental/ControlRig/Source/ControlRig/Private/Units/Hierarchy/RigUnit_SetBoneTransform.cpp
+++ b/Engine/Plugins/Experimental/ControlRig/Source/ControlRig/Private/Units/Hierarchy/RigUnit_SetBoneTransform.cpp
@@ -8,7 +8,7 @@ FString FRigUnit_SetBoneTransform::GetUnitLabel() const
 	return FString::Printf(TEXT("Set Transform %s"), *Bone.ToString());
 }
 
-void FRigUnit_SetBoneTransform::Execute(const FRigUnitContext& Context)
+UE_RigUnit_SetBoneTransform_IMPLEMENT_STATIC_VIRTUAL_METHOD(void, Execute, const FRigUnitContext& Context)
 {
     DECLARE_SCOPE_HIERARCHICAL_COUNTER_RIGUNIT()
 	FRigBoneHierarchy* Hierarchy = ExecuteContext.GetBones();
diff --git a/Engine/Plugins/Experimental/ControlRig/Source/ControlRig/Private/Units/Hierarchy/RigUnit_SetBoneTransform.h b/Engine/Plugins/Experimental/ControlRig/Source/ControlRig/Private/Units/Hierarchy/RigUnit_SetBoneTransform.h
--- a/Engine/Plugins/Experimental/ControlRig/Source/ControlRig/Private/Units/Hierarchy/RigUnit_SetBoneTransform.h
+++ b/Engine/Plugins/Experimental/ControlRig/Source/ControlRig/Private/Units/Hierarchy/RigUnit_SetBoneTransform.h
@@ -21,6 +21,8 @@ struct FRigUnit_SetBoneTransform : public FRigUnitMutable
 	{}
 
 	virtual FString GetUnitLabel() const override;
+
+	STATIC_VIRTUAL_METHOD()
 	virtual void Execute(const FRigUnitContext& Context) override;
 
 	/**
diff --git a/Engine/Plugins/Experimental/ControlRig/Source/ControlRig/Private/Units/Hierarchy/RigUnit_SetBoneTranslation.cpp b/Engine/Plugins/Experimental/ControlRig/Source/ControlRig/Private/Units/Hierarchy/RigUnit_SetBoneTranslation.cpp
--- a/Engine/Plugins/Experimental/ControlRig/Source/ControlRig/Private/Units/Hierarchy/RigUnit_SetBoneTranslation.cpp
+++ b/Engine/Plugins/Experimental/ControlRig/Source/ControlRig/Private/Units/Hierarchy/RigUnit_SetBoneTranslation.cpp
@@ -8,7 +8,7 @@ FString FRigUnit_SetBoneTranslation::GetUnitLabel() const
 	return FString::Printf(TEXT("Set Translation %s"), *Bone.ToString());
 }
 
-void FRigUnit_SetBoneTranslation::Execute(const FRigUnitContext& Context)
+UE_RigUnit_SetBoneTranslation_IMPLEMENT_STATIC_VIRTUAL_METHOD(void, Execute, const FRigUnitContext& Context)
 {
     DECLARE_SCOPE_HIERARCHICAL_COUNTER_RIGUNIT()
 	FRigBoneHierarchy* Hierarchy = ExecuteContext.GetBones();
diff --git a/Engine/Plugins/Experimental/ControlRig/Source/ControlRig/Private/Units/Hierarchy/RigUnit_SetBoneTranslation.h b/Engine/Plugins/Experimental/ControlRig/Source/ControlRig/Private/Units/Hierarchy/RigUnit_SetBoneTranslation.h
--- a/Engine/Plugins/Experimental/ControlRig/Source/ControlRig/Private/Units/Hierarchy/RigUnit_SetBoneTranslation.h
+++ b/Engine/Plugins/Experimental/ControlRig/Source/ControlRig/Private/Units/Hierarchy/RigUnit_SetBoneTranslation.h
@@ -22,6 +22,7 @@ struct FRigUnit_SetBoneTranslation : public FRigUnitMutable
 	{}
 
 	virtual FString GetUnitLabel() const override;
+	STATIC_VIRTUAL_METHOD()
 	virtual void Execute(const FRigUnitContext& Context) override;
 
 	/**
diff --git a/Engine/Plugins/Experimental/ControlRig/Source/ControlRig/Private/Units/Hierarchy/RigUnit_SetCurveValue.cpp b/Engine/Plugins/Experimental/ControlRig/Source/ControlRig/Private/Units/Hierarchy/RigUnit_SetCurveValue.cpp
--- a/Engine/Plugins/Experimental/ControlRig/Source/ControlRig/Private/Units/Hierarchy/RigUnit_SetCurveValue.cpp
+++ b/Engine/Plugins/Experimental/ControlRig/Source/ControlRig/Private/Units/Hierarchy/RigUnit_SetCurveValue.cpp
@@ -8,7 +8,7 @@ FString FRigUnit_SetCurveValue::GetUnitLabel() const
 	return FString::Printf(TEXT("Set Curve [%s]"), *Curve.ToString());
 }
 
-void FRigUnit_SetCurveValue::Execute(const FRigUnitContext& Context)
+UE_RigUnit_SetCurveValue_IMPLEMENT_STATIC_VIRTUAL_METHOD(void, Execute, const FRigUnitContext& Context)
 {
     DECLARE_SCOPE_HIERARCHICAL_COUNTER_RIGUNIT()
 	FRigCurveContainer* CurveContainer = ExecuteContext.GetCurves();
diff --git a/Engine/Plugins/Experimental/ControlRig/Source/ControlRig/Private/Units/Hierarchy/RigUnit_SetCurveValue.h b/Engine/Plugins/Experimental/ControlRig/Source/ControlRig/Private/Units/Hierarchy/RigUnit_SetCurveValue.h
--- a/Engine/Plugins/Experimental/ControlRig/Source/ControlRig/Private/Units/Hierarchy/RigUnit_SetCurveValue.h
+++ b/Engine/Plugins/Experimental/ControlRig/Source/ControlRig/Private/Units/Hierarchy/RigUnit_SetCurveValue.h
@@ -20,6 +20,7 @@ struct FRigUnit_SetCurveValue : public FRigUnitMutable
 	{}
 
 	virtual FString GetUnitLabel() const override;
+	STATIC_VIRTUAL_METHOD()
 	virtual void Execute(const FRigUnitContext& Context) override;
 
 	/**
diff --git a/Engine/Plugins/Experimental/ControlRig/Source/ControlRig/Private/Units/Hierarchy/RigUnit_SetRelativeBoneTransform.cpp b/Engine/Plugins/Experimental/ControlRig/Source/ControlRig/Private/Units/Hierarchy/RigUnit_SetRelativeBoneTransform.cpp
--- a/Engine/Plugins/Experimental/ControlRig/Source/ControlRig/Private/Units/Hierarchy/RigUnit_SetRelativeBoneTransform.cpp
+++ b/Engine/Plugins/Experimental/ControlRig/Source/ControlRig/Private/Units/Hierarchy/RigUnit_SetRelativeBoneTransform.cpp
@@ -8,7 +8,7 @@ FString FRigUnit_SetRelativeBoneTransform::GetUnitLabel() const
 	return FString::Printf(TEXT("Set Relative Transform %s"), *Bone.ToString());
 }
 
-void FRigUnit_SetRelativeBoneTransform::Execute(const FRigUnitContext& Context)
+UE_RigUnit_SetRelativeBoneTransform_IMPLEMENT_STATIC_VIRTUAL_METHOD(void, Execute, const FRigUnitContext& Context)
 {
     DECLARE_SCOPE_HIERARCHICAL_COUNTER_RIGUNIT()
 	FRigBoneHierarchy* Hierarchy = ExecuteContext.GetBones();
diff --git a/Engine/Plugins/Experimental/ControlRig/Source/ControlRig/Private/Units/Hierarchy/RigUnit_SetRelativeBoneTransform.h b/Engine/Plugins/Experimental/ControlRig/Source/ControlRig/Private/Units/Hierarchy/RigUnit_SetRelativeBoneTransform.h
--- a/Engine/Plugins/Experimental/ControlRig/Source/ControlRig/Private/Units/Hierarchy/RigUnit_SetRelativeBoneTransform.h
+++ b/Engine/Plugins/Experimental/ControlRig/Source/ControlRig/Private/Units/Hierarchy/RigUnit_SetRelativeBoneTransform.h
@@ -21,6 +21,7 @@ struct FRigUnit_SetRelativeBoneTransform : public FRigUnitMutable
 	{}
 
 	virtual FString GetUnitLabel() const override;
+	STATIC_VIRTUAL_METHOD()
 	virtual void Execute(const FRigUnitContext& Context) override;
 
 	/**
diff --git a/Engine/Plugins/Experimental/ControlRig/Source/ControlRig/Private/Units/Highlevel/Harmonics/RigUnit_BoneHarmonics.cpp b/Engine/Plugins/Experimental/ControlRig/Source/ControlRig/Private/Units/Highlevel/Harmonics/RigUnit_BoneHarmonics.cpp
--- a/Engine/Plugins/Experimental/ControlRig/Source/ControlRig/Private/Units/Highlevel/Harmonics/RigUnit_BoneHarmonics.cpp
+++ b/Engine/Plugins/Experimental/ControlRig/Source/ControlRig/Private/Units/Highlevel/Harmonics/RigUnit_BoneHarmonics.cpp
@@ -3,7 +3,7 @@
 #include "Units/Highlevel/Harmonics/RigUnit_BoneHarmonics.h"
 #include "Units/RigUnitContext.h"
 
-void FRigUnit_BoneHarmonics::Execute(const FRigUnitContext& Context)
+UE_RigUnit_BoneHarmonics_IMPLEMENT_STATIC_VIRTUAL_METHOD(void, Execute, const FRigUnitContext& Context)
 {
     DECLARE_SCOPE_HIERARCHICAL_COUNTER_RIGUNIT()
 	FRigBoneHierarchy* Hierarchy = ExecuteContext.GetBones();
diff --git a/Engine/Plugins/Experimental/ControlRig/Source/ControlRig/Private/Units/Highlevel/Harmonics/RigUnit_BoneHarmonics.h b/Engine/Plugins/Experimental/ControlRig/Source/ControlRig/Private/Units/Highlevel/Harmonics/RigUnit_BoneHarmonics.h
--- a/Engine/Plugins/Experimental/ControlRig/Source/ControlRig/Private/Units/Highlevel/Harmonics/RigUnit_BoneHarmonics.h
+++ b/Engine/Plugins/Experimental/ControlRig/Source/ControlRig/Private/Units/Highlevel/Harmonics/RigUnit_BoneHarmonics.h
@@ -54,6 +54,7 @@ struct FRigUnit_BoneHarmonics : public FRigUnit_HighlevelBaseMutable
 		WaveTime = FVector::ZeroVector;
 	}
 
+	STATIC_VIRTUAL_METHOD()
 	virtual void Execute(const FRigUnitContext& Context) override;
 
 	/** The bones to drive. */
diff --git a/Engine/Plugins/Experimental/ControlRig/Source/ControlRig/Private/Units/Highlevel/Harmonics/RigUnit_ChainHarmonics.cpp b/Engine/Plugins/Experimental/ControlRig/Source/ControlRig/Private/Units/Highlevel/Harmonics/RigUnit_ChainHarmonics.cpp
--- a/Engine/Plugins/Experimental/ControlRig/Source/ControlRig/Private/Units/Highlevel/Harmonics/RigUnit_ChainHarmonics.cpp
+++ b/Engine/Plugins/Experimental/ControlRig/Source/ControlRig/Private/Units/Highlevel/Harmonics/RigUnit_ChainHarmonics.cpp
@@ -3,7 +3,7 @@
 #include "Units/Highlevel/Harmonics/RigUnit_ChainHarmonics.h"
 #include "Units/RigUnitContext.h"
 
-void FRigUnit_ChainHarmonics::Execute(const FRigUnitContext& Context)
+UE_RigUnit_ChainHarmonics_IMPLEMENT_STATIC_VIRTUAL_METHOD(void, Execute, const FRigUnitContext& Context)
 {
     DECLARE_SCOPE_HIERARCHICAL_COUNTER_RIGUNIT()
 	FRigBoneHierarchy* Hierarchy = ExecuteContext.GetBones();
diff --git a/Engine/Plugins/Experimental/ControlRig/Source/ControlRig/Private/Units/Highlevel/Harmonics/RigUnit_ChainHarmonics.h b/Engine/Plugins/Experimental/ControlRig/Source/ControlRig/Private/Units/Highlevel/Harmonics/RigUnit_ChainHarmonics.h
--- a/Engine/Plugins/Experimental/ControlRig/Source/ControlRig/Private/Units/Highlevel/Harmonics/RigUnit_ChainHarmonics.h
+++ b/Engine/Plugins/Experimental/ControlRig/Source/ControlRig/Private/Units/Highlevel/Harmonics/RigUnit_ChainHarmonics.h
@@ -161,6 +161,7 @@ struct FRigUnit_ChainHarmonics : public FRigUnit_HighlevelBaseMutable
 		Time = FVector::ZeroVector;
 	}
 
+	STATIC_VIRTUAL_METHOD()
 	virtual void Execute(const FRigUnitContext& Context) override;
 
 	UPROPERTY(meta = (Input, Constant, BoneName))
diff --git a/Engine/Plugins/Experimental/ControlRig/Source/ControlRig/Private/Units/Highlevel/Hierarchy/RigUnit_AimBone.cpp b/Engine/Plugins/Experimental/ControlRig/Source/ControlRig/Private/Units/Highlevel/Hierarchy/RigUnit_AimBone.cpp
--- a/Engine/Plugins/Experimental/ControlRig/Source/ControlRig/Private/Units/Highlevel/Hierarchy/RigUnit_AimBone.cpp
+++ b/Engine/Plugins/Experimental/ControlRig/Source/ControlRig/Private/Units/Highlevel/Hierarchy/RigUnit_AimBone.cpp
@@ -3,7 +3,7 @@
 #include "Units/Highlevel/Hierarchy/RigUnit_AimBone.h"
 #include "Units/RigUnitContext.h"
 
-void FRigUnit_AimBone::Execute(const FRigUnitContext& Context)
+UE_RigUnit_AimBone_IMPLEMENT_STATIC_VIRTUAL_METHOD(void, Execute, const FRigUnitContext& Context)
 {
     DECLARE_SCOPE_HIERARCHICAL_COUNTER_RIGUNIT()
 	FRigBoneHierarchy* Hierarchy = ExecuteContext.GetBones();
diff --git a/Engine/Plugins/Experimental/ControlRig/Source/ControlRig/Private/Units/Highlevel/Hierarchy/RigUnit_AimBone.h b/Engine/Plugins/Experimental/ControlRig/Source/ControlRig/Private/Units/Highlevel/Hierarchy/RigUnit_AimBone.h
--- a/Engine/Plugins/Experimental/ControlRig/Source/ControlRig/Private/Units/Highlevel/Hierarchy/RigUnit_AimBone.h
+++ b/Engine/Plugins/Experimental/ControlRig/Source/ControlRig/Private/Units/Highlevel/Hierarchy/RigUnit_AimBone.h
@@ -106,6 +106,7 @@ struct FRigUnit_AimBone : public FRigUnit_HighlevelBaseMutable
 		SecondaryCachedSpaceIndex = INDEX_NONE;
 	}
 
+	STATIC_VIRTUAL_METHOD()
 	virtual void Execute(const FRigUnitContext& Context) override;
 
 	/** 
diff --git a/Engine/Plugins/Experimental/ControlRig/Source/ControlRig/Private/Units/Highlevel/Hierarchy/RigUnit_CCDIK.cpp b/Engine/Plugins/Experimental/ControlRig/Source/ControlRig/Private/Units/Highlevel/Hierarchy/RigUnit_CCDIK.cpp
--- a/Engine/Plugins/Experimental/ControlRig/Source/ControlRig/Private/Units/Highlevel/Hierarchy/RigUnit_CCDIK.cpp
+++ b/Engine/Plugins/Experimental/ControlRig/Source/ControlRig/Private/Units/Highlevel/Hierarchy/RigUnit_CCDIK.cpp
@@ -3,7 +3,7 @@
 #include "RigUnit_CCDIK.h"
 #include "Units/RigUnitContext.h"
 
-void FRigUnit_CCDIK::Execute(const FRigUnitContext& Context)
+UE_RigUnit_CCDIK_IMPLEMENT_STATIC_VIRTUAL_METHOD(void, Execute, const FRigUnitContext& Context)
 {
     DECLARE_SCOPE_HIERARCHICAL_COUNTER_RIGUNIT()
 	FRigBoneHierarchy* Hierarchy = ExecuteContext.GetBones();
@@ -72,7 +72,7 @@ void FRigUnit_CCDIK::Execute(const FRigUnitContext& Context)
 			{
 				const FTransform& GlobalTransform = Hierarchy->GetGlobalTransform(BoneIndices[ChainIndex]);
 				const FTransform& LocalTransform = Hierarchy->GetLocalTransform(BoneIndices[ChainIndex]);
-				Chain.Add(CCDIKChainLink(GlobalTransform, LocalTransform, ChainIndex));
+				Chain.Add(FCCDIKChainLink(GlobalTransform, LocalTransform, ChainIndex));
 			}
 
 			for (float& Limit : RotationLimitsPerBone)
@@ -95,7 +95,7 @@ void FRigUnit_CCDIK::Execute(const FRigUnitContext& Context)
 			{
 				for (int32 LinkIndex = 0; LinkIndex < BoneIndices.Num(); LinkIndex++)
 				{
-					const CCDIKChainLink& CurrentLink = Chain[LinkIndex];
+					const FCCDIKChainLink& CurrentLink = Chain[LinkIndex];
 					Hierarchy->SetGlobalTransform(BoneIndices[LinkIndex], CurrentLink.Transform, bPropagateToChildren);
 				}
 
diff --git a/Engine/Plugins/Experimental/ControlRig/Source/ControlRig/Private/Units/Highlevel/Hierarchy/RigUnit_CCDIK.h b/Engine/Plugins/Experimental/ControlRig/Source/ControlRig/Private/Units/Highlevel/Hierarchy/RigUnit_CCDIK.h
--- a/Engine/Plugins/Experimental/ControlRig/Source/ControlRig/Private/Units/Highlevel/Hierarchy/RigUnit_CCDIK.h
+++ b/Engine/Plugins/Experimental/ControlRig/Source/ControlRig/Private/Units/Highlevel/Hierarchy/RigUnit_CCDIK.h
@@ -40,6 +40,7 @@ struct FRigUnit_CCDIK : public FRigUnit_HighlevelBaseMutable
 {
 	GENERATED_BODY()
 
+	STATIC_VIRTUAL_METHOD()
 	virtual void Execute(const FRigUnitContext& Context) override;
 
 	FRigUnit_CCDIK()
@@ -109,9 +110,14 @@ struct FRigUnit_CCDIK : public FRigUnit_HighlevelBaseMutable
 	UPROPERTY(meta = (Input))
 	bool bPropagateToChildren;
 
-	TArray<CCDIKChainLink> Chain;
+	UPROPERTY(transient)
+	TArray<FCCDIKChainLink> Chain;
+	UPROPERTY(transient)
 	TArray<int32> BoneIndices;
+	UPROPERTY(transient)
 	TArray<int32> RotationLimitIndex;
+	UPROPERTY(transient)
 	TArray<float> RotationLimitsPerBone;
+	UPROPERTY(transient)
 	int32 EffectorIndex;
 };
diff --git a/Engine/Plugins/Experimental/ControlRig/Source/ControlRig/Private/Units/Highlevel/Hierarchy/RigUnit_DistributeRotation.cpp b/Engine/Plugins/Experimental/ControlRig/Source/ControlRig/Private/Units/Highlevel/Hierarchy/RigUnit_DistributeRotation.cpp
--- a/Engine/Plugins/Experimental/ControlRig/Source/ControlRig/Private/Units/Highlevel/Hierarchy/RigUnit_DistributeRotation.cpp
+++ b/Engine/Plugins/Experimental/ControlRig/Source/ControlRig/Private/Units/Highlevel/Hierarchy/RigUnit_DistributeRotation.cpp
@@ -3,7 +3,7 @@
 #include "Units/Highlevel/Hierarchy/RigUnit_DistributeRotation.h"
 #include "Units/RigUnitContext.h"
 
-void FRigUnit_DistributeRotation::Execute(const FRigUnitContext& Context)
+UE_RigUnit_DistributeRotation_IMPLEMENT_STATIC_VIRTUAL_METHOD(void, Execute, const FRigUnitContext& Context)
 {
     DECLARE_SCOPE_HIERARCHICAL_COUNTER_RIGUNIT()
 	FRigBoneHierarchy* Hierarchy = ExecuteContext.GetBones();
diff --git a/Engine/Plugins/Experimental/ControlRig/Source/ControlRig/Private/Units/Highlevel/Hierarchy/RigUnit_DistributeRotation.h b/Engine/Plugins/Experimental/ControlRig/Source/ControlRig/Private/Units/Highlevel/Hierarchy/RigUnit_DistributeRotation.h
--- a/Engine/Plugins/Experimental/ControlRig/Source/ControlRig/Private/Units/Highlevel/Hierarchy/RigUnit_DistributeRotation.h
+++ b/Engine/Plugins/Experimental/ControlRig/Source/ControlRig/Private/Units/Highlevel/Hierarchy/RigUnit_DistributeRotation.h
@@ -47,6 +47,7 @@ struct FRigUnit_DistributeRotation : public FRigUnit_HighlevelBaseMutable
 		bPropagateToChildren = false;
 	}
 
+	STATIC_VIRTUAL_METHOD()
 	virtual void Execute(const FRigUnitContext& Context) override;
 
 	/** 
diff --git a/Engine/Plugins/Experimental/ControlRig/Source/ControlRig/Private/Units/Highlevel/Hierarchy/RigUnit_FABRIK.cpp b/Engine/Plugins/Experimental/ControlRig/Source/ControlRig/Private/Units/Highlevel/Hierarchy/RigUnit_FABRIK.cpp
--- a/Engine/Plugins/Experimental/ControlRig/Source/ControlRig/Private/Units/Highlevel/Hierarchy/RigUnit_FABRIK.cpp
+++ b/Engine/Plugins/Experimental/ControlRig/Source/ControlRig/Private/Units/Highlevel/Hierarchy/RigUnit_FABRIK.cpp
@@ -3,7 +3,7 @@
 #include "RigUnit_FABRIK.h"
 #include "Units/RigUnitContext.h"
 
-void FRigUnit_FABRIK::Execute(const FRigUnitContext& Context)
+UE_RigUnit_FABRIK_IMPLEMENT_STATIC_VIRTUAL_METHOD(void, Execute, const FRigUnitContext& Context)
 {
     DECLARE_SCOPE_HIERARCHICAL_COUNTER_RIGUNIT()
 	FRigBoneHierarchy* Hierarchy = ExecuteContext.GetBones();
@@ -55,7 +55,7 @@ void FRigUnit_FABRIK::Execute(const FRigUnitContext& Context)
 			float MaximumReach = 0.f;
 			int32 const NumChainLinks = BoneIndices.Num();
 			const int32 RootIndex = BoneIndices.Last();
-			Chain.Add(FABRIKChainLink(Hierarchy->GetGlobalTransform(RootIndex).GetLocation(), 0.f, RootIndex, 0));
+			Chain.Add(FFABRIKChainLink(Hierarchy->GetGlobalTransform(RootIndex).GetLocation(), 0.f, RootIndex, 0));
 			Transforms[0] = Hierarchy->GetGlobalTransform(RootIndex);
 
 			// start from child to up
@@ -68,7 +68,7 @@ void FRigUnit_FABRIK::Execute(const FRigUnitContext& Context)
 				float const BoneLength = FVector::Dist(BoneTransform.GetLocation(), ParentTransform.GetLocation());
 
 				const int32 TransformIndex = Chain.Num();
-				Chain.Add(FABRIKChainLink(BoneTransform.GetLocation(), BoneLength, BoneIndices[ChainIndex], TransformIndex));
+				Chain.Add(FFABRIKChainLink(BoneTransform.GetLocation(), BoneLength, BoneIndices[ChainIndex], TransformIndex));
 				MaximumReach += BoneLength;
 
 				Transforms[TransformIndex] = BoneTransform;
@@ -82,8 +82,8 @@ void FRigUnit_FABRIK::Execute(const FRigUnitContext& Context)
 				// FABRIK algorithm - re-orientation of bone local axes after translation calculation
 				for (int32 LinkIndex = 0; LinkIndex < NumChainLinks - 1; LinkIndex++)
 				{
-					const FABRIKChainLink& CurrentLink = Chain[LinkIndex];
-					const FABRIKChainLink& ChildLink = Chain[LinkIndex + 1];
+					const FFABRIKChainLink& CurrentLink = Chain[LinkIndex];
+					const FFABRIKChainLink& ChildLink = Chain[LinkIndex + 1];
 
 					// Calculate pre-translation vector between this bone and child
 					FVector const OldDir = (Hierarchy->GetGlobalTransform(ChildLink.BoneIndex).GetLocation() - Hierarchy->GetGlobalTransform(CurrentLink.BoneIndex).GetLocation()).GetUnsafeNormal();
@@ -106,14 +106,14 @@ void FRigUnit_FABRIK::Execute(const FRigUnitContext& Context)
 				}
 
 				// fill up the last data transform
-				const FABRIKChainLink & CurrentLink = Chain[NumChainLinks - 1];
+				const FFABRIKChainLink & CurrentLink = Chain[NumChainLinks - 1];
 				FTransform& CurrentBoneTransform = Transforms[CurrentLink.TransformIndex];
 				CurrentBoneTransform.SetTranslation(CurrentLink.Position);
 				CurrentBoneTransform.SetRotation(Hierarchy->GetGlobalTransform(CurrentLink.BoneIndex).GetRotation());
 
 				for (int32 LinkIndex = 0; LinkIndex < NumChainLinks; LinkIndex++)
 				{
-					FABRIKChainLink const & LocalLink = Chain[LinkIndex];
+					FFABRIKChainLink const & LocalLink = Chain[LinkIndex];
 					Hierarchy->SetGlobalTransform(LocalLink.BoneIndex, Transforms[LocalLink.TransformIndex], bPropagateToChildren);
 				}
 
diff --git a/Engine/Plugins/Experimental/ControlRig/Source/ControlRig/Private/Units/Highlevel/Hierarchy/RigUnit_FABRIK.h b/Engine/Plugins/Experimental/ControlRig/Source/ControlRig/Private/Units/Highlevel/Hierarchy/RigUnit_FABRIK.h
--- a/Engine/Plugins/Experimental/ControlRig/Source/ControlRig/Private/Units/Highlevel/Hierarchy/RigUnit_FABRIK.h
+++ b/Engine/Plugins/Experimental/ControlRig/Source/ControlRig/Private/Units/Highlevel/Hierarchy/RigUnit_FABRIK.h
@@ -16,6 +16,7 @@ struct FRigUnit_FABRIK : public FRigUnit_HighlevelBaseMutable
 {
 	GENERATED_BODY()
 
+	STATIC_VIRTUAL_METHOD()
 	virtual void Execute(const FRigUnitContext& Context) override;
 
 	FRigUnit_FABRIK()
@@ -65,7 +66,12 @@ struct FRigUnit_FABRIK : public FRigUnit_HighlevelBaseMutable
 	UPROPERTY(meta = (Input))
 	int32 MaxIterations;
 
-	TArray<FABRIKChainLink> Chain;
+	UPROPERTY(transient)
+	TArray<FFABRIKChainLink> Chain;
+
+	UPROPERTY(transient)
 	TArray<int32> BoneIndices;
+
+	UPROPERTY(transient)
 	int32 EffectorIndex;
 };
diff --git a/Engine/Plugins/Experimental/ControlRig/Source/ControlRig/Private/Units/Highlevel/Hierarchy/RigUnit_FitChainToCurve.cpp b/Engine/Plugins/Experimental/ControlRig/Source/ControlRig/Private/Units/Highlevel/Hierarchy/RigUnit_FitChainToCurve.cpp
--- a/Engine/Plugins/Experimental/ControlRig/Source/ControlRig/Private/Units/Highlevel/Hierarchy/RigUnit_FitChainToCurve.cpp
+++ b/Engine/Plugins/Experimental/ControlRig/Source/ControlRig/Private/Units/Highlevel/Hierarchy/RigUnit_FitChainToCurve.cpp
@@ -3,7 +3,7 @@
 #include "Units/Highlevel/Hierarchy/RigUnit_FitChainToCurve.h"
 #include "Units/RigUnitContext.h"
 
-void FRigUnit_FitChainToCurve::Execute(const FRigUnitContext& Context)
+UE_RigUnit_FitChainToCurve_IMPLEMENT_STATIC_VIRTUAL_METHOD(void, Execute, const FRigUnitContext& Context)
 {
     DECLARE_SCOPE_HIERARCHICAL_COUNTER_RIGUNIT()
 	FRigBoneHierarchy* Hierarchy = ExecuteContext.GetBones();
diff --git a/Engine/Plugins/Experimental/ControlRig/Source/ControlRig/Private/Units/Highlevel/Hierarchy/RigUnit_FitChainToCurve.h b/Engine/Plugins/Experimental/ControlRig/Source/ControlRig/Private/Units/Highlevel/Hierarchy/RigUnit_FitChainToCurve.h
--- a/Engine/Plugins/Experimental/ControlRig/Source/ControlRig/Private/Units/Highlevel/Hierarchy/RigUnit_FitChainToCurve.h
+++ b/Engine/Plugins/Experimental/ControlRig/Source/ControlRig/Private/Units/Highlevel/Hierarchy/RigUnit_FitChainToCurve.h
@@ -108,6 +108,7 @@ struct FRigUnit_FitChainToCurve : public FRigUnit_HighlevelBaseMutable
 		ChainLength = 0.f;
 	}
 
+	STATIC_VIRTUAL_METHOD()
 	virtual void Execute(const FRigUnitContext& Context) override;
 
 	/** 
diff --git a/Engine/Plugins/Experimental/ControlRig/Source/ControlRig/Private/Units/Highlevel/Hierarchy/RigUnit_ModifyBoneTransforms.cpp b/Engine/Plugins/Experimental/ControlRig/Source/ControlRig/Private/Units/Highlevel/Hierarchy/RigUnit_ModifyBoneTransforms.cpp
--- a/Engine/Plugins/Experimental/ControlRig/Source/ControlRig/Private/Units/Highlevel/Hierarchy/RigUnit_ModifyBoneTransforms.cpp
+++ b/Engine/Plugins/Experimental/ControlRig/Source/ControlRig/Private/Units/Highlevel/Hierarchy/RigUnit_ModifyBoneTransforms.cpp
@@ -3,7 +3,7 @@
 #include "RigUnit_ModifyBoneTransforms.h"
 #include "Units/RigUnitContext.h"
 
-void FRigUnit_ModifyBoneTransforms::Execute(const FRigUnitContext& Context)
+UE_RigUnit_ModifyBoneTransforms_IMPLEMENT_STATIC_VIRTUAL_METHOD(void, Execute, const FRigUnitContext& Context)
 {
     DECLARE_SCOPE_HIERARCHICAL_COUNTER_RIGUNIT()
 	FRigBoneHierarchy* Hierarchy = ExecuteContext.GetBones();
diff --git a/Engine/Plugins/Experimental/ControlRig/Source/ControlRig/Private/Units/Highlevel/Hierarchy/RigUnit_ModifyBoneTransforms.h b/Engine/Plugins/Experimental/ControlRig/Source/ControlRig/Private/Units/Highlevel/Hierarchy/RigUnit_ModifyBoneTransforms.h
--- a/Engine/Plugins/Experimental/ControlRig/Source/ControlRig/Private/Units/Highlevel/Hierarchy/RigUnit_ModifyBoneTransforms.h
+++ b/Engine/Plugins/Experimental/ControlRig/Source/ControlRig/Private/Units/Highlevel/Hierarchy/RigUnit_ModifyBoneTransforms.h
@@ -70,6 +70,7 @@ struct FRigUnit_ModifyBoneTransforms : public FRigUnit_HighlevelBaseMutable
 		, Mode(EControlRigModifyBoneMode::AdditiveLocal)
 	{}
 
+	STATIC_VIRTUAL_METHOD()
 	virtual void Execute(const FRigUnitContext& Context) override;
 
 	/**
diff --git a/Engine/Plugins/Experimental/ControlRig/Source/ControlRig/Private/Units/Highlevel/Hierarchy/RigUnit_SlideChain.cpp b/Engine/Plugins/Experimental/ControlRig/Source/ControlRig/Private/Units/Highlevel/Hierarchy/RigUnit_SlideChain.cpp
--- a/Engine/Plugins/Experimental/ControlRig/Source/ControlRig/Private/Units/Highlevel/Hierarchy/RigUnit_SlideChain.cpp
+++ b/Engine/Plugins/Experimental/ControlRig/Source/ControlRig/Private/Units/Highlevel/Hierarchy/RigUnit_SlideChain.cpp
@@ -4,7 +4,7 @@
 #include "Math/ControlRigMathLibrary.h"
 #include "Units/RigUnitContext.h"
 
-void FRigUnit_SlideChain::Execute(const FRigUnitContext& Context)
+UE_RigUnit_SlideChain_IMPLEMENT_STATIC_VIRTUAL_METHOD(void, Execute, const FRigUnitContext& Context)
 {
 
 	FRigBoneHierarchy* Hierarchy = ExecuteContext.GetBones();
diff --git a/Engine/Plugins/Experimental/ControlRig/Source/ControlRig/Private/Units/Highlevel/Hierarchy/RigUnit_SlideChain.h b/Engine/Plugins/Experimental/ControlRig/Source/ControlRig/Private/Units/Highlevel/Hierarchy/RigUnit_SlideChain.h
--- a/Engine/Plugins/Experimental/ControlRig/Source/ControlRig/Private/Units/Highlevel/Hierarchy/RigUnit_SlideChain.h
+++ b/Engine/Plugins/Experimental/ControlRig/Source/ControlRig/Private/Units/Highlevel/Hierarchy/RigUnit_SlideChain.h
@@ -22,6 +22,7 @@ struct FRigUnit_SlideChain: public FRigUnit_HighlevelBaseMutable
 		ChainLength = 0.f;
 	}
 
+	STATIC_VIRTUAL_METHOD()
 	virtual void Execute(const FRigUnitContext& Context) override;
 
 	/** 
diff --git a/Engine/Plugins/Experimental/ControlRig/Source/ControlRig/Private/Units/Highlevel/Hierarchy/RigUnit_SpringIK.cpp b/Engine/Plugins/Experimental/ControlRig/Source/ControlRig/Private/Units/Highlevel/Hierarchy/RigUnit_SpringIK.cpp
--- a/Engine/Plugins/Experimental/ControlRig/Source/ControlRig/Private/Units/Highlevel/Hierarchy/RigUnit_SpringIK.cpp
+++ b/Engine/Plugins/Experimental/ControlRig/Source/ControlRig/Private/Units/Highlevel/Hierarchy/RigUnit_SpringIK.cpp
@@ -3,7 +3,7 @@
 #include "Units/Highlevel/Hierarchy/RigUnit_SpringIK.h"
 #include "Units/RigUnitContext.h"
 
-void FRigUnit_SpringIK::Execute(const FRigUnitContext& Context)
+UE_RigUnit_SpringIK_IMPLEMENT_STATIC_VIRTUAL_METHOD(void, Execute, const FRigUnitContext& Context)
 {
     DECLARE_SCOPE_HIERARCHICAL_COUNTER_RIGUNIT()
 	FRigBoneHierarchy* Hierarchy = ExecuteContext.GetBones();
diff --git a/Engine/Plugins/Experimental/ControlRig/Source/ControlRig/Private/Units/Highlevel/Hierarchy/RigUnit_SpringIK.h b/Engine/Plugins/Experimental/ControlRig/Source/ControlRig/Private/Units/Highlevel/Hierarchy/RigUnit_SpringIK.h
--- a/Engine/Plugins/Experimental/ControlRig/Source/ControlRig/Private/Units/Highlevel/Hierarchy/RigUnit_SpringIK.h
+++ b/Engine/Plugins/Experimental/ControlRig/Source/ControlRig/Private/Units/Highlevel/Hierarchy/RigUnit_SpringIK.h
@@ -74,6 +74,7 @@ struct FRigUnit_SpringIK : public FRigUnit_HighlevelBaseMutable
 		PoleVectorIndex = INDEX_NONE;
 	}
 
+	STATIC_VIRTUAL_METHOD()
 	virtual void Execute(const FRigUnitContext& Context) override;
 
 	/** 
diff --git a/Engine/Plugins/Experimental/ControlRig/Source/ControlRig/Private/Units/Highlevel/Hierarchy/RigUnit_TransformConstraint.cpp b/Engine/Plugins/Experimental/ControlRig/Source/ControlRig/Private/Units/Highlevel/Hierarchy/RigUnit_TransformConstraint.cpp
--- a/Engine/Plugins/Experimental/ControlRig/Source/ControlRig/Private/Units/Highlevel/Hierarchy/RigUnit_TransformConstraint.cpp
+++ b/Engine/Plugins/Experimental/ControlRig/Source/ControlRig/Private/Units/Highlevel/Hierarchy/RigUnit_TransformConstraint.cpp
@@ -5,7 +5,7 @@
 #include "HelperUtil.h"
 #include "AnimationCoreLibrary.h"
 
-void FRigUnit_TransformConstraint::Execute(const FRigUnitContext& Context)
+UE_RigUnit_TransformConstraint_IMPLEMENT_STATIC_VIRTUAL_METHOD(void, Execute, const FRigUnitContext& Context)
 {
     DECLARE_SCOPE_HIERARCHICAL_COUNTER_RIGUNIT()
 
@@ -35,23 +35,23 @@ void FRigUnit_TransformConstraint::Execute(const FRigUnitContext& Context)
 
 						if (bTranslationFilterValid && bRotationFilterValid && bScaleFilterValid)
 						{
-							AddConstraintData(ETransformConstraintType::Parent, TargetIndex, SourceTransform, InputBaseTransform);
+							AddConstraintData(Targets, ETransformConstraintType::Parent, TargetIndex, SourceTransform, InputBaseTransform, ConstraintData, ConstraintDataToTargets);
 						}
 						else
 						{
 							if (bTranslationFilterValid)
 							{
-								AddConstraintData(ETransformConstraintType::Translation, TargetIndex, SourceTransform, InputBaseTransform);
+								AddConstraintData(Targets, ETransformConstraintType::Translation, TargetIndex, SourceTransform, InputBaseTransform, ConstraintData, ConstraintDataToTargets);
 							}
 
 							if (bRotationFilterValid)
 							{
-								AddConstraintData(ETransformConstraintType::Rotation, TargetIndex, SourceTransform, InputBaseTransform);
+								AddConstraintData(Targets, ETransformConstraintType::Rotation, TargetIndex, SourceTransform, InputBaseTransform, ConstraintData, ConstraintDataToTargets);
 							}
 
 							if (bScaleFilterValid)
 							{
-								AddConstraintData(ETransformConstraintType::Scale, TargetIndex, SourceTransform, InputBaseTransform);
+								AddConstraintData(Targets, ETransformConstraintType::Scale, TargetIndex, SourceTransform, InputBaseTransform, ConstraintData, ConstraintDataToTargets);
 							}
 						}
 					}
@@ -93,13 +93,13 @@ void FRigUnit_TransformConstraint::Execute(const FRigUnitContext& Context)
 	}
 }
 
-void FRigUnit_TransformConstraint::AddConstraintData(ETransformConstraintType ConstraintType, const int32 TargetIndex, const FTransform& SourceTransform, const FTransform& InBaseTransform)
+void FRigUnit_TransformConstraint::AddConstraintData(const TArray<FConstraintTarget>& Targets, ETransformConstraintType ConstraintType, const int32 TargetIndex, const FTransform& SourceTransform, const FTransform& InBaseTransform, TArray<FConstraintData>& OutConstraintData, TMap<int32, int32>& OutConstraintDataToTargets)
 {
 	const FConstraintTarget& Target = Targets[TargetIndex];
 
-	int32 NewIndex = ConstraintData.AddDefaulted();
+	int32 NewIndex = OutConstraintData.AddDefaulted();
 	check(NewIndex != INDEX_NONE);
-	FConstraintData& NewData = ConstraintData[NewIndex];
+	FConstraintData& NewData = OutConstraintData[NewIndex];
 	NewData.Constraint = FTransformConstraintDescription(ConstraintType);
 	NewData.bMaintainOffset = Target.bMaintainOffset;
 	NewData.Weight = Target.Weight;
@@ -109,5 +109,5 @@ void FRigUnit_TransformConstraint::AddConstraintData(ETransformConstraintType Co
 		NewData.SaveInverseOffset(SourceTransform, Target.Transform, InBaseTransform);
 	}
 
-	ConstraintDataToTargets.Add(NewIndex, TargetIndex);
+	OutConstraintDataToTargets.Add(NewIndex, TargetIndex);
 }
\ No newline at end of file
diff --git a/Engine/Plugins/Experimental/ControlRig/Source/ControlRig/Private/Units/Highlevel/Hierarchy/RigUnit_TransformConstraint.h b/Engine/Plugins/Experimental/ControlRig/Source/ControlRig/Private/Units/Highlevel/Hierarchy/RigUnit_TransformConstraint.h
--- a/Engine/Plugins/Experimental/ControlRig/Source/ControlRig/Private/Units/Highlevel/Hierarchy/RigUnit_TransformConstraint.h
+++ b/Engine/Plugins/Experimental/ControlRig/Source/ControlRig/Private/Units/Highlevel/Hierarchy/RigUnit_TransformConstraint.h
@@ -45,6 +45,7 @@ struct FRigUnit_TransformConstraint : public FRigUnit_HighlevelBaseMutable
 		: BaseTransformSpace(ETransformSpaceMode::GlobalSpace)
 	{}
 
+	STATIC_VIRTUAL_METHOD()
 	virtual void Execute(const FRigUnitContext& Context) override;
 
 	UPROPERTY(EditAnywhere, Category = "Constraint", meta = (Input, Constant, BoneName))
@@ -65,10 +66,12 @@ struct FRigUnit_TransformConstraint : public FRigUnit_HighlevelBaseMutable
 	TArray<FConstraintTarget> Targets;
 
 private:
+	static void AddConstraintData(const TArray<FConstraintTarget>& Targets, ETransformConstraintType ConstraintType, const int32 TargetIndex, const FTransform& SourceTransform, const FTransform& InBaseTransform, TArray<FConstraintData>& OutConstraintData, TMap<int32, int32>& OutConstraintDataToTargets);
+
 	// note that Targets.Num () != ConstraintData.Num()
+	UPROPERTY(transient)
 	TArray<FConstraintData>	ConstraintData;
 
-	void AddConstraintData(ETransformConstraintType ConstraintType, const int32 TargetIndex, const FTransform& SourceTransform, const FTransform& InBaseTransform);
-
+	UPROPERTY(transient)
 	TMap<int32, int32> ConstraintDataToTargets;
 };
diff --git a/Engine/Plugins/Experimental/ControlRig/Source/ControlRig/Private/Units/Highlevel/Hierarchy/RigUnit_TwoBoneIKSimple.cpp b/Engine/Plugins/Experimental/ControlRig/Source/ControlRig/Private/Units/Highlevel/Hierarchy/RigUnit_TwoBoneIKSimple.cpp
--- a/Engine/Plugins/Experimental/ControlRig/Source/ControlRig/Private/Units/Highlevel/Hierarchy/RigUnit_TwoBoneIKSimple.cpp
+++ b/Engine/Plugins/Experimental/ControlRig/Source/ControlRig/Private/Units/Highlevel/Hierarchy/RigUnit_TwoBoneIKSimple.cpp
@@ -4,7 +4,7 @@
 #include "Units/RigUnitContext.h"
 #include "Math/ControlRigMathLibrary.h"
 
-void FRigUnit_TwoBoneIKSimple::Execute(const FRigUnitContext& Context)
+UE_RigUnit_TwoBoneIKSimple_IMPLEMENT_STATIC_VIRTUAL_METHOD(void, Execute, const FRigUnitContext& Context)
 {
     DECLARE_SCOPE_HIERARCHICAL_COUNTER_RIGUNIT()
 	FRigBoneHierarchy* Hierarchy = ExecuteContext.GetBones();
diff --git a/Engine/Plugins/Experimental/ControlRig/Source/ControlRig/Private/Units/Highlevel/Hierarchy/RigUnit_TwoBoneIKSimple.h b/Engine/Plugins/Experimental/ControlRig/Source/ControlRig/Private/Units/Highlevel/Hierarchy/RigUnit_TwoBoneIKSimple.h
--- a/Engine/Plugins/Experimental/ControlRig/Source/ControlRig/Private/Units/Highlevel/Hierarchy/RigUnit_TwoBoneIKSimple.h
+++ b/Engine/Plugins/Experimental/ControlRig/Source/ControlRig/Private/Units/Highlevel/Hierarchy/RigUnit_TwoBoneIKSimple.h
@@ -67,6 +67,7 @@ struct FRigUnit_TwoBoneIKSimple : public FRigUnit_HighlevelBaseMutable
 		PoleVectorSpaceIndex = INDEX_NONE;
 	}
 
+	STATIC_VIRTUAL_METHOD()
 	virtual void Execute(const FRigUnitContext& Context) override;
 
 	/** 
diff --git a/Engine/Plugins/Experimental/ControlRig/Source/ControlRig/Private/Units/Math/RigUnit_MathBool.cpp b/Engine/Plugins/Experimental/ControlRig/Source/ControlRig/Private/Units/Math/RigUnit_MathBool.cpp
--- a/Engine/Plugins/Experimental/ControlRig/Source/ControlRig/Private/Units/Math/RigUnit_MathBool.cpp
+++ b/Engine/Plugins/Experimental/ControlRig/Source/ControlRig/Private/Units/Math/RigUnit_MathBool.cpp
@@ -3,37 +3,37 @@
 #include "Units/Math/RigUnit_MathBool.h"
 #include "Units/RigUnitContext.h"
 
-void FRigUnit_MathBoolNot::Execute(const FRigUnitContext& Context)
+UE_RigUnit_MathBoolNot_IMPLEMENT_STATIC_VIRTUAL_METHOD(void, Execute, const FRigUnitContext& Context)
 {
     DECLARE_SCOPE_HIERARCHICAL_COUNTER_RIGUNIT()
 	Result = !Value;
 }
 
-void FRigUnit_MathBoolAnd::Execute(const FRigUnitContext& Context)
+UE_RigUnit_MathBoolAnd_IMPLEMENT_STATIC_VIRTUAL_METHOD(void, Execute, const FRigUnitContext& Context)
 {
     DECLARE_SCOPE_HIERARCHICAL_COUNTER_RIGUNIT()
 	Result = A && B;
 }
 
-void FRigUnit_MathBoolNand::Execute(const FRigUnitContext& Context)
+UE_RigUnit_MathBoolNand_IMPLEMENT_STATIC_VIRTUAL_METHOD(void, Execute, const FRigUnitContext& Context)
 {
     DECLARE_SCOPE_HIERARCHICAL_COUNTER_RIGUNIT()
 	Result = (!A) && (!B);
 }
 
-void FRigUnit_MathBoolOr::Execute(const FRigUnitContext& Context)
+UE_RigUnit_MathBoolOr_IMPLEMENT_STATIC_VIRTUAL_METHOD(void, Execute, const FRigUnitContext& Context)
 {
     DECLARE_SCOPE_HIERARCHICAL_COUNTER_RIGUNIT()
 	Result = A || B;
 }
 
-void FRigUnit_MathBoolEquals::Execute(const FRigUnitContext& Context)
+UE_RigUnit_MathBoolEquals_IMPLEMENT_STATIC_VIRTUAL_METHOD(void, Execute, const FRigUnitContext& Context)
 {
     DECLARE_SCOPE_HIERARCHICAL_COUNTER_RIGUNIT()
 	Result = A == B;
 }
 
-void FRigUnit_MathBoolNotEquals::Execute(const FRigUnitContext& Context)
+UE_RigUnit_MathBoolNotEquals_IMPLEMENT_STATIC_VIRTUAL_METHOD(void, Execute, const FRigUnitContext& Context)
 {
     DECLARE_SCOPE_HIERARCHICAL_COUNTER_RIGUNIT()
 	Result = A != B;
diff --git a/Engine/Plugins/Experimental/ControlRig/Source/ControlRig/Private/Units/Math/RigUnit_MathBool.h b/Engine/Plugins/Experimental/ControlRig/Source/ControlRig/Private/Units/Math/RigUnit_MathBool.h
--- a/Engine/Plugins/Experimental/ControlRig/Source/ControlRig/Private/Units/Math/RigUnit_MathBool.h
+++ b/Engine/Plugins/Experimental/ControlRig/Source/ControlRig/Private/Units/Math/RigUnit_MathBool.h
@@ -98,6 +98,7 @@ struct FRigUnit_MathBoolNot : public FRigUnit_MathBoolUnaryOp
 {
 	GENERATED_BODY()
 
+	STATIC_VIRTUAL_METHOD()
 	virtual void Execute(const FRigUnitContext& Context) override;
 };
 
@@ -109,6 +110,7 @@ struct FRigUnit_MathBoolAnd : public FRigUnit_MathBoolBinaryOp
 {
 	GENERATED_BODY()
 
+	STATIC_VIRTUAL_METHOD()
 	virtual void Execute(const FRigUnitContext& Context) override;
 };
 
@@ -120,6 +122,7 @@ struct FRigUnit_MathBoolNand : public FRigUnit_MathBoolBinaryOp
 {
 	GENERATED_BODY()
 
+	STATIC_VIRTUAL_METHOD()
 	virtual void Execute(const FRigUnitContext& Context) override;
 };
 
@@ -131,6 +134,7 @@ struct FRigUnit_MathBoolOr : public FRigUnit_MathBoolBinaryOp
 {
 	GENERATED_BODY()
 
+	STATIC_VIRTUAL_METHOD()
 	virtual void Execute(const FRigUnitContext& Context) override;
 };
 
@@ -141,6 +145,7 @@ USTRUCT(meta=(DisplayName="Equals", PrototypeName="Equals", Keywords="Same,=="))
 struct FRigUnit_MathBoolEquals : public FRigUnit_MathBoolBase
 {
 	GENERATED_BODY()
+	STATIC_VIRTUAL_METHOD()
 	virtual void Execute(const FRigUnitContext& Context) override;
 
 	FRigUnit_MathBoolEquals()
@@ -166,6 +171,7 @@ USTRUCT(meta=(DisplayName="Not Equals", PrototypeName="NotEquals", Keywords="Dif
 struct FRigUnit_MathBoolNotEquals : public FRigUnit_MathBoolBase
 {
 	GENERATED_BODY()
+	STATIC_VIRTUAL_METHOD()
 	virtual void Execute(const FRigUnitContext& Context) override;
 
 	FRigUnit_MathBoolNotEquals()
diff --git a/Engine/Plugins/Experimental/ControlRig/Source/ControlRig/Private/Units/Math/RigUnit_MathFloat.cpp b/Engine/Plugins/Experimental/ControlRig/Source/ControlRig/Private/Units/Math/RigUnit_MathFloat.cpp
--- a/Engine/Plugins/Experimental/ControlRig/Source/ControlRig/Private/Units/Math/RigUnit_MathFloat.cpp
+++ b/Engine/Plugins/Experimental/ControlRig/Source/ControlRig/Private/Units/Math/RigUnit_MathFloat.cpp
@@ -3,25 +3,25 @@
 #include "Units/Math/RigUnit_MathFloat.h"
 #include "Units/RigUnitContext.h"
 
-void FRigUnit_MathFloatAdd::Execute(const FRigUnitContext& Context)
+UE_RigUnit_MathFloatAdd_IMPLEMENT_STATIC_VIRTUAL_METHOD(void, Execute, const FRigUnitContext& Context)
 {
     DECLARE_SCOPE_HIERARCHICAL_COUNTER_RIGUNIT()
 	Result = A + B;
 }
 
-void FRigUnit_MathFloatSub::Execute(const FRigUnitContext& Context)
+UE_RigUnit_MathFloatSub_IMPLEMENT_STATIC_VIRTUAL_METHOD(void, Execute, const FRigUnitContext& Context)
 {
     DECLARE_SCOPE_HIERARCHICAL_COUNTER_RIGUNIT()
 	Result = A - B;
 }
 
-void FRigUnit_MathFloatMul::Execute(const FRigUnitContext& Context)
+UE_RigUnit_MathFloatMul_IMPLEMENT_STATIC_VIRTUAL_METHOD(void, Execute, const FRigUnitContext& Context)
 {
     DECLARE_SCOPE_HIERARCHICAL_COUNTER_RIGUNIT()
 	Result = A * B;
 }
 
-void FRigUnit_MathFloatDiv::Execute(const FRigUnitContext& Context)
+UE_RigUnit_MathFloatDiv_IMPLEMENT_STATIC_VIRTUAL_METHOD(void, Execute, const FRigUnitContext& Context)
 {
     DECLARE_SCOPE_HIERARCHICAL_COUNTER_RIGUNIT()
 	if(FMath::IsNearlyZero(B))
@@ -33,7 +33,7 @@ void FRigUnit_MathFloatDiv::Execute(const FRigUnitContext& Context)
 	Result = A / B;
 }
 
-void FRigUnit_MathFloatMod::Execute(const FRigUnitContext& Context)
+UE_RigUnit_MathFloatMod_IMPLEMENT_STATIC_VIRTUAL_METHOD(void, Execute, const FRigUnitContext& Context)
 {
     DECLARE_SCOPE_HIERARCHICAL_COUNTER_RIGUNIT()
 	if(FMath::IsNearlyZero(B) || B < 0.f)
@@ -45,25 +45,25 @@ void FRigUnit_MathFloatMod::Execute(const FRigUnitContext& Context)
 	Result = FMath::Fmod(A, B);
 }
 
-void FRigUnit_MathFloatMin::Execute(const FRigUnitContext& Context)
+UE_RigUnit_MathFloatMin_IMPLEMENT_STATIC_VIRTUAL_METHOD(void, Execute, const FRigUnitContext& Context)
 {
     DECLARE_SCOPE_HIERARCHICAL_COUNTER_RIGUNIT()
 	Result = FMath::Min<float>(A, B);
 }
 
-void FRigUnit_MathFloatMax::Execute(const FRigUnitContext& Context)
+UE_RigUnit_MathFloatMax_IMPLEMENT_STATIC_VIRTUAL_METHOD(void, Execute, const FRigUnitContext& Context)
 {
     DECLARE_SCOPE_HIERARCHICAL_COUNTER_RIGUNIT()
 	Result = FMath::Max<float>(A, B);
 }
 
-void FRigUnit_MathFloatPow::Execute(const FRigUnitContext& Context)
+UE_RigUnit_MathFloatPow_IMPLEMENT_STATIC_VIRTUAL_METHOD(void, Execute, const FRigUnitContext& Context)
 {
     DECLARE_SCOPE_HIERARCHICAL_COUNTER_RIGUNIT()
 	Result = FMath::Pow(A, B);
 }
 
-void FRigUnit_MathFloatSqrt::Execute(const FRigUnitContext& Context)
+UE_RigUnit_MathFloatSqrt_IMPLEMENT_STATIC_VIRTUAL_METHOD(void, Execute, const FRigUnitContext& Context)
 {
     DECLARE_SCOPE_HIERARCHICAL_COUNTER_RIGUNIT()
 	if(Value < 0.f)
@@ -75,55 +75,55 @@ void FRigUnit_MathFloatSqrt::Execute(const FRigUnitContext& Context)
 	Result = FMath::Sqrt(Value);
 }
 
-void FRigUnit_MathFloatNegate::Execute(const FRigUnitContext& Context)
+UE_RigUnit_MathFloatNegate_IMPLEMENT_STATIC_VIRTUAL_METHOD(void, Execute, const FRigUnitContext& Context)
 {
     DECLARE_SCOPE_HIERARCHICAL_COUNTER_RIGUNIT()
 	Result = -Value;
 }
 
-void FRigUnit_MathFloatAbs::Execute(const FRigUnitContext& Context)
+UE_RigUnit_MathFloatAbs_IMPLEMENT_STATIC_VIRTUAL_METHOD(void, Execute, const FRigUnitContext& Context)
 {
     DECLARE_SCOPE_HIERARCHICAL_COUNTER_RIGUNIT()
 	Result = FMath::Abs(Value);
 }
 
-void FRigUnit_MathFloatFloor::Execute(const FRigUnitContext& Context)
+UE_RigUnit_MathFloatFloor_IMPLEMENT_STATIC_VIRTUAL_METHOD(void, Execute, const FRigUnitContext& Context)
 {
     DECLARE_SCOPE_HIERARCHICAL_COUNTER_RIGUNIT()
 	Result = FMath::FloorToFloat(Value);
 }
 
-void FRigUnit_MathFloatCeil::Execute(const FRigUnitContext& Context)
+UE_RigUnit_MathFloatCeil_IMPLEMENT_STATIC_VIRTUAL_METHOD(void, Execute, const FRigUnitContext& Context)
 {
     DECLARE_SCOPE_HIERARCHICAL_COUNTER_RIGUNIT()
 	Result = FMath::CeilToFloat(Value);
 }
 
-void FRigUnit_MathFloatRound::Execute(const FRigUnitContext& Context)
+UE_RigUnit_MathFloatRound_IMPLEMENT_STATIC_VIRTUAL_METHOD(void, Execute, const FRigUnitContext& Context)
 {
     DECLARE_SCOPE_HIERARCHICAL_COUNTER_RIGUNIT()
 	Result = FMath::RoundToFloat(Value);
 }
 
-void FRigUnit_MathFloatSign::Execute(const FRigUnitContext& Context)
+UE_RigUnit_MathFloatSign_IMPLEMENT_STATIC_VIRTUAL_METHOD(void, Execute, const FRigUnitContext& Context)
 {
     DECLARE_SCOPE_HIERARCHICAL_COUNTER_RIGUNIT()
 	Result = Value < 0.f ? -1.f : 1.f;
 }
 
-void FRigUnit_MathFloatClamp::Execute(const FRigUnitContext& Context)
+UE_RigUnit_MathFloatClamp_IMPLEMENT_STATIC_VIRTUAL_METHOD(void, Execute, const FRigUnitContext& Context)
 {
     DECLARE_SCOPE_HIERARCHICAL_COUNTER_RIGUNIT()
 	Result = FMath::Clamp<float>(Value, Minimum, Maximum);
 }
 
-void FRigUnit_MathFloatLerp::Execute(const FRigUnitContext& Context)
+UE_RigUnit_MathFloatLerp_IMPLEMENT_STATIC_VIRTUAL_METHOD(void, Execute, const FRigUnitContext& Context)
 {
     DECLARE_SCOPE_HIERARCHICAL_COUNTER_RIGUNIT()
 	Result = FMath::Lerp<float>(A, B, T);
 }
 
-void FRigUnit_MathFloatRemap::Execute(const FRigUnitContext& Context)
+UE_RigUnit_MathFloatRemap_IMPLEMENT_STATIC_VIRTUAL_METHOD(void, Execute, const FRigUnitContext& Context)
 {
     DECLARE_SCOPE_HIERARCHICAL_COUNTER_RIGUNIT()
 	float Ratio = 0.f;
@@ -142,43 +142,43 @@ void FRigUnit_MathFloatRemap::Execute(const FRigUnitContext& Context)
 	Result = FMath::Lerp<float>(TargetMinimum, TargetMaximum, Ratio);
 }
 
-void FRigUnit_MathFloatEquals::Execute(const FRigUnitContext& Context)
+UE_RigUnit_MathFloatEquals_IMPLEMENT_STATIC_VIRTUAL_METHOD(void, Execute, const FRigUnitContext& Context)
 {
     DECLARE_SCOPE_HIERARCHICAL_COUNTER_RIGUNIT()
 	Result = A == B;
 }
 
-void FRigUnit_MathFloatNotEquals::Execute(const FRigUnitContext& Context)
+UE_RigUnit_MathFloatNotEquals_IMPLEMENT_STATIC_VIRTUAL_METHOD(void, Execute, const FRigUnitContext& Context)
 {
     DECLARE_SCOPE_HIERARCHICAL_COUNTER_RIGUNIT()
 	Result = A != B;
 }
 
-void FRigUnit_MathFloatGreater::Execute(const FRigUnitContext& Context)
+UE_RigUnit_MathFloatGreater_IMPLEMENT_STATIC_VIRTUAL_METHOD(void, Execute, const FRigUnitContext& Context)
 {
     DECLARE_SCOPE_HIERARCHICAL_COUNTER_RIGUNIT()
 	Result = A > B;
 }
 
-void FRigUnit_MathFloatLess::Execute(const FRigUnitContext& Context)
+UE_RigUnit_MathFloatLess_IMPLEMENT_STATIC_VIRTUAL_METHOD(void, Execute, const FRigUnitContext& Context)
 {
     DECLARE_SCOPE_HIERARCHICAL_COUNTER_RIGUNIT()
 	Result = A < B;
 }
 
-void FRigUnit_MathFloatGreaterEqual::Execute(const FRigUnitContext& Context)
+UE_RigUnit_MathFloatGreaterEqual_IMPLEMENT_STATIC_VIRTUAL_METHOD(void, Execute, const FRigUnitContext& Context)
 {
     DECLARE_SCOPE_HIERARCHICAL_COUNTER_RIGUNIT()
 	Result = A >= B;
 }
 
-void FRigUnit_MathFloatLessEqual::Execute(const FRigUnitContext& Context)
+UE_RigUnit_MathFloatLessEqual_IMPLEMENT_STATIC_VIRTUAL_METHOD(void, Execute, const FRigUnitContext& Context)
 {
     DECLARE_SCOPE_HIERARCHICAL_COUNTER_RIGUNIT()
 	Result = A <= B;
 }
 
-void FRigUnit_MathFloatIsNearlyZero::Execute(const FRigUnitContext& Context)
+UE_RigUnit_MathFloatIsNearlyZero_IMPLEMENT_STATIC_VIRTUAL_METHOD(void, Execute, const FRigUnitContext& Context)
 {
     DECLARE_SCOPE_HIERARCHICAL_COUNTER_RIGUNIT()
 	if(Tolerance < 0.f)
@@ -188,7 +188,7 @@ void FRigUnit_MathFloatIsNearlyZero::Execute(const FRigUnitContext& Context)
 	Result = FMath::IsNearlyZero(Value, FMath::Max<float>(Tolerance, SMALL_NUMBER));
 }
 
-void FRigUnit_MathFloatIsNearlyEqual::Execute(const FRigUnitContext& Context)
+UE_RigUnit_MathFloatIsNearlyEqual_IMPLEMENT_STATIC_VIRTUAL_METHOD(void, Execute, const FRigUnitContext& Context)
 {
     DECLARE_SCOPE_HIERARCHICAL_COUNTER_RIGUNIT()
 	if(Tolerance < 0.f)
@@ -198,43 +198,43 @@ void FRigUnit_MathFloatIsNearlyEqual::Execute(const FRigUnitContext& Context)
 	Result = FMath::IsNearlyEqual(A, B, FMath::Max<float>(Tolerance, SMALL_NUMBER));
 }
 
-void FRigUnit_MathFloatSelectBool::Execute(const FRigUnitContext& Context)
+UE_RigUnit_MathFloatSelectBool_IMPLEMENT_STATIC_VIRTUAL_METHOD(void, Execute, const FRigUnitContext& Context)
 {
     DECLARE_SCOPE_HIERARCHICAL_COUNTER_RIGUNIT()
 	Result = Condition ? IfTrue : IfFalse;
 }
 
-void FRigUnit_MathFloatDeg::Execute(const FRigUnitContext& Context)
+UE_RigUnit_MathFloatDeg_IMPLEMENT_STATIC_VIRTUAL_METHOD(void, Execute, const FRigUnitContext& Context)
 {
     DECLARE_SCOPE_HIERARCHICAL_COUNTER_RIGUNIT()
 	Result = FMath::RadiansToDegrees(Value);
 }
 
-void FRigUnit_MathFloatRad::Execute(const FRigUnitContext& Context)
+UE_RigUnit_MathFloatRad_IMPLEMENT_STATIC_VIRTUAL_METHOD(void, Execute, const FRigUnitContext& Context)
 {
     DECLARE_SCOPE_HIERARCHICAL_COUNTER_RIGUNIT()
 	Result = FMath::DegreesToRadians(Value);
 }
 
-void FRigUnit_MathFloatSin::Execute(const FRigUnitContext& Context)
+UE_RigUnit_MathFloatSin_IMPLEMENT_STATIC_VIRTUAL_METHOD(void, Execute, const FRigUnitContext& Context)
 {
     DECLARE_SCOPE_HIERARCHICAL_COUNTER_RIGUNIT()
 	Result = FMath::Sin(Value);
 }
 
-void FRigUnit_MathFloatCos::Execute(const FRigUnitContext& Context)
+UE_RigUnit_MathFloatCos_IMPLEMENT_STATIC_VIRTUAL_METHOD(void, Execute, const FRigUnitContext& Context)
 {
     DECLARE_SCOPE_HIERARCHICAL_COUNTER_RIGUNIT()
 	Result = FMath::Cos(Value);
 }
 
-void FRigUnit_MathFloatTan::Execute(const FRigUnitContext& Context)
+UE_RigUnit_MathFloatTan_IMPLEMENT_STATIC_VIRTUAL_METHOD(void, Execute, const FRigUnitContext& Context)
 {
     DECLARE_SCOPE_HIERARCHICAL_COUNTER_RIGUNIT()
 	Result = FMath::Tan(Value);
 }
 
-void FRigUnit_MathFloatAsin::Execute(const FRigUnitContext& Context)
+UE_RigUnit_MathFloatAsin_IMPLEMENT_STATIC_VIRTUAL_METHOD(void, Execute, const FRigUnitContext& Context)
 {
     DECLARE_SCOPE_HIERARCHICAL_COUNTER_RIGUNIT()
 	if (!FMath::IsWithinInclusive<float>(Value, -1.f, 1.f))
@@ -246,7 +246,7 @@ void FRigUnit_MathFloatAsin::Execute(const FRigUnitContext& Context)
 	Result = FMath::Asin(Value);
 }
 
-void FRigUnit_MathFloatAcos::Execute(const FRigUnitContext& Context)
+UE_RigUnit_MathFloatAcos_IMPLEMENT_STATIC_VIRTUAL_METHOD(void, Execute, const FRigUnitContext& Context)
 {
     DECLARE_SCOPE_HIERARCHICAL_COUNTER_RIGUNIT()
 	if (!FMath::IsWithinInclusive<float>(Value, -1.f, 1.f))
@@ -258,13 +258,13 @@ void FRigUnit_MathFloatAcos::Execute(const FRigUnitContext& Context)
 	Result = FMath::Acos(Value);
 }
 
-void FRigUnit_MathFloatAtan::Execute(const FRigUnitContext& Context)
+UE_RigUnit_MathFloatAtan_IMPLEMENT_STATIC_VIRTUAL_METHOD(void, Execute, const FRigUnitContext& Context)
 {
     DECLARE_SCOPE_HIERARCHICAL_COUNTER_RIGUNIT()
 	Result = FMath::Atan(Value);
 }
 
-void FRigUnit_MathFloatLawOfCosine::Execute(const FRigUnitContext& Context)
+UE_RigUnit_MathFloatLawOfCosine_IMPLEMENT_STATIC_VIRTUAL_METHOD(void, Execute, const FRigUnitContext& Context)
 {
     DECLARE_SCOPE_HIERARCHICAL_COUNTER_RIGUNIT()
 	if ((A <= 0.f) || (B <= 0.f) || (C <= 0.f) || (A + B < C) || (A + C < B) || (B + C < A))
diff --git a/Engine/Plugins/Experimental/ControlRig/Source/ControlRig/Private/Units/Math/RigUnit_MathFloat.h b/Engine/Plugins/Experimental/ControlRig/Source/ControlRig/Private/Units/Math/RigUnit_MathFloat.h
--- a/Engine/Plugins/Experimental/ControlRig/Source/ControlRig/Private/Units/Math/RigUnit_MathFloat.h
+++ b/Engine/Plugins/Experimental/ControlRig/Source/ControlRig/Private/Units/Math/RigUnit_MathFloat.h
@@ -112,6 +112,7 @@ struct FRigUnit_MathFloatAdd : public FRigUnit_MathFloatBinaryOp
 {
 	GENERATED_BODY()
 
+	STATIC_VIRTUAL_METHOD()
 	virtual void Execute(const FRigUnitContext& Context) override;
 };
 
@@ -123,6 +124,7 @@ struct FRigUnit_MathFloatSub : public FRigUnit_MathFloatBinaryOp
 {
 	GENERATED_BODY()
 
+	STATIC_VIRTUAL_METHOD()
 	virtual void Execute(const FRigUnitContext& Context) override;
 };
 
@@ -140,6 +142,7 @@ struct FRigUnit_MathFloatMul : public FRigUnit_MathFloatBinaryOp
 		B = 1.f;
 	}
 
+	STATIC_VIRTUAL_METHOD()
 	virtual void Execute(const FRigUnitContext& Context) override;
 };
 
@@ -156,6 +159,7 @@ struct FRigUnit_MathFloatDiv : public FRigUnit_MathFloatBinaryOp
 		B = 1.f;
 	}
 
+	STATIC_VIRTUAL_METHOD()
 	virtual void Execute(const FRigUnitContext& Context) override;
 };
 
@@ -173,6 +177,7 @@ struct FRigUnit_MathFloatMod : public FRigUnit_MathFloatBinaryOp
 		B = 1.f;
 	}
 
+	STATIC_VIRTUAL_METHOD()
 	virtual void Execute(const FRigUnitContext& Context) override;
 };
 
@@ -184,6 +189,7 @@ struct FRigUnit_MathFloatMin : public FRigUnit_MathFloatBinaryOp
 {
 	GENERATED_BODY()
 
+	STATIC_VIRTUAL_METHOD()
 	virtual void Execute(const FRigUnitContext& Context) override;
 };
 
@@ -195,6 +201,7 @@ struct FRigUnit_MathFloatMax : public FRigUnit_MathFloatBinaryOp
 {
 	GENERATED_BODY()
 
+	STATIC_VIRTUAL_METHOD()
 	virtual void Execute(const FRigUnitContext& Context) override;
 };
 
@@ -212,6 +219,7 @@ struct FRigUnit_MathFloatPow : public FRigUnit_MathFloatBinaryOp
 		B = 1.f;
 	}
 
+	STATIC_VIRTUAL_METHOD()
 	virtual void Execute(const FRigUnitContext& Context) override;
 };
 
@@ -223,6 +231,7 @@ struct FRigUnit_MathFloatSqrt : public FRigUnit_MathFloatUnaryOp
 {
 	GENERATED_BODY()
 
+	STATIC_VIRTUAL_METHOD()
 	virtual void Execute(const FRigUnitContext& Context) override;
 };
 
@@ -234,6 +243,7 @@ struct FRigUnit_MathFloatNegate : public FRigUnit_MathFloatUnaryOp
 {
 	GENERATED_BODY()
 
+	STATIC_VIRTUAL_METHOD()
 	virtual void Execute(const FRigUnitContext& Context) override;
 };
 
@@ -245,6 +255,7 @@ struct FRigUnit_MathFloatAbs : public FRigUnit_MathFloatUnaryOp
 {
 	GENERATED_BODY()
 
+	STATIC_VIRTUAL_METHOD()
 	virtual void Execute(const FRigUnitContext& Context) override;
 };
 
@@ -256,6 +267,7 @@ struct FRigUnit_MathFloatFloor : public FRigUnit_MathFloatUnaryOp
 {
 	GENERATED_BODY()
 
+	STATIC_VIRTUAL_METHOD()
 	virtual void Execute(const FRigUnitContext& Context) override;
 };
 
@@ -267,6 +279,7 @@ struct FRigUnit_MathFloatCeil : public FRigUnit_MathFloatUnaryOp
 {
 	GENERATED_BODY()
 
+	STATIC_VIRTUAL_METHOD()
 	virtual void Execute(const FRigUnitContext& Context) override;
 };
 
@@ -278,6 +291,7 @@ struct FRigUnit_MathFloatRound : public FRigUnit_MathFloatUnaryOp
 {
 	GENERATED_BODY()
 
+	STATIC_VIRTUAL_METHOD()
 	virtual void Execute(const FRigUnitContext& Context) override;
 };
 
@@ -289,6 +303,7 @@ struct FRigUnit_MathFloatSign : public FRigUnit_MathFloatUnaryOp
 {
 	GENERATED_BODY()
 
+	STATIC_VIRTUAL_METHOD()
 	virtual void Execute(const FRigUnitContext& Context) override;
 };
 
@@ -305,6 +320,7 @@ struct FRigUnit_MathFloatClamp : public FRigUnit_MathFloatBase
 		Value = Minimum = Maximum = Result = 0.f;
 	}
 
+	STATIC_VIRTUAL_METHOD()
 	virtual void Execute(const FRigUnitContext& Context) override;
 
 	UPROPERTY(meta=(Input))
@@ -334,6 +350,7 @@ struct FRigUnit_MathFloatLerp : public FRigUnit_MathFloatBase
 		B = 1.f;
 	}
 
+	STATIC_VIRTUAL_METHOD()
 	virtual void Execute(const FRigUnitContext& Context) override;
 
 	UPROPERTY(meta=(Input))
@@ -364,6 +381,7 @@ struct FRigUnit_MathFloatRemap : public FRigUnit_MathFloatBase
 		bClamp = false;
 	}
 
+	STATIC_VIRTUAL_METHOD()
 	virtual void Execute(const FRigUnitContext& Context) override;
 
 	UPROPERTY(meta=(Input))
@@ -396,6 +414,7 @@ USTRUCT(meta=(DisplayName="Equals", PrototypeName="Equals", Keywords="Same,=="))
 struct FRigUnit_MathFloatEquals : public FRigUnit_MathFloatBase
 {
 	GENERATED_BODY()
+	STATIC_VIRTUAL_METHOD()
 	virtual void Execute(const FRigUnitContext& Context) override;
 
 	FRigUnit_MathFloatEquals()
@@ -421,6 +440,7 @@ USTRUCT(meta=(DisplayName="Not Equals", PrototypeName="NotEquals", Keywords="Dif
 struct FRigUnit_MathFloatNotEquals : public FRigUnit_MathFloatBase
 {
 	GENERATED_BODY()
+	STATIC_VIRTUAL_METHOD()
 	virtual void Execute(const FRigUnitContext& Context) override;
 
 	FRigUnit_MathFloatNotEquals()
@@ -453,6 +473,7 @@ struct FRigUnit_MathFloatGreater : public FRigUnit_MathFloatBase
 		Result = false;
 	}
 
+	STATIC_VIRTUAL_METHOD()
 	virtual void Execute(const FRigUnitContext& Context) override;
 
 	UPROPERTY(meta=(Input))
@@ -479,6 +500,7 @@ struct FRigUnit_MathFloatLess : public FRigUnit_MathFloatBase
 		Result = false;
 	}
 
+	STATIC_VIRTUAL_METHOD()
 	virtual void Execute(const FRigUnitContext& Context) override;
 
 	UPROPERTY(meta=(Input))
@@ -505,6 +527,7 @@ struct FRigUnit_MathFloatGreaterEqual : public FRigUnit_MathFloatBase
 		Result = true;
 	}
 
+	STATIC_VIRTUAL_METHOD()
 	virtual void Execute(const FRigUnitContext& Context) override;
 
 	UPROPERTY(meta=(Input))
@@ -531,6 +554,7 @@ struct FRigUnit_MathFloatLessEqual : public FRigUnit_MathFloatBase
 		Result = true;
 	}
 
+	STATIC_VIRTUAL_METHOD()
 	virtual void Execute(const FRigUnitContext& Context) override;
 
 	UPROPERTY(meta=(Input))
@@ -557,6 +581,7 @@ struct FRigUnit_MathFloatIsNearlyZero : public FRigUnit_MathFloatBase
 		Result = true;
 	}
 	
+	STATIC_VIRTUAL_METHOD()
 	virtual void Execute(const FRigUnitContext& Context) override;
 
 	UPROPERTY(meta=(Input))
@@ -583,6 +608,7 @@ struct FRigUnit_MathFloatIsNearlyEqual : public FRigUnit_MathFloatBase
 		Result = true;
 	}
 
+	STATIC_VIRTUAL_METHOD()
 	virtual void Execute(const FRigUnitContext& Context) override;
 
 	UPROPERTY(meta=(Input))
@@ -612,6 +638,7 @@ struct FRigUnit_MathFloatSelectBool : public FRigUnit_MathFloatBase
 		IfTrue = IfFalse = Result = 0.f;
 	}
 
+	STATIC_VIRTUAL_METHOD()
 	virtual void Execute(const FRigUnitContext& Context) override;
 
 	UPROPERTY(meta=(Input))
@@ -635,6 +662,7 @@ struct FRigUnit_MathFloatDeg : public FRigUnit_MathFloatUnaryOp
 {
 	GENERATED_BODY()
 
+	STATIC_VIRTUAL_METHOD()
 	virtual void Execute(const FRigUnitContext& Context) override;
 };
 
@@ -646,6 +674,7 @@ struct FRigUnit_MathFloatRad : public FRigUnit_MathFloatUnaryOp
 {
 	GENERATED_BODY()
 
+	STATIC_VIRTUAL_METHOD()
 	virtual void Execute(const FRigUnitContext& Context) override;
 };
 
@@ -657,6 +686,7 @@ struct FRigUnit_MathFloatSin : public FRigUnit_MathFloatUnaryOp
 {
 	GENERATED_BODY()
 
+	STATIC_VIRTUAL_METHOD()
 	virtual void Execute(const FRigUnitContext& Context) override;
 };
 
@@ -668,6 +698,7 @@ struct FRigUnit_MathFloatCos : public FRigUnit_MathFloatUnaryOp
 {
 	GENERATED_BODY()
 
+	STATIC_VIRTUAL_METHOD()
 	virtual void Execute(const FRigUnitContext& Context) override;
 };
 
@@ -679,6 +710,7 @@ struct FRigUnit_MathFloatTan : public FRigUnit_MathFloatUnaryOp
 {
 	GENERATED_BODY()
 
+	STATIC_VIRTUAL_METHOD()
 	virtual void Execute(const FRigUnitContext& Context) override;
 };
 
@@ -690,6 +722,7 @@ struct FRigUnit_MathFloatAsin : public FRigUnit_MathFloatUnaryOp
 {
 	GENERATED_BODY()
 
+	STATIC_VIRTUAL_METHOD()
 	virtual void Execute(const FRigUnitContext& Context) override;
 };
 
@@ -701,6 +734,7 @@ struct FRigUnit_MathFloatAcos : public FRigUnit_MathFloatUnaryOp
 {
 	GENERATED_BODY()
 
+	STATIC_VIRTUAL_METHOD()
 	virtual void Execute(const FRigUnitContext& Context) override;
 };
 
@@ -712,6 +746,7 @@ struct FRigUnit_MathFloatAtan : public FRigUnit_MathFloatUnaryOp
 {
 	GENERATED_BODY()
 
+	STATIC_VIRTUAL_METHOD()
 	virtual void Execute(const FRigUnitContext& Context) override;
 };
 
@@ -729,6 +764,7 @@ struct FRigUnit_MathFloatLawOfCosine : public FRigUnit_MathFloatBase
 		bValid = false;
 	}
 
+	STATIC_VIRTUAL_METHOD()
 	virtual void Execute(const FRigUnitContext& Context) override;
 
 	UPROPERTY(meta = (Input))
diff --git a/Engine/Plugins/Experimental/ControlRig/Source/ControlRig/Private/Units/Math/RigUnit_MathQuaternion.cpp b/Engine/Plugins/Experimental/ControlRig/Source/ControlRig/Private/Units/Math/RigUnit_MathQuaternion.cpp
--- a/Engine/Plugins/Experimental/ControlRig/Source/ControlRig/Private/Units/Math/RigUnit_MathQuaternion.cpp
+++ b/Engine/Plugins/Experimental/ControlRig/Source/ControlRig/Private/Units/Math/RigUnit_MathQuaternion.cpp
@@ -3,7 +3,7 @@
 #include "Units/Math/RigUnit_MathQuaternion.h"
 #include "Units/RigUnitContext.h"
 
-void FRigUnit_MathQuaternionFromAxisAndAngle::Execute(const FRigUnitContext& Context)
+UE_RigUnit_MathQuaternionFromAxisAndAngle_IMPLEMENT_STATIC_VIRTUAL_METHOD(void, Execute, const FRigUnitContext& Context)
 {
     DECLARE_SCOPE_HIERARCHICAL_COUNTER_RIGUNIT()
 	if (Axis.IsNearlyZero())
@@ -15,19 +15,19 @@ void FRigUnit_MathQuaternionFromAxisAndAngle::Execute(const FRigUnitContext& Con
 	Result = FQuat(Axis.GetUnsafeNormal(), Angle);
 }
 
-void FRigUnit_MathQuaternionFromEuler::Execute(const FRigUnitContext& Context)
+UE_RigUnit_MathQuaternionFromEuler_IMPLEMENT_STATIC_VIRTUAL_METHOD(void, Execute, const FRigUnitContext& Context)
 {
     DECLARE_SCOPE_HIERARCHICAL_COUNTER_RIGUNIT()
 	Result = FControlRigMathLibrary::QuatFromEuler(Euler, RotationOrder);
 }
 
-void FRigUnit_MathQuaternionFromRotator::Execute(const FRigUnitContext& Context)
+UE_RigUnit_MathQuaternionFromRotator_IMPLEMENT_STATIC_VIRTUAL_METHOD(void, Execute, const FRigUnitContext& Context)
 {
     DECLARE_SCOPE_HIERARCHICAL_COUNTER_RIGUNIT()
 	Result = FQuat(Rotator);
 }
 
-void FRigUnit_MathQuaternionFromTwoVectors::Execute(const FRigUnitContext& Context)
+UE_RigUnit_MathQuaternionFromTwoVectors_IMPLEMENT_STATIC_VIRTUAL_METHOD(void, Execute, const FRigUnitContext& Context)
 {
     DECLARE_SCOPE_HIERARCHICAL_COUNTER_RIGUNIT()
 	if (A.IsNearlyZero() || B.IsNearlyZero())
@@ -46,79 +46,79 @@ void FRigUnit_MathQuaternionFromTwoVectors::Execute(const FRigUnitContext& Conte
 	Result = FQuat::FindBetweenVectors(A, B).GetNormalized();
 }
 
-void FRigUnit_MathQuaternionToAxisAndAngle::Execute(const FRigUnitContext& Context)
+UE_RigUnit_MathQuaternionToAxisAndAngle_IMPLEMENT_STATIC_VIRTUAL_METHOD(void, Execute, const FRigUnitContext& Context)
 {
     DECLARE_SCOPE_HIERARCHICAL_COUNTER_RIGUNIT()
 	Value.ToAxisAndAngle(Axis, Angle);
 }
 
-void FRigUnit_MathQuaternionToEuler::Execute(const FRigUnitContext& Context)
+UE_RigUnit_MathQuaternionToEuler_IMPLEMENT_STATIC_VIRTUAL_METHOD(void, Execute, const FRigUnitContext& Context)
 {
     DECLARE_SCOPE_HIERARCHICAL_COUNTER_RIGUNIT()
 	Result = FControlRigMathLibrary::EulerFromQuat(Value, RotationOrder);
 }
 
-void FRigUnit_MathQuaternionToRotator::Execute(const FRigUnitContext& Context)
+UE_RigUnit_MathQuaternionToRotator_IMPLEMENT_STATIC_VIRTUAL_METHOD(void, Execute, const FRigUnitContext& Context)
 {
     DECLARE_SCOPE_HIERARCHICAL_COUNTER_RIGUNIT()
 	Result = Value.Rotator();
 }
 
-void FRigUnit_MathQuaternionMul::Execute(const FRigUnitContext& Context)
+UE_RigUnit_MathQuaternionMul_IMPLEMENT_STATIC_VIRTUAL_METHOD(void, Execute, const FRigUnitContext& Context)
 {
     DECLARE_SCOPE_HIERARCHICAL_COUNTER_RIGUNIT()
 	Result = A * B;
 }
 
-void FRigUnit_MathQuaternionInverse::Execute(const FRigUnitContext& Context)
+UE_RigUnit_MathQuaternionInverse_IMPLEMENT_STATIC_VIRTUAL_METHOD(void, Execute, const FRigUnitContext& Context)
 {
     DECLARE_SCOPE_HIERARCHICAL_COUNTER_RIGUNIT()
 	Result = Value.Inverse();
 }
 
-void FRigUnit_MathQuaternionSlerp::Execute(const FRigUnitContext& Context)
+UE_RigUnit_MathQuaternionSlerp_IMPLEMENT_STATIC_VIRTUAL_METHOD(void, Execute, const FRigUnitContext& Context)
 {
     DECLARE_SCOPE_HIERARCHICAL_COUNTER_RIGUNIT()
 	Result = FQuat::Slerp(A, B, T);
 }
 
-void FRigUnit_MathQuaternionEquals::Execute(const FRigUnitContext& Context)
+UE_RigUnit_MathQuaternionEquals_IMPLEMENT_STATIC_VIRTUAL_METHOD(void, Execute, const FRigUnitContext& Context)
 {
     DECLARE_SCOPE_HIERARCHICAL_COUNTER_RIGUNIT()
 	Result = A == B;
 }
 
-void FRigUnit_MathQuaternionNotEquals::Execute(const FRigUnitContext& Context)
+UE_RigUnit_MathQuaternionNotEquals_IMPLEMENT_STATIC_VIRTUAL_METHOD(void, Execute, const FRigUnitContext& Context)
 {
     DECLARE_SCOPE_HIERARCHICAL_COUNTER_RIGUNIT()
 	Result = A != B;
 }
 
-void FRigUnit_MathQuaternionSelectBool::Execute(const FRigUnitContext& Context)
+UE_RigUnit_MathQuaternionSelectBool_IMPLEMENT_STATIC_VIRTUAL_METHOD(void, Execute, const FRigUnitContext& Context)
 {
     DECLARE_SCOPE_HIERARCHICAL_COUNTER_RIGUNIT()
 	Result = Condition ? IfTrue : IfFalse;
 }
 
-void FRigUnit_MathQuaternionDot::Execute(const FRigUnitContext& Context)
+UE_RigUnit_MathQuaternionDot_IMPLEMENT_STATIC_VIRTUAL_METHOD(void, Execute, const FRigUnitContext& Context)
 {
     DECLARE_SCOPE_HIERARCHICAL_COUNTER_RIGUNIT()
 	Result = A | B;
 }
 
-void FRigUnit_MathQuaternionUnit::Execute(const FRigUnitContext& Context)
+UE_RigUnit_MathQuaternionUnit_IMPLEMENT_STATIC_VIRTUAL_METHOD(void, Execute, const FRigUnitContext& Context)
 {
     DECLARE_SCOPE_HIERARCHICAL_COUNTER_RIGUNIT()
 	Result = Value.GetNormalized();
 }
 
-void FRigUnit_MathQuaternionRotateVector::Execute(const FRigUnitContext& Context)
+UE_RigUnit_MathQuaternionRotateVector_IMPLEMENT_STATIC_VIRTUAL_METHOD(void, Execute, const FRigUnitContext& Context)
 {
     DECLARE_SCOPE_HIERARCHICAL_COUNTER_RIGUNIT()
 	Result = Quaternion.RotateVector(Vector);
 }
 
-void FRigUnit_MathQuaternionGetAxis::Execute(const FRigUnitContext& Context)
+UE_RigUnit_MathQuaternionGetAxis_IMPLEMENT_STATIC_VIRTUAL_METHOD(void, Execute, const FRigUnitContext& Context)
 {
     DECLARE_SCOPE_HIERARCHICAL_COUNTER_RIGUNIT()
 	switch (Axis)
@@ -143,7 +143,7 @@ void FRigUnit_MathQuaternionGetAxis::Execute(const FRigUnitContext& Context)
 }
 
 
-void FRigUnit_MathQuaternionSwingTwist::Execute(const FRigUnitContext& Context)
+UE_RigUnit_MathQuaternionSwingTwist_IMPLEMENT_STATIC_VIRTUAL_METHOD(void, Execute, const FRigUnitContext& Context)
 {
     DECLARE_SCOPE_HIERARCHICAL_COUNTER_RIGUNIT()
 	if (TwistAxis.IsNearlyZero())
diff --git a/Engine/Plugins/Experimental/ControlRig/Source/ControlRig/Private/Units/Math/RigUnit_MathQuaternion.h b/Engine/Plugins/Experimental/ControlRig/Source/ControlRig/Private/Units/Math/RigUnit_MathQuaternion.h
--- a/Engine/Plugins/Experimental/ControlRig/Source/ControlRig/Private/Units/Math/RigUnit_MathQuaternion.h
+++ b/Engine/Plugins/Experimental/ControlRig/Source/ControlRig/Private/Units/Math/RigUnit_MathQuaternion.h
@@ -64,6 +64,7 @@ struct FRigUnit_MathQuaternionFromAxisAndAngle : public FRigUnit_MathQuaternionB
 		Result = FQuat::Identity;
 	}
 
+	STATIC_VIRTUAL_METHOD()
 	virtual void Execute(const FRigUnitContext& Context) override;
 
 	UPROPERTY(meta=(Input))
@@ -91,6 +92,7 @@ struct FRigUnit_MathQuaternionFromEuler : public FRigUnit_MathQuaternionBase
 		Result = FQuat::Identity;
 	}
 
+	STATIC_VIRTUAL_METHOD()
 	virtual void Execute(const FRigUnitContext& Context) override;
 
 	UPROPERTY(meta = (Input))
@@ -117,6 +119,7 @@ struct FRigUnit_MathQuaternionFromRotator : public FRigUnit_MathQuaternionBase
 		Result = FQuat::Identity;
 	}
 
+	STATIC_VIRTUAL_METHOD()
 	virtual void Execute(const FRigUnitContext& Context) override;
 
 	UPROPERTY(meta=(Input))
@@ -140,6 +143,7 @@ struct FRigUnit_MathQuaternionFromTwoVectors : public FRigUnit_MathQuaternionBas
 		Result = FQuat::Identity;
 	}
 
+	STATIC_VIRTUAL_METHOD()
 	virtual void Execute(const FRigUnitContext& Context) override;
 
 	UPROPERTY(meta=(Input))
@@ -167,6 +171,7 @@ struct FRigUnit_MathQuaternionToAxisAndAngle : public FRigUnit_MathQuaternionBas
 		Angle = 0.f;
 	}
 
+	STATIC_VIRTUAL_METHOD()
 	virtual void Execute(const FRigUnitContext& Context) override;
 
 	UPROPERTY(meta=(Input))
@@ -194,6 +199,7 @@ struct FRigUnit_MathQuaternionToEuler : public FRigUnit_MathQuaternionBase
 		Result = FVector::ZeroVector;
 	}
 
+	STATIC_VIRTUAL_METHOD()
 	virtual void Execute(const FRigUnitContext& Context) override;
 
 	UPROPERTY(meta=(Input))
@@ -220,6 +226,7 @@ struct FRigUnit_MathQuaternionToRotator : public FRigUnit_MathQuaternionBase
 		Result = FRotator::ZeroRotator;
 	}
 
+	STATIC_VIRTUAL_METHOD()
 	virtual void Execute(const FRigUnitContext& Context) override;
 
 	UPROPERTY(meta=(Input))
@@ -236,6 +243,7 @@ USTRUCT(meta=(DisplayName="Multiply", PrototypeName="Multiply", Keywords="Produc
 struct FRigUnit_MathQuaternionMul : public FRigUnit_MathQuaternionBinaryOp
 {
 	GENERATED_BODY()
+	STATIC_VIRTUAL_METHOD()
 	virtual void Execute(const FRigUnitContext& Context) override;
 };
 
@@ -246,6 +254,7 @@ USTRUCT(meta=(DisplayName="Inverse", PrototypeName="Inverse"))
 struct FRigUnit_MathQuaternionInverse : public FRigUnit_MathQuaternionUnaryOp
 {
 	GENERATED_BODY()
+	STATIC_VIRTUAL_METHOD()
 	virtual void Execute(const FRigUnitContext& Context) override;
 };
 
@@ -263,6 +272,7 @@ struct FRigUnit_MathQuaternionSlerp : public FRigUnit_MathQuaternionBase
 		T = 0.f;
 	}
 
+	STATIC_VIRTUAL_METHOD()
 	virtual void Execute(const FRigUnitContext& Context) override;
 
 	UPROPERTY(meta=(Input))
@@ -292,6 +302,7 @@ struct FRigUnit_MathQuaternionEquals : public FRigUnit_MathQuaternionBase
 		Result = true;
 	}	
 
+	STATIC_VIRTUAL_METHOD()
 	virtual void Execute(const FRigUnitContext& Context) override;
 
 	UPROPERTY(meta=(Input))
@@ -318,6 +329,7 @@ struct FRigUnit_MathQuaternionNotEquals : public FRigUnit_MathQuaternionBase
 		Result = false;
 	}
 
+	STATIC_VIRTUAL_METHOD()
 	virtual void Execute(const FRigUnitContext& Context) override;
 
 	UPROPERTY(meta=(Input))
@@ -344,6 +356,7 @@ struct FRigUnit_MathQuaternionSelectBool : public FRigUnit_MathQuaternionBase
 		Condition = true;
 	}
 
+	STATIC_VIRTUAL_METHOD()
 	virtual void Execute(const FRigUnitContext& Context) override;
 
 	UPROPERTY(meta=(Input))
@@ -373,6 +386,7 @@ struct FRigUnit_MathQuaternionDot : public FRigUnit_MathQuaternionBase
 		Result = 0.f;
 	}
 
+	STATIC_VIRTUAL_METHOD()
 	virtual void Execute(const FRigUnitContext& Context) override;
 
 	UPROPERTY(meta=(Input))
@@ -392,6 +406,7 @@ USTRUCT(meta=(DisplayName="Unit", PrototypeName="Unit", Keywords="Normalize"))
 struct FRigUnit_MathQuaternionUnit : public FRigUnit_MathQuaternionUnaryOp
 {
 	GENERATED_BODY()
+	STATIC_VIRTUAL_METHOD()
 	virtual void Execute(const FRigUnitContext& Context) override;
 };
 
@@ -409,6 +424,7 @@ struct FRigUnit_MathQuaternionRotateVector : public FRigUnit_MathQuaternionBase
 		Vector = Result = FVector::ZeroVector;
 	}
 
+	STATIC_VIRTUAL_METHOD()
 	virtual void Execute(const FRigUnitContext& Context) override;
 
 	UPROPERTY(meta=(Input))
@@ -436,6 +452,7 @@ struct FRigUnit_MathQuaternionGetAxis: public FRigUnit_MathQuaternionBase
 		Result = FVector(1.f, 0.f, 0.f);
 	}
 
+	STATIC_VIRTUAL_METHOD()
 	virtual void Execute(const FRigUnitContext& Context) override;
 
 	UPROPERTY(meta = (Input))
@@ -463,6 +480,7 @@ struct FRigUnit_MathQuaternionSwingTwist : public FRigUnit_MathQuaternionBase
 		TwistAxis = FVector(1.f, 0.f, 0.f);
 	}
 
+	STATIC_VIRTUAL_METHOD()
 	virtual void Execute(const FRigUnitContext& Context) override;
 
 	UPROPERTY(meta = (Input))
diff --git a/Engine/Plugins/Experimental/ControlRig/Source/ControlRig/Private/Units/Math/RigUnit_MathTransform.cpp b/Engine/Plugins/Experimental/ControlRig/Source/ControlRig/Private/Units/Math/RigUnit_MathTransform.cpp
--- a/Engine/Plugins/Experimental/ControlRig/Source/ControlRig/Private/Units/Math/RigUnit_MathTransform.cpp
+++ b/Engine/Plugins/Experimental/ControlRig/Source/ControlRig/Private/Units/Math/RigUnit_MathTransform.cpp
@@ -4,61 +4,61 @@
 #include "Math/ControlRigMathLibrary.h"
 #include "Units/RigUnitContext.h"
 
-void FRigUnit_MathTransformFromEulerTransform::Execute(const FRigUnitContext& Context)
+UE_RigUnit_MathTransformFromEulerTransform_IMPLEMENT_STATIC_VIRTUAL_METHOD(void, Execute, const FRigUnitContext& Context)
 {
     DECLARE_SCOPE_HIERARCHICAL_COUNTER_RIGUNIT()
 	Result = EulerTransform.ToFTransform();
 }
 
-void FRigUnit_MathTransformToEulerTransform::Execute(const FRigUnitContext& Context)
+UE_RigUnit_MathTransformToEulerTransform_IMPLEMENT_STATIC_VIRTUAL_METHOD(void, Execute, const FRigUnitContext& Context)
 {
     DECLARE_SCOPE_HIERARCHICAL_COUNTER_RIGUNIT()
 	Result.FromFTransform(Value);
 }
 
-void FRigUnit_MathTransformMul::Execute(const FRigUnitContext& Context)
+UE_RigUnit_MathTransformMul_IMPLEMENT_STATIC_VIRTUAL_METHOD(void, Execute, const FRigUnitContext& Context)
 {
     DECLARE_SCOPE_HIERARCHICAL_COUNTER_RIGUNIT()
 	Result = A * B;
 }
 
-void FRigUnit_MathTransformMakeRelative::Execute(const FRigUnitContext& Context)
+UE_RigUnit_MathTransformMakeRelative_IMPLEMENT_STATIC_VIRTUAL_METHOD(void, Execute, const FRigUnitContext& Context)
 {
     DECLARE_SCOPE_HIERARCHICAL_COUNTER_RIGUNIT()
 	Local = Global.GetRelativeTransform(Parent);
 }
 
-void FRigUnit_MathTransformInverse::Execute(const FRigUnitContext& Context)
+UE_RigUnit_MathTransformInverse_IMPLEMENT_STATIC_VIRTUAL_METHOD(void, Execute, const FRigUnitContext& Context)
 {
     DECLARE_SCOPE_HIERARCHICAL_COUNTER_RIGUNIT()
 	Result = Value.Inverse();
 }
 
-void FRigUnit_MathTransformLerp::Execute(const FRigUnitContext& Context)
+UE_RigUnit_MathTransformLerp_IMPLEMENT_STATIC_VIRTUAL_METHOD(void, Execute, const FRigUnitContext& Context)
 {
     DECLARE_SCOPE_HIERARCHICAL_COUNTER_RIGUNIT()
 	Result = FControlRigMathLibrary::LerpTransform(A, B, T);
 }
 
-void FRigUnit_MathTransformSelectBool::Execute(const FRigUnitContext& Context)
+UE_RigUnit_MathTransformSelectBool_IMPLEMENT_STATIC_VIRTUAL_METHOD(void, Execute, const FRigUnitContext& Context)
 {
     DECLARE_SCOPE_HIERARCHICAL_COUNTER_RIGUNIT()
 	Result = Condition ? IfTrue : IfFalse;
 }
 
-void FRigUnit_MathTransformRotateVector::Execute(const FRigUnitContext& Context)
+UE_RigUnit_MathTransformRotateVector_IMPLEMENT_STATIC_VIRTUAL_METHOD(void, Execute, const FRigUnitContext& Context)
 {
     DECLARE_SCOPE_HIERARCHICAL_COUNTER_RIGUNIT()
 	Result = Transform.TransformVector(Direction);
 }
 
-void FRigUnit_MathTransformTransformVector::Execute(const FRigUnitContext& Context)
+UE_RigUnit_MathTransformTransformVector_IMPLEMENT_STATIC_VIRTUAL_METHOD(void, Execute, const FRigUnitContext& Context)
 {
     DECLARE_SCOPE_HIERARCHICAL_COUNTER_RIGUNIT()
 	Result = Transform.TransformPosition(Location);
 }
 
-void FRigUnit_MathTransformFromSRT::Execute(const FRigUnitContext& Context)
+UE_RigUnit_MathTransformFromSRT_IMPLEMENT_STATIC_VIRTUAL_METHOD(void, Execute, const FRigUnitContext& Context)
 {
     DECLARE_SCOPE_HIERARCHICAL_COUNTER_RIGUNIT()
 	Transform.SetLocation(Location);
diff --git a/Engine/Plugins/Experimental/ControlRig/Source/ControlRig/Private/Units/Math/RigUnit_MathTransform.h b/Engine/Plugins/Experimental/ControlRig/Source/ControlRig/Private/Units/Math/RigUnit_MathTransform.h
--- a/Engine/Plugins/Experimental/ControlRig/Source/ControlRig/Private/Units/Math/RigUnit_MathTransform.h
+++ b/Engine/Plugins/Experimental/ControlRig/Source/ControlRig/Private/Units/Math/RigUnit_MathTransform.h
@@ -64,6 +64,7 @@ struct FRigUnit_MathTransformFromEulerTransform : public FRigUnit_MathTransformB
 		Result = FTransform::Identity;
 	}
 
+	STATIC_VIRTUAL_METHOD()
 	virtual void Execute(const FRigUnitContext& Context) override;
 
 	UPROPERTY(meta=(Input))
@@ -87,6 +88,7 @@ struct FRigUnit_MathTransformToEulerTransform : public FRigUnit_MathTransformBas
 		Result = FEulerTransform::Identity;
 	}
 
+	STATIC_VIRTUAL_METHOD()
 	virtual void Execute(const FRigUnitContext& Context) override;
 
 	UPROPERTY(meta=(Input))
@@ -103,6 +105,7 @@ USTRUCT(meta=(DisplayName="Multiply", PrototypeName="Multiply", Keywords="Produc
 struct FRigUnit_MathTransformMul : public FRigUnit_MathTransformBinaryOp
 {
 	GENERATED_BODY()
+	STATIC_VIRTUAL_METHOD()
 	virtual void Execute(const FRigUnitContext& Context) override;
 };
 
@@ -119,6 +122,7 @@ struct FRigUnit_MathTransformMakeRelative : public FRigUnit_MathTransformBase
 		Global = Parent = Local = FTransform::Identity;
 	}
 
+	STATIC_VIRTUAL_METHOD()
 	virtual void Execute(const FRigUnitContext& Context) override;
 
 	UPROPERTY(meta=(Input))
@@ -138,6 +142,7 @@ USTRUCT(meta=(DisplayName="Inverse", PrototypeName="Inverse"))
 struct FRigUnit_MathTransformInverse : public FRigUnit_MathTransformUnaryOp
 {
 	GENERATED_BODY()
+	STATIC_VIRTUAL_METHOD()
 	virtual void Execute(const FRigUnitContext& Context) override;
 };
 
@@ -155,6 +160,7 @@ struct FRigUnit_MathTransformLerp : public FRigUnit_MathTransformBase
 		T = 0.f;
 	}
 
+	STATIC_VIRTUAL_METHOD()
 	virtual void Execute(const FRigUnitContext& Context) override;
 
 	UPROPERTY(meta=(Input))
@@ -184,6 +190,7 @@ struct FRigUnit_MathTransformSelectBool : public FRigUnit_MathTransformBase
 		IfTrue = IfFalse = Result = FTransform::Identity;
 	}
 
+	STATIC_VIRTUAL_METHOD()
 	virtual void Execute(const FRigUnitContext& Context) override;
 
 	UPROPERTY(meta=(Input))
@@ -213,6 +220,7 @@ struct FRigUnit_MathTransformRotateVector : public FRigUnit_MathTransformBase
 		Direction = Result = FVector::ZeroVector;
 	}
 
+	STATIC_VIRTUAL_METHOD()
 	virtual void Execute(const FRigUnitContext& Context) override;
 
 	UPROPERTY(meta=(Input))
@@ -239,6 +247,7 @@ struct FRigUnit_MathTransformTransformVector : public FRigUnit_MathTransformBase
 		Location = Result = FVector::ZeroVector;
 	}
 	
+	STATIC_VIRTUAL_METHOD()
 	virtual void Execute(const FRigUnitContext& Context) override;
 
 	UPROPERTY(meta=(Input))
@@ -269,6 +278,7 @@ struct FRigUnit_MathTransformFromSRT : public FRigUnit_MathTransformBase
 		EulerTransform = FEulerTransform::Identity;
 	}
 	
+	STATIC_VIRTUAL_METHOD()
 	virtual void Execute(const FRigUnitContext& Context) override;
 
 	UPROPERTY(meta=(Input))
diff --git a/Engine/Plugins/Experimental/ControlRig/Source/ControlRig/Private/Units/Math/RigUnit_MathVector.cpp b/Engine/Plugins/Experimental/ControlRig/Source/ControlRig/Private/Units/Math/RigUnit_MathVector.cpp
--- a/Engine/Plugins/Experimental/ControlRig/Source/ControlRig/Private/Units/Math/RigUnit_MathVector.cpp
+++ b/Engine/Plugins/Experimental/ControlRig/Source/ControlRig/Private/Units/Math/RigUnit_MathVector.cpp
@@ -4,37 +4,37 @@
 #include "Units/RigUnitContext.h"
 #include "Math/ControlRigMathLibrary.h"
 
-void FRigUnit_MathVectorFromFloat::Execute(const FRigUnitContext& Context)
+UE_RigUnit_MathVectorFromFloat_IMPLEMENT_STATIC_VIRTUAL_METHOD(void, Execute, const FRigUnitContext& Context)
 {
     DECLARE_SCOPE_HIERARCHICAL_COUNTER_RIGUNIT()
 	Result = FVector(Value, Value, Value);
 }
 
-void FRigUnit_MathVectorAdd::Execute(const FRigUnitContext& Context)
+UE_RigUnit_MathVectorAdd_IMPLEMENT_STATIC_VIRTUAL_METHOD(void, Execute, const FRigUnitContext& Context)
 {
     DECLARE_SCOPE_HIERARCHICAL_COUNTER_RIGUNIT()
 	Result = A + B;
 }
 
-void FRigUnit_MathVectorSub::Execute(const FRigUnitContext& Context)
+UE_RigUnit_MathVectorSub_IMPLEMENT_STATIC_VIRTUAL_METHOD(void, Execute, const FRigUnitContext& Context)
 {
     DECLARE_SCOPE_HIERARCHICAL_COUNTER_RIGUNIT()
 	Result = A - B;
 }
 
-void FRigUnit_MathVectorMul::Execute(const FRigUnitContext& Context)
+UE_RigUnit_MathVectorMul_IMPLEMENT_STATIC_VIRTUAL_METHOD(void, Execute, const FRigUnitContext& Context)
 {
     DECLARE_SCOPE_HIERARCHICAL_COUNTER_RIGUNIT()
 	Result = A * B;
 }
 
-void FRigUnit_MathVectorScale::Execute(const FRigUnitContext& Context)
+UE_RigUnit_MathVectorScale_IMPLEMENT_STATIC_VIRTUAL_METHOD(void, Execute, const FRigUnitContext& Context)
 {
     DECLARE_SCOPE_HIERARCHICAL_COUNTER_RIGUNIT()
 	Result = Value * Factor;
 }
 
-void FRigUnit_MathVectorDiv::Execute(const FRigUnitContext& Context)
+UE_RigUnit_MathVectorDiv_IMPLEMENT_STATIC_VIRTUAL_METHOD(void, Execute, const FRigUnitContext& Context)
 {
     DECLARE_SCOPE_HIERARCHICAL_COUNTER_RIGUNIT()
 	if(FMath::IsNearlyZero(B.X) || FMath::IsNearlyZero(B.Y) || FMath::IsNearlyZero(B.Z))
@@ -57,7 +57,7 @@ void FRigUnit_MathVectorDiv::Execute(const FRigUnitContext& Context)
 	Result = A / B;
 }
 
-void FRigUnit_MathVectorMod::Execute(const FRigUnitContext& Context)
+UE_RigUnit_MathVectorMod_IMPLEMENT_STATIC_VIRTUAL_METHOD(void, Execute, const FRigUnitContext& Context)
 {
     DECLARE_SCOPE_HIERARCHICAL_COUNTER_RIGUNIT()
 	if(FMath::IsNearlyZero(B.X) || FMath::IsNearlyZero(B.Y) || FMath::IsNearlyZero(B.Z) || B.X < 0.f || B.Y < 0.f || B.Z < 0.f)
@@ -83,7 +83,7 @@ void FRigUnit_MathVectorMod::Execute(const FRigUnitContext& Context)
 	Result.Z = FMath::Fmod(A.Z, B.Z);
 }
 
-void FRigUnit_MathVectorMin::Execute(const FRigUnitContext& Context)
+UE_RigUnit_MathVectorMin_IMPLEMENT_STATIC_VIRTUAL_METHOD(void, Execute, const FRigUnitContext& Context)
 {
     DECLARE_SCOPE_HIERARCHICAL_COUNTER_RIGUNIT()
 	Result.X = FMath::Min<float>(A.X, B.X);
@@ -91,7 +91,7 @@ void FRigUnit_MathVectorMin::Execute(const FRigUnitContext& Context)
 	Result.Z = FMath::Min<float>(A.Z, B.Z);
 }
 
-void FRigUnit_MathVectorMax::Execute(const FRigUnitContext& Context)
+UE_RigUnit_MathVectorMax_IMPLEMENT_STATIC_VIRTUAL_METHOD(void, Execute, const FRigUnitContext& Context)
 {
     DECLARE_SCOPE_HIERARCHICAL_COUNTER_RIGUNIT()
 	Result.X = FMath::Max<float>(A.X, B.X);
@@ -99,13 +99,13 @@ void FRigUnit_MathVectorMax::Execute(const FRigUnitContext& Context)
 	Result.Z = FMath::Max<float>(A.Z, B.Z);
 }
 
-void FRigUnit_MathVectorNegate::Execute(const FRigUnitContext& Context)
+UE_RigUnit_MathVectorNegate_IMPLEMENT_STATIC_VIRTUAL_METHOD(void, Execute, const FRigUnitContext& Context)
 {
     DECLARE_SCOPE_HIERARCHICAL_COUNTER_RIGUNIT()
 	Result = -Value;
 }
 
-void FRigUnit_MathVectorAbs::Execute(const FRigUnitContext& Context)
+UE_RigUnit_MathVectorAbs_IMPLEMENT_STATIC_VIRTUAL_METHOD(void, Execute, const FRigUnitContext& Context)
 {
     DECLARE_SCOPE_HIERARCHICAL_COUNTER_RIGUNIT()
 	Result.X = FMath::Abs(Value.X);
@@ -113,7 +113,7 @@ void FRigUnit_MathVectorAbs::Execute(const FRigUnitContext& Context)
 	Result.Z = FMath::Abs(Value.Z);
 }
 
-void FRigUnit_MathVectorFloor::Execute(const FRigUnitContext& Context)
+UE_RigUnit_MathVectorFloor_IMPLEMENT_STATIC_VIRTUAL_METHOD(void, Execute, const FRigUnitContext& Context)
 {
     DECLARE_SCOPE_HIERARCHICAL_COUNTER_RIGUNIT()
 	Result.X = FMath::FloorToFloat(Value.X);
@@ -121,7 +121,7 @@ void FRigUnit_MathVectorFloor::Execute(const FRigUnitContext& Context)
 	Result.Z = FMath::FloorToFloat(Value.Z);
 }
 
-void FRigUnit_MathVectorCeil::Execute(const FRigUnitContext& Context)
+UE_RigUnit_MathVectorCeil_IMPLEMENT_STATIC_VIRTUAL_METHOD(void, Execute, const FRigUnitContext& Context)
 {
     DECLARE_SCOPE_HIERARCHICAL_COUNTER_RIGUNIT()
 	Result.X = FMath::CeilToFloat(Value.X);
@@ -129,7 +129,7 @@ void FRigUnit_MathVectorCeil::Execute(const FRigUnitContext& Context)
 	Result.Z = FMath::CeilToFloat(Value.Z);
 }
 
-void FRigUnit_MathVectorRound::Execute(const FRigUnitContext& Context)
+UE_RigUnit_MathVectorRound_IMPLEMENT_STATIC_VIRTUAL_METHOD(void, Execute, const FRigUnitContext& Context)
 {
     DECLARE_SCOPE_HIERARCHICAL_COUNTER_RIGUNIT()
 	Result.X = FMath::RoundToFloat(Value.X);
@@ -137,13 +137,13 @@ void FRigUnit_MathVectorRound::Execute(const FRigUnitContext& Context)
 	Result.Z = FMath::RoundToFloat(Value.Z);
 }
 
-void FRigUnit_MathVectorSign::Execute(const FRigUnitContext& Context)
+UE_RigUnit_MathVectorSign_IMPLEMENT_STATIC_VIRTUAL_METHOD(void, Execute, const FRigUnitContext& Context)
 {
     DECLARE_SCOPE_HIERARCHICAL_COUNTER_RIGUNIT()
 	Result = Value.GetSignVector();
 }
 
-void FRigUnit_MathVectorClamp::Execute(const FRigUnitContext& Context)
+UE_RigUnit_MathVectorClamp_IMPLEMENT_STATIC_VIRTUAL_METHOD(void, Execute, const FRigUnitContext& Context)
 {
     DECLARE_SCOPE_HIERARCHICAL_COUNTER_RIGUNIT()
 	Result.X = FMath::Clamp<float>(Value.X, Minimum.X, Maximum.X);
@@ -151,13 +151,13 @@ void FRigUnit_MathVectorClamp::Execute(const FRigUnitContext& Context)
 	Result.Z = FMath::Clamp<float>(Value.Z, Minimum.Z, Maximum.Z);
 }
 
-void FRigUnit_MathVectorLerp::Execute(const FRigUnitContext& Context)
+UE_RigUnit_MathVectorLerp_IMPLEMENT_STATIC_VIRTUAL_METHOD(void, Execute, const FRigUnitContext& Context)
 {
     DECLARE_SCOPE_HIERARCHICAL_COUNTER_RIGUNIT()
 	Result = FMath::Lerp<FVector>(A, B, T);
 }
 
-void FRigUnit_MathVectorRemap::Execute(const FRigUnitContext& Context)
+UE_RigUnit_MathVectorRemap_IMPLEMENT_STATIC_VIRTUAL_METHOD(void, Execute, const FRigUnitContext& Context)
 {
     DECLARE_SCOPE_HIERARCHICAL_COUNTER_RIGUNIT()
 	FVector Ratio(0.f, 0.f, 0.f);
@@ -189,19 +189,19 @@ void FRigUnit_MathVectorRemap::Execute(const FRigUnitContext& Context)
 	Result = FMath::Lerp<FVector>(TargetMinimum, TargetMaximum, Ratio);
 }
 
-void FRigUnit_MathVectorEquals::Execute(const FRigUnitContext& Context)
+UE_RigUnit_MathVectorEquals_IMPLEMENT_STATIC_VIRTUAL_METHOD(void, Execute, const FRigUnitContext& Context)
 {
     DECLARE_SCOPE_HIERARCHICAL_COUNTER_RIGUNIT()
 	Result = A == B;
 }
 
-void FRigUnit_MathVectorNotEquals::Execute(const FRigUnitContext& Context)
+UE_RigUnit_MathVectorNotEquals_IMPLEMENT_STATIC_VIRTUAL_METHOD(void, Execute, const FRigUnitContext& Context)
 {
     DECLARE_SCOPE_HIERARCHICAL_COUNTER_RIGUNIT()
 	Result = A != B;
 }
 
-void FRigUnit_MathVectorIsNearlyZero::Execute(const FRigUnitContext& Context)
+UE_RigUnit_MathVectorIsNearlyZero_IMPLEMENT_STATIC_VIRTUAL_METHOD(void, Execute, const FRigUnitContext& Context)
 {
     DECLARE_SCOPE_HIERARCHICAL_COUNTER_RIGUNIT()
 	if(Tolerance < 0.f)
@@ -211,7 +211,7 @@ void FRigUnit_MathVectorIsNearlyZero::Execute(const FRigUnitContext& Context)
 	Result = Value.IsNearlyZero(FMath::Max<float>(Tolerance, SMALL_NUMBER));
 }
 
-void FRigUnit_MathVectorIsNearlyEqual::Execute(const FRigUnitContext& Context)
+UE_RigUnit_MathVectorIsNearlyEqual_IMPLEMENT_STATIC_VIRTUAL_METHOD(void, Execute, const FRigUnitContext& Context)
 {
     DECLARE_SCOPE_HIERARCHICAL_COUNTER_RIGUNIT()
 	if(Tolerance < 0.f)
@@ -221,55 +221,55 @@ void FRigUnit_MathVectorIsNearlyEqual::Execute(const FRigUnitContext& Context)
 	Result = (A - B).IsNearlyZero(FMath::Max<float>(Tolerance, SMALL_NUMBER));
 }
 
-void FRigUnit_MathVectorSelectBool::Execute(const FRigUnitContext& Context)
+UE_RigUnit_MathVectorSelectBool_IMPLEMENT_STATIC_VIRTUAL_METHOD(void, Execute, const FRigUnitContext& Context)
 {
     DECLARE_SCOPE_HIERARCHICAL_COUNTER_RIGUNIT()
 	Result = Condition ? IfTrue : IfFalse;
 }
 
-void FRigUnit_MathVectorDeg::Execute(const FRigUnitContext& Context)
+UE_RigUnit_MathVectorDeg_IMPLEMENT_STATIC_VIRTUAL_METHOD(void, Execute, const FRigUnitContext& Context)
 {
     DECLARE_SCOPE_HIERARCHICAL_COUNTER_RIGUNIT()
 	Result = FMath::RadiansToDegrees(Value);
 }
 
-void FRigUnit_MathVectorRad::Execute(const FRigUnitContext& Context)
+UE_RigUnit_MathVectorRad_IMPLEMENT_STATIC_VIRTUAL_METHOD(void, Execute, const FRigUnitContext& Context)
 {
     DECLARE_SCOPE_HIERARCHICAL_COUNTER_RIGUNIT()
 	Result = FMath::DegreesToRadians(Value);
 }
 
-void FRigUnit_MathVectorLengthSquared::Execute(const FRigUnitContext& Context)
+UE_RigUnit_MathVectorLengthSquared_IMPLEMENT_STATIC_VIRTUAL_METHOD(void, Execute, const FRigUnitContext& Context)
 {
     DECLARE_SCOPE_HIERARCHICAL_COUNTER_RIGUNIT()
 	Result = Value.SizeSquared();
 }
 
-void FRigUnit_MathVectorLength::Execute(const FRigUnitContext& Context)
+UE_RigUnit_MathVectorLength_IMPLEMENT_STATIC_VIRTUAL_METHOD(void, Execute, const FRigUnitContext& Context)
 {
     DECLARE_SCOPE_HIERARCHICAL_COUNTER_RIGUNIT()
 	Result = Value.Size();
 }
 
-void FRigUnit_MathVectorDistance::Execute(const FRigUnitContext& Context)
+UE_RigUnit_MathVectorDistance_IMPLEMENT_STATIC_VIRTUAL_METHOD(void, Execute, const FRigUnitContext& Context)
 {
     DECLARE_SCOPE_HIERARCHICAL_COUNTER_RIGUNIT()
 	Result = FVector::Distance(A, B);
 }
 
-void FRigUnit_MathVectorCross::Execute(const FRigUnitContext& Context)
+UE_RigUnit_MathVectorCross_IMPLEMENT_STATIC_VIRTUAL_METHOD(void, Execute, const FRigUnitContext& Context)
 {
     DECLARE_SCOPE_HIERARCHICAL_COUNTER_RIGUNIT()
 	Result = A ^ B;
 }
 
-void FRigUnit_MathVectorDot::Execute(const FRigUnitContext& Context)
+UE_RigUnit_MathVectorDot_IMPLEMENT_STATIC_VIRTUAL_METHOD(void, Execute, const FRigUnitContext& Context)
 {
     DECLARE_SCOPE_HIERARCHICAL_COUNTER_RIGUNIT()
 	Result = A | B;
 }
 
-void FRigUnit_MathVectorUnit::Execute(const FRigUnitContext& Context)
+UE_RigUnit_MathVectorUnit_IMPLEMENT_STATIC_VIRTUAL_METHOD(void, Execute, const FRigUnitContext& Context)
 {
     DECLARE_SCOPE_HIERARCHICAL_COUNTER_RIGUNIT()
 	if (Value.IsNearlyZero())
@@ -281,7 +281,7 @@ void FRigUnit_MathVectorUnit::Execute(const FRigUnitContext& Context)
 	Result = Value.GetUnsafeNormal();
 }
 
-void FRigUnit_MathVectorMirror::Execute(const FRigUnitContext& Context)
+UE_RigUnit_MathVectorMirror_IMPLEMENT_STATIC_VIRTUAL_METHOD(void, Execute, const FRigUnitContext& Context)
 {
     DECLARE_SCOPE_HIERARCHICAL_COUNTER_RIGUNIT()
 	if (Normal.IsNearlyZero())
@@ -293,7 +293,7 @@ void FRigUnit_MathVectorMirror::Execute(const FRigUnitContext& Context)
 	Result = Value.MirrorByVector(Normal.GetSafeNormal());
 }
 
-void FRigUnit_MathVectorAngle::Execute(const FRigUnitContext& Context)
+UE_RigUnit_MathVectorAngle_IMPLEMENT_STATIC_VIRTUAL_METHOD(void, Execute, const FRigUnitContext& Context)
 {
     DECLARE_SCOPE_HIERARCHICAL_COUNTER_RIGUNIT()
 	if (A.IsNearlyZero() || B.IsNearlyZero())
@@ -312,7 +312,7 @@ void FRigUnit_MathVectorAngle::Execute(const FRigUnitContext& Context)
 	Result = FQuat::FindBetween(A, B).GetAngle();
 }
 
-void FRigUnit_MathVectorParallel::Execute(const FRigUnitContext& Context)
+UE_RigUnit_MathVectorParallel_IMPLEMENT_STATIC_VIRTUAL_METHOD(void, Execute, const FRigUnitContext& Context)
 {
     DECLARE_SCOPE_HIERARCHICAL_COUNTER_RIGUNIT()
 	if (A.IsNearlyZero() || B.IsNearlyZero())
@@ -331,7 +331,7 @@ void FRigUnit_MathVectorParallel::Execute(const FRigUnitContext& Context)
 	Result = FVector::Parallel(A, B);
 }
 
-void FRigUnit_MathVectorOrthogonal::Execute(const FRigUnitContext& Context)
+UE_RigUnit_MathVectorOrthogonal_IMPLEMENT_STATIC_VIRTUAL_METHOD(void, Execute, const FRigUnitContext& Context)
 {
     DECLARE_SCOPE_HIERARCHICAL_COUNTER_RIGUNIT()
 	if (A.IsNearlyZero() || B.IsNearlyZero())
@@ -350,7 +350,7 @@ void FRigUnit_MathVectorOrthogonal::Execute(const FRigUnitContext& Context)
 	Result = FVector::Orthogonal(A, B);
 }
 
-void FRigUnit_MathVectorBezierFourPoint::Execute(const FRigUnitContext& Context)
+UE_RigUnit_MathVectorBezierFourPoint_IMPLEMENT_STATIC_VIRTUAL_METHOD(void, Execute, const FRigUnitContext& Context)
 {
     DECLARE_SCOPE_HIERARCHICAL_COUNTER_RIGUNIT()
 	FControlRigMathLibrary::FourPointBezier(Bezier, T, Result, Tangent);
diff --git a/Engine/Plugins/Experimental/ControlRig/Source/ControlRig/Private/Units/Math/RigUnit_MathVector.h b/Engine/Plugins/Experimental/ControlRig/Source/ControlRig/Private/Units/Math/RigUnit_MathVector.h
--- a/Engine/Plugins/Experimental/ControlRig/Source/ControlRig/Private/Units/Math/RigUnit_MathVector.h
+++ b/Engine/Plugins/Experimental/ControlRig/Source/ControlRig/Private/Units/Math/RigUnit_MathVector.h
@@ -56,6 +56,7 @@ USTRUCT(meta=(DisplayName="From Float", PrototypeName="FromFloat", Keywords="Mak
 struct FRigUnit_MathVectorFromFloat : public FRigUnit_MathVectorBase
 {
 	GENERATED_BODY()
+	STATIC_VIRTUAL_METHOD()
 	virtual void Execute(const FRigUnitContext& Context) override;
 
 	FRigUnit_MathVectorFromFloat()
@@ -79,6 +80,7 @@ struct FRigUnit_MathVectorAdd : public FRigUnit_MathVectorBinaryOp
 {
 	GENERATED_BODY()
 
+	STATIC_VIRTUAL_METHOD()
 	virtual void Execute(const FRigUnitContext& Context) override;
 };
 
@@ -90,6 +92,7 @@ struct FRigUnit_MathVectorSub : public FRigUnit_MathVectorBinaryOp
 {
 	GENERATED_BODY()
 
+	STATIC_VIRTUAL_METHOD()
 	virtual void Execute(const FRigUnitContext& Context) override;
 };
 
@@ -106,6 +109,7 @@ struct FRigUnit_MathVectorMul : public FRigUnit_MathVectorBinaryOp
 		A = B = FVector::OneVector;
 	}
 
+	STATIC_VIRTUAL_METHOD()
 	virtual void Execute(const FRigUnitContext& Context) override;
 };
 
@@ -123,6 +127,7 @@ struct FRigUnit_MathVectorScale : public FRigUnit_MathVectorBase
 		Factor = 1.f;
 	}
 
+	STATIC_VIRTUAL_METHOD()
 	virtual void Execute(const FRigUnitContext& Context) override;
 
 	UPROPERTY(meta = (Input))
@@ -148,6 +153,7 @@ struct FRigUnit_MathVectorDiv : public FRigUnit_MathVectorBinaryOp
 		B = FVector::OneVector;
 	}
 
+	STATIC_VIRTUAL_METHOD()
 	virtual void Execute(const FRigUnitContext& Context) override;
 };
 
@@ -165,6 +171,7 @@ struct FRigUnit_MathVectorMod : public FRigUnit_MathVectorBinaryOp
 		B = FVector::OneVector;
 	}
 
+	STATIC_VIRTUAL_METHOD()
 	virtual void Execute(const FRigUnitContext& Context) override;
 };
 
@@ -176,6 +183,7 @@ struct FRigUnit_MathVectorMin : public FRigUnit_MathVectorBinaryOp
 {
 	GENERATED_BODY()
 
+	STATIC_VIRTUAL_METHOD()
 	virtual void Execute(const FRigUnitContext& Context) override;
 };
 
@@ -187,6 +195,7 @@ struct FRigUnit_MathVectorMax : public FRigUnit_MathVectorBinaryOp
 {
 	GENERATED_BODY()
 
+	STATIC_VIRTUAL_METHOD()
 	virtual void Execute(const FRigUnitContext& Context) override;
 };
 
@@ -198,6 +207,7 @@ struct FRigUnit_MathVectorNegate : public FRigUnit_MathVectorUnaryOp
 {
 	GENERATED_BODY()
 
+	STATIC_VIRTUAL_METHOD()
 	virtual void Execute(const FRigUnitContext& Context) override;
 };
 
@@ -209,6 +219,7 @@ struct FRigUnit_MathVectorAbs : public FRigUnit_MathVectorUnaryOp
 {
 	GENERATED_BODY()
 
+	STATIC_VIRTUAL_METHOD()
 	virtual void Execute(const FRigUnitContext& Context) override;
 };
 
@@ -220,6 +231,7 @@ struct FRigUnit_MathVectorFloor : public FRigUnit_MathVectorUnaryOp
 {
 	GENERATED_BODY()
 
+	STATIC_VIRTUAL_METHOD()
 	virtual void Execute(const FRigUnitContext& Context) override;
 };
 
@@ -231,6 +243,7 @@ struct FRigUnit_MathVectorCeil : public FRigUnit_MathVectorUnaryOp
 {
 	GENERATED_BODY()
 
+	STATIC_VIRTUAL_METHOD()
 	virtual void Execute(const FRigUnitContext& Context) override;
 };
 
@@ -242,6 +255,7 @@ struct FRigUnit_MathVectorRound : public FRigUnit_MathVectorUnaryOp
 {
 	GENERATED_BODY()
 
+	STATIC_VIRTUAL_METHOD()
 	virtual void Execute(const FRigUnitContext& Context) override;
 };
 
@@ -253,6 +267,7 @@ struct FRigUnit_MathVectorSign : public FRigUnit_MathVectorUnaryOp
 {
 	GENERATED_BODY()
 
+	STATIC_VIRTUAL_METHOD()
 	virtual void Execute(const FRigUnitContext& Context) override;
 };
 
@@ -270,6 +285,7 @@ struct FRigUnit_MathVectorClamp : public FRigUnit_MathVectorBase
 		Maximum = FVector::OneVector;
 	}
 
+	STATIC_VIRTUAL_METHOD()
 	virtual void Execute(const FRigUnitContext& Context) override;
 
 	UPROPERTY(meta=(Input))
@@ -300,6 +316,7 @@ struct FRigUnit_MathVectorLerp : public FRigUnit_MathVectorBase
 		T = 0.f;
 	}
 
+	STATIC_VIRTUAL_METHOD()
 	virtual void Execute(const FRigUnitContext& Context) override;
 
 	UPROPERTY(meta=(Input))
@@ -330,6 +347,7 @@ struct FRigUnit_MathVectorRemap : public FRigUnit_MathVectorBase
 		bClamp = false;
 	}
 
+	STATIC_VIRTUAL_METHOD()
 	virtual void Execute(const FRigUnitContext& Context) override;
 
 	UPROPERTY(meta=(Input))
@@ -369,6 +387,7 @@ struct FRigUnit_MathVectorEquals : public FRigUnit_MathVectorBase
 		Result = true;
 	}
 
+	STATIC_VIRTUAL_METHOD()
 	virtual void Execute(const FRigUnitContext& Context) override;
 
 	UPROPERTY(meta=(Input))
@@ -395,6 +414,7 @@ struct FRigUnit_MathVectorNotEquals : public FRigUnit_MathVectorBase
 		Result = false;
 	}
 
+	STATIC_VIRTUAL_METHOD()
 	virtual void Execute(const FRigUnitContext& Context) override;
 
 	UPROPERTY(meta=(Input))
@@ -422,6 +442,7 @@ struct FRigUnit_MathVectorIsNearlyZero : public FRigUnit_MathVectorBase
 		Result = true;
 	}
 
+	STATIC_VIRTUAL_METHOD()
 	virtual void Execute(const FRigUnitContext& Context) override;
 
 	UPROPERTY(meta=(Input))
@@ -449,6 +470,7 @@ struct FRigUnit_MathVectorIsNearlyEqual : public FRigUnit_MathVectorBase
 		Result = true;
 	}
 
+	STATIC_VIRTUAL_METHOD()
 	virtual void Execute(const FRigUnitContext& Context) override;
 
 	UPROPERTY(meta=(Input))
@@ -478,6 +500,7 @@ struct FRigUnit_MathVectorSelectBool : public FRigUnit_MathVectorBase
 		IfTrue = IfFalse = Result = FVector::ZeroVector;
 	}
 
+	STATIC_VIRTUAL_METHOD()
 	virtual void Execute(const FRigUnitContext& Context) override;
 
 	UPROPERTY(meta=(Input))
@@ -501,6 +524,7 @@ struct FRigUnit_MathVectorDeg : public FRigUnit_MathVectorUnaryOp
 {
 	GENERATED_BODY()
 
+	STATIC_VIRTUAL_METHOD()
 	virtual void Execute(const FRigUnitContext& Context) override;
 };
 
@@ -512,6 +536,7 @@ struct FRigUnit_MathVectorRad : public FRigUnit_MathVectorUnaryOp
 {
 	GENERATED_BODY()
 
+	STATIC_VIRTUAL_METHOD()
 	virtual void Execute(const FRigUnitContext& Context) override;
 };
 
@@ -529,6 +554,7 @@ struct FRigUnit_MathVectorLengthSquared : public FRigUnit_MathVectorBase
 		Result = 0.f;
 	}
 
+	STATIC_VIRTUAL_METHOD()
 	virtual void Execute(const FRigUnitContext& Context) override;
 
 	UPROPERTY(meta=(Input))
@@ -552,6 +578,7 @@ struct FRigUnit_MathVectorLength : public FRigUnit_MathVectorBase
 		Result = 0.f;
 	}
 
+	STATIC_VIRTUAL_METHOD()
 	virtual void Execute(const FRigUnitContext& Context) override;
 
 	UPROPERTY(meta=(Input))
@@ -575,6 +602,7 @@ struct FRigUnit_MathVectorDistance : public FRigUnit_MathVectorBase
 		Result = 0.f;
 	}
 
+	STATIC_VIRTUAL_METHOD()
 	virtual void Execute(const FRigUnitContext& Context) override;
 
 	UPROPERTY(meta=(Input))
@@ -594,6 +622,7 @@ USTRUCT(meta=(DisplayName="Cross", PrototypeName="Cross", Keywords="^"))
 struct FRigUnit_MathVectorCross : public FRigUnit_MathVectorBinaryOp
 {
 	GENERATED_BODY()
+	STATIC_VIRTUAL_METHOD()
 	virtual void Execute(const FRigUnitContext& Context) override;
 };
 
@@ -611,6 +640,7 @@ struct FRigUnit_MathVectorDot : public FRigUnit_MathVectorBase
 		Result = 0.f;
 	}
 
+	STATIC_VIRTUAL_METHOD()
 	virtual void Execute(const FRigUnitContext& Context) override;
 
 	UPROPERTY(meta=(Input))
@@ -631,6 +661,7 @@ struct FRigUnit_MathVectorUnit : public FRigUnit_MathVectorUnaryOp
 {
 	GENERATED_BODY()
 
+	STATIC_VIRTUAL_METHOD()
 	virtual void Execute(const FRigUnitContext& Context) override;
 };
 
@@ -648,6 +679,7 @@ struct FRigUnit_MathVectorMirror : public FRigUnit_MathVectorBase
 		Normal = FVector(1.f, 0.f, 0.f);
 	}
 
+	STATIC_VIRTUAL_METHOD()
 	virtual void Execute(const FRigUnitContext& Context) override;
 
 	UPROPERTY(meta=(Input))
@@ -674,6 +706,7 @@ struct FRigUnit_MathVectorAngle : public FRigUnit_MathVectorBase
 		Result = 0.f;
 	}
 
+	STATIC_VIRTUAL_METHOD()
 	virtual void Execute(const FRigUnitContext& Context) override;
 
 	UPROPERTY(meta=(Input))
@@ -700,6 +733,7 @@ struct FRigUnit_MathVectorParallel : public FRigUnit_MathVectorBase
 		Result = true;
 	}
 
+	STATIC_VIRTUAL_METHOD()
 	virtual void Execute(const FRigUnitContext& Context) override;
 
 	UPROPERTY(meta=(Input))
@@ -727,6 +761,7 @@ struct FRigUnit_MathVectorOrthogonal : public FRigUnit_MathVectorBase
 		Result = true;
 	}
 
+	STATIC_VIRTUAL_METHOD()
 	virtual void Execute(const FRigUnitContext& Context) override;
 
 	UPROPERTY(meta=(Input))
@@ -754,6 +789,7 @@ struct FRigUnit_MathVectorBezierFourPoint : public FRigUnit_MathVectorBase
 		Result = Tangent = FVector::ZeroVector;
 	}
 
+	STATIC_VIRTUAL_METHOD()
 	virtual void Execute(const FRigUnitContext& Context) override;
 
 	UPROPERTY(meta=(Input))
diff --git a/Engine/Plugins/Experimental/ControlRig/Source/ControlRig/Private/Units/Math/RigUnit_Noise.cpp b/Engine/Plugins/Experimental/ControlRig/Source/ControlRig/Private/Units/Math/RigUnit_Noise.cpp
--- a/Engine/Plugins/Experimental/ControlRig/Source/ControlRig/Private/Units/Math/RigUnit_Noise.cpp
+++ b/Engine/Plugins/Experimental/ControlRig/Source/ControlRig/Private/Units/Math/RigUnit_Noise.cpp
@@ -3,7 +3,7 @@
 #include "Units/Math/RigUnit_Noise.h"
 #include "Units/RigUnitContext.h"
 
-void FRigUnit_NoiseFloat::Execute(const FRigUnitContext& Context)
+UE_RigUnit_NoiseFloat_IMPLEMENT_STATIC_VIRTUAL_METHOD(void, Execute, const FRigUnitContext& Context)
 {
     DECLARE_SCOPE_HIERARCHICAL_COUNTER_RIGUNIT()
 	if (Context.State == EControlRigState::Init)
@@ -17,7 +17,7 @@ void FRigUnit_NoiseFloat::Execute(const FRigUnitContext& Context)
 	Time = Time + Speed * Context.DeltaTime;
 }
 
-void FRigUnit_NoiseVector::Execute(const FRigUnitContext& Context)
+UE_RigUnit_NoiseVector_IMPLEMENT_STATIC_VIRTUAL_METHOD(void, Execute, const FRigUnitContext& Context)
 {
     DECLARE_SCOPE_HIERARCHICAL_COUNTER_RIGUNIT()
 	if (Context.State == EControlRigState::Init)
diff --git a/Engine/Plugins/Experimental/ControlRig/Source/ControlRig/Private/Units/Math/RigUnit_Noise.h b/Engine/Plugins/Experimental/ControlRig/Source/ControlRig/Private/Units/Math/RigUnit_Noise.h
--- a/Engine/Plugins/Experimental/ControlRig/Source/ControlRig/Private/Units/Math/RigUnit_Noise.h
+++ b/Engine/Plugins/Experimental/ControlRig/Source/ControlRig/Private/Units/Math/RigUnit_Noise.h
@@ -17,6 +17,7 @@ struct FRigUnit_NoiseFloat : public FRigUnit_MathBase
 		Frequency = Maximum = 1.f;
 	}
 
+	STATIC_VIRTUAL_METHOD()
 	virtual void Execute(const FRigUnitContext& Context) override;
 
 	UPROPERTY(meta = (Input))
@@ -55,6 +56,7 @@ struct FRigUnit_NoiseVector : public FRigUnit_MathBase
 		Maximum = 1.f;
 	}
 
+	STATIC_VIRTUAL_METHOD()
 	virtual void Execute(const FRigUnitContext& Context) override;
 
 	UPROPERTY(meta = (Input))
diff --git a/Engine/Plugins/Experimental/ControlRig/Source/ControlRig/Private/Units/Math/RigUnit_Random.cpp b/Engine/Plugins/Experimental/ControlRig/Source/ControlRig/Private/Units/Math/RigUnit_Random.cpp
--- a/Engine/Plugins/Experimental/ControlRig/Source/ControlRig/Private/Units/Math/RigUnit_Random.cpp
+++ b/Engine/Plugins/Experimental/ControlRig/Source/ControlRig/Private/Units/Math/RigUnit_Random.cpp
@@ -15,7 +15,7 @@ float FRigUnit_Random_Helper(int32& Seed)
 	return FPlatformMath::Fractional(Result.f);
 }
 
-void FRigUnit_RandomFloat::Execute(const FRigUnitContext& Context)
+UE_RigUnit_RandomFloat_IMPLEMENT_STATIC_VIRTUAL_METHOD(void, Execute, const FRigUnitContext& Context)
 {
     DECLARE_SCOPE_HIERARCHICAL_COUNTER_RIGUNIT()
 	if (Context.State == EControlRigState::Init)
@@ -38,7 +38,7 @@ void FRigUnit_RandomFloat::Execute(const FRigUnitContext& Context)
 	LastResult = Result;
 }
 
-void FRigUnit_RandomVector::Execute(const FRigUnitContext& Context)
+UE_RigUnit_RandomVector_IMPLEMENT_STATIC_VIRTUAL_METHOD(void, Execute, const FRigUnitContext& Context)
 {
     DECLARE_SCOPE_HIERARCHICAL_COUNTER_RIGUNIT()
 	if (Context.State == EControlRigState::Init)
diff --git a/Engine/Plugins/Experimental/ControlRig/Source/ControlRig/Private/Units/Math/RigUnit_Random.h b/Engine/Plugins/Experimental/ControlRig/Source/ControlRig/Private/Units/Math/RigUnit_Random.h
--- a/Engine/Plugins/Experimental/ControlRig/Source/ControlRig/Private/Units/Math/RigUnit_Random.h
+++ b/Engine/Plugins/Experimental/ControlRig/Source/ControlRig/Private/Units/Math/RigUnit_Random.h
@@ -19,6 +19,7 @@ struct FRigUnit_RandomFloat : public FRigUnit_MathBase
 		TimeLeft = 0.f;
 	}
 
+	STATIC_VIRTUAL_METHOD()
 	virtual void Execute(const FRigUnitContext& Context) override;
 
 	UPROPERTY(meta = (Input, Constant))
@@ -62,6 +63,7 @@ struct FRigUnit_RandomVector: public FRigUnit_MathBase
 		Result = LastResult = FVector::ZeroVector;
 	}
 
+	STATIC_VIRTUAL_METHOD()
 	virtual void Execute(const FRigUnitContext& Context) override;
 
 	UPROPERTY(meta = (Input))
diff --git a/Engine/Plugins/Experimental/ControlRig/Source/ControlRig/Private/Units/Simulation/RigUnit_Accumulate.cpp b/Engine/Plugins/Experimental/ControlRig/Source/ControlRig/Private/Units/Simulation/RigUnit_Accumulate.cpp
--- a/Engine/Plugins/Experimental/ControlRig/Source/ControlRig/Private/Units/Simulation/RigUnit_Accumulate.cpp
+++ b/Engine/Plugins/Experimental/ControlRig/Source/ControlRig/Private/Units/Simulation/RigUnit_Accumulate.cpp
@@ -3,7 +3,7 @@
 #include "Units/Simulation/RigUnit_Accumulate.h"
 #include "Units/RigUnitContext.h"
 
-void FRigUnit_AccumulateFloatAdd::Execute(const FRigUnitContext& Context)
+UE_RigUnit_AccumulateFloatAdd_IMPLEMENT_STATIC_VIRTUAL_METHOD(void, Execute, const FRigUnitContext& Context)
 {
     DECLARE_SCOPE_HIERARCHICAL_COUNTER_RIGUNIT()
 	if (Context.State == EControlRigState::Init)
@@ -20,7 +20,7 @@ void FRigUnit_AccumulateFloatAdd::Execute(const FRigUnitContext& Context)
 	}
 }
 
-void FRigUnit_AccumulateVectorAdd::Execute(const FRigUnitContext& Context)
+UE_RigUnit_AccumulateVectorAdd_IMPLEMENT_STATIC_VIRTUAL_METHOD(void, Execute, const FRigUnitContext& Context)
 {
     DECLARE_SCOPE_HIERARCHICAL_COUNTER_RIGUNIT()
 	if (Context.State == EControlRigState::Init)
@@ -37,7 +37,7 @@ void FRigUnit_AccumulateVectorAdd::Execute(const FRigUnitContext& Context)
 	}
 }
 
-void FRigUnit_AccumulateFloatMul::Execute(const FRigUnitContext& Context)
+UE_RigUnit_AccumulateFloatMul_IMPLEMENT_STATIC_VIRTUAL_METHOD(void, Execute, const FRigUnitContext& Context)
 {
     DECLARE_SCOPE_HIERARCHICAL_COUNTER_RIGUNIT()
 	if (Context.State == EControlRigState::Init)
@@ -55,7 +55,7 @@ void FRigUnit_AccumulateFloatMul::Execute(const FRigUnitContext& Context)
 	}
 }
 
-void FRigUnit_AccumulateVectorMul::Execute(const FRigUnitContext& Context)
+UE_RigUnit_AccumulateVectorMul_IMPLEMENT_STATIC_VIRTUAL_METHOD(void, Execute, const FRigUnitContext& Context)
 {
     DECLARE_SCOPE_HIERARCHICAL_COUNTER_RIGUNIT()
 	if (Context.State == EControlRigState::Init)
@@ -73,7 +73,7 @@ void FRigUnit_AccumulateVectorMul::Execute(const FRigUnitContext& Context)
 	}
 }
 
-void FRigUnit_AccumulateQuatMul::Execute(const FRigUnitContext& Context)
+UE_RigUnit_AccumulateQuatMul_IMPLEMENT_STATIC_VIRTUAL_METHOD(void, Execute, const FRigUnitContext& Context)
 {
     DECLARE_SCOPE_HIERARCHICAL_COUNTER_RIGUNIT()
 	if (Context.State == EControlRigState::Init)
@@ -98,7 +98,7 @@ void FRigUnit_AccumulateQuatMul::Execute(const FRigUnitContext& Context)
 	}
 }
 
-void FRigUnit_AccumulateTransformMul::Execute(const FRigUnitContext& Context)
+UE_RigUnit_AccumulateTransformMul_IMPLEMENT_STATIC_VIRTUAL_METHOD(void, Execute, const FRigUnitContext& Context)
 {
     DECLARE_SCOPE_HIERARCHICAL_COUNTER_RIGUNIT()
 	if (Context.State == EControlRigState::Init)
@@ -125,7 +125,7 @@ void FRigUnit_AccumulateTransformMul::Execute(const FRigUnitContext& Context)
 	}
 }
 
-void FRigUnit_AccumulateFloatLerp::Execute(const FRigUnitContext& Context)
+UE_RigUnit_AccumulateFloatLerp_IMPLEMENT_STATIC_VIRTUAL_METHOD(void, Execute, const FRigUnitContext& Context)
 {
     DECLARE_SCOPE_HIERARCHICAL_COUNTER_RIGUNIT()
 	if (Context.State == EControlRigState::Init)
@@ -138,7 +138,7 @@ void FRigUnit_AccumulateFloatLerp::Execute(const FRigUnitContext& Context)
 	}
 }
 
-void FRigUnit_AccumulateVectorLerp::Execute(const FRigUnitContext& Context)
+UE_RigUnit_AccumulateVectorLerp_IMPLEMENT_STATIC_VIRTUAL_METHOD(void, Execute, const FRigUnitContext& Context)
 {
     DECLARE_SCOPE_HIERARCHICAL_COUNTER_RIGUNIT()
 	if (Context.State == EControlRigState::Init)
@@ -151,7 +151,7 @@ void FRigUnit_AccumulateVectorLerp::Execute(const FRigUnitContext& Context)
 	}
 }
 
-void FRigUnit_AccumulateQuatLerp::Execute(const FRigUnitContext& Context)
+UE_RigUnit_AccumulateQuatLerp_IMPLEMENT_STATIC_VIRTUAL_METHOD(void, Execute, const FRigUnitContext& Context)
 {
     DECLARE_SCOPE_HIERARCHICAL_COUNTER_RIGUNIT()
 	if (Context.State == EControlRigState::Init)
@@ -164,7 +164,7 @@ void FRigUnit_AccumulateQuatLerp::Execute(const FRigUnitContext& Context)
 	}
 }
 
-void FRigUnit_AccumulateTransformLerp::Execute(const FRigUnitContext& Context)
+UE_RigUnit_AccumulateTransformLerp_IMPLEMENT_STATIC_VIRTUAL_METHOD(void, Execute, const FRigUnitContext& Context)
 {
     DECLARE_SCOPE_HIERARCHICAL_COUNTER_RIGUNIT()
 	if (Context.State == EControlRigState::Init)
@@ -181,7 +181,7 @@ void FRigUnit_AccumulateTransformLerp::Execute(const FRigUnitContext& Context)
 	}
 }
 
-void FRigUnit_AccumulateFloatRange::Execute(const FRigUnitContext& Context)
+UE_RigUnit_AccumulateFloatRange_IMPLEMENT_STATIC_VIRTUAL_METHOD(void, Execute, const FRigUnitContext& Context)
 {
     DECLARE_SCOPE_HIERARCHICAL_COUNTER_RIGUNIT()
 	if (Context.State == EControlRigState::Init)
@@ -195,7 +195,7 @@ void FRigUnit_AccumulateFloatRange::Execute(const FRigUnitContext& Context)
 	}
 }
 
-void FRigUnit_AccumulateVectorRange::Execute(const FRigUnitContext& Context)
+UE_RigUnit_AccumulateVectorRange_IMPLEMENT_STATIC_VIRTUAL_METHOD(void, Execute, const FRigUnitContext& Context)
 {
     DECLARE_SCOPE_HIERARCHICAL_COUNTER_RIGUNIT()
 	if (Context.State == EControlRigState::Init)
diff --git a/Engine/Plugins/Experimental/ControlRig/Source/ControlRig/Private/Units/Simulation/RigUnit_Accumulate.h b/Engine/Plugins/Experimental/ControlRig/Source/ControlRig/Private/Units/Simulation/RigUnit_Accumulate.h
--- a/Engine/Plugins/Experimental/ControlRig/Source/ControlRig/Private/Units/Simulation/RigUnit_Accumulate.h
+++ b/Engine/Plugins/Experimental/ControlRig/Source/ControlRig/Private/Units/Simulation/RigUnit_Accumulate.h
@@ -19,6 +19,7 @@ struct FRigUnit_AccumulateFloatAdd : public FRigUnit_SimBase
 		bIntegrateDeltaTime = false;
 	}
 
+	STATIC_VIRTUAL_METHOD()
 	virtual void Execute(const FRigUnitContext& Context) override;
 
 	UPROPERTY(meta=(Input))
@@ -51,6 +52,7 @@ struct FRigUnit_AccumulateVectorAdd : public FRigUnit_SimBase
 		bIntegrateDeltaTime = false;
 	}
 
+	STATIC_VIRTUAL_METHOD()
 	virtual void Execute(const FRigUnitContext& Context) override;
 
 	UPROPERTY(meta = (Input))
@@ -83,6 +85,7 @@ struct FRigUnit_AccumulateFloatMul : public FRigUnit_SimBase
 		bIntegrateDeltaTime = false;
 	}
 
+	STATIC_VIRTUAL_METHOD()
 	virtual void Execute(const FRigUnitContext& Context) override;
 
 	UPROPERTY(meta=(Input))
@@ -115,6 +118,7 @@ struct FRigUnit_AccumulateVectorMul : public FRigUnit_SimBase
 		bIntegrateDeltaTime = false;
 	}
 
+	STATIC_VIRTUAL_METHOD()
 	virtual void Execute(const FRigUnitContext& Context) override;
 
 	UPROPERTY(meta = (Input))
@@ -147,6 +151,7 @@ struct FRigUnit_AccumulateQuatMul : public FRigUnit_SimBase
 		bIntegrateDeltaTime = bFlipOrder = false;
 	}
 
+	STATIC_VIRTUAL_METHOD()
 	virtual void Execute(const FRigUnitContext& Context) override;
 
 	UPROPERTY(meta=(Input))
@@ -182,6 +187,7 @@ struct FRigUnit_AccumulateTransformMul : public FRigUnit_SimBase
 		bIntegrateDeltaTime = bFlipOrder = false;
 	}
 
+	STATIC_VIRTUAL_METHOD()
 	virtual void Execute(const FRigUnitContext& Context) override;
 
 	UPROPERTY(meta = (Input))
@@ -217,6 +223,7 @@ struct FRigUnit_AccumulateFloatLerp : public FRigUnit_SimBase
 		bIntegrateDeltaTime = false;
 	}
 
+	STATIC_VIRTUAL_METHOD()
 	virtual void Execute(const FRigUnitContext& Context) override;
 
 	UPROPERTY(meta=(Input))
@@ -253,6 +260,7 @@ struct FRigUnit_AccumulateVectorLerp : public FRigUnit_SimBase
 		bIntegrateDeltaTime = false;
 	}
 
+	STATIC_VIRTUAL_METHOD()
 	virtual void Execute(const FRigUnitContext& Context) override;
 
 	UPROPERTY(meta = (Input))
@@ -289,6 +297,7 @@ struct FRigUnit_AccumulateQuatLerp : public FRigUnit_SimBase
 		bIntegrateDeltaTime = false;
 	}
 
+	STATIC_VIRTUAL_METHOD()
 	virtual void Execute(const FRigUnitContext& Context) override;
 
 	UPROPERTY(meta=(Input))
@@ -325,6 +334,7 @@ struct FRigUnit_AccumulateTransformLerp : public FRigUnit_SimBase
 		bIntegrateDeltaTime = false;
 	}
 
+	STATIC_VIRTUAL_METHOD()
 	virtual void Execute(const FRigUnitContext& Context) override;
 
 	UPROPERTY(meta = (Input))
@@ -359,6 +369,7 @@ struct FRigUnit_AccumulateFloatRange : public FRigUnit_SimBase
 		Value = Minimum = Maximum = AccumulatedMinimum = AccumulatedMaximum = 0.f;
 	}
 
+	STATIC_VIRTUAL_METHOD()
 	virtual void Execute(const FRigUnitContext& Context) override;
 
 	UPROPERTY(meta = (Input))
@@ -390,6 +401,7 @@ struct FRigUnit_AccumulateVectorRange : public FRigUnit_SimBase
 		Value = Minimum = Maximum = AccumulatedMinimum = AccumulatedMaximum = FVector::ZeroVector;
 	}
 
+	STATIC_VIRTUAL_METHOD()
 	virtual void Execute(const FRigUnitContext& Context) override;
 
 	UPROPERTY(meta = (Input))
diff --git a/Engine/Plugins/Experimental/ControlRig/Source/ControlRig/Private/Units/Simulation/RigUnit_DeltaFromPrevious.cpp b/Engine/Plugins/Experimental/ControlRig/Source/ControlRig/Private/Units/Simulation/RigUnit_DeltaFromPrevious.cpp
--- a/Engine/Plugins/Experimental/ControlRig/Source/ControlRig/Private/Units/Simulation/RigUnit_DeltaFromPrevious.cpp
+++ b/Engine/Plugins/Experimental/ControlRig/Source/ControlRig/Private/Units/Simulation/RigUnit_DeltaFromPrevious.cpp
@@ -3,7 +3,7 @@
 #include "Units/Simulation/RigUnit_DeltaFromPrevious.h"
 #include "Units/RigUnitContext.h"
 
-void FRigUnit_DeltaFromPreviousFloat::Execute(const FRigUnitContext& Context)
+UE_RigUnit_DeltaFromPreviousFloat_IMPLEMENT_STATIC_VIRTUAL_METHOD(void, Execute, const FRigUnitContext& Context)
 {
     DECLARE_SCOPE_HIERARCHICAL_COUNTER_RIGUNIT()
 	if (Context.State == EControlRigState::Init)
@@ -18,7 +18,7 @@ void FRigUnit_DeltaFromPreviousFloat::Execute(const FRigUnitContext& Context)
 	Cache = Value;
 }
 
-void FRigUnit_DeltaFromPreviousVector::Execute(const FRigUnitContext& Context)
+UE_RigUnit_DeltaFromPreviousVector_IMPLEMENT_STATIC_VIRTUAL_METHOD(void, Execute, const FRigUnitContext& Context)
 {
     DECLARE_SCOPE_HIERARCHICAL_COUNTER_RIGUNIT()
 	if (Context.State == EControlRigState::Init)
@@ -33,7 +33,7 @@ void FRigUnit_DeltaFromPreviousVector::Execute(const FRigUnitContext& Context)
 	Cache = Value;
 }
 
-void FRigUnit_DeltaFromPreviousQuat::Execute(const FRigUnitContext& Context)
+UE_RigUnit_DeltaFromPreviousQuat_IMPLEMENT_STATIC_VIRTUAL_METHOD(void, Execute, const FRigUnitContext& Context)
 {
     DECLARE_SCOPE_HIERARCHICAL_COUNTER_RIGUNIT()
 	if (Context.State == EControlRigState::Init)
@@ -48,7 +48,7 @@ void FRigUnit_DeltaFromPreviousQuat::Execute(const FRigUnitContext& Context)
 	Cache = Value;
 }
 
-void FRigUnit_DeltaFromPreviousTransform::Execute(const FRigUnitContext& Context)
+UE_RigUnit_DeltaFromPreviousTransform_IMPLEMENT_STATIC_VIRTUAL_METHOD(void, Execute, const FRigUnitContext& Context)
 {
     DECLARE_SCOPE_HIERARCHICAL_COUNTER_RIGUNIT()
 	if (Context.State == EControlRigState::Init)
diff --git a/Engine/Plugins/Experimental/ControlRig/Source/ControlRig/Private/Units/Simulation/RigUnit_DeltaFromPrevious.h b/Engine/Plugins/Experimental/ControlRig/Source/ControlRig/Private/Units/Simulation/RigUnit_DeltaFromPrevious.h
--- a/Engine/Plugins/Experimental/ControlRig/Source/ControlRig/Private/Units/Simulation/RigUnit_DeltaFromPrevious.h
+++ b/Engine/Plugins/Experimental/ControlRig/Source/ControlRig/Private/Units/Simulation/RigUnit_DeltaFromPrevious.h
@@ -18,6 +18,7 @@ struct FRigUnit_DeltaFromPreviousFloat : public FRigUnit_SimBase
 		Value = Delta = PreviousValue = Cache = 0.f;
 	}
 
+	STATIC_VIRTUAL_METHOD()
 	virtual void Execute(const FRigUnitContext& Context) override;
 
 	UPROPERTY(meta = (Input))
@@ -46,6 +47,7 @@ struct FRigUnit_DeltaFromPreviousVector : public FRigUnit_SimBase
 		Value = Delta = PreviousValue = Cache = FVector::ZeroVector;
 	}
 
+	STATIC_VIRTUAL_METHOD()
 	virtual void Execute(const FRigUnitContext& Context) override;
 
 	UPROPERTY(meta = (Input))
@@ -74,6 +76,7 @@ struct FRigUnit_DeltaFromPreviousQuat : public FRigUnit_SimBase
 		Value = Delta = PreviousValue = Cache = FQuat::Identity;
 	}
 
+	STATIC_VIRTUAL_METHOD()
 	virtual void Execute(const FRigUnitContext& Context) override;
 
 	UPROPERTY(meta = (Input))
@@ -102,6 +105,7 @@ struct FRigUnit_DeltaFromPreviousTransform : public FRigUnit_SimBase
 		Value = Delta = PreviousValue = Cache = FTransform::Identity;
 	}
 
+	STATIC_VIRTUAL_METHOD()
 	virtual void Execute(const FRigUnitContext& Context) override;
 
 	UPROPERTY(meta = (Input))
diff --git a/Engine/Plugins/Experimental/ControlRig/Source/ControlRig/Private/Units/Simulation/RigUnit_Kalman.cpp b/Engine/Plugins/Experimental/ControlRig/Source/ControlRig/Private/Units/Simulation/RigUnit_Kalman.cpp
--- a/Engine/Plugins/Experimental/ControlRig/Source/ControlRig/Private/Units/Simulation/RigUnit_Kalman.cpp
+++ b/Engine/Plugins/Experimental/ControlRig/Source/ControlRig/Private/Units/Simulation/RigUnit_Kalman.cpp
@@ -3,7 +3,7 @@
 #include "Units/Simulation/RigUnit_Kalman.h"
 #include "Units/RigUnitContext.h"
 
-void FRigUnit_KalmanFloat::Execute(const FRigUnitContext& Context)
+UE_RigUnit_KalmanFloat_IMPLEMENT_STATIC_VIRTUAL_METHOD(void, Execute, const FRigUnitContext& Context)
 {
     DECLARE_SCOPE_HIERARCHICAL_COUNTER_RIGUNIT()
 	if(BufferSize <= 0)
@@ -49,7 +49,7 @@ void FRigUnit_KalmanFloat::Execute(const FRigUnitContext& Context)
 	}
 }
 
-void FRigUnit_KalmanVector::Execute(const FRigUnitContext& Context)
+UE_RigUnit_KalmanVector_IMPLEMENT_STATIC_VIRTUAL_METHOD(void, Execute, const FRigUnitContext& Context)
 {
     DECLARE_SCOPE_HIERARCHICAL_COUNTER_RIGUNIT()
 	if(BufferSize <= 0)
@@ -95,7 +95,7 @@ void FRigUnit_KalmanVector::Execute(const FRigUnitContext& Context)
 	}
 }
 
-void FRigUnit_KalmanTransform::Execute(const FRigUnitContext& Context)
+UE_RigUnit_KalmanTransform_IMPLEMENT_STATIC_VIRTUAL_METHOD(void, Execute, const FRigUnitContext& Context)
 {
     DECLARE_SCOPE_HIERARCHICAL_COUNTER_RIGUNIT()
 	if(BufferSize <= 0)
diff --git a/Engine/Plugins/Experimental/ControlRig/Source/ControlRig/Private/Units/Simulation/RigUnit_Kalman.h b/Engine/Plugins/Experimental/ControlRig/Source/ControlRig/Private/Units/Simulation/RigUnit_Kalman.h
--- a/Engine/Plugins/Experimental/ControlRig/Source/ControlRig/Private/Units/Simulation/RigUnit_Kalman.h
+++ b/Engine/Plugins/Experimental/ControlRig/Source/ControlRig/Private/Units/Simulation/RigUnit_Kalman.h
@@ -21,6 +21,7 @@ struct FRigUnit_KalmanFloat : public FRigUnit_SimBase
 		LastInsertIndex = 0;
 	}
 
+	STATIC_VIRTUAL_METHOD()
 	virtual void Execute(const FRigUnitContext& Context) override;
 
 	UPROPERTY(meta = (Input))
@@ -55,6 +56,7 @@ struct FRigUnit_KalmanVector : public FRigUnit_SimBase
 		LastInsertIndex = 0;
 	}
 
+	STATIC_VIRTUAL_METHOD()
 	virtual void Execute(const FRigUnitContext& Context) override;
 
 	UPROPERTY(meta = (Input))
@@ -89,6 +91,7 @@ struct FRigUnit_KalmanTransform : public FRigUnit_SimBase
 		LastInsertIndex = 0;
 	}
 
+	STATIC_VIRTUAL_METHOD()
 	virtual void Execute(const FRigUnitContext& Context) override;
 
 	UPROPERTY(meta = (Input))
diff --git a/Engine/Plugins/Experimental/ControlRig/Source/ControlRig/Private/Units/Simulation/RigUnit_PointSimulation.cpp b/Engine/Plugins/Experimental/ControlRig/Source/ControlRig/Private/Units/Simulation/RigUnit_PointSimulation.cpp
--- a/Engine/Plugins/Experimental/ControlRig/Source/ControlRig/Private/Units/Simulation/RigUnit_PointSimulation.cpp
+++ b/Engine/Plugins/Experimental/ControlRig/Source/ControlRig/Private/Units/Simulation/RigUnit_PointSimulation.cpp
@@ -3,7 +3,7 @@
 #include "Units/Simulation/RigUnit_PointSimulation.h"
 #include "Units/RigUnitContext.h"
 
-void FRigUnit_PointSimulation::Execute(const FRigUnitContext& Context)
+UE_RigUnit_PointSimulation_IMPLEMENT_STATIC_VIRTUAL_METHOD(void, Execute, const FRigUnitContext& Context)
 {
     DECLARE_SCOPE_HIERARCHICAL_COUNTER_RIGUNIT()
 	FRigBoneHierarchy* Hierarchy = ExecuteContext.GetBones();
diff --git a/Engine/Plugins/Experimental/ControlRig/Source/ControlRig/Private/Units/Simulation/RigUnit_PointSimulation.h b/Engine/Plugins/Experimental/ControlRig/Source/ControlRig/Private/Units/Simulation/RigUnit_PointSimulation.h
--- a/Engine/Plugins/Experimental/ControlRig/Source/ControlRig/Private/Units/Simulation/RigUnit_PointSimulation.h
+++ b/Engine/Plugins/Experimental/ControlRig/Source/ControlRig/Private/Units/Simulation/RigUnit_PointSimulation.h
@@ -123,6 +123,7 @@ struct FRigUnit_PointSimulation : public FRigUnit_SimBaseMutable
 		Bezier = FCRFourPointBezier();
 	}
 
+	STATIC_VIRTUAL_METHOD()
 	virtual void Execute(const FRigUnitContext& Context) override;
 
 	/** The points to simulate */
diff --git a/Engine/Plugins/Experimental/ControlRig/Source/ControlRig/Private/Units/Simulation/RigUnit_TimeOffset.cpp b/Engine/Plugins/Experimental/ControlRig/Source/ControlRig/Private/Units/Simulation/RigUnit_TimeOffset.cpp
--- a/Engine/Plugins/Experimental/ControlRig/Source/ControlRig/Private/Units/Simulation/RigUnit_TimeOffset.cpp
+++ b/Engine/Plugins/Experimental/ControlRig/Source/ControlRig/Private/Units/Simulation/RigUnit_TimeOffset.cpp
@@ -3,7 +3,7 @@
 #include "Units/Simulation/RigUnit_TimeOffset.h"
 #include "Units/RigUnitContext.h"
 
-void FRigUnit_TimeOffsetFloat::Execute(const FRigUnitContext& Context)
+UE_RigUnit_TimeOffsetFloat_IMPLEMENT_STATIC_VIRTUAL_METHOD(void, Execute, const FRigUnitContext& Context)
 {
     DECLARE_SCOPE_HIERARCHICAL_COUNTER_RIGUNIT()
 	Result = Value;
@@ -111,7 +111,7 @@ void FRigUnit_TimeOffsetFloat::Execute(const FRigUnitContext& Context)
 	}
 }
 
-void FRigUnit_TimeOffsetVector::Execute(const FRigUnitContext& Context)
+UE_RigUnit_TimeOffsetVector_IMPLEMENT_STATIC_VIRTUAL_METHOD(void, Execute, const FRigUnitContext& Context)
 {
     DECLARE_SCOPE_HIERARCHICAL_COUNTER_RIGUNIT()
 	Result = Value;
@@ -219,7 +219,7 @@ void FRigUnit_TimeOffsetVector::Execute(const FRigUnitContext& Context)
 	}
 }
 
-void FRigUnit_TimeOffsetTransform::Execute(const FRigUnitContext& Context)
+UE_RigUnit_TimeOffsetTransform_IMPLEMENT_STATIC_VIRTUAL_METHOD(void, Execute, const FRigUnitContext& Context)
 {
     DECLARE_SCOPE_HIERARCHICAL_COUNTER_RIGUNIT()
 	Result = Value;
diff --git a/Engine/Plugins/Experimental/ControlRig/Source/ControlRig/Private/Units/Simulation/RigUnit_TimeOffset.h b/Engine/Plugins/Experimental/ControlRig/Source/ControlRig/Private/Units/Simulation/RigUnit_TimeOffset.h
--- a/Engine/Plugins/Experimental/ControlRig/Source/ControlRig/Private/Units/Simulation/RigUnit_TimeOffset.h
+++ b/Engine/Plugins/Experimental/ControlRig/Source/ControlRig/Private/Units/Simulation/RigUnit_TimeOffset.h
@@ -22,6 +22,7 @@ struct FRigUnit_TimeOffsetFloat : public FRigUnit_SimBase
 		LastInsertIndex = 0;
 	}
 
+	STATIC_VIRTUAL_METHOD()
 	virtual void Execute(const FRigUnitContext& Context) override;
 	
 	/** The value to record */
@@ -70,6 +71,7 @@ struct FRigUnit_TimeOffsetVector : public FRigUnit_SimBase
 		LastInsertIndex = 0;
 	}
 
+	STATIC_VIRTUAL_METHOD()
 	virtual void Execute(const FRigUnitContext& Context) override;
 	
 	/** The value to record */
@@ -118,6 +120,7 @@ struct FRigUnit_TimeOffsetTransform : public FRigUnit_SimBase
 		LastInsertIndex = 0;
 	}
 
+	STATIC_VIRTUAL_METHOD()
 	virtual void Execute(const FRigUnitContext& Context) override;
 	
 	/** The value to record */
diff --git a/Engine/Plugins/Experimental/ControlRig/Source/ControlRig/Private/Units/Simulation/RigUnit_Timeline.cpp b/Engine/Plugins/Experimental/ControlRig/Source/ControlRig/Private/Units/Simulation/RigUnit_Timeline.cpp
--- a/Engine/Plugins/Experimental/ControlRig/Source/ControlRig/Private/Units/Simulation/RigUnit_Timeline.cpp
+++ b/Engine/Plugins/Experimental/ControlRig/Source/ControlRig/Private/Units/Simulation/RigUnit_Timeline.cpp
@@ -3,7 +3,7 @@
 #include "Units/Simulation/RigUnit_Timeline.h"
 #include "Units/RigUnitContext.h"
 
-void FRigUnit_Timeline::Execute(const FRigUnitContext& Context)
+UE_RigUnit_Timeline_IMPLEMENT_STATIC_VIRTUAL_METHOD(void, Execute, const FRigUnitContext& Context)
 {
     DECLARE_SCOPE_HIERARCHICAL_COUNTER_RIGUNIT()
 	if (Context.State == EControlRigState::Init)
diff --git a/Engine/Plugins/Experimental/ControlRig/Source/ControlRig/Private/Units/Simulation/RigUnit_Timeline.h b/Engine/Plugins/Experimental/ControlRig/Source/ControlRig/Private/Units/Simulation/RigUnit_Timeline.h
--- a/Engine/Plugins/Experimental/ControlRig/Source/ControlRig/Private/Units/Simulation/RigUnit_Timeline.h
+++ b/Engine/Plugins/Experimental/ControlRig/Source/ControlRig/Private/Units/Simulation/RigUnit_Timeline.h
@@ -19,6 +19,7 @@ struct FRigUnit_Timeline : public FRigUnit_SimBase
 		Time = AccumulatedValue = 0.f;
 	}
 
+	STATIC_VIRTUAL_METHOD()
 	virtual void Execute(const FRigUnitContext& Context) override;
 
 	UPROPERTY(meta = (Input))
diff --git a/Engine/Plugins/Experimental/ControlRig/Source/ControlRig/Private/Units/Simulation/RigUnit_Verlet.cpp b/Engine/Plugins/Experimental/ControlRig/Source/ControlRig/Private/Units/Simulation/RigUnit_Verlet.cpp
--- a/Engine/Plugins/Experimental/ControlRig/Source/ControlRig/Private/Units/Simulation/RigUnit_Verlet.cpp
+++ b/Engine/Plugins/Experimental/ControlRig/Source/ControlRig/Private/Units/Simulation/RigUnit_Verlet.cpp
@@ -3,7 +3,7 @@
 #include "Units/Simulation/RigUnit_Verlet.h"
 #include "Units/RigUnitContext.h"
 
-void FRigUnit_VerletIntegrateVector::Execute(const FRigUnitContext& Context)
+UE_RigUnit_VerletIntegrateVector_IMPLEMENT_STATIC_VIRTUAL_METHOD(void, Execute, const FRigUnitContext& Context)
 {
     DECLARE_SCOPE_HIERARCHICAL_COUNTER_RIGUNIT()
 	if (Context.State == EControlRigState::Init)
diff --git a/Engine/Plugins/Experimental/ControlRig/Source/ControlRig/Private/Units/Simulation/RigUnit_Verlet.h b/Engine/Plugins/Experimental/ControlRig/Source/ControlRig/Private/Units/Simulation/RigUnit_Verlet.h
--- a/Engine/Plugins/Experimental/ControlRig/Source/ControlRig/Private/Units/Simulation/RigUnit_Verlet.h
+++ b/Engine/Plugins/Experimental/ControlRig/Source/ControlRig/Private/Units/Simulation/RigUnit_Verlet.h
@@ -23,6 +23,7 @@ struct FRigUnit_VerletIntegrateVector : public FRigUnit_SimBase
 		Point = FCRSimPoint();
 	}
 
+	STATIC_VIRTUAL_METHOD()
 	virtual void Execute(const FRigUnitContext& Context) override;
 
 	UPROPERTY(meta = (Input))
diff --git a/Engine/Plugins/Experimental/ControlRig/Source/ControlRig/Public/Units/Control/RigUnit_Control.h b/Engine/Plugins/Experimental/ControlRig/Source/ControlRig/Public/Units/Control/RigUnit_Control.h
--- a/Engine/Plugins/Experimental/ControlRig/Source/ControlRig/Public/Units/Control/RigUnit_Control.h
+++ b/Engine/Plugins/Experimental/ControlRig/Source/ControlRig/Public/Units/Control/RigUnit_Control.h
@@ -21,22 +21,28 @@ struct CONTROLRIG_API FRigUnit_Control : public FRigUnit
 	{
 	}
 
+	STATIC_VIRTUAL_METHOD()
 	virtual void Execute(const FRigUnitContext& Context) override;
 
 	/** Combine Transform and Base to make the resultant transform */
 	FTransform GetResultantTransform() const;
+	static FTransform StaticGetResultantTransform(const FEulerTransform& InTransform, const FTransform& InBase, const FTransformFilter& InFilter);
 
 	/** Combine Transform and Base to make the resultant transform (as a matrix) */
 	FMatrix GetResultantMatrix() const;
+	static FMatrix StaticGetResultantMatrix(const FEulerTransform& InTransform, const FTransform& InBase, const FTransformFilter& InFilter);
 
 	/** Set the transform using a resultant transform (already incorporating Base) */
 	void SetResultantTransform(const FTransform& InResultantTransform);
+	static void StaticSetResultantTransform(const FTransform& InResultantTransform, const FTransform& InBase, FEulerTransform& OutTransform);
 
 	/** Set the transform using a resultant matrix (already incorporating Base) */
 	void SetResultantMatrix(const FMatrix& InResultantMatrix);
+	static void StaticSetResultantMatrix(const FMatrix& InResultantMatrix, const FTransform& InBase, FEulerTransform& OutTransform);
 
 	/** Get the local transform (i.e. without base) with filter applied */
 	FEulerTransform GetFilteredTransform() const;
+	static FEulerTransform StaticGetFilteredTransform(const FEulerTransform& InTransform, const FTransformFilter& InFilter);
 
 #if WITH_EDITORONLY_DATA
 	/** Actor class to use to display this in the viewport */
diff --git a/Engine/Plugins/Experimental/ControlRig/Source/ControlRig/Public/Units/RigUnit.h b/Engine/Plugins/Experimental/ControlRig/Source/ControlRig/Public/Units/RigUnit.h
--- a/Engine/Plugins/Experimental/ControlRig/Source/ControlRig/Public/Units/RigUnit.h
+++ b/Engine/Plugins/Experimental/ControlRig/Source/ControlRig/Public/Units/RigUnit.h
@@ -42,16 +42,16 @@ struct CONTROLRIG_API FRigUnit
 	/* 
 	 * This is property name given by ControlRig as transient when initialized, so only available in run-time
 	 */
-	UPROPERTY(BlueprintReadOnly, transient, Category=FRigUnit)
+	UPROPERTY(BlueprintReadOnly, transient, Category=FRigUnit, meta = (Constant))
 	FName RigUnitName;
 	
 	/* 
 	 * This is struct name given by ControlRig as transient when initialized, so only available in run-time
 	 */
-	UPROPERTY(BlueprintReadOnly, transient, Category=FRigUnit)
+	UPROPERTY(BlueprintReadOnly, transient, Category=FRigUnit, meta = (Constant))
 	FName RigUnitStructName;
 
-	UPROPERTY(EditAnywhere, Category = FRigUnit)
+	UPROPERTY(EditAnywhere, Category = FRigUnit, meta = (Constant))
 	EUnitExecutionType ExecutionType;
 };
 
diff --git a/Engine/Source/Runtime/AnimGraphRuntime/Private/BoneControllers/AnimNode_CCDIK.cpp b/Engine/Source/Runtime/AnimGraphRuntime/Private/BoneControllers/AnimNode_CCDIK.cpp
--- a/Engine/Source/Runtime/AnimGraphRuntime/Private/BoneControllers/AnimNode_CCDIK.cpp
+++ b/Engine/Source/Runtime/AnimGraphRuntime/Private/BoneControllers/AnimNode_CCDIK.cpp
@@ -72,7 +72,7 @@ void FAnimNode_CCDIK::EvaluateSkeletalControl_AnyThread(FComponentSpacePoseConte
 	OutBoneTransforms.AddUninitialized(NumTransforms);
 
 	// Gather chain links. These are non zero length bones.
-	TArray<CCDIKChainLink> Chain;
+	TArray<FCCDIKChainLink> Chain;
 	Chain.Reserve(NumTransforms);
 	// Start with Root Bone
 	{
@@ -81,7 +81,7 @@ void FAnimNode_CCDIK::EvaluateSkeletalControl_AnyThread(FComponentSpacePoseConte
 		const FTransform& BoneCSTransform = Output.Pose.GetComponentSpaceTransform(RootBoneIndex);
 
 		OutBoneTransforms[0] = FBoneTransform(RootBoneIndex, BoneCSTransform);
-		Chain.Add(CCDIKChainLink(BoneCSTransform, LocalTransform, 0));
+		Chain.Add(FCCDIKChainLink(BoneCSTransform, LocalTransform, 0));
 	}
 
 	// Go through remaining transforms
@@ -100,13 +100,13 @@ void FAnimNode_CCDIK::EvaluateSkeletalControl_AnyThread(FComponentSpacePoseConte
 
 		if (!FMath::IsNearlyZero(BoneLength))
 		{
-			Chain.Add(CCDIKChainLink(BoneCSTransform, LocalTransform, TransformIndex));
+			Chain.Add(FCCDIKChainLink(BoneCSTransform, LocalTransform, TransformIndex));
 		}
 		else
 		{
 			// Mark this transform as a zero length child of the last link.
 			// It will inherit position and delta rotation from parent link.
-			CCDIKChainLink & ParentLink = Chain[Chain.Num() - 1];
+			FCCDIKChainLink & ParentLink = Chain[Chain.Num() - 1];
 			ParentLink.ChildZeroLengthTransformIndices.Add(TransformIndex);
 		}
 	}
@@ -122,7 +122,7 @@ void FAnimNode_CCDIK::EvaluateSkeletalControl_AnyThread(FComponentSpacePoseConte
 		// First step: update bone transform positions from chain links.
 		for (int32 LinkIndex = 0; LinkIndex < NumChainLinks; LinkIndex++)
 		{
-			CCDIKChainLink const & ChainLink = Chain[LinkIndex];
+			FCCDIKChainLink const & ChainLink = Chain[LinkIndex];
 			OutBoneTransforms[ChainLink.TransformIndex].Transform = ChainLink.Transform;
 
 			// If there are any zero length children, update position of those
diff --git a/Engine/Source/Runtime/AnimGraphRuntime/Private/BoneControllers/AnimNode_Fabrik.cpp b/Engine/Source/Runtime/AnimGraphRuntime/Private/BoneControllers/AnimNode_Fabrik.cpp
--- a/Engine/Source/Runtime/AnimGraphRuntime/Private/BoneControllers/AnimNode_Fabrik.cpp
+++ b/Engine/Source/Runtime/AnimGraphRuntime/Private/BoneControllers/AnimNode_Fabrik.cpp
@@ -84,7 +84,7 @@ void FAnimNode_Fabrik::EvaluateSkeletalControl_AnyThread(FComponentSpacePoseCont
 	OutBoneTransforms.AddUninitialized(NumTransforms);
 
 	// Gather chain links. These are non zero length bones.
-	TArray<FABRIKChainLink> Chain;
+	TArray<FFABRIKChainLink> Chain;
 	Chain.Reserve(NumTransforms);
 
 	// Start with Root Bone
@@ -93,7 +93,7 @@ void FAnimNode_Fabrik::EvaluateSkeletalControl_AnyThread(FComponentSpacePoseCont
 		const FTransform& BoneCSTransform = Output.Pose.GetComponentSpaceTransform(RootBoneIndex);
 
 		OutBoneTransforms[0] = FBoneTransform(RootBoneIndex, BoneCSTransform);
-		Chain.Add(FABRIKChainLink(BoneCSTransform.GetLocation(), 0.f, RootBoneIndex, 0));
+		Chain.Add(FFABRIKChainLink(BoneCSTransform.GetLocation(), 0.f, RootBoneIndex, 0));
 	}
 
 	// Go through remaining transforms
@@ -111,14 +111,14 @@ void FAnimNode_Fabrik::EvaluateSkeletalControl_AnyThread(FComponentSpacePoseCont
 
 		if (!FMath::IsNearlyZero(BoneLength))
 		{
-			Chain.Add(FABRIKChainLink(BoneCSPosition, BoneLength, BoneIndex, TransformIndex));
+			Chain.Add(FFABRIKChainLink(BoneCSPosition, BoneLength, BoneIndex, TransformIndex));
 			MaximumReach += BoneLength;
 		}
 		else
 		{
 			// Mark this transform as a zero length child of the last link.
 			// It will inherit position and delta rotation from parent link.
-			FABRIKChainLink & ParentLink = Chain[Chain.Num()-1];
+			FFABRIKChainLink & ParentLink = Chain[Chain.Num()-1];
 			ParentLink.ChildZeroLengthTransformIndices.Add(TransformIndex);
 		}
 	}
@@ -131,7 +131,7 @@ void FAnimNode_Fabrik::EvaluateSkeletalControl_AnyThread(FComponentSpacePoseCont
 		// First step: update bone transform positions from chain links.
 		for (int32 LinkIndex = 0; LinkIndex < NumChainLinks; LinkIndex++)
 		{
-			FABRIKChainLink const & ChainLink = Chain[LinkIndex];
+			FFABRIKChainLink const & ChainLink = Chain[LinkIndex];
 			OutBoneTransforms[ChainLink.TransformIndex].Transform.SetTranslation(ChainLink.Position);
 
 			// If there are any zero length children, update position of those
@@ -145,8 +145,8 @@ void FAnimNode_Fabrik::EvaluateSkeletalControl_AnyThread(FComponentSpacePoseCont
 		// FABRIK algorithm - re-orientation of bone local axes after translation calculation
 		for (int32 LinkIndex = 0; LinkIndex < NumChainLinks - 1; LinkIndex++)
 		{
-			FABRIKChainLink const & CurrentLink = Chain[LinkIndex];
-			FABRIKChainLink const & ChildLink = Chain[LinkIndex + 1];
+			FFABRIKChainLink const & CurrentLink = Chain[LinkIndex];
+			FFABRIKChainLink const & ChildLink = Chain[LinkIndex + 1];
 
 			// Calculate pre-translation vector between this bone and child
 			FVector const OldDir = (GetCurrentLocation(Output.Pose, FCompactPoseBoneIndex(ChildLink.BoneIndex)) - GetCurrentLocation(Output.Pose, FCompactPoseBoneIndex(CurrentLink.BoneIndex))).GetUnsafeNormal();
diff --git a/Engine/Source/Runtime/AnimationCore/Private/CCDIK.cpp b/Engine/Source/Runtime/AnimationCore/Private/CCDIK.cpp
--- a/Engine/Source/Runtime/AnimationCore/Private/CCDIK.cpp
+++ b/Engine/Source/Runtime/AnimationCore/Private/CCDIK.cpp
@@ -5,16 +5,16 @@
 namespace AnimationCore
 {
 
-	bool SolveCCDIK(TArray<CCDIKChainLink>& InOutChain, const FVector& TargetPosition, float Precision, int32 MaxIteration, bool bStartFromTail, bool bEnableRotationLimit, const TArray<float>& RotationLimitPerJoints)
+	bool SolveCCDIK(TArray<FCCDIKChainLink>& InOutChain, const FVector& TargetPosition, float Precision, int32 MaxIteration, bool bStartFromTail, bool bEnableRotationLimit, const TArray<float>& RotationLimitPerJoints)
 	{
 		struct Local
 		{
-			static bool UpdateChainLink(TArray<CCDIKChainLink>& Chain, int32 LinkIndex, const FVector& TargetPos, bool bInEnableRotationLimit, const TArray<float>& InRotationLimitPerJoints)
+			static bool UpdateChainLink(TArray<FCCDIKChainLink>& Chain, int32 LinkIndex, const FVector& TargetPos, bool bInEnableRotationLimit, const TArray<float>& InRotationLimitPerJoints)
 			{
 				int32 const TipBoneLinkIndex = Chain.Num() - 1;
 
 				ensure(Chain.IsValidIndex(TipBoneLinkIndex));
-				CCDIKChainLink& CurrentLink = Chain[LinkIndex];
+				FCCDIKChainLink& CurrentLink = Chain[LinkIndex];
 
 				// update new tip pos
 				FVector TipPos = Chain[TipBoneLinkIndex].Transform.GetLocation();
@@ -61,7 +61,7 @@ namespace AnimationCore
 						// if I have parent, make sure to refresh local transform since my current transform has changed
 						if (LinkIndex > 0)
 						{
-							CCDIKChainLink const & Parent = Chain[LinkIndex - 1];
+							FCCDIKChainLink const & Parent = Chain[LinkIndex - 1];
 							CurrentLink.LocalTransform = CurrentLinkTransform.GetRelativeTransform(Parent.Transform);
 							CurrentLink.LocalTransform.NormalizeRotation();
 						}
@@ -72,7 +72,7 @@ namespace AnimationCore
 						// now update all chain
 						for (int32 ChildLinkIndex = LinkIndex + 1; ChildLinkIndex <= TipBoneLinkIndex; ++ChildLinkIndex)
 						{
-							CCDIKChainLink& ChildIterLink = Chain[ChildLinkIndex];
+							FCCDIKChainLink& ChildIterLink = Chain[ChildLinkIndex];
 							const FTransform LocalTransform = ChildIterLink.LocalTransform;
 							ChildIterLink.Transform = LocalTransform * CurrentParentTransform;
 							ChildIterLink.Transform.NormalizeRotation();
diff --git a/Engine/Source/Runtime/AnimationCore/Private/FABRIK.cpp b/Engine/Source/Runtime/AnimationCore/Private/FABRIK.cpp
--- a/Engine/Source/Runtime/AnimationCore/Private/FABRIK.cpp
+++ b/Engine/Source/Runtime/AnimationCore/Private/FABRIK.cpp
@@ -3,7 +3,7 @@
 #include "FABRIK.h"
 #include "BoneContainer.h"
 
-FVector FABRIKChainLink::GetDirectionToParent(const FBoneContainer& BoneContainer, FCompactPoseBoneIndex BoneIndex)
+FVector FFABRIKChainLink::GetDirectionToParent(const FBoneContainer& BoneContainer, FCompactPoseBoneIndex BoneIndex)
 {
 	const FTransform& RefTransform = BoneContainer.GetRefPoseTransform(BoneIndex);
 	return -RefTransform.GetTranslation().GetSafeNormal();
@@ -15,7 +15,7 @@ namespace AnimationCore
 	// Implementation of the FABRIK IK Algorithm
 	// Please see http://andreasaristidou.com/publications/FABRIK.pdf for more details
 
-	bool SolveFabrik(TArray<FABRIKChainLink>& InOutChain, const FVector& TargetPosition, float MaximumReach, float Precision, int32 MaxIterations)
+	bool SolveFabrik(TArray<FFABRIKChainLink>& InOutChain, const FVector& TargetPosition, float MaximumReach, float Precision, int32 MaxIterations)
 	{
 		bool bBoneLocationUpdated = false;
 		float const RootToTargetDistSq = FVector::DistSquared(InOutChain[0].Position, TargetPosition);
@@ -27,8 +27,8 @@ namespace AnimationCore
 		{
 			for (int32 LinkIndex = 1; LinkIndex < NumChainLinks; LinkIndex++)
 			{
-				FABRIKChainLink const & ParentLink = InOutChain[LinkIndex - 1];
-				FABRIKChainLink & CurrentLink = InOutChain[LinkIndex];
+				FFABRIKChainLink const & ParentLink = InOutChain[LinkIndex - 1];
+				FFABRIKChainLink & CurrentLink = InOutChain[LinkIndex];
 				CurrentLink.Position = ParentLink.Position + (TargetPosition - ParentLink.Position).GetUnsafeNormal() * CurrentLink.Length;
 			}
 			bBoneLocationUpdated = true;
@@ -50,8 +50,8 @@ namespace AnimationCore
 					// "Forward Reaching" stage - adjust bones from end effector.
 					for (int32 LinkIndex = TipBoneLinkIndex - 1; LinkIndex > 0; LinkIndex--)
 					{
-						FABRIKChainLink & CurrentLink = InOutChain[LinkIndex];
-						FABRIKChainLink const & ChildLink = InOutChain[LinkIndex + 1];
+						FFABRIKChainLink & CurrentLink = InOutChain[LinkIndex];
+						FFABRIKChainLink const & ChildLink = InOutChain[LinkIndex + 1];
 
 						CurrentLink.Position = ChildLink.Position + (CurrentLink.Position - ChildLink.Position).GetUnsafeNormal() * ChildLink.Length;
 					}
@@ -59,8 +59,8 @@ namespace AnimationCore
 					// "Backward Reaching" stage - adjust bones from root.
 					for (int32 LinkIndex = 1; LinkIndex < TipBoneLinkIndex; LinkIndex++)
 					{
-						FABRIKChainLink const & ParentLink = InOutChain[LinkIndex - 1];
-						FABRIKChainLink & CurrentLink = InOutChain[LinkIndex];
+						FFABRIKChainLink const & ParentLink = InOutChain[LinkIndex - 1];
+						FFABRIKChainLink & CurrentLink = InOutChain[LinkIndex];
 
 						CurrentLink.Position = ParentLink.Position + (CurrentLink.Position - ParentLink.Position).GetUnsafeNormal() * CurrentLink.Length;
 					}
@@ -72,8 +72,8 @@ namespace AnimationCore
 
 				// Place tip bone based on how close we got to target.
 				{
-					FABRIKChainLink const & ParentLink = InOutChain[TipBoneLinkIndex - 1];
-					FABRIKChainLink & CurrentLink = InOutChain[TipBoneLinkIndex];
+					FFABRIKChainLink const & ParentLink = InOutChain[TipBoneLinkIndex - 1];
+					FFABRIKChainLink & CurrentLink = InOutChain[TipBoneLinkIndex];
 
 					CurrentLink.Position = ParentLink.Position + (CurrentLink.Position - ParentLink.Position).GetUnsafeNormal() * CurrentLink.Length;
 				}
diff --git a/Engine/Source/Runtime/AnimationCore/Public/CCDIK.h b/Engine/Source/Runtime/AnimationCore/Public/CCDIK.h
--- a/Engine/Source/Runtime/AnimationCore/Public/CCDIK.h
+++ b/Engine/Source/Runtime/AnimationCore/Public/CCDIK.h
@@ -5,10 +5,14 @@
 #include "CoreMinimal.h"
 #include "UObject/ObjectMacros.h"
 #include "BoneIndices.h"
+#include "CCDIK.Generated.h"
 
 /** Transient structure for CCDIK node evaluation */
-struct CCDIKChainLink
+USTRUCT()
+struct FCCDIKChainLink
 {
+	GENERATED_USTRUCT_BODY()
+
 public:
 	/** Transform of bone in component space. */
 	FTransform Transform;
@@ -25,13 +29,13 @@ public:
 
 	float CurrentAngleDelta;
 
-	CCDIKChainLink()
+	FCCDIKChainLink()
 		: TransformIndex(INDEX_NONE)
 		, CurrentAngleDelta(0.f)
 	{
 	}
 
-	CCDIKChainLink(const FTransform& InTransform, const FTransform& InLocalTransform, const int32& InTransformIndex)
+	FCCDIKChainLink(const FTransform& InTransform, const FTransform& InLocalTransform, const int32& InTransformIndex)
 		: Transform(InTransform)
 		, LocalTransform(InLocalTransform)
 		, TransformIndex(InTransformIndex)
@@ -42,5 +46,5 @@ public:
 
 namespace AnimationCore
 {
-	ANIMATIONCORE_API bool SolveCCDIK(TArray<CCDIKChainLink>& InOutChain, const FVector& TargetPosition, float Precision, int32 MaxIteration, bool bStartFromTail, bool bEnableRotationLimit, const TArray<float>& RotationLimitPerJoints);
+	ANIMATIONCORE_API bool SolveCCDIK(TArray<FCCDIKChainLink>& InOutChain, const FVector& TargetPosition, float Precision, int32 MaxIteration, bool bStartFromTail, bool bEnableRotationLimit, const TArray<float>& RotationLimitPerJoints);
 };
diff --git a/Engine/Source/Runtime/AnimationCore/Public/FABRIK.h b/Engine/Source/Runtime/AnimationCore/Public/FABRIK.h
--- a/Engine/Source/Runtime/AnimationCore/Public/FABRIK.h
+++ b/Engine/Source/Runtime/AnimationCore/Public/FABRIK.h
@@ -5,6 +5,7 @@
 #include "CoreMinimal.h"
 #include "UObject/ObjectMacros.h"
 #include "BoneIndices.h"
+#include "FABRIK.Generated.h"
 
 struct FBoneContainer;
 /**
@@ -14,8 +15,11 @@ struct FBoneContainer;
 /** Transient structure for FABRIK node evaluation */
 
 /** Transient structure for FABRIK node evaluation */
-struct FABRIKChainLink
+USTRUCT()
+struct FFABRIKChainLink
 {
+	GENERATED_USTRUCT_BODY()
+
 public:
 	/** Position of bone in component space. */
 	FVector Position;
@@ -36,7 +40,7 @@ public:
 	* They have a zero length distance, so they will inherit this bone's transformation. */
 	TArray<int32> ChildZeroLengthTransformIndices;
 
-	FABRIKChainLink()
+	FFABRIKChainLink()
 		: Position(FVector::ZeroVector)
 		, Length(0.f)
 		, BoneIndex(INDEX_NONE)
@@ -45,7 +49,7 @@ public:
 	{
 	}
 
-	FABRIKChainLink(const FVector& InPosition, const float InLength, const FCompactPoseBoneIndex& InBoneIndex, const int32& InTransformIndex)
+	FFABRIKChainLink(const FVector& InPosition, const float InLength, const FCompactPoseBoneIndex& InBoneIndex, const int32& InTransformIndex)
 		: Position(InPosition)
 		, Length(InLength)
 		, BoneIndex(InBoneIndex.GetInt())
@@ -54,7 +58,7 @@ public:
 	{
 	}
 
-	FABRIKChainLink(const FVector& InPosition, const float InLength, const FCompactPoseBoneIndex& InBoneIndex, const int32& InTransformIndex, const FVector& InDefaultDirToParent)
+	FFABRIKChainLink(const FVector& InPosition, const float InLength, const FCompactPoseBoneIndex& InBoneIndex, const int32& InTransformIndex, const FVector& InDefaultDirToParent)
 		: Position(InPosition)
 		, Length(InLength)
 		, BoneIndex(InBoneIndex.GetInt())
@@ -63,7 +67,7 @@ public:
 	{
 	}
 
-	FABRIKChainLink(const FVector& InPosition, const float InLength, const int32 InBoneIndex, const int32 InTransformIndex)
+	FFABRIKChainLink(const FVector& InPosition, const float InLength, const int32 InBoneIndex, const int32 InTransformIndex)
 		: Position(InPosition)
 		, Length(InLength)
 		, BoneIndex(InBoneIndex)
@@ -90,5 +94,5 @@ namespace AnimationCore
 	*
 	* @return  true if modified. False if not.
 	*/
-	ANIMATIONCORE_API bool SolveFabrik(TArray<FABRIKChainLink>& InOutChain, const FVector& TargetPosition, float MaximumReach, float Precision, int32 MaxIteration);
+	ANIMATIONCORE_API bool SolveFabrik(TArray<FFABRIKChainLink>& InOutChain, const FVector& TargetPosition, float MaximumReach, float Precision, int32 MaxIteration);
 };
