Fix build break: ControlClass not declared.

diff --git a/Engine/Plugins/Experimental/ControlRig/Source/ControlRig/Public/Units/Control/RigUnit_Control.h b/Engine/Plugins/Experimental/ControlRig/Source/ControlRig/Public/Units/Control/RigUnit_Control.h
--- a/Engine/Plugins/Experimental/ControlRig/Source/ControlRig/Public/Units/Control/RigUnit_Control.h
+++ b/Engine/Plugins/Experimental/ControlRig/Source/ControlRig/Public/Units/Control/RigUnit_Control.h
@@ -44,11 +44,9 @@ struct CONTROLRIG_API FRigUnit_Control : public FRigUnit
 	FEulerTransform GetFilteredTransform() const;
 	static FEulerTransform StaticGetFilteredTransform(const FEulerTransform& InTransform, const FTransformFilter& InFilter);
 
-#if WITH_EDITORONLY_DATA
 	/** Actor class to use to display this in the viewport */
 	UPROPERTY()
 	TSubclassOf<AControlRigControl> ControlClass;
-#endif
 
 	/** The transform of this control */
 	UPROPERTY(EditAnywhere, Category="Control", Interp, meta = (AnimationInput))
diff --git a/Engine/Source/Runtime/AnimationCore/Public/CCDIK.h b/Engine/Source/Runtime/AnimationCore/Public/CCDIK.h
--- a/Engine/Source/Runtime/AnimationCore/Public/CCDIK.h
+++ b/Engine/Source/Runtime/AnimationCore/Public/CCDIK.h
@@ -5,7 +5,7 @@
 #include "CoreMinimal.h"
 #include "UObject/ObjectMacros.h"
 #include "BoneIndices.h"
-#include "CCDIK.Generated.h"
+#include "CCDIK.generated.h"
 
 /** Transient structure for CCDIK node evaluation */
 USTRUCT()
diff --git a/Engine/Source/Runtime/AnimationCore/Public/FABRIK.h b/Engine/Source/Runtime/AnimationCore/Public/FABRIK.h
--- a/Engine/Source/Runtime/AnimationCore/Public/FABRIK.h
+++ b/Engine/Source/Runtime/AnimationCore/Public/FABRIK.h
@@ -5,7 +5,7 @@
 #include "CoreMinimal.h"
 #include "UObject/ObjectMacros.h"
 #include "BoneIndices.h"
-#include "FABRIK.Generated.h"
+#include "FABRIK.generated.h"
 
 struct FBoneContainer;
 /**
