ControlRig: Adapting unit execution to new multiplex macro + new registry unit test

diff --git a/Engine/Plugins/Experimental/ControlRig/Source/ControlRig/Private/Units/Animation/RigUnit_AnimEasing.cpp b/Engine/Plugins/Experimental/ControlRig/Source/ControlRig/Private/Units/Animation/RigUnit_AnimEasing.cpp
--- a/Engine/Plugins/Experimental/ControlRig/Source/ControlRig/Private/Units/Animation/RigUnit_AnimEasing.cpp
+++ b/Engine/Plugins/Experimental/ControlRig/Source/ControlRig/Private/Units/Animation/RigUnit_AnimEasing.cpp
@@ -3,7 +3,7 @@
 #include "Units/Animation/RigUnit_AnimEasing.h"
 #include "Units/RigUnitContext.h"
 
-UE_RigUnit_AnimEasing_IMPLEMENT_STATIC_VIRTUAL_METHOD(void, Execute, const FRigUnitContext& Context)
+UE_RigUnit_AnimEasing_IMPLEMENT_MULTIPLEX(void, Execute, const FRigUnitContext& Context)
 {
     DECLARE_SCOPE_HIERARCHICAL_COUNTER_RIGUNIT()
 	if (FMath::IsNearlyEqual(SourceMinimum, SourceMaximum))
diff --git a/Engine/Plugins/Experimental/ControlRig/Source/ControlRig/Private/Units/Animation/RigUnit_AnimEasing.h b/Engine/Plugins/Experimental/ControlRig/Source/ControlRig/Private/Units/Animation/RigUnit_AnimEasing.h
--- a/Engine/Plugins/Experimental/ControlRig/Source/ControlRig/Private/Units/Animation/RigUnit_AnimEasing.h
+++ b/Engine/Plugins/Experimental/ControlRig/Source/ControlRig/Private/Units/Animation/RigUnit_AnimEasing.h
@@ -39,7 +39,7 @@ struct FRigUnit_AnimEasing : public FRigUnit_AnimBase
 		SourceMaximum = TargetMaximum = 1.f;
 	}
 
-	STATIC_VIRTUAL_METHOD()
+	MULTIPLEX_METHOD()
 	virtual void Execute(const FRigUnitContext& Context) override;
 
 	UPROPERTY(meta=(Input))
diff --git a/Engine/Plugins/Experimental/ControlRig/Source/ControlRig/Private/Units/Animation/RigUnit_AnimEvalRichCurve.cpp b/Engine/Plugins/Experimental/ControlRig/Source/ControlRig/Private/Units/Animation/RigUnit_AnimEvalRichCurve.cpp
--- a/Engine/Plugins/Experimental/ControlRig/Source/ControlRig/Private/Units/Animation/RigUnit_AnimEvalRichCurve.cpp
+++ b/Engine/Plugins/Experimental/ControlRig/Source/ControlRig/Private/Units/Animation/RigUnit_AnimEvalRichCurve.cpp
@@ -3,7 +3,7 @@
 #include "Units/Animation/RigUnit_AnimEvalRichCurve.h"
 #include "Units/RigUnitContext.h"
 
-UE_RigUnit_AnimEvalRichCurve_IMPLEMENT_STATIC_VIRTUAL_METHOD(void, Execute, const FRigUnitContext& Context)
+UE_RigUnit_AnimEvalRichCurve_IMPLEMENT_MULTIPLEX(void, Execute, const FRigUnitContext& Context)
 {
     DECLARE_SCOPE_HIERARCHICAL_COUNTER_RIGUNIT()
 	if (FMath::IsNearlyEqual(SourceMinimum, SourceMaximum))
diff --git a/Engine/Plugins/Experimental/ControlRig/Source/ControlRig/Private/Units/Animation/RigUnit_AnimEvalRichCurve.h b/Engine/Plugins/Experimental/ControlRig/Source/ControlRig/Private/Units/Animation/RigUnit_AnimEvalRichCurve.h
--- a/Engine/Plugins/Experimental/ControlRig/Source/ControlRig/Private/Units/Animation/RigUnit_AnimEvalRichCurve.h
+++ b/Engine/Plugins/Experimental/ControlRig/Source/ControlRig/Private/Units/Animation/RigUnit_AnimEvalRichCurve.h
@@ -24,7 +24,7 @@ struct FRigUnit_AnimEvalRichCurve : public FRigUnit_AnimBase
 		SourceMaximum = TargetMaximum = 1.f;
 	}
 
-	STATIC_VIRTUAL_METHOD()
+	MULTIPLEX_METHOD()
 	virtual void Execute(const FRigUnitContext& Context) override;
 
 	UPROPERTY(meta=(Input))
diff --git a/Engine/Plugins/Experimental/ControlRig/Source/ControlRig/Private/Units/Animation/RigUnit_AnimRichCurve.cpp b/Engine/Plugins/Experimental/ControlRig/Source/ControlRig/Private/Units/Animation/RigUnit_AnimRichCurve.cpp
--- a/Engine/Plugins/Experimental/ControlRig/Source/ControlRig/Private/Units/Animation/RigUnit_AnimRichCurve.cpp
+++ b/Engine/Plugins/Experimental/ControlRig/Source/ControlRig/Private/Units/Animation/RigUnit_AnimRichCurve.cpp
@@ -3,7 +3,7 @@
 #include "Units/Animation/RigUnit_AnimRichCurve.h"
 #include "Units/RigUnitContext.h"
 
-UE_RigUnit_AnimRichCurve_IMPLEMENT_STATIC_VIRTUAL_METHOD(void, Execute, const FRigUnitContext& Context)
+UE_RigUnit_AnimRichCurve_IMPLEMENT_MULTIPLEX(void, Execute, const FRigUnitContext& Context)
 {
     DECLARE_SCOPE_HIERARCHICAL_COUNTER_RIGUNIT()
 }
diff --git a/Engine/Plugins/Experimental/ControlRig/Source/ControlRig/Private/Units/Animation/RigUnit_AnimRichCurve.h b/Engine/Plugins/Experimental/ControlRig/Source/ControlRig/Private/Units/Animation/RigUnit_AnimRichCurve.h
--- a/Engine/Plugins/Experimental/ControlRig/Source/ControlRig/Private/Units/Animation/RigUnit_AnimRichCurve.h
+++ b/Engine/Plugins/Experimental/ControlRig/Source/ControlRig/Private/Units/Animation/RigUnit_AnimRichCurve.h
@@ -21,7 +21,7 @@ struct FRigUnit_AnimRichCurve : public FRigUnit_AnimBase
 		Curve.GetRichCurve()->AddKey(1.f, 1.f);
 	}
 
-	STATIC_VIRTUAL_METHOD()
+	MULTIPLEX_METHOD()
 	virtual void Execute(const FRigUnitContext& Context) override;
 
 	UPROPERTY(meta=(Input, Output, Constant))
diff --git a/Engine/Plugins/Experimental/ControlRig/Source/ControlRig/Private/Units/Animation/RigUnit_GetDeltaTime.cpp b/Engine/Plugins/Experimental/ControlRig/Source/ControlRig/Private/Units/Animation/RigUnit_GetDeltaTime.cpp
--- a/Engine/Plugins/Experimental/ControlRig/Source/ControlRig/Private/Units/Animation/RigUnit_GetDeltaTime.cpp
+++ b/Engine/Plugins/Experimental/ControlRig/Source/ControlRig/Private/Units/Animation/RigUnit_GetDeltaTime.cpp
@@ -3,7 +3,7 @@
 #include "Units/Animation/RigUnit_GetDeltaTime.h"
 #include "Units/RigUnitContext.h"
 
-UE_RigUnit_GetDeltaTime_IMPLEMENT_STATIC_VIRTUAL_METHOD(void, Execute, const FRigUnitContext& Context)
+UE_RigUnit_GetDeltaTime_IMPLEMENT_MULTIPLEX(void, Execute, const FRigUnitContext& Context)
 {
     DECLARE_SCOPE_HIERARCHICAL_COUNTER_RIGUNIT()
 	Result = Context.DeltaTime;
diff --git a/Engine/Plugins/Experimental/ControlRig/Source/ControlRig/Private/Units/Animation/RigUnit_GetDeltaTime.h b/Engine/Plugins/Experimental/ControlRig/Source/ControlRig/Private/Units/Animation/RigUnit_GetDeltaTime.h
--- a/Engine/Plugins/Experimental/ControlRig/Source/ControlRig/Private/Units/Animation/RigUnit_GetDeltaTime.h
+++ b/Engine/Plugins/Experimental/ControlRig/Source/ControlRig/Private/Units/Animation/RigUnit_GetDeltaTime.h
@@ -18,7 +18,7 @@ struct FRigUnit_GetDeltaTime : public FRigUnit_AnimBase
 		Result = 0.f;
 	}
 
-	STATIC_VIRTUAL_METHOD()
+	MULTIPLEX_METHOD()
 	virtual void Execute(const FRigUnitContext& Context) override;
 
 	UPROPERTY(meta=(Output))
diff --git a/Engine/Plugins/Experimental/ControlRig/Source/ControlRig/Private/Units/Animation/RigUnit_GetWorldTime.cpp b/Engine/Plugins/Experimental/ControlRig/Source/ControlRig/Private/Units/Animation/RigUnit_GetWorldTime.cpp
--- a/Engine/Plugins/Experimental/ControlRig/Source/ControlRig/Private/Units/Animation/RigUnit_GetWorldTime.cpp
+++ b/Engine/Plugins/Experimental/ControlRig/Source/ControlRig/Private/Units/Animation/RigUnit_GetWorldTime.cpp
@@ -3,7 +3,7 @@
 #include "Units/Animation/RigUnit_GetWorldTime.h"
 #include "Units/RigUnitContext.h"
 
-UE_RigUnit_GetWorldTime_IMPLEMENT_STATIC_VIRTUAL_METHOD(void, Execute, const FRigUnitContext& Context)
+UE_RigUnit_GetWorldTime_IMPLEMENT_MULTIPLEX(void, Execute, const FRigUnitContext& Context)
 {
     DECLARE_SCOPE_HIERARCHICAL_COUNTER_RIGUNIT()
 	FDateTime Now = FDateTime::Now();
diff --git a/Engine/Plugins/Experimental/ControlRig/Source/ControlRig/Private/Units/Animation/RigUnit_GetWorldTime.h b/Engine/Plugins/Experimental/ControlRig/Source/ControlRig/Private/Units/Animation/RigUnit_GetWorldTime.h
--- a/Engine/Plugins/Experimental/ControlRig/Source/ControlRig/Private/Units/Animation/RigUnit_GetWorldTime.h
+++ b/Engine/Plugins/Experimental/ControlRig/Source/ControlRig/Private/Units/Animation/RigUnit_GetWorldTime.h
@@ -18,7 +18,7 @@ struct FRigUnit_GetWorldTime : public FRigUnit_AnimBase
 		Year = Month = Day = WeekDay = Hours = Minutes = Seconds = OverallSeconds = 0.f;
 	}
 
-	STATIC_VIRTUAL_METHOD()
+	MULTIPLEX_METHOD()
 	virtual void Execute(const FRigUnitContext& Context) override;
 
 	UPROPERTY(meta = (Output))
diff --git a/Engine/Plugins/Experimental/ControlRig/Source/ControlRig/Private/Units/Control/RigUnit_Control.cpp b/Engine/Plugins/Experimental/ControlRig/Source/ControlRig/Private/Units/Control/RigUnit_Control.cpp
--- a/Engine/Plugins/Experimental/ControlRig/Source/ControlRig/Private/Units/Control/RigUnit_Control.cpp
+++ b/Engine/Plugins/Experimental/ControlRig/Source/ControlRig/Private/Units/Control/RigUnit_Control.cpp
@@ -3,7 +3,7 @@
 #include "Units/Control/RigUnit_Control.h"
 #include "Units/RigUnitContext.h"
 
-UE_RigUnit_Control_IMPLEMENT_STATIC_VIRTUAL_METHOD(void, Execute, const FRigUnitContext& Context)
+UE_RigUnit_Control_IMPLEMENT_MULTIPLEX(void, Execute, const FRigUnitContext& Context)
 {
     DECLARE_SCOPE_HIERARCHICAL_COUNTER_RIGUNIT()
 	if (Context.State == EControlRigState::Init)
diff --git a/Engine/Plugins/Experimental/ControlRig/Source/ControlRig/Private/Units/Debug/RigUnit_DebugBezier.cpp b/Engine/Plugins/Experimental/ControlRig/Source/ControlRig/Private/Units/Debug/RigUnit_DebugBezier.cpp
--- a/Engine/Plugins/Experimental/ControlRig/Source/ControlRig/Private/Units/Debug/RigUnit_DebugBezier.cpp
+++ b/Engine/Plugins/Experimental/ControlRig/Source/ControlRig/Private/Units/Debug/RigUnit_DebugBezier.cpp
@@ -3,7 +3,7 @@
 #include "Units/Debug/RigUnit_DebugBezier.h"
 #include "Units/RigUnitContext.h"
 
-UE_RigUnit_DebugBezier_IMPLEMENT_STATIC_VIRTUAL_METHOD(void, Execute, const FRigUnitContext& Context)
+UE_RigUnit_DebugBezier_IMPLEMENT_MULTIPLEX(void, Execute, const FRigUnitContext& Context)
 {
     DECLARE_SCOPE_HIERARCHICAL_COUNTER_RIGUNIT()
 	if (Context.State == EControlRigState::Init)
diff --git a/Engine/Plugins/Experimental/ControlRig/Source/ControlRig/Private/Units/Debug/RigUnit_DebugBezier.h b/Engine/Plugins/Experimental/ControlRig/Source/ControlRig/Private/Units/Debug/RigUnit_DebugBezier.h
--- a/Engine/Plugins/Experimental/ControlRig/Source/ControlRig/Private/Units/Debug/RigUnit_DebugBezier.h
+++ b/Engine/Plugins/Experimental/ControlRig/Source/ControlRig/Private/Units/Debug/RigUnit_DebugBezier.h
@@ -23,7 +23,7 @@ struct FRigUnit_DebugBezier : public FRigUnit_DebugBaseMutable
 		bEnabled = true;
 	}
 
-	STATIC_VIRTUAL_METHOD()
+	MULTIPLEX_METHOD()
 	virtual void Execute(const FRigUnitContext& Context) override;
 
 	UPROPERTY(meta = (Input))
diff --git a/Engine/Plugins/Experimental/ControlRig/Source/ControlRig/Private/Units/Debug/RigUnit_DebugHierarchy.cpp b/Engine/Plugins/Experimental/ControlRig/Source/ControlRig/Private/Units/Debug/RigUnit_DebugHierarchy.cpp
--- a/Engine/Plugins/Experimental/ControlRig/Source/ControlRig/Private/Units/Debug/RigUnit_DebugHierarchy.cpp
+++ b/Engine/Plugins/Experimental/ControlRig/Source/ControlRig/Private/Units/Debug/RigUnit_DebugHierarchy.cpp
@@ -3,7 +3,7 @@
 #include "Units/Debug/RigUnit_DebugHierarchy.h"
 #include "Units/RigUnitContext.h"
 
-UE_RigUnit_DebugHierarchy_IMPLEMENT_STATIC_VIRTUAL_METHOD(void, Execute, const FRigUnitContext& Context)
+UE_RigUnit_DebugHierarchy_IMPLEMENT_MULTIPLEX(void, Execute, const FRigUnitContext& Context)
 {
     DECLARE_SCOPE_HIERARCHICAL_COUNTER_RIGUNIT()
 	if (Context.State == EControlRigState::Init)
diff --git a/Engine/Plugins/Experimental/ControlRig/Source/ControlRig/Private/Units/Debug/RigUnit_DebugHierarchy.h b/Engine/Plugins/Experimental/ControlRig/Source/ControlRig/Private/Units/Debug/RigUnit_DebugHierarchy.h
--- a/Engine/Plugins/Experimental/ControlRig/Source/ControlRig/Private/Units/Debug/RigUnit_DebugHierarchy.h
+++ b/Engine/Plugins/Experimental/ControlRig/Source/ControlRig/Private/Units/Debug/RigUnit_DebugHierarchy.h
@@ -21,7 +21,7 @@ struct FRigUnit_DebugHierarchy : public FRigUnit_DebugBaseMutable
 		bEnabled = true;
 	}
 
-	STATIC_VIRTUAL_METHOD()
+	MULTIPLEX_METHOD()
 	virtual void Execute(const FRigUnitContext& Context) override;
 
 	UPROPERTY(meta = (Input))
diff --git a/Engine/Plugins/Experimental/ControlRig/Source/ControlRig/Private/Units/Debug/RigUnit_DebugLine.cpp b/Engine/Plugins/Experimental/ControlRig/Source/ControlRig/Private/Units/Debug/RigUnit_DebugLine.cpp
--- a/Engine/Plugins/Experimental/ControlRig/Source/ControlRig/Private/Units/Debug/RigUnit_DebugLine.cpp
+++ b/Engine/Plugins/Experimental/ControlRig/Source/ControlRig/Private/Units/Debug/RigUnit_DebugLine.cpp
@@ -3,7 +3,7 @@
 #include "Units/Debug/RigUnit_DebugLine.h"
 #include "Units/RigUnitContext.h"
 
-UE_RigUnit_DebugLine_IMPLEMENT_STATIC_VIRTUAL_METHOD(void, Execute, const FRigUnitContext& Context)
+UE_RigUnit_DebugLine_IMPLEMENT_MULTIPLEX(void, Execute, const FRigUnitContext& Context)
 {
     DECLARE_SCOPE_HIERARCHICAL_COUNTER_RIGUNIT()
 	if (Context.State == EControlRigState::Init)
diff --git a/Engine/Plugins/Experimental/ControlRig/Source/ControlRig/Private/Units/Debug/RigUnit_DebugLine.h b/Engine/Plugins/Experimental/ControlRig/Source/ControlRig/Private/Units/Debug/RigUnit_DebugLine.h
--- a/Engine/Plugins/Experimental/ControlRig/Source/ControlRig/Private/Units/Debug/RigUnit_DebugLine.h
+++ b/Engine/Plugins/Experimental/ControlRig/Source/ControlRig/Private/Units/Debug/RigUnit_DebugLine.h
@@ -19,7 +19,7 @@ struct FRigUnit_DebugLine : public FRigUnit_DebugBaseMutable
 		bEnabled = true;
 	}
 
-	STATIC_VIRTUAL_METHOD()
+	MULTIPLEX_METHOD()
 	virtual void Execute(const FRigUnitContext& Context) override;
 
 	UPROPERTY(meta = (Input))
diff --git a/Engine/Plugins/Experimental/ControlRig/Source/ControlRig/Private/Units/Debug/RigUnit_DebugLineStrip.cpp b/Engine/Plugins/Experimental/ControlRig/Source/ControlRig/Private/Units/Debug/RigUnit_DebugLineStrip.cpp
--- a/Engine/Plugins/Experimental/ControlRig/Source/ControlRig/Private/Units/Debug/RigUnit_DebugLineStrip.cpp
+++ b/Engine/Plugins/Experimental/ControlRig/Source/ControlRig/Private/Units/Debug/RigUnit_DebugLineStrip.cpp
@@ -3,7 +3,7 @@
 #include "Units/Debug/RigUnit_DebugLineStrip.h"
 #include "Units/RigUnitContext.h"
 
-UE_RigUnit_DebugLineStrip_IMPLEMENT_STATIC_VIRTUAL_METHOD(void, Execute, const FRigUnitContext& Context)
+UE_RigUnit_DebugLineStrip_IMPLEMENT_MULTIPLEX(void, Execute, const FRigUnitContext& Context)
 {
     DECLARE_SCOPE_HIERARCHICAL_COUNTER_RIGUNIT()
 	if (Context.State == EControlRigState::Init)
diff --git a/Engine/Plugins/Experimental/ControlRig/Source/ControlRig/Private/Units/Debug/RigUnit_DebugLineStrip.h b/Engine/Plugins/Experimental/ControlRig/Source/ControlRig/Private/Units/Debug/RigUnit_DebugLineStrip.h
--- a/Engine/Plugins/Experimental/ControlRig/Source/ControlRig/Private/Units/Debug/RigUnit_DebugLineStrip.h
+++ b/Engine/Plugins/Experimental/ControlRig/Source/ControlRig/Private/Units/Debug/RigUnit_DebugLineStrip.h
@@ -18,7 +18,7 @@ struct FRigUnit_DebugLineStrip : public FRigUnit_DebugBaseMutable
 		bEnabled = true;
 	}
 
-	STATIC_VIRTUAL_METHOD()
+	MULTIPLEX_METHOD()
 	virtual void Execute(const FRigUnitContext& Context) override;
 
 	UPROPERTY(meta = (Input))
diff --git a/Engine/Plugins/Experimental/ControlRig/Source/ControlRig/Private/Units/Debug/RigUnit_DebugPoint.cpp b/Engine/Plugins/Experimental/ControlRig/Source/ControlRig/Private/Units/Debug/RigUnit_DebugPoint.cpp
--- a/Engine/Plugins/Experimental/ControlRig/Source/ControlRig/Private/Units/Debug/RigUnit_DebugPoint.cpp
+++ b/Engine/Plugins/Experimental/ControlRig/Source/ControlRig/Private/Units/Debug/RigUnit_DebugPoint.cpp
@@ -3,7 +3,7 @@
 #include "Units/Debug/RigUnit_DebugPoint.h"
 #include "Units/RigUnitContext.h"
 
-UE_RigUnit_DebugPoint_IMPLEMENT_STATIC_VIRTUAL_METHOD(void, Execute, const FRigUnitContext& Context)
+UE_RigUnit_DebugPoint_IMPLEMENT_MULTIPLEX(void, Execute, const FRigUnitContext& Context)
 {
     DECLARE_SCOPE_HIERARCHICAL_COUNTER_RIGUNIT()
 	if (Context.State == EControlRigState::Init)
@@ -40,7 +40,7 @@ UE_RigUnit_DebugPoint_IMPLEMENT_STATIC_VIRTUAL_METHOD(void, Execute, const FRigU
 	}
 }
 
-UE_RigUnit_DebugPointMutable_IMPLEMENT_STATIC_VIRTUAL_METHOD(void, Execute, const FRigUnitContext& Context)
+UE_RigUnit_DebugPointMutable_IMPLEMENT_MULTIPLEX(void, Execute, const FRigUnitContext& Context)
 {
     DECLARE_SCOPE_HIERARCHICAL_COUNTER_RIGUNIT()
 	if (Context.State == EControlRigState::Init)
diff --git a/Engine/Plugins/Experimental/ControlRig/Source/ControlRig/Private/Units/Debug/RigUnit_DebugPoint.h b/Engine/Plugins/Experimental/ControlRig/Source/ControlRig/Private/Units/Debug/RigUnit_DebugPoint.h
--- a/Engine/Plugins/Experimental/ControlRig/Source/ControlRig/Private/Units/Debug/RigUnit_DebugPoint.h
+++ b/Engine/Plugins/Experimental/ControlRig/Source/ControlRig/Private/Units/Debug/RigUnit_DebugPoint.h
@@ -34,7 +34,7 @@ struct FRigUnit_DebugPoint : public FRigUnit_DebugBase
 		bEnabled = true;
 	}
 
-	STATIC_VIRTUAL_METHOD()
+	MULTIPLEX_METHOD()
 	virtual void Execute(const FRigUnitContext& Context) override;
 
 	UPROPERTY(meta = (Input, Output))
@@ -78,7 +78,7 @@ struct FRigUnit_DebugPointMutable : public FRigUnit_DebugBaseMutable
 		bEnabled = true;
 	}
 
-	STATIC_VIRTUAL_METHOD()
+	MULTIPLEX_METHOD()
 	virtual void Execute(const FRigUnitContext& Context) override;
 
 	UPROPERTY(meta = (Input))
diff --git a/Engine/Plugins/Experimental/ControlRig/Source/ControlRig/Private/Units/Debug/RigUnit_DebugPrimitives.cpp b/Engine/Plugins/Experimental/ControlRig/Source/ControlRig/Private/Units/Debug/RigUnit_DebugPrimitives.cpp
--- a/Engine/Plugins/Experimental/ControlRig/Source/ControlRig/Private/Units/Debug/RigUnit_DebugPrimitives.cpp
+++ b/Engine/Plugins/Experimental/ControlRig/Source/ControlRig/Private/Units/Debug/RigUnit_DebugPrimitives.cpp
@@ -3,7 +3,7 @@
 #include "Units/Debug/RigUnit_DebugPrimitives.h"
 #include "Units/RigUnitContext.h"
 
-UE_RigUnit_DebugRectangle_IMPLEMENT_STATIC_VIRTUAL_METHOD(void, Execute, const FRigUnitContext& Context)
+UE_RigUnit_DebugRectangle_IMPLEMENT_MULTIPLEX(void, Execute, const FRigUnitContext& Context)
 {
     DECLARE_SCOPE_HIERARCHICAL_COUNTER_RIGUNIT()
 	if (Context.State == EControlRigState::Init)
@@ -25,7 +25,7 @@ UE_RigUnit_DebugRectangle_IMPLEMENT_STATIC_VIRTUAL_METHOD(void, Execute, const F
 	Context.DrawInterface->DrawRectangle(WorldOffset, DrawTransform, Scale, Color, Thickness);
 }
 
-UE_RigUnit_DebugArc_IMPLEMENT_STATIC_VIRTUAL_METHOD(void, Execute, const FRigUnitContext& Context)
+UE_RigUnit_DebugArc_IMPLEMENT_MULTIPLEX(void, Execute, const FRigUnitContext& Context)
 {
     DECLARE_SCOPE_HIERARCHICAL_COUNTER_RIGUNIT()
 	if (Context.State == EControlRigState::Init)
diff --git a/Engine/Plugins/Experimental/ControlRig/Source/ControlRig/Private/Units/Debug/RigUnit_DebugPrimitives.h b/Engine/Plugins/Experimental/ControlRig/Source/ControlRig/Private/Units/Debug/RigUnit_DebugPrimitives.h
--- a/Engine/Plugins/Experimental/ControlRig/Source/ControlRig/Private/Units/Debug/RigUnit_DebugPrimitives.h
+++ b/Engine/Plugins/Experimental/ControlRig/Source/ControlRig/Private/Units/Debug/RigUnit_DebugPrimitives.h
@@ -20,7 +20,7 @@ struct FRigUnit_DebugRectangle : public FRigUnit_DebugBaseMutable
 		bEnabled = true;
 	}
 
-	STATIC_VIRTUAL_METHOD()
+	MULTIPLEX_METHOD()
 	virtual void Execute(const FRigUnitContext& Context) override;
 
 	UPROPERTY(meta = (Input))
@@ -62,7 +62,7 @@ struct FRigUnit_DebugArc : public FRigUnit_DebugBaseMutable
 		bEnabled = true;
 	}
 
-	STATIC_VIRTUAL_METHOD()
+	MULTIPLEX_METHOD()
 	virtual void Execute(const FRigUnitContext& Context) override;
 
 	UPROPERTY(meta = (Input))
diff --git a/Engine/Plugins/Experimental/ControlRig/Source/ControlRig/Private/Units/Debug/RigUnit_DebugTransform.cpp b/Engine/Plugins/Experimental/ControlRig/Source/ControlRig/Private/Units/Debug/RigUnit_DebugTransform.cpp
--- a/Engine/Plugins/Experimental/ControlRig/Source/ControlRig/Private/Units/Debug/RigUnit_DebugTransform.cpp
+++ b/Engine/Plugins/Experimental/ControlRig/Source/ControlRig/Private/Units/Debug/RigUnit_DebugTransform.cpp
@@ -3,7 +3,7 @@
 #include "Units/Debug/RigUnit_DebugTransform.h"
 #include "Units/RigUnitContext.h"
 
-UE_RigUnit_DebugTransform_IMPLEMENT_STATIC_VIRTUAL_METHOD(void, Execute, const FRigUnitContext& Context)
+UE_RigUnit_DebugTransform_IMPLEMENT_MULTIPLEX(void, Execute, const FRigUnitContext& Context)
 {
     DECLARE_SCOPE_HIERARCHICAL_COUNTER_RIGUNIT()
 	if (Context.State == EControlRigState::Init)
@@ -43,7 +43,7 @@ UE_RigUnit_DebugTransform_IMPLEMENT_STATIC_VIRTUAL_METHOD(void, Execute, const F
 	}
 }
 
-UE_RigUnit_DebugTransformMutable_IMPLEMENT_STATIC_VIRTUAL_METHOD(void, Execute, const FRigUnitContext& Context)
+UE_RigUnit_DebugTransformMutable_IMPLEMENT_MULTIPLEX(void, Execute, const FRigUnitContext& Context)
 {
     DECLARE_SCOPE_HIERARCHICAL_COUNTER_RIGUNIT()
 	if (Context.State == EControlRigState::Init)
@@ -83,7 +83,7 @@ UE_RigUnit_DebugTransformMutable_IMPLEMENT_STATIC_VIRTUAL_METHOD(void, Execute,
 	}
 }
 
-UE_RigUnit_DebugTransformArrayMutable_IMPLEMENT_STATIC_VIRTUAL_METHOD(void, Execute, const FRigUnitContext& Context)
+UE_RigUnit_DebugTransformArrayMutable_IMPLEMENT_MULTIPLEX(void, Execute, const FRigUnitContext& Context)
 {
     DECLARE_SCOPE_HIERARCHICAL_COUNTER_RIGUNIT()
 	if (Context.State == EControlRigState::Init)
diff --git a/Engine/Plugins/Experimental/ControlRig/Source/ControlRig/Private/Units/Debug/RigUnit_DebugTransform.h b/Engine/Plugins/Experimental/ControlRig/Source/ControlRig/Private/Units/Debug/RigUnit_DebugTransform.h
--- a/Engine/Plugins/Experimental/ControlRig/Source/ControlRig/Private/Units/Debug/RigUnit_DebugTransform.h
+++ b/Engine/Plugins/Experimental/ControlRig/Source/ControlRig/Private/Units/Debug/RigUnit_DebugTransform.h
@@ -36,7 +36,7 @@ struct FRigUnit_DebugTransform : public FRigUnit_DebugBase
 		bEnabled = true;
 	}
 
-	STATIC_VIRTUAL_METHOD()
+	MULTIPLEX_METHOD()
 	virtual void Execute(const FRigUnitContext& Context) override;
 
 	UPROPERTY(meta = (Input, Output))
@@ -79,7 +79,7 @@ struct FRigUnit_DebugTransformMutable : public FRigUnit_DebugBaseMutable
 		bEnabled = true;
 	}
 
-	STATIC_VIRTUAL_METHOD()
+	MULTIPLEX_METHOD()
 	virtual void Execute(const FRigUnitContext& Context) override;
 
 	UPROPERTY(meta = (Input))
@@ -122,7 +122,7 @@ struct FRigUnit_DebugTransformArrayMutable : public FRigUnit_DebugBaseMutable
 		bEnabled = true;
 	}
 
-	STATIC_VIRTUAL_METHOD()
+	MULTIPLEX_METHOD()
 	virtual void Execute(const FRigUnitContext& Context) override;
 
 	UPROPERTY(meta = (Input))
diff --git a/Engine/Plugins/Experimental/ControlRig/Source/ControlRig/Private/Units/Debug/RigUnit_ProfilingBracket.cpp b/Engine/Plugins/Experimental/ControlRig/Source/ControlRig/Private/Units/Debug/RigUnit_ProfilingBracket.cpp
--- a/Engine/Plugins/Experimental/ControlRig/Source/ControlRig/Private/Units/Debug/RigUnit_ProfilingBracket.cpp
+++ b/Engine/Plugins/Experimental/ControlRig/Source/ControlRig/Private/Units/Debug/RigUnit_ProfilingBracket.cpp
@@ -4,7 +4,7 @@
 #include "Units/RigUnitContext.h"
 #include "KismetAnimationLibrary.h"
 
-UE_RigUnit_StartProfilingTimer_IMPLEMENT_STATIC_VIRTUAL_METHOD(void, Execute, const FRigUnitContext& Context)
+UE_RigUnit_StartProfilingTimer_IMPLEMENT_MULTIPLEX(void, Execute, const FRigUnitContext& Context)
 {
     DECLARE_SCOPE_HIERARCHICAL_COUNTER_RIGUNIT()
 	if (Context.State == EControlRigState::Update)
@@ -13,7 +13,7 @@ UE_RigUnit_StartProfilingTimer_IMPLEMENT_STATIC_VIRTUAL_METHOD(void, Execute, co
 	}
 }
 
-UE_RigUnit_EndProfilingTimer_IMPLEMENT_STATIC_VIRTUAL_METHOD(void, Execute, const FRigUnitContext& Context)
+UE_RigUnit_EndProfilingTimer_IMPLEMENT_MULTIPLEX(void, Execute, const FRigUnitContext& Context)
 {
     DECLARE_SCOPE_HIERARCHICAL_COUNTER_RIGUNIT()
 	if (Context.State == EControlRigState::Init)
diff --git a/Engine/Plugins/Experimental/ControlRig/Source/ControlRig/Private/Units/Debug/RigUnit_ProfilingBracket.h b/Engine/Plugins/Experimental/ControlRig/Source/ControlRig/Private/Units/Debug/RigUnit_ProfilingBracket.h
--- a/Engine/Plugins/Experimental/ControlRig/Source/ControlRig/Private/Units/Debug/RigUnit_ProfilingBracket.h
+++ b/Engine/Plugins/Experimental/ControlRig/Source/ControlRig/Private/Units/Debug/RigUnit_ProfilingBracket.h
@@ -14,7 +14,7 @@ struct FRigUnit_StartProfilingTimer : public FRigUnit_DebugBaseMutable
 	{
 	}
 
-	STATIC_VIRTUAL_METHOD()
+	MULTIPLEX_METHOD()
 	virtual void Execute(const FRigUnitContext& Context) override;
 };
 
@@ -31,7 +31,7 @@ struct FRigUnit_EndProfilingTimer : public FRigUnit_DebugBaseMutable
 		Prefix = TEXT("Timer");
 	}
 
-	STATIC_VIRTUAL_METHOD()
+	MULTIPLEX_METHOD()
 	virtual void Execute(const FRigUnitContext& Context) override;
 
 	UPROPERTY(meta = (Input, Constant))
diff --git a/Engine/Plugins/Experimental/ControlRig/Source/ControlRig/Private/Units/Execution/RigUnit_BeginExecution.cpp b/Engine/Plugins/Experimental/ControlRig/Source/ControlRig/Private/Units/Execution/RigUnit_BeginExecution.cpp
--- a/Engine/Plugins/Experimental/ControlRig/Source/ControlRig/Private/Units/Execution/RigUnit_BeginExecution.cpp
+++ b/Engine/Plugins/Experimental/ControlRig/Source/ControlRig/Private/Units/Execution/RigUnit_BeginExecution.cpp
@@ -3,7 +3,7 @@
 #include "RigUnit_BeginExecution.h"
 #include "Units/RigUnitContext.h"
 
-UE_RigUnit_BeginExecution_IMPLEMENT_STATIC_VIRTUAL_METHOD(void, Execute, const FRigUnitContext& Context)
+UE_RigUnit_BeginExecution_IMPLEMENT_MULTIPLEX(void, Execute, const FRigUnitContext& Context)
 {
     DECLARE_SCOPE_HIERARCHICAL_COUNTER_RIGUNIT()
 	ExecuteContext.Hierarchy = (FRigHierarchyContainer*)Context.Hierarchy;
diff --git a/Engine/Plugins/Experimental/ControlRig/Source/ControlRig/Private/Units/Execution/RigUnit_BeginExecution.h b/Engine/Plugins/Experimental/ControlRig/Source/ControlRig/Private/Units/Execution/RigUnit_BeginExecution.h
--- a/Engine/Plugins/Experimental/ControlRig/Source/ControlRig/Private/Units/Execution/RigUnit_BeginExecution.h
+++ b/Engine/Plugins/Experimental/ControlRig/Source/ControlRig/Private/Units/Execution/RigUnit_BeginExecution.h
@@ -11,7 +11,7 @@ struct FRigUnit_BeginExecution : public FRigUnit
 {
 	GENERATED_BODY()
 
-	STATIC_VIRTUAL_METHOD()
+	MULTIPLEX_METHOD()
 	virtual void Execute(const FRigUnitContext& Context) override;
 
 	// The execution result
diff --git a/Engine/Plugins/Experimental/ControlRig/Source/ControlRig/Private/Units/Hierarchy/RigUnit_AddBoneTransform.cpp b/Engine/Plugins/Experimental/ControlRig/Source/ControlRig/Private/Units/Hierarchy/RigUnit_AddBoneTransform.cpp
--- a/Engine/Plugins/Experimental/ControlRig/Source/ControlRig/Private/Units/Hierarchy/RigUnit_AddBoneTransform.cpp
+++ b/Engine/Plugins/Experimental/ControlRig/Source/ControlRig/Private/Units/Hierarchy/RigUnit_AddBoneTransform.cpp
@@ -8,7 +8,7 @@ FString FRigUnit_AddBoneTransform::GetUnitLabel() const
 	return FString::Printf(TEXT("Offset Transform %s"), *Bone.ToString());
 }
 
-UE_RigUnit_AddBoneTransform_IMPLEMENT_STATIC_VIRTUAL_METHOD(void, Execute, const FRigUnitContext& Context)
+UE_RigUnit_AddBoneTransform_IMPLEMENT_MULTIPLEX(void, Execute, const FRigUnitContext& Context)
 {
     DECLARE_SCOPE_HIERARCHICAL_COUNTER_RIGUNIT()
 	FRigBoneHierarchy* Hierarchy = ExecuteContext.GetBones();
diff --git a/Engine/Plugins/Experimental/ControlRig/Source/ControlRig/Private/Units/Hierarchy/RigUnit_AddBoneTransform.h b/Engine/Plugins/Experimental/ControlRig/Source/ControlRig/Private/Units/Hierarchy/RigUnit_AddBoneTransform.h
--- a/Engine/Plugins/Experimental/ControlRig/Source/ControlRig/Private/Units/Hierarchy/RigUnit_AddBoneTransform.h
+++ b/Engine/Plugins/Experimental/ControlRig/Source/ControlRig/Private/Units/Hierarchy/RigUnit_AddBoneTransform.h
@@ -21,7 +21,7 @@ struct FRigUnit_AddBoneTransform : public FRigUnitMutable
 	{}
 
 	virtual FString GetUnitLabel() const override;
-	STATIC_VIRTUAL_METHOD()
+	MULTIPLEX_METHOD()
 	virtual void Execute(const FRigUnitContext& Context) override;
 
 	/**
diff --git a/Engine/Plugins/Experimental/ControlRig/Source/ControlRig/Private/Units/Hierarchy/RigUnit_GetBoneTransform.cpp b/Engine/Plugins/Experimental/ControlRig/Source/ControlRig/Private/Units/Hierarchy/RigUnit_GetBoneTransform.cpp
--- a/Engine/Plugins/Experimental/ControlRig/Source/ControlRig/Private/Units/Hierarchy/RigUnit_GetBoneTransform.cpp
+++ b/Engine/Plugins/Experimental/ControlRig/Source/ControlRig/Private/Units/Hierarchy/RigUnit_GetBoneTransform.cpp
@@ -8,7 +8,7 @@ FString FRigUnit_GetBoneTransform::GetUnitLabel() const
 	return FString::Printf(TEXT("Get Transform %s"), *Bone.ToString());
 }
 
-UE_RigUnit_GetBoneTransform_IMPLEMENT_STATIC_VIRTUAL_METHOD(void, Execute, const FRigUnitContext& Context)
+UE_RigUnit_GetBoneTransform_IMPLEMENT_MULTIPLEX(void, Execute, const FRigUnitContext& Context)
 {
     DECLARE_SCOPE_HIERARCHICAL_COUNTER_RIGUNIT()
 	const FRigBoneHierarchy* Hierarchy = Context.GetBones();
diff --git a/Engine/Plugins/Experimental/ControlRig/Source/ControlRig/Private/Units/Hierarchy/RigUnit_GetBoneTransform.h b/Engine/Plugins/Experimental/ControlRig/Source/ControlRig/Private/Units/Hierarchy/RigUnit_GetBoneTransform.h
--- a/Engine/Plugins/Experimental/ControlRig/Source/ControlRig/Private/Units/Hierarchy/RigUnit_GetBoneTransform.h
+++ b/Engine/Plugins/Experimental/ControlRig/Source/ControlRig/Private/Units/Hierarchy/RigUnit_GetBoneTransform.h
@@ -19,7 +19,7 @@ struct FRigUnit_GetBoneTransform : public FRigUnit
 	{}
 
 	virtual FString GetUnitLabel() const override;
-	STATIC_VIRTUAL_METHOD()
+	MULTIPLEX_METHOD()
 	virtual void Execute(const FRigUnitContext& Context) override;
 
 	/**
diff --git a/Engine/Plugins/Experimental/ControlRig/Source/ControlRig/Private/Units/Hierarchy/RigUnit_GetCurveValue.cpp b/Engine/Plugins/Experimental/ControlRig/Source/ControlRig/Private/Units/Hierarchy/RigUnit_GetCurveValue.cpp
--- a/Engine/Plugins/Experimental/ControlRig/Source/ControlRig/Private/Units/Hierarchy/RigUnit_GetCurveValue.cpp
+++ b/Engine/Plugins/Experimental/ControlRig/Source/ControlRig/Private/Units/Hierarchy/RigUnit_GetCurveValue.cpp
@@ -8,7 +8,7 @@ FString FRigUnit_GetCurveValue::GetUnitLabel() const
 	return FString::Printf(TEXT("Get Curve %s"), *Curve.ToString());
 }
 
-UE_RigUnit_GetCurveValue_IMPLEMENT_STATIC_VIRTUAL_METHOD(void, Execute, const FRigUnitContext& Context)
+UE_RigUnit_GetCurveValue_IMPLEMENT_MULTIPLEX(void, Execute, const FRigUnitContext& Context)
 {
     DECLARE_SCOPE_HIERARCHICAL_COUNTER_RIGUNIT()
 	const FRigCurveContainer* CurveContainer = Context.GetCurves();
diff --git a/Engine/Plugins/Experimental/ControlRig/Source/ControlRig/Private/Units/Hierarchy/RigUnit_GetCurveValue.h b/Engine/Plugins/Experimental/ControlRig/Source/ControlRig/Private/Units/Hierarchy/RigUnit_GetCurveValue.h
--- a/Engine/Plugins/Experimental/ControlRig/Source/ControlRig/Private/Units/Hierarchy/RigUnit_GetCurveValue.h
+++ b/Engine/Plugins/Experimental/ControlRig/Source/ControlRig/Private/Units/Hierarchy/RigUnit_GetCurveValue.h
@@ -20,7 +20,7 @@ struct FRigUnit_GetCurveValue : public FRigUnit
 	{}
 
 	virtual FString GetUnitLabel() const override;
-	STATIC_VIRTUAL_METHOD()
+	MULTIPLEX_METHOD()
 	virtual void Execute(const FRigUnitContext& Context) override;
 
 	/**
diff --git a/Engine/Plugins/Experimental/ControlRig/Source/ControlRig/Private/Units/Hierarchy/RigUnit_GetInitialBoneTransform.cpp b/Engine/Plugins/Experimental/ControlRig/Source/ControlRig/Private/Units/Hierarchy/RigUnit_GetInitialBoneTransform.cpp
--- a/Engine/Plugins/Experimental/ControlRig/Source/ControlRig/Private/Units/Hierarchy/RigUnit_GetInitialBoneTransform.cpp
+++ b/Engine/Plugins/Experimental/ControlRig/Source/ControlRig/Private/Units/Hierarchy/RigUnit_GetInitialBoneTransform.cpp
@@ -8,7 +8,7 @@ FString FRigUnit_GetInitialBoneTransform::GetUnitLabel() const
 	return FString::Printf(TEXT("Get Initial Transform %s"), *Bone.ToString());
 }
 
-UE_RigUnit_GetInitialBoneTransform_IMPLEMENT_STATIC_VIRTUAL_METHOD(void, Execute, const FRigUnitContext& Context)
+UE_RigUnit_GetInitialBoneTransform_IMPLEMENT_MULTIPLEX(void, Execute, const FRigUnitContext& Context)
 {
     DECLARE_SCOPE_HIERARCHICAL_COUNTER_RIGUNIT()
 	const FRigBoneHierarchy* Hierarchy = Context.GetBones();
diff --git a/Engine/Plugins/Experimental/ControlRig/Source/ControlRig/Private/Units/Hierarchy/RigUnit_GetInitialBoneTransform.h b/Engine/Plugins/Experimental/ControlRig/Source/ControlRig/Private/Units/Hierarchy/RigUnit_GetInitialBoneTransform.h
--- a/Engine/Plugins/Experimental/ControlRig/Source/ControlRig/Private/Units/Hierarchy/RigUnit_GetInitialBoneTransform.h
+++ b/Engine/Plugins/Experimental/ControlRig/Source/ControlRig/Private/Units/Hierarchy/RigUnit_GetInitialBoneTransform.h
@@ -19,7 +19,7 @@ struct FRigUnit_GetInitialBoneTransform : public FRigUnit
 	{}
 
 	virtual FString GetUnitLabel() const override;
-	STATIC_VIRTUAL_METHOD()
+	MULTIPLEX_METHOD()
 	virtual void Execute(const FRigUnitContext& Context) override;
 
 	/**
diff --git a/Engine/Plugins/Experimental/ControlRig/Source/ControlRig/Private/Units/Hierarchy/RigUnit_GetRelativeBoneTransform.cpp b/Engine/Plugins/Experimental/ControlRig/Source/ControlRig/Private/Units/Hierarchy/RigUnit_GetRelativeBoneTransform.cpp
--- a/Engine/Plugins/Experimental/ControlRig/Source/ControlRig/Private/Units/Hierarchy/RigUnit_GetRelativeBoneTransform.cpp
+++ b/Engine/Plugins/Experimental/ControlRig/Source/ControlRig/Private/Units/Hierarchy/RigUnit_GetRelativeBoneTransform.cpp
@@ -8,7 +8,7 @@ FString FRigUnit_GetRelativeBoneTransform::GetUnitLabel() const
 	return FString::Printf(TEXT("Get Relative Transform %s"), *Bone.ToString());
 }
 
-UE_RigUnit_GetRelativeBoneTransform_IMPLEMENT_STATIC_VIRTUAL_METHOD(void, Execute, const FRigUnitContext& Context)
+UE_RigUnit_GetRelativeBoneTransform_IMPLEMENT_MULTIPLEX(void, Execute, const FRigUnitContext& Context)
 {
     DECLARE_SCOPE_HIERARCHICAL_COUNTER_RIGUNIT()
 	const FRigBoneHierarchy* Hierarchy = Context.GetBones();
diff --git a/Engine/Plugins/Experimental/ControlRig/Source/ControlRig/Private/Units/Hierarchy/RigUnit_GetRelativeBoneTransform.h b/Engine/Plugins/Experimental/ControlRig/Source/ControlRig/Private/Units/Hierarchy/RigUnit_GetRelativeBoneTransform.h
--- a/Engine/Plugins/Experimental/ControlRig/Source/ControlRig/Private/Units/Hierarchy/RigUnit_GetRelativeBoneTransform.h
+++ b/Engine/Plugins/Experimental/ControlRig/Source/ControlRig/Private/Units/Hierarchy/RigUnit_GetRelativeBoneTransform.h
@@ -19,7 +19,7 @@ struct FRigUnit_GetRelativeBoneTransform : public FRigUnit
 	{}
 
 	virtual FString GetUnitLabel() const override;
-	STATIC_VIRTUAL_METHOD()
+	MULTIPLEX_METHOD()
 	virtual void Execute(const FRigUnitContext& Context) override;
 
 	/**
diff --git a/Engine/Plugins/Experimental/ControlRig/Source/ControlRig/Private/Units/Hierarchy/RigUnit_SetBoneRotation.cpp b/Engine/Plugins/Experimental/ControlRig/Source/ControlRig/Private/Units/Hierarchy/RigUnit_SetBoneRotation.cpp
--- a/Engine/Plugins/Experimental/ControlRig/Source/ControlRig/Private/Units/Hierarchy/RigUnit_SetBoneRotation.cpp
+++ b/Engine/Plugins/Experimental/ControlRig/Source/ControlRig/Private/Units/Hierarchy/RigUnit_SetBoneRotation.cpp
@@ -8,7 +8,7 @@ FString FRigUnit_SetBoneRotation::GetUnitLabel() const
 	return FString::Printf(TEXT("Set Rotation %s"), *Bone.ToString());
 }
 
-UE_RigUnit_SetBoneRotation_IMPLEMENT_STATIC_VIRTUAL_METHOD(void, Execute, const FRigUnitContext& Context)
+UE_RigUnit_SetBoneRotation_IMPLEMENT_MULTIPLEX(void, Execute, const FRigUnitContext& Context)
 {
     DECLARE_SCOPE_HIERARCHICAL_COUNTER_RIGUNIT()
 	FRigBoneHierarchy* Hierarchy = ExecuteContext.GetBones();
diff --git a/Engine/Plugins/Experimental/ControlRig/Source/ControlRig/Private/Units/Hierarchy/RigUnit_SetBoneRotation.h b/Engine/Plugins/Experimental/ControlRig/Source/ControlRig/Private/Units/Hierarchy/RigUnit_SetBoneRotation.h
--- a/Engine/Plugins/Experimental/ControlRig/Source/ControlRig/Private/Units/Hierarchy/RigUnit_SetBoneRotation.h
+++ b/Engine/Plugins/Experimental/ControlRig/Source/ControlRig/Private/Units/Hierarchy/RigUnit_SetBoneRotation.h
@@ -22,7 +22,7 @@ struct FRigUnit_SetBoneRotation : public FRigUnitMutable
 	{}
 
 	virtual FString GetUnitLabel() const override;
-	STATIC_VIRTUAL_METHOD()
+	MULTIPLEX_METHOD()
 	virtual void Execute(const FRigUnitContext& Context) override;
 
 	/**
diff --git a/Engine/Plugins/Experimental/ControlRig/Source/ControlRig/Private/Units/Hierarchy/RigUnit_SetBoneTransform.cpp b/Engine/Plugins/Experimental/ControlRig/Source/ControlRig/Private/Units/Hierarchy/RigUnit_SetBoneTransform.cpp
--- a/Engine/Plugins/Experimental/ControlRig/Source/ControlRig/Private/Units/Hierarchy/RigUnit_SetBoneTransform.cpp
+++ b/Engine/Plugins/Experimental/ControlRig/Source/ControlRig/Private/Units/Hierarchy/RigUnit_SetBoneTransform.cpp
@@ -8,7 +8,7 @@ FString FRigUnit_SetBoneTransform::GetUnitLabel() const
 	return FString::Printf(TEXT("Set Transform %s"), *Bone.ToString());
 }
 
-UE_RigUnit_SetBoneTransform_IMPLEMENT_STATIC_VIRTUAL_METHOD(void, Execute, const FRigUnitContext& Context)
+UE_RigUnit_SetBoneTransform_IMPLEMENT_MULTIPLEX(void, Execute, const FRigUnitContext& Context)
 {
     DECLARE_SCOPE_HIERARCHICAL_COUNTER_RIGUNIT()
 	FRigBoneHierarchy* Hierarchy = ExecuteContext.GetBones();
diff --git a/Engine/Plugins/Experimental/ControlRig/Source/ControlRig/Private/Units/Hierarchy/RigUnit_SetBoneTransform.h b/Engine/Plugins/Experimental/ControlRig/Source/ControlRig/Private/Units/Hierarchy/RigUnit_SetBoneTransform.h
--- a/Engine/Plugins/Experimental/ControlRig/Source/ControlRig/Private/Units/Hierarchy/RigUnit_SetBoneTransform.h
+++ b/Engine/Plugins/Experimental/ControlRig/Source/ControlRig/Private/Units/Hierarchy/RigUnit_SetBoneTransform.h
@@ -22,7 +22,7 @@ struct FRigUnit_SetBoneTransform : public FRigUnitMutable
 
 	virtual FString GetUnitLabel() const override;
 
-	STATIC_VIRTUAL_METHOD()
+	MULTIPLEX_METHOD()
 	virtual void Execute(const FRigUnitContext& Context) override;
 
 	/**
diff --git a/Engine/Plugins/Experimental/ControlRig/Source/ControlRig/Private/Units/Hierarchy/RigUnit_SetBoneTranslation.cpp b/Engine/Plugins/Experimental/ControlRig/Source/ControlRig/Private/Units/Hierarchy/RigUnit_SetBoneTranslation.cpp
--- a/Engine/Plugins/Experimental/ControlRig/Source/ControlRig/Private/Units/Hierarchy/RigUnit_SetBoneTranslation.cpp
+++ b/Engine/Plugins/Experimental/ControlRig/Source/ControlRig/Private/Units/Hierarchy/RigUnit_SetBoneTranslation.cpp
@@ -8,7 +8,7 @@ FString FRigUnit_SetBoneTranslation::GetUnitLabel() const
 	return FString::Printf(TEXT("Set Translation %s"), *Bone.ToString());
 }
 
-UE_RigUnit_SetBoneTranslation_IMPLEMENT_STATIC_VIRTUAL_METHOD(void, Execute, const FRigUnitContext& Context)
+UE_RigUnit_SetBoneTranslation_IMPLEMENT_MULTIPLEX(void, Execute, const FRigUnitContext& Context)
 {
     DECLARE_SCOPE_HIERARCHICAL_COUNTER_RIGUNIT()
 	FRigBoneHierarchy* Hierarchy = ExecuteContext.GetBones();
diff --git a/Engine/Plugins/Experimental/ControlRig/Source/ControlRig/Private/Units/Hierarchy/RigUnit_SetBoneTranslation.h b/Engine/Plugins/Experimental/ControlRig/Source/ControlRig/Private/Units/Hierarchy/RigUnit_SetBoneTranslation.h
--- a/Engine/Plugins/Experimental/ControlRig/Source/ControlRig/Private/Units/Hierarchy/RigUnit_SetBoneTranslation.h
+++ b/Engine/Plugins/Experimental/ControlRig/Source/ControlRig/Private/Units/Hierarchy/RigUnit_SetBoneTranslation.h
@@ -22,7 +22,7 @@ struct FRigUnit_SetBoneTranslation : public FRigUnitMutable
 	{}
 
 	virtual FString GetUnitLabel() const override;
-	STATIC_VIRTUAL_METHOD()
+	MULTIPLEX_METHOD()
 	virtual void Execute(const FRigUnitContext& Context) override;
 
 	/**
diff --git a/Engine/Plugins/Experimental/ControlRig/Source/ControlRig/Private/Units/Hierarchy/RigUnit_SetCurveValue.cpp b/Engine/Plugins/Experimental/ControlRig/Source/ControlRig/Private/Units/Hierarchy/RigUnit_SetCurveValue.cpp
--- a/Engine/Plugins/Experimental/ControlRig/Source/ControlRig/Private/Units/Hierarchy/RigUnit_SetCurveValue.cpp
+++ b/Engine/Plugins/Experimental/ControlRig/Source/ControlRig/Private/Units/Hierarchy/RigUnit_SetCurveValue.cpp
@@ -8,7 +8,7 @@ FString FRigUnit_SetCurveValue::GetUnitLabel() const
 	return FString::Printf(TEXT("Set Curve [%s]"), *Curve.ToString());
 }
 
-UE_RigUnit_SetCurveValue_IMPLEMENT_STATIC_VIRTUAL_METHOD(void, Execute, const FRigUnitContext& Context)
+UE_RigUnit_SetCurveValue_IMPLEMENT_MULTIPLEX(void, Execute, const FRigUnitContext& Context)
 {
     DECLARE_SCOPE_HIERARCHICAL_COUNTER_RIGUNIT()
 	FRigCurveContainer* CurveContainer = ExecuteContext.GetCurves();
diff --git a/Engine/Plugins/Experimental/ControlRig/Source/ControlRig/Private/Units/Hierarchy/RigUnit_SetCurveValue.h b/Engine/Plugins/Experimental/ControlRig/Source/ControlRig/Private/Units/Hierarchy/RigUnit_SetCurveValue.h
--- a/Engine/Plugins/Experimental/ControlRig/Source/ControlRig/Private/Units/Hierarchy/RigUnit_SetCurveValue.h
+++ b/Engine/Plugins/Experimental/ControlRig/Source/ControlRig/Private/Units/Hierarchy/RigUnit_SetCurveValue.h
@@ -20,7 +20,7 @@ struct FRigUnit_SetCurveValue : public FRigUnitMutable
 	{}
 
 	virtual FString GetUnitLabel() const override;
-	STATIC_VIRTUAL_METHOD()
+	MULTIPLEX_METHOD()
 	virtual void Execute(const FRigUnitContext& Context) override;
 
 	/**
diff --git a/Engine/Plugins/Experimental/ControlRig/Source/ControlRig/Private/Units/Hierarchy/RigUnit_SetRelativeBoneTransform.cpp b/Engine/Plugins/Experimental/ControlRig/Source/ControlRig/Private/Units/Hierarchy/RigUnit_SetRelativeBoneTransform.cpp
--- a/Engine/Plugins/Experimental/ControlRig/Source/ControlRig/Private/Units/Hierarchy/RigUnit_SetRelativeBoneTransform.cpp
+++ b/Engine/Plugins/Experimental/ControlRig/Source/ControlRig/Private/Units/Hierarchy/RigUnit_SetRelativeBoneTransform.cpp
@@ -8,7 +8,7 @@ FString FRigUnit_SetRelativeBoneTransform::GetUnitLabel() const
 	return FString::Printf(TEXT("Set Relative Transform %s"), *Bone.ToString());
 }
 
-UE_RigUnit_SetRelativeBoneTransform_IMPLEMENT_STATIC_VIRTUAL_METHOD(void, Execute, const FRigUnitContext& Context)
+UE_RigUnit_SetRelativeBoneTransform_IMPLEMENT_MULTIPLEX(void, Execute, const FRigUnitContext& Context)
 {
     DECLARE_SCOPE_HIERARCHICAL_COUNTER_RIGUNIT()
 	FRigBoneHierarchy* Hierarchy = ExecuteContext.GetBones();
diff --git a/Engine/Plugins/Experimental/ControlRig/Source/ControlRig/Private/Units/Hierarchy/RigUnit_SetRelativeBoneTransform.h b/Engine/Plugins/Experimental/ControlRig/Source/ControlRig/Private/Units/Hierarchy/RigUnit_SetRelativeBoneTransform.h
--- a/Engine/Plugins/Experimental/ControlRig/Source/ControlRig/Private/Units/Hierarchy/RigUnit_SetRelativeBoneTransform.h
+++ b/Engine/Plugins/Experimental/ControlRig/Source/ControlRig/Private/Units/Hierarchy/RigUnit_SetRelativeBoneTransform.h
@@ -21,7 +21,7 @@ struct FRigUnit_SetRelativeBoneTransform : public FRigUnitMutable
 	{}
 
 	virtual FString GetUnitLabel() const override;
-	STATIC_VIRTUAL_METHOD()
+	MULTIPLEX_METHOD()
 	virtual void Execute(const FRigUnitContext& Context) override;
 
 	/**
diff --git a/Engine/Plugins/Experimental/ControlRig/Source/ControlRig/Private/Units/Highlevel/Harmonics/RigUnit_BoneHarmonics.cpp b/Engine/Plugins/Experimental/ControlRig/Source/ControlRig/Private/Units/Highlevel/Harmonics/RigUnit_BoneHarmonics.cpp
--- a/Engine/Plugins/Experimental/ControlRig/Source/ControlRig/Private/Units/Highlevel/Harmonics/RigUnit_BoneHarmonics.cpp
+++ b/Engine/Plugins/Experimental/ControlRig/Source/ControlRig/Private/Units/Highlevel/Harmonics/RigUnit_BoneHarmonics.cpp
@@ -3,7 +3,7 @@
 #include "Units/Highlevel/Harmonics/RigUnit_BoneHarmonics.h"
 #include "Units/RigUnitContext.h"
 
-UE_RigUnit_BoneHarmonics_IMPLEMENT_STATIC_VIRTUAL_METHOD(void, Execute, const FRigUnitContext& Context)
+UE_RigUnit_BoneHarmonics_IMPLEMENT_MULTIPLEX(void, Execute, const FRigUnitContext& Context)
 {
     DECLARE_SCOPE_HIERARCHICAL_COUNTER_RIGUNIT()
 	FRigBoneHierarchy* Hierarchy = ExecuteContext.GetBones();
diff --git a/Engine/Plugins/Experimental/ControlRig/Source/ControlRig/Private/Units/Highlevel/Harmonics/RigUnit_BoneHarmonics.h b/Engine/Plugins/Experimental/ControlRig/Source/ControlRig/Private/Units/Highlevel/Harmonics/RigUnit_BoneHarmonics.h
--- a/Engine/Plugins/Experimental/ControlRig/Source/ControlRig/Private/Units/Highlevel/Harmonics/RigUnit_BoneHarmonics.h
+++ b/Engine/Plugins/Experimental/ControlRig/Source/ControlRig/Private/Units/Highlevel/Harmonics/RigUnit_BoneHarmonics.h
@@ -54,7 +54,7 @@ struct FRigUnit_BoneHarmonics : public FRigUnit_HighlevelBaseMutable
 		WaveTime = FVector::ZeroVector;
 	}
 
-	STATIC_VIRTUAL_METHOD()
+	MULTIPLEX_METHOD()
 	virtual void Execute(const FRigUnitContext& Context) override;
 
 	/** The bones to drive. */
diff --git a/Engine/Plugins/Experimental/ControlRig/Source/ControlRig/Private/Units/Highlevel/Harmonics/RigUnit_ChainHarmonics.cpp b/Engine/Plugins/Experimental/ControlRig/Source/ControlRig/Private/Units/Highlevel/Harmonics/RigUnit_ChainHarmonics.cpp
--- a/Engine/Plugins/Experimental/ControlRig/Source/ControlRig/Private/Units/Highlevel/Harmonics/RigUnit_ChainHarmonics.cpp
+++ b/Engine/Plugins/Experimental/ControlRig/Source/ControlRig/Private/Units/Highlevel/Harmonics/RigUnit_ChainHarmonics.cpp
@@ -3,7 +3,7 @@
 #include "Units/Highlevel/Harmonics/RigUnit_ChainHarmonics.h"
 #include "Units/RigUnitContext.h"
 
-UE_RigUnit_ChainHarmonics_IMPLEMENT_STATIC_VIRTUAL_METHOD(void, Execute, const FRigUnitContext& Context)
+UE_RigUnit_ChainHarmonics_IMPLEMENT_MULTIPLEX(void, Execute, const FRigUnitContext& Context)
 {
     DECLARE_SCOPE_HIERARCHICAL_COUNTER_RIGUNIT()
 	FRigBoneHierarchy* Hierarchy = ExecuteContext.GetBones();
diff --git a/Engine/Plugins/Experimental/ControlRig/Source/ControlRig/Private/Units/Highlevel/Harmonics/RigUnit_ChainHarmonics.h b/Engine/Plugins/Experimental/ControlRig/Source/ControlRig/Private/Units/Highlevel/Harmonics/RigUnit_ChainHarmonics.h
--- a/Engine/Plugins/Experimental/ControlRig/Source/ControlRig/Private/Units/Highlevel/Harmonics/RigUnit_ChainHarmonics.h
+++ b/Engine/Plugins/Experimental/ControlRig/Source/ControlRig/Private/Units/Highlevel/Harmonics/RigUnit_ChainHarmonics.h
@@ -161,7 +161,7 @@ struct FRigUnit_ChainHarmonics : public FRigUnit_HighlevelBaseMutable
 		Time = FVector::ZeroVector;
 	}
 
-	STATIC_VIRTUAL_METHOD()
+	MULTIPLEX_METHOD()
 	virtual void Execute(const FRigUnitContext& Context) override;
 
 	UPROPERTY(meta = (Input, Constant, BoneName))
diff --git a/Engine/Plugins/Experimental/ControlRig/Source/ControlRig/Private/Units/Highlevel/Hierarchy/RigUnit_AimBone.cpp b/Engine/Plugins/Experimental/ControlRig/Source/ControlRig/Private/Units/Highlevel/Hierarchy/RigUnit_AimBone.cpp
--- a/Engine/Plugins/Experimental/ControlRig/Source/ControlRig/Private/Units/Highlevel/Hierarchy/RigUnit_AimBone.cpp
+++ b/Engine/Plugins/Experimental/ControlRig/Source/ControlRig/Private/Units/Highlevel/Hierarchy/RigUnit_AimBone.cpp
@@ -3,7 +3,7 @@
 #include "Units/Highlevel/Hierarchy/RigUnit_AimBone.h"
 #include "Units/RigUnitContext.h"
 
-UE_RigUnit_AimBone_IMPLEMENT_STATIC_VIRTUAL_METHOD(void, Execute, const FRigUnitContext& Context)
+UE_RigUnit_AimBone_IMPLEMENT_MULTIPLEX(void, Execute, const FRigUnitContext& Context)
 {
     DECLARE_SCOPE_HIERARCHICAL_COUNTER_RIGUNIT()
 	FRigBoneHierarchy* Hierarchy = ExecuteContext.GetBones();
diff --git a/Engine/Plugins/Experimental/ControlRig/Source/ControlRig/Private/Units/Highlevel/Hierarchy/RigUnit_AimBone.h b/Engine/Plugins/Experimental/ControlRig/Source/ControlRig/Private/Units/Highlevel/Hierarchy/RigUnit_AimBone.h
--- a/Engine/Plugins/Experimental/ControlRig/Source/ControlRig/Private/Units/Highlevel/Hierarchy/RigUnit_AimBone.h
+++ b/Engine/Plugins/Experimental/ControlRig/Source/ControlRig/Private/Units/Highlevel/Hierarchy/RigUnit_AimBone.h
@@ -106,7 +106,7 @@ struct FRigUnit_AimBone : public FRigUnit_HighlevelBaseMutable
 		SecondaryCachedSpaceIndex = INDEX_NONE;
 	}
 
-	STATIC_VIRTUAL_METHOD()
+	MULTIPLEX_METHOD()
 	virtual void Execute(const FRigUnitContext& Context) override;
 
 	/** 
diff --git a/Engine/Plugins/Experimental/ControlRig/Source/ControlRig/Private/Units/Highlevel/Hierarchy/RigUnit_CCDIK.cpp b/Engine/Plugins/Experimental/ControlRig/Source/ControlRig/Private/Units/Highlevel/Hierarchy/RigUnit_CCDIK.cpp
--- a/Engine/Plugins/Experimental/ControlRig/Source/ControlRig/Private/Units/Highlevel/Hierarchy/RigUnit_CCDIK.cpp
+++ b/Engine/Plugins/Experimental/ControlRig/Source/ControlRig/Private/Units/Highlevel/Hierarchy/RigUnit_CCDIK.cpp
@@ -3,7 +3,7 @@
 #include "RigUnit_CCDIK.h"
 #include "Units/RigUnitContext.h"
 
-UE_RigUnit_CCDIK_IMPLEMENT_STATIC_VIRTUAL_METHOD(void, Execute, const FRigUnitContext& Context)
+UE_RigUnit_CCDIK_IMPLEMENT_MULTIPLEX(void, Execute, const FRigUnitContext& Context)
 {
     DECLARE_SCOPE_HIERARCHICAL_COUNTER_RIGUNIT()
 	FRigBoneHierarchy* Hierarchy = ExecuteContext.GetBones();
diff --git a/Engine/Plugins/Experimental/ControlRig/Source/ControlRig/Private/Units/Highlevel/Hierarchy/RigUnit_CCDIK.h b/Engine/Plugins/Experimental/ControlRig/Source/ControlRig/Private/Units/Highlevel/Hierarchy/RigUnit_CCDIK.h
--- a/Engine/Plugins/Experimental/ControlRig/Source/ControlRig/Private/Units/Highlevel/Hierarchy/RigUnit_CCDIK.h
+++ b/Engine/Plugins/Experimental/ControlRig/Source/ControlRig/Private/Units/Highlevel/Hierarchy/RigUnit_CCDIK.h
@@ -40,7 +40,7 @@ struct FRigUnit_CCDIK : public FRigUnit_HighlevelBaseMutable
 {
 	GENERATED_BODY()
 
-	STATIC_VIRTUAL_METHOD()
+	MULTIPLEX_METHOD()
 	virtual void Execute(const FRigUnitContext& Context) override;
 
 	FRigUnit_CCDIK()
diff --git a/Engine/Plugins/Experimental/ControlRig/Source/ControlRig/Private/Units/Highlevel/Hierarchy/RigUnit_DistributeRotation.cpp b/Engine/Plugins/Experimental/ControlRig/Source/ControlRig/Private/Units/Highlevel/Hierarchy/RigUnit_DistributeRotation.cpp
--- a/Engine/Plugins/Experimental/ControlRig/Source/ControlRig/Private/Units/Highlevel/Hierarchy/RigUnit_DistributeRotation.cpp
+++ b/Engine/Plugins/Experimental/ControlRig/Source/ControlRig/Private/Units/Highlevel/Hierarchy/RigUnit_DistributeRotation.cpp
@@ -3,7 +3,7 @@
 #include "Units/Highlevel/Hierarchy/RigUnit_DistributeRotation.h"
 #include "Units/RigUnitContext.h"
 
-UE_RigUnit_DistributeRotation_IMPLEMENT_STATIC_VIRTUAL_METHOD(void, Execute, const FRigUnitContext& Context)
+UE_RigUnit_DistributeRotation_IMPLEMENT_MULTIPLEX(void, Execute, const FRigUnitContext& Context)
 {
     DECLARE_SCOPE_HIERARCHICAL_COUNTER_RIGUNIT()
 	FRigBoneHierarchy* Hierarchy = ExecuteContext.GetBones();
diff --git a/Engine/Plugins/Experimental/ControlRig/Source/ControlRig/Private/Units/Highlevel/Hierarchy/RigUnit_DistributeRotation.h b/Engine/Plugins/Experimental/ControlRig/Source/ControlRig/Private/Units/Highlevel/Hierarchy/RigUnit_DistributeRotation.h
--- a/Engine/Plugins/Experimental/ControlRig/Source/ControlRig/Private/Units/Highlevel/Hierarchy/RigUnit_DistributeRotation.h
+++ b/Engine/Plugins/Experimental/ControlRig/Source/ControlRig/Private/Units/Highlevel/Hierarchy/RigUnit_DistributeRotation.h
@@ -47,7 +47,7 @@ struct FRigUnit_DistributeRotation : public FRigUnit_HighlevelBaseMutable
 		bPropagateToChildren = false;
 	}
 
-	STATIC_VIRTUAL_METHOD()
+	MULTIPLEX_METHOD()
 	virtual void Execute(const FRigUnitContext& Context) override;
 
 	/** 
diff --git a/Engine/Plugins/Experimental/ControlRig/Source/ControlRig/Private/Units/Highlevel/Hierarchy/RigUnit_FABRIK.cpp b/Engine/Plugins/Experimental/ControlRig/Source/ControlRig/Private/Units/Highlevel/Hierarchy/RigUnit_FABRIK.cpp
--- a/Engine/Plugins/Experimental/ControlRig/Source/ControlRig/Private/Units/Highlevel/Hierarchy/RigUnit_FABRIK.cpp
+++ b/Engine/Plugins/Experimental/ControlRig/Source/ControlRig/Private/Units/Highlevel/Hierarchy/RigUnit_FABRIK.cpp
@@ -3,7 +3,7 @@
 #include "RigUnit_FABRIK.h"
 #include "Units/RigUnitContext.h"
 
-UE_RigUnit_FABRIK_IMPLEMENT_STATIC_VIRTUAL_METHOD(void, Execute, const FRigUnitContext& Context)
+UE_RigUnit_FABRIK_IMPLEMENT_MULTIPLEX(void, Execute, const FRigUnitContext& Context)
 {
     DECLARE_SCOPE_HIERARCHICAL_COUNTER_RIGUNIT()
 	FRigBoneHierarchy* Hierarchy = ExecuteContext.GetBones();
diff --git a/Engine/Plugins/Experimental/ControlRig/Source/ControlRig/Private/Units/Highlevel/Hierarchy/RigUnit_FABRIK.h b/Engine/Plugins/Experimental/ControlRig/Source/ControlRig/Private/Units/Highlevel/Hierarchy/RigUnit_FABRIK.h
--- a/Engine/Plugins/Experimental/ControlRig/Source/ControlRig/Private/Units/Highlevel/Hierarchy/RigUnit_FABRIK.h
+++ b/Engine/Plugins/Experimental/ControlRig/Source/ControlRig/Private/Units/Highlevel/Hierarchy/RigUnit_FABRIK.h
@@ -16,7 +16,7 @@ struct FRigUnit_FABRIK : public FRigUnit_HighlevelBaseMutable
 {
 	GENERATED_BODY()
 
-	STATIC_VIRTUAL_METHOD()
+	MULTIPLEX_METHOD()
 	virtual void Execute(const FRigUnitContext& Context) override;
 
 	FRigUnit_FABRIK()
diff --git a/Engine/Plugins/Experimental/ControlRig/Source/ControlRig/Private/Units/Highlevel/Hierarchy/RigUnit_FitChainToCurve.cpp b/Engine/Plugins/Experimental/ControlRig/Source/ControlRig/Private/Units/Highlevel/Hierarchy/RigUnit_FitChainToCurve.cpp
--- a/Engine/Plugins/Experimental/ControlRig/Source/ControlRig/Private/Units/Highlevel/Hierarchy/RigUnit_FitChainToCurve.cpp
+++ b/Engine/Plugins/Experimental/ControlRig/Source/ControlRig/Private/Units/Highlevel/Hierarchy/RigUnit_FitChainToCurve.cpp
@@ -3,7 +3,7 @@
 #include "Units/Highlevel/Hierarchy/RigUnit_FitChainToCurve.h"
 #include "Units/RigUnitContext.h"
 
-UE_RigUnit_FitChainToCurve_IMPLEMENT_STATIC_VIRTUAL_METHOD(void, Execute, const FRigUnitContext& Context)
+UE_RigUnit_FitChainToCurve_IMPLEMENT_MULTIPLEX(void, Execute, const FRigUnitContext& Context)
 {
     DECLARE_SCOPE_HIERARCHICAL_COUNTER_RIGUNIT()
 	FRigBoneHierarchy* Hierarchy = ExecuteContext.GetBones();
diff --git a/Engine/Plugins/Experimental/ControlRig/Source/ControlRig/Private/Units/Highlevel/Hierarchy/RigUnit_FitChainToCurve.h b/Engine/Plugins/Experimental/ControlRig/Source/ControlRig/Private/Units/Highlevel/Hierarchy/RigUnit_FitChainToCurve.h
--- a/Engine/Plugins/Experimental/ControlRig/Source/ControlRig/Private/Units/Highlevel/Hierarchy/RigUnit_FitChainToCurve.h
+++ b/Engine/Plugins/Experimental/ControlRig/Source/ControlRig/Private/Units/Highlevel/Hierarchy/RigUnit_FitChainToCurve.h
@@ -108,7 +108,7 @@ struct FRigUnit_FitChainToCurve : public FRigUnit_HighlevelBaseMutable
 		ChainLength = 0.f;
 	}
 
-	STATIC_VIRTUAL_METHOD()
+	MULTIPLEX_METHOD()
 	virtual void Execute(const FRigUnitContext& Context) override;
 
 	/** 
diff --git a/Engine/Plugins/Experimental/ControlRig/Source/ControlRig/Private/Units/Highlevel/Hierarchy/RigUnit_ModifyBoneTransforms.cpp b/Engine/Plugins/Experimental/ControlRig/Source/ControlRig/Private/Units/Highlevel/Hierarchy/RigUnit_ModifyBoneTransforms.cpp
--- a/Engine/Plugins/Experimental/ControlRig/Source/ControlRig/Private/Units/Highlevel/Hierarchy/RigUnit_ModifyBoneTransforms.cpp
+++ b/Engine/Plugins/Experimental/ControlRig/Source/ControlRig/Private/Units/Highlevel/Hierarchy/RigUnit_ModifyBoneTransforms.cpp
@@ -3,7 +3,7 @@
 #include "RigUnit_ModifyBoneTransforms.h"
 #include "Units/RigUnitContext.h"
 
-UE_RigUnit_ModifyBoneTransforms_IMPLEMENT_STATIC_VIRTUAL_METHOD(void, Execute, const FRigUnitContext& Context)
+UE_RigUnit_ModifyBoneTransforms_IMPLEMENT_MULTIPLEX(void, Execute, const FRigUnitContext& Context)
 {
     DECLARE_SCOPE_HIERARCHICAL_COUNTER_RIGUNIT()
 	FRigBoneHierarchy* Hierarchy = ExecuteContext.GetBones();
diff --git a/Engine/Plugins/Experimental/ControlRig/Source/ControlRig/Private/Units/Highlevel/Hierarchy/RigUnit_ModifyBoneTransforms.h b/Engine/Plugins/Experimental/ControlRig/Source/ControlRig/Private/Units/Highlevel/Hierarchy/RigUnit_ModifyBoneTransforms.h
--- a/Engine/Plugins/Experimental/ControlRig/Source/ControlRig/Private/Units/Highlevel/Hierarchy/RigUnit_ModifyBoneTransforms.h
+++ b/Engine/Plugins/Experimental/ControlRig/Source/ControlRig/Private/Units/Highlevel/Hierarchy/RigUnit_ModifyBoneTransforms.h
@@ -70,7 +70,7 @@ struct FRigUnit_ModifyBoneTransforms : public FRigUnit_HighlevelBaseMutable
 		, Mode(EControlRigModifyBoneMode::AdditiveLocal)
 	{}
 
-	STATIC_VIRTUAL_METHOD()
+	MULTIPLEX_METHOD()
 	virtual void Execute(const FRigUnitContext& Context) override;
 
 	/**
diff --git a/Engine/Plugins/Experimental/ControlRig/Source/ControlRig/Private/Units/Highlevel/Hierarchy/RigUnit_SlideChain.cpp b/Engine/Plugins/Experimental/ControlRig/Source/ControlRig/Private/Units/Highlevel/Hierarchy/RigUnit_SlideChain.cpp
--- a/Engine/Plugins/Experimental/ControlRig/Source/ControlRig/Private/Units/Highlevel/Hierarchy/RigUnit_SlideChain.cpp
+++ b/Engine/Plugins/Experimental/ControlRig/Source/ControlRig/Private/Units/Highlevel/Hierarchy/RigUnit_SlideChain.cpp
@@ -4,7 +4,7 @@
 #include "Math/ControlRigMathLibrary.h"
 #include "Units/RigUnitContext.h"
 
-UE_RigUnit_SlideChain_IMPLEMENT_STATIC_VIRTUAL_METHOD(void, Execute, const FRigUnitContext& Context)
+UE_RigUnit_SlideChain_IMPLEMENT_MULTIPLEX(void, Execute, const FRigUnitContext& Context)
 {
 
 	FRigBoneHierarchy* Hierarchy = ExecuteContext.GetBones();
diff --git a/Engine/Plugins/Experimental/ControlRig/Source/ControlRig/Private/Units/Highlevel/Hierarchy/RigUnit_SlideChain.h b/Engine/Plugins/Experimental/ControlRig/Source/ControlRig/Private/Units/Highlevel/Hierarchy/RigUnit_SlideChain.h
--- a/Engine/Plugins/Experimental/ControlRig/Source/ControlRig/Private/Units/Highlevel/Hierarchy/RigUnit_SlideChain.h
+++ b/Engine/Plugins/Experimental/ControlRig/Source/ControlRig/Private/Units/Highlevel/Hierarchy/RigUnit_SlideChain.h
@@ -22,7 +22,7 @@ struct FRigUnit_SlideChain: public FRigUnit_HighlevelBaseMutable
 		ChainLength = 0.f;
 	}
 
-	STATIC_VIRTUAL_METHOD()
+	MULTIPLEX_METHOD()
 	virtual void Execute(const FRigUnitContext& Context) override;
 
 	/** 
diff --git a/Engine/Plugins/Experimental/ControlRig/Source/ControlRig/Private/Units/Highlevel/Hierarchy/RigUnit_SpringIK.cpp b/Engine/Plugins/Experimental/ControlRig/Source/ControlRig/Private/Units/Highlevel/Hierarchy/RigUnit_SpringIK.cpp
--- a/Engine/Plugins/Experimental/ControlRig/Source/ControlRig/Private/Units/Highlevel/Hierarchy/RigUnit_SpringIK.cpp
+++ b/Engine/Plugins/Experimental/ControlRig/Source/ControlRig/Private/Units/Highlevel/Hierarchy/RigUnit_SpringIK.cpp
@@ -3,7 +3,7 @@
 #include "Units/Highlevel/Hierarchy/RigUnit_SpringIK.h"
 #include "Units/RigUnitContext.h"
 
-UE_RigUnit_SpringIK_IMPLEMENT_STATIC_VIRTUAL_METHOD(void, Execute, const FRigUnitContext& Context)
+UE_RigUnit_SpringIK_IMPLEMENT_MULTIPLEX(void, Execute, const FRigUnitContext& Context)
 {
     DECLARE_SCOPE_HIERARCHICAL_COUNTER_RIGUNIT()
 	FRigBoneHierarchy* Hierarchy = ExecuteContext.GetBones();
diff --git a/Engine/Plugins/Experimental/ControlRig/Source/ControlRig/Private/Units/Highlevel/Hierarchy/RigUnit_SpringIK.h b/Engine/Plugins/Experimental/ControlRig/Source/ControlRig/Private/Units/Highlevel/Hierarchy/RigUnit_SpringIK.h
--- a/Engine/Plugins/Experimental/ControlRig/Source/ControlRig/Private/Units/Highlevel/Hierarchy/RigUnit_SpringIK.h
+++ b/Engine/Plugins/Experimental/ControlRig/Source/ControlRig/Private/Units/Highlevel/Hierarchy/RigUnit_SpringIK.h
@@ -74,7 +74,7 @@ struct FRigUnit_SpringIK : public FRigUnit_HighlevelBaseMutable
 		PoleVectorIndex = INDEX_NONE;
 	}
 
-	STATIC_VIRTUAL_METHOD()
+	MULTIPLEX_METHOD()
 	virtual void Execute(const FRigUnitContext& Context) override;
 
 	/** 
diff --git a/Engine/Plugins/Experimental/ControlRig/Source/ControlRig/Private/Units/Highlevel/Hierarchy/RigUnit_TransformConstraint.cpp b/Engine/Plugins/Experimental/ControlRig/Source/ControlRig/Private/Units/Highlevel/Hierarchy/RigUnit_TransformConstraint.cpp
--- a/Engine/Plugins/Experimental/ControlRig/Source/ControlRig/Private/Units/Highlevel/Hierarchy/RigUnit_TransformConstraint.cpp
+++ b/Engine/Plugins/Experimental/ControlRig/Source/ControlRig/Private/Units/Highlevel/Hierarchy/RigUnit_TransformConstraint.cpp
@@ -5,7 +5,7 @@
 #include "HelperUtil.h"
 #include "AnimationCoreLibrary.h"
 
-UE_RigUnit_TransformConstraint_IMPLEMENT_STATIC_VIRTUAL_METHOD(void, Execute, const FRigUnitContext& Context)
+UE_RigUnit_TransformConstraint_IMPLEMENT_MULTIPLEX(void, Execute, const FRigUnitContext& Context)
 {
     DECLARE_SCOPE_HIERARCHICAL_COUNTER_RIGUNIT()
 
diff --git a/Engine/Plugins/Experimental/ControlRig/Source/ControlRig/Private/Units/Highlevel/Hierarchy/RigUnit_TransformConstraint.h b/Engine/Plugins/Experimental/ControlRig/Source/ControlRig/Private/Units/Highlevel/Hierarchy/RigUnit_TransformConstraint.h
--- a/Engine/Plugins/Experimental/ControlRig/Source/ControlRig/Private/Units/Highlevel/Hierarchy/RigUnit_TransformConstraint.h
+++ b/Engine/Plugins/Experimental/ControlRig/Source/ControlRig/Private/Units/Highlevel/Hierarchy/RigUnit_TransformConstraint.h
@@ -45,7 +45,7 @@ struct FRigUnit_TransformConstraint : public FRigUnit_HighlevelBaseMutable
 		: BaseTransformSpace(ETransformSpaceMode::GlobalSpace)
 	{}
 
-	STATIC_VIRTUAL_METHOD()
+	MULTIPLEX_METHOD()
 	virtual void Execute(const FRigUnitContext& Context) override;
 
 	UPROPERTY(EditAnywhere, Category = "Constraint", meta = (Input, Constant, BoneName))
diff --git a/Engine/Plugins/Experimental/ControlRig/Source/ControlRig/Private/Units/Highlevel/Hierarchy/RigUnit_TwoBoneIKSimple.cpp b/Engine/Plugins/Experimental/ControlRig/Source/ControlRig/Private/Units/Highlevel/Hierarchy/RigUnit_TwoBoneIKSimple.cpp
--- a/Engine/Plugins/Experimental/ControlRig/Source/ControlRig/Private/Units/Highlevel/Hierarchy/RigUnit_TwoBoneIKSimple.cpp
+++ b/Engine/Plugins/Experimental/ControlRig/Source/ControlRig/Private/Units/Highlevel/Hierarchy/RigUnit_TwoBoneIKSimple.cpp
@@ -4,7 +4,7 @@
 #include "Units/RigUnitContext.h"
 #include "Math/ControlRigMathLibrary.h"
 
-UE_RigUnit_TwoBoneIKSimple_IMPLEMENT_STATIC_VIRTUAL_METHOD(void, Execute, const FRigUnitContext& Context)
+UE_RigUnit_TwoBoneIKSimple_IMPLEMENT_MULTIPLEX(void, Execute, const FRigUnitContext& Context)
 {
     DECLARE_SCOPE_HIERARCHICAL_COUNTER_RIGUNIT()
 	FRigBoneHierarchy* Hierarchy = ExecuteContext.GetBones();
diff --git a/Engine/Plugins/Experimental/ControlRig/Source/ControlRig/Private/Units/Highlevel/Hierarchy/RigUnit_TwoBoneIKSimple.h b/Engine/Plugins/Experimental/ControlRig/Source/ControlRig/Private/Units/Highlevel/Hierarchy/RigUnit_TwoBoneIKSimple.h
--- a/Engine/Plugins/Experimental/ControlRig/Source/ControlRig/Private/Units/Highlevel/Hierarchy/RigUnit_TwoBoneIKSimple.h
+++ b/Engine/Plugins/Experimental/ControlRig/Source/ControlRig/Private/Units/Highlevel/Hierarchy/RigUnit_TwoBoneIKSimple.h
@@ -67,7 +67,7 @@ struct FRigUnit_TwoBoneIKSimple : public FRigUnit_HighlevelBaseMutable
 		PoleVectorSpaceIndex = INDEX_NONE;
 	}
 
-	STATIC_VIRTUAL_METHOD()
+	MULTIPLEX_METHOD()
 	virtual void Execute(const FRigUnitContext& Context) override;
 
 	/** 
diff --git a/Engine/Plugins/Experimental/ControlRig/Source/ControlRig/Private/Units/Math/RigUnit_MathBool.cpp b/Engine/Plugins/Experimental/ControlRig/Source/ControlRig/Private/Units/Math/RigUnit_MathBool.cpp
--- a/Engine/Plugins/Experimental/ControlRig/Source/ControlRig/Private/Units/Math/RigUnit_MathBool.cpp
+++ b/Engine/Plugins/Experimental/ControlRig/Source/ControlRig/Private/Units/Math/RigUnit_MathBool.cpp
@@ -3,37 +3,37 @@
 #include "Units/Math/RigUnit_MathBool.h"
 #include "Units/RigUnitContext.h"
 
-UE_RigUnit_MathBoolNot_IMPLEMENT_STATIC_VIRTUAL_METHOD(void, Execute, const FRigUnitContext& Context)
+UE_RigUnit_MathBoolNot_IMPLEMENT_MULTIPLEX(void, Execute, const FRigUnitContext& Context)
 {
     DECLARE_SCOPE_HIERARCHICAL_COUNTER_RIGUNIT()
 	Result = !Value;
 }
 
-UE_RigUnit_MathBoolAnd_IMPLEMENT_STATIC_VIRTUAL_METHOD(void, Execute, const FRigUnitContext& Context)
+UE_RigUnit_MathBoolAnd_IMPLEMENT_MULTIPLEX(void, Execute, const FRigUnitContext& Context)
 {
     DECLARE_SCOPE_HIERARCHICAL_COUNTER_RIGUNIT()
 	Result = A && B;
 }
 
-UE_RigUnit_MathBoolNand_IMPLEMENT_STATIC_VIRTUAL_METHOD(void, Execute, const FRigUnitContext& Context)
+UE_RigUnit_MathBoolNand_IMPLEMENT_MULTIPLEX(void, Execute, const FRigUnitContext& Context)
 {
     DECLARE_SCOPE_HIERARCHICAL_COUNTER_RIGUNIT()
 	Result = (!A) && (!B);
 }
 
-UE_RigUnit_MathBoolOr_IMPLEMENT_STATIC_VIRTUAL_METHOD(void, Execute, const FRigUnitContext& Context)
+UE_RigUnit_MathBoolOr_IMPLEMENT_MULTIPLEX(void, Execute, const FRigUnitContext& Context)
 {
     DECLARE_SCOPE_HIERARCHICAL_COUNTER_RIGUNIT()
 	Result = A || B;
 }
 
-UE_RigUnit_MathBoolEquals_IMPLEMENT_STATIC_VIRTUAL_METHOD(void, Execute, const FRigUnitContext& Context)
+UE_RigUnit_MathBoolEquals_IMPLEMENT_MULTIPLEX(void, Execute, const FRigUnitContext& Context)
 {
     DECLARE_SCOPE_HIERARCHICAL_COUNTER_RIGUNIT()
 	Result = A == B;
 }
 
-UE_RigUnit_MathBoolNotEquals_IMPLEMENT_STATIC_VIRTUAL_METHOD(void, Execute, const FRigUnitContext& Context)
+UE_RigUnit_MathBoolNotEquals_IMPLEMENT_MULTIPLEX(void, Execute, const FRigUnitContext& Context)
 {
     DECLARE_SCOPE_HIERARCHICAL_COUNTER_RIGUNIT()
 	Result = A != B;
diff --git a/Engine/Plugins/Experimental/ControlRig/Source/ControlRig/Private/Units/Math/RigUnit_MathBool.h b/Engine/Plugins/Experimental/ControlRig/Source/ControlRig/Private/Units/Math/RigUnit_MathBool.h
--- a/Engine/Plugins/Experimental/ControlRig/Source/ControlRig/Private/Units/Math/RigUnit_MathBool.h
+++ b/Engine/Plugins/Experimental/ControlRig/Source/ControlRig/Private/Units/Math/RigUnit_MathBool.h
@@ -98,7 +98,7 @@ struct FRigUnit_MathBoolNot : public FRigUnit_MathBoolUnaryOp
 {
 	GENERATED_BODY()
 
-	STATIC_VIRTUAL_METHOD()
+	MULTIPLEX_METHOD()
 	virtual void Execute(const FRigUnitContext& Context) override;
 };
 
@@ -110,7 +110,7 @@ struct FRigUnit_MathBoolAnd : public FRigUnit_MathBoolBinaryOp
 {
 	GENERATED_BODY()
 
-	STATIC_VIRTUAL_METHOD()
+	MULTIPLEX_METHOD()
 	virtual void Execute(const FRigUnitContext& Context) override;
 };
 
@@ -122,7 +122,7 @@ struct FRigUnit_MathBoolNand : public FRigUnit_MathBoolBinaryOp
 {
 	GENERATED_BODY()
 
-	STATIC_VIRTUAL_METHOD()
+	MULTIPLEX_METHOD()
 	virtual void Execute(const FRigUnitContext& Context) override;
 };
 
@@ -134,7 +134,7 @@ struct FRigUnit_MathBoolOr : public FRigUnit_MathBoolBinaryOp
 {
 	GENERATED_BODY()
 
-	STATIC_VIRTUAL_METHOD()
+	MULTIPLEX_METHOD()
 	virtual void Execute(const FRigUnitContext& Context) override;
 };
 
@@ -145,7 +145,7 @@ USTRUCT(meta=(DisplayName="Equals", PrototypeName="Equals", Keywords="Same,=="))
 struct FRigUnit_MathBoolEquals : public FRigUnit_MathBoolBase
 {
 	GENERATED_BODY()
-	STATIC_VIRTUAL_METHOD()
+	MULTIPLEX_METHOD()
 	virtual void Execute(const FRigUnitContext& Context) override;
 
 	FRigUnit_MathBoolEquals()
@@ -171,7 +171,7 @@ USTRUCT(meta=(DisplayName="Not Equals", PrototypeName="NotEquals", Keywords="Dif
 struct FRigUnit_MathBoolNotEquals : public FRigUnit_MathBoolBase
 {
 	GENERATED_BODY()
-	STATIC_VIRTUAL_METHOD()
+	MULTIPLEX_METHOD()
 	virtual void Execute(const FRigUnitContext& Context) override;
 
 	FRigUnit_MathBoolNotEquals()
diff --git a/Engine/Plugins/Experimental/ControlRig/Source/ControlRig/Private/Units/Math/RigUnit_MathFloat.cpp b/Engine/Plugins/Experimental/ControlRig/Source/ControlRig/Private/Units/Math/RigUnit_MathFloat.cpp
--- a/Engine/Plugins/Experimental/ControlRig/Source/ControlRig/Private/Units/Math/RigUnit_MathFloat.cpp
+++ b/Engine/Plugins/Experimental/ControlRig/Source/ControlRig/Private/Units/Math/RigUnit_MathFloat.cpp
@@ -3,25 +3,25 @@
 #include "Units/Math/RigUnit_MathFloat.h"
 #include "Units/RigUnitContext.h"
 
-UE_RigUnit_MathFloatAdd_IMPLEMENT_STATIC_VIRTUAL_METHOD(void, Execute, const FRigUnitContext& Context)
+UE_RigUnit_MathFloatAdd_IMPLEMENT_MULTIPLEX(void, Execute, const FRigUnitContext& Context)
 {
     DECLARE_SCOPE_HIERARCHICAL_COUNTER_RIGUNIT()
 	Result = A + B;
 }
 
-UE_RigUnit_MathFloatSub_IMPLEMENT_STATIC_VIRTUAL_METHOD(void, Execute, const FRigUnitContext& Context)
+UE_RigUnit_MathFloatSub_IMPLEMENT_MULTIPLEX(void, Execute, const FRigUnitContext& Context)
 {
     DECLARE_SCOPE_HIERARCHICAL_COUNTER_RIGUNIT()
 	Result = A - B;
 }
 
-UE_RigUnit_MathFloatMul_IMPLEMENT_STATIC_VIRTUAL_METHOD(void, Execute, const FRigUnitContext& Context)
+UE_RigUnit_MathFloatMul_IMPLEMENT_MULTIPLEX(void, Execute, const FRigUnitContext& Context)
 {
     DECLARE_SCOPE_HIERARCHICAL_COUNTER_RIGUNIT()
 	Result = A * B;
 }
 
-UE_RigUnit_MathFloatDiv_IMPLEMENT_STATIC_VIRTUAL_METHOD(void, Execute, const FRigUnitContext& Context)
+UE_RigUnit_MathFloatDiv_IMPLEMENT_MULTIPLEX(void, Execute, const FRigUnitContext& Context)
 {
     DECLARE_SCOPE_HIERARCHICAL_COUNTER_RIGUNIT()
 	if(FMath::IsNearlyZero(B))
@@ -33,7 +33,7 @@ UE_RigUnit_MathFloatDiv_IMPLEMENT_STATIC_VIRTUAL_METHOD(void, Execute, const FRi
 	Result = A / B;
 }
 
-UE_RigUnit_MathFloatMod_IMPLEMENT_STATIC_VIRTUAL_METHOD(void, Execute, const FRigUnitContext& Context)
+UE_RigUnit_MathFloatMod_IMPLEMENT_MULTIPLEX(void, Execute, const FRigUnitContext& Context)
 {
     DECLARE_SCOPE_HIERARCHICAL_COUNTER_RIGUNIT()
 	if(FMath::IsNearlyZero(B) || B < 0.f)
@@ -45,25 +45,25 @@ UE_RigUnit_MathFloatMod_IMPLEMENT_STATIC_VIRTUAL_METHOD(void, Execute, const FRi
 	Result = FMath::Fmod(A, B);
 }
 
-UE_RigUnit_MathFloatMin_IMPLEMENT_STATIC_VIRTUAL_METHOD(void, Execute, const FRigUnitContext& Context)
+UE_RigUnit_MathFloatMin_IMPLEMENT_MULTIPLEX(void, Execute, const FRigUnitContext& Context)
 {
     DECLARE_SCOPE_HIERARCHICAL_COUNTER_RIGUNIT()
 	Result = FMath::Min<float>(A, B);
 }
 
-UE_RigUnit_MathFloatMax_IMPLEMENT_STATIC_VIRTUAL_METHOD(void, Execute, const FRigUnitContext& Context)
+UE_RigUnit_MathFloatMax_IMPLEMENT_MULTIPLEX(void, Execute, const FRigUnitContext& Context)
 {
     DECLARE_SCOPE_HIERARCHICAL_COUNTER_RIGUNIT()
 	Result = FMath::Max<float>(A, B);
 }
 
-UE_RigUnit_MathFloatPow_IMPLEMENT_STATIC_VIRTUAL_METHOD(void, Execute, const FRigUnitContext& Context)
+UE_RigUnit_MathFloatPow_IMPLEMENT_MULTIPLEX(void, Execute, const FRigUnitContext& Context)
 {
     DECLARE_SCOPE_HIERARCHICAL_COUNTER_RIGUNIT()
 	Result = FMath::Pow(A, B);
 }
 
-UE_RigUnit_MathFloatSqrt_IMPLEMENT_STATIC_VIRTUAL_METHOD(void, Execute, const FRigUnitContext& Context)
+UE_RigUnit_MathFloatSqrt_IMPLEMENT_MULTIPLEX(void, Execute, const FRigUnitContext& Context)
 {
     DECLARE_SCOPE_HIERARCHICAL_COUNTER_RIGUNIT()
 	if(Value < 0.f)
@@ -75,55 +75,55 @@ UE_RigUnit_MathFloatSqrt_IMPLEMENT_STATIC_VIRTUAL_METHOD(void, Execute, const FR
 	Result = FMath::Sqrt(Value);
 }
 
-UE_RigUnit_MathFloatNegate_IMPLEMENT_STATIC_VIRTUAL_METHOD(void, Execute, const FRigUnitContext& Context)
+UE_RigUnit_MathFloatNegate_IMPLEMENT_MULTIPLEX(void, Execute, const FRigUnitContext& Context)
 {
     DECLARE_SCOPE_HIERARCHICAL_COUNTER_RIGUNIT()
 	Result = -Value;
 }
 
-UE_RigUnit_MathFloatAbs_IMPLEMENT_STATIC_VIRTUAL_METHOD(void, Execute, const FRigUnitContext& Context)
+UE_RigUnit_MathFloatAbs_IMPLEMENT_MULTIPLEX(void, Execute, const FRigUnitContext& Context)
 {
     DECLARE_SCOPE_HIERARCHICAL_COUNTER_RIGUNIT()
 	Result = FMath::Abs(Value);
 }
 
-UE_RigUnit_MathFloatFloor_IMPLEMENT_STATIC_VIRTUAL_METHOD(void, Execute, const FRigUnitContext& Context)
+UE_RigUnit_MathFloatFloor_IMPLEMENT_MULTIPLEX(void, Execute, const FRigUnitContext& Context)
 {
     DECLARE_SCOPE_HIERARCHICAL_COUNTER_RIGUNIT()
 	Result = FMath::FloorToFloat(Value);
 }
 
-UE_RigUnit_MathFloatCeil_IMPLEMENT_STATIC_VIRTUAL_METHOD(void, Execute, const FRigUnitContext& Context)
+UE_RigUnit_MathFloatCeil_IMPLEMENT_MULTIPLEX(void, Execute, const FRigUnitContext& Context)
 {
     DECLARE_SCOPE_HIERARCHICAL_COUNTER_RIGUNIT()
 	Result = FMath::CeilToFloat(Value);
 }
 
-UE_RigUnit_MathFloatRound_IMPLEMENT_STATIC_VIRTUAL_METHOD(void, Execute, const FRigUnitContext& Context)
+UE_RigUnit_MathFloatRound_IMPLEMENT_MULTIPLEX(void, Execute, const FRigUnitContext& Context)
 {
     DECLARE_SCOPE_HIERARCHICAL_COUNTER_RIGUNIT()
 	Result = FMath::RoundToFloat(Value);
 }
 
-UE_RigUnit_MathFloatSign_IMPLEMENT_STATIC_VIRTUAL_METHOD(void, Execute, const FRigUnitContext& Context)
+UE_RigUnit_MathFloatSign_IMPLEMENT_MULTIPLEX(void, Execute, const FRigUnitContext& Context)
 {
     DECLARE_SCOPE_HIERARCHICAL_COUNTER_RIGUNIT()
 	Result = Value < 0.f ? -1.f : 1.f;
 }
 
-UE_RigUnit_MathFloatClamp_IMPLEMENT_STATIC_VIRTUAL_METHOD(void, Execute, const FRigUnitContext& Context)
+UE_RigUnit_MathFloatClamp_IMPLEMENT_MULTIPLEX(void, Execute, const FRigUnitContext& Context)
 {
     DECLARE_SCOPE_HIERARCHICAL_COUNTER_RIGUNIT()
 	Result = FMath::Clamp<float>(Value, Minimum, Maximum);
 }
 
-UE_RigUnit_MathFloatLerp_IMPLEMENT_STATIC_VIRTUAL_METHOD(void, Execute, const FRigUnitContext& Context)
+UE_RigUnit_MathFloatLerp_IMPLEMENT_MULTIPLEX(void, Execute, const FRigUnitContext& Context)
 {
     DECLARE_SCOPE_HIERARCHICAL_COUNTER_RIGUNIT()
 	Result = FMath::Lerp<float>(A, B, T);
 }
 
-UE_RigUnit_MathFloatRemap_IMPLEMENT_STATIC_VIRTUAL_METHOD(void, Execute, const FRigUnitContext& Context)
+UE_RigUnit_MathFloatRemap_IMPLEMENT_MULTIPLEX(void, Execute, const FRigUnitContext& Context)
 {
     DECLARE_SCOPE_HIERARCHICAL_COUNTER_RIGUNIT()
 	float Ratio = 0.f;
@@ -142,43 +142,43 @@ UE_RigUnit_MathFloatRemap_IMPLEMENT_STATIC_VIRTUAL_METHOD(void, Execute, const F
 	Result = FMath::Lerp<float>(TargetMinimum, TargetMaximum, Ratio);
 }
 
-UE_RigUnit_MathFloatEquals_IMPLEMENT_STATIC_VIRTUAL_METHOD(void, Execute, const FRigUnitContext& Context)
+UE_RigUnit_MathFloatEquals_IMPLEMENT_MULTIPLEX(void, Execute, const FRigUnitContext& Context)
 {
     DECLARE_SCOPE_HIERARCHICAL_COUNTER_RIGUNIT()
 	Result = A == B;
 }
 
-UE_RigUnit_MathFloatNotEquals_IMPLEMENT_STATIC_VIRTUAL_METHOD(void, Execute, const FRigUnitContext& Context)
+UE_RigUnit_MathFloatNotEquals_IMPLEMENT_MULTIPLEX(void, Execute, const FRigUnitContext& Context)
 {
     DECLARE_SCOPE_HIERARCHICAL_COUNTER_RIGUNIT()
 	Result = A != B;
 }
 
-UE_RigUnit_MathFloatGreater_IMPLEMENT_STATIC_VIRTUAL_METHOD(void, Execute, const FRigUnitContext& Context)
+UE_RigUnit_MathFloatGreater_IMPLEMENT_MULTIPLEX(void, Execute, const FRigUnitContext& Context)
 {
     DECLARE_SCOPE_HIERARCHICAL_COUNTER_RIGUNIT()
 	Result = A > B;
 }
 
-UE_RigUnit_MathFloatLess_IMPLEMENT_STATIC_VIRTUAL_METHOD(void, Execute, const FRigUnitContext& Context)
+UE_RigUnit_MathFloatLess_IMPLEMENT_MULTIPLEX(void, Execute, const FRigUnitContext& Context)
 {
     DECLARE_SCOPE_HIERARCHICAL_COUNTER_RIGUNIT()
 	Result = A < B;
 }
 
-UE_RigUnit_MathFloatGreaterEqual_IMPLEMENT_STATIC_VIRTUAL_METHOD(void, Execute, const FRigUnitContext& Context)
+UE_RigUnit_MathFloatGreaterEqual_IMPLEMENT_MULTIPLEX(void, Execute, const FRigUnitContext& Context)
 {
     DECLARE_SCOPE_HIERARCHICAL_COUNTER_RIGUNIT()
 	Result = A >= B;
 }
 
-UE_RigUnit_MathFloatLessEqual_IMPLEMENT_STATIC_VIRTUAL_METHOD(void, Execute, const FRigUnitContext& Context)
+UE_RigUnit_MathFloatLessEqual_IMPLEMENT_MULTIPLEX(void, Execute, const FRigUnitContext& Context)
 {
     DECLARE_SCOPE_HIERARCHICAL_COUNTER_RIGUNIT()
 	Result = A <= B;
 }
 
-UE_RigUnit_MathFloatIsNearlyZero_IMPLEMENT_STATIC_VIRTUAL_METHOD(void, Execute, const FRigUnitContext& Context)
+UE_RigUnit_MathFloatIsNearlyZero_IMPLEMENT_MULTIPLEX(void, Execute, const FRigUnitContext& Context)
 {
     DECLARE_SCOPE_HIERARCHICAL_COUNTER_RIGUNIT()
 	if(Tolerance < 0.f)
@@ -188,7 +188,7 @@ UE_RigUnit_MathFloatIsNearlyZero_IMPLEMENT_STATIC_VIRTUAL_METHOD(void, Execute,
 	Result = FMath::IsNearlyZero(Value, FMath::Max<float>(Tolerance, SMALL_NUMBER));
 }
 
-UE_RigUnit_MathFloatIsNearlyEqual_IMPLEMENT_STATIC_VIRTUAL_METHOD(void, Execute, const FRigUnitContext& Context)
+UE_RigUnit_MathFloatIsNearlyEqual_IMPLEMENT_MULTIPLEX(void, Execute, const FRigUnitContext& Context)
 {
     DECLARE_SCOPE_HIERARCHICAL_COUNTER_RIGUNIT()
 	if(Tolerance < 0.f)
@@ -198,43 +198,43 @@ UE_RigUnit_MathFloatIsNearlyEqual_IMPLEMENT_STATIC_VIRTUAL_METHOD(void, Execute,
 	Result = FMath::IsNearlyEqual(A, B, FMath::Max<float>(Tolerance, SMALL_NUMBER));
 }
 
-UE_RigUnit_MathFloatSelectBool_IMPLEMENT_STATIC_VIRTUAL_METHOD(void, Execute, const FRigUnitContext& Context)
+UE_RigUnit_MathFloatSelectBool_IMPLEMENT_MULTIPLEX(void, Execute, const FRigUnitContext& Context)
 {
     DECLARE_SCOPE_HIERARCHICAL_COUNTER_RIGUNIT()
 	Result = Condition ? IfTrue : IfFalse;
 }
 
-UE_RigUnit_MathFloatDeg_IMPLEMENT_STATIC_VIRTUAL_METHOD(void, Execute, const FRigUnitContext& Context)
+UE_RigUnit_MathFloatDeg_IMPLEMENT_MULTIPLEX(void, Execute, const FRigUnitContext& Context)
 {
     DECLARE_SCOPE_HIERARCHICAL_COUNTER_RIGUNIT()
 	Result = FMath::RadiansToDegrees(Value);
 }
 
-UE_RigUnit_MathFloatRad_IMPLEMENT_STATIC_VIRTUAL_METHOD(void, Execute, const FRigUnitContext& Context)
+UE_RigUnit_MathFloatRad_IMPLEMENT_MULTIPLEX(void, Execute, const FRigUnitContext& Context)
 {
     DECLARE_SCOPE_HIERARCHICAL_COUNTER_RIGUNIT()
 	Result = FMath::DegreesToRadians(Value);
 }
 
-UE_RigUnit_MathFloatSin_IMPLEMENT_STATIC_VIRTUAL_METHOD(void, Execute, const FRigUnitContext& Context)
+UE_RigUnit_MathFloatSin_IMPLEMENT_MULTIPLEX(void, Execute, const FRigUnitContext& Context)
 {
     DECLARE_SCOPE_HIERARCHICAL_COUNTER_RIGUNIT()
 	Result = FMath::Sin(Value);
 }
 
-UE_RigUnit_MathFloatCos_IMPLEMENT_STATIC_VIRTUAL_METHOD(void, Execute, const FRigUnitContext& Context)
+UE_RigUnit_MathFloatCos_IMPLEMENT_MULTIPLEX(void, Execute, const FRigUnitContext& Context)
 {
     DECLARE_SCOPE_HIERARCHICAL_COUNTER_RIGUNIT()
 	Result = FMath::Cos(Value);
 }
 
-UE_RigUnit_MathFloatTan_IMPLEMENT_STATIC_VIRTUAL_METHOD(void, Execute, const FRigUnitContext& Context)
+UE_RigUnit_MathFloatTan_IMPLEMENT_MULTIPLEX(void, Execute, const FRigUnitContext& Context)
 {
     DECLARE_SCOPE_HIERARCHICAL_COUNTER_RIGUNIT()
 	Result = FMath::Tan(Value);
 }
 
-UE_RigUnit_MathFloatAsin_IMPLEMENT_STATIC_VIRTUAL_METHOD(void, Execute, const FRigUnitContext& Context)
+UE_RigUnit_MathFloatAsin_IMPLEMENT_MULTIPLEX(void, Execute, const FRigUnitContext& Context)
 {
     DECLARE_SCOPE_HIERARCHICAL_COUNTER_RIGUNIT()
 	if (!FMath::IsWithinInclusive<float>(Value, -1.f, 1.f))
@@ -246,7 +246,7 @@ UE_RigUnit_MathFloatAsin_IMPLEMENT_STATIC_VIRTUAL_METHOD(void, Execute, const FR
 	Result = FMath::Asin(Value);
 }
 
-UE_RigUnit_MathFloatAcos_IMPLEMENT_STATIC_VIRTUAL_METHOD(void, Execute, const FRigUnitContext& Context)
+UE_RigUnit_MathFloatAcos_IMPLEMENT_MULTIPLEX(void, Execute, const FRigUnitContext& Context)
 {
     DECLARE_SCOPE_HIERARCHICAL_COUNTER_RIGUNIT()
 	if (!FMath::IsWithinInclusive<float>(Value, -1.f, 1.f))
@@ -258,13 +258,13 @@ UE_RigUnit_MathFloatAcos_IMPLEMENT_STATIC_VIRTUAL_METHOD(void, Execute, const FR
 	Result = FMath::Acos(Value);
 }
 
-UE_RigUnit_MathFloatAtan_IMPLEMENT_STATIC_VIRTUAL_METHOD(void, Execute, const FRigUnitContext& Context)
+UE_RigUnit_MathFloatAtan_IMPLEMENT_MULTIPLEX(void, Execute, const FRigUnitContext& Context)
 {
     DECLARE_SCOPE_HIERARCHICAL_COUNTER_RIGUNIT()
 	Result = FMath::Atan(Value);
 }
 
-UE_RigUnit_MathFloatLawOfCosine_IMPLEMENT_STATIC_VIRTUAL_METHOD(void, Execute, const FRigUnitContext& Context)
+UE_RigUnit_MathFloatLawOfCosine_IMPLEMENT_MULTIPLEX(void, Execute, const FRigUnitContext& Context)
 {
     DECLARE_SCOPE_HIERARCHICAL_COUNTER_RIGUNIT()
 	if ((A <= 0.f) || (B <= 0.f) || (C <= 0.f) || (A + B < C) || (A + C < B) || (B + C < A))
diff --git a/Engine/Plugins/Experimental/ControlRig/Source/ControlRig/Private/Units/Math/RigUnit_MathFloat.h b/Engine/Plugins/Experimental/ControlRig/Source/ControlRig/Private/Units/Math/RigUnit_MathFloat.h
--- a/Engine/Plugins/Experimental/ControlRig/Source/ControlRig/Private/Units/Math/RigUnit_MathFloat.h
+++ b/Engine/Plugins/Experimental/ControlRig/Source/ControlRig/Private/Units/Math/RigUnit_MathFloat.h
@@ -112,7 +112,7 @@ struct FRigUnit_MathFloatAdd : public FRigUnit_MathFloatBinaryOp
 {
 	GENERATED_BODY()
 
-	STATIC_VIRTUAL_METHOD()
+	MULTIPLEX_METHOD()
 	virtual void Execute(const FRigUnitContext& Context) override;
 };
 
@@ -124,7 +124,7 @@ struct FRigUnit_MathFloatSub : public FRigUnit_MathFloatBinaryOp
 {
 	GENERATED_BODY()
 
-	STATIC_VIRTUAL_METHOD()
+	MULTIPLEX_METHOD()
 	virtual void Execute(const FRigUnitContext& Context) override;
 };
 
@@ -142,7 +142,7 @@ struct FRigUnit_MathFloatMul : public FRigUnit_MathFloatBinaryOp
 		B = 1.f;
 	}
 
-	STATIC_VIRTUAL_METHOD()
+	MULTIPLEX_METHOD()
 	virtual void Execute(const FRigUnitContext& Context) override;
 };
 
@@ -159,7 +159,7 @@ struct FRigUnit_MathFloatDiv : public FRigUnit_MathFloatBinaryOp
 		B = 1.f;
 	}
 
-	STATIC_VIRTUAL_METHOD()
+	MULTIPLEX_METHOD()
 	virtual void Execute(const FRigUnitContext& Context) override;
 };
 
@@ -177,7 +177,7 @@ struct FRigUnit_MathFloatMod : public FRigUnit_MathFloatBinaryOp
 		B = 1.f;
 	}
 
-	STATIC_VIRTUAL_METHOD()
+	MULTIPLEX_METHOD()
 	virtual void Execute(const FRigUnitContext& Context) override;
 };
 
@@ -189,7 +189,7 @@ struct FRigUnit_MathFloatMin : public FRigUnit_MathFloatBinaryOp
 {
 	GENERATED_BODY()
 
-	STATIC_VIRTUAL_METHOD()
+	MULTIPLEX_METHOD()
 	virtual void Execute(const FRigUnitContext& Context) override;
 };
 
@@ -201,7 +201,7 @@ struct FRigUnit_MathFloatMax : public FRigUnit_MathFloatBinaryOp
 {
 	GENERATED_BODY()
 
-	STATIC_VIRTUAL_METHOD()
+	MULTIPLEX_METHOD()
 	virtual void Execute(const FRigUnitContext& Context) override;
 };
 
@@ -219,7 +219,7 @@ struct FRigUnit_MathFloatPow : public FRigUnit_MathFloatBinaryOp
 		B = 1.f;
 	}
 
-	STATIC_VIRTUAL_METHOD()
+	MULTIPLEX_METHOD()
 	virtual void Execute(const FRigUnitContext& Context) override;
 };
 
@@ -231,7 +231,7 @@ struct FRigUnit_MathFloatSqrt : public FRigUnit_MathFloatUnaryOp
 {
 	GENERATED_BODY()
 
-	STATIC_VIRTUAL_METHOD()
+	MULTIPLEX_METHOD()
 	virtual void Execute(const FRigUnitContext& Context) override;
 };
 
@@ -243,7 +243,7 @@ struct FRigUnit_MathFloatNegate : public FRigUnit_MathFloatUnaryOp
 {
 	GENERATED_BODY()
 
-	STATIC_VIRTUAL_METHOD()
+	MULTIPLEX_METHOD()
 	virtual void Execute(const FRigUnitContext& Context) override;
 };
 
@@ -255,7 +255,7 @@ struct FRigUnit_MathFloatAbs : public FRigUnit_MathFloatUnaryOp
 {
 	GENERATED_BODY()
 
-	STATIC_VIRTUAL_METHOD()
+	MULTIPLEX_METHOD()
 	virtual void Execute(const FRigUnitContext& Context) override;
 };
 
@@ -267,7 +267,7 @@ struct FRigUnit_MathFloatFloor : public FRigUnit_MathFloatUnaryOp
 {
 	GENERATED_BODY()
 
-	STATIC_VIRTUAL_METHOD()
+	MULTIPLEX_METHOD()
 	virtual void Execute(const FRigUnitContext& Context) override;
 };
 
@@ -279,7 +279,7 @@ struct FRigUnit_MathFloatCeil : public FRigUnit_MathFloatUnaryOp
 {
 	GENERATED_BODY()
 
-	STATIC_VIRTUAL_METHOD()
+	MULTIPLEX_METHOD()
 	virtual void Execute(const FRigUnitContext& Context) override;
 };
 
@@ -291,7 +291,7 @@ struct FRigUnit_MathFloatRound : public FRigUnit_MathFloatUnaryOp
 {
 	GENERATED_BODY()
 
-	STATIC_VIRTUAL_METHOD()
+	MULTIPLEX_METHOD()
 	virtual void Execute(const FRigUnitContext& Context) override;
 };
 
@@ -303,7 +303,7 @@ struct FRigUnit_MathFloatSign : public FRigUnit_MathFloatUnaryOp
 {
 	GENERATED_BODY()
 
-	STATIC_VIRTUAL_METHOD()
+	MULTIPLEX_METHOD()
 	virtual void Execute(const FRigUnitContext& Context) override;
 };
 
@@ -320,7 +320,7 @@ struct FRigUnit_MathFloatClamp : public FRigUnit_MathFloatBase
 		Value = Minimum = Maximum = Result = 0.f;
 	}
 
-	STATIC_VIRTUAL_METHOD()
+	MULTIPLEX_METHOD()
 	virtual void Execute(const FRigUnitContext& Context) override;
 
 	UPROPERTY(meta=(Input))
@@ -350,7 +350,7 @@ struct FRigUnit_MathFloatLerp : public FRigUnit_MathFloatBase
 		B = 1.f;
 	}
 
-	STATIC_VIRTUAL_METHOD()
+	MULTIPLEX_METHOD()
 	virtual void Execute(const FRigUnitContext& Context) override;
 
 	UPROPERTY(meta=(Input))
@@ -381,7 +381,7 @@ struct FRigUnit_MathFloatRemap : public FRigUnit_MathFloatBase
 		bClamp = false;
 	}
 
-	STATIC_VIRTUAL_METHOD()
+	MULTIPLEX_METHOD()
 	virtual void Execute(const FRigUnitContext& Context) override;
 
 	UPROPERTY(meta=(Input))
@@ -414,7 +414,7 @@ USTRUCT(meta=(DisplayName="Equals", PrototypeName="Equals", Keywords="Same,=="))
 struct FRigUnit_MathFloatEquals : public FRigUnit_MathFloatBase
 {
 	GENERATED_BODY()
-	STATIC_VIRTUAL_METHOD()
+	MULTIPLEX_METHOD()
 	virtual void Execute(const FRigUnitContext& Context) override;
 
 	FRigUnit_MathFloatEquals()
@@ -440,7 +440,7 @@ USTRUCT(meta=(DisplayName="Not Equals", PrototypeName="NotEquals", Keywords="Dif
 struct FRigUnit_MathFloatNotEquals : public FRigUnit_MathFloatBase
 {
 	GENERATED_BODY()
-	STATIC_VIRTUAL_METHOD()
+	MULTIPLEX_METHOD()
 	virtual void Execute(const FRigUnitContext& Context) override;
 
 	FRigUnit_MathFloatNotEquals()
@@ -473,7 +473,7 @@ struct FRigUnit_MathFloatGreater : public FRigUnit_MathFloatBase
 		Result = false;
 	}
 
-	STATIC_VIRTUAL_METHOD()
+	MULTIPLEX_METHOD()
 	virtual void Execute(const FRigUnitContext& Context) override;
 
 	UPROPERTY(meta=(Input))
@@ -500,7 +500,7 @@ struct FRigUnit_MathFloatLess : public FRigUnit_MathFloatBase
 		Result = false;
 	}
 
-	STATIC_VIRTUAL_METHOD()
+	MULTIPLEX_METHOD()
 	virtual void Execute(const FRigUnitContext& Context) override;
 
 	UPROPERTY(meta=(Input))
@@ -527,7 +527,7 @@ struct FRigUnit_MathFloatGreaterEqual : public FRigUnit_MathFloatBase
 		Result = true;
 	}
 
-	STATIC_VIRTUAL_METHOD()
+	MULTIPLEX_METHOD()
 	virtual void Execute(const FRigUnitContext& Context) override;
 
 	UPROPERTY(meta=(Input))
@@ -554,7 +554,7 @@ struct FRigUnit_MathFloatLessEqual : public FRigUnit_MathFloatBase
 		Result = true;
 	}
 
-	STATIC_VIRTUAL_METHOD()
+	MULTIPLEX_METHOD()
 	virtual void Execute(const FRigUnitContext& Context) override;
 
 	UPROPERTY(meta=(Input))
@@ -581,7 +581,7 @@ struct FRigUnit_MathFloatIsNearlyZero : public FRigUnit_MathFloatBase
 		Result = true;
 	}
 	
-	STATIC_VIRTUAL_METHOD()
+	MULTIPLEX_METHOD()
 	virtual void Execute(const FRigUnitContext& Context) override;
 
 	UPROPERTY(meta=(Input))
@@ -608,7 +608,7 @@ struct FRigUnit_MathFloatIsNearlyEqual : public FRigUnit_MathFloatBase
 		Result = true;
 	}
 
-	STATIC_VIRTUAL_METHOD()
+	MULTIPLEX_METHOD()
 	virtual void Execute(const FRigUnitContext& Context) override;
 
 	UPROPERTY(meta=(Input))
@@ -638,7 +638,7 @@ struct FRigUnit_MathFloatSelectBool : public FRigUnit_MathFloatBase
 		IfTrue = IfFalse = Result = 0.f;
 	}
 
-	STATIC_VIRTUAL_METHOD()
+	MULTIPLEX_METHOD()
 	virtual void Execute(const FRigUnitContext& Context) override;
 
 	UPROPERTY(meta=(Input))
@@ -662,7 +662,7 @@ struct FRigUnit_MathFloatDeg : public FRigUnit_MathFloatUnaryOp
 {
 	GENERATED_BODY()
 
-	STATIC_VIRTUAL_METHOD()
+	MULTIPLEX_METHOD()
 	virtual void Execute(const FRigUnitContext& Context) override;
 };
 
@@ -674,7 +674,7 @@ struct FRigUnit_MathFloatRad : public FRigUnit_MathFloatUnaryOp
 {
 	GENERATED_BODY()
 
-	STATIC_VIRTUAL_METHOD()
+	MULTIPLEX_METHOD()
 	virtual void Execute(const FRigUnitContext& Context) override;
 };
 
@@ -686,7 +686,7 @@ struct FRigUnit_MathFloatSin : public FRigUnit_MathFloatUnaryOp
 {
 	GENERATED_BODY()
 
-	STATIC_VIRTUAL_METHOD()
+	MULTIPLEX_METHOD()
 	virtual void Execute(const FRigUnitContext& Context) override;
 };
 
@@ -698,7 +698,7 @@ struct FRigUnit_MathFloatCos : public FRigUnit_MathFloatUnaryOp
 {
 	GENERATED_BODY()
 
-	STATIC_VIRTUAL_METHOD()
+	MULTIPLEX_METHOD()
 	virtual void Execute(const FRigUnitContext& Context) override;
 };
 
@@ -710,7 +710,7 @@ struct FRigUnit_MathFloatTan : public FRigUnit_MathFloatUnaryOp
 {
 	GENERATED_BODY()
 
-	STATIC_VIRTUAL_METHOD()
+	MULTIPLEX_METHOD()
 	virtual void Execute(const FRigUnitContext& Context) override;
 };
 
@@ -722,7 +722,7 @@ struct FRigUnit_MathFloatAsin : public FRigUnit_MathFloatUnaryOp
 {
 	GENERATED_BODY()
 
-	STATIC_VIRTUAL_METHOD()
+	MULTIPLEX_METHOD()
 	virtual void Execute(const FRigUnitContext& Context) override;
 };
 
@@ -734,7 +734,7 @@ struct FRigUnit_MathFloatAcos : public FRigUnit_MathFloatUnaryOp
 {
 	GENERATED_BODY()
 
-	STATIC_VIRTUAL_METHOD()
+	MULTIPLEX_METHOD()
 	virtual void Execute(const FRigUnitContext& Context) override;
 };
 
@@ -746,7 +746,7 @@ struct FRigUnit_MathFloatAtan : public FRigUnit_MathFloatUnaryOp
 {
 	GENERATED_BODY()
 
-	STATIC_VIRTUAL_METHOD()
+	MULTIPLEX_METHOD()
 	virtual void Execute(const FRigUnitContext& Context) override;
 };
 
@@ -764,7 +764,7 @@ struct FRigUnit_MathFloatLawOfCosine : public FRigUnit_MathFloatBase
 		bValid = false;
 	}
 
-	STATIC_VIRTUAL_METHOD()
+	MULTIPLEX_METHOD()
 	virtual void Execute(const FRigUnitContext& Context) override;
 
 	UPROPERTY(meta = (Input))
diff --git a/Engine/Plugins/Experimental/ControlRig/Source/ControlRig/Private/Units/Math/RigUnit_MathQuaternion.cpp b/Engine/Plugins/Experimental/ControlRig/Source/ControlRig/Private/Units/Math/RigUnit_MathQuaternion.cpp
--- a/Engine/Plugins/Experimental/ControlRig/Source/ControlRig/Private/Units/Math/RigUnit_MathQuaternion.cpp
+++ b/Engine/Plugins/Experimental/ControlRig/Source/ControlRig/Private/Units/Math/RigUnit_MathQuaternion.cpp
@@ -3,7 +3,7 @@
 #include "Units/Math/RigUnit_MathQuaternion.h"
 #include "Units/RigUnitContext.h"
 
-UE_RigUnit_MathQuaternionFromAxisAndAngle_IMPLEMENT_STATIC_VIRTUAL_METHOD(void, Execute, const FRigUnitContext& Context)
+UE_RigUnit_MathQuaternionFromAxisAndAngle_IMPLEMENT_MULTIPLEX(void, Execute, const FRigUnitContext& Context)
 {
     DECLARE_SCOPE_HIERARCHICAL_COUNTER_RIGUNIT()
 	if (Axis.IsNearlyZero())
@@ -15,19 +15,19 @@ UE_RigUnit_MathQuaternionFromAxisAndAngle_IMPLEMENT_STATIC_VIRTUAL_METHOD(void,
 	Result = FQuat(Axis.GetUnsafeNormal(), Angle);
 }
 
-UE_RigUnit_MathQuaternionFromEuler_IMPLEMENT_STATIC_VIRTUAL_METHOD(void, Execute, const FRigUnitContext& Context)
+UE_RigUnit_MathQuaternionFromEuler_IMPLEMENT_MULTIPLEX(void, Execute, const FRigUnitContext& Context)
 {
     DECLARE_SCOPE_HIERARCHICAL_COUNTER_RIGUNIT()
 	Result = FControlRigMathLibrary::QuatFromEuler(Euler, RotationOrder);
 }
 
-UE_RigUnit_MathQuaternionFromRotator_IMPLEMENT_STATIC_VIRTUAL_METHOD(void, Execute, const FRigUnitContext& Context)
+UE_RigUnit_MathQuaternionFromRotator_IMPLEMENT_MULTIPLEX(void, Execute, const FRigUnitContext& Context)
 {
     DECLARE_SCOPE_HIERARCHICAL_COUNTER_RIGUNIT()
 	Result = FQuat(Rotator);
 }
 
-UE_RigUnit_MathQuaternionFromTwoVectors_IMPLEMENT_STATIC_VIRTUAL_METHOD(void, Execute, const FRigUnitContext& Context)
+UE_RigUnit_MathQuaternionFromTwoVectors_IMPLEMENT_MULTIPLEX(void, Execute, const FRigUnitContext& Context)
 {
     DECLARE_SCOPE_HIERARCHICAL_COUNTER_RIGUNIT()
 	if (A.IsNearlyZero() || B.IsNearlyZero())
@@ -46,79 +46,79 @@ UE_RigUnit_MathQuaternionFromTwoVectors_IMPLEMENT_STATIC_VIRTUAL_METHOD(void, Ex
 	Result = FQuat::FindBetweenVectors(A, B).GetNormalized();
 }
 
-UE_RigUnit_MathQuaternionToAxisAndAngle_IMPLEMENT_STATIC_VIRTUAL_METHOD(void, Execute, const FRigUnitContext& Context)
+UE_RigUnit_MathQuaternionToAxisAndAngle_IMPLEMENT_MULTIPLEX(void, Execute, const FRigUnitContext& Context)
 {
     DECLARE_SCOPE_HIERARCHICAL_COUNTER_RIGUNIT()
 	Value.ToAxisAndAngle(Axis, Angle);
 }
 
-UE_RigUnit_MathQuaternionToEuler_IMPLEMENT_STATIC_VIRTUAL_METHOD(void, Execute, const FRigUnitContext& Context)
+UE_RigUnit_MathQuaternionToEuler_IMPLEMENT_MULTIPLEX(void, Execute, const FRigUnitContext& Context)
 {
     DECLARE_SCOPE_HIERARCHICAL_COUNTER_RIGUNIT()
 	Result = FControlRigMathLibrary::EulerFromQuat(Value, RotationOrder);
 }
 
-UE_RigUnit_MathQuaternionToRotator_IMPLEMENT_STATIC_VIRTUAL_METHOD(void, Execute, const FRigUnitContext& Context)
+UE_RigUnit_MathQuaternionToRotator_IMPLEMENT_MULTIPLEX(void, Execute, const FRigUnitContext& Context)
 {
     DECLARE_SCOPE_HIERARCHICAL_COUNTER_RIGUNIT()
 	Result = Value.Rotator();
 }
 
-UE_RigUnit_MathQuaternionMul_IMPLEMENT_STATIC_VIRTUAL_METHOD(void, Execute, const FRigUnitContext& Context)
+UE_RigUnit_MathQuaternionMul_IMPLEMENT_MULTIPLEX(void, Execute, const FRigUnitContext& Context)
 {
     DECLARE_SCOPE_HIERARCHICAL_COUNTER_RIGUNIT()
 	Result = A * B;
 }
 
-UE_RigUnit_MathQuaternionInverse_IMPLEMENT_STATIC_VIRTUAL_METHOD(void, Execute, const FRigUnitContext& Context)
+UE_RigUnit_MathQuaternionInverse_IMPLEMENT_MULTIPLEX(void, Execute, const FRigUnitContext& Context)
 {
     DECLARE_SCOPE_HIERARCHICAL_COUNTER_RIGUNIT()
 	Result = Value.Inverse();
 }
 
-UE_RigUnit_MathQuaternionSlerp_IMPLEMENT_STATIC_VIRTUAL_METHOD(void, Execute, const FRigUnitContext& Context)
+UE_RigUnit_MathQuaternionSlerp_IMPLEMENT_MULTIPLEX(void, Execute, const FRigUnitContext& Context)
 {
     DECLARE_SCOPE_HIERARCHICAL_COUNTER_RIGUNIT()
 	Result = FQuat::Slerp(A, B, T);
 }
 
-UE_RigUnit_MathQuaternionEquals_IMPLEMENT_STATIC_VIRTUAL_METHOD(void, Execute, const FRigUnitContext& Context)
+UE_RigUnit_MathQuaternionEquals_IMPLEMENT_MULTIPLEX(void, Execute, const FRigUnitContext& Context)
 {
     DECLARE_SCOPE_HIERARCHICAL_COUNTER_RIGUNIT()
 	Result = A == B;
 }
 
-UE_RigUnit_MathQuaternionNotEquals_IMPLEMENT_STATIC_VIRTUAL_METHOD(void, Execute, const FRigUnitContext& Context)
+UE_RigUnit_MathQuaternionNotEquals_IMPLEMENT_MULTIPLEX(void, Execute, const FRigUnitContext& Context)
 {
     DECLARE_SCOPE_HIERARCHICAL_COUNTER_RIGUNIT()
 	Result = A != B;
 }
 
-UE_RigUnit_MathQuaternionSelectBool_IMPLEMENT_STATIC_VIRTUAL_METHOD(void, Execute, const FRigUnitContext& Context)
+UE_RigUnit_MathQuaternionSelectBool_IMPLEMENT_MULTIPLEX(void, Execute, const FRigUnitContext& Context)
 {
     DECLARE_SCOPE_HIERARCHICAL_COUNTER_RIGUNIT()
 	Result = Condition ? IfTrue : IfFalse;
 }
 
-UE_RigUnit_MathQuaternionDot_IMPLEMENT_STATIC_VIRTUAL_METHOD(void, Execute, const FRigUnitContext& Context)
+UE_RigUnit_MathQuaternionDot_IMPLEMENT_MULTIPLEX(void, Execute, const FRigUnitContext& Context)
 {
     DECLARE_SCOPE_HIERARCHICAL_COUNTER_RIGUNIT()
 	Result = A | B;
 }
 
-UE_RigUnit_MathQuaternionUnit_IMPLEMENT_STATIC_VIRTUAL_METHOD(void, Execute, const FRigUnitContext& Context)
+UE_RigUnit_MathQuaternionUnit_IMPLEMENT_MULTIPLEX(void, Execute, const FRigUnitContext& Context)
 {
     DECLARE_SCOPE_HIERARCHICAL_COUNTER_RIGUNIT()
 	Result = Value.GetNormalized();
 }
 
-UE_RigUnit_MathQuaternionRotateVector_IMPLEMENT_STATIC_VIRTUAL_METHOD(void, Execute, const FRigUnitContext& Context)
+UE_RigUnit_MathQuaternionRotateVector_IMPLEMENT_MULTIPLEX(void, Execute, const FRigUnitContext& Context)
 {
     DECLARE_SCOPE_HIERARCHICAL_COUNTER_RIGUNIT()
 	Result = Quaternion.RotateVector(Vector);
 }
 
-UE_RigUnit_MathQuaternionGetAxis_IMPLEMENT_STATIC_VIRTUAL_METHOD(void, Execute, const FRigUnitContext& Context)
+UE_RigUnit_MathQuaternionGetAxis_IMPLEMENT_MULTIPLEX(void, Execute, const FRigUnitContext& Context)
 {
     DECLARE_SCOPE_HIERARCHICAL_COUNTER_RIGUNIT()
 	switch (Axis)
@@ -143,7 +143,7 @@ UE_RigUnit_MathQuaternionGetAxis_IMPLEMENT_STATIC_VIRTUAL_METHOD(void, Execute,
 }
 
 
-UE_RigUnit_MathQuaternionSwingTwist_IMPLEMENT_STATIC_VIRTUAL_METHOD(void, Execute, const FRigUnitContext& Context)
+UE_RigUnit_MathQuaternionSwingTwist_IMPLEMENT_MULTIPLEX(void, Execute, const FRigUnitContext& Context)
 {
     DECLARE_SCOPE_HIERARCHICAL_COUNTER_RIGUNIT()
 	if (TwistAxis.IsNearlyZero())
diff --git a/Engine/Plugins/Experimental/ControlRig/Source/ControlRig/Private/Units/Math/RigUnit_MathQuaternion.h b/Engine/Plugins/Experimental/ControlRig/Source/ControlRig/Private/Units/Math/RigUnit_MathQuaternion.h
--- a/Engine/Plugins/Experimental/ControlRig/Source/ControlRig/Private/Units/Math/RigUnit_MathQuaternion.h
+++ b/Engine/Plugins/Experimental/ControlRig/Source/ControlRig/Private/Units/Math/RigUnit_MathQuaternion.h
@@ -64,7 +64,7 @@ struct FRigUnit_MathQuaternionFromAxisAndAngle : public FRigUnit_MathQuaternionB
 		Result = FQuat::Identity;
 	}
 
-	STATIC_VIRTUAL_METHOD()
+	MULTIPLEX_METHOD()
 	virtual void Execute(const FRigUnitContext& Context) override;
 
 	UPROPERTY(meta=(Input))
@@ -92,7 +92,7 @@ struct FRigUnit_MathQuaternionFromEuler : public FRigUnit_MathQuaternionBase
 		Result = FQuat::Identity;
 	}
 
-	STATIC_VIRTUAL_METHOD()
+	MULTIPLEX_METHOD()
 	virtual void Execute(const FRigUnitContext& Context) override;
 
 	UPROPERTY(meta = (Input))
@@ -119,7 +119,7 @@ struct FRigUnit_MathQuaternionFromRotator : public FRigUnit_MathQuaternionBase
 		Result = FQuat::Identity;
 	}
 
-	STATIC_VIRTUAL_METHOD()
+	MULTIPLEX_METHOD()
 	virtual void Execute(const FRigUnitContext& Context) override;
 
 	UPROPERTY(meta=(Input))
@@ -143,7 +143,7 @@ struct FRigUnit_MathQuaternionFromTwoVectors : public FRigUnit_MathQuaternionBas
 		Result = FQuat::Identity;
 	}
 
-	STATIC_VIRTUAL_METHOD()
+	MULTIPLEX_METHOD()
 	virtual void Execute(const FRigUnitContext& Context) override;
 
 	UPROPERTY(meta=(Input))
@@ -171,7 +171,7 @@ struct FRigUnit_MathQuaternionToAxisAndAngle : public FRigUnit_MathQuaternionBas
 		Angle = 0.f;
 	}
 
-	STATIC_VIRTUAL_METHOD()
+	MULTIPLEX_METHOD()
 	virtual void Execute(const FRigUnitContext& Context) override;
 
 	UPROPERTY(meta=(Input))
@@ -199,7 +199,7 @@ struct FRigUnit_MathQuaternionToEuler : public FRigUnit_MathQuaternionBase
 		Result = FVector::ZeroVector;
 	}
 
-	STATIC_VIRTUAL_METHOD()
+	MULTIPLEX_METHOD()
 	virtual void Execute(const FRigUnitContext& Context) override;
 
 	UPROPERTY(meta=(Input))
@@ -226,7 +226,7 @@ struct FRigUnit_MathQuaternionToRotator : public FRigUnit_MathQuaternionBase
 		Result = FRotator::ZeroRotator;
 	}
 
-	STATIC_VIRTUAL_METHOD()
+	MULTIPLEX_METHOD()
 	virtual void Execute(const FRigUnitContext& Context) override;
 
 	UPROPERTY(meta=(Input))
@@ -243,7 +243,7 @@ USTRUCT(meta=(DisplayName="Multiply", PrototypeName="Multiply", Keywords="Produc
 struct FRigUnit_MathQuaternionMul : public FRigUnit_MathQuaternionBinaryOp
 {
 	GENERATED_BODY()
-	STATIC_VIRTUAL_METHOD()
+	MULTIPLEX_METHOD()
 	virtual void Execute(const FRigUnitContext& Context) override;
 };
 
@@ -254,7 +254,7 @@ USTRUCT(meta=(DisplayName="Inverse", PrototypeName="Inverse"))
 struct FRigUnit_MathQuaternionInverse : public FRigUnit_MathQuaternionUnaryOp
 {
 	GENERATED_BODY()
-	STATIC_VIRTUAL_METHOD()
+	MULTIPLEX_METHOD()
 	virtual void Execute(const FRigUnitContext& Context) override;
 };
 
@@ -272,7 +272,7 @@ struct FRigUnit_MathQuaternionSlerp : public FRigUnit_MathQuaternionBase
 		T = 0.f;
 	}
 
-	STATIC_VIRTUAL_METHOD()
+	MULTIPLEX_METHOD()
 	virtual void Execute(const FRigUnitContext& Context) override;
 
 	UPROPERTY(meta=(Input))
@@ -302,7 +302,7 @@ struct FRigUnit_MathQuaternionEquals : public FRigUnit_MathQuaternionBase
 		Result = true;
 	}	
 
-	STATIC_VIRTUAL_METHOD()
+	MULTIPLEX_METHOD()
 	virtual void Execute(const FRigUnitContext& Context) override;
 
 	UPROPERTY(meta=(Input))
@@ -329,7 +329,7 @@ struct FRigUnit_MathQuaternionNotEquals : public FRigUnit_MathQuaternionBase
 		Result = false;
 	}
 
-	STATIC_VIRTUAL_METHOD()
+	MULTIPLEX_METHOD()
 	virtual void Execute(const FRigUnitContext& Context) override;
 
 	UPROPERTY(meta=(Input))
@@ -356,7 +356,7 @@ struct FRigUnit_MathQuaternionSelectBool : public FRigUnit_MathQuaternionBase
 		Condition = true;
 	}
 
-	STATIC_VIRTUAL_METHOD()
+	MULTIPLEX_METHOD()
 	virtual void Execute(const FRigUnitContext& Context) override;
 
 	UPROPERTY(meta=(Input))
@@ -386,7 +386,7 @@ struct FRigUnit_MathQuaternionDot : public FRigUnit_MathQuaternionBase
 		Result = 0.f;
 	}
 
-	STATIC_VIRTUAL_METHOD()
+	MULTIPLEX_METHOD()
 	virtual void Execute(const FRigUnitContext& Context) override;
 
 	UPROPERTY(meta=(Input))
@@ -406,7 +406,7 @@ USTRUCT(meta=(DisplayName="Unit", PrototypeName="Unit", Keywords="Normalize"))
 struct FRigUnit_MathQuaternionUnit : public FRigUnit_MathQuaternionUnaryOp
 {
 	GENERATED_BODY()
-	STATIC_VIRTUAL_METHOD()
+	MULTIPLEX_METHOD()
 	virtual void Execute(const FRigUnitContext& Context) override;
 };
 
@@ -424,7 +424,7 @@ struct FRigUnit_MathQuaternionRotateVector : public FRigUnit_MathQuaternionBase
 		Vector = Result = FVector::ZeroVector;
 	}
 
-	STATIC_VIRTUAL_METHOD()
+	MULTIPLEX_METHOD()
 	virtual void Execute(const FRigUnitContext& Context) override;
 
 	UPROPERTY(meta=(Input))
@@ -452,7 +452,7 @@ struct FRigUnit_MathQuaternionGetAxis: public FRigUnit_MathQuaternionBase
 		Result = FVector(1.f, 0.f, 0.f);
 	}
 
-	STATIC_VIRTUAL_METHOD()
+	MULTIPLEX_METHOD()
 	virtual void Execute(const FRigUnitContext& Context) override;
 
 	UPROPERTY(meta = (Input))
@@ -480,7 +480,7 @@ struct FRigUnit_MathQuaternionSwingTwist : public FRigUnit_MathQuaternionBase
 		TwistAxis = FVector(1.f, 0.f, 0.f);
 	}
 
-	STATIC_VIRTUAL_METHOD()
+	MULTIPLEX_METHOD()
 	virtual void Execute(const FRigUnitContext& Context) override;
 
 	UPROPERTY(meta = (Input))
diff --git a/Engine/Plugins/Experimental/ControlRig/Source/ControlRig/Private/Units/Math/RigUnit_MathTransform.cpp b/Engine/Plugins/Experimental/ControlRig/Source/ControlRig/Private/Units/Math/RigUnit_MathTransform.cpp
--- a/Engine/Plugins/Experimental/ControlRig/Source/ControlRig/Private/Units/Math/RigUnit_MathTransform.cpp
+++ b/Engine/Plugins/Experimental/ControlRig/Source/ControlRig/Private/Units/Math/RigUnit_MathTransform.cpp
@@ -4,61 +4,61 @@
 #include "Math/ControlRigMathLibrary.h"
 #include "Units/RigUnitContext.h"
 
-UE_RigUnit_MathTransformFromEulerTransform_IMPLEMENT_STATIC_VIRTUAL_METHOD(void, Execute, const FRigUnitContext& Context)
+UE_RigUnit_MathTransformFromEulerTransform_IMPLEMENT_MULTIPLEX(void, Execute, const FRigUnitContext& Context)
 {
     DECLARE_SCOPE_HIERARCHICAL_COUNTER_RIGUNIT()
 	Result = EulerTransform.ToFTransform();
 }
 
-UE_RigUnit_MathTransformToEulerTransform_IMPLEMENT_STATIC_VIRTUAL_METHOD(void, Execute, const FRigUnitContext& Context)
+UE_RigUnit_MathTransformToEulerTransform_IMPLEMENT_MULTIPLEX(void, Execute, const FRigUnitContext& Context)
 {
     DECLARE_SCOPE_HIERARCHICAL_COUNTER_RIGUNIT()
 	Result.FromFTransform(Value);
 }
 
-UE_RigUnit_MathTransformMul_IMPLEMENT_STATIC_VIRTUAL_METHOD(void, Execute, const FRigUnitContext& Context)
+UE_RigUnit_MathTransformMul_IMPLEMENT_MULTIPLEX(void, Execute, const FRigUnitContext& Context)
 {
     DECLARE_SCOPE_HIERARCHICAL_COUNTER_RIGUNIT()
 	Result = A * B;
 }
 
-UE_RigUnit_MathTransformMakeRelative_IMPLEMENT_STATIC_VIRTUAL_METHOD(void, Execute, const FRigUnitContext& Context)
+UE_RigUnit_MathTransformMakeRelative_IMPLEMENT_MULTIPLEX(void, Execute, const FRigUnitContext& Context)
 {
     DECLARE_SCOPE_HIERARCHICAL_COUNTER_RIGUNIT()
 	Local = Global.GetRelativeTransform(Parent);
 }
 
-UE_RigUnit_MathTransformInverse_IMPLEMENT_STATIC_VIRTUAL_METHOD(void, Execute, const FRigUnitContext& Context)
+UE_RigUnit_MathTransformInverse_IMPLEMENT_MULTIPLEX(void, Execute, const FRigUnitContext& Context)
 {
     DECLARE_SCOPE_HIERARCHICAL_COUNTER_RIGUNIT()
 	Result = Value.Inverse();
 }
 
-UE_RigUnit_MathTransformLerp_IMPLEMENT_STATIC_VIRTUAL_METHOD(void, Execute, const FRigUnitContext& Context)
+UE_RigUnit_MathTransformLerp_IMPLEMENT_MULTIPLEX(void, Execute, const FRigUnitContext& Context)
 {
     DECLARE_SCOPE_HIERARCHICAL_COUNTER_RIGUNIT()
 	Result = FControlRigMathLibrary::LerpTransform(A, B, T);
 }
 
-UE_RigUnit_MathTransformSelectBool_IMPLEMENT_STATIC_VIRTUAL_METHOD(void, Execute, const FRigUnitContext& Context)
+UE_RigUnit_MathTransformSelectBool_IMPLEMENT_MULTIPLEX(void, Execute, const FRigUnitContext& Context)
 {
     DECLARE_SCOPE_HIERARCHICAL_COUNTER_RIGUNIT()
 	Result = Condition ? IfTrue : IfFalse;
 }
 
-UE_RigUnit_MathTransformRotateVector_IMPLEMENT_STATIC_VIRTUAL_METHOD(void, Execute, const FRigUnitContext& Context)
+UE_RigUnit_MathTransformRotateVector_IMPLEMENT_MULTIPLEX(void, Execute, const FRigUnitContext& Context)
 {
     DECLARE_SCOPE_HIERARCHICAL_COUNTER_RIGUNIT()
 	Result = Transform.TransformVector(Direction);
 }
 
-UE_RigUnit_MathTransformTransformVector_IMPLEMENT_STATIC_VIRTUAL_METHOD(void, Execute, const FRigUnitContext& Context)
+UE_RigUnit_MathTransformTransformVector_IMPLEMENT_MULTIPLEX(void, Execute, const FRigUnitContext& Context)
 {
     DECLARE_SCOPE_HIERARCHICAL_COUNTER_RIGUNIT()
 	Result = Transform.TransformPosition(Location);
 }
 
-UE_RigUnit_MathTransformFromSRT_IMPLEMENT_STATIC_VIRTUAL_METHOD(void, Execute, const FRigUnitContext& Context)
+UE_RigUnit_MathTransformFromSRT_IMPLEMENT_MULTIPLEX(void, Execute, const FRigUnitContext& Context)
 {
     DECLARE_SCOPE_HIERARCHICAL_COUNTER_RIGUNIT()
 	Transform.SetLocation(Location);
diff --git a/Engine/Plugins/Experimental/ControlRig/Source/ControlRig/Private/Units/Math/RigUnit_MathTransform.h b/Engine/Plugins/Experimental/ControlRig/Source/ControlRig/Private/Units/Math/RigUnit_MathTransform.h
--- a/Engine/Plugins/Experimental/ControlRig/Source/ControlRig/Private/Units/Math/RigUnit_MathTransform.h
+++ b/Engine/Plugins/Experimental/ControlRig/Source/ControlRig/Private/Units/Math/RigUnit_MathTransform.h
@@ -64,7 +64,7 @@ struct FRigUnit_MathTransformFromEulerTransform : public FRigUnit_MathTransformB
 		Result = FTransform::Identity;
 	}
 
-	STATIC_VIRTUAL_METHOD()
+	MULTIPLEX_METHOD()
 	virtual void Execute(const FRigUnitContext& Context) override;
 
 	UPROPERTY(meta=(Input))
@@ -88,7 +88,7 @@ struct FRigUnit_MathTransformToEulerTransform : public FRigUnit_MathTransformBas
 		Result = FEulerTransform::Identity;
 	}
 
-	STATIC_VIRTUAL_METHOD()
+	MULTIPLEX_METHOD()
 	virtual void Execute(const FRigUnitContext& Context) override;
 
 	UPROPERTY(meta=(Input))
@@ -105,7 +105,7 @@ USTRUCT(meta=(DisplayName="Multiply", PrototypeName="Multiply", Keywords="Produc
 struct FRigUnit_MathTransformMul : public FRigUnit_MathTransformBinaryOp
 {
 	GENERATED_BODY()
-	STATIC_VIRTUAL_METHOD()
+	MULTIPLEX_METHOD()
 	virtual void Execute(const FRigUnitContext& Context) override;
 };
 
@@ -122,7 +122,7 @@ struct FRigUnit_MathTransformMakeRelative : public FRigUnit_MathTransformBase
 		Global = Parent = Local = FTransform::Identity;
 	}
 
-	STATIC_VIRTUAL_METHOD()
+	MULTIPLEX_METHOD()
 	virtual void Execute(const FRigUnitContext& Context) override;
 
 	UPROPERTY(meta=(Input))
@@ -142,7 +142,7 @@ USTRUCT(meta=(DisplayName="Inverse", PrototypeName="Inverse"))
 struct FRigUnit_MathTransformInverse : public FRigUnit_MathTransformUnaryOp
 {
 	GENERATED_BODY()
-	STATIC_VIRTUAL_METHOD()
+	MULTIPLEX_METHOD()
 	virtual void Execute(const FRigUnitContext& Context) override;
 };
 
@@ -160,7 +160,7 @@ struct FRigUnit_MathTransformLerp : public FRigUnit_MathTransformBase
 		T = 0.f;
 	}
 
-	STATIC_VIRTUAL_METHOD()
+	MULTIPLEX_METHOD()
 	virtual void Execute(const FRigUnitContext& Context) override;
 
 	UPROPERTY(meta=(Input))
@@ -190,7 +190,7 @@ struct FRigUnit_MathTransformSelectBool : public FRigUnit_MathTransformBase
 		IfTrue = IfFalse = Result = FTransform::Identity;
 	}
 
-	STATIC_VIRTUAL_METHOD()
+	MULTIPLEX_METHOD()
 	virtual void Execute(const FRigUnitContext& Context) override;
 
 	UPROPERTY(meta=(Input))
@@ -220,7 +220,7 @@ struct FRigUnit_MathTransformRotateVector : public FRigUnit_MathTransformBase
 		Direction = Result = FVector::ZeroVector;
 	}
 
-	STATIC_VIRTUAL_METHOD()
+	MULTIPLEX_METHOD()
 	virtual void Execute(const FRigUnitContext& Context) override;
 
 	UPROPERTY(meta=(Input))
@@ -247,7 +247,7 @@ struct FRigUnit_MathTransformTransformVector : public FRigUnit_MathTransformBase
 		Location = Result = FVector::ZeroVector;
 	}
 	
-	STATIC_VIRTUAL_METHOD()
+	MULTIPLEX_METHOD()
 	virtual void Execute(const FRigUnitContext& Context) override;
 
 	UPROPERTY(meta=(Input))
@@ -278,7 +278,7 @@ struct FRigUnit_MathTransformFromSRT : public FRigUnit_MathTransformBase
 		EulerTransform = FEulerTransform::Identity;
 	}
 	
-	STATIC_VIRTUAL_METHOD()
+	MULTIPLEX_METHOD()
 	virtual void Execute(const FRigUnitContext& Context) override;
 
 	UPROPERTY(meta=(Input))
diff --git a/Engine/Plugins/Experimental/ControlRig/Source/ControlRig/Private/Units/Math/RigUnit_MathVector.cpp b/Engine/Plugins/Experimental/ControlRig/Source/ControlRig/Private/Units/Math/RigUnit_MathVector.cpp
--- a/Engine/Plugins/Experimental/ControlRig/Source/ControlRig/Private/Units/Math/RigUnit_MathVector.cpp
+++ b/Engine/Plugins/Experimental/ControlRig/Source/ControlRig/Private/Units/Math/RigUnit_MathVector.cpp
@@ -4,37 +4,37 @@
 #include "Units/RigUnitContext.h"
 #include "Math/ControlRigMathLibrary.h"
 
-UE_RigUnit_MathVectorFromFloat_IMPLEMENT_STATIC_VIRTUAL_METHOD(void, Execute, const FRigUnitContext& Context)
+UE_RigUnit_MathVectorFromFloat_IMPLEMENT_MULTIPLEX(void, Execute, const FRigUnitContext& Context)
 {
     DECLARE_SCOPE_HIERARCHICAL_COUNTER_RIGUNIT()
 	Result = FVector(Value, Value, Value);
 }
 
-UE_RigUnit_MathVectorAdd_IMPLEMENT_STATIC_VIRTUAL_METHOD(void, Execute, const FRigUnitContext& Context)
+UE_RigUnit_MathVectorAdd_IMPLEMENT_MULTIPLEX(void, Execute, const FRigUnitContext& Context)
 {
     DECLARE_SCOPE_HIERARCHICAL_COUNTER_RIGUNIT()
 	Result = A + B;
 }
 
-UE_RigUnit_MathVectorSub_IMPLEMENT_STATIC_VIRTUAL_METHOD(void, Execute, const FRigUnitContext& Context)
+UE_RigUnit_MathVectorSub_IMPLEMENT_MULTIPLEX(void, Execute, const FRigUnitContext& Context)
 {
     DECLARE_SCOPE_HIERARCHICAL_COUNTER_RIGUNIT()
 	Result = A - B;
 }
 
-UE_RigUnit_MathVectorMul_IMPLEMENT_STATIC_VIRTUAL_METHOD(void, Execute, const FRigUnitContext& Context)
+UE_RigUnit_MathVectorMul_IMPLEMENT_MULTIPLEX(void, Execute, const FRigUnitContext& Context)
 {
     DECLARE_SCOPE_HIERARCHICAL_COUNTER_RIGUNIT()
 	Result = A * B;
 }
 
-UE_RigUnit_MathVectorScale_IMPLEMENT_STATIC_VIRTUAL_METHOD(void, Execute, const FRigUnitContext& Context)
+UE_RigUnit_MathVectorScale_IMPLEMENT_MULTIPLEX(void, Execute, const FRigUnitContext& Context)
 {
     DECLARE_SCOPE_HIERARCHICAL_COUNTER_RIGUNIT()
 	Result = Value * Factor;
 }
 
-UE_RigUnit_MathVectorDiv_IMPLEMENT_STATIC_VIRTUAL_METHOD(void, Execute, const FRigUnitContext& Context)
+UE_RigUnit_MathVectorDiv_IMPLEMENT_MULTIPLEX(void, Execute, const FRigUnitContext& Context)
 {
     DECLARE_SCOPE_HIERARCHICAL_COUNTER_RIGUNIT()
 	if(FMath::IsNearlyZero(B.X) || FMath::IsNearlyZero(B.Y) || FMath::IsNearlyZero(B.Z))
@@ -57,7 +57,7 @@ UE_RigUnit_MathVectorDiv_IMPLEMENT_STATIC_VIRTUAL_METHOD(void, Execute, const FR
 	Result = A / B;
 }
 
-UE_RigUnit_MathVectorMod_IMPLEMENT_STATIC_VIRTUAL_METHOD(void, Execute, const FRigUnitContext& Context)
+UE_RigUnit_MathVectorMod_IMPLEMENT_MULTIPLEX(void, Execute, const FRigUnitContext& Context)
 {
     DECLARE_SCOPE_HIERARCHICAL_COUNTER_RIGUNIT()
 	if(FMath::IsNearlyZero(B.X) || FMath::IsNearlyZero(B.Y) || FMath::IsNearlyZero(B.Z) || B.X < 0.f || B.Y < 0.f || B.Z < 0.f)
@@ -83,7 +83,7 @@ UE_RigUnit_MathVectorMod_IMPLEMENT_STATIC_VIRTUAL_METHOD(void, Execute, const FR
 	Result.Z = FMath::Fmod(A.Z, B.Z);
 }
 
-UE_RigUnit_MathVectorMin_IMPLEMENT_STATIC_VIRTUAL_METHOD(void, Execute, const FRigUnitContext& Context)
+UE_RigUnit_MathVectorMin_IMPLEMENT_MULTIPLEX(void, Execute, const FRigUnitContext& Context)
 {
     DECLARE_SCOPE_HIERARCHICAL_COUNTER_RIGUNIT()
 	Result.X = FMath::Min<float>(A.X, B.X);
@@ -91,7 +91,7 @@ UE_RigUnit_MathVectorMin_IMPLEMENT_STATIC_VIRTUAL_METHOD(void, Execute, const FR
 	Result.Z = FMath::Min<float>(A.Z, B.Z);
 }
 
-UE_RigUnit_MathVectorMax_IMPLEMENT_STATIC_VIRTUAL_METHOD(void, Execute, const FRigUnitContext& Context)
+UE_RigUnit_MathVectorMax_IMPLEMENT_MULTIPLEX(void, Execute, const FRigUnitContext& Context)
 {
     DECLARE_SCOPE_HIERARCHICAL_COUNTER_RIGUNIT()
 	Result.X = FMath::Max<float>(A.X, B.X);
@@ -99,13 +99,13 @@ UE_RigUnit_MathVectorMax_IMPLEMENT_STATIC_VIRTUAL_METHOD(void, Execute, const FR
 	Result.Z = FMath::Max<float>(A.Z, B.Z);
 }
 
-UE_RigUnit_MathVectorNegate_IMPLEMENT_STATIC_VIRTUAL_METHOD(void, Execute, const FRigUnitContext& Context)
+UE_RigUnit_MathVectorNegate_IMPLEMENT_MULTIPLEX(void, Execute, const FRigUnitContext& Context)
 {
     DECLARE_SCOPE_HIERARCHICAL_COUNTER_RIGUNIT()
 	Result = -Value;
 }
 
-UE_RigUnit_MathVectorAbs_IMPLEMENT_STATIC_VIRTUAL_METHOD(void, Execute, const FRigUnitContext& Context)
+UE_RigUnit_MathVectorAbs_IMPLEMENT_MULTIPLEX(void, Execute, const FRigUnitContext& Context)
 {
     DECLARE_SCOPE_HIERARCHICAL_COUNTER_RIGUNIT()
 	Result.X = FMath::Abs(Value.X);
@@ -113,7 +113,7 @@ UE_RigUnit_MathVectorAbs_IMPLEMENT_STATIC_VIRTUAL_METHOD(void, Execute, const FR
 	Result.Z = FMath::Abs(Value.Z);
 }
 
-UE_RigUnit_MathVectorFloor_IMPLEMENT_STATIC_VIRTUAL_METHOD(void, Execute, const FRigUnitContext& Context)
+UE_RigUnit_MathVectorFloor_IMPLEMENT_MULTIPLEX(void, Execute, const FRigUnitContext& Context)
 {
     DECLARE_SCOPE_HIERARCHICAL_COUNTER_RIGUNIT()
 	Result.X = FMath::FloorToFloat(Value.X);
@@ -121,7 +121,7 @@ UE_RigUnit_MathVectorFloor_IMPLEMENT_STATIC_VIRTUAL_METHOD(void, Execute, const
 	Result.Z = FMath::FloorToFloat(Value.Z);
 }
 
-UE_RigUnit_MathVectorCeil_IMPLEMENT_STATIC_VIRTUAL_METHOD(void, Execute, const FRigUnitContext& Context)
+UE_RigUnit_MathVectorCeil_IMPLEMENT_MULTIPLEX(void, Execute, const FRigUnitContext& Context)
 {
     DECLARE_SCOPE_HIERARCHICAL_COUNTER_RIGUNIT()
 	Result.X = FMath::CeilToFloat(Value.X);
@@ -129,7 +129,7 @@ UE_RigUnit_MathVectorCeil_IMPLEMENT_STATIC_VIRTUAL_METHOD(void, Execute, const F
 	Result.Z = FMath::CeilToFloat(Value.Z);
 }
 
-UE_RigUnit_MathVectorRound_IMPLEMENT_STATIC_VIRTUAL_METHOD(void, Execute, const FRigUnitContext& Context)
+UE_RigUnit_MathVectorRound_IMPLEMENT_MULTIPLEX(void, Execute, const FRigUnitContext& Context)
 {
     DECLARE_SCOPE_HIERARCHICAL_COUNTER_RIGUNIT()
 	Result.X = FMath::RoundToFloat(Value.X);
@@ -137,13 +137,13 @@ UE_RigUnit_MathVectorRound_IMPLEMENT_STATIC_VIRTUAL_METHOD(void, Execute, const
 	Result.Z = FMath::RoundToFloat(Value.Z);
 }
 
-UE_RigUnit_MathVectorSign_IMPLEMENT_STATIC_VIRTUAL_METHOD(void, Execute, const FRigUnitContext& Context)
+UE_RigUnit_MathVectorSign_IMPLEMENT_MULTIPLEX(void, Execute, const FRigUnitContext& Context)
 {
     DECLARE_SCOPE_HIERARCHICAL_COUNTER_RIGUNIT()
 	Result = Value.GetSignVector();
 }
 
-UE_RigUnit_MathVectorClamp_IMPLEMENT_STATIC_VIRTUAL_METHOD(void, Execute, const FRigUnitContext& Context)
+UE_RigUnit_MathVectorClamp_IMPLEMENT_MULTIPLEX(void, Execute, const FRigUnitContext& Context)
 {
     DECLARE_SCOPE_HIERARCHICAL_COUNTER_RIGUNIT()
 	Result.X = FMath::Clamp<float>(Value.X, Minimum.X, Maximum.X);
@@ -151,13 +151,13 @@ UE_RigUnit_MathVectorClamp_IMPLEMENT_STATIC_VIRTUAL_METHOD(void, Execute, const
 	Result.Z = FMath::Clamp<float>(Value.Z, Minimum.Z, Maximum.Z);
 }
 
-UE_RigUnit_MathVectorLerp_IMPLEMENT_STATIC_VIRTUAL_METHOD(void, Execute, const FRigUnitContext& Context)
+UE_RigUnit_MathVectorLerp_IMPLEMENT_MULTIPLEX(void, Execute, const FRigUnitContext& Context)
 {
     DECLARE_SCOPE_HIERARCHICAL_COUNTER_RIGUNIT()
 	Result = FMath::Lerp<FVector>(A, B, T);
 }
 
-UE_RigUnit_MathVectorRemap_IMPLEMENT_STATIC_VIRTUAL_METHOD(void, Execute, const FRigUnitContext& Context)
+UE_RigUnit_MathVectorRemap_IMPLEMENT_MULTIPLEX(void, Execute, const FRigUnitContext& Context)
 {
     DECLARE_SCOPE_HIERARCHICAL_COUNTER_RIGUNIT()
 	FVector Ratio(0.f, 0.f, 0.f);
@@ -189,19 +189,19 @@ UE_RigUnit_MathVectorRemap_IMPLEMENT_STATIC_VIRTUAL_METHOD(void, Execute, const
 	Result = FMath::Lerp<FVector>(TargetMinimum, TargetMaximum, Ratio);
 }
 
-UE_RigUnit_MathVectorEquals_IMPLEMENT_STATIC_VIRTUAL_METHOD(void, Execute, const FRigUnitContext& Context)
+UE_RigUnit_MathVectorEquals_IMPLEMENT_MULTIPLEX(void, Execute, const FRigUnitContext& Context)
 {
     DECLARE_SCOPE_HIERARCHICAL_COUNTER_RIGUNIT()
 	Result = A == B;
 }
 
-UE_RigUnit_MathVectorNotEquals_IMPLEMENT_STATIC_VIRTUAL_METHOD(void, Execute, const FRigUnitContext& Context)
+UE_RigUnit_MathVectorNotEquals_IMPLEMENT_MULTIPLEX(void, Execute, const FRigUnitContext& Context)
 {
     DECLARE_SCOPE_HIERARCHICAL_COUNTER_RIGUNIT()
 	Result = A != B;
 }
 
-UE_RigUnit_MathVectorIsNearlyZero_IMPLEMENT_STATIC_VIRTUAL_METHOD(void, Execute, const FRigUnitContext& Context)
+UE_RigUnit_MathVectorIsNearlyZero_IMPLEMENT_MULTIPLEX(void, Execute, const FRigUnitContext& Context)
 {
     DECLARE_SCOPE_HIERARCHICAL_COUNTER_RIGUNIT()
 	if(Tolerance < 0.f)
@@ -211,7 +211,7 @@ UE_RigUnit_MathVectorIsNearlyZero_IMPLEMENT_STATIC_VIRTUAL_METHOD(void, Execute,
 	Result = Value.IsNearlyZero(FMath::Max<float>(Tolerance, SMALL_NUMBER));
 }
 
-UE_RigUnit_MathVectorIsNearlyEqual_IMPLEMENT_STATIC_VIRTUAL_METHOD(void, Execute, const FRigUnitContext& Context)
+UE_RigUnit_MathVectorIsNearlyEqual_IMPLEMENT_MULTIPLEX(void, Execute, const FRigUnitContext& Context)
 {
     DECLARE_SCOPE_HIERARCHICAL_COUNTER_RIGUNIT()
 	if(Tolerance < 0.f)
@@ -221,55 +221,55 @@ UE_RigUnit_MathVectorIsNearlyEqual_IMPLEMENT_STATIC_VIRTUAL_METHOD(void, Execute
 	Result = (A - B).IsNearlyZero(FMath::Max<float>(Tolerance, SMALL_NUMBER));
 }
 
-UE_RigUnit_MathVectorSelectBool_IMPLEMENT_STATIC_VIRTUAL_METHOD(void, Execute, const FRigUnitContext& Context)
+UE_RigUnit_MathVectorSelectBool_IMPLEMENT_MULTIPLEX(void, Execute, const FRigUnitContext& Context)
 {
     DECLARE_SCOPE_HIERARCHICAL_COUNTER_RIGUNIT()
 	Result = Condition ? IfTrue : IfFalse;
 }
 
-UE_RigUnit_MathVectorDeg_IMPLEMENT_STATIC_VIRTUAL_METHOD(void, Execute, const FRigUnitContext& Context)
+UE_RigUnit_MathVectorDeg_IMPLEMENT_MULTIPLEX(void, Execute, const FRigUnitContext& Context)
 {
     DECLARE_SCOPE_HIERARCHICAL_COUNTER_RIGUNIT()
 	Result = FMath::RadiansToDegrees(Value);
 }
 
-UE_RigUnit_MathVectorRad_IMPLEMENT_STATIC_VIRTUAL_METHOD(void, Execute, const FRigUnitContext& Context)
+UE_RigUnit_MathVectorRad_IMPLEMENT_MULTIPLEX(void, Execute, const FRigUnitContext& Context)
 {
     DECLARE_SCOPE_HIERARCHICAL_COUNTER_RIGUNIT()
 	Result = FMath::DegreesToRadians(Value);
 }
 
-UE_RigUnit_MathVectorLengthSquared_IMPLEMENT_STATIC_VIRTUAL_METHOD(void, Execute, const FRigUnitContext& Context)
+UE_RigUnit_MathVectorLengthSquared_IMPLEMENT_MULTIPLEX(void, Execute, const FRigUnitContext& Context)
 {
     DECLARE_SCOPE_HIERARCHICAL_COUNTER_RIGUNIT()
 	Result = Value.SizeSquared();
 }
 
-UE_RigUnit_MathVectorLength_IMPLEMENT_STATIC_VIRTUAL_METHOD(void, Execute, const FRigUnitContext& Context)
+UE_RigUnit_MathVectorLength_IMPLEMENT_MULTIPLEX(void, Execute, const FRigUnitContext& Context)
 {
     DECLARE_SCOPE_HIERARCHICAL_COUNTER_RIGUNIT()
 	Result = Value.Size();
 }
 
-UE_RigUnit_MathVectorDistance_IMPLEMENT_STATIC_VIRTUAL_METHOD(void, Execute, const FRigUnitContext& Context)
+UE_RigUnit_MathVectorDistance_IMPLEMENT_MULTIPLEX(void, Execute, const FRigUnitContext& Context)
 {
     DECLARE_SCOPE_HIERARCHICAL_COUNTER_RIGUNIT()
 	Result = FVector::Distance(A, B);
 }
 
-UE_RigUnit_MathVectorCross_IMPLEMENT_STATIC_VIRTUAL_METHOD(void, Execute, const FRigUnitContext& Context)
+UE_RigUnit_MathVectorCross_IMPLEMENT_MULTIPLEX(void, Execute, const FRigUnitContext& Context)
 {
     DECLARE_SCOPE_HIERARCHICAL_COUNTER_RIGUNIT()
 	Result = A ^ B;
 }
 
-UE_RigUnit_MathVectorDot_IMPLEMENT_STATIC_VIRTUAL_METHOD(void, Execute, const FRigUnitContext& Context)
+UE_RigUnit_MathVectorDot_IMPLEMENT_MULTIPLEX(void, Execute, const FRigUnitContext& Context)
 {
     DECLARE_SCOPE_HIERARCHICAL_COUNTER_RIGUNIT()
 	Result = A | B;
 }
 
-UE_RigUnit_MathVectorUnit_IMPLEMENT_STATIC_VIRTUAL_METHOD(void, Execute, const FRigUnitContext& Context)
+UE_RigUnit_MathVectorUnit_IMPLEMENT_MULTIPLEX(void, Execute, const FRigUnitContext& Context)
 {
     DECLARE_SCOPE_HIERARCHICAL_COUNTER_RIGUNIT()
 	if (Value.IsNearlyZero())
@@ -281,7 +281,7 @@ UE_RigUnit_MathVectorUnit_IMPLEMENT_STATIC_VIRTUAL_METHOD(void, Execute, const F
 	Result = Value.GetUnsafeNormal();
 }
 
-UE_RigUnit_MathVectorMirror_IMPLEMENT_STATIC_VIRTUAL_METHOD(void, Execute, const FRigUnitContext& Context)
+UE_RigUnit_MathVectorMirror_IMPLEMENT_MULTIPLEX(void, Execute, const FRigUnitContext& Context)
 {
     DECLARE_SCOPE_HIERARCHICAL_COUNTER_RIGUNIT()
 	if (Normal.IsNearlyZero())
@@ -293,7 +293,7 @@ UE_RigUnit_MathVectorMirror_IMPLEMENT_STATIC_VIRTUAL_METHOD(void, Execute, const
 	Result = Value.MirrorByVector(Normal.GetSafeNormal());
 }
 
-UE_RigUnit_MathVectorAngle_IMPLEMENT_STATIC_VIRTUAL_METHOD(void, Execute, const FRigUnitContext& Context)
+UE_RigUnit_MathVectorAngle_IMPLEMENT_MULTIPLEX(void, Execute, const FRigUnitContext& Context)
 {
     DECLARE_SCOPE_HIERARCHICAL_COUNTER_RIGUNIT()
 	if (A.IsNearlyZero() || B.IsNearlyZero())
@@ -312,7 +312,7 @@ UE_RigUnit_MathVectorAngle_IMPLEMENT_STATIC_VIRTUAL_METHOD(void, Execute, const
 	Result = FQuat::FindBetween(A, B).GetAngle();
 }
 
-UE_RigUnit_MathVectorParallel_IMPLEMENT_STATIC_VIRTUAL_METHOD(void, Execute, const FRigUnitContext& Context)
+UE_RigUnit_MathVectorParallel_IMPLEMENT_MULTIPLEX(void, Execute, const FRigUnitContext& Context)
 {
     DECLARE_SCOPE_HIERARCHICAL_COUNTER_RIGUNIT()
 	if (A.IsNearlyZero() || B.IsNearlyZero())
@@ -331,7 +331,7 @@ UE_RigUnit_MathVectorParallel_IMPLEMENT_STATIC_VIRTUAL_METHOD(void, Execute, con
 	Result = FVector::Parallel(A, B);
 }
 
-UE_RigUnit_MathVectorOrthogonal_IMPLEMENT_STATIC_VIRTUAL_METHOD(void, Execute, const FRigUnitContext& Context)
+UE_RigUnit_MathVectorOrthogonal_IMPLEMENT_MULTIPLEX(void, Execute, const FRigUnitContext& Context)
 {
     DECLARE_SCOPE_HIERARCHICAL_COUNTER_RIGUNIT()
 	if (A.IsNearlyZero() || B.IsNearlyZero())
@@ -350,7 +350,7 @@ UE_RigUnit_MathVectorOrthogonal_IMPLEMENT_STATIC_VIRTUAL_METHOD(void, Execute, c
 	Result = FVector::Orthogonal(A, B);
 }
 
-UE_RigUnit_MathVectorBezierFourPoint_IMPLEMENT_STATIC_VIRTUAL_METHOD(void, Execute, const FRigUnitContext& Context)
+UE_RigUnit_MathVectorBezierFourPoint_IMPLEMENT_MULTIPLEX(void, Execute, const FRigUnitContext& Context)
 {
     DECLARE_SCOPE_HIERARCHICAL_COUNTER_RIGUNIT()
 	FControlRigMathLibrary::FourPointBezier(Bezier, T, Result, Tangent);
diff --git a/Engine/Plugins/Experimental/ControlRig/Source/ControlRig/Private/Units/Math/RigUnit_MathVector.h b/Engine/Plugins/Experimental/ControlRig/Source/ControlRig/Private/Units/Math/RigUnit_MathVector.h
--- a/Engine/Plugins/Experimental/ControlRig/Source/ControlRig/Private/Units/Math/RigUnit_MathVector.h
+++ b/Engine/Plugins/Experimental/ControlRig/Source/ControlRig/Private/Units/Math/RigUnit_MathVector.h
@@ -56,7 +56,7 @@ USTRUCT(meta=(DisplayName="From Float", PrototypeName="FromFloat", Keywords="Mak
 struct FRigUnit_MathVectorFromFloat : public FRigUnit_MathVectorBase
 {
 	GENERATED_BODY()
-	STATIC_VIRTUAL_METHOD()
+	MULTIPLEX_METHOD()
 	virtual void Execute(const FRigUnitContext& Context) override;
 
 	FRigUnit_MathVectorFromFloat()
@@ -80,7 +80,7 @@ struct FRigUnit_MathVectorAdd : public FRigUnit_MathVectorBinaryOp
 {
 	GENERATED_BODY()
 
-	STATIC_VIRTUAL_METHOD()
+	MULTIPLEX_METHOD()
 	virtual void Execute(const FRigUnitContext& Context) override;
 };
 
@@ -92,7 +92,7 @@ struct FRigUnit_MathVectorSub : public FRigUnit_MathVectorBinaryOp
 {
 	GENERATED_BODY()
 
-	STATIC_VIRTUAL_METHOD()
+	MULTIPLEX_METHOD()
 	virtual void Execute(const FRigUnitContext& Context) override;
 };
 
@@ -109,7 +109,7 @@ struct FRigUnit_MathVectorMul : public FRigUnit_MathVectorBinaryOp
 		A = B = FVector::OneVector;
 	}
 
-	STATIC_VIRTUAL_METHOD()
+	MULTIPLEX_METHOD()
 	virtual void Execute(const FRigUnitContext& Context) override;
 };
 
@@ -127,7 +127,7 @@ struct FRigUnit_MathVectorScale : public FRigUnit_MathVectorBase
 		Factor = 1.f;
 	}
 
-	STATIC_VIRTUAL_METHOD()
+	MULTIPLEX_METHOD()
 	virtual void Execute(const FRigUnitContext& Context) override;
 
 	UPROPERTY(meta = (Input))
@@ -153,7 +153,7 @@ struct FRigUnit_MathVectorDiv : public FRigUnit_MathVectorBinaryOp
 		B = FVector::OneVector;
 	}
 
-	STATIC_VIRTUAL_METHOD()
+	MULTIPLEX_METHOD()
 	virtual void Execute(const FRigUnitContext& Context) override;
 };
 
@@ -171,7 +171,7 @@ struct FRigUnit_MathVectorMod : public FRigUnit_MathVectorBinaryOp
 		B = FVector::OneVector;
 	}
 
-	STATIC_VIRTUAL_METHOD()
+	MULTIPLEX_METHOD()
 	virtual void Execute(const FRigUnitContext& Context) override;
 };
 
@@ -183,7 +183,7 @@ struct FRigUnit_MathVectorMin : public FRigUnit_MathVectorBinaryOp
 {
 	GENERATED_BODY()
 
-	STATIC_VIRTUAL_METHOD()
+	MULTIPLEX_METHOD()
 	virtual void Execute(const FRigUnitContext& Context) override;
 };
 
@@ -195,7 +195,7 @@ struct FRigUnit_MathVectorMax : public FRigUnit_MathVectorBinaryOp
 {
 	GENERATED_BODY()
 
-	STATIC_VIRTUAL_METHOD()
+	MULTIPLEX_METHOD()
 	virtual void Execute(const FRigUnitContext& Context) override;
 };
 
@@ -207,7 +207,7 @@ struct FRigUnit_MathVectorNegate : public FRigUnit_MathVectorUnaryOp
 {
 	GENERATED_BODY()
 
-	STATIC_VIRTUAL_METHOD()
+	MULTIPLEX_METHOD()
 	virtual void Execute(const FRigUnitContext& Context) override;
 };
 
@@ -219,7 +219,7 @@ struct FRigUnit_MathVectorAbs : public FRigUnit_MathVectorUnaryOp
 {
 	GENERATED_BODY()
 
-	STATIC_VIRTUAL_METHOD()
+	MULTIPLEX_METHOD()
 	virtual void Execute(const FRigUnitContext& Context) override;
 };
 
@@ -231,7 +231,7 @@ struct FRigUnit_MathVectorFloor : public FRigUnit_MathVectorUnaryOp
 {
 	GENERATED_BODY()
 
-	STATIC_VIRTUAL_METHOD()
+	MULTIPLEX_METHOD()
 	virtual void Execute(const FRigUnitContext& Context) override;
 };
 
@@ -243,7 +243,7 @@ struct FRigUnit_MathVectorCeil : public FRigUnit_MathVectorUnaryOp
 {
 	GENERATED_BODY()
 
-	STATIC_VIRTUAL_METHOD()
+	MULTIPLEX_METHOD()
 	virtual void Execute(const FRigUnitContext& Context) override;
 };
 
@@ -255,7 +255,7 @@ struct FRigUnit_MathVectorRound : public FRigUnit_MathVectorUnaryOp
 {
 	GENERATED_BODY()
 
-	STATIC_VIRTUAL_METHOD()
+	MULTIPLEX_METHOD()
 	virtual void Execute(const FRigUnitContext& Context) override;
 };
 
@@ -267,7 +267,7 @@ struct FRigUnit_MathVectorSign : public FRigUnit_MathVectorUnaryOp
 {
 	GENERATED_BODY()
 
-	STATIC_VIRTUAL_METHOD()
+	MULTIPLEX_METHOD()
 	virtual void Execute(const FRigUnitContext& Context) override;
 };
 
@@ -285,7 +285,7 @@ struct FRigUnit_MathVectorClamp : public FRigUnit_MathVectorBase
 		Maximum = FVector::OneVector;
 	}
 
-	STATIC_VIRTUAL_METHOD()
+	MULTIPLEX_METHOD()
 	virtual void Execute(const FRigUnitContext& Context) override;
 
 	UPROPERTY(meta=(Input))
@@ -316,7 +316,7 @@ struct FRigUnit_MathVectorLerp : public FRigUnit_MathVectorBase
 		T = 0.f;
 	}
 
-	STATIC_VIRTUAL_METHOD()
+	MULTIPLEX_METHOD()
 	virtual void Execute(const FRigUnitContext& Context) override;
 
 	UPROPERTY(meta=(Input))
@@ -347,7 +347,7 @@ struct FRigUnit_MathVectorRemap : public FRigUnit_MathVectorBase
 		bClamp = false;
 	}
 
-	STATIC_VIRTUAL_METHOD()
+	MULTIPLEX_METHOD()
 	virtual void Execute(const FRigUnitContext& Context) override;
 
 	UPROPERTY(meta=(Input))
@@ -387,7 +387,7 @@ struct FRigUnit_MathVectorEquals : public FRigUnit_MathVectorBase
 		Result = true;
 	}
 
-	STATIC_VIRTUAL_METHOD()
+	MULTIPLEX_METHOD()
 	virtual void Execute(const FRigUnitContext& Context) override;
 
 	UPROPERTY(meta=(Input))
@@ -414,7 +414,7 @@ struct FRigUnit_MathVectorNotEquals : public FRigUnit_MathVectorBase
 		Result = false;
 	}
 
-	STATIC_VIRTUAL_METHOD()
+	MULTIPLEX_METHOD()
 	virtual void Execute(const FRigUnitContext& Context) override;
 
 	UPROPERTY(meta=(Input))
@@ -442,7 +442,7 @@ struct FRigUnit_MathVectorIsNearlyZero : public FRigUnit_MathVectorBase
 		Result = true;
 	}
 
-	STATIC_VIRTUAL_METHOD()
+	MULTIPLEX_METHOD()
 	virtual void Execute(const FRigUnitContext& Context) override;
 
 	UPROPERTY(meta=(Input))
@@ -470,7 +470,7 @@ struct FRigUnit_MathVectorIsNearlyEqual : public FRigUnit_MathVectorBase
 		Result = true;
 	}
 
-	STATIC_VIRTUAL_METHOD()
+	MULTIPLEX_METHOD()
 	virtual void Execute(const FRigUnitContext& Context) override;
 
 	UPROPERTY(meta=(Input))
@@ -500,7 +500,7 @@ struct FRigUnit_MathVectorSelectBool : public FRigUnit_MathVectorBase
 		IfTrue = IfFalse = Result = FVector::ZeroVector;
 	}
 
-	STATIC_VIRTUAL_METHOD()
+	MULTIPLEX_METHOD()
 	virtual void Execute(const FRigUnitContext& Context) override;
 
 	UPROPERTY(meta=(Input))
@@ -524,7 +524,7 @@ struct FRigUnit_MathVectorDeg : public FRigUnit_MathVectorUnaryOp
 {
 	GENERATED_BODY()
 
-	STATIC_VIRTUAL_METHOD()
+	MULTIPLEX_METHOD()
 	virtual void Execute(const FRigUnitContext& Context) override;
 };
 
@@ -536,7 +536,7 @@ struct FRigUnit_MathVectorRad : public FRigUnit_MathVectorUnaryOp
 {
 	GENERATED_BODY()
 
-	STATIC_VIRTUAL_METHOD()
+	MULTIPLEX_METHOD()
 	virtual void Execute(const FRigUnitContext& Context) override;
 };
 
@@ -554,7 +554,7 @@ struct FRigUnit_MathVectorLengthSquared : public FRigUnit_MathVectorBase
 		Result = 0.f;
 	}
 
-	STATIC_VIRTUAL_METHOD()
+	MULTIPLEX_METHOD()
 	virtual void Execute(const FRigUnitContext& Context) override;
 
 	UPROPERTY(meta=(Input))
@@ -578,7 +578,7 @@ struct FRigUnit_MathVectorLength : public FRigUnit_MathVectorBase
 		Result = 0.f;
 	}
 
-	STATIC_VIRTUAL_METHOD()
+	MULTIPLEX_METHOD()
 	virtual void Execute(const FRigUnitContext& Context) override;
 
 	UPROPERTY(meta=(Input))
@@ -602,7 +602,7 @@ struct FRigUnit_MathVectorDistance : public FRigUnit_MathVectorBase
 		Result = 0.f;
 	}
 
-	STATIC_VIRTUAL_METHOD()
+	MULTIPLEX_METHOD()
 	virtual void Execute(const FRigUnitContext& Context) override;
 
 	UPROPERTY(meta=(Input))
@@ -622,7 +622,7 @@ USTRUCT(meta=(DisplayName="Cross", PrototypeName="Cross", Keywords="^"))
 struct FRigUnit_MathVectorCross : public FRigUnit_MathVectorBinaryOp
 {
 	GENERATED_BODY()
-	STATIC_VIRTUAL_METHOD()
+	MULTIPLEX_METHOD()
 	virtual void Execute(const FRigUnitContext& Context) override;
 };
 
@@ -640,7 +640,7 @@ struct FRigUnit_MathVectorDot : public FRigUnit_MathVectorBase
 		Result = 0.f;
 	}
 
-	STATIC_VIRTUAL_METHOD()
+	MULTIPLEX_METHOD()
 	virtual void Execute(const FRigUnitContext& Context) override;
 
 	UPROPERTY(meta=(Input))
@@ -661,7 +661,7 @@ struct FRigUnit_MathVectorUnit : public FRigUnit_MathVectorUnaryOp
 {
 	GENERATED_BODY()
 
-	STATIC_VIRTUAL_METHOD()
+	MULTIPLEX_METHOD()
 	virtual void Execute(const FRigUnitContext& Context) override;
 };
 
@@ -679,7 +679,7 @@ struct FRigUnit_MathVectorMirror : public FRigUnit_MathVectorBase
 		Normal = FVector(1.f, 0.f, 0.f);
 	}
 
-	STATIC_VIRTUAL_METHOD()
+	MULTIPLEX_METHOD()
 	virtual void Execute(const FRigUnitContext& Context) override;
 
 	UPROPERTY(meta=(Input))
@@ -706,7 +706,7 @@ struct FRigUnit_MathVectorAngle : public FRigUnit_MathVectorBase
 		Result = 0.f;
 	}
 
-	STATIC_VIRTUAL_METHOD()
+	MULTIPLEX_METHOD()
 	virtual void Execute(const FRigUnitContext& Context) override;
 
 	UPROPERTY(meta=(Input))
@@ -733,7 +733,7 @@ struct FRigUnit_MathVectorParallel : public FRigUnit_MathVectorBase
 		Result = true;
 	}
 
-	STATIC_VIRTUAL_METHOD()
+	MULTIPLEX_METHOD()
 	virtual void Execute(const FRigUnitContext& Context) override;
 
 	UPROPERTY(meta=(Input))
@@ -761,7 +761,7 @@ struct FRigUnit_MathVectorOrthogonal : public FRigUnit_MathVectorBase
 		Result = true;
 	}
 
-	STATIC_VIRTUAL_METHOD()
+	MULTIPLEX_METHOD()
 	virtual void Execute(const FRigUnitContext& Context) override;
 
 	UPROPERTY(meta=(Input))
@@ -789,7 +789,7 @@ struct FRigUnit_MathVectorBezierFourPoint : public FRigUnit_MathVectorBase
 		Result = Tangent = FVector::ZeroVector;
 	}
 
-	STATIC_VIRTUAL_METHOD()
+	MULTIPLEX_METHOD()
 	virtual void Execute(const FRigUnitContext& Context) override;
 
 	UPROPERTY(meta=(Input))
diff --git a/Engine/Plugins/Experimental/ControlRig/Source/ControlRig/Private/Units/Math/RigUnit_Noise.cpp b/Engine/Plugins/Experimental/ControlRig/Source/ControlRig/Private/Units/Math/RigUnit_Noise.cpp
--- a/Engine/Plugins/Experimental/ControlRig/Source/ControlRig/Private/Units/Math/RigUnit_Noise.cpp
+++ b/Engine/Plugins/Experimental/ControlRig/Source/ControlRig/Private/Units/Math/RigUnit_Noise.cpp
@@ -3,7 +3,7 @@
 #include "Units/Math/RigUnit_Noise.h"
 #include "Units/RigUnitContext.h"
 
-UE_RigUnit_NoiseFloat_IMPLEMENT_STATIC_VIRTUAL_METHOD(void, Execute, const FRigUnitContext& Context)
+UE_RigUnit_NoiseFloat_IMPLEMENT_MULTIPLEX(void, Execute, const FRigUnitContext& Context)
 {
     DECLARE_SCOPE_HIERARCHICAL_COUNTER_RIGUNIT()
 	if (Context.State == EControlRigState::Init)
@@ -17,7 +17,7 @@ UE_RigUnit_NoiseFloat_IMPLEMENT_STATIC_VIRTUAL_METHOD(void, Execute, const FRigU
 	Time = Time + Speed * Context.DeltaTime;
 }
 
-UE_RigUnit_NoiseVector_IMPLEMENT_STATIC_VIRTUAL_METHOD(void, Execute, const FRigUnitContext& Context)
+UE_RigUnit_NoiseVector_IMPLEMENT_MULTIPLEX(void, Execute, const FRigUnitContext& Context)
 {
     DECLARE_SCOPE_HIERARCHICAL_COUNTER_RIGUNIT()
 	if (Context.State == EControlRigState::Init)
diff --git a/Engine/Plugins/Experimental/ControlRig/Source/ControlRig/Private/Units/Math/RigUnit_Noise.h b/Engine/Plugins/Experimental/ControlRig/Source/ControlRig/Private/Units/Math/RigUnit_Noise.h
--- a/Engine/Plugins/Experimental/ControlRig/Source/ControlRig/Private/Units/Math/RigUnit_Noise.h
+++ b/Engine/Plugins/Experimental/ControlRig/Source/ControlRig/Private/Units/Math/RigUnit_Noise.h
@@ -17,7 +17,7 @@ struct FRigUnit_NoiseFloat : public FRigUnit_MathBase
 		Frequency = Maximum = 1.f;
 	}
 
-	STATIC_VIRTUAL_METHOD()
+	MULTIPLEX_METHOD()
 	virtual void Execute(const FRigUnitContext& Context) override;
 
 	UPROPERTY(meta = (Input))
@@ -56,7 +56,7 @@ struct FRigUnit_NoiseVector : public FRigUnit_MathBase
 		Maximum = 1.f;
 	}
 
-	STATIC_VIRTUAL_METHOD()
+	MULTIPLEX_METHOD()
 	virtual void Execute(const FRigUnitContext& Context) override;
 
 	UPROPERTY(meta = (Input))
diff --git a/Engine/Plugins/Experimental/ControlRig/Source/ControlRig/Private/Units/Math/RigUnit_Random.cpp b/Engine/Plugins/Experimental/ControlRig/Source/ControlRig/Private/Units/Math/RigUnit_Random.cpp
--- a/Engine/Plugins/Experimental/ControlRig/Source/ControlRig/Private/Units/Math/RigUnit_Random.cpp
+++ b/Engine/Plugins/Experimental/ControlRig/Source/ControlRig/Private/Units/Math/RigUnit_Random.cpp
@@ -15,7 +15,7 @@ float FRigUnit_Random_Helper(int32& Seed)
 	return FPlatformMath::Fractional(Result.f);
 }
 
-UE_RigUnit_RandomFloat_IMPLEMENT_STATIC_VIRTUAL_METHOD(void, Execute, const FRigUnitContext& Context)
+UE_RigUnit_RandomFloat_IMPLEMENT_MULTIPLEX(void, Execute, const FRigUnitContext& Context)
 {
     DECLARE_SCOPE_HIERARCHICAL_COUNTER_RIGUNIT()
 	if (Context.State == EControlRigState::Init)
@@ -38,7 +38,7 @@ UE_RigUnit_RandomFloat_IMPLEMENT_STATIC_VIRTUAL_METHOD(void, Execute, const FRig
 	LastResult = Result;
 }
 
-UE_RigUnit_RandomVector_IMPLEMENT_STATIC_VIRTUAL_METHOD(void, Execute, const FRigUnitContext& Context)
+UE_RigUnit_RandomVector_IMPLEMENT_MULTIPLEX(void, Execute, const FRigUnitContext& Context)
 {
     DECLARE_SCOPE_HIERARCHICAL_COUNTER_RIGUNIT()
 	if (Context.State == EControlRigState::Init)
diff --git a/Engine/Plugins/Experimental/ControlRig/Source/ControlRig/Private/Units/Math/RigUnit_Random.h b/Engine/Plugins/Experimental/ControlRig/Source/ControlRig/Private/Units/Math/RigUnit_Random.h
--- a/Engine/Plugins/Experimental/ControlRig/Source/ControlRig/Private/Units/Math/RigUnit_Random.h
+++ b/Engine/Plugins/Experimental/ControlRig/Source/ControlRig/Private/Units/Math/RigUnit_Random.h
@@ -19,7 +19,7 @@ struct FRigUnit_RandomFloat : public FRigUnit_MathBase
 		TimeLeft = 0.f;
 	}
 
-	STATIC_VIRTUAL_METHOD()
+	MULTIPLEX_METHOD()
 	virtual void Execute(const FRigUnitContext& Context) override;
 
 	UPROPERTY(meta = (Input, Constant))
@@ -63,7 +63,7 @@ struct FRigUnit_RandomVector: public FRigUnit_MathBase
 		Result = LastResult = FVector::ZeroVector;
 	}
 
-	STATIC_VIRTUAL_METHOD()
+	MULTIPLEX_METHOD()
 	virtual void Execute(const FRigUnitContext& Context) override;
 
 	UPROPERTY(meta = (Input))
diff --git a/Engine/Plugins/Experimental/ControlRig/Source/ControlRig/Private/Units/Simulation/RigUnit_Accumulate.cpp b/Engine/Plugins/Experimental/ControlRig/Source/ControlRig/Private/Units/Simulation/RigUnit_Accumulate.cpp
--- a/Engine/Plugins/Experimental/ControlRig/Source/ControlRig/Private/Units/Simulation/RigUnit_Accumulate.cpp
+++ b/Engine/Plugins/Experimental/ControlRig/Source/ControlRig/Private/Units/Simulation/RigUnit_Accumulate.cpp
@@ -3,7 +3,7 @@
 #include "Units/Simulation/RigUnit_Accumulate.h"
 #include "Units/RigUnitContext.h"
 
-UE_RigUnit_AccumulateFloatAdd_IMPLEMENT_STATIC_VIRTUAL_METHOD(void, Execute, const FRigUnitContext& Context)
+UE_RigUnit_AccumulateFloatAdd_IMPLEMENT_MULTIPLEX(void, Execute, const FRigUnitContext& Context)
 {
     DECLARE_SCOPE_HIERARCHICAL_COUNTER_RIGUNIT()
 	if (Context.State == EControlRigState::Init)
@@ -20,7 +20,7 @@ UE_RigUnit_AccumulateFloatAdd_IMPLEMENT_STATIC_VIRTUAL_METHOD(void, Execute, con
 	}
 }
 
-UE_RigUnit_AccumulateVectorAdd_IMPLEMENT_STATIC_VIRTUAL_METHOD(void, Execute, const FRigUnitContext& Context)
+UE_RigUnit_AccumulateVectorAdd_IMPLEMENT_MULTIPLEX(void, Execute, const FRigUnitContext& Context)
 {
     DECLARE_SCOPE_HIERARCHICAL_COUNTER_RIGUNIT()
 	if (Context.State == EControlRigState::Init)
@@ -37,7 +37,7 @@ UE_RigUnit_AccumulateVectorAdd_IMPLEMENT_STATIC_VIRTUAL_METHOD(void, Execute, co
 	}
 }
 
-UE_RigUnit_AccumulateFloatMul_IMPLEMENT_STATIC_VIRTUAL_METHOD(void, Execute, const FRigUnitContext& Context)
+UE_RigUnit_AccumulateFloatMul_IMPLEMENT_MULTIPLEX(void, Execute, const FRigUnitContext& Context)
 {
     DECLARE_SCOPE_HIERARCHICAL_COUNTER_RIGUNIT()
 	if (Context.State == EControlRigState::Init)
@@ -55,7 +55,7 @@ UE_RigUnit_AccumulateFloatMul_IMPLEMENT_STATIC_VIRTUAL_METHOD(void, Execute, con
 	}
 }
 
-UE_RigUnit_AccumulateVectorMul_IMPLEMENT_STATIC_VIRTUAL_METHOD(void, Execute, const FRigUnitContext& Context)
+UE_RigUnit_AccumulateVectorMul_IMPLEMENT_MULTIPLEX(void, Execute, const FRigUnitContext& Context)
 {
     DECLARE_SCOPE_HIERARCHICAL_COUNTER_RIGUNIT()
 	if (Context.State == EControlRigState::Init)
@@ -73,7 +73,7 @@ UE_RigUnit_AccumulateVectorMul_IMPLEMENT_STATIC_VIRTUAL_METHOD(void, Execute, co
 	}
 }
 
-UE_RigUnit_AccumulateQuatMul_IMPLEMENT_STATIC_VIRTUAL_METHOD(void, Execute, const FRigUnitContext& Context)
+UE_RigUnit_AccumulateQuatMul_IMPLEMENT_MULTIPLEX(void, Execute, const FRigUnitContext& Context)
 {
     DECLARE_SCOPE_HIERARCHICAL_COUNTER_RIGUNIT()
 	if (Context.State == EControlRigState::Init)
@@ -98,7 +98,7 @@ UE_RigUnit_AccumulateQuatMul_IMPLEMENT_STATIC_VIRTUAL_METHOD(void, Execute, cons
 	}
 }
 
-UE_RigUnit_AccumulateTransformMul_IMPLEMENT_STATIC_VIRTUAL_METHOD(void, Execute, const FRigUnitContext& Context)
+UE_RigUnit_AccumulateTransformMul_IMPLEMENT_MULTIPLEX(void, Execute, const FRigUnitContext& Context)
 {
     DECLARE_SCOPE_HIERARCHICAL_COUNTER_RIGUNIT()
 	if (Context.State == EControlRigState::Init)
@@ -125,7 +125,7 @@ UE_RigUnit_AccumulateTransformMul_IMPLEMENT_STATIC_VIRTUAL_METHOD(void, Execute,
 	}
 }
 
-UE_RigUnit_AccumulateFloatLerp_IMPLEMENT_STATIC_VIRTUAL_METHOD(void, Execute, const FRigUnitContext& Context)
+UE_RigUnit_AccumulateFloatLerp_IMPLEMENT_MULTIPLEX(void, Execute, const FRigUnitContext& Context)
 {
     DECLARE_SCOPE_HIERARCHICAL_COUNTER_RIGUNIT()
 	if (Context.State == EControlRigState::Init)
@@ -138,7 +138,7 @@ UE_RigUnit_AccumulateFloatLerp_IMPLEMENT_STATIC_VIRTUAL_METHOD(void, Execute, co
 	}
 }
 
-UE_RigUnit_AccumulateVectorLerp_IMPLEMENT_STATIC_VIRTUAL_METHOD(void, Execute, const FRigUnitContext& Context)
+UE_RigUnit_AccumulateVectorLerp_IMPLEMENT_MULTIPLEX(void, Execute, const FRigUnitContext& Context)
 {
     DECLARE_SCOPE_HIERARCHICAL_COUNTER_RIGUNIT()
 	if (Context.State == EControlRigState::Init)
@@ -151,7 +151,7 @@ UE_RigUnit_AccumulateVectorLerp_IMPLEMENT_STATIC_VIRTUAL_METHOD(void, Execute, c
 	}
 }
 
-UE_RigUnit_AccumulateQuatLerp_IMPLEMENT_STATIC_VIRTUAL_METHOD(void, Execute, const FRigUnitContext& Context)
+UE_RigUnit_AccumulateQuatLerp_IMPLEMENT_MULTIPLEX(void, Execute, const FRigUnitContext& Context)
 {
     DECLARE_SCOPE_HIERARCHICAL_COUNTER_RIGUNIT()
 	if (Context.State == EControlRigState::Init)
@@ -164,7 +164,7 @@ UE_RigUnit_AccumulateQuatLerp_IMPLEMENT_STATIC_VIRTUAL_METHOD(void, Execute, con
 	}
 }
 
-UE_RigUnit_AccumulateTransformLerp_IMPLEMENT_STATIC_VIRTUAL_METHOD(void, Execute, const FRigUnitContext& Context)
+UE_RigUnit_AccumulateTransformLerp_IMPLEMENT_MULTIPLEX(void, Execute, const FRigUnitContext& Context)
 {
     DECLARE_SCOPE_HIERARCHICAL_COUNTER_RIGUNIT()
 	if (Context.State == EControlRigState::Init)
@@ -181,7 +181,7 @@ UE_RigUnit_AccumulateTransformLerp_IMPLEMENT_STATIC_VIRTUAL_METHOD(void, Execute
 	}
 }
 
-UE_RigUnit_AccumulateFloatRange_IMPLEMENT_STATIC_VIRTUAL_METHOD(void, Execute, const FRigUnitContext& Context)
+UE_RigUnit_AccumulateFloatRange_IMPLEMENT_MULTIPLEX(void, Execute, const FRigUnitContext& Context)
 {
     DECLARE_SCOPE_HIERARCHICAL_COUNTER_RIGUNIT()
 	if (Context.State == EControlRigState::Init)
@@ -195,7 +195,7 @@ UE_RigUnit_AccumulateFloatRange_IMPLEMENT_STATIC_VIRTUAL_METHOD(void, Execute, c
 	}
 }
 
-UE_RigUnit_AccumulateVectorRange_IMPLEMENT_STATIC_VIRTUAL_METHOD(void, Execute, const FRigUnitContext& Context)
+UE_RigUnit_AccumulateVectorRange_IMPLEMENT_MULTIPLEX(void, Execute, const FRigUnitContext& Context)
 {
     DECLARE_SCOPE_HIERARCHICAL_COUNTER_RIGUNIT()
 	if (Context.State == EControlRigState::Init)
diff --git a/Engine/Plugins/Experimental/ControlRig/Source/ControlRig/Private/Units/Simulation/RigUnit_Accumulate.h b/Engine/Plugins/Experimental/ControlRig/Source/ControlRig/Private/Units/Simulation/RigUnit_Accumulate.h
--- a/Engine/Plugins/Experimental/ControlRig/Source/ControlRig/Private/Units/Simulation/RigUnit_Accumulate.h
+++ b/Engine/Plugins/Experimental/ControlRig/Source/ControlRig/Private/Units/Simulation/RigUnit_Accumulate.h
@@ -19,7 +19,7 @@ struct FRigUnit_AccumulateFloatAdd : public FRigUnit_SimBase
 		bIntegrateDeltaTime = false;
 	}
 
-	STATIC_VIRTUAL_METHOD()
+	MULTIPLEX_METHOD()
 	virtual void Execute(const FRigUnitContext& Context) override;
 
 	UPROPERTY(meta=(Input))
@@ -52,7 +52,7 @@ struct FRigUnit_AccumulateVectorAdd : public FRigUnit_SimBase
 		bIntegrateDeltaTime = false;
 	}
 
-	STATIC_VIRTUAL_METHOD()
+	MULTIPLEX_METHOD()
 	virtual void Execute(const FRigUnitContext& Context) override;
 
 	UPROPERTY(meta = (Input))
@@ -85,7 +85,7 @@ struct FRigUnit_AccumulateFloatMul : public FRigUnit_SimBase
 		bIntegrateDeltaTime = false;
 	}
 
-	STATIC_VIRTUAL_METHOD()
+	MULTIPLEX_METHOD()
 	virtual void Execute(const FRigUnitContext& Context) override;
 
 	UPROPERTY(meta=(Input))
@@ -118,7 +118,7 @@ struct FRigUnit_AccumulateVectorMul : public FRigUnit_SimBase
 		bIntegrateDeltaTime = false;
 	}
 
-	STATIC_VIRTUAL_METHOD()
+	MULTIPLEX_METHOD()
 	virtual void Execute(const FRigUnitContext& Context) override;
 
 	UPROPERTY(meta = (Input))
@@ -151,7 +151,7 @@ struct FRigUnit_AccumulateQuatMul : public FRigUnit_SimBase
 		bIntegrateDeltaTime = bFlipOrder = false;
 	}
 
-	STATIC_VIRTUAL_METHOD()
+	MULTIPLEX_METHOD()
 	virtual void Execute(const FRigUnitContext& Context) override;
 
 	UPROPERTY(meta=(Input))
@@ -187,7 +187,7 @@ struct FRigUnit_AccumulateTransformMul : public FRigUnit_SimBase
 		bIntegrateDeltaTime = bFlipOrder = false;
 	}
 
-	STATIC_VIRTUAL_METHOD()
+	MULTIPLEX_METHOD()
 	virtual void Execute(const FRigUnitContext& Context) override;
 
 	UPROPERTY(meta = (Input))
@@ -223,7 +223,7 @@ struct FRigUnit_AccumulateFloatLerp : public FRigUnit_SimBase
 		bIntegrateDeltaTime = false;
 	}
 
-	STATIC_VIRTUAL_METHOD()
+	MULTIPLEX_METHOD()
 	virtual void Execute(const FRigUnitContext& Context) override;
 
 	UPROPERTY(meta=(Input))
@@ -260,7 +260,7 @@ struct FRigUnit_AccumulateVectorLerp : public FRigUnit_SimBase
 		bIntegrateDeltaTime = false;
 	}
 
-	STATIC_VIRTUAL_METHOD()
+	MULTIPLEX_METHOD()
 	virtual void Execute(const FRigUnitContext& Context) override;
 
 	UPROPERTY(meta = (Input))
@@ -297,7 +297,7 @@ struct FRigUnit_AccumulateQuatLerp : public FRigUnit_SimBase
 		bIntegrateDeltaTime = false;
 	}
 
-	STATIC_VIRTUAL_METHOD()
+	MULTIPLEX_METHOD()
 	virtual void Execute(const FRigUnitContext& Context) override;
 
 	UPROPERTY(meta=(Input))
@@ -334,7 +334,7 @@ struct FRigUnit_AccumulateTransformLerp : public FRigUnit_SimBase
 		bIntegrateDeltaTime = false;
 	}
 
-	STATIC_VIRTUAL_METHOD()
+	MULTIPLEX_METHOD()
 	virtual void Execute(const FRigUnitContext& Context) override;
 
 	UPROPERTY(meta = (Input))
@@ -369,7 +369,7 @@ struct FRigUnit_AccumulateFloatRange : public FRigUnit_SimBase
 		Value = Minimum = Maximum = AccumulatedMinimum = AccumulatedMaximum = 0.f;
 	}
 
-	STATIC_VIRTUAL_METHOD()
+	MULTIPLEX_METHOD()
 	virtual void Execute(const FRigUnitContext& Context) override;
 
 	UPROPERTY(meta = (Input))
@@ -401,7 +401,7 @@ struct FRigUnit_AccumulateVectorRange : public FRigUnit_SimBase
 		Value = Minimum = Maximum = AccumulatedMinimum = AccumulatedMaximum = FVector::ZeroVector;
 	}
 
-	STATIC_VIRTUAL_METHOD()
+	MULTIPLEX_METHOD()
 	virtual void Execute(const FRigUnitContext& Context) override;
 
 	UPROPERTY(meta = (Input))
diff --git a/Engine/Plugins/Experimental/ControlRig/Source/ControlRig/Private/Units/Simulation/RigUnit_DeltaFromPrevious.cpp b/Engine/Plugins/Experimental/ControlRig/Source/ControlRig/Private/Units/Simulation/RigUnit_DeltaFromPrevious.cpp
--- a/Engine/Plugins/Experimental/ControlRig/Source/ControlRig/Private/Units/Simulation/RigUnit_DeltaFromPrevious.cpp
+++ b/Engine/Plugins/Experimental/ControlRig/Source/ControlRig/Private/Units/Simulation/RigUnit_DeltaFromPrevious.cpp
@@ -3,7 +3,7 @@
 #include "Units/Simulation/RigUnit_DeltaFromPrevious.h"
 #include "Units/RigUnitContext.h"
 
-UE_RigUnit_DeltaFromPreviousFloat_IMPLEMENT_STATIC_VIRTUAL_METHOD(void, Execute, const FRigUnitContext& Context)
+UE_RigUnit_DeltaFromPreviousFloat_IMPLEMENT_MULTIPLEX(void, Execute, const FRigUnitContext& Context)
 {
     DECLARE_SCOPE_HIERARCHICAL_COUNTER_RIGUNIT()
 	if (Context.State == EControlRigState::Init)
@@ -18,7 +18,7 @@ UE_RigUnit_DeltaFromPreviousFloat_IMPLEMENT_STATIC_VIRTUAL_METHOD(void, Execute,
 	Cache = Value;
 }
 
-UE_RigUnit_DeltaFromPreviousVector_IMPLEMENT_STATIC_VIRTUAL_METHOD(void, Execute, const FRigUnitContext& Context)
+UE_RigUnit_DeltaFromPreviousVector_IMPLEMENT_MULTIPLEX(void, Execute, const FRigUnitContext& Context)
 {
     DECLARE_SCOPE_HIERARCHICAL_COUNTER_RIGUNIT()
 	if (Context.State == EControlRigState::Init)
@@ -33,7 +33,7 @@ UE_RigUnit_DeltaFromPreviousVector_IMPLEMENT_STATIC_VIRTUAL_METHOD(void, Execute
 	Cache = Value;
 }
 
-UE_RigUnit_DeltaFromPreviousQuat_IMPLEMENT_STATIC_VIRTUAL_METHOD(void, Execute, const FRigUnitContext& Context)
+UE_RigUnit_DeltaFromPreviousQuat_IMPLEMENT_MULTIPLEX(void, Execute, const FRigUnitContext& Context)
 {
     DECLARE_SCOPE_HIERARCHICAL_COUNTER_RIGUNIT()
 	if (Context.State == EControlRigState::Init)
@@ -48,7 +48,7 @@ UE_RigUnit_DeltaFromPreviousQuat_IMPLEMENT_STATIC_VIRTUAL_METHOD(void, Execute,
 	Cache = Value;
 }
 
-UE_RigUnit_DeltaFromPreviousTransform_IMPLEMENT_STATIC_VIRTUAL_METHOD(void, Execute, const FRigUnitContext& Context)
+UE_RigUnit_DeltaFromPreviousTransform_IMPLEMENT_MULTIPLEX(void, Execute, const FRigUnitContext& Context)
 {
     DECLARE_SCOPE_HIERARCHICAL_COUNTER_RIGUNIT()
 	if (Context.State == EControlRigState::Init)
diff --git a/Engine/Plugins/Experimental/ControlRig/Source/ControlRig/Private/Units/Simulation/RigUnit_DeltaFromPrevious.h b/Engine/Plugins/Experimental/ControlRig/Source/ControlRig/Private/Units/Simulation/RigUnit_DeltaFromPrevious.h
--- a/Engine/Plugins/Experimental/ControlRig/Source/ControlRig/Private/Units/Simulation/RigUnit_DeltaFromPrevious.h
+++ b/Engine/Plugins/Experimental/ControlRig/Source/ControlRig/Private/Units/Simulation/RigUnit_DeltaFromPrevious.h
@@ -18,7 +18,7 @@ struct FRigUnit_DeltaFromPreviousFloat : public FRigUnit_SimBase
 		Value = Delta = PreviousValue = Cache = 0.f;
 	}
 
-	STATIC_VIRTUAL_METHOD()
+	MULTIPLEX_METHOD()
 	virtual void Execute(const FRigUnitContext& Context) override;
 
 	UPROPERTY(meta = (Input))
@@ -47,7 +47,7 @@ struct FRigUnit_DeltaFromPreviousVector : public FRigUnit_SimBase
 		Value = Delta = PreviousValue = Cache = FVector::ZeroVector;
 	}
 
-	STATIC_VIRTUAL_METHOD()
+	MULTIPLEX_METHOD()
 	virtual void Execute(const FRigUnitContext& Context) override;
 
 	UPROPERTY(meta = (Input))
@@ -76,7 +76,7 @@ struct FRigUnit_DeltaFromPreviousQuat : public FRigUnit_SimBase
 		Value = Delta = PreviousValue = Cache = FQuat::Identity;
 	}
 
-	STATIC_VIRTUAL_METHOD()
+	MULTIPLEX_METHOD()
 	virtual void Execute(const FRigUnitContext& Context) override;
 
 	UPROPERTY(meta = (Input))
@@ -105,7 +105,7 @@ struct FRigUnit_DeltaFromPreviousTransform : public FRigUnit_SimBase
 		Value = Delta = PreviousValue = Cache = FTransform::Identity;
 	}
 
-	STATIC_VIRTUAL_METHOD()
+	MULTIPLEX_METHOD()
 	virtual void Execute(const FRigUnitContext& Context) override;
 
 	UPROPERTY(meta = (Input))
diff --git a/Engine/Plugins/Experimental/ControlRig/Source/ControlRig/Private/Units/Simulation/RigUnit_Kalman.cpp b/Engine/Plugins/Experimental/ControlRig/Source/ControlRig/Private/Units/Simulation/RigUnit_Kalman.cpp
--- a/Engine/Plugins/Experimental/ControlRig/Source/ControlRig/Private/Units/Simulation/RigUnit_Kalman.cpp
+++ b/Engine/Plugins/Experimental/ControlRig/Source/ControlRig/Private/Units/Simulation/RigUnit_Kalman.cpp
@@ -3,7 +3,7 @@
 #include "Units/Simulation/RigUnit_Kalman.h"
 #include "Units/RigUnitContext.h"
 
-UE_RigUnit_KalmanFloat_IMPLEMENT_STATIC_VIRTUAL_METHOD(void, Execute, const FRigUnitContext& Context)
+UE_RigUnit_KalmanFloat_IMPLEMENT_MULTIPLEX(void, Execute, const FRigUnitContext& Context)
 {
     DECLARE_SCOPE_HIERARCHICAL_COUNTER_RIGUNIT()
 	if(BufferSize <= 0)
@@ -52,7 +52,7 @@ UE_RigUnit_KalmanFloat_IMPLEMENT_STATIC_VIRTUAL_METHOD(void, Execute, const FRig
 	}
 }
 
-UE_RigUnit_KalmanVector_IMPLEMENT_STATIC_VIRTUAL_METHOD(void, Execute, const FRigUnitContext& Context)
+UE_RigUnit_KalmanVector_IMPLEMENT_MULTIPLEX(void, Execute, const FRigUnitContext& Context)
 {
     DECLARE_SCOPE_HIERARCHICAL_COUNTER_RIGUNIT()
 	if(BufferSize <= 0)
@@ -101,7 +101,7 @@ UE_RigUnit_KalmanVector_IMPLEMENT_STATIC_VIRTUAL_METHOD(void, Execute, const FRi
 	}
 }
 
-UE_RigUnit_KalmanTransform_IMPLEMENT_STATIC_VIRTUAL_METHOD(void, Execute, const FRigUnitContext& Context)
+UE_RigUnit_KalmanTransform_IMPLEMENT_MULTIPLEX(void, Execute, const FRigUnitContext& Context)
 {
     DECLARE_SCOPE_HIERARCHICAL_COUNTER_RIGUNIT()
 	if(BufferSize <= 0)
diff --git a/Engine/Plugins/Experimental/ControlRig/Source/ControlRig/Private/Units/Simulation/RigUnit_Kalman.h b/Engine/Plugins/Experimental/ControlRig/Source/ControlRig/Private/Units/Simulation/RigUnit_Kalman.h
--- a/Engine/Plugins/Experimental/ControlRig/Source/ControlRig/Private/Units/Simulation/RigUnit_Kalman.h
+++ b/Engine/Plugins/Experimental/ControlRig/Source/ControlRig/Private/Units/Simulation/RigUnit_Kalman.h
@@ -21,7 +21,7 @@ struct FRigUnit_KalmanFloat : public FRigUnit_SimBase
 		LastInsertIndex = 0;
 	}
 
-	STATIC_VIRTUAL_METHOD()
+	MULTIPLEX_METHOD()
 	virtual void Execute(const FRigUnitContext& Context) override;
 
 	UPROPERTY(meta = (Input))
@@ -56,7 +56,7 @@ struct FRigUnit_KalmanVector : public FRigUnit_SimBase
 		LastInsertIndex = 0;
 	}
 
-	STATIC_VIRTUAL_METHOD()
+	MULTIPLEX_METHOD()
 	virtual void Execute(const FRigUnitContext& Context) override;
 
 	UPROPERTY(meta = (Input))
@@ -91,7 +91,7 @@ struct FRigUnit_KalmanTransform : public FRigUnit_SimBase
 		LastInsertIndex = 0;
 	}
 
-	STATIC_VIRTUAL_METHOD()
+	MULTIPLEX_METHOD()
 	virtual void Execute(const FRigUnitContext& Context) override;
 
 	UPROPERTY(meta = (Input))
diff --git a/Engine/Plugins/Experimental/ControlRig/Source/ControlRig/Private/Units/Simulation/RigUnit_PointSimulation.cpp b/Engine/Plugins/Experimental/ControlRig/Source/ControlRig/Private/Units/Simulation/RigUnit_PointSimulation.cpp
--- a/Engine/Plugins/Experimental/ControlRig/Source/ControlRig/Private/Units/Simulation/RigUnit_PointSimulation.cpp
+++ b/Engine/Plugins/Experimental/ControlRig/Source/ControlRig/Private/Units/Simulation/RigUnit_PointSimulation.cpp
@@ -3,7 +3,7 @@
 #include "Units/Simulation/RigUnit_PointSimulation.h"
 #include "Units/RigUnitContext.h"
 
-UE_RigUnit_PointSimulation_IMPLEMENT_STATIC_VIRTUAL_METHOD(void, Execute, const FRigUnitContext& Context)
+UE_RigUnit_PointSimulation_IMPLEMENT_MULTIPLEX(void, Execute, const FRigUnitContext& Context)
 {
     DECLARE_SCOPE_HIERARCHICAL_COUNTER_RIGUNIT()
 	FRigBoneHierarchy* Hierarchy = ExecuteContext.GetBones();
diff --git a/Engine/Plugins/Experimental/ControlRig/Source/ControlRig/Private/Units/Simulation/RigUnit_PointSimulation.h b/Engine/Plugins/Experimental/ControlRig/Source/ControlRig/Private/Units/Simulation/RigUnit_PointSimulation.h
--- a/Engine/Plugins/Experimental/ControlRig/Source/ControlRig/Private/Units/Simulation/RigUnit_PointSimulation.h
+++ b/Engine/Plugins/Experimental/ControlRig/Source/ControlRig/Private/Units/Simulation/RigUnit_PointSimulation.h
@@ -123,7 +123,7 @@ struct FRigUnit_PointSimulation : public FRigUnit_SimBaseMutable
 		Bezier = FCRFourPointBezier();
 	}
 
-	STATIC_VIRTUAL_METHOD()
+	MULTIPLEX_METHOD()
 	virtual void Execute(const FRigUnitContext& Context) override;
 
 	/** The points to simulate */
diff --git a/Engine/Plugins/Experimental/ControlRig/Source/ControlRig/Private/Units/Simulation/RigUnit_TimeOffset.cpp b/Engine/Plugins/Experimental/ControlRig/Source/ControlRig/Private/Units/Simulation/RigUnit_TimeOffset.cpp
--- a/Engine/Plugins/Experimental/ControlRig/Source/ControlRig/Private/Units/Simulation/RigUnit_TimeOffset.cpp
+++ b/Engine/Plugins/Experimental/ControlRig/Source/ControlRig/Private/Units/Simulation/RigUnit_TimeOffset.cpp
@@ -3,7 +3,7 @@
 #include "Units/Simulation/RigUnit_TimeOffset.h"
 #include "Units/RigUnitContext.h"
 
-UE_RigUnit_TimeOffsetFloat_IMPLEMENT_STATIC_VIRTUAL_METHOD(void, Execute, const FRigUnitContext& Context)
+UE_RigUnit_TimeOffsetFloat_IMPLEMENT_MULTIPLEX(void, Execute, const FRigUnitContext& Context)
 {
     DECLARE_SCOPE_HIERARCHICAL_COUNTER_RIGUNIT()
 	Result = Value;
@@ -111,7 +111,7 @@ UE_RigUnit_TimeOffsetFloat_IMPLEMENT_STATIC_VIRTUAL_METHOD(void, Execute, const
 	}
 }
 
-UE_RigUnit_TimeOffsetVector_IMPLEMENT_STATIC_VIRTUAL_METHOD(void, Execute, const FRigUnitContext& Context)
+UE_RigUnit_TimeOffsetVector_IMPLEMENT_MULTIPLEX(void, Execute, const FRigUnitContext& Context)
 {
     DECLARE_SCOPE_HIERARCHICAL_COUNTER_RIGUNIT()
 	Result = Value;
@@ -219,7 +219,7 @@ UE_RigUnit_TimeOffsetVector_IMPLEMENT_STATIC_VIRTUAL_METHOD(void, Execute, const
 	}
 }
 
-UE_RigUnit_TimeOffsetTransform_IMPLEMENT_STATIC_VIRTUAL_METHOD(void, Execute, const FRigUnitContext& Context)
+UE_RigUnit_TimeOffsetTransform_IMPLEMENT_MULTIPLEX(void, Execute, const FRigUnitContext& Context)
 {
     DECLARE_SCOPE_HIERARCHICAL_COUNTER_RIGUNIT()
 	Result = Value;
diff --git a/Engine/Plugins/Experimental/ControlRig/Source/ControlRig/Private/Units/Simulation/RigUnit_TimeOffset.h b/Engine/Plugins/Experimental/ControlRig/Source/ControlRig/Private/Units/Simulation/RigUnit_TimeOffset.h
--- a/Engine/Plugins/Experimental/ControlRig/Source/ControlRig/Private/Units/Simulation/RigUnit_TimeOffset.h
+++ b/Engine/Plugins/Experimental/ControlRig/Source/ControlRig/Private/Units/Simulation/RigUnit_TimeOffset.h
@@ -22,7 +22,7 @@ struct FRigUnit_TimeOffsetFloat : public FRigUnit_SimBase
 		LastInsertIndex = 0;
 	}
 
-	STATIC_VIRTUAL_METHOD()
+	MULTIPLEX_METHOD()
 	virtual void Execute(const FRigUnitContext& Context) override;
 	
 	/** The value to record */
@@ -71,7 +71,7 @@ struct FRigUnit_TimeOffsetVector : public FRigUnit_SimBase
 		LastInsertIndex = 0;
 	}
 
-	STATIC_VIRTUAL_METHOD()
+	MULTIPLEX_METHOD()
 	virtual void Execute(const FRigUnitContext& Context) override;
 	
 	/** The value to record */
@@ -120,7 +120,7 @@ struct FRigUnit_TimeOffsetTransform : public FRigUnit_SimBase
 		LastInsertIndex = 0;
 	}
 
-	STATIC_VIRTUAL_METHOD()
+	MULTIPLEX_METHOD()
 	virtual void Execute(const FRigUnitContext& Context) override;
 	
 	/** The value to record */
diff --git a/Engine/Plugins/Experimental/ControlRig/Source/ControlRig/Private/Units/Simulation/RigUnit_Timeline.cpp b/Engine/Plugins/Experimental/ControlRig/Source/ControlRig/Private/Units/Simulation/RigUnit_Timeline.cpp
--- a/Engine/Plugins/Experimental/ControlRig/Source/ControlRig/Private/Units/Simulation/RigUnit_Timeline.cpp
+++ b/Engine/Plugins/Experimental/ControlRig/Source/ControlRig/Private/Units/Simulation/RigUnit_Timeline.cpp
@@ -3,7 +3,7 @@
 #include "Units/Simulation/RigUnit_Timeline.h"
 #include "Units/RigUnitContext.h"
 
-UE_RigUnit_Timeline_IMPLEMENT_STATIC_VIRTUAL_METHOD(void, Execute, const FRigUnitContext& Context)
+UE_RigUnit_Timeline_IMPLEMENT_MULTIPLEX(void, Execute, const FRigUnitContext& Context)
 {
     DECLARE_SCOPE_HIERARCHICAL_COUNTER_RIGUNIT()
 	if (Context.State == EControlRigState::Init)
diff --git a/Engine/Plugins/Experimental/ControlRig/Source/ControlRig/Private/Units/Simulation/RigUnit_Timeline.h b/Engine/Plugins/Experimental/ControlRig/Source/ControlRig/Private/Units/Simulation/RigUnit_Timeline.h
--- a/Engine/Plugins/Experimental/ControlRig/Source/ControlRig/Private/Units/Simulation/RigUnit_Timeline.h
+++ b/Engine/Plugins/Experimental/ControlRig/Source/ControlRig/Private/Units/Simulation/RigUnit_Timeline.h
@@ -19,7 +19,7 @@ struct FRigUnit_Timeline : public FRigUnit_SimBase
 		Time = AccumulatedValue = 0.f;
 	}
 
-	STATIC_VIRTUAL_METHOD()
+	MULTIPLEX_METHOD()
 	virtual void Execute(const FRigUnitContext& Context) override;
 
 	UPROPERTY(meta = (Input))
diff --git a/Engine/Plugins/Experimental/ControlRig/Source/ControlRig/Private/Units/Simulation/RigUnit_Verlet.cpp b/Engine/Plugins/Experimental/ControlRig/Source/ControlRig/Private/Units/Simulation/RigUnit_Verlet.cpp
--- a/Engine/Plugins/Experimental/ControlRig/Source/ControlRig/Private/Units/Simulation/RigUnit_Verlet.cpp
+++ b/Engine/Plugins/Experimental/ControlRig/Source/ControlRig/Private/Units/Simulation/RigUnit_Verlet.cpp
@@ -3,7 +3,7 @@
 #include "Units/Simulation/RigUnit_Verlet.h"
 #include "Units/RigUnitContext.h"
 
-UE_RigUnit_VerletIntegrateVector_IMPLEMENT_STATIC_VIRTUAL_METHOD(void, Execute, const FRigUnitContext& Context)
+UE_RigUnit_VerletIntegrateVector_IMPLEMENT_MULTIPLEX(void, Execute, const FRigUnitContext& Context)
 {
     DECLARE_SCOPE_HIERARCHICAL_COUNTER_RIGUNIT()
 	if (Context.State == EControlRigState::Init)
diff --git a/Engine/Plugins/Experimental/ControlRig/Source/ControlRig/Private/Units/Simulation/RigUnit_Verlet.h b/Engine/Plugins/Experimental/ControlRig/Source/ControlRig/Private/Units/Simulation/RigUnit_Verlet.h
--- a/Engine/Plugins/Experimental/ControlRig/Source/ControlRig/Private/Units/Simulation/RigUnit_Verlet.h
+++ b/Engine/Plugins/Experimental/ControlRig/Source/ControlRig/Private/Units/Simulation/RigUnit_Verlet.h
@@ -23,7 +23,7 @@ struct FRigUnit_VerletIntegrateVector : public FRigUnit_SimBase
 		Point = FCRSimPoint();
 	}
 
-	STATIC_VIRTUAL_METHOD()
+	MULTIPLEX_METHOD()
 	virtual void Execute(const FRigUnitContext& Context) override;
 
 	UPROPERTY(meta = (Input))
diff --git a/Engine/Plugins/Experimental/ControlRig/Source/ControlRig/Public/Units/Control/RigUnit_Control.h b/Engine/Plugins/Experimental/ControlRig/Source/ControlRig/Public/Units/Control/RigUnit_Control.h
--- a/Engine/Plugins/Experimental/ControlRig/Source/ControlRig/Public/Units/Control/RigUnit_Control.h
+++ b/Engine/Plugins/Experimental/ControlRig/Source/ControlRig/Public/Units/Control/RigUnit_Control.h
@@ -21,7 +21,7 @@ struct CONTROLRIG_API FRigUnit_Control : public FRigUnit
 	{
 	}
 
-	STATIC_VIRTUAL_METHOD()
+	MULTIPLEX_METHOD()
 	virtual void Execute(const FRigUnitContext& Context) override;
 
 	/** Combine Transform and Base to make the resultant transform */
diff --git a/Engine/Plugins/Experimental/ControlRig/Source/ControlRig/Public/Units/RigUnit.h b/Engine/Plugins/Experimental/ControlRig/Source/ControlRig/Public/Units/RigUnit.h
--- a/Engine/Plugins/Experimental/ControlRig/Source/ControlRig/Public/Units/RigUnit.h
+++ b/Engine/Plugins/Experimental/ControlRig/Source/ControlRig/Public/Units/RigUnit.h
@@ -3,6 +3,7 @@
 #pragma once
 
 #include "ControlRigDefines.h"
+#include "MultiplexRegistry.h"
 #include "RigUnit.generated.h"
 
 struct FRigUnitContext;
