Multiplex: Conforming all ControlRig rig units to multiplex method policies (fix build break caused by UHT changes)

diff --git a/Engine/Plugins/Experimental/ControlRig/Source/ControlRig/Private/Units/Debug/RigUnit_DebugTransform.cpp b/Engine/Plugins/Experimental/ControlRig/Source/ControlRig/Private/Units/Debug/RigUnit_DebugTransform.cpp
--- a/Engine/Plugins/Experimental/ControlRig/Source/ControlRig/Private/Units/Debug/RigUnit_DebugTransform.cpp
+++ b/Engine/Plugins/Experimental/ControlRig/Source/ControlRig/Private/Units/Debug/RigUnit_DebugTransform.cpp
@@ -96,24 +96,24 @@ UE_RigUnit_DebugTransformArrayMutable_IMPLEMENT_MULTIPLEX(void, Execute, const F
 		return;
 	}
 
-	DrawTransforms.SetNumUninitialized(Transforms.Num());
+	WorkData.DrawTransforms.SetNumUninitialized(Transforms.Num());
 	if (Space != NAME_None && Context.GetBones() != nullptr)
 	{
 		FTransform SpaceTransform = Context.GetBones()->GetGlobalTransform(Space);
 		for(int32 Index=0;Index<Transforms.Num();Index++)
 		{
-			DrawTransforms[Index] = Transforms[Index] * SpaceTransform;
+			WorkData.DrawTransforms[Index] = Transforms[Index] * SpaceTransform;
 		}
 	}
 	else
 	{
 		for(int32 Index=0;Index<Transforms.Num();Index++)
 		{
-			DrawTransforms[Index] = Transforms[Index];
+			WorkData.DrawTransforms[Index] = Transforms[Index];
 		}
 	}
 
-	for(FTransform& DrawTransform : DrawTransforms)
+	for(FTransform& DrawTransform : WorkData.DrawTransforms)
 	{
 		switch (Mode)
 		{
diff --git a/Engine/Plugins/Experimental/ControlRig/Source/ControlRig/Private/Units/Debug/RigUnit_DebugTransform.h b/Engine/Plugins/Experimental/ControlRig/Source/ControlRig/Private/Units/Debug/RigUnit_DebugTransform.h
--- a/Engine/Plugins/Experimental/ControlRig/Source/ControlRig/Private/Units/Debug/RigUnit_DebugTransform.h
+++ b/Engine/Plugins/Experimental/ControlRig/Source/ControlRig/Private/Units/Debug/RigUnit_DebugTransform.h
@@ -107,10 +107,19 @@ struct FRigUnit_DebugTransformMutable : public FRigUnit_DebugBaseMutable
 	bool bEnabled;
 };
 
+USTRUCT()
+struct FRigUnit_DebugTransformArrayMutable_WorkData
+{
+	GENERATED_BODY()
+		
+	UPROPERTY()
+	TArray<FTransform> DrawTransforms;
+};
+
 USTRUCT(meta=(DisplayName="Draw Transform Array"))
 struct FRigUnit_DebugTransformArrayMutable : public FRigUnit_DebugBaseMutable
 {
-	GENERATED_BODY()
+ 	GENERATED_BODY()
 
 	FRigUnit_DebugTransformArrayMutable()
 	{
@@ -149,6 +158,6 @@ struct FRigUnit_DebugTransformArrayMutable : public FRigUnit_DebugBaseMutable
 	UPROPERTY(meta = (Input, Constant, BoneName))
 	bool bEnabled;
 
-	UPROPERTY()
-	TArray<FTransform> DrawTransforms;
+	UPROPERTY(transient)
+	FRigUnit_DebugTransformArrayMutable_WorkData WorkData;
 };
\ No newline at end of file
diff --git a/Engine/Plugins/Experimental/ControlRig/Source/ControlRig/Private/Units/Highlevel/Harmonics/RigUnit_BoneHarmonics.cpp b/Engine/Plugins/Experimental/ControlRig/Source/ControlRig/Private/Units/Highlevel/Harmonics/RigUnit_BoneHarmonics.cpp
--- a/Engine/Plugins/Experimental/ControlRig/Source/ControlRig/Private/Units/Highlevel/Harmonics/RigUnit_BoneHarmonics.cpp
+++ b/Engine/Plugins/Experimental/ControlRig/Source/ControlRig/Private/Units/Highlevel/Harmonics/RigUnit_BoneHarmonics.cpp
@@ -12,6 +12,9 @@ UE_RigUnit_BoneHarmonics_IMPLEMENT_MULTIPLEX(void, Execute, const FRigUnitContex
 		return;
 	}
 
+	TArray<int32>& BoneIndices = WorkData.BoneIndices;
+	FVector& WaveTime = WorkData.WaveTime;
+
 	if (Context.State == EControlRigState::Init ||
 		BoneIndices.Num() != Bones.Num())
 	{
diff --git a/Engine/Plugins/Experimental/ControlRig/Source/ControlRig/Private/Units/Highlevel/Harmonics/RigUnit_BoneHarmonics.h b/Engine/Plugins/Experimental/ControlRig/Source/ControlRig/Private/Units/Highlevel/Harmonics/RigUnit_BoneHarmonics.h
--- a/Engine/Plugins/Experimental/ControlRig/Source/ControlRig/Private/Units/Highlevel/Harmonics/RigUnit_BoneHarmonics.h
+++ b/Engine/Plugins/Experimental/ControlRig/Source/ControlRig/Private/Units/Highlevel/Harmonics/RigUnit_BoneHarmonics.h
@@ -31,6 +31,23 @@ struct FRigUnit_BoneHarmonics_BoneTarget
 	float Ratio;
 };
 
+USTRUCT()
+struct FRigUnit_BoneHarmonics_WorkData
+{
+	GENERATED_BODY()
+
+	FRigUnit_BoneHarmonics_WorkData()
+	{
+		WaveTime = FVector::ZeroVector;
+	}
+
+	UPROPERTY()
+	TArray<int32> BoneIndices;
+
+	UPROPERTY()
+	FVector WaveTime;
+};
+
 /**
  * Performs point based simulation
  */
@@ -51,7 +68,6 @@ struct FRigUnit_BoneHarmonics : public FRigUnit_HighlevelBaseMutable
 		WaveMaximum = 1.f;
 		RotationOrder = EControlRigRotationOrder::YZX;
 		bPropagateToChildren = false;
-		WaveTime = FVector::ZeroVector;
 	}
 
 	MULTIPLEX_METHOD()
@@ -97,10 +113,7 @@ struct FRigUnit_BoneHarmonics : public FRigUnit_HighlevelBaseMutable
 	UPROPERTY(meta = (Input))
 	bool bPropagateToChildren;
 
-	UPROPERTY()
-	TArray<int32> BoneIndices;
-
-	UPROPERTY()
-	FVector WaveTime;
+	UPROPERTY(transient)
+	FRigUnit_BoneHarmonics_WorkData WorkData;
 };
 
diff --git a/Engine/Plugins/Experimental/ControlRig/Source/ControlRig/Private/Units/Highlevel/Harmonics/RigUnit_ChainHarmonics.cpp b/Engine/Plugins/Experimental/ControlRig/Source/ControlRig/Private/Units/Highlevel/Harmonics/RigUnit_ChainHarmonics.cpp
--- a/Engine/Plugins/Experimental/ControlRig/Source/ControlRig/Private/Units/Highlevel/Harmonics/RigUnit_ChainHarmonics.cpp
+++ b/Engine/Plugins/Experimental/ControlRig/Source/ControlRig/Private/Units/Highlevel/Harmonics/RigUnit_ChainHarmonics.cpp
@@ -11,6 +11,16 @@ UE_RigUnit_ChainHarmonics_IMPLEMENT_MULTIPLEX(void, Execute, const FRigUnitConte
 	{
 		return;
 	}
+
+	FVector& Time = WorkData.Time;
+	TArray<int32>& Bones = WorkData.Bones;
+	TArray<float>& Ratio = WorkData.Ratio;
+	TArray<FVector>& LocalTip = WorkData.LocalTip;
+	TArray<FVector>& PendulumTip = WorkData.PendulumTip;
+	TArray<FVector>& PendulumPosition = WorkData.PendulumPosition;
+	TArray<FVector>& PendulumVelocity = WorkData.PendulumVelocity;
+	TArray<FVector>& HierarchyLine = WorkData.HierarchyLine;
+	TArray<FVector>& VelocityLines = WorkData.VelocityLines;
 	
 	if (Context.State == EControlRigState::Init)
 	{
diff --git a/Engine/Plugins/Experimental/ControlRig/Source/ControlRig/Private/Units/Highlevel/Harmonics/RigUnit_ChainHarmonics.h b/Engine/Plugins/Experimental/ControlRig/Source/ControlRig/Private/Units/Highlevel/Harmonics/RigUnit_ChainHarmonics.h
--- a/Engine/Plugins/Experimental/ControlRig/Source/ControlRig/Private/Units/Highlevel/Harmonics/RigUnit_ChainHarmonics.h
+++ b/Engine/Plugins/Experimental/ControlRig/Source/ControlRig/Private/Units/Highlevel/Harmonics/RigUnit_ChainHarmonics.h
@@ -137,6 +137,44 @@ struct FRigUnit_ChainHarmonics_Pendulum
 	float UnwindMaximum;
 };
 
+USTRUCT()
+struct FRigUnit_ChainHarmonics_WorkData
+{
+	GENERATED_BODY()
+
+	FRigUnit_ChainHarmonics_WorkData()
+	{
+		Time = FVector::ZeroVector;
+	}
+
+	UPROPERTY()
+	FVector Time;
+
+	UPROPERTY()
+	TArray<int32> Bones;
+
+	UPROPERTY()
+	TArray<float> Ratio;
+
+	UPROPERTY()
+	TArray<FVector> LocalTip;
+
+	UPROPERTY()
+	TArray<FVector> PendulumTip;
+
+	UPROPERTY()
+	TArray<FVector> PendulumPosition;
+
+	UPROPERTY()
+	TArray<FVector> PendulumVelocity;
+
+	UPROPERTY()
+	TArray<FVector> HierarchyLine;
+
+	UPROPERTY()
+	TArray<FVector> VelocityLines;
+};
+
 USTRUCT(meta=(DisplayName="ChainHarmonics"))
 struct FRigUnit_ChainHarmonics : public FRigUnit_HighlevelBaseMutable
 {
@@ -157,8 +195,6 @@ struct FRigUnit_ChainHarmonics : public FRigUnit_HighlevelBaseMutable
 
 		bDrawDebug = true;
 		DrawWorldOffset = FTransform::Identity;
-
-		Time = FVector::ZeroVector;
 	}
 
 	MULTIPLEX_METHOD()
@@ -188,30 +224,6 @@ struct FRigUnit_ChainHarmonics : public FRigUnit_HighlevelBaseMutable
 	UPROPERTY(meta = (Input))
 	FTransform DrawWorldOffset;
 
-	UPROPERTY()
-	FVector Time;
-
-	UPROPERTY()
-	TArray<int32> Bones;
-
-	UPROPERTY()
-	TArray<float> Ratio;
-
-	UPROPERTY()
-	TArray<FVector> LocalTip;
-
-	UPROPERTY()
-	TArray<FVector> PendulumTip;
-
-	UPROPERTY()
-	TArray<FVector> PendulumPosition;
-
-	UPROPERTY()
-	TArray<FVector> PendulumVelocity;
-
-	UPROPERTY()
-	TArray<FVector> HierarchyLine;
-
-	UPROPERTY()
-	TArray<FVector> VelocityLines;
+	UPROPERTY(transient)
+	FRigUnit_ChainHarmonics_WorkData WorkData;
 };
diff --git a/Engine/Plugins/Experimental/ControlRig/Source/ControlRig/Private/Units/Highlevel/Hierarchy/RigUnit_CCDIK.cpp b/Engine/Plugins/Experimental/ControlRig/Source/ControlRig/Private/Units/Highlevel/Hierarchy/RigUnit_CCDIK.cpp
--- a/Engine/Plugins/Experimental/ControlRig/Source/ControlRig/Private/Units/Highlevel/Hierarchy/RigUnit_CCDIK.cpp
+++ b/Engine/Plugins/Experimental/ControlRig/Source/ControlRig/Private/Units/Highlevel/Hierarchy/RigUnit_CCDIK.cpp
@@ -12,6 +12,12 @@ UE_RigUnit_CCDIK_IMPLEMENT_MULTIPLEX(void, Execute, const FRigUnitContext& Conte
 		return;
 	}
 
+	TArray<FCCDIKChainLink>& Chain = WorkData.Chain;
+	TArray<int32>& BoneIndices = WorkData.BoneIndices;
+	TArray<int32>& RotationLimitIndex = WorkData.RotationLimitIndex;
+	TArray<float>& RotationLimitsPerBone = WorkData.RotationLimitsPerBone;
+	int32& EffectorIndex = WorkData.EffectorIndex;
+
 	if (Context.State == EControlRigState::Init ||
 		RotationLimits.Num() != RotationLimitIndex.Num())
 	{
diff --git a/Engine/Plugins/Experimental/ControlRig/Source/ControlRig/Private/Units/Highlevel/Hierarchy/RigUnit_CCDIK.h b/Engine/Plugins/Experimental/ControlRig/Source/ControlRig/Private/Units/Highlevel/Hierarchy/RigUnit_CCDIK.h
--- a/Engine/Plugins/Experimental/ControlRig/Source/ControlRig/Private/Units/Highlevel/Hierarchy/RigUnit_CCDIK.h
+++ b/Engine/Plugins/Experimental/ControlRig/Source/ControlRig/Private/Units/Highlevel/Hierarchy/RigUnit_CCDIK.h
@@ -30,6 +30,32 @@ struct FRigUnit_CCDIK_RotationLimit
 	float Limit;
 };
 
+USTRUCT()
+struct FRigUnit_CCDIK_WorkData
+{
+	GENERATED_BODY()
+
+	FRigUnit_CCDIK_WorkData()
+	{
+		EffectorIndex = INDEX_NONE;
+	}
+
+	UPROPERTY()
+	TArray<FCCDIKChainLink> Chain;
+
+	UPROPERTY()
+	TArray<int32> BoneIndices;
+
+	UPROPERTY()
+	TArray<int32> RotationLimitIndex;
+
+	UPROPERTY()
+	TArray<float> RotationLimitsPerBone;
+
+	UPROPERTY()
+	int32 EffectorIndex;
+};
+
 /**
  * The CCID solver can solve N-Bone chains using 
  * the Cyclic Coordinate Descent Inverse Kinematics algorithm.
@@ -45,7 +71,6 @@ struct FRigUnit_CCDIK : public FRigUnit_HighlevelBaseMutable
 
 	FRigUnit_CCDIK()
 	{
-		EffectorIndex = INDEX_NONE;
 		EffectorTransform = FTransform::Identity;
 		Precision = 1.f;
 		MaxIterations = 10;
@@ -111,13 +136,5 @@ struct FRigUnit_CCDIK : public FRigUnit_HighlevelBaseMutable
 	bool bPropagateToChildren;
 
 	UPROPERTY(transient)
-	TArray<FCCDIKChainLink> Chain;
-	UPROPERTY(transient)
-	TArray<int32> BoneIndices;
-	UPROPERTY(transient)
-	TArray<int32> RotationLimitIndex;
-	UPROPERTY(transient)
-	TArray<float> RotationLimitsPerBone;
-	UPROPERTY(transient)
-	int32 EffectorIndex;
+	FRigUnit_CCDIK_WorkData WorkData;
 };
diff --git a/Engine/Plugins/Experimental/ControlRig/Source/ControlRig/Private/Units/Highlevel/Hierarchy/RigUnit_DistributeRotation.cpp b/Engine/Plugins/Experimental/ControlRig/Source/ControlRig/Private/Units/Highlevel/Hierarchy/RigUnit_DistributeRotation.cpp
--- a/Engine/Plugins/Experimental/ControlRig/Source/ControlRig/Private/Units/Highlevel/Hierarchy/RigUnit_DistributeRotation.cpp
+++ b/Engine/Plugins/Experimental/ControlRig/Source/ControlRig/Private/Units/Highlevel/Hierarchy/RigUnit_DistributeRotation.cpp
@@ -12,6 +12,12 @@ UE_RigUnit_DistributeRotation_IMPLEMENT_MULTIPLEX(void, Execute, const FRigUnitC
 		return;
 	}
 
+	TArray<int32>& BoneIndices = WorkData.BoneIndices;
+	TArray<int32>& BoneRotationA = WorkData.BoneRotationA;
+	TArray<int32>& BoneRotationB = WorkData.BoneRotationB;
+	TArray<float>& BoneRotationT = WorkData.BoneRotationT;
+	TArray<FTransform>& BoneLocalTransforms = WorkData.BoneLocalTransforms;
+
 	if (Context.State == EControlRigState::Init)
 	{
 		BoneIndices.Reset();
@@ -208,21 +214,21 @@ IMPLEMENT_RIGUNIT_AUTOMATION_TEST(FRigUnit_DistributeRotation)
 
 	Init();
 
-	AddErrorIfFalse(Unit.BoneRotationA[0] == 0, TEXT("unexpected bone a"));
-	AddErrorIfFalse(Unit.BoneRotationB[0] == 0, TEXT("unexpected bone b"));
-	AddErrorIfFalse(FMath::IsNearlyEqual(Unit.BoneRotationT[0], 0.f, 0.001f), TEXT("unexpected bone t"));
-	AddErrorIfFalse(Unit.BoneRotationA[1] == 0, TEXT("unexpected bone a"));
-	AddErrorIfFalse(Unit.BoneRotationB[1] == 2, TEXT("unexpected bone b"));
-	AddErrorIfFalse(FMath::IsNearlyEqual(Unit.BoneRotationT[1], 0.5f, 0.001f), TEXT("unexpected bone t"));
-	AddErrorIfFalse(Unit.BoneRotationA[2] == 2, TEXT("unexpected bone a"));
-	AddErrorIfFalse(Unit.BoneRotationB[2] == 2, TEXT("unexpected bone b"));
-	AddErrorIfFalse(FMath::IsNearlyEqual(Unit.BoneRotationT[2], 0.f, 0.001f), TEXT("unexpected bone t"));
-	AddErrorIfFalse(Unit.BoneRotationA[3] == 2, TEXT("unexpected bone a"));
-	AddErrorIfFalse(Unit.BoneRotationB[3] == 1, TEXT("unexpected bone b"));
-	AddErrorIfFalse(FMath::IsNearlyEqual(Unit.BoneRotationT[3], 0.5f, 0.001f), TEXT("unexpected bone t"));
-	AddErrorIfFalse(Unit.BoneRotationA[4] == 1, TEXT("unexpected bone a"));
-	AddErrorIfFalse(Unit.BoneRotationB[4] == 1, TEXT("unexpected bone b"));
-	AddErrorIfFalse(FMath::IsNearlyEqual(Unit.BoneRotationT[4], 0.0f, 0.001f), TEXT("unexpected bone t"));
+	AddErrorIfFalse(Unit.WorkData.BoneRotationA[0] == 0, TEXT("unexpected bone a"));
+	AddErrorIfFalse(Unit.WorkData.BoneRotationB[0] == 0, TEXT("unexpected bone b"));
+	AddErrorIfFalse(FMath::IsNearlyEqual(Unit.WorkData.BoneRotationT[0], 0.f, 0.001f), TEXT("unexpected bone t"));
+	AddErrorIfFalse(Unit.WorkData.BoneRotationA[1] == 0, TEXT("unexpected bone a"));
+	AddErrorIfFalse(Unit.WorkData.BoneRotationB[1] == 2, TEXT("unexpected bone b"));
+	AddErrorIfFalse(FMath::IsNearlyEqual(Unit.WorkData.BoneRotationT[1], 0.5f, 0.001f), TEXT("unexpected bone t"));
+	AddErrorIfFalse(Unit.WorkData.BoneRotationA[2] == 2, TEXT("unexpected bone a"));
+	AddErrorIfFalse(Unit.WorkData.BoneRotationB[2] == 2, TEXT("unexpected bone b"));
+	AddErrorIfFalse(FMath::IsNearlyEqual(Unit.WorkData.BoneRotationT[2], 0.f, 0.001f), TEXT("unexpected bone t"));
+	AddErrorIfFalse(Unit.WorkData.BoneRotationA[3] == 2, TEXT("unexpected bone a"));
+	AddErrorIfFalse(Unit.WorkData.BoneRotationB[3] == 1, TEXT("unexpected bone b"));
+	AddErrorIfFalse(FMath::IsNearlyEqual(Unit.WorkData.BoneRotationT[3], 0.5f, 0.001f), TEXT("unexpected bone t"));
+	AddErrorIfFalse(Unit.WorkData.BoneRotationA[4] == 1, TEXT("unexpected bone a"));
+	AddErrorIfFalse(Unit.WorkData.BoneRotationB[4] == 1, TEXT("unexpected bone b"));
+	AddErrorIfFalse(FMath::IsNearlyEqual(Unit.WorkData.BoneRotationT[4], 0.0f, 0.001f), TEXT("unexpected bone t"));
 
 	Execute();
 
diff --git a/Engine/Plugins/Experimental/ControlRig/Source/ControlRig/Private/Units/Highlevel/Hierarchy/RigUnit_DistributeRotation.h b/Engine/Plugins/Experimental/ControlRig/Source/ControlRig/Private/Units/Highlevel/Hierarchy/RigUnit_DistributeRotation.h
--- a/Engine/Plugins/Experimental/ControlRig/Source/ControlRig/Private/Units/Highlevel/Hierarchy/RigUnit_DistributeRotation.h
+++ b/Engine/Plugins/Experimental/ControlRig/Source/ControlRig/Private/Units/Highlevel/Hierarchy/RigUnit_DistributeRotation.h
@@ -30,6 +30,27 @@ struct FRigUnit_DistributeRotation_Rotation
 	float Ratio;
 };
 
+USTRUCT()
+struct FRigUnit_DistributeRotation_WorkData
+{
+	GENERATED_BODY()
+
+	UPROPERTY()
+	TArray<int32> BoneIndices;
+
+	UPROPERTY()
+	TArray<int32> BoneRotationA;
+
+	UPROPERTY()
+	TArray<int32> BoneRotationB;
+
+	UPROPERTY()
+	TArray<float> BoneRotationT;
+
+	UPROPERTY()
+	TArray<FTransform> BoneLocalTransforms;
+};
+
 /**
  * Distributes rotations provided along a chain.
  * Each rotation is expressed by a quaternion and a ratio, where the ratio is between 0.0 and 1.0
@@ -82,18 +103,6 @@ struct FRigUnit_DistributeRotation : public FRigUnit_HighlevelBaseMutable
 	UPROPERTY(meta = (Input))
 	bool bPropagateToChildren;
 
-	UPROPERTY()
-	TArray<int32> BoneIndices;
-
-	UPROPERTY()
-	TArray<int32> BoneRotationA;
-
-	UPROPERTY()
-	TArray<int32> BoneRotationB;
-
-	UPROPERTY()
-	TArray<float> BoneRotationT;
-
-	UPROPERTY()
-	TArray<FTransform> BoneLocalTransforms;
+	UPROPERTY(transient)
+	FRigUnit_DistributeRotation_WorkData WorkData;
 };
diff --git a/Engine/Plugins/Experimental/ControlRig/Source/ControlRig/Private/Units/Highlevel/Hierarchy/RigUnit_FABRIK.cpp b/Engine/Plugins/Experimental/ControlRig/Source/ControlRig/Private/Units/Highlevel/Hierarchy/RigUnit_FABRIK.cpp
--- a/Engine/Plugins/Experimental/ControlRig/Source/ControlRig/Private/Units/Highlevel/Hierarchy/RigUnit_FABRIK.cpp
+++ b/Engine/Plugins/Experimental/ControlRig/Source/ControlRig/Private/Units/Highlevel/Hierarchy/RigUnit_FABRIK.cpp
@@ -12,6 +12,10 @@ UE_RigUnit_FABRIK_IMPLEMENT_MULTIPLEX(void, Execute, const FRigUnitContext& Cont
 		return;
 	}
 
+	TArray<FFABRIKChainLink>& Chain = WorkData.Chain;
+	TArray<int32>& BoneIndices = WorkData.BoneIndices;
+	int32& EffectorIndex = WorkData.EffectorIndex;
+
 	if (Context.State == EControlRigState::Init)
 	{
 		BoneIndices.Reset();
diff --git a/Engine/Plugins/Experimental/ControlRig/Source/ControlRig/Private/Units/Highlevel/Hierarchy/RigUnit_FABRIK.h b/Engine/Plugins/Experimental/ControlRig/Source/ControlRig/Private/Units/Highlevel/Hierarchy/RigUnit_FABRIK.h
--- a/Engine/Plugins/Experimental/ControlRig/Source/ControlRig/Private/Units/Highlevel/Hierarchy/RigUnit_FABRIK.h
+++ b/Engine/Plugins/Experimental/ControlRig/Source/ControlRig/Private/Units/Highlevel/Hierarchy/RigUnit_FABRIK.h
@@ -6,6 +6,26 @@
 #include "FABRIK.h"
 #include "RigUnit_FABRIK.generated.h"
 
+USTRUCT()
+struct FRigUnit_FABRIK_WorkData
+{
+	GENERATED_BODY()
+
+	FRigUnit_FABRIK_WorkData()
+	{
+		EffectorIndex = INDEX_NONE;
+	}
+
+	UPROPERTY()
+	TArray<FFABRIKChainLink> Chain;
+
+	UPROPERTY()
+	TArray<int32> BoneIndices;
+
+	UPROPERTY()
+	int32 EffectorIndex;
+};
+
 /**
  * The FABRIK solver can solve N-Bone chains using 
  * the Forward and Backward Reaching Inverse Kinematics algorithm.
@@ -21,7 +41,6 @@ struct FRigUnit_FABRIK : public FRigUnit_HighlevelBaseMutable
 
 	FRigUnit_FABRIK()
 	{
-		EffectorIndex = INDEX_NONE;
 		Precision = 1.f;
 		MaxIterations = 10;
 		EffectorTransform = FTransform::Identity;
@@ -67,11 +86,5 @@ struct FRigUnit_FABRIK : public FRigUnit_HighlevelBaseMutable
 	int32 MaxIterations;
 
 	UPROPERTY(transient)
-	TArray<FFABRIKChainLink> Chain;
-
-	UPROPERTY(transient)
-	TArray<int32> BoneIndices;
-
-	UPROPERTY(transient)
-	int32 EffectorIndex;
+	FRigUnit_FABRIK_WorkData WorkData;
 };
diff --git a/Engine/Plugins/Experimental/ControlRig/Source/ControlRig/Private/Units/Highlevel/Hierarchy/RigUnit_FitChainToCurve.cpp b/Engine/Plugins/Experimental/ControlRig/Source/ControlRig/Private/Units/Highlevel/Hierarchy/RigUnit_FitChainToCurve.cpp
--- a/Engine/Plugins/Experimental/ControlRig/Source/ControlRig/Private/Units/Highlevel/Hierarchy/RigUnit_FitChainToCurve.cpp
+++ b/Engine/Plugins/Experimental/ControlRig/Source/ControlRig/Private/Units/Highlevel/Hierarchy/RigUnit_FitChainToCurve.cpp
@@ -12,6 +12,17 @@ UE_RigUnit_FitChainToCurve_IMPLEMENT_MULTIPLEX(void, Execute, const FRigUnitCont
 		return;
 	}
 
+	float& ChainLength = WorkData.ChainLength;
+	TArray<FVector>& BonePositions = WorkData.BonePositions;
+	TArray<float>& BoneSegments = WorkData.BoneSegments;
+	TArray<FVector>& CurvePositions = WorkData.CurvePositions;
+	TArray<float>& CurveSegments = WorkData.CurveSegments;
+	TArray<int32>& BoneIndices = WorkData.BoneIndices;
+	TArray<int32>& BoneRotationA = WorkData.BoneRotationA;
+	TArray<int32>& BoneRotationB = WorkData.BoneRotationB;
+	TArray<float>& BoneRotationT = WorkData.BoneRotationT;
+	TArray<FTransform>& BoneLocalTransforms = WorkData.BoneLocalTransforms;
+
 	if (Context.State == EControlRigState::Init)
 	{
 		BoneIndices.Reset();
diff --git a/Engine/Plugins/Experimental/ControlRig/Source/ControlRig/Private/Units/Highlevel/Hierarchy/RigUnit_FitChainToCurve.h b/Engine/Plugins/Experimental/ControlRig/Source/ControlRig/Private/Units/Highlevel/Hierarchy/RigUnit_FitChainToCurve.h
--- a/Engine/Plugins/Experimental/ControlRig/Source/ControlRig/Private/Units/Highlevel/Hierarchy/RigUnit_FitChainToCurve.h
+++ b/Engine/Plugins/Experimental/ControlRig/Source/ControlRig/Private/Units/Highlevel/Hierarchy/RigUnit_FitChainToCurve.h
@@ -82,6 +82,46 @@
 	FTransform WorldOffset;
 };
 
+USTRUCT()
+struct FRigUnit_FitChainToCurve_WorkData
+{
+	GENERATED_BODY()
+
+	FRigUnit_FitChainToCurve_WorkData()
+	{
+		ChainLength = 0.f;
+	}
+
+	UPROPERTY()
+	float ChainLength;
+
+	UPROPERTY()
+	TArray<FVector> BonePositions;
+
+	UPROPERTY()
+	TArray<float> BoneSegments;
+
+	UPROPERTY()
+	TArray<FVector> CurvePositions;
+
+	UPROPERTY()
+	TArray<float> CurveSegments;
+
+	UPROPERTY()
+	TArray<int32> BoneIndices;
+
+	UPROPERTY()
+	TArray<int32> BoneRotationA;
+
+	UPROPERTY()
+	TArray<int32> BoneRotationB;
+
+	UPROPERTY()
+	TArray<float> BoneRotationT;
+
+	UPROPERTY()
+	TArray<FTransform> BoneLocalTransforms;
+};
 /**
  * Fits a given chain to a four point bezier curve.
  * Additionally provides rotational control matching the features of the Distribute Rotation node.
@@ -105,7 +145,6 @@
 		RotationEaseType = EControlRigAnimEasingType::Linear;
 		bPropagateToChildren = false;
 		DebugSettings = FRigUnit_FitChainToCurve_DebugSettings();
-		ChainLength = 0.f;
 	}
 
 	MULTIPLEX_METHOD()
@@ -196,33 +235,6 @@
 	UPROPERTY(meta = (Input))
 	FRigUnit_FitChainToCurve_DebugSettings DebugSettings;
 
-	UPROPERTY()
-	float ChainLength;
-
-	UPROPERTY()
-	TArray<FVector> BonePositions;
-
-	UPROPERTY()
-	TArray<float> BoneSegments;
-
-	UPROPERTY()
-	TArray<FVector> CurvePositions;
-
-	UPROPERTY()
-	TArray<float> CurveSegments;
-
-	UPROPERTY()
-	TArray<int32> BoneIndices;
-
-	UPROPERTY()
-	TArray<int32> BoneRotationA;
-
-	UPROPERTY()
-	TArray<int32> BoneRotationB;
-
-	UPROPERTY()
-	TArray<float> BoneRotationT;
-
-	UPROPERTY()
-	TArray<FTransform> BoneLocalTransforms;
+	UPROPERTY(transient)
+	FRigUnit_FitChainToCurve_WorkData WorkData;
 };
diff --git a/Engine/Plugins/Experimental/ControlRig/Source/ControlRig/Private/Units/Highlevel/Hierarchy/RigUnit_ModifyBoneTransforms.cpp b/Engine/Plugins/Experimental/ControlRig/Source/ControlRig/Private/Units/Highlevel/Hierarchy/RigUnit_ModifyBoneTransforms.cpp
--- a/Engine/Plugins/Experimental/ControlRig/Source/ControlRig/Private/Units/Highlevel/Hierarchy/RigUnit_ModifyBoneTransforms.cpp
+++ b/Engine/Plugins/Experimental/ControlRig/Source/ControlRig/Private/Units/Highlevel/Hierarchy/RigUnit_ModifyBoneTransforms.cpp
@@ -6,6 +6,9 @@
 UE_RigUnit_ModifyBoneTransforms_IMPLEMENT_MULTIPLEX(void, Execute, const FRigUnitContext& Context)
 {
     DECLARE_SCOPE_HIERARCHICAL_COUNTER_RIGUNIT()
+
+	TArray<int32>& CachedBoneIndices = WorkData.CachedBoneIndices;
+
 	FRigBoneHierarchy* Hierarchy = ExecuteContext.GetBones();
 	if (Hierarchy)
 	{
diff --git a/Engine/Plugins/Experimental/ControlRig/Source/ControlRig/Private/Units/Highlevel/Hierarchy/RigUnit_ModifyBoneTransforms.h b/Engine/Plugins/Experimental/ControlRig/Source/ControlRig/Private/Units/Highlevel/Hierarchy/RigUnit_ModifyBoneTransforms.h
--- a/Engine/Plugins/Experimental/ControlRig/Source/ControlRig/Private/Units/Highlevel/Hierarchy/RigUnit_ModifyBoneTransforms.h
+++ b/Engine/Plugins/Experimental/ControlRig/Source/ControlRig/Private/Units/Highlevel/Hierarchy/RigUnit_ModifyBoneTransforms.h
@@ -54,6 +54,14 @@ struct FRigUnit_ModifyBoneTransforms_PerBone
 	FTransform Transform;
 };
 
+USTRUCT()
+struct FRigUnit_ModifyBoneTransforms_WorkData
+{
+	GENERATED_BODY()
+
+	UPROPERTY()
+	TArray<int32> CachedBoneIndices;
+};
 
 /**
  * SetBoneTransform is used to perform a change in the hierarchy by setting a single bone's transform.
@@ -105,6 +113,6 @@ struct FRigUnit_ModifyBoneTransforms : public FRigUnit_HighlevelBaseMutable
 	EControlRigModifyBoneMode Mode;
 
 	// Used to cache the internally used bone index
-	UPROPERTY()
-	TArray<int32> CachedBoneIndices;
+	UPROPERTY(transient)
+	FRigUnit_ModifyBoneTransforms_WorkData WorkData;
 };
diff --git a/Engine/Plugins/Experimental/ControlRig/Source/ControlRig/Private/Units/Highlevel/Hierarchy/RigUnit_SlideChain.cpp b/Engine/Plugins/Experimental/ControlRig/Source/ControlRig/Private/Units/Highlevel/Hierarchy/RigUnit_SlideChain.cpp
--- a/Engine/Plugins/Experimental/ControlRig/Source/ControlRig/Private/Units/Highlevel/Hierarchy/RigUnit_SlideChain.cpp
+++ b/Engine/Plugins/Experimental/ControlRig/Source/ControlRig/Private/Units/Highlevel/Hierarchy/RigUnit_SlideChain.cpp
@@ -13,6 +13,12 @@ UE_RigUnit_SlideChain_IMPLEMENT_MULTIPLEX(void, Execute, const FRigUnitContext&
 		return;
 	}
 
+	float& ChainLength = WorkData.ChainLength;
+	TArray<float>& BoneSegments = WorkData.BoneSegments;
+	TArray<int32>& BoneIndices = WorkData.BoneIndices;
+	TArray<FTransform>& Transforms = WorkData.Transforms;
+	TArray<FTransform>& BlendedTransforms = WorkData.BlendedTransforms;
+
 	if (Context.State == EControlRigState::Init)
 	{
 		BoneSegments.Reset();
diff --git a/Engine/Plugins/Experimental/ControlRig/Source/ControlRig/Private/Units/Highlevel/Hierarchy/RigUnit_SlideChain.h b/Engine/Plugins/Experimental/ControlRig/Source/ControlRig/Private/Units/Highlevel/Hierarchy/RigUnit_SlideChain.h
--- a/Engine/Plugins/Experimental/ControlRig/Source/ControlRig/Private/Units/Highlevel/Hierarchy/RigUnit_SlideChain.h
+++ b/Engine/Plugins/Experimental/ControlRig/Source/ControlRig/Private/Units/Highlevel/Hierarchy/RigUnit_SlideChain.h
@@ -6,6 +6,32 @@
 #include "Math/ControlRigMathLibrary.h"
 #include "RigUnit_SlideChain.generated.h"
 
+USTRUCT()
+struct FRigUnit_SlideChain_WorkData
+{
+	GENERATED_BODY()
+
+	FRigUnit_SlideChain_WorkData()
+	{
+		ChainLength = 0.f;
+	}
+
+	UPROPERTY()
+	float ChainLength;
+
+	UPROPERTY()
+	TArray<float> BoneSegments;
+
+	UPROPERTY()
+	TArray<int32> BoneIndices;
+
+	UPROPERTY()
+	TArray<FTransform> Transforms;
+
+	UPROPERTY()
+	TArray<FTransform> BlendedTransforms;
+};
+
 /**
  * Slides an existing chain along itself with control over extrapolation.
  */
@@ -19,7 +45,6 @@ struct FRigUnit_SlideChain: public FRigUnit_HighlevelBaseMutable
 		StartBone = EndBone = NAME_None;
 		SlideAmount = 0.f;
 		bPropagateToChildren = false;
-		ChainLength = 0.f;
 	}
 
 	MULTIPLEX_METHOD()
@@ -51,18 +76,6 @@ struct FRigUnit_SlideChain: public FRigUnit_HighlevelBaseMutable
 	UPROPERTY(meta = (Input))
 	bool bPropagateToChildren;
 
-	UPROPERTY()
-	float ChainLength;
-
-	UPROPERTY()
-	TArray<float> BoneSegments;
-
-	UPROPERTY()
-	TArray<int32> BoneIndices;
-
-	UPROPERTY()
-	TArray<FTransform> Transforms;
-
-	UPROPERTY()
-	TArray<FTransform> BlendedTransforms;
+	UPROPERTY(transient)
+	FRigUnit_SlideChain_WorkData WorkData;
 };
diff --git a/Engine/Plugins/Experimental/ControlRig/Source/ControlRig/Private/Units/Highlevel/Hierarchy/RigUnit_SpringIK.cpp b/Engine/Plugins/Experimental/ControlRig/Source/ControlRig/Private/Units/Highlevel/Hierarchy/RigUnit_SpringIK.cpp
--- a/Engine/Plugins/Experimental/ControlRig/Source/ControlRig/Private/Units/Highlevel/Hierarchy/RigUnit_SpringIK.cpp
+++ b/Engine/Plugins/Experimental/ControlRig/Source/ControlRig/Private/Units/Highlevel/Hierarchy/RigUnit_SpringIK.cpp
@@ -12,6 +12,11 @@ UE_RigUnit_SpringIK_IMPLEMENT_MULTIPLEX(void, Execute, const FRigUnitContext& Co
 		return;
 	}
 
+	TArray<int32>& BoneIndices = WorkData.BoneIndices;
+	int32& PoleVectorIndex = WorkData.PoleVectorIndex;
+	TArray<FTransform>& Transforms = WorkData.Transforms;
+	FCRSimPointContainer& Simulation = WorkData.Simulation;
+
 	if (Context.State == EControlRigState::Init)
 	{
 		BoneIndices.Reset();
diff --git a/Engine/Plugins/Experimental/ControlRig/Source/ControlRig/Private/Units/Highlevel/Hierarchy/RigUnit_SpringIK.h b/Engine/Plugins/Experimental/ControlRig/Source/ControlRig/Private/Units/Highlevel/Hierarchy/RigUnit_SpringIK.h
--- a/Engine/Plugins/Experimental/ControlRig/Source/ControlRig/Private/Units/Highlevel/Hierarchy/RigUnit_SpringIK.h
+++ b/Engine/Plugins/Experimental/ControlRig/Source/ControlRig/Private/Units/Highlevel/Hierarchy/RigUnit_SpringIK.h
@@ -44,6 +44,29 @@ struct FRigUnit_SpringIK_DebugSettings
 	FTransform WorldOffset;
 };
 
+USTRUCT()
+struct FRigUnit_SpringIK_WorkData
+{
+	GENERATED_BODY()
+
+	FRigUnit_SpringIK_WorkData()
+	{
+		PoleVectorIndex = INDEX_NONE;
+	}
+
+	UPROPERTY()
+	TArray<int32> BoneIndices;
+
+	UPROPERTY()
+	int32 PoleVectorIndex;
+
+	UPROPERTY()
+	TArray<FTransform> Transforms;
+
+	UPROPERTY()
+	FCRSimPointContainer Simulation;
+};
+
 /**
  * The Spring IK solver uses a verlet integrator to perform an IK solve.
  * It support custom constraints including distance, length etc.
@@ -71,7 +94,6 @@ struct FRigUnit_SpringIK : public FRigUnit_HighlevelBaseMutable
 		bLimitLocalPosition = true;
 		bPropagateToChildren = false;
 		DebugSettings = FRigUnit_SpringIK_DebugSettings();
-		PoleVectorIndex = INDEX_NONE;
 	}
 
 	MULTIPLEX_METHOD()
@@ -196,15 +218,6 @@ struct FRigUnit_SpringIK : public FRigUnit_HighlevelBaseMutable
 	UPROPERTY(meta = (Input))
 	FRigUnit_SpringIK_DebugSettings DebugSettings;
 
-	UPROPERTY()
-	TArray<int32> BoneIndices;
-
-	UPROPERTY()
-	int32 PoleVectorIndex;
-
-	UPROPERTY()
-	TArray<FTransform> Transforms;
-
-	UPROPERTY()
-	FCRSimPointContainer Simulation;
+	UPROPERTY(transient)
+	FRigUnit_SpringIK_WorkData WorkData;
 };
diff --git a/Engine/Plugins/Experimental/ControlRig/Source/ControlRig/Private/Units/Highlevel/Hierarchy/RigUnit_TransformConstraint.cpp b/Engine/Plugins/Experimental/ControlRig/Source/ControlRig/Private/Units/Highlevel/Hierarchy/RigUnit_TransformConstraint.cpp
--- a/Engine/Plugins/Experimental/ControlRig/Source/ControlRig/Private/Units/Highlevel/Hierarchy/RigUnit_TransformConstraint.cpp
+++ b/Engine/Plugins/Experimental/ControlRig/Source/ControlRig/Private/Units/Highlevel/Hierarchy/RigUnit_TransformConstraint.cpp
@@ -9,6 +9,9 @@ UE_RigUnit_TransformConstraint_IMPLEMENT_MULTIPLEX(void, Execute, const FRigUnit
 {
     DECLARE_SCOPE_HIERARCHICAL_COUNTER_RIGUNIT()
 
+	TArray<FConstraintData>&	ConstraintData = WorkData.ConstraintData;
+	TMap<int32, int32>& ConstraintDataToTargets = WorkData.ConstraintDataToTargets;
+
 	if (Context.State == EControlRigState::Init)
 	{
 		ConstraintData.Reset();
diff --git a/Engine/Plugins/Experimental/ControlRig/Source/ControlRig/Private/Units/Highlevel/Hierarchy/RigUnit_TransformConstraint.h b/Engine/Plugins/Experimental/ControlRig/Source/ControlRig/Private/Units/Highlevel/Hierarchy/RigUnit_TransformConstraint.h
--- a/Engine/Plugins/Experimental/ControlRig/Source/ControlRig/Private/Units/Highlevel/Hierarchy/RigUnit_TransformConstraint.h
+++ b/Engine/Plugins/Experimental/ControlRig/Source/ControlRig/Private/Units/Highlevel/Hierarchy/RigUnit_TransformConstraint.h
@@ -36,6 +36,19 @@ struct FConstraintTarget
 	{}
 };
 
+USTRUCT()
+struct FRigUnit_TransformConstraint_WorkData
+{
+	GENERATED_BODY()
+
+	// note that Targets.Num () != ConstraintData.Num()
+	UPROPERTY()
+	TArray<FConstraintData>	ConstraintData;
+
+	UPROPERTY()
+	TMap<int32, int32> ConstraintDataToTargets;
+};
+
 USTRUCT(meta=(DisplayName="Transform Constraint", Category="Transforms"))
 struct FRigUnit_TransformConstraint : public FRigUnit_HighlevelBaseMutable
 {
@@ -68,10 +81,6 @@ struct FRigUnit_TransformConstraint : public FRigUnit_HighlevelBaseMutable
 private:
 	static void AddConstraintData(const TArrayView<FConstraintTarget>& Targets, ETransformConstraintType ConstraintType, const int32 TargetIndex, const FTransform& SourceTransform, const FTransform& InBaseTransform, TArray<FConstraintData>& OutConstraintData, TMap<int32, int32>& OutConstraintDataToTargets);
 
-	// note that Targets.Num () != ConstraintData.Num()
 	UPROPERTY(transient)
-	TArray<FConstraintData>	ConstraintData;
-
-	UPROPERTY(transient)
-	TMap<int32, int32> ConstraintDataToTargets;
+	FRigUnit_TransformConstraint_WorkData WorkData;
 };
diff --git a/Engine/Plugins/Experimental/ControlRig/Source/ControlRig/Private/Units/Simulation/RigUnit_PointSimulation.cpp b/Engine/Plugins/Experimental/ControlRig/Source/ControlRig/Private/Units/Simulation/RigUnit_PointSimulation.cpp
--- a/Engine/Plugins/Experimental/ControlRig/Source/ControlRig/Private/Units/Simulation/RigUnit_PointSimulation.cpp
+++ b/Engine/Plugins/Experimental/ControlRig/Source/ControlRig/Private/Units/Simulation/RigUnit_PointSimulation.cpp
@@ -10,6 +10,9 @@ UE_RigUnit_PointSimulation_IMPLEMENT_MULTIPLEX(void, Execute, const FRigUnitCont
 
 	float DeltaTime = Context.DeltaTime;
 
+	FCRSimPointContainer& Simulation = WorkData.Simulation;
+	TArray<int32>& BoneIndices = WorkData.BoneIndices;
+
 	if (Context.State == EControlRigState::Init ||
 		Simulation.Points.Num() != Points.Num() ||
 		Simulation.Springs.Num() != Links.Num() ||
diff --git a/Engine/Plugins/Experimental/ControlRig/Source/ControlRig/Private/Units/Simulation/RigUnit_PointSimulation.h b/Engine/Plugins/Experimental/ControlRig/Source/ControlRig/Private/Units/Simulation/RigUnit_PointSimulation.h
--- a/Engine/Plugins/Experimental/ControlRig/Source/ControlRig/Private/Units/Simulation/RigUnit_PointSimulation.h
+++ b/Engine/Plugins/Experimental/ControlRig/Source/ControlRig/Private/Units/Simulation/RigUnit_PointSimulation.h
@@ -101,6 +101,18 @@ struct FRigUnit_PointSimulation_BoneTarget
 	int32 SecondaryAimPoint;
 };
 
+USTRUCT()
+struct FRigUnit_PointSimulation_WorkData
+{
+	GENERATED_BODY()
+
+	UPROPERTY()
+	FCRSimPointContainer Simulation;
+
+	UPROPERTY()
+	TArray<int32> BoneIndices;
+};
+
 /**
  * Performs point based simulation
  * Note: Disabled for now.
@@ -189,10 +201,7 @@ struct FRigUnit_PointSimulation : public FRigUnit_SimBaseMutable
 	UPROPERTY(meta = (Output))
 	FCRFourPointBezier Bezier;
 
-	UPROPERTY()
-	FCRSimPointContainer Simulation;
-
-	UPROPERTY()
-	TArray<int32> BoneIndices;
+	UPROPERTY(transient)
+	FRigUnit_PointSimulation_WorkData WorkData;
 };
 
diff --git a/Engine/Plugins/Experimental/ControlRig/Source/ControlRig/Private/Units/Simulation/RigUnit_TimeOffset.cpp b/Engine/Plugins/Experimental/ControlRig/Source/ControlRig/Private/Units/Simulation/RigUnit_TimeOffset.cpp
--- a/Engine/Plugins/Experimental/ControlRig/Source/ControlRig/Private/Units/Simulation/RigUnit_TimeOffset.cpp
+++ b/Engine/Plugins/Experimental/ControlRig/Source/ControlRig/Private/Units/Simulation/RigUnit_TimeOffset.cpp
@@ -31,15 +31,12 @@ UE_RigUnit_TimeOffsetFloat_IMPLEMENT_MULTIPLEX(void, Execute, const FRigUnitCont
 	int32 MaxSize = FMath::Clamp<int32>(BufferSize, 2, 512);
 	if (Context.State == EControlRigState::Init)
 	{
-		Buffer.Reset();
-		Buffer.Reserve(MaxSize);
-		DeltaTimes.Reset();
-		DeltaTimes.Reserve(MaxSize);
+		UpperBound = 0;
 		Result = Value;
 		return;
 	}
 
-	if (Buffer.Num() == 0)
+	if (UpperBound == 0)
 	{
 		Result = Value;
 	}
@@ -50,7 +47,7 @@ UE_RigUnit_TimeOffsetFloat_IMPLEMENT_MULTIPLEX(void, Execute, const FRigUnitCont
 		float T = 0.f;
 
 		int32 Index = LastInsertIndex;
-		if (SecondsAgo < DeltaTimes[Index] || Buffer.Num() == 1)
+		if (SecondsAgo < DeltaTimes[Index] || UpperBound == 1)
 		{
 			B = Buffer[Index];
 			T = SecondsAgo / DeltaTimes[Index];
@@ -58,7 +55,7 @@ UE_RigUnit_TimeOffsetFloat_IMPLEMENT_MULTIPLEX(void, Execute, const FRigUnitCont
 		else
 		{
 			float AccumulatedTime = 0.f;
-			for (int32 TimeIndex = 0; TimeIndex < DeltaTimes.Num(); TimeIndex++)
+			for (int32 TimeIndex = 0; TimeIndex < UpperBound; TimeIndex++)
 			{
 				B = Buffer[Index];
 
@@ -70,7 +67,7 @@ UE_RigUnit_TimeOffsetFloat_IMPLEMENT_MULTIPLEX(void, Execute, const FRigUnitCont
 				AccumulatedTime = AccumulatedTime + DeltaTimes[Index];;
 				A = B;
 
-				Index = (Index - 1 + DeltaTimes.Num()) % DeltaTimes.Num();
+				Index = (Index - 1 + UpperBound) % UpperBound;
 			}
 		}
 
@@ -79,26 +76,26 @@ UE_RigUnit_TimeOffsetFloat_IMPLEMENT_MULTIPLEX(void, Execute, const FRigUnitCont
 
 	if (Context.DeltaTime > SMALL_NUMBER)
 	{
-		if (Buffer.Num() == 0)
+		if (UpperBound == 0)
 		{
-			LastInsertIndex = Buffer.Num();
-			Buffer.Add(Value);
-			DeltaTimes.Add(Context.DeltaTime);
+			LastInsertIndex = UpperBound;
+			Buffer[UpperBound] = Value;
+			DeltaTimes[UpperBound++] = Context.DeltaTime;
 		}
 		else
 		{
 			float SecondsPerEntry = TimeRange / float(MaxSize - 1);
 			if (DeltaTimes[LastInsertIndex] > SecondsPerEntry - Context.DeltaTime * 0.5f)
 			{
-				if (Buffer.Num() < MaxSize)
+				if (UpperBound < MaxSize)
 				{
-					LastInsertIndex = Buffer.Num();
-					Buffer.Add(Value);
-					DeltaTimes.Add(Context.DeltaTime);
+					LastInsertIndex = UpperBound;
+					Buffer[UpperBound] = Value;
+					DeltaTimes[UpperBound++] = Context.DeltaTime;
 				}
 				else
 				{
-					LastInsertIndex = (LastInsertIndex + 1) % Buffer.Num();
+					LastInsertIndex = (LastInsertIndex + 1) % UpperBound;
 					Buffer[LastInsertIndex] = Value;
 					DeltaTimes[LastInsertIndex] = Context.DeltaTime;
 				}
@@ -139,15 +136,12 @@ UE_RigUnit_TimeOffsetVector_IMPLEMENT_MULTIPLEX(void, Execute, const FRigUnitCon
 	int32 MaxSize = FMath::Clamp<int32>(BufferSize, 2, 512);
 	if (Context.State == EControlRigState::Init)
 	{
-		Buffer.Reset();
-		Buffer.Reserve(MaxSize);
-		DeltaTimes.Reset();
-		DeltaTimes.Reserve(MaxSize);
+		UpperBound = 0;
 		Result = Value;
 		return;
 	}
 
-	if (Buffer.Num() == 0)
+	if (UpperBound == 0)
 	{
 		Result = Value;
 	}
@@ -158,7 +152,7 @@ UE_RigUnit_TimeOffsetVector_IMPLEMENT_MULTIPLEX(void, Execute, const FRigUnitCon
 		float T = 0.f;
 
 		int32 Index = LastInsertIndex;
-		if (SecondsAgo < DeltaTimes[Index] || Buffer.Num() == 1)
+		if (SecondsAgo < DeltaTimes[Index] || UpperBound == 1)
 		{
 			B = Buffer[Index];
 			T = SecondsAgo / DeltaTimes[Index];
@@ -166,7 +160,7 @@ UE_RigUnit_TimeOffsetVector_IMPLEMENT_MULTIPLEX(void, Execute, const FRigUnitCon
 		else
 		{
 			float AccumulatedTime = 0.f;
-			for (int32 TimeIndex = 0; TimeIndex < DeltaTimes.Num(); TimeIndex++)
+			for (int32 TimeIndex = 0; TimeIndex < UpperBound; TimeIndex++)
 			{
 				B = Buffer[Index];
 
@@ -178,7 +172,7 @@ UE_RigUnit_TimeOffsetVector_IMPLEMENT_MULTIPLEX(void, Execute, const FRigUnitCon
 				AccumulatedTime = AccumulatedTime + DeltaTimes[Index];;
 				A = B;
 
-				Index = (Index - 1 + DeltaTimes.Num()) % DeltaTimes.Num();
+				Index = (Index - 1 + UpperBound) % UpperBound;
 			}
 		}
 
@@ -187,26 +181,26 @@ UE_RigUnit_TimeOffsetVector_IMPLEMENT_MULTIPLEX(void, Execute, const FRigUnitCon
 
 	if (Context.DeltaTime > SMALL_NUMBER)
 	{
-		if (Buffer.Num() == 0)
+		if (UpperBound == 0)
 		{
-			LastInsertIndex = Buffer.Num();
-			Buffer.Add(Value);
-			DeltaTimes.Add(Context.DeltaTime);
+			LastInsertIndex = UpperBound;
+			Buffer[UpperBound] = Value;
+			DeltaTimes[UpperBound++] = Context.DeltaTime;
 		}
 		else
 		{
 			float SecondsPerEntry = TimeRange / float(MaxSize - 1);
 			if (DeltaTimes[LastInsertIndex] > SecondsPerEntry - Context.DeltaTime * 0.5f)
 			{
-				if (Buffer.Num() < MaxSize)
+				if (UpperBound < MaxSize)
 				{
-					LastInsertIndex = Buffer.Num();
-					Buffer.Add(Value);
-					DeltaTimes.Add(Context.DeltaTime);
+					LastInsertIndex = UpperBound;
+					Buffer[UpperBound] = Value;
+					DeltaTimes[UpperBound++] = Context.DeltaTime;
 				}
 				else
 				{
-					LastInsertIndex = (LastInsertIndex + 1) % Buffer.Num();
+					LastInsertIndex = (LastInsertIndex + 1) % UpperBound;
 					Buffer[LastInsertIndex] = Value;
 					DeltaTimes[LastInsertIndex] = Context.DeltaTime;
 				}
@@ -247,15 +241,12 @@ UE_RigUnit_TimeOffsetTransform_IMPLEMENT_MULTIPLEX(void, Execute, const FRigUnit
 	int32 MaxSize = FMath::Clamp<int32>(BufferSize, 2, 512);
 	if (Context.State == EControlRigState::Init)
 	{
-		Buffer.Reset();
-		Buffer.Reserve(MaxSize);
-		DeltaTimes.Reset();
-		DeltaTimes.Reserve(MaxSize);
+		UpperBound = 0;
 		Result = Value;
 		return;
 	}
 
-	if (Buffer.Num() == 0)
+	if (UpperBound == 0)
 	{
 		Result = Value;
 	}
@@ -266,7 +257,7 @@ UE_RigUnit_TimeOffsetTransform_IMPLEMENT_MULTIPLEX(void, Execute, const FRigUnit
 		float T = 0.f;
 
 		int32 Index = LastInsertIndex;
-		if (SecondsAgo < DeltaTimes[Index] || Buffer.Num() == 1)
+		if (SecondsAgo < DeltaTimes[Index] || UpperBound == 1)
 		{
 			B = Buffer[Index];
 			T = SecondsAgo / DeltaTimes[Index];
@@ -274,7 +265,7 @@ UE_RigUnit_TimeOffsetTransform_IMPLEMENT_MULTIPLEX(void, Execute, const FRigUnit
 		else
 		{
 			float AccumulatedTime = 0.f;
-			for (int32 TimeIndex = 0; TimeIndex < DeltaTimes.Num(); TimeIndex++)
+			for (int32 TimeIndex = 0; TimeIndex < UpperBound; TimeIndex++)
 			{
 				B = Buffer[Index];
 
@@ -286,7 +277,7 @@ UE_RigUnit_TimeOffsetTransform_IMPLEMENT_MULTIPLEX(void, Execute, const FRigUnit
 				AccumulatedTime = AccumulatedTime + DeltaTimes[Index];;
 				A = B;
 
-				Index = (Index - 1 + DeltaTimes.Num()) % DeltaTimes.Num();
+				Index = (Index - 1 + UpperBound) % UpperBound;
 			}
 		}
 
@@ -297,26 +288,26 @@ UE_RigUnit_TimeOffsetTransform_IMPLEMENT_MULTIPLEX(void, Execute, const FRigUnit
 
 	if (Context.DeltaTime > SMALL_NUMBER)
 	{
-		if (Buffer.Num() == 0)
+		if (UpperBound == 0)
 		{
-			LastInsertIndex = Buffer.Num();
-			Buffer.Add(Value);
-			DeltaTimes.Add(Context.DeltaTime);
+			LastInsertIndex = UpperBound;
+			Buffer[UpperBound] = Value;
+			DeltaTimes[UpperBound++] = Context.DeltaTime;
 		}
 		else
 		{
 			float SecondsPerEntry = TimeRange / float(MaxSize - 1);
 			if (DeltaTimes[LastInsertIndex] > SecondsPerEntry - Context.DeltaTime * 0.5f)
 			{
-				if (Buffer.Num() < MaxSize)
+				if (UpperBound < MaxSize)
 				{
-					LastInsertIndex = Buffer.Num();
-					Buffer.Add(Value);
-					DeltaTimes.Add(Context.DeltaTime);
+					LastInsertIndex = UpperBound;
+					Buffer[UpperBound] = Value;
+					DeltaTimes[UpperBound++] = Context.DeltaTime;
 				}
 				else
 				{
-					LastInsertIndex = (LastInsertIndex + 1) % Buffer.Num();
+					LastInsertIndex = (LastInsertIndex + 1) % UpperBound;
 					Buffer[LastInsertIndex] = Value;
 					DeltaTimes[LastInsertIndex] = Context.DeltaTime;
 				}
diff --git a/Engine/Plugins/Experimental/ControlRig/Source/ControlRig/Private/Units/Simulation/RigUnit_TimeOffset.h b/Engine/Plugins/Experimental/ControlRig/Source/ControlRig/Private/Units/Simulation/RigUnit_TimeOffset.h
--- a/Engine/Plugins/Experimental/ControlRig/Source/ControlRig/Private/Units/Simulation/RigUnit_TimeOffset.h
+++ b/Engine/Plugins/Experimental/ControlRig/Source/ControlRig/Private/Units/Simulation/RigUnit_TimeOffset.h
@@ -20,6 +20,7 @@ struct FRigUnit_TimeOffsetFloat : public FRigUnit_SimBase
 		BufferSize = 16;
 		TimeRange = 1.f;
 		LastInsertIndex = 0;
+		UpperBound = 0;
 	}
 
 	MULTIPLEX_METHOD()
@@ -44,14 +45,17 @@ struct FRigUnit_TimeOffsetFloat : public FRigUnit_SimBase
 	UPROPERTY(meta=(Output))
 	float Result;
 
-	UPROPERTY()
+	UPROPERTY(meta = (MaxArraySize = "FMath::Clamp<int32>(BufferSize, 2, 512)"))
 	TArray<float> Buffer;
 
-	UPROPERTY()
+	UPROPERTY(meta = (MaxArraySize = "FMath::Clamp<int32>(BufferSize, 2, 512)"))
 	TArray<float> DeltaTimes;
 
 	UPROPERTY()
 	int32 LastInsertIndex;
+
+	UPROPERTY()
+	int32 UpperBound;
 };
 
 /**
@@ -69,6 +73,7 @@ struct FRigUnit_TimeOffsetVector : public FRigUnit_SimBase
 		BufferSize = 16;
 		TimeRange = 1.f;
 		LastInsertIndex = 0;
+		UpperBound = 0;
 	}
 
 	MULTIPLEX_METHOD()
@@ -93,14 +98,17 @@ struct FRigUnit_TimeOffsetVector : public FRigUnit_SimBase
 	UPROPERTY(meta=(Output))
 	FVector Result;
 
-	UPROPERTY()
+	UPROPERTY(meta = (MaxArraySize = "FMath::Clamp<int32>(BufferSize, 2, 512)"))
 	TArray<FVector> Buffer;
 
-	UPROPERTY()
+	UPROPERTY(meta = (MaxArraySize = "FMath::Clamp<int32>(BufferSize, 2, 512)"))
 	TArray<float> DeltaTimes;
 
 	UPROPERTY()
 	int32 LastInsertIndex;
+
+	UPROPERTY()
+	int32 UpperBound;
 };
 
 /**
@@ -118,6 +126,7 @@ struct FRigUnit_TimeOffsetTransform : public FRigUnit_SimBase
 		BufferSize = 16;
 		TimeRange = 1.f;
 		LastInsertIndex = 0;
+		UpperBound = 0;
 	}
 
 	MULTIPLEX_METHOD()
@@ -142,12 +151,15 @@ struct FRigUnit_TimeOffsetTransform : public FRigUnit_SimBase
 	UPROPERTY(meta=(Output))
 	FTransform Result;
 
-	UPROPERTY()
+	UPROPERTY(meta = (MaxArraySize = "FMath::Clamp<int32>(BufferSize, 2, 512)"))
 	TArray<FTransform> Buffer;
 
-	UPROPERTY()
+	UPROPERTY(meta = (MaxArraySize = "FMath::Clamp<int32>(BufferSize, 2, 512)"))
 	TArray<float> DeltaTimes;
 
 	UPROPERTY()
 	int32 LastInsertIndex;
+
+	UPROPERTY()
+	int32 UpperBound;
 };
