RigVM: Move from IMPLEMENT_RIGVM macro towards a more descriptive macro.

diff --git a/Engine/Plugins/Experimental/ControlRig/Source/ControlRig/Private/Units/Animation/RigUnit_AnimEasing.cpp b/Engine/Plugins/Experimental/ControlRig/Source/ControlRig/Private/Units/Animation/RigUnit_AnimEasing.cpp
--- a/Engine/Plugins/Experimental/ControlRig/Source/ControlRig/Private/Units/Animation/RigUnit_AnimEasing.cpp
+++ b/Engine/Plugins/Experimental/ControlRig/Source/ControlRig/Private/Units/Animation/RigUnit_AnimEasing.cpp
@@ -3,7 +3,7 @@
 #include "Units/Animation/RigUnit_AnimEasing.h"
 #include "Units/RigUnitContext.h"
 
-UE_RigUnit_AnimEasing_IMPLEMENT_RIGVM(void, Execute, const FRigUnitContext& Context)
+FRigUnit_AnimEasing_Execute()
 {
     DECLARE_SCOPE_HIERARCHICAL_COUNTER_RIGUNIT()
 	if (FMath::IsNearlyEqual(SourceMinimum, SourceMaximum))
diff --git a/Engine/Plugins/Experimental/ControlRig/Source/ControlRig/Private/Units/Animation/RigUnit_AnimEvalRichCurve.cpp b/Engine/Plugins/Experimental/ControlRig/Source/ControlRig/Private/Units/Animation/RigUnit_AnimEvalRichCurve.cpp
--- a/Engine/Plugins/Experimental/ControlRig/Source/ControlRig/Private/Units/Animation/RigUnit_AnimEvalRichCurve.cpp
+++ b/Engine/Plugins/Experimental/ControlRig/Source/ControlRig/Private/Units/Animation/RigUnit_AnimEvalRichCurve.cpp
@@ -3,7 +3,7 @@
 #include "Units/Animation/RigUnit_AnimEvalRichCurve.h"
 #include "Units/RigUnitContext.h"
 
-UE_RigUnit_AnimEvalRichCurve_IMPLEMENT_RIGVM(void, Execute, const FRigUnitContext& Context)
+FRigUnit_AnimEvalRichCurve_Execute()
 {
     DECLARE_SCOPE_HIERARCHICAL_COUNTER_RIGUNIT()
 	if (FMath::IsNearlyEqual(SourceMinimum, SourceMaximum))
diff --git a/Engine/Plugins/Experimental/ControlRig/Source/ControlRig/Private/Units/Animation/RigUnit_AnimRichCurve.cpp b/Engine/Plugins/Experimental/ControlRig/Source/ControlRig/Private/Units/Animation/RigUnit_AnimRichCurve.cpp
--- a/Engine/Plugins/Experimental/ControlRig/Source/ControlRig/Private/Units/Animation/RigUnit_AnimRichCurve.cpp
+++ b/Engine/Plugins/Experimental/ControlRig/Source/ControlRig/Private/Units/Animation/RigUnit_AnimRichCurve.cpp
@@ -3,7 +3,7 @@
 #include "Units/Animation/RigUnit_AnimRichCurve.h"
 #include "Units/RigUnitContext.h"
 
-UE_RigUnit_AnimRichCurve_IMPLEMENT_RIGVM(void, Execute, const FRigUnitContext& Context)
+FRigUnit_AnimRichCurve_Execute()
 {
     DECLARE_SCOPE_HIERARCHICAL_COUNTER_RIGUNIT()
 }
diff --git a/Engine/Plugins/Experimental/ControlRig/Source/ControlRig/Private/Units/Animation/RigUnit_GetDeltaTime.cpp b/Engine/Plugins/Experimental/ControlRig/Source/ControlRig/Private/Units/Animation/RigUnit_GetDeltaTime.cpp
--- a/Engine/Plugins/Experimental/ControlRig/Source/ControlRig/Private/Units/Animation/RigUnit_GetDeltaTime.cpp
+++ b/Engine/Plugins/Experimental/ControlRig/Source/ControlRig/Private/Units/Animation/RigUnit_GetDeltaTime.cpp
@@ -3,7 +3,7 @@
 #include "Units/Animation/RigUnit_GetDeltaTime.h"
 #include "Units/RigUnitContext.h"
 
-UE_RigUnit_GetDeltaTime_IMPLEMENT_RIGVM(void, Execute, const FRigUnitContext& Context)
+FRigUnit_GetDeltaTime_Execute()
 {
     DECLARE_SCOPE_HIERARCHICAL_COUNTER_RIGUNIT()
 	Result = Context.DeltaTime;
diff --git a/Engine/Plugins/Experimental/ControlRig/Source/ControlRig/Private/Units/Animation/RigUnit_GetWorldTime.cpp b/Engine/Plugins/Experimental/ControlRig/Source/ControlRig/Private/Units/Animation/RigUnit_GetWorldTime.cpp
--- a/Engine/Plugins/Experimental/ControlRig/Source/ControlRig/Private/Units/Animation/RigUnit_GetWorldTime.cpp
+++ b/Engine/Plugins/Experimental/ControlRig/Source/ControlRig/Private/Units/Animation/RigUnit_GetWorldTime.cpp
@@ -3,7 +3,7 @@
 #include "Units/Animation/RigUnit_GetWorldTime.h"
 #include "Units/RigUnitContext.h"
 
-UE_RigUnit_GetWorldTime_IMPLEMENT_RIGVM(void, Execute, const FRigUnitContext& Context)
+FRigUnit_GetWorldTime_Execute()
 {
     DECLARE_SCOPE_HIERARCHICAL_COUNTER_RIGUNIT()
 	FDateTime Now = FDateTime::Now();
diff --git a/Engine/Plugins/Experimental/ControlRig/Source/ControlRig/Private/Units/Control/RigUnit_Control.cpp b/Engine/Plugins/Experimental/ControlRig/Source/ControlRig/Private/Units/Control/RigUnit_Control.cpp
--- a/Engine/Plugins/Experimental/ControlRig/Source/ControlRig/Private/Units/Control/RigUnit_Control.cpp
+++ b/Engine/Plugins/Experimental/ControlRig/Source/ControlRig/Private/Units/Control/RigUnit_Control.cpp
@@ -3,7 +3,7 @@
 #include "Units/Control/RigUnit_Control.h"
 #include "Units/RigUnitContext.h"
 
-UE_RigUnit_Control_IMPLEMENT_RIGVM(void, Execute, const FRigUnitContext& Context)
+FRigUnit_Control_Execute()
 {
     DECLARE_SCOPE_HIERARCHICAL_COUNTER_RIGUNIT()
 	if (Context.State == EControlRigState::Init)
diff --git a/Engine/Plugins/Experimental/ControlRig/Source/ControlRig/Private/Units/Debug/RigUnit_DebugBezier.cpp b/Engine/Plugins/Experimental/ControlRig/Source/ControlRig/Private/Units/Debug/RigUnit_DebugBezier.cpp
--- a/Engine/Plugins/Experimental/ControlRig/Source/ControlRig/Private/Units/Debug/RigUnit_DebugBezier.cpp
+++ b/Engine/Plugins/Experimental/ControlRig/Source/ControlRig/Private/Units/Debug/RigUnit_DebugBezier.cpp
@@ -3,7 +3,7 @@
 #include "Units/Debug/RigUnit_DebugBezier.h"
 #include "Units/RigUnitContext.h"
 
-UE_RigUnit_DebugBezier_IMPLEMENT_RIGVM(void, Execute, const FRigUnitContext& Context)
+FRigUnit_DebugBezier_Execute()
 {
     DECLARE_SCOPE_HIERARCHICAL_COUNTER_RIGUNIT()
 	if (Context.State == EControlRigState::Init)
diff --git a/Engine/Plugins/Experimental/ControlRig/Source/ControlRig/Private/Units/Debug/RigUnit_DebugHierarchy.cpp b/Engine/Plugins/Experimental/ControlRig/Source/ControlRig/Private/Units/Debug/RigUnit_DebugHierarchy.cpp
--- a/Engine/Plugins/Experimental/ControlRig/Source/ControlRig/Private/Units/Debug/RigUnit_DebugHierarchy.cpp
+++ b/Engine/Plugins/Experimental/ControlRig/Source/ControlRig/Private/Units/Debug/RigUnit_DebugHierarchy.cpp
@@ -3,7 +3,7 @@
 #include "Units/Debug/RigUnit_DebugHierarchy.h"
 #include "Units/RigUnitContext.h"
 
-UE_RigUnit_DebugHierarchy_IMPLEMENT_RIGVM(void, Execute, const FRigUnitContext& Context)
+FRigUnit_DebugHierarchy_Execute()
 {
     DECLARE_SCOPE_HIERARCHICAL_COUNTER_RIGUNIT()
 	if (Context.State == EControlRigState::Init)
diff --git a/Engine/Plugins/Experimental/ControlRig/Source/ControlRig/Private/Units/Debug/RigUnit_DebugLine.cpp b/Engine/Plugins/Experimental/ControlRig/Source/ControlRig/Private/Units/Debug/RigUnit_DebugLine.cpp
--- a/Engine/Plugins/Experimental/ControlRig/Source/ControlRig/Private/Units/Debug/RigUnit_DebugLine.cpp
+++ b/Engine/Plugins/Experimental/ControlRig/Source/ControlRig/Private/Units/Debug/RigUnit_DebugLine.cpp
@@ -3,7 +3,7 @@
 #include "Units/Debug/RigUnit_DebugLine.h"
 #include "Units/RigUnitContext.h"
 
-UE_RigUnit_DebugLine_IMPLEMENT_RIGVM(void, Execute, const FRigUnitContext& Context)
+FRigUnit_DebugLine_Execute()
 {
     DECLARE_SCOPE_HIERARCHICAL_COUNTER_RIGUNIT()
 	if (Context.State == EControlRigState::Init)
diff --git a/Engine/Plugins/Experimental/ControlRig/Source/ControlRig/Private/Units/Debug/RigUnit_DebugLineStrip.cpp b/Engine/Plugins/Experimental/ControlRig/Source/ControlRig/Private/Units/Debug/RigUnit_DebugLineStrip.cpp
--- a/Engine/Plugins/Experimental/ControlRig/Source/ControlRig/Private/Units/Debug/RigUnit_DebugLineStrip.cpp
+++ b/Engine/Plugins/Experimental/ControlRig/Source/ControlRig/Private/Units/Debug/RigUnit_DebugLineStrip.cpp
@@ -3,7 +3,7 @@
 #include "Units/Debug/RigUnit_DebugLineStrip.h"
 #include "Units/RigUnitContext.h"
 
-UE_RigUnit_DebugLineStrip_IMPLEMENT_RIGVM(void, Execute, const FRigUnitContext& Context)
+FRigUnit_DebugLineStrip_Execute()
 {
     DECLARE_SCOPE_HIERARCHICAL_COUNTER_RIGUNIT()
 	if (Context.State == EControlRigState::Init)
diff --git a/Engine/Plugins/Experimental/ControlRig/Source/ControlRig/Private/Units/Debug/RigUnit_DebugPoint.cpp b/Engine/Plugins/Experimental/ControlRig/Source/ControlRig/Private/Units/Debug/RigUnit_DebugPoint.cpp
--- a/Engine/Plugins/Experimental/ControlRig/Source/ControlRig/Private/Units/Debug/RigUnit_DebugPoint.cpp
+++ b/Engine/Plugins/Experimental/ControlRig/Source/ControlRig/Private/Units/Debug/RigUnit_DebugPoint.cpp
@@ -3,7 +3,7 @@
 #include "Units/Debug/RigUnit_DebugPoint.h"
 #include "Units/RigUnitContext.h"
 
-UE_RigUnit_DebugPoint_IMPLEMENT_RIGVM(void, Execute, const FRigUnitContext& Context)
+FRigUnit_DebugPoint_Execute()
 {
     DECLARE_SCOPE_HIERARCHICAL_COUNTER_RIGUNIT()
 	if (Context.State == EControlRigState::Init)
@@ -40,7 +40,7 @@ UE_RigUnit_DebugPoint_IMPLEMENT_RIGVM(void, Execute, const FRigUnitContext& Cont
 	}
 }
 
-UE_RigUnit_DebugPointMutable_IMPLEMENT_RIGVM(void, Execute, const FRigUnitContext& Context)
+FRigUnit_DebugPointMutable_Execute()
 {
     DECLARE_SCOPE_HIERARCHICAL_COUNTER_RIGUNIT()
 	if (Context.State == EControlRigState::Init)
diff --git a/Engine/Plugins/Experimental/ControlRig/Source/ControlRig/Private/Units/Debug/RigUnit_DebugPrimitives.cpp b/Engine/Plugins/Experimental/ControlRig/Source/ControlRig/Private/Units/Debug/RigUnit_DebugPrimitives.cpp
--- a/Engine/Plugins/Experimental/ControlRig/Source/ControlRig/Private/Units/Debug/RigUnit_DebugPrimitives.cpp
+++ b/Engine/Plugins/Experimental/ControlRig/Source/ControlRig/Private/Units/Debug/RigUnit_DebugPrimitives.cpp
@@ -3,7 +3,7 @@
 #include "Units/Debug/RigUnit_DebugPrimitives.h"
 #include "Units/RigUnitContext.h"
 
-UE_RigUnit_DebugRectangle_IMPLEMENT_RIGVM(void, Execute, const FRigUnitContext& Context)
+FRigUnit_DebugRectangle_Execute()
 {
     DECLARE_SCOPE_HIERARCHICAL_COUNTER_RIGUNIT()
 	if (Context.State == EControlRigState::Init)
@@ -25,7 +25,7 @@ UE_RigUnit_DebugRectangle_IMPLEMENT_RIGVM(void, Execute, const FRigUnitContext&
 	Context.DrawInterface->DrawRectangle(WorldOffset, DrawTransform, Scale, Color, Thickness);
 }
 
-UE_RigUnit_DebugArc_IMPLEMENT_RIGVM(void, Execute, const FRigUnitContext& Context)
+FRigUnit_DebugArc_Execute()
 {
     DECLARE_SCOPE_HIERARCHICAL_COUNTER_RIGUNIT()
 	if (Context.State == EControlRigState::Init)
diff --git a/Engine/Plugins/Experimental/ControlRig/Source/ControlRig/Private/Units/Debug/RigUnit_DebugTransform.cpp b/Engine/Plugins/Experimental/ControlRig/Source/ControlRig/Private/Units/Debug/RigUnit_DebugTransform.cpp
--- a/Engine/Plugins/Experimental/ControlRig/Source/ControlRig/Private/Units/Debug/RigUnit_DebugTransform.cpp
+++ b/Engine/Plugins/Experimental/ControlRig/Source/ControlRig/Private/Units/Debug/RigUnit_DebugTransform.cpp
@@ -3,7 +3,7 @@
 #include "Units/Debug/RigUnit_DebugTransform.h"
 #include "Units/RigUnitContext.h"
 
-UE_RigUnit_DebugTransform_IMPLEMENT_RIGVM(void, Execute, const FRigUnitContext& Context)
+FRigUnit_DebugTransform_Execute()
 {
     DECLARE_SCOPE_HIERARCHICAL_COUNTER_RIGUNIT()
 	if (Context.State == EControlRigState::Init)
@@ -43,7 +43,7 @@ UE_RigUnit_DebugTransform_IMPLEMENT_RIGVM(void, Execute, const FRigUnitContext&
 	}
 }
 
-UE_RigUnit_DebugTransformMutable_IMPLEMENT_RIGVM(void, Execute, const FRigUnitContext& Context)
+FRigUnit_DebugTransformMutable_Execute()
 {
     DECLARE_SCOPE_HIERARCHICAL_COUNTER_RIGUNIT()
 	if (Context.State == EControlRigState::Init)
@@ -83,7 +83,7 @@ UE_RigUnit_DebugTransformMutable_IMPLEMENT_RIGVM(void, Execute, const FRigUnitCo
 	}
 }
 
-UE_RigUnit_DebugTransformArrayMutable_IMPLEMENT_RIGVM(void, Execute, const FRigUnitContext& Context)
+FRigUnit_DebugTransformArrayMutable_Execute()
 {
     DECLARE_SCOPE_HIERARCHICAL_COUNTER_RIGUNIT()
 	if (Context.State == EControlRigState::Init)
diff --git a/Engine/Plugins/Experimental/ControlRig/Source/ControlRig/Private/Units/Debug/RigUnit_ProfilingBracket.cpp b/Engine/Plugins/Experimental/ControlRig/Source/ControlRig/Private/Units/Debug/RigUnit_ProfilingBracket.cpp
--- a/Engine/Plugins/Experimental/ControlRig/Source/ControlRig/Private/Units/Debug/RigUnit_ProfilingBracket.cpp
+++ b/Engine/Plugins/Experimental/ControlRig/Source/ControlRig/Private/Units/Debug/RigUnit_ProfilingBracket.cpp
@@ -4,7 +4,7 @@
 #include "Units/RigUnitContext.h"
 #include "KismetAnimationLibrary.h"
 
-UE_RigUnit_StartProfilingTimer_IMPLEMENT_RIGVM(void, Execute, const FRigUnitContext& Context)
+FRigUnit_StartProfilingTimer_Execute()
 {
     DECLARE_SCOPE_HIERARCHICAL_COUNTER_RIGUNIT()
 	if (Context.State == EControlRigState::Update)
@@ -13,7 +13,7 @@ UE_RigUnit_StartProfilingTimer_IMPLEMENT_RIGVM(void, Execute, const FRigUnitCont
 	}
 }
 
-UE_RigUnit_EndProfilingTimer_IMPLEMENT_RIGVM(void, Execute, const FRigUnitContext& Context)
+FRigUnit_EndProfilingTimer_Execute()
 {
     DECLARE_SCOPE_HIERARCHICAL_COUNTER_RIGUNIT()
 	if (Context.State == EControlRigState::Init)
diff --git a/Engine/Plugins/Experimental/ControlRig/Source/ControlRig/Private/Units/Execution/RigUnit_BeginExecution.cpp b/Engine/Plugins/Experimental/ControlRig/Source/ControlRig/Private/Units/Execution/RigUnit_BeginExecution.cpp
--- a/Engine/Plugins/Experimental/ControlRig/Source/ControlRig/Private/Units/Execution/RigUnit_BeginExecution.cpp
+++ b/Engine/Plugins/Experimental/ControlRig/Source/ControlRig/Private/Units/Execution/RigUnit_BeginExecution.cpp
@@ -3,7 +3,7 @@
 #include "RigUnit_BeginExecution.h"
 #include "Units/RigUnitContext.h"
 
-UE_RigUnit_BeginExecution_IMPLEMENT_RIGVM(void, Execute, const FRigUnitContext& Context)
+FRigUnit_BeginExecution_Execute()
 {
     DECLARE_SCOPE_HIERARCHICAL_COUNTER_RIGUNIT()
 	ExecuteContext.Hierarchy = (FRigHierarchyContainer*)Context.Hierarchy;
diff --git a/Engine/Plugins/Experimental/ControlRig/Source/ControlRig/Private/Units/Hierarchy/RigUnit_AddBoneTransform.cpp b/Engine/Plugins/Experimental/ControlRig/Source/ControlRig/Private/Units/Hierarchy/RigUnit_AddBoneTransform.cpp
--- a/Engine/Plugins/Experimental/ControlRig/Source/ControlRig/Private/Units/Hierarchy/RigUnit_AddBoneTransform.cpp
+++ b/Engine/Plugins/Experimental/ControlRig/Source/ControlRig/Private/Units/Hierarchy/RigUnit_AddBoneTransform.cpp
@@ -8,7 +8,7 @@ FString FRigUnit_AddBoneTransform::GetUnitLabel() const
 	return FString::Printf(TEXT("Offset Transform %s"), *Bone.ToString());
 }
 
-UE_RigUnit_AddBoneTransform_IMPLEMENT_RIGVM(void, Execute, const FRigUnitContext& Context)
+FRigUnit_AddBoneTransform_Execute()
 {
     DECLARE_SCOPE_HIERARCHICAL_COUNTER_RIGUNIT()
 	FRigBoneHierarchy* Hierarchy = ExecuteContext.GetBones();
diff --git a/Engine/Plugins/Experimental/ControlRig/Source/ControlRig/Private/Units/Hierarchy/RigUnit_GetBoneTransform.cpp b/Engine/Plugins/Experimental/ControlRig/Source/ControlRig/Private/Units/Hierarchy/RigUnit_GetBoneTransform.cpp
--- a/Engine/Plugins/Experimental/ControlRig/Source/ControlRig/Private/Units/Hierarchy/RigUnit_GetBoneTransform.cpp
+++ b/Engine/Plugins/Experimental/ControlRig/Source/ControlRig/Private/Units/Hierarchy/RigUnit_GetBoneTransform.cpp
@@ -8,7 +8,7 @@ FString FRigUnit_GetBoneTransform::GetUnitLabel() const
 	return FString::Printf(TEXT("Get Transform %s"), *Bone.ToString());
 }
 
-UE_RigUnit_GetBoneTransform_IMPLEMENT_RIGVM(void, Execute, const FRigUnitContext& Context)
+FRigUnit_GetBoneTransform_Execute()
 {
     DECLARE_SCOPE_HIERARCHICAL_COUNTER_RIGUNIT()
 	const FRigBoneHierarchy* Hierarchy = Context.GetBones();
diff --git a/Engine/Plugins/Experimental/ControlRig/Source/ControlRig/Private/Units/Hierarchy/RigUnit_GetCurveValue.cpp b/Engine/Plugins/Experimental/ControlRig/Source/ControlRig/Private/Units/Hierarchy/RigUnit_GetCurveValue.cpp
--- a/Engine/Plugins/Experimental/ControlRig/Source/ControlRig/Private/Units/Hierarchy/RigUnit_GetCurveValue.cpp
+++ b/Engine/Plugins/Experimental/ControlRig/Source/ControlRig/Private/Units/Hierarchy/RigUnit_GetCurveValue.cpp
@@ -8,7 +8,7 @@ FString FRigUnit_GetCurveValue::GetUnitLabel() const
 	return FString::Printf(TEXT("Get Curve %s"), *Curve.ToString());
 }
 
-UE_RigUnit_GetCurveValue_IMPLEMENT_RIGVM(void, Execute, const FRigUnitContext& Context)
+FRigUnit_GetCurveValue_Execute()
 {
     DECLARE_SCOPE_HIERARCHICAL_COUNTER_RIGUNIT()
 	const FRigCurveContainer* CurveContainer = Context.GetCurves();
diff --git a/Engine/Plugins/Experimental/ControlRig/Source/ControlRig/Private/Units/Hierarchy/RigUnit_GetInitialBoneTransform.cpp b/Engine/Plugins/Experimental/ControlRig/Source/ControlRig/Private/Units/Hierarchy/RigUnit_GetInitialBoneTransform.cpp
--- a/Engine/Plugins/Experimental/ControlRig/Source/ControlRig/Private/Units/Hierarchy/RigUnit_GetInitialBoneTransform.cpp
+++ b/Engine/Plugins/Experimental/ControlRig/Source/ControlRig/Private/Units/Hierarchy/RigUnit_GetInitialBoneTransform.cpp
@@ -8,7 +8,7 @@ FString FRigUnit_GetInitialBoneTransform::GetUnitLabel() const
 	return FString::Printf(TEXT("Get Initial Transform %s"), *Bone.ToString());
 }
 
-UE_RigUnit_GetInitialBoneTransform_IMPLEMENT_RIGVM(void, Execute, const FRigUnitContext& Context)
+FRigUnit_GetInitialBoneTransform_Execute()
 {
     DECLARE_SCOPE_HIERARCHICAL_COUNTER_RIGUNIT()
 	const FRigBoneHierarchy* Hierarchy = Context.GetBones();
diff --git a/Engine/Plugins/Experimental/ControlRig/Source/ControlRig/Private/Units/Hierarchy/RigUnit_GetRelativeBoneTransform.cpp b/Engine/Plugins/Experimental/ControlRig/Source/ControlRig/Private/Units/Hierarchy/RigUnit_GetRelativeBoneTransform.cpp
--- a/Engine/Plugins/Experimental/ControlRig/Source/ControlRig/Private/Units/Hierarchy/RigUnit_GetRelativeBoneTransform.cpp
+++ b/Engine/Plugins/Experimental/ControlRig/Source/ControlRig/Private/Units/Hierarchy/RigUnit_GetRelativeBoneTransform.cpp
@@ -8,7 +8,7 @@ FString FRigUnit_GetRelativeBoneTransform::GetUnitLabel() const
 	return FString::Printf(TEXT("Get Relative Transform %s"), *Bone.ToString());
 }
 
-UE_RigUnit_GetRelativeBoneTransform_IMPLEMENT_RIGVM(void, Execute, const FRigUnitContext& Context)
+FRigUnit_GetRelativeBoneTransform_Execute()
 {
     DECLARE_SCOPE_HIERARCHICAL_COUNTER_RIGUNIT()
 	const FRigBoneHierarchy* Hierarchy = Context.GetBones();
diff --git a/Engine/Plugins/Experimental/ControlRig/Source/ControlRig/Private/Units/Hierarchy/RigUnit_SetBoneRotation.cpp b/Engine/Plugins/Experimental/ControlRig/Source/ControlRig/Private/Units/Hierarchy/RigUnit_SetBoneRotation.cpp
--- a/Engine/Plugins/Experimental/ControlRig/Source/ControlRig/Private/Units/Hierarchy/RigUnit_SetBoneRotation.cpp
+++ b/Engine/Plugins/Experimental/ControlRig/Source/ControlRig/Private/Units/Hierarchy/RigUnit_SetBoneRotation.cpp
@@ -8,7 +8,7 @@ FString FRigUnit_SetBoneRotation::GetUnitLabel() const
 	return FString::Printf(TEXT("Set Rotation %s"), *Bone.ToString());
 }
 
-UE_RigUnit_SetBoneRotation_IMPLEMENT_RIGVM(void, Execute, const FRigUnitContext& Context)
+FRigUnit_SetBoneRotation_Execute()
 {
     DECLARE_SCOPE_HIERARCHICAL_COUNTER_RIGUNIT()
 	FRigBoneHierarchy* Hierarchy = ExecuteContext.GetBones();
diff --git a/Engine/Plugins/Experimental/ControlRig/Source/ControlRig/Private/Units/Hierarchy/RigUnit_SetBoneTransform.cpp b/Engine/Plugins/Experimental/ControlRig/Source/ControlRig/Private/Units/Hierarchy/RigUnit_SetBoneTransform.cpp
--- a/Engine/Plugins/Experimental/ControlRig/Source/ControlRig/Private/Units/Hierarchy/RigUnit_SetBoneTransform.cpp
+++ b/Engine/Plugins/Experimental/ControlRig/Source/ControlRig/Private/Units/Hierarchy/RigUnit_SetBoneTransform.cpp
@@ -8,7 +8,7 @@ FString FRigUnit_SetBoneTransform::GetUnitLabel() const
 	return FString::Printf(TEXT("Set Transform %s"), *Bone.ToString());
 }
 
-UE_RigUnit_SetBoneTransform_IMPLEMENT_RIGVM(void, Execute, const FRigUnitContext& Context)
+FRigUnit_SetBoneTransform_Execute()
 {
     DECLARE_SCOPE_HIERARCHICAL_COUNTER_RIGUNIT()
 	FRigBoneHierarchy* Hierarchy = ExecuteContext.GetBones();
diff --git a/Engine/Plugins/Experimental/ControlRig/Source/ControlRig/Private/Units/Hierarchy/RigUnit_SetBoneTranslation.cpp b/Engine/Plugins/Experimental/ControlRig/Source/ControlRig/Private/Units/Hierarchy/RigUnit_SetBoneTranslation.cpp
--- a/Engine/Plugins/Experimental/ControlRig/Source/ControlRig/Private/Units/Hierarchy/RigUnit_SetBoneTranslation.cpp
+++ b/Engine/Plugins/Experimental/ControlRig/Source/ControlRig/Private/Units/Hierarchy/RigUnit_SetBoneTranslation.cpp
@@ -8,7 +8,7 @@ FString FRigUnit_SetBoneTranslation::GetUnitLabel() const
 	return FString::Printf(TEXT("Set Translation %s"), *Bone.ToString());
 }
 
-UE_RigUnit_SetBoneTranslation_IMPLEMENT_RIGVM(void, Execute, const FRigUnitContext& Context)
+FRigUnit_SetBoneTranslation_Execute()
 {
     DECLARE_SCOPE_HIERARCHICAL_COUNTER_RIGUNIT()
 	FRigBoneHierarchy* Hierarchy = ExecuteContext.GetBones();
diff --git a/Engine/Plugins/Experimental/ControlRig/Source/ControlRig/Private/Units/Hierarchy/RigUnit_SetCurveValue.cpp b/Engine/Plugins/Experimental/ControlRig/Source/ControlRig/Private/Units/Hierarchy/RigUnit_SetCurveValue.cpp
--- a/Engine/Plugins/Experimental/ControlRig/Source/ControlRig/Private/Units/Hierarchy/RigUnit_SetCurveValue.cpp
+++ b/Engine/Plugins/Experimental/ControlRig/Source/ControlRig/Private/Units/Hierarchy/RigUnit_SetCurveValue.cpp
@@ -8,7 +8,7 @@ FString FRigUnit_SetCurveValue::GetUnitLabel() const
 	return FString::Printf(TEXT("Set Curve [%s]"), *Curve.ToString());
 }
 
-UE_RigUnit_SetCurveValue_IMPLEMENT_RIGVM(void, Execute, const FRigUnitContext& Context)
+FRigUnit_SetCurveValue_Execute()
 {
     DECLARE_SCOPE_HIERARCHICAL_COUNTER_RIGUNIT()
 	FRigCurveContainer* CurveContainer = ExecuteContext.GetCurves();
diff --git a/Engine/Plugins/Experimental/ControlRig/Source/ControlRig/Private/Units/Hierarchy/RigUnit_SetRelativeBoneTransform.cpp b/Engine/Plugins/Experimental/ControlRig/Source/ControlRig/Private/Units/Hierarchy/RigUnit_SetRelativeBoneTransform.cpp
--- a/Engine/Plugins/Experimental/ControlRig/Source/ControlRig/Private/Units/Hierarchy/RigUnit_SetRelativeBoneTransform.cpp
+++ b/Engine/Plugins/Experimental/ControlRig/Source/ControlRig/Private/Units/Hierarchy/RigUnit_SetRelativeBoneTransform.cpp
@@ -8,7 +8,7 @@ FString FRigUnit_SetRelativeBoneTransform::GetUnitLabel() const
 	return FString::Printf(TEXT("Set Relative Transform %s"), *Bone.ToString());
 }
 
-UE_RigUnit_SetRelativeBoneTransform_IMPLEMENT_RIGVM(void, Execute, const FRigUnitContext& Context)
+FRigUnit_SetRelativeBoneTransform_Execute()
 {
     DECLARE_SCOPE_HIERARCHICAL_COUNTER_RIGUNIT()
 	FRigBoneHierarchy* Hierarchy = ExecuteContext.GetBones();
diff --git a/Engine/Plugins/Experimental/ControlRig/Source/ControlRig/Private/Units/Highlevel/Harmonics/RigUnit_BoneHarmonics.cpp b/Engine/Plugins/Experimental/ControlRig/Source/ControlRig/Private/Units/Highlevel/Harmonics/RigUnit_BoneHarmonics.cpp
--- a/Engine/Plugins/Experimental/ControlRig/Source/ControlRig/Private/Units/Highlevel/Harmonics/RigUnit_BoneHarmonics.cpp
+++ b/Engine/Plugins/Experimental/ControlRig/Source/ControlRig/Private/Units/Highlevel/Harmonics/RigUnit_BoneHarmonics.cpp
@@ -3,7 +3,7 @@
 #include "Units/Highlevel/Harmonics/RigUnit_BoneHarmonics.h"
 #include "Units/RigUnitContext.h"
 
-UE_RigUnit_BoneHarmonics_IMPLEMENT_RIGVM(void, Execute, const FRigUnitContext& Context)
+FRigUnit_BoneHarmonics_Execute()
 {
     DECLARE_SCOPE_HIERARCHICAL_COUNTER_RIGUNIT()
 	FRigBoneHierarchy* Hierarchy = ExecuteContext.GetBones();
diff --git a/Engine/Plugins/Experimental/ControlRig/Source/ControlRig/Private/Units/Highlevel/Harmonics/RigUnit_ChainHarmonics.cpp b/Engine/Plugins/Experimental/ControlRig/Source/ControlRig/Private/Units/Highlevel/Harmonics/RigUnit_ChainHarmonics.cpp
--- a/Engine/Plugins/Experimental/ControlRig/Source/ControlRig/Private/Units/Highlevel/Harmonics/RigUnit_ChainHarmonics.cpp
+++ b/Engine/Plugins/Experimental/ControlRig/Source/ControlRig/Private/Units/Highlevel/Harmonics/RigUnit_ChainHarmonics.cpp
@@ -3,7 +3,7 @@
 #include "Units/Highlevel/Harmonics/RigUnit_ChainHarmonics.h"
 #include "Units/RigUnitContext.h"
 
-UE_RigUnit_ChainHarmonics_IMPLEMENT_RIGVM(void, Execute, const FRigUnitContext& Context)
+FRigUnit_ChainHarmonics_Execute()
 {
     DECLARE_SCOPE_HIERARCHICAL_COUNTER_RIGUNIT()
 	FRigBoneHierarchy* Hierarchy = ExecuteContext.GetBones();
diff --git a/Engine/Plugins/Experimental/ControlRig/Source/ControlRig/Private/Units/Highlevel/Hierarchy/RigUnit_AimBone.cpp b/Engine/Plugins/Experimental/ControlRig/Source/ControlRig/Private/Units/Highlevel/Hierarchy/RigUnit_AimBone.cpp
--- a/Engine/Plugins/Experimental/ControlRig/Source/ControlRig/Private/Units/Highlevel/Hierarchy/RigUnit_AimBone.cpp
+++ b/Engine/Plugins/Experimental/ControlRig/Source/ControlRig/Private/Units/Highlevel/Hierarchy/RigUnit_AimBone.cpp
@@ -3,7 +3,7 @@
 #include "Units/Highlevel/Hierarchy/RigUnit_AimBone.h"
 #include "Units/RigUnitContext.h"
 
-UE_RigUnit_AimBone_IMPLEMENT_RIGVM(void, Execute, const FRigUnitContext& Context)
+FRigUnit_AimBone_Execute()
 {
     DECLARE_SCOPE_HIERARCHICAL_COUNTER_RIGUNIT()
 	FRigBoneHierarchy* Hierarchy = ExecuteContext.GetBones();
diff --git a/Engine/Plugins/Experimental/ControlRig/Source/ControlRig/Private/Units/Highlevel/Hierarchy/RigUnit_CCDIK.cpp b/Engine/Plugins/Experimental/ControlRig/Source/ControlRig/Private/Units/Highlevel/Hierarchy/RigUnit_CCDIK.cpp
--- a/Engine/Plugins/Experimental/ControlRig/Source/ControlRig/Private/Units/Highlevel/Hierarchy/RigUnit_CCDIK.cpp
+++ b/Engine/Plugins/Experimental/ControlRig/Source/ControlRig/Private/Units/Highlevel/Hierarchy/RigUnit_CCDIK.cpp
@@ -3,7 +3,7 @@
 #include "RigUnit_CCDIK.h"
 #include "Units/RigUnitContext.h"
 
-UE_RigUnit_CCDIK_IMPLEMENT_RIGVM(void, Execute, const FRigUnitContext& Context)
+FRigUnit_CCDIK_Execute()
 {
     DECLARE_SCOPE_HIERARCHICAL_COUNTER_RIGUNIT()
 	FRigBoneHierarchy* Hierarchy = ExecuteContext.GetBones();
diff --git a/Engine/Plugins/Experimental/ControlRig/Source/ControlRig/Private/Units/Highlevel/Hierarchy/RigUnit_DistributeRotation.cpp b/Engine/Plugins/Experimental/ControlRig/Source/ControlRig/Private/Units/Highlevel/Hierarchy/RigUnit_DistributeRotation.cpp
--- a/Engine/Plugins/Experimental/ControlRig/Source/ControlRig/Private/Units/Highlevel/Hierarchy/RigUnit_DistributeRotation.cpp
+++ b/Engine/Plugins/Experimental/ControlRig/Source/ControlRig/Private/Units/Highlevel/Hierarchy/RigUnit_DistributeRotation.cpp
@@ -3,7 +3,7 @@
 #include "Units/Highlevel/Hierarchy/RigUnit_DistributeRotation.h"
 #include "Units/RigUnitContext.h"
 
-UE_RigUnit_DistributeRotation_IMPLEMENT_RIGVM(void, Execute, const FRigUnitContext& Context)
+FRigUnit_DistributeRotation_Execute()
 {
     DECLARE_SCOPE_HIERARCHICAL_COUNTER_RIGUNIT()
 	FRigBoneHierarchy* Hierarchy = ExecuteContext.GetBones();
diff --git a/Engine/Plugins/Experimental/ControlRig/Source/ControlRig/Private/Units/Highlevel/Hierarchy/RigUnit_FABRIK.cpp b/Engine/Plugins/Experimental/ControlRig/Source/ControlRig/Private/Units/Highlevel/Hierarchy/RigUnit_FABRIK.cpp
--- a/Engine/Plugins/Experimental/ControlRig/Source/ControlRig/Private/Units/Highlevel/Hierarchy/RigUnit_FABRIK.cpp
+++ b/Engine/Plugins/Experimental/ControlRig/Source/ControlRig/Private/Units/Highlevel/Hierarchy/RigUnit_FABRIK.cpp
@@ -3,7 +3,7 @@
 #include "RigUnit_FABRIK.h"
 #include "Units/RigUnitContext.h"
 
-UE_RigUnit_FABRIK_IMPLEMENT_RIGVM(void, Execute, const FRigUnitContext& Context)
+FRigUnit_FABRIK_Execute()
 {
     DECLARE_SCOPE_HIERARCHICAL_COUNTER_RIGUNIT()
 	FRigBoneHierarchy* Hierarchy = ExecuteContext.GetBones();
diff --git a/Engine/Plugins/Experimental/ControlRig/Source/ControlRig/Private/Units/Highlevel/Hierarchy/RigUnit_FitChainToCurve.cpp b/Engine/Plugins/Experimental/ControlRig/Source/ControlRig/Private/Units/Highlevel/Hierarchy/RigUnit_FitChainToCurve.cpp
--- a/Engine/Plugins/Experimental/ControlRig/Source/ControlRig/Private/Units/Highlevel/Hierarchy/RigUnit_FitChainToCurve.cpp
+++ b/Engine/Plugins/Experimental/ControlRig/Source/ControlRig/Private/Units/Highlevel/Hierarchy/RigUnit_FitChainToCurve.cpp
@@ -3,7 +3,7 @@
 #include "Units/Highlevel/Hierarchy/RigUnit_FitChainToCurve.h"
 #include "Units/RigUnitContext.h"
 
-UE_RigUnit_FitChainToCurve_IMPLEMENT_RIGVM(void, Execute, const FRigUnitContext& Context)
+FRigUnit_FitChainToCurve_Execute()
 {
     DECLARE_SCOPE_HIERARCHICAL_COUNTER_RIGUNIT()
 	FRigBoneHierarchy* Hierarchy = ExecuteContext.GetBones();
diff --git a/Engine/Plugins/Experimental/ControlRig/Source/ControlRig/Private/Units/Highlevel/Hierarchy/RigUnit_ModifyBoneTransforms.cpp b/Engine/Plugins/Experimental/ControlRig/Source/ControlRig/Private/Units/Highlevel/Hierarchy/RigUnit_ModifyBoneTransforms.cpp
--- a/Engine/Plugins/Experimental/ControlRig/Source/ControlRig/Private/Units/Highlevel/Hierarchy/RigUnit_ModifyBoneTransforms.cpp
+++ b/Engine/Plugins/Experimental/ControlRig/Source/ControlRig/Private/Units/Highlevel/Hierarchy/RigUnit_ModifyBoneTransforms.cpp
@@ -3,7 +3,7 @@
 #include "RigUnit_ModifyBoneTransforms.h"
 #include "Units/RigUnitContext.h"
 
-UE_RigUnit_ModifyBoneTransforms_IMPLEMENT_RIGVM(void, Execute, const FRigUnitContext& Context)
+FRigUnit_ModifyBoneTransforms_Execute()
 {
     DECLARE_SCOPE_HIERARCHICAL_COUNTER_RIGUNIT()
 
diff --git a/Engine/Plugins/Experimental/ControlRig/Source/ControlRig/Private/Units/Highlevel/Hierarchy/RigUnit_SlideChain.cpp b/Engine/Plugins/Experimental/ControlRig/Source/ControlRig/Private/Units/Highlevel/Hierarchy/RigUnit_SlideChain.cpp
--- a/Engine/Plugins/Experimental/ControlRig/Source/ControlRig/Private/Units/Highlevel/Hierarchy/RigUnit_SlideChain.cpp
+++ b/Engine/Plugins/Experimental/ControlRig/Source/ControlRig/Private/Units/Highlevel/Hierarchy/RigUnit_SlideChain.cpp
@@ -4,7 +4,7 @@
 #include "Math/ControlRigMathLibrary.h"
 #include "Units/RigUnitContext.h"
 
-UE_RigUnit_SlideChain_IMPLEMENT_RIGVM(void, Execute, const FRigUnitContext& Context)
+FRigUnit_SlideChain_Execute()
 {
 
 	FRigBoneHierarchy* Hierarchy = ExecuteContext.GetBones();
diff --git a/Engine/Plugins/Experimental/ControlRig/Source/ControlRig/Private/Units/Highlevel/Hierarchy/RigUnit_SpringIK.cpp b/Engine/Plugins/Experimental/ControlRig/Source/ControlRig/Private/Units/Highlevel/Hierarchy/RigUnit_SpringIK.cpp
--- a/Engine/Plugins/Experimental/ControlRig/Source/ControlRig/Private/Units/Highlevel/Hierarchy/RigUnit_SpringIK.cpp
+++ b/Engine/Plugins/Experimental/ControlRig/Source/ControlRig/Private/Units/Highlevel/Hierarchy/RigUnit_SpringIK.cpp
@@ -3,7 +3,7 @@
 #include "Units/Highlevel/Hierarchy/RigUnit_SpringIK.h"
 #include "Units/RigUnitContext.h"
 
-UE_RigUnit_SpringIK_IMPLEMENT_RIGVM(void, Execute, const FRigUnitContext& Context)
+FRigUnit_SpringIK_Execute()
 {
     DECLARE_SCOPE_HIERARCHICAL_COUNTER_RIGUNIT()
 	FRigBoneHierarchy* Hierarchy = ExecuteContext.GetBones();
diff --git a/Engine/Plugins/Experimental/ControlRig/Source/ControlRig/Private/Units/Highlevel/Hierarchy/RigUnit_TransformConstraint.cpp b/Engine/Plugins/Experimental/ControlRig/Source/ControlRig/Private/Units/Highlevel/Hierarchy/RigUnit_TransformConstraint.cpp
--- a/Engine/Plugins/Experimental/ControlRig/Source/ControlRig/Private/Units/Highlevel/Hierarchy/RigUnit_TransformConstraint.cpp
+++ b/Engine/Plugins/Experimental/ControlRig/Source/ControlRig/Private/Units/Highlevel/Hierarchy/RigUnit_TransformConstraint.cpp
@@ -5,7 +5,7 @@
 #include "HelperUtil.h"
 #include "AnimationCoreLibrary.h"
 
-UE_RigUnit_TransformConstraint_IMPLEMENT_RIGVM(void, Execute, const FRigUnitContext& Context)
+FRigUnit_TransformConstraint_Execute()
 {
     DECLARE_SCOPE_HIERARCHICAL_COUNTER_RIGUNIT()
 
diff --git a/Engine/Plugins/Experimental/ControlRig/Source/ControlRig/Private/Units/Highlevel/Hierarchy/RigUnit_TwoBoneIKSimple.cpp b/Engine/Plugins/Experimental/ControlRig/Source/ControlRig/Private/Units/Highlevel/Hierarchy/RigUnit_TwoBoneIKSimple.cpp
--- a/Engine/Plugins/Experimental/ControlRig/Source/ControlRig/Private/Units/Highlevel/Hierarchy/RigUnit_TwoBoneIKSimple.cpp
+++ b/Engine/Plugins/Experimental/ControlRig/Source/ControlRig/Private/Units/Highlevel/Hierarchy/RigUnit_TwoBoneIKSimple.cpp
@@ -4,7 +4,7 @@
 #include "Units/RigUnitContext.h"
 #include "Math/ControlRigMathLibrary.h"
 
-UE_RigUnit_TwoBoneIKSimple_IMPLEMENT_RIGVM(void, Execute, const FRigUnitContext& Context)
+FRigUnit_TwoBoneIKSimple_Execute()
 {
     DECLARE_SCOPE_HIERARCHICAL_COUNTER_RIGUNIT()
 	FRigBoneHierarchy* Hierarchy = ExecuteContext.GetBones();
diff --git a/Engine/Plugins/Experimental/ControlRig/Source/ControlRig/Private/Units/Math/RigUnit_MathBool.cpp b/Engine/Plugins/Experimental/ControlRig/Source/ControlRig/Private/Units/Math/RigUnit_MathBool.cpp
--- a/Engine/Plugins/Experimental/ControlRig/Source/ControlRig/Private/Units/Math/RigUnit_MathBool.cpp
+++ b/Engine/Plugins/Experimental/ControlRig/Source/ControlRig/Private/Units/Math/RigUnit_MathBool.cpp
@@ -3,37 +3,37 @@
 #include "Units/Math/RigUnit_MathBool.h"
 #include "Units/RigUnitContext.h"
 
-UE_RigUnit_MathBoolNot_IMPLEMENT_RIGVM(void, Execute, const FRigUnitContext& Context)
+FRigUnit_MathBoolNot_Execute()
 {
     DECLARE_SCOPE_HIERARCHICAL_COUNTER_RIGUNIT()
 	Result = !Value;
 }
 
-UE_RigUnit_MathBoolAnd_IMPLEMENT_RIGVM(void, Execute, const FRigUnitContext& Context)
+FRigUnit_MathBoolAnd_Execute()
 {
     DECLARE_SCOPE_HIERARCHICAL_COUNTER_RIGUNIT()
 	Result = A && B;
 }
 
-UE_RigUnit_MathBoolNand_IMPLEMENT_RIGVM(void, Execute, const FRigUnitContext& Context)
+FRigUnit_MathBoolNand_Execute()
 {
     DECLARE_SCOPE_HIERARCHICAL_COUNTER_RIGUNIT()
 	Result = (!A) && (!B);
 }
 
-UE_RigUnit_MathBoolOr_IMPLEMENT_RIGVM(void, Execute, const FRigUnitContext& Context)
+FRigUnit_MathBoolOr_Execute()
 {
     DECLARE_SCOPE_HIERARCHICAL_COUNTER_RIGUNIT()
 	Result = A || B;
 }
 
-UE_RigUnit_MathBoolEquals_IMPLEMENT_RIGVM(void, Execute, const FRigUnitContext& Context)
+FRigUnit_MathBoolEquals_Execute()
 {
     DECLARE_SCOPE_HIERARCHICAL_COUNTER_RIGUNIT()
 	Result = A == B;
 }
 
-UE_RigUnit_MathBoolNotEquals_IMPLEMENT_RIGVM(void, Execute, const FRigUnitContext& Context)
+FRigUnit_MathBoolNotEquals_Execute()
 {
     DECLARE_SCOPE_HIERARCHICAL_COUNTER_RIGUNIT()
 	Result = A != B;
diff --git a/Engine/Plugins/Experimental/ControlRig/Source/ControlRig/Private/Units/Math/RigUnit_MathFloat.cpp b/Engine/Plugins/Experimental/ControlRig/Source/ControlRig/Private/Units/Math/RigUnit_MathFloat.cpp
--- a/Engine/Plugins/Experimental/ControlRig/Source/ControlRig/Private/Units/Math/RigUnit_MathFloat.cpp
+++ b/Engine/Plugins/Experimental/ControlRig/Source/ControlRig/Private/Units/Math/RigUnit_MathFloat.cpp
@@ -3,25 +3,25 @@
 #include "Units/Math/RigUnit_MathFloat.h"
 #include "Units/RigUnitContext.h"
 
-UE_RigUnit_MathFloatAdd_IMPLEMENT_RIGVM(void, Execute, const FRigUnitContext& Context)
+FRigUnit_MathFloatAdd_Execute()
 {
     DECLARE_SCOPE_HIERARCHICAL_COUNTER_RIGUNIT()
 	Result = A + B;
 }
 
-UE_RigUnit_MathFloatSub_IMPLEMENT_RIGVM(void, Execute, const FRigUnitContext& Context)
+FRigUnit_MathFloatSub_Execute()
 {
     DECLARE_SCOPE_HIERARCHICAL_COUNTER_RIGUNIT()
 	Result = A - B;
 }
 
-UE_RigUnit_MathFloatMul_IMPLEMENT_RIGVM(void, Execute, const FRigUnitContext& Context)
+FRigUnit_MathFloatMul_Execute()
 {
     DECLARE_SCOPE_HIERARCHICAL_COUNTER_RIGUNIT()
 	Result = A * B;
 }
 
-UE_RigUnit_MathFloatDiv_IMPLEMENT_RIGVM(void, Execute, const FRigUnitContext& Context)
+FRigUnit_MathFloatDiv_Execute()
 {
     DECLARE_SCOPE_HIERARCHICAL_COUNTER_RIGUNIT()
 	if(FMath::IsNearlyZero(B))
@@ -33,7 +33,7 @@ UE_RigUnit_MathFloatDiv_IMPLEMENT_RIGVM(void, Execute, const FRigUnitContext& Co
 	Result = A / B;
 }
 
-UE_RigUnit_MathFloatMod_IMPLEMENT_RIGVM(void, Execute, const FRigUnitContext& Context)
+FRigUnit_MathFloatMod_Execute()
 {
     DECLARE_SCOPE_HIERARCHICAL_COUNTER_RIGUNIT()
 	if(FMath::IsNearlyZero(B) || B < 0.f)
@@ -45,25 +45,25 @@ UE_RigUnit_MathFloatMod_IMPLEMENT_RIGVM(void, Execute, const FRigUnitContext& Co
 	Result = FMath::Fmod(A, B);
 }
 
-UE_RigUnit_MathFloatMin_IMPLEMENT_RIGVM(void, Execute, const FRigUnitContext& Context)
+FRigUnit_MathFloatMin_Execute()
 {
     DECLARE_SCOPE_HIERARCHICAL_COUNTER_RIGUNIT()
 	Result = FMath::Min<float>(A, B);
 }
 
-UE_RigUnit_MathFloatMax_IMPLEMENT_RIGVM(void, Execute, const FRigUnitContext& Context)
+FRigUnit_MathFloatMax_Execute()
 {
     DECLARE_SCOPE_HIERARCHICAL_COUNTER_RIGUNIT()
 	Result = FMath::Max<float>(A, B);
 }
 
-UE_RigUnit_MathFloatPow_IMPLEMENT_RIGVM(void, Execute, const FRigUnitContext& Context)
+FRigUnit_MathFloatPow_Execute()
 {
     DECLARE_SCOPE_HIERARCHICAL_COUNTER_RIGUNIT()
 	Result = FMath::Pow(A, B);
 }
 
-UE_RigUnit_MathFloatSqrt_IMPLEMENT_RIGVM(void, Execute, const FRigUnitContext& Context)
+FRigUnit_MathFloatSqrt_Execute()
 {
     DECLARE_SCOPE_HIERARCHICAL_COUNTER_RIGUNIT()
 	if(Value < 0.f)
@@ -75,55 +75,55 @@ UE_RigUnit_MathFloatSqrt_IMPLEMENT_RIGVM(void, Execute, const FRigUnitContext& C
 	Result = FMath::Sqrt(Value);
 }
 
-UE_RigUnit_MathFloatNegate_IMPLEMENT_RIGVM(void, Execute, const FRigUnitContext& Context)
+FRigUnit_MathFloatNegate_Execute()
 {
     DECLARE_SCOPE_HIERARCHICAL_COUNTER_RIGUNIT()
 	Result = -Value;
 }
 
-UE_RigUnit_MathFloatAbs_IMPLEMENT_RIGVM(void, Execute, const FRigUnitContext& Context)
+FRigUnit_MathFloatAbs_Execute()
 {
     DECLARE_SCOPE_HIERARCHICAL_COUNTER_RIGUNIT()
 	Result = FMath::Abs(Value);
 }
 
-UE_RigUnit_MathFloatFloor_IMPLEMENT_RIGVM(void, Execute, const FRigUnitContext& Context)
+FRigUnit_MathFloatFloor_Execute()
 {
     DECLARE_SCOPE_HIERARCHICAL_COUNTER_RIGUNIT()
 	Result = FMath::FloorToFloat(Value);
 }
 
-UE_RigUnit_MathFloatCeil_IMPLEMENT_RIGVM(void, Execute, const FRigUnitContext& Context)
+FRigUnit_MathFloatCeil_Execute()
 {
     DECLARE_SCOPE_HIERARCHICAL_COUNTER_RIGUNIT()
 	Result = FMath::CeilToFloat(Value);
 }
 
-UE_RigUnit_MathFloatRound_IMPLEMENT_RIGVM(void, Execute, const FRigUnitContext& Context)
+FRigUnit_MathFloatRound_Execute()
 {
     DECLARE_SCOPE_HIERARCHICAL_COUNTER_RIGUNIT()
 	Result = FMath::RoundToFloat(Value);
 }
 
-UE_RigUnit_MathFloatSign_IMPLEMENT_RIGVM(void, Execute, const FRigUnitContext& Context)
+FRigUnit_MathFloatSign_Execute()
 {
     DECLARE_SCOPE_HIERARCHICAL_COUNTER_RIGUNIT()
 	Result = Value < 0.f ? -1.f : 1.f;
 }
 
-UE_RigUnit_MathFloatClamp_IMPLEMENT_RIGVM(void, Execute, const FRigUnitContext& Context)
+FRigUnit_MathFloatClamp_Execute()
 {
     DECLARE_SCOPE_HIERARCHICAL_COUNTER_RIGUNIT()
 	Result = FMath::Clamp<float>(Value, Minimum, Maximum);
 }
 
-UE_RigUnit_MathFloatLerp_IMPLEMENT_RIGVM(void, Execute, const FRigUnitContext& Context)
+FRigUnit_MathFloatLerp_Execute()
 {
     DECLARE_SCOPE_HIERARCHICAL_COUNTER_RIGUNIT()
 	Result = FMath::Lerp<float>(A, B, T);
 }
 
-UE_RigUnit_MathFloatRemap_IMPLEMENT_RIGVM(void, Execute, const FRigUnitContext& Context)
+FRigUnit_MathFloatRemap_Execute()
 {
     DECLARE_SCOPE_HIERARCHICAL_COUNTER_RIGUNIT()
 	float Ratio = 0.f;
@@ -142,43 +142,43 @@ UE_RigUnit_MathFloatRemap_IMPLEMENT_RIGVM(void, Execute, const FRigUnitContext&
 	Result = FMath::Lerp<float>(TargetMinimum, TargetMaximum, Ratio);
 }
 
-UE_RigUnit_MathFloatEquals_IMPLEMENT_RIGVM(void, Execute, const FRigUnitContext& Context)
+FRigUnit_MathFloatEquals_Execute()
 {
     DECLARE_SCOPE_HIERARCHICAL_COUNTER_RIGUNIT()
 	Result = A == B;
 }
 
-UE_RigUnit_MathFloatNotEquals_IMPLEMENT_RIGVM(void, Execute, const FRigUnitContext& Context)
+FRigUnit_MathFloatNotEquals_Execute()
 {
     DECLARE_SCOPE_HIERARCHICAL_COUNTER_RIGUNIT()
 	Result = A != B;
 }
 
-UE_RigUnit_MathFloatGreater_IMPLEMENT_RIGVM(void, Execute, const FRigUnitContext& Context)
+FRigUnit_MathFloatGreater_Execute()
 {
     DECLARE_SCOPE_HIERARCHICAL_COUNTER_RIGUNIT()
 	Result = A > B;
 }
 
-UE_RigUnit_MathFloatLess_IMPLEMENT_RIGVM(void, Execute, const FRigUnitContext& Context)
+FRigUnit_MathFloatLess_Execute()
 {
     DECLARE_SCOPE_HIERARCHICAL_COUNTER_RIGUNIT()
 	Result = A < B;
 }
 
-UE_RigUnit_MathFloatGreaterEqual_IMPLEMENT_RIGVM(void, Execute, const FRigUnitContext& Context)
+FRigUnit_MathFloatGreaterEqual_Execute()
 {
     DECLARE_SCOPE_HIERARCHICAL_COUNTER_RIGUNIT()
 	Result = A >= B;
 }
 
-UE_RigUnit_MathFloatLessEqual_IMPLEMENT_RIGVM(void, Execute, const FRigUnitContext& Context)
+FRigUnit_MathFloatLessEqual_Execute()
 {
     DECLARE_SCOPE_HIERARCHICAL_COUNTER_RIGUNIT()
 	Result = A <= B;
 }
 
-UE_RigUnit_MathFloatIsNearlyZero_IMPLEMENT_RIGVM(void, Execute, const FRigUnitContext& Context)
+FRigUnit_MathFloatIsNearlyZero_Execute()
 {
     DECLARE_SCOPE_HIERARCHICAL_COUNTER_RIGUNIT()
 	if(Tolerance < 0.f)
@@ -188,7 +188,7 @@ UE_RigUnit_MathFloatIsNearlyZero_IMPLEMENT_RIGVM(void, Execute, const FRigUnitCo
 	Result = FMath::IsNearlyZero(Value, FMath::Max<float>(Tolerance, SMALL_NUMBER));
 }
 
-UE_RigUnit_MathFloatIsNearlyEqual_IMPLEMENT_RIGVM(void, Execute, const FRigUnitContext& Context)
+FRigUnit_MathFloatIsNearlyEqual_Execute()
 {
     DECLARE_SCOPE_HIERARCHICAL_COUNTER_RIGUNIT()
 	if(Tolerance < 0.f)
@@ -198,43 +198,43 @@ UE_RigUnit_MathFloatIsNearlyEqual_IMPLEMENT_RIGVM(void, Execute, const FRigUnitC
 	Result = FMath::IsNearlyEqual(A, B, FMath::Max<float>(Tolerance, SMALL_NUMBER));
 }
 
-UE_RigUnit_MathFloatSelectBool_IMPLEMENT_RIGVM(void, Execute, const FRigUnitContext& Context)
+FRigUnit_MathFloatSelectBool_Execute()
 {
     DECLARE_SCOPE_HIERARCHICAL_COUNTER_RIGUNIT()
 	Result = Condition ? IfTrue : IfFalse;
 }
 
-UE_RigUnit_MathFloatDeg_IMPLEMENT_RIGVM(void, Execute, const FRigUnitContext& Context)
+FRigUnit_MathFloatDeg_Execute()
 {
     DECLARE_SCOPE_HIERARCHICAL_COUNTER_RIGUNIT()
 	Result = FMath::RadiansToDegrees(Value);
 }
 
-UE_RigUnit_MathFloatRad_IMPLEMENT_RIGVM(void, Execute, const FRigUnitContext& Context)
+FRigUnit_MathFloatRad_Execute()
 {
     DECLARE_SCOPE_HIERARCHICAL_COUNTER_RIGUNIT()
 	Result = FMath::DegreesToRadians(Value);
 }
 
-UE_RigUnit_MathFloatSin_IMPLEMENT_RIGVM(void, Execute, const FRigUnitContext& Context)
+FRigUnit_MathFloatSin_Execute()
 {
     DECLARE_SCOPE_HIERARCHICAL_COUNTER_RIGUNIT()
 	Result = FMath::Sin(Value);
 }
 
-UE_RigUnit_MathFloatCos_IMPLEMENT_RIGVM(void, Execute, const FRigUnitContext& Context)
+FRigUnit_MathFloatCos_Execute()
 {
     DECLARE_SCOPE_HIERARCHICAL_COUNTER_RIGUNIT()
 	Result = FMath::Cos(Value);
 }
 
-UE_RigUnit_MathFloatTan_IMPLEMENT_RIGVM(void, Execute, const FRigUnitContext& Context)
+FRigUnit_MathFloatTan_Execute()
 {
     DECLARE_SCOPE_HIERARCHICAL_COUNTER_RIGUNIT()
 	Result = FMath::Tan(Value);
 }
 
-UE_RigUnit_MathFloatAsin_IMPLEMENT_RIGVM(void, Execute, const FRigUnitContext& Context)
+FRigUnit_MathFloatAsin_Execute()
 {
     DECLARE_SCOPE_HIERARCHICAL_COUNTER_RIGUNIT()
 	if (!FMath::IsWithinInclusive<float>(Value, -1.f, 1.f))
@@ -246,7 +246,7 @@ UE_RigUnit_MathFloatAsin_IMPLEMENT_RIGVM(void, Execute, const FRigUnitContext& C
 	Result = FMath::Asin(Value);
 }
 
-UE_RigUnit_MathFloatAcos_IMPLEMENT_RIGVM(void, Execute, const FRigUnitContext& Context)
+FRigUnit_MathFloatAcos_Execute()
 {
     DECLARE_SCOPE_HIERARCHICAL_COUNTER_RIGUNIT()
 	if (!FMath::IsWithinInclusive<float>(Value, -1.f, 1.f))
@@ -258,13 +258,13 @@ UE_RigUnit_MathFloatAcos_IMPLEMENT_RIGVM(void, Execute, const FRigUnitContext& C
 	Result = FMath::Acos(Value);
 }
 
-UE_RigUnit_MathFloatAtan_IMPLEMENT_RIGVM(void, Execute, const FRigUnitContext& Context)
+FRigUnit_MathFloatAtan_Execute()
 {
     DECLARE_SCOPE_HIERARCHICAL_COUNTER_RIGUNIT()
 	Result = FMath::Atan(Value);
 }
 
-UE_RigUnit_MathFloatLawOfCosine_IMPLEMENT_RIGVM(void, Execute, const FRigUnitContext& Context)
+FRigUnit_MathFloatLawOfCosine_Execute()
 {
     DECLARE_SCOPE_HIERARCHICAL_COUNTER_RIGUNIT()
 	if ((A <= 0.f) || (B <= 0.f) || (C <= 0.f) || (A + B < C) || (A + C < B) || (B + C < A))
diff --git a/Engine/Plugins/Experimental/ControlRig/Source/ControlRig/Private/Units/Math/RigUnit_MathQuaternion.cpp b/Engine/Plugins/Experimental/ControlRig/Source/ControlRig/Private/Units/Math/RigUnit_MathQuaternion.cpp
--- a/Engine/Plugins/Experimental/ControlRig/Source/ControlRig/Private/Units/Math/RigUnit_MathQuaternion.cpp
+++ b/Engine/Plugins/Experimental/ControlRig/Source/ControlRig/Private/Units/Math/RigUnit_MathQuaternion.cpp
@@ -3,7 +3,7 @@
 #include "Units/Math/RigUnit_MathQuaternion.h"
 #include "Units/RigUnitContext.h"
 
-UE_RigUnit_MathQuaternionFromAxisAndAngle_IMPLEMENT_RIGVM(void, Execute, const FRigUnitContext& Context)
+FRigUnit_MathQuaternionFromAxisAndAngle_Execute()
 {
     DECLARE_SCOPE_HIERARCHICAL_COUNTER_RIGUNIT()
 	if (Axis.IsNearlyZero())
@@ -15,19 +15,19 @@ UE_RigUnit_MathQuaternionFromAxisAndAngle_IMPLEMENT_RIGVM(void, Execute, const F
 	Result = FQuat(Axis.GetUnsafeNormal(), Angle);
 }
 
-UE_RigUnit_MathQuaternionFromEuler_IMPLEMENT_RIGVM(void, Execute, const FRigUnitContext& Context)
+FRigUnit_MathQuaternionFromEuler_Execute()
 {
     DECLARE_SCOPE_HIERARCHICAL_COUNTER_RIGUNIT()
 	Result = FControlRigMathLibrary::QuatFromEuler(Euler, RotationOrder);
 }
 
-UE_RigUnit_MathQuaternionFromRotator_IMPLEMENT_RIGVM(void, Execute, const FRigUnitContext& Context)
+FRigUnit_MathQuaternionFromRotator_Execute()
 {
     DECLARE_SCOPE_HIERARCHICAL_COUNTER_RIGUNIT()
 	Result = FQuat(Rotator);
 }
 
-UE_RigUnit_MathQuaternionFromTwoVectors_IMPLEMENT_RIGVM(void, Execute, const FRigUnitContext& Context)
+FRigUnit_MathQuaternionFromTwoVectors_Execute()
 {
     DECLARE_SCOPE_HIERARCHICAL_COUNTER_RIGUNIT()
 	if (A.IsNearlyZero() || B.IsNearlyZero())
@@ -46,79 +46,79 @@ UE_RigUnit_MathQuaternionFromTwoVectors_IMPLEMENT_RIGVM(void, Execute, const FRi
 	Result = FQuat::FindBetweenVectors(A, B).GetNormalized();
 }
 
-UE_RigUnit_MathQuaternionToAxisAndAngle_IMPLEMENT_RIGVM(void, Execute, const FRigUnitContext& Context)
+FRigUnit_MathQuaternionToAxisAndAngle_Execute()
 {
     DECLARE_SCOPE_HIERARCHICAL_COUNTER_RIGUNIT()
 	Value.ToAxisAndAngle(Axis, Angle);
 }
 
-UE_RigUnit_MathQuaternionToEuler_IMPLEMENT_RIGVM(void, Execute, const FRigUnitContext& Context)
+FRigUnit_MathQuaternionToEuler_Execute()
 {
     DECLARE_SCOPE_HIERARCHICAL_COUNTER_RIGUNIT()
 	Result = FControlRigMathLibrary::EulerFromQuat(Value, RotationOrder);
 }
 
-UE_RigUnit_MathQuaternionToRotator_IMPLEMENT_RIGVM(void, Execute, const FRigUnitContext& Context)
+FRigUnit_MathQuaternionToRotator_Execute()
 {
     DECLARE_SCOPE_HIERARCHICAL_COUNTER_RIGUNIT()
 	Result = Value.Rotator();
 }
 
-UE_RigUnit_MathQuaternionMul_IMPLEMENT_RIGVM(void, Execute, const FRigUnitContext& Context)
+FRigUnit_MathQuaternionMul_Execute()
 {
     DECLARE_SCOPE_HIERARCHICAL_COUNTER_RIGUNIT()
 	Result = A * B;
 }
 
-UE_RigUnit_MathQuaternionInverse_IMPLEMENT_RIGVM(void, Execute, const FRigUnitContext& Context)
+FRigUnit_MathQuaternionInverse_Execute()
 {
     DECLARE_SCOPE_HIERARCHICAL_COUNTER_RIGUNIT()
 	Result = Value.Inverse();
 }
 
-UE_RigUnit_MathQuaternionSlerp_IMPLEMENT_RIGVM(void, Execute, const FRigUnitContext& Context)
+FRigUnit_MathQuaternionSlerp_Execute()
 {
     DECLARE_SCOPE_HIERARCHICAL_COUNTER_RIGUNIT()
 	Result = FQuat::Slerp(A, B, T);
 }
 
-UE_RigUnit_MathQuaternionEquals_IMPLEMENT_RIGVM(void, Execute, const FRigUnitContext& Context)
+FRigUnit_MathQuaternionEquals_Execute()
 {
     DECLARE_SCOPE_HIERARCHICAL_COUNTER_RIGUNIT()
 	Result = A == B;
 }
 
-UE_RigUnit_MathQuaternionNotEquals_IMPLEMENT_RIGVM(void, Execute, const FRigUnitContext& Context)
+FRigUnit_MathQuaternionNotEquals_Execute()
 {
     DECLARE_SCOPE_HIERARCHICAL_COUNTER_RIGUNIT()
 	Result = A != B;
 }
 
-UE_RigUnit_MathQuaternionSelectBool_IMPLEMENT_RIGVM(void, Execute, const FRigUnitContext& Context)
+FRigUnit_MathQuaternionSelectBool_Execute()
 {
     DECLARE_SCOPE_HIERARCHICAL_COUNTER_RIGUNIT()
 	Result = Condition ? IfTrue : IfFalse;
 }
 
-UE_RigUnit_MathQuaternionDot_IMPLEMENT_RIGVM(void, Execute, const FRigUnitContext& Context)
+FRigUnit_MathQuaternionDot_Execute()
 {
     DECLARE_SCOPE_HIERARCHICAL_COUNTER_RIGUNIT()
 	Result = A | B;
 }
 
-UE_RigUnit_MathQuaternionUnit_IMPLEMENT_RIGVM(void, Execute, const FRigUnitContext& Context)
+FRigUnit_MathQuaternionUnit_Execute()
 {
     DECLARE_SCOPE_HIERARCHICAL_COUNTER_RIGUNIT()
 	Result = Value.GetNormalized();
 }
 
-UE_RigUnit_MathQuaternionRotateVector_IMPLEMENT_RIGVM(void, Execute, const FRigUnitContext& Context)
+FRigUnit_MathQuaternionRotateVector_Execute()
 {
     DECLARE_SCOPE_HIERARCHICAL_COUNTER_RIGUNIT()
 	Result = Quaternion.RotateVector(Vector);
 }
 
-UE_RigUnit_MathQuaternionGetAxis_IMPLEMENT_RIGVM(void, Execute, const FRigUnitContext& Context)
+FRigUnit_MathQuaternionGetAxis_Execute()
 {
     DECLARE_SCOPE_HIERARCHICAL_COUNTER_RIGUNIT()
 	switch (Axis)
@@ -143,7 +143,7 @@ UE_RigUnit_MathQuaternionGetAxis_IMPLEMENT_RIGVM(void, Execute, const FRigUnitCo
 }
 
 
-UE_RigUnit_MathQuaternionSwingTwist_IMPLEMENT_RIGVM(void, Execute, const FRigUnitContext& Context)
+FRigUnit_MathQuaternionSwingTwist_Execute()
 {
     DECLARE_SCOPE_HIERARCHICAL_COUNTER_RIGUNIT()
 	if (TwistAxis.IsNearlyZero())
diff --git a/Engine/Plugins/Experimental/ControlRig/Source/ControlRig/Private/Units/Math/RigUnit_MathTransform.cpp b/Engine/Plugins/Experimental/ControlRig/Source/ControlRig/Private/Units/Math/RigUnit_MathTransform.cpp
--- a/Engine/Plugins/Experimental/ControlRig/Source/ControlRig/Private/Units/Math/RigUnit_MathTransform.cpp
+++ b/Engine/Plugins/Experimental/ControlRig/Source/ControlRig/Private/Units/Math/RigUnit_MathTransform.cpp
@@ -4,61 +4,61 @@
 #include "Math/ControlRigMathLibrary.h"
 #include "Units/RigUnitContext.h"
 
-UE_RigUnit_MathTransformFromEulerTransform_IMPLEMENT_RIGVM(void, Execute, const FRigUnitContext& Context)
+FRigUnit_MathTransformFromEulerTransform_Execute()
 {
     DECLARE_SCOPE_HIERARCHICAL_COUNTER_RIGUNIT()
 	Result = EulerTransform.ToFTransform();
 }
 
-UE_RigUnit_MathTransformToEulerTransform_IMPLEMENT_RIGVM(void, Execute, const FRigUnitContext& Context)
+FRigUnit_MathTransformToEulerTransform_Execute()
 {
     DECLARE_SCOPE_HIERARCHICAL_COUNTER_RIGUNIT()
 	Result.FromFTransform(Value);
 }
 
-UE_RigUnit_MathTransformMul_IMPLEMENT_RIGVM(void, Execute, const FRigUnitContext& Context)
+FRigUnit_MathTransformMul_Execute()
 {
     DECLARE_SCOPE_HIERARCHICAL_COUNTER_RIGUNIT()
 	Result = A * B;
 }
 
-UE_RigUnit_MathTransformMakeRelative_IMPLEMENT_RIGVM(void, Execute, const FRigUnitContext& Context)
+FRigUnit_MathTransformMakeRelative_Execute()
 {
     DECLARE_SCOPE_HIERARCHICAL_COUNTER_RIGUNIT()
 	Local = Global.GetRelativeTransform(Parent);
 }
 
-UE_RigUnit_MathTransformInverse_IMPLEMENT_RIGVM(void, Execute, const FRigUnitContext& Context)
+FRigUnit_MathTransformInverse_Execute()
 {
     DECLARE_SCOPE_HIERARCHICAL_COUNTER_RIGUNIT()
 	Result = Value.Inverse();
 }
 
-UE_RigUnit_MathTransformLerp_IMPLEMENT_RIGVM(void, Execute, const FRigUnitContext& Context)
+FRigUnit_MathTransformLerp_Execute()
 {
     DECLARE_SCOPE_HIERARCHICAL_COUNTER_RIGUNIT()
 	Result = FControlRigMathLibrary::LerpTransform(A, B, T);
 }
 
-UE_RigUnit_MathTransformSelectBool_IMPLEMENT_RIGVM(void, Execute, const FRigUnitContext& Context)
+FRigUnit_MathTransformSelectBool_Execute()
 {
     DECLARE_SCOPE_HIERARCHICAL_COUNTER_RIGUNIT()
 	Result = Condition ? IfTrue : IfFalse;
 }
 
-UE_RigUnit_MathTransformRotateVector_IMPLEMENT_RIGVM(void, Execute, const FRigUnitContext& Context)
+FRigUnit_MathTransformRotateVector_Execute()
 {
     DECLARE_SCOPE_HIERARCHICAL_COUNTER_RIGUNIT()
 	Result = Transform.TransformVector(Direction);
 }
 
-UE_RigUnit_MathTransformTransformVector_IMPLEMENT_RIGVM(void, Execute, const FRigUnitContext& Context)
+FRigUnit_MathTransformTransformVector_Execute()
 {
     DECLARE_SCOPE_HIERARCHICAL_COUNTER_RIGUNIT()
 	Result = Transform.TransformPosition(Location);
 }
 
-UE_RigUnit_MathTransformFromSRT_IMPLEMENT_RIGVM(void, Execute, const FRigUnitContext& Context)
+FRigUnit_MathTransformFromSRT_Execute()
 {
     DECLARE_SCOPE_HIERARCHICAL_COUNTER_RIGUNIT()
 	Transform.SetLocation(Location);
diff --git a/Engine/Plugins/Experimental/ControlRig/Source/ControlRig/Private/Units/Math/RigUnit_MathVector.cpp b/Engine/Plugins/Experimental/ControlRig/Source/ControlRig/Private/Units/Math/RigUnit_MathVector.cpp
--- a/Engine/Plugins/Experimental/ControlRig/Source/ControlRig/Private/Units/Math/RigUnit_MathVector.cpp
+++ b/Engine/Plugins/Experimental/ControlRig/Source/ControlRig/Private/Units/Math/RigUnit_MathVector.cpp
@@ -4,37 +4,37 @@
 #include "Units/RigUnitContext.h"
 #include "Math/ControlRigMathLibrary.h"
 
-UE_RigUnit_MathVectorFromFloat_IMPLEMENT_RIGVM(void, Execute, const FRigUnitContext& Context)
+FRigUnit_MathVectorFromFloat_Execute()
 {
     DECLARE_SCOPE_HIERARCHICAL_COUNTER_RIGUNIT()
 	Result = FVector(Value, Value, Value);
 }
 
-UE_RigUnit_MathVectorAdd_IMPLEMENT_RIGVM(void, Execute, const FRigUnitContext& Context)
+FRigUnit_MathVectorAdd_Execute()
 {
     DECLARE_SCOPE_HIERARCHICAL_COUNTER_RIGUNIT()
 	Result = A + B;
 }
 
-UE_RigUnit_MathVectorSub_IMPLEMENT_RIGVM(void, Execute, const FRigUnitContext& Context)
+FRigUnit_MathVectorSub_Execute()
 {
     DECLARE_SCOPE_HIERARCHICAL_COUNTER_RIGUNIT()
 	Result = A - B;
 }
 
-UE_RigUnit_MathVectorMul_IMPLEMENT_RIGVM(void, Execute, const FRigUnitContext& Context)
+FRigUnit_MathVectorMul_Execute()
 {
     DECLARE_SCOPE_HIERARCHICAL_COUNTER_RIGUNIT()
 	Result = A * B;
 }
 
-UE_RigUnit_MathVectorScale_IMPLEMENT_RIGVM(void, Execute, const FRigUnitContext& Context)
+FRigUnit_MathVectorScale_Execute()
 {
     DECLARE_SCOPE_HIERARCHICAL_COUNTER_RIGUNIT()
 	Result = Value * Factor;
 }
 
-UE_RigUnit_MathVectorDiv_IMPLEMENT_RIGVM(void, Execute, const FRigUnitContext& Context)
+FRigUnit_MathVectorDiv_Execute()
 {
     DECLARE_SCOPE_HIERARCHICAL_COUNTER_RIGUNIT()
 	if(FMath::IsNearlyZero(B.X) || FMath::IsNearlyZero(B.Y) || FMath::IsNearlyZero(B.Z))
@@ -57,7 +57,7 @@ UE_RigUnit_MathVectorDiv_IMPLEMENT_RIGVM(void, Execute, const FRigUnitContext& C
 	Result = A / B;
 }
 
-UE_RigUnit_MathVectorMod_IMPLEMENT_RIGVM(void, Execute, const FRigUnitContext& Context)
+FRigUnit_MathVectorMod_Execute()
 {
     DECLARE_SCOPE_HIERARCHICAL_COUNTER_RIGUNIT()
 	if(FMath::IsNearlyZero(B.X) || FMath::IsNearlyZero(B.Y) || FMath::IsNearlyZero(B.Z) || B.X < 0.f || B.Y < 0.f || B.Z < 0.f)
@@ -83,7 +83,7 @@ UE_RigUnit_MathVectorMod_IMPLEMENT_RIGVM(void, Execute, const FRigUnitContext& C
 	Result.Z = FMath::Fmod(A.Z, B.Z);
 }
 
-UE_RigUnit_MathVectorMin_IMPLEMENT_RIGVM(void, Execute, const FRigUnitContext& Context)
+FRigUnit_MathVectorMin_Execute()
 {
     DECLARE_SCOPE_HIERARCHICAL_COUNTER_RIGUNIT()
 	Result.X = FMath::Min<float>(A.X, B.X);
@@ -91,7 +91,7 @@ UE_RigUnit_MathVectorMin_IMPLEMENT_RIGVM(void, Execute, const FRigUnitContext& C
 	Result.Z = FMath::Min<float>(A.Z, B.Z);
 }
 
-UE_RigUnit_MathVectorMax_IMPLEMENT_RIGVM(void, Execute, const FRigUnitContext& Context)
+FRigUnit_MathVectorMax_Execute()
 {
     DECLARE_SCOPE_HIERARCHICAL_COUNTER_RIGUNIT()
 	Result.X = FMath::Max<float>(A.X, B.X);
@@ -99,13 +99,13 @@ UE_RigUnit_MathVectorMax_IMPLEMENT_RIGVM(void, Execute, const FRigUnitContext& C
 	Result.Z = FMath::Max<float>(A.Z, B.Z);
 }
 
-UE_RigUnit_MathVectorNegate_IMPLEMENT_RIGVM(void, Execute, const FRigUnitContext& Context)
+FRigUnit_MathVectorNegate_Execute()
 {
     DECLARE_SCOPE_HIERARCHICAL_COUNTER_RIGUNIT()
 	Result = -Value;
 }
 
-UE_RigUnit_MathVectorAbs_IMPLEMENT_RIGVM(void, Execute, const FRigUnitContext& Context)
+FRigUnit_MathVectorAbs_Execute()
 {
     DECLARE_SCOPE_HIERARCHICAL_COUNTER_RIGUNIT()
 	Result.X = FMath::Abs(Value.X);
@@ -113,7 +113,7 @@ UE_RigUnit_MathVectorAbs_IMPLEMENT_RIGVM(void, Execute, const FRigUnitContext& C
 	Result.Z = FMath::Abs(Value.Z);
 }
 
-UE_RigUnit_MathVectorFloor_IMPLEMENT_RIGVM(void, Execute, const FRigUnitContext& Context)
+FRigUnit_MathVectorFloor_Execute()
 {
     DECLARE_SCOPE_HIERARCHICAL_COUNTER_RIGUNIT()
 	Result.X = FMath::FloorToFloat(Value.X);
@@ -121,7 +121,7 @@ UE_RigUnit_MathVectorFloor_IMPLEMENT_RIGVM(void, Execute, const FRigUnitContext&
 	Result.Z = FMath::FloorToFloat(Value.Z);
 }
 
-UE_RigUnit_MathVectorCeil_IMPLEMENT_RIGVM(void, Execute, const FRigUnitContext& Context)
+FRigUnit_MathVectorCeil_Execute()
 {
     DECLARE_SCOPE_HIERARCHICAL_COUNTER_RIGUNIT()
 	Result.X = FMath::CeilToFloat(Value.X);
@@ -129,7 +129,7 @@ UE_RigUnit_MathVectorCeil_IMPLEMENT_RIGVM(void, Execute, const FRigUnitContext&
 	Result.Z = FMath::CeilToFloat(Value.Z);
 }
 
-UE_RigUnit_MathVectorRound_IMPLEMENT_RIGVM(void, Execute, const FRigUnitContext& Context)
+FRigUnit_MathVectorRound_Execute()
 {
     DECLARE_SCOPE_HIERARCHICAL_COUNTER_RIGUNIT()
 	Result.X = FMath::RoundToFloat(Value.X);
@@ -137,13 +137,13 @@ UE_RigUnit_MathVectorRound_IMPLEMENT_RIGVM(void, Execute, const FRigUnitContext&
 	Result.Z = FMath::RoundToFloat(Value.Z);
 }
 
-UE_RigUnit_MathVectorSign_IMPLEMENT_RIGVM(void, Execute, const FRigUnitContext& Context)
+FRigUnit_MathVectorSign_Execute()
 {
     DECLARE_SCOPE_HIERARCHICAL_COUNTER_RIGUNIT()
 	Result = Value.GetSignVector();
 }
 
-UE_RigUnit_MathVectorClamp_IMPLEMENT_RIGVM(void, Execute, const FRigUnitContext& Context)
+FRigUnit_MathVectorClamp_Execute()
 {
     DECLARE_SCOPE_HIERARCHICAL_COUNTER_RIGUNIT()
 	Result.X = FMath::Clamp<float>(Value.X, Minimum.X, Maximum.X);
@@ -151,13 +151,13 @@ UE_RigUnit_MathVectorClamp_IMPLEMENT_RIGVM(void, Execute, const FRigUnitContext&
 	Result.Z = FMath::Clamp<float>(Value.Z, Minimum.Z, Maximum.Z);
 }
 
-UE_RigUnit_MathVectorLerp_IMPLEMENT_RIGVM(void, Execute, const FRigUnitContext& Context)
+FRigUnit_MathVectorLerp_Execute()
 {
     DECLARE_SCOPE_HIERARCHICAL_COUNTER_RIGUNIT()
 	Result = FMath::Lerp<FVector>(A, B, T);
 }
 
-UE_RigUnit_MathVectorRemap_IMPLEMENT_RIGVM(void, Execute, const FRigUnitContext& Context)
+FRigUnit_MathVectorRemap_Execute()
 {
     DECLARE_SCOPE_HIERARCHICAL_COUNTER_RIGUNIT()
 	FVector Ratio(0.f, 0.f, 0.f);
@@ -189,19 +189,19 @@ UE_RigUnit_MathVectorRemap_IMPLEMENT_RIGVM(void, Execute, const FRigUnitContext&
 	Result = FMath::Lerp<FVector>(TargetMinimum, TargetMaximum, Ratio);
 }
 
-UE_RigUnit_MathVectorEquals_IMPLEMENT_RIGVM(void, Execute, const FRigUnitContext& Context)
+FRigUnit_MathVectorEquals_Execute()
 {
     DECLARE_SCOPE_HIERARCHICAL_COUNTER_RIGUNIT()
 	Result = A == B;
 }
 
-UE_RigUnit_MathVectorNotEquals_IMPLEMENT_RIGVM(void, Execute, const FRigUnitContext& Context)
+FRigUnit_MathVectorNotEquals_Execute()
 {
     DECLARE_SCOPE_HIERARCHICAL_COUNTER_RIGUNIT()
 	Result = A != B;
 }
 
-UE_RigUnit_MathVectorIsNearlyZero_IMPLEMENT_RIGVM(void, Execute, const FRigUnitContext& Context)
+FRigUnit_MathVectorIsNearlyZero_Execute()
 {
     DECLARE_SCOPE_HIERARCHICAL_COUNTER_RIGUNIT()
 	if(Tolerance < 0.f)
@@ -211,7 +211,7 @@ UE_RigUnit_MathVectorIsNearlyZero_IMPLEMENT_RIGVM(void, Execute, const FRigUnitC
 	Result = Value.IsNearlyZero(FMath::Max<float>(Tolerance, SMALL_NUMBER));
 }
 
-UE_RigUnit_MathVectorIsNearlyEqual_IMPLEMENT_RIGVM(void, Execute, const FRigUnitContext& Context)
+FRigUnit_MathVectorIsNearlyEqual_Execute()
 {
     DECLARE_SCOPE_HIERARCHICAL_COUNTER_RIGUNIT()
 	if(Tolerance < 0.f)
@@ -221,55 +221,55 @@ UE_RigUnit_MathVectorIsNearlyEqual_IMPLEMENT_RIGVM(void, Execute, const FRigUnit
 	Result = (A - B).IsNearlyZero(FMath::Max<float>(Tolerance, SMALL_NUMBER));
 }
 
-UE_RigUnit_MathVectorSelectBool_IMPLEMENT_RIGVM(void, Execute, const FRigUnitContext& Context)
+FRigUnit_MathVectorSelectBool_Execute()
 {
     DECLARE_SCOPE_HIERARCHICAL_COUNTER_RIGUNIT()
 	Result = Condition ? IfTrue : IfFalse;
 }
 
-UE_RigUnit_MathVectorDeg_IMPLEMENT_RIGVM(void, Execute, const FRigUnitContext& Context)
+FRigUnit_MathVectorDeg_Execute()
 {
     DECLARE_SCOPE_HIERARCHICAL_COUNTER_RIGUNIT()
 	Result = FMath::RadiansToDegrees(Value);
 }
 
-UE_RigUnit_MathVectorRad_IMPLEMENT_RIGVM(void, Execute, const FRigUnitContext& Context)
+FRigUnit_MathVectorRad_Execute()
 {
     DECLARE_SCOPE_HIERARCHICAL_COUNTER_RIGUNIT()
 	Result = FMath::DegreesToRadians(Value);
 }
 
-UE_RigUnit_MathVectorLengthSquared_IMPLEMENT_RIGVM(void, Execute, const FRigUnitContext& Context)
+FRigUnit_MathVectorLengthSquared_Execute()
 {
     DECLARE_SCOPE_HIERARCHICAL_COUNTER_RIGUNIT()
 	Result = Value.SizeSquared();
 }
 
-UE_RigUnit_MathVectorLength_IMPLEMENT_RIGVM(void, Execute, const FRigUnitContext& Context)
+FRigUnit_MathVectorLength_Execute()
 {
     DECLARE_SCOPE_HIERARCHICAL_COUNTER_RIGUNIT()
 	Result = Value.Size();
 }
 
-UE_RigUnit_MathVectorDistance_IMPLEMENT_RIGVM(void, Execute, const FRigUnitContext& Context)
+FRigUnit_MathVectorDistance_Execute()
 {
     DECLARE_SCOPE_HIERARCHICAL_COUNTER_RIGUNIT()
 	Result = FVector::Distance(A, B);
 }
 
-UE_RigUnit_MathVectorCross_IMPLEMENT_RIGVM(void, Execute, const FRigUnitContext& Context)
+FRigUnit_MathVectorCross_Execute()
 {
     DECLARE_SCOPE_HIERARCHICAL_COUNTER_RIGUNIT()
 	Result = A ^ B;
 }
 
-UE_RigUnit_MathVectorDot_IMPLEMENT_RIGVM(void, Execute, const FRigUnitContext& Context)
+FRigUnit_MathVectorDot_Execute()
 {
     DECLARE_SCOPE_HIERARCHICAL_COUNTER_RIGUNIT()
 	Result = A | B;
 }
 
-UE_RigUnit_MathVectorUnit_IMPLEMENT_RIGVM(void, Execute, const FRigUnitContext& Context)
+FRigUnit_MathVectorUnit_Execute()
 {
     DECLARE_SCOPE_HIERARCHICAL_COUNTER_RIGUNIT()
 	if (Value.IsNearlyZero())
@@ -281,7 +281,7 @@ UE_RigUnit_MathVectorUnit_IMPLEMENT_RIGVM(void, Execute, const FRigUnitContext&
 	Result = Value.GetUnsafeNormal();
 }
 
-UE_RigUnit_MathVectorMirror_IMPLEMENT_RIGVM(void, Execute, const FRigUnitContext& Context)
+FRigUnit_MathVectorMirror_Execute()
 {
     DECLARE_SCOPE_HIERARCHICAL_COUNTER_RIGUNIT()
 	if (Normal.IsNearlyZero())
@@ -293,7 +293,7 @@ UE_RigUnit_MathVectorMirror_IMPLEMENT_RIGVM(void, Execute, const FRigUnitContext
 	Result = Value.MirrorByVector(Normal.GetSafeNormal());
 }
 
-UE_RigUnit_MathVectorAngle_IMPLEMENT_RIGVM(void, Execute, const FRigUnitContext& Context)
+FRigUnit_MathVectorAngle_Execute()
 {
     DECLARE_SCOPE_HIERARCHICAL_COUNTER_RIGUNIT()
 	if (A.IsNearlyZero() || B.IsNearlyZero())
@@ -312,7 +312,7 @@ UE_RigUnit_MathVectorAngle_IMPLEMENT_RIGVM(void, Execute, const FRigUnitContext&
 	Result = FQuat::FindBetween(A, B).GetAngle();
 }
 
-UE_RigUnit_MathVectorParallel_IMPLEMENT_RIGVM(void, Execute, const FRigUnitContext& Context)
+FRigUnit_MathVectorParallel_Execute()
 {
     DECLARE_SCOPE_HIERARCHICAL_COUNTER_RIGUNIT()
 	if (A.IsNearlyZero() || B.IsNearlyZero())
@@ -331,7 +331,7 @@ UE_RigUnit_MathVectorParallel_IMPLEMENT_RIGVM(void, Execute, const FRigUnitConte
 	Result = FVector::Parallel(A, B);
 }
 
-UE_RigUnit_MathVectorOrthogonal_IMPLEMENT_RIGVM(void, Execute, const FRigUnitContext& Context)
+FRigUnit_MathVectorOrthogonal_Execute()
 {
     DECLARE_SCOPE_HIERARCHICAL_COUNTER_RIGUNIT()
 	if (A.IsNearlyZero() || B.IsNearlyZero())
@@ -350,7 +350,7 @@ UE_RigUnit_MathVectorOrthogonal_IMPLEMENT_RIGVM(void, Execute, const FRigUnitCon
 	Result = FVector::Orthogonal(A, B);
 }
 
-UE_RigUnit_MathVectorBezierFourPoint_IMPLEMENT_RIGVM(void, Execute, const FRigUnitContext& Context)
+FRigUnit_MathVectorBezierFourPoint_Execute()
 {
     DECLARE_SCOPE_HIERARCHICAL_COUNTER_RIGUNIT()
 	FControlRigMathLibrary::FourPointBezier(Bezier, T, Result, Tangent);
diff --git a/Engine/Plugins/Experimental/ControlRig/Source/ControlRig/Private/Units/Math/RigUnit_Noise.cpp b/Engine/Plugins/Experimental/ControlRig/Source/ControlRig/Private/Units/Math/RigUnit_Noise.cpp
--- a/Engine/Plugins/Experimental/ControlRig/Source/ControlRig/Private/Units/Math/RigUnit_Noise.cpp
+++ b/Engine/Plugins/Experimental/ControlRig/Source/ControlRig/Private/Units/Math/RigUnit_Noise.cpp
@@ -3,7 +3,7 @@
 #include "Units/Math/RigUnit_Noise.h"
 #include "Units/RigUnitContext.h"
 
-UE_RigUnit_NoiseFloat_IMPLEMENT_RIGVM(void, Execute, const FRigUnitContext& Context)
+FRigUnit_NoiseFloat_Execute()
 {
     DECLARE_SCOPE_HIERARCHICAL_COUNTER_RIGUNIT()
 	if (Context.State == EControlRigState::Init)
@@ -17,7 +17,7 @@ UE_RigUnit_NoiseFloat_IMPLEMENT_RIGVM(void, Execute, const FRigUnitContext& Cont
 	Time = Time + Speed * Context.DeltaTime;
 }
 
-UE_RigUnit_NoiseVector_IMPLEMENT_RIGVM(void, Execute, const FRigUnitContext& Context)
+FRigUnit_NoiseVector_Execute()
 {
     DECLARE_SCOPE_HIERARCHICAL_COUNTER_RIGUNIT()
 	if (Context.State == EControlRigState::Init)
diff --git a/Engine/Plugins/Experimental/ControlRig/Source/ControlRig/Private/Units/Math/RigUnit_Random.cpp b/Engine/Plugins/Experimental/ControlRig/Source/ControlRig/Private/Units/Math/RigUnit_Random.cpp
--- a/Engine/Plugins/Experimental/ControlRig/Source/ControlRig/Private/Units/Math/RigUnit_Random.cpp
+++ b/Engine/Plugins/Experimental/ControlRig/Source/ControlRig/Private/Units/Math/RigUnit_Random.cpp
@@ -15,7 +15,7 @@ float FRigUnit_Random_Helper(int32& Seed)
 	return FPlatformMath::Fractional(Result.f);
 }
 
-UE_RigUnit_RandomFloat_IMPLEMENT_RIGVM(void, Execute, const FRigUnitContext& Context)
+FRigUnit_RandomFloat_Execute()
 {
     DECLARE_SCOPE_HIERARCHICAL_COUNTER_RIGUNIT()
 	if (Context.State == EControlRigState::Init)
@@ -38,7 +38,7 @@ UE_RigUnit_RandomFloat_IMPLEMENT_RIGVM(void, Execute, const FRigUnitContext& Con
 	LastResult = Result;
 }
 
-UE_RigUnit_RandomVector_IMPLEMENT_RIGVM(void, Execute, const FRigUnitContext& Context)
+FRigUnit_RandomVector_Execute()
 {
     DECLARE_SCOPE_HIERARCHICAL_COUNTER_RIGUNIT()
 	if (Context.State == EControlRigState::Init)
diff --git a/Engine/Plugins/Experimental/ControlRig/Source/ControlRig/Private/Units/Simulation/RigUnit_Accumulate.cpp b/Engine/Plugins/Experimental/ControlRig/Source/ControlRig/Private/Units/Simulation/RigUnit_Accumulate.cpp
--- a/Engine/Plugins/Experimental/ControlRig/Source/ControlRig/Private/Units/Simulation/RigUnit_Accumulate.cpp
+++ b/Engine/Plugins/Experimental/ControlRig/Source/ControlRig/Private/Units/Simulation/RigUnit_Accumulate.cpp
@@ -3,7 +3,7 @@
 #include "Units/Simulation/RigUnit_Accumulate.h"
 #include "Units/RigUnitContext.h"
 
-UE_RigUnit_AccumulateFloatAdd_IMPLEMENT_RIGVM(void, Execute, const FRigUnitContext& Context)
+FRigUnit_AccumulateFloatAdd_Execute()
 {
     DECLARE_SCOPE_HIERARCHICAL_COUNTER_RIGUNIT()
 	if (Context.State == EControlRigState::Init)
@@ -20,7 +20,7 @@ UE_RigUnit_AccumulateFloatAdd_IMPLEMENT_RIGVM(void, Execute, const FRigUnitConte
 	}
 }
 
-UE_RigUnit_AccumulateVectorAdd_IMPLEMENT_RIGVM(void, Execute, const FRigUnitContext& Context)
+FRigUnit_AccumulateVectorAdd_Execute()
 {
     DECLARE_SCOPE_HIERARCHICAL_COUNTER_RIGUNIT()
 	if (Context.State == EControlRigState::Init)
@@ -37,7 +37,7 @@ UE_RigUnit_AccumulateVectorAdd_IMPLEMENT_RIGVM(void, Execute, const FRigUnitCont
 	}
 }
 
-UE_RigUnit_AccumulateFloatMul_IMPLEMENT_RIGVM(void, Execute, const FRigUnitContext& Context)
+FRigUnit_AccumulateFloatMul_Execute()
 {
     DECLARE_SCOPE_HIERARCHICAL_COUNTER_RIGUNIT()
 	if (Context.State == EControlRigState::Init)
@@ -55,7 +55,7 @@ UE_RigUnit_AccumulateFloatMul_IMPLEMENT_RIGVM(void, Execute, const FRigUnitConte
 	}
 }
 
-UE_RigUnit_AccumulateVectorMul_IMPLEMENT_RIGVM(void, Execute, const FRigUnitContext& Context)
+FRigUnit_AccumulateVectorMul_Execute()
 {
     DECLARE_SCOPE_HIERARCHICAL_COUNTER_RIGUNIT()
 	if (Context.State == EControlRigState::Init)
@@ -73,7 +73,7 @@ UE_RigUnit_AccumulateVectorMul_IMPLEMENT_RIGVM(void, Execute, const FRigUnitCont
 	}
 }
 
-UE_RigUnit_AccumulateQuatMul_IMPLEMENT_RIGVM(void, Execute, const FRigUnitContext& Context)
+FRigUnit_AccumulateQuatMul_Execute()
 {
     DECLARE_SCOPE_HIERARCHICAL_COUNTER_RIGUNIT()
 	if (Context.State == EControlRigState::Init)
@@ -98,7 +98,7 @@ UE_RigUnit_AccumulateQuatMul_IMPLEMENT_RIGVM(void, Execute, const FRigUnitContex
 	}
 }
 
-UE_RigUnit_AccumulateTransformMul_IMPLEMENT_RIGVM(void, Execute, const FRigUnitContext& Context)
+FRigUnit_AccumulateTransformMul_Execute()
 {
     DECLARE_SCOPE_HIERARCHICAL_COUNTER_RIGUNIT()
 	if (Context.State == EControlRigState::Init)
@@ -125,7 +125,7 @@ UE_RigUnit_AccumulateTransformMul_IMPLEMENT_RIGVM(void, Execute, const FRigUnitC
 	}
 }
 
-UE_RigUnit_AccumulateFloatLerp_IMPLEMENT_RIGVM(void, Execute, const FRigUnitContext& Context)
+FRigUnit_AccumulateFloatLerp_Execute()
 {
     DECLARE_SCOPE_HIERARCHICAL_COUNTER_RIGUNIT()
 	if (Context.State == EControlRigState::Init)
@@ -138,7 +138,7 @@ UE_RigUnit_AccumulateFloatLerp_IMPLEMENT_RIGVM(void, Execute, const FRigUnitCont
 	}
 }
 
-UE_RigUnit_AccumulateVectorLerp_IMPLEMENT_RIGVM(void, Execute, const FRigUnitContext& Context)
+FRigUnit_AccumulateVectorLerp_Execute()
 {
     DECLARE_SCOPE_HIERARCHICAL_COUNTER_RIGUNIT()
 	if (Context.State == EControlRigState::Init)
@@ -151,7 +151,7 @@ UE_RigUnit_AccumulateVectorLerp_IMPLEMENT_RIGVM(void, Execute, const FRigUnitCon
 	}
 }
 
-UE_RigUnit_AccumulateQuatLerp_IMPLEMENT_RIGVM(void, Execute, const FRigUnitContext& Context)
+FRigUnit_AccumulateQuatLerp_Execute()
 {
     DECLARE_SCOPE_HIERARCHICAL_COUNTER_RIGUNIT()
 	if (Context.State == EControlRigState::Init)
@@ -164,7 +164,7 @@ UE_RigUnit_AccumulateQuatLerp_IMPLEMENT_RIGVM(void, Execute, const FRigUnitConte
 	}
 }
 
-UE_RigUnit_AccumulateTransformLerp_IMPLEMENT_RIGVM(void, Execute, const FRigUnitContext& Context)
+FRigUnit_AccumulateTransformLerp_Execute()
 {
     DECLARE_SCOPE_HIERARCHICAL_COUNTER_RIGUNIT()
 	if (Context.State == EControlRigState::Init)
@@ -181,7 +181,7 @@ UE_RigUnit_AccumulateTransformLerp_IMPLEMENT_RIGVM(void, Execute, const FRigUnit
 	}
 }
 
-UE_RigUnit_AccumulateFloatRange_IMPLEMENT_RIGVM(void, Execute, const FRigUnitContext& Context)
+FRigUnit_AccumulateFloatRange_Execute()
 {
     DECLARE_SCOPE_HIERARCHICAL_COUNTER_RIGUNIT()
 	if (Context.State == EControlRigState::Init)
@@ -195,7 +195,7 @@ UE_RigUnit_AccumulateFloatRange_IMPLEMENT_RIGVM(void, Execute, const FRigUnitCon
 	}
 }
 
-UE_RigUnit_AccumulateVectorRange_IMPLEMENT_RIGVM(void, Execute, const FRigUnitContext& Context)
+FRigUnit_AccumulateVectorRange_Execute()
 {
     DECLARE_SCOPE_HIERARCHICAL_COUNTER_RIGUNIT()
 	if (Context.State == EControlRigState::Init)
diff --git a/Engine/Plugins/Experimental/ControlRig/Source/ControlRig/Private/Units/Simulation/RigUnit_DeltaFromPrevious.cpp b/Engine/Plugins/Experimental/ControlRig/Source/ControlRig/Private/Units/Simulation/RigUnit_DeltaFromPrevious.cpp
--- a/Engine/Plugins/Experimental/ControlRig/Source/ControlRig/Private/Units/Simulation/RigUnit_DeltaFromPrevious.cpp
+++ b/Engine/Plugins/Experimental/ControlRig/Source/ControlRig/Private/Units/Simulation/RigUnit_DeltaFromPrevious.cpp
@@ -3,7 +3,7 @@
 #include "Units/Simulation/RigUnit_DeltaFromPrevious.h"
 #include "Units/RigUnitContext.h"
 
-UE_RigUnit_DeltaFromPreviousFloat_IMPLEMENT_RIGVM(void, Execute, const FRigUnitContext& Context)
+FRigUnit_DeltaFromPreviousFloat_Execute()
 {
     DECLARE_SCOPE_HIERARCHICAL_COUNTER_RIGUNIT()
 	if (Context.State == EControlRigState::Init)
@@ -18,7 +18,7 @@ UE_RigUnit_DeltaFromPreviousFloat_IMPLEMENT_RIGVM(void, Execute, const FRigUnitC
 	Cache = Value;
 }
 
-UE_RigUnit_DeltaFromPreviousVector_IMPLEMENT_RIGVM(void, Execute, const FRigUnitContext& Context)
+FRigUnit_DeltaFromPreviousVector_Execute()
 {
     DECLARE_SCOPE_HIERARCHICAL_COUNTER_RIGUNIT()
 	if (Context.State == EControlRigState::Init)
@@ -33,7 +33,7 @@ UE_RigUnit_DeltaFromPreviousVector_IMPLEMENT_RIGVM(void, Execute, const FRigUnit
 	Cache = Value;
 }
 
-UE_RigUnit_DeltaFromPreviousQuat_IMPLEMENT_RIGVM(void, Execute, const FRigUnitContext& Context)
+FRigUnit_DeltaFromPreviousQuat_Execute()
 {
     DECLARE_SCOPE_HIERARCHICAL_COUNTER_RIGUNIT()
 	if (Context.State == EControlRigState::Init)
@@ -48,7 +48,7 @@ UE_RigUnit_DeltaFromPreviousQuat_IMPLEMENT_RIGVM(void, Execute, const FRigUnitCo
 	Cache = Value;
 }
 
-UE_RigUnit_DeltaFromPreviousTransform_IMPLEMENT_RIGVM(void, Execute, const FRigUnitContext& Context)
+FRigUnit_DeltaFromPreviousTransform_Execute()
 {
     DECLARE_SCOPE_HIERARCHICAL_COUNTER_RIGUNIT()
 	if (Context.State == EControlRigState::Init)
diff --git a/Engine/Plugins/Experimental/ControlRig/Source/ControlRig/Private/Units/Simulation/RigUnit_Kalman.cpp b/Engine/Plugins/Experimental/ControlRig/Source/ControlRig/Private/Units/Simulation/RigUnit_Kalman.cpp
--- a/Engine/Plugins/Experimental/ControlRig/Source/ControlRig/Private/Units/Simulation/RigUnit_Kalman.cpp
+++ b/Engine/Plugins/Experimental/ControlRig/Source/ControlRig/Private/Units/Simulation/RigUnit_Kalman.cpp
@@ -3,7 +3,7 @@
 #include "Units/Simulation/RigUnit_Kalman.h"
 #include "Units/RigUnitContext.h"
 
-UE_RigUnit_KalmanFloat_IMPLEMENT_RIGVM(void, Execute, const FRigUnitContext& Context)
+FRigUnit_KalmanFloat_Execute()
 {
     DECLARE_SCOPE_HIERARCHICAL_COUNTER_RIGUNIT()
 	if(BufferSize <= 0)
@@ -52,7 +52,7 @@ UE_RigUnit_KalmanFloat_IMPLEMENT_RIGVM(void, Execute, const FRigUnitContext& Con
 	}
 }
 
-UE_RigUnit_KalmanVector_IMPLEMENT_RIGVM(void, Execute, const FRigUnitContext& Context)
+FRigUnit_KalmanVector_Execute()
 {
     DECLARE_SCOPE_HIERARCHICAL_COUNTER_RIGUNIT()
 	if(BufferSize <= 0)
@@ -101,7 +101,7 @@ UE_RigUnit_KalmanVector_IMPLEMENT_RIGVM(void, Execute, const FRigUnitContext& Co
 	}
 }
 
-UE_RigUnit_KalmanTransform_IMPLEMENT_RIGVM(void, Execute, const FRigUnitContext& Context)
+FRigUnit_KalmanTransform_Execute()
 {
     DECLARE_SCOPE_HIERARCHICAL_COUNTER_RIGUNIT()
 	if(BufferSize <= 0)
diff --git a/Engine/Plugins/Experimental/ControlRig/Source/ControlRig/Private/Units/Simulation/RigUnit_PointSimulation.cpp b/Engine/Plugins/Experimental/ControlRig/Source/ControlRig/Private/Units/Simulation/RigUnit_PointSimulation.cpp
--- a/Engine/Plugins/Experimental/ControlRig/Source/ControlRig/Private/Units/Simulation/RigUnit_PointSimulation.cpp
+++ b/Engine/Plugins/Experimental/ControlRig/Source/ControlRig/Private/Units/Simulation/RigUnit_PointSimulation.cpp
@@ -3,7 +3,7 @@
 #include "Units/Simulation/RigUnit_PointSimulation.h"
 #include "Units/RigUnitContext.h"
 
-UE_RigUnit_PointSimulation_IMPLEMENT_RIGVM(void, Execute, const FRigUnitContext& Context)
+FRigUnit_PointSimulation_Execute()
 {
     DECLARE_SCOPE_HIERARCHICAL_COUNTER_RIGUNIT()
 	FRigBoneHierarchy* Hierarchy = ExecuteContext.GetBones();
diff --git a/Engine/Plugins/Experimental/ControlRig/Source/ControlRig/Private/Units/Simulation/RigUnit_TimeOffset.cpp b/Engine/Plugins/Experimental/ControlRig/Source/ControlRig/Private/Units/Simulation/RigUnit_TimeOffset.cpp
--- a/Engine/Plugins/Experimental/ControlRig/Source/ControlRig/Private/Units/Simulation/RigUnit_TimeOffset.cpp
+++ b/Engine/Plugins/Experimental/ControlRig/Source/ControlRig/Private/Units/Simulation/RigUnit_TimeOffset.cpp
@@ -3,7 +3,7 @@
 #include "Units/Simulation/RigUnit_TimeOffset.h"
 #include "Units/RigUnitContext.h"
 
-UE_RigUnit_TimeOffsetFloat_IMPLEMENT_RIGVM(void, Execute, const FRigUnitContext& Context)
+FRigUnit_TimeOffsetFloat_Execute()
 {
     DECLARE_SCOPE_HIERARCHICAL_COUNTER_RIGUNIT()
 	Result = Value;
@@ -108,7 +108,7 @@ UE_RigUnit_TimeOffsetFloat_IMPLEMENT_RIGVM(void, Execute, const FRigUnitContext&
 	}
 }
 
-UE_RigUnit_TimeOffsetVector_IMPLEMENT_RIGVM(void, Execute, const FRigUnitContext& Context)
+FRigUnit_TimeOffsetVector_Execute()
 {
     DECLARE_SCOPE_HIERARCHICAL_COUNTER_RIGUNIT()
 	Result = Value;
@@ -213,7 +213,7 @@ UE_RigUnit_TimeOffsetVector_IMPLEMENT_RIGVM(void, Execute, const FRigUnitContext
 	}
 }
 
-UE_RigUnit_TimeOffsetTransform_IMPLEMENT_RIGVM(void, Execute, const FRigUnitContext& Context)
+FRigUnit_TimeOffsetTransform_Execute()
 {
     DECLARE_SCOPE_HIERARCHICAL_COUNTER_RIGUNIT()
 	Result = Value;
diff --git a/Engine/Plugins/Experimental/ControlRig/Source/ControlRig/Private/Units/Simulation/RigUnit_Timeline.cpp b/Engine/Plugins/Experimental/ControlRig/Source/ControlRig/Private/Units/Simulation/RigUnit_Timeline.cpp
--- a/Engine/Plugins/Experimental/ControlRig/Source/ControlRig/Private/Units/Simulation/RigUnit_Timeline.cpp
+++ b/Engine/Plugins/Experimental/ControlRig/Source/ControlRig/Private/Units/Simulation/RigUnit_Timeline.cpp
@@ -3,7 +3,7 @@
 #include "Units/Simulation/RigUnit_Timeline.h"
 #include "Units/RigUnitContext.h"
 
-UE_RigUnit_Timeline_IMPLEMENT_RIGVM(void, Execute, const FRigUnitContext& Context)
+FRigUnit_Timeline_Execute()
 {
     DECLARE_SCOPE_HIERARCHICAL_COUNTER_RIGUNIT()
 	if (Context.State == EControlRigState::Init)
diff --git a/Engine/Plugins/Experimental/ControlRig/Source/ControlRig/Private/Units/Simulation/RigUnit_Verlet.cpp b/Engine/Plugins/Experimental/ControlRig/Source/ControlRig/Private/Units/Simulation/RigUnit_Verlet.cpp
--- a/Engine/Plugins/Experimental/ControlRig/Source/ControlRig/Private/Units/Simulation/RigUnit_Verlet.cpp
+++ b/Engine/Plugins/Experimental/ControlRig/Source/ControlRig/Private/Units/Simulation/RigUnit_Verlet.cpp
@@ -3,7 +3,7 @@
 #include "Units/Simulation/RigUnit_Verlet.h"
 #include "Units/RigUnitContext.h"
 
-UE_RigUnit_VerletIntegrateVector_IMPLEMENT_RIGVM(void, Execute, const FRigUnitContext& Context)
+FRigUnit_VerletIntegrateVector_Execute()
 {
     DECLARE_SCOPE_HIERARCHICAL_COUNTER_RIGUNIT()
 	if (Context.State == EControlRigState::Init)
diff --git a/Engine/Source/Programs/UnrealHeaderTool/Private/CodeGenerator.cpp b/Engine/Source/Programs/UnrealHeaderTool/Private/CodeGenerator.cpp
--- a/Engine/Source/Programs/UnrealHeaderTool/Private/CodeGenerator.cpp
+++ b/Engine/Source/Programs/UnrealHeaderTool/Private/CodeGenerator.cpp
@@ -3349,8 +3349,12 @@ void FNativeClassHeaderGenerator::ExportGeneratedStructBodyMacros(FOutputDevice&
 		FString StructMembers = StructRigVMInfo->Members.Declarations(false, TEXT(", \\\r\n\t\t"), true, false);
 
 		OutGeneratedHeaderText.Log(TEXT("\n"));
-		OutGeneratedHeaderText.Logf(TEXT("#define %s_IMPLEMENT_RIGVM(ReturnType, FunctionName, ...) \\\r\n"), *RigVMMacroPrefix);
-		OutGeneratedHeaderText.Logf(TEXT("\tReturnType %s::Static##FunctionName(\\\r\n\t\t%s,\\\r\n\t\t##__VA_ARGS__ \\\r\n\t)\n"), StructNameCPP, *StructMembers);
+		for (const FRigVMMethodInfo& MethodInfo : StructRigVMInfo->Methods)
+		{
+			FString ParameterSuffix = MethodInfo.Parameters.Declarations(true, TEXT(", \\\r\n\t\t"));
+			OutGeneratedHeaderText.Logf(TEXT("#define %s_%s() \\\r\n"), StructNameCPP, *MethodInfo.Name);
+			OutGeneratedHeaderText.Logf(TEXT("\t%s %s::Static%s( \\\r\n\t\t%s%s \\\r\n\t)\n"), *MethodInfo.ReturnType, StructNameCPP, *MethodInfo.Name, *StructMembers, *ParameterSuffix);
+		}
 		OutGeneratedHeaderText.Log(TEXT("\n"));
 	}
 
@@ -3602,7 +3606,7 @@ void FNativeClassHeaderGenerator::ExportGeneratedStructBodyMacros(FOutputDevice&
 
 	// if this struct has RigVM methods we need to implement both the 
 	// virtual function as well as the stub method here.
-	// The static method is implemented by the user using the IMPLEMENT_RIGVM macro.
+	// The static method is implemented by the user using a macro.
 	if (StructRigVMInfo)
 	{
 		FString StructMembersForVirtualFunc = StructRigVMInfo->Members.Names(false, TEXT(",\r\n\t\t"), true);
