COPY from //Dev-Anim to //Dev-Main

diff --git a/Engine/Config/BaseEngine.ini b/Engine/Config/BaseEngine.ini
--- a/Engine/Config/BaseEngine.ini
+++ b/Engine/Config/BaseEngine.ini
@@ -757,6 +757,9 @@ bUseStreamingPause=false
 +ClassRedirects=(OldName="PixelStreamingInputComponent",NewName="PixelStreamerInputComponent")
 +PropertyRedirects=(OldName="PixelStreamerInputComponent.OnPixelStreamingInputEvent",NewName="PixelStreamerInputComponent.OnInputEvent")
 
+; 4.25
++EnumRedirects=(OldName="ECurveBlendOption",ValueChanges=(("ECurveBlendOption::MaxWeight", "ECurveBlendOption::Override"))
+
 [CoreUObject.Metadata]
 ; Note: UnrealHeaderTool should be rerun after making changes to MetadataRedirects to catch any keys specified in class headers
 +MetadataRedirects=(OldKey="K2Protected", NewKey="BlueprintProtected")
@@ -2334,6 +2337,8 @@ HMDWornMovementThreshold = 50.0
 bStripAnimationDataOnDedicatedServer=False
 
 [Animation.DefaultObjectSettings]
+BoneCompressionSettings="/Engine/Animation/DefaultAnimBoneCompressionSettings"
+AnimationRecorderBoneCompressionSettings="/Engine/Animation/DefaultRecorderBoneCompression"
 CurveCompressionSettings="/Engine/Animation/DefaultAnimCurveCompressionSettings"
 
 [/Script/Engine.MeshSimplificationSettings]
diff --git a/Engine/Plugins/Animation/LiveLink/Source/LiveLink/Private/LiveLinkBlueprintLibrary.cpp b/Engine/Plugins/Animation/LiveLink/Source/LiveLink/Private/LiveLinkBlueprintLibrary.cpp
--- a/Engine/Plugins/Animation/LiveLink/Source/LiveLink/Private/LiveLinkBlueprintLibrary.cpp
+++ b/Engine/Plugins/Animation/LiveLink/Source/LiveLink/Private/LiveLinkBlueprintLibrary.cpp
@@ -47,11 +47,11 @@ void ULiveLinkBlueprintLibrary::GetMetadata(UPARAM(ref) FSubjectFrameHandle& Sub
 
 void ULiveLinkBlueprintLibrary::GetBasicData(UPARAM(ref) FSubjectFrameHandle& SubjectFrameHandle, FLiveLinkBasicBlueprintData& BasicBlueprintData)
 {
-	if (FLiveLinkSkeletonStaticData* StaticData = SubjectFrameHandle.GetSourceSkeletonStaticData())
+	if (const FLiveLinkSkeletonStaticData* StaticData = SubjectFrameHandle.GetSourceSkeletonStaticData())
 	{
 		BasicBlueprintData.StaticData = *StaticData;
 	}
-	if (FLiveLinkAnimationFrameData* FrameData = SubjectFrameHandle.GetSourceAnimationFrameData())
+	if (const FLiveLinkAnimationFrameData* FrameData = SubjectFrameHandle.GetSourceAnimationFrameData())
 	{
 		BasicBlueprintData.FrameData = *FrameData;
 	}
diff --git a/Engine/Plugins/Experimental/AlembicImporter/Source/AlembicLibrary/Private/AbcImporter.cpp b/Engine/Plugins/Experimental/AlembicImporter/Source/AlembicLibrary/Private/AbcImporter.cpp
--- a/Engine/Plugins/Experimental/AlembicImporter/Source/AlembicLibrary/Private/AbcImporter.cpp
+++ b/Engine/Plugins/Experimental/AlembicImporter/Source/AlembicLibrary/Private/AbcImporter.cpp
@@ -1596,6 +1596,7 @@ bool FAbcImporter::BuildSkeletalMesh( FSkeletalMeshLODModel& LODModel, const FRe
 		TargetSection.BoneMap.Add(0);
 
 		TargetSection.CalcMaxBoneInfluences();
+		TargetSection.CalcUse16BitBoneIndex();
 	}
 
 	// Only using bone zero
diff --git a/Engine/Plugins/Experimental/ControlRig/ControlRig.uplugin b/Engine/Plugins/Experimental/ControlRig/ControlRig.uplugin
--- a/Engine/Plugins/Experimental/ControlRig/ControlRig.uplugin
+++ b/Engine/Plugins/Experimental/ControlRig/ControlRig.uplugin
@@ -14,6 +14,13 @@
 	"CanContainContent" : true,
 	"IsBetaVersion" : true,
 	"Installed" : false,
+	"Plugins":
+	[
+		{
+			"Name": "LevelSequenceEditor",
+			"Enabled": true
+		}
+	],
 	"Modules" :
 	[
 		{
diff --git a/Engine/Plugins/Experimental/ControlRig/Source/ControlRig/ControlRig.Build.cs b/Engine/Plugins/Experimental/ControlRig/Source/ControlRig/ControlRig.Build.cs
--- a/Engine/Plugins/Experimental/ControlRig/Source/ControlRig/ControlRig.Build.cs
+++ b/Engine/Plugins/Experimental/ControlRig/Source/ControlRig/ControlRig.Build.cs
@@ -36,12 +36,21 @@ namespace UnrealBuildTool.Rules
 
             if (Target.bBuildEditor == true)
             {
+                PublicDependencyModuleNames.AddRange(
+				    new string[]
+					{
+						"RigVMDeveloper",
+                        "AnimGraph",
+                    }
+                );
+
                 PrivateDependencyModuleNames.AddRange(
                     new string[]
                     {
                         "UnrealEd",
                         "BlueprintGraph",
                         "PropertyEditor",
+                        "RigVMDeveloper",
                     }
                 );
 
diff --git a/Engine/Plugins/Experimental/ControlRig/Source/ControlRig/Private/AnimNode_ControlRig.cpp b/Engine/Plugins/Experimental/ControlRig/Source/ControlRig/Private/AnimNode_ControlRig.cpp
--- a/Engine/Plugins/Experimental/ControlRig/Source/ControlRig/Private/AnimNode_ControlRig.cpp
+++ b/Engine/Plugins/Experimental/ControlRig/Source/ControlRig/Private/AnimNode_ControlRig.cpp
@@ -8,7 +8,6 @@
 #include "GameFramework/Actor.h"
 #include "Animation/NodeMappingContainer.h"
 #include "AnimationRuntime.h"
-#include "ControlRigVariables.h"
 
 #if WITH_EDITOR
 #include "Editor.h"
@@ -186,6 +185,7 @@ void FAnimNode_ControlRig::UpdateInput(UControlRig* InControlRig, const FPoseCon
 	DECLARE_SCOPE_HIERARCHICAL_COUNTER_FUNC()
 
 	FAnimNode_ControlRigBase::UpdateInput(InControlRig, InOutput);
+
 	// now go through variable mapping table and see if anything is mapping through input
 	if (InputMapping.Num() > 0 && InControlRig)
 	{
@@ -201,12 +201,15 @@ void FAnimNode_ControlRig::UpdateInput(UControlRig* InControlRig, const FPoseCon
 				if (UID != SmartName::MaxUID)
 				{
 					const float Value = InOutput.Curve.Get(UID);
-	
-					// helper function to set input value for ControlRig
-					// This converts to the proper destination type, and sets the float type Value
-					if (!FControlRigIOHelper::SetInputValue(InControlRig, SourcePath, FControlRigIOTypes::GetTypeString<float>(), Value))
+
+					FRigVMParameter Parameter = InControlRig->GetVM()->GetParameterByName(SourcePath);
+					if (Parameter.GetType() == ERigVMParameterType::Input && Parameter.GetCPPType() == TEXT("float"))
 					{
-						UE_LOG(LogAnimation, Warning, TEXT("[%s] Missing Input Property [%s]"), *GetNameSafe(InControlRig->GetClass()), *SourcePath.ToString());
+						InControlRig->GetVM()->SetParameterValue<float>(Parameter, Value);
+					}
+					else
+					{
+						UE_LOG(LogAnimation, Warning, TEXT("[%s] Missing Input Parameter [%s]"), *GetNameSafe(InControlRig->GetClass()), *SourcePath.ToString());
 					}
 				}
 			}
@@ -220,7 +223,6 @@ void FAnimNode_ControlRig::UpdateOutput(UControlRig* InControlRig, FPoseContext&
 
 	FAnimNode_ControlRigBase::UpdateOutput(InControlRig, InOutput);
 
-	// update output curves
 	if (OutputMapping.Num() > 0 && InControlRig)
 	{
 		for (auto Iter = OutputMapping.CreateConstIterator(); Iter; ++Iter)
@@ -231,11 +233,10 @@ void FAnimNode_ControlRig::UpdateOutput(UControlRig* InControlRig, FPoseContext&
 
 			if (SourcePath != NAME_None)
 			{
-				// find Segment is right value
-				float Value;
-				// helper function to get output value and convert to float 
-				if (FControlRigIOHelper::GetOutputValue(InControlRig, SourcePath, FControlRigIOTypes::GetTypeString<float>(), Value))
+				FRigVMParameter Parameter = InControlRig->GetVM()->GetParameterByName(SourcePath);
+				if (Parameter.GetType() == ERigVMParameterType::Output && Parameter.GetCPPType() == TEXT("float"))
 				{
+					float Value = InControlRig->GetVM()->GetParameterValue<float>(Parameter);
 					SmartName::UID_Type* UID = InputToCurveMappingUIDs.Find(Iter.Value());
 					if (UID)
 					{
@@ -244,7 +245,7 @@ void FAnimNode_ControlRig::UpdateOutput(UControlRig* InControlRig, FPoseContext&
 				}
 				else
 				{
-					UE_LOG(LogAnimation, Warning, TEXT("[%s] Missing Output Property [%s]"), *GetNameSafe(ControlRig->GetClass()), *SourcePath.ToString());
+					UE_LOG(LogAnimation, Warning, TEXT("[%s] Missing Output Parameter [%s]"), *GetNameSafe(ControlRig->GetClass()), *SourcePath.ToString());
 				}
 			}
 		}
@@ -264,7 +265,9 @@ void FAnimNode_ControlRig::SetIOMapping(bool bInput, const FName& SourceProperty
 			TMap<FName, FName>& MappingData = (bInput) ? InputMapping : OutputMapping;
 
 			// if it's valid as of now, we add it
-			if (CDO->IsValidIOVariables(bInput, SourceProperty))
+			ERigVMParameterType ParameterType = CDO->GetVM()->GetParameterByName(SourceProperty).GetType();
+			if ((bInput && (ParameterType == ERigVMParameterType::Input)) ||
+				(!bInput && (ParameterType == ERigVMParameterType::Output)))
 			{
 				if (TargetCurve == NAME_None)
 				{
@@ -292,6 +295,77 @@ FName FAnimNode_ControlRig::GetIOMapping(bool bInput, const FName& SourcePropert
 	return NAME_None;
 }
 
+void FAnimNode_ControlRig::InitializeProperties(const UObject* InSourceInstance, UClass* InTargetClass)
+{
+	// Build property lists
+	SourceProperties.Reset(SourcePropertyNames.Num());
+	DestProperties.Reset(SourcePropertyNames.Num());
+
+	check(SourcePropertyNames.Num() == DestPropertyNames.Num());
+
+	for (int32 Idx = 0; Idx < SourcePropertyNames.Num(); ++Idx)
+	{
+		FName& SourceName = SourcePropertyNames[Idx];
+		UClass* SourceClass = InSourceInstance->GetClass();
+
+		FProperty* SourceProperty = FindField<FProperty>(SourceClass, SourceName);
+		SourceProperties.Add(SourceProperty);
+		DestProperties.Add(nullptr);
+	}
+}
+
+void FAnimNode_ControlRig::PropagateInputProperties(const UObject* InSourceInstance)
+{
+	if (TargetInstance)
+	{
+		UControlRig* TargetControlRig = Cast<UControlRig>((UObject*)TargetInstance);
+
+		// First copy properties
+		check(SourceProperties.Num() == DestProperties.Num());
+		for (int32 PropIdx = 0; PropIdx < SourceProperties.Num(); ++PropIdx)
+		{
+			FProperty* CallerProperty = SourceProperties[PropIdx];
+
+			FRigVMParameter DestParameter = TargetControlRig->GetVM()->GetParameterByName(DestPropertyNames[PropIdx]);
+			if (DestParameter.GetType() != ERigVMParameterType::Input)
+			{
+				continue;
+			}
+
+			const uint8* SrcPtr = CallerProperty->ContainerPtrToValuePtr<uint8>(InSourceInstance);
+
+			if (CastField<FBoolProperty>(CallerProperty) != nullptr && DestParameter.GetCPPType() == TEXT("bool"))
+			{
+				const bool Value = *(const bool*)SrcPtr;
+				TargetControlRig->GetVM()->SetParameterValue<bool>(DestParameter, Value);
+			}
+			else if (CastField<FFloatProperty>(CallerProperty) != nullptr && DestParameter.GetCPPType() == TEXT("float"))
+			{
+				const float Value = *(const float*)SrcPtr;
+				TargetControlRig->GetVM()->SetParameterValue<float>(DestParameter, Value);
+			}
+			else if (CastField<FIntProperty>(CallerProperty) != nullptr && DestParameter.GetCPPType() == TEXT("int32"))
+			{
+				const int32 Value = *(const int32*)SrcPtr;
+				TargetControlRig->GetVM()->SetParameterValue<int32>(DestParameter, Value);
+			}
+			else if (CastField<FNameProperty>(CallerProperty) != nullptr && DestParameter.GetCPPType() == TEXT("FName"))
+			{
+				const FName Value = *(const FName*)SrcPtr;
+				TargetControlRig->GetVM()->SetParameterValue<FName>(DestParameter, Value);
+			}
+			else if (FStructProperty* StructProperty = CastField<FStructProperty>(CallerProperty))
+			{
+				if (StructProperty->Struct == DestParameter.GetScriptStruct())
+				{
+					uint8* DestPtr = TargetControlRig->GetVM()->WorkMemory.GetData(DestParameter.GetRegisterIndex());
+					StructProperty->Struct->CopyScriptStruct(DestPtr, SrcPtr, 1);
+				}
+			}
+		}
+	}
+}
+
 #if WITH_EDITOR
 void FAnimNode_ControlRig::OnObjectsReplaced(const TMap<UObject*, UObject*>& OldToNewInstanceMap)
 {
diff --git a/Engine/Plugins/Experimental/ControlRig/Source/ControlRig/Private/AnimNode_ControlRigBase.cpp b/Engine/Plugins/Experimental/ControlRig/Source/ControlRig/Private/AnimNode_ControlRigBase.cpp
--- a/Engine/Plugins/Experimental/ControlRig/Source/ControlRig/Private/AnimNode_ControlRigBase.cpp
+++ b/Engine/Plugins/Experimental/ControlRig/Source/ControlRig/Private/AnimNode_ControlRigBase.cpp
@@ -9,7 +9,8 @@
 
 FAnimNode_ControlRigBase::FAnimNode_ControlRigBase()
 	: FAnimNode_CustomProperty()
-	, bUpdateInput(true)
+	, InputSettings(FControlRigIOSettings())
+	, OutputSettings(FControlRigIOSettings())
 	, bExecute(true)
 	, InternalBlendAlpha (1.f)
 {
@@ -33,6 +34,8 @@ void FAnimNode_ControlRigBase::OnInitializeAnimInstance(const FAnimInstanceProxy
 			// node mapping container will be saved on the initialization part
 			NodeMappingContainer = Component->SkeletalMesh->GetNodeMappingContainer(Blueprint);
 		}
+		// register skeletalmesh component for now
+		ControlRig->GetDataSourceRegistry()->RegisterDataSource(UControlRig::OwnerComponent, InAnimInstance->GetOwningComponent());
 	}
 }
 
@@ -45,8 +48,7 @@ void FAnimNode_ControlRigBase::Initialize_AnyThread(const FAnimationInitializeCo
 
 	if (UControlRig* ControlRig = GetControlRig())
 	{
-		ControlRig->Initialize();
-
+		//Don't Inititialize the Control Rig here it may have the wrong VM on the CDO
 		SetTargetInstance(ControlRig);
 	}
 }
@@ -79,36 +81,42 @@ void FAnimNode_ControlRigBase::UpdateInput(UControlRig* ControlRig, const FPoseC
 {
 	DECLARE_SCOPE_HIERARCHICAL_COUNTER_FUNC()
 
-	const FBoneContainer& RequiredBones = InOutput.Pose.GetBoneContainer();
-
-	// get component pose from control rig
-	FCSPose<FCompactPose> MeshPoses;
-	// first I need to convert to local pose
-	MeshPoses.InitPose(InOutput.Pose);
-
-	// @re-think - now control rig contains init pose from their default hierarchy and current pose from this instance.
-	// we may need this init pose somewhere (instance refpose)
-	for (auto Iter = ControlRigBoneMapping.CreateConstIterator(); Iter; ++Iter)
+	if (InputSettings.bUpdatePose)
 	{
-		const FName& Name = Iter.Key();
-		const uint16 Index = Iter.Value();
+		const FBoneContainer& RequiredBones = InOutput.Pose.GetBoneContainer();
 
-		FTransform ComponentTransform = MeshPoses.GetComponentSpaceTransform(FCompactPoseBoneIndex(Index));
-		if (NodeMappingContainer.IsValid())
+		// get component pose from control rig
+		FCSPose<FCompactPose> MeshPoses;
+		// first I need to convert to local pose
+		MeshPoses.InitPose(InOutput.Pose);
+
+		// @re-think - now control rig contains init pose from their default hierarchy and current pose from this instance.
+		// we may need this init pose somewhere (instance refpose)
+		for (auto Iter = ControlRigBoneMapping.CreateConstIterator(); Iter; ++Iter)
 		{
-			ComponentTransform = NodeMappingContainer->GetSourceToTargetTransform(Name).GetRelativeTransformReverse(ComponentTransform);
-		}
+			const FName& Name = Iter.Key();
+			const uint16 Index = Iter.Value();
 
-		ControlRig->SetGlobalTransform(Name, ComponentTransform, false);
+			FTransform ComponentTransform = MeshPoses.GetComponentSpaceTransform(FCompactPoseBoneIndex(Index));
+			if (NodeMappingContainer.IsValid())
+			{
+				ComponentTransform = NodeMappingContainer->GetSourceToTargetTransform(Name).GetRelativeTransformReverse(ComponentTransform);
+			}
+
+			ControlRig->SetGlobalTransform(Name, ComponentTransform, false);
+		}
 	}
 
-	// we just do name mapping 
-	for (auto Iter = ControlRigCurveMapping.CreateConstIterator(); Iter; ++Iter)
+	if (InputSettings.bUpdateCurves)
 	{
-		const FName& Name = Iter.Key();
-		const uint16 Index = Iter.Value();
+		// we just do name mapping 
+		for (auto Iter = ControlRigCurveMapping.CreateConstIterator(); Iter; ++Iter)
+		{
+			const FName& Name = Iter.Key();
+			const uint16 Index = Iter.Value();
 
-		ControlRig->SetCurveValue(Name, InOutput.Curve.Get(Index));
+			ControlRig->SetCurveValue(Name, InOutput.Curve.Get(Index));
+		}
 	}
 }
 
@@ -116,38 +124,44 @@ void FAnimNode_ControlRigBase::UpdateOutput(UControlRig* ControlRig, FPoseContex
 {
 	DECLARE_SCOPE_HIERARCHICAL_COUNTER_FUNC()
 
-	// copy output of the rig
-	const FBoneContainer& RequiredBones = InOutput.Pose.GetBoneContainer();
-
-	// get component pose from control rig
-	FCSPose<FCompactPose> MeshPoses;
-	MeshPoses.InitPose(InOutput.Pose);
-
-	for (auto Iter = ControlRigBoneMapping.CreateConstIterator(); Iter; ++Iter)
+	if (OutputSettings.bUpdatePose)
 	{
-		const FName& Name = Iter.Key();
-		const uint16 Index = Iter.Value();
+		// copy output of the rig
+		const FBoneContainer& RequiredBones = InOutput.Pose.GetBoneContainer();
 
-		FCompactPoseBoneIndex CompactPoseIndex(Index);
-		FTransform ComponentTransform = ControlRig->GetGlobalTransform(Name);
-		if (NodeMappingContainer.IsValid())
+		// get component pose from control rig
+		FCSPose<FCompactPose> MeshPoses;
+		MeshPoses.InitPose(InOutput.Pose);
+
+		for (auto Iter = ControlRigBoneMapping.CreateConstIterator(); Iter; ++Iter)
 		{
-			ComponentTransform = NodeMappingContainer->GetSourceToTargetTransform(Name) * ComponentTransform;
+			const FName& Name = Iter.Key();
+			const uint16 Index = Iter.Value();
+
+			FCompactPoseBoneIndex CompactPoseIndex(Index);
+			FTransform ComponentTransform = ControlRig->GetGlobalTransform(Name);
+			if (NodeMappingContainer.IsValid())
+			{
+				ComponentTransform = NodeMappingContainer->GetSourceToTargetTransform(Name) * ComponentTransform;
+			}
+
+			MeshPoses.SetComponentSpaceTransform(CompactPoseIndex, ComponentTransform);
 		}
 
-		MeshPoses.SetComponentSpaceTransform(CompactPoseIndex, ComponentTransform);
+		FCSPose<FCompactPose>::ConvertComponentPosesToLocalPoses(MeshPoses, InOutput.Pose);
 	}
 
-	FCSPose<FCompactPose>::ConvertComponentPosesToLocalPoses(MeshPoses, InOutput.Pose);
-
-	// update curve
-	for (auto Iter = ControlRigCurveMapping.CreateConstIterator(); Iter; ++Iter)
+	if (OutputSettings.bUpdateCurves)
 	{
-		const FName& Name = Iter.Key();
-		const uint16 Index = Iter.Value();
+		// update curve
+		for (auto Iter = ControlRigCurveMapping.CreateConstIterator(); Iter; ++Iter)
+		{
+			const FName& Name = Iter.Key();
+			const uint16 Index = Iter.Value();
 
-		const float Value = ControlRig->GetCurveValue(Name);
-		InOutput.Curve.Set(Index, Value);
+			const float Value = ControlRig->GetCurveValue(Name);
+			InOutput.Curve.Set(Index, Value);
+		}
 	}
 }
 
@@ -200,11 +214,8 @@ void FAnimNode_ControlRigBase::ExecuteControlRig(FPoseContext& InOutput)
 {
 	if (UControlRig* ControlRig = GetControlRig())
 	{
-		if (bUpdateInput)
-		{
-			// first update input to the system
-			UpdateInput(ControlRig, InOutput);
-		}
+		// first update input to the system
+		UpdateInput(ControlRig, InOutput);
 
 		if (bExecute)
 		{
diff --git a/Engine/Plugins/Experimental/ControlRig/Source/ControlRig/Private/ControlRig.cpp b/Engine/Plugins/Experimental/ControlRig/Source/ControlRig/Private/ControlRig.cpp
--- a/Engine/Plugins/Experimental/ControlRig/Source/ControlRig/Private/ControlRig.cpp
+++ b/Engine/Plugins/Experimental/ControlRig/Source/ControlRig/Private/ControlRig.cpp
@@ -3,15 +3,14 @@
 #include "ControlRig.h"
 #include "GameFramework/Actor.h"
 #include "Misc/RuntimeErrors.h"
-#include "ControlRigVM.h"
 #include "IControlRigObjectBinding.h"
 #include "HelperUtil.h"
 #include "ControlRigBlueprintGeneratedClass.h"
 #include "ControlRigObjectVersion.h"
-#include "ControlRigVariables.h"
 #if WITH_EDITOR
 #include "ControlRigModule.h"
 #include "Modules/ModuleManager.h"
+#include "RigVMModel/RigVMNode.h"
 #endif// WITH_EDITOR
 
 #define LOCTEXT_NAMESPACE "ControlRig"
@@ -30,108 +29,45 @@ const FName UControlRig::CategoryMetaName("Category");
 const FName UControlRig::DisplayNameMetaName("DisplayName");
 const FName UControlRig::MenuDescSuffixMetaName("MenuDescSuffix");
 const FName UControlRig::ShowVariableNameInTitleMetaName("ShowVariableNameInTitle");
-const FName UControlRig::BoneNameMetaName("BoneName");
-const FName UControlRig::ControlNameMetaName("ControlName");
-const FName UControlRig::SpaceNameMetaName("SpaceName");
-const FName UControlRig::CurveNameMetaName("CurveName");
+const FName UControlRig::CustomWidgetMetaName("CustomWidget");
 const FName UControlRig::ConstantMetaName("Constant");
 const FName UControlRig::TitleColorMetaName("TitleColor");
 const FName UControlRig::NodeColorMetaName("NodeColor");
 const FName UControlRig::KeywordsMetaName("Keywords");
 const FName UControlRig::PrototypeNameMetaName("PrototypeName");
 
-const FName UControlRig::AnimationInputMetaName("AnimationInput");
-const FName UControlRig::AnimationOutputMetaName("AnimationOutput");
-
 const FName UControlRig::ExpandPinByDefaultMetaName("ExpandByDefault");
 const FName UControlRig::DefaultArraySizeMetaName("DefaultArraySize");
 
-UControlRig::UControlRig()
-	: DeltaTime(0.0f)
-	, ExecutionType(ERigExecutionType::Runtime)
+const FName UControlRig::OwnerComponent("OwnerComponent");
+
+//CVar to specify if we should create a float control for each curve in the curve container
+//By default we don't but it may be useful to do so for debugging
+static TAutoConsoleVariable<int32> CVarControlRigCreateFloatControlsForCurves(
+	TEXT("ControlRig.CreateFloatControlsForCurves"),
+	0,
+	TEXT("If nonzero we create a float control for each curve in the curve container, useful for debugging low level controls."),
+	ECVF_Default);
+
+UControlRig::UControlRig(const FObjectInitializer& ObjectInitializer)
+	: Super(ObjectInitializer)
+	, DeltaTime(0.0f)
 #if WITH_EDITOR
 	, ControlRigLog(nullptr)
 	, bEnableControlRigLogging(true)
 	, DrawInterface(nullptr)
 	, DataSourceRegistry(nullptr)
+	, PreviewInstance(nullptr)
+	, bRequiresInitExecution(false)
 #endif
 {
-#if DEBUG_CONTROLRIG_PROPERTYCHANGE
-	CacheDebugClassData();
-#endif // #if DEBUG_CONTROLRIG_PROPERTYCHANGE
-}
-
-#if DEBUG_CONTROLRIG_PROPERTYCHANGE
-void UControlRig::CacheDebugClassData()
-{
-	DECLARE_SCOPE_HIERARCHICAL_COUNTER_FUNC()
-
-	// this can be debug only
-	UControlRigBlueprintGeneratedClass* CurrentClass = Cast<UControlRigBlueprintGeneratedClass>(GetClass());
-	if (CurrentClass)
-	{
-		DebugClassSize = CurrentClass->PropertiesSize;
-
-		Destructors.Reset();
-		for (FProperty* P = CurrentClass->DestructorLink; P; P = P->DestructorLinkNext)
-		{
-			FStructProperty* StructProperty = CastField<FStructProperty>(P);
-			if (StructProperty)
-			{
-				Destructors.Add(StructProperty->Struct);
-			}
-		}
-
-		PropertyData.Reset();
-		for (FProperty* P = CurrentClass->PropertyLink; P; P = P->PropertyLinkNext)
-		{
-			FPropertyData PropData;
-			PropData.Offset = P->GetOffset_ForDebug();
-			PropData.Size = P->GetSize();
-			PropData.PropertyName = P->GetFName();
-
-			PropertyData.Add(PropData);
-		}
-
-		ensure(CurrentClass->UberGraphFunction == nullptr);
-	}
+	VM = ObjectInitializer.CreateDefaultSubobject<URigVM>(this, TEXT("VM"));
+	// create default source registry
+	DataSourceRegistry = CreateDefaultSubobject<UAnimationDataSourceRegistry>(TEXT("DataSourceRegistry"));
 }
 
-void UControlRig::ValidateDebugClassData()
-{
-	DECLARE_SCOPE_HIERARCHICAL_COUNTER_FUNC()
-
-	UControlRigBlueprintGeneratedClass* CurrentClass = Cast<UControlRigBlueprintGeneratedClass>(GetClass());
-	if (CurrentClass)
-	{
-		ensureAlwaysMsgf(DebugClassSize == CurrentClass->PropertiesSize,
-			TEXT("Class [%s] size has changed : used be [size %d], and current class is [size %d]"), *GetNameSafe(CurrentClass), DebugClassSize, CurrentClass->GetStructureSize());
-
-		int32 Index = 0;
-		for (FProperty* P = CurrentClass->DestructorLink; P; P = P->DestructorLinkNext)
-		{
-			FStructProperty* StructProperty = CastField<FStructProperty>(P);
-			if (StructProperty)
-			{
-				ensureAlways(Destructors[Index++] == (StructProperty->Struct));
-			}
-		}
-
-		Index = 0;
-		for (FProperty* P = CurrentClass->PropertyLink; P; P = P->PropertyLinkNext)
-		{
-			const FPropertyData& PropData = PropertyData[Index++];
-			ensureAlwaysMsgf(PropData.Offset == P->GetOffset_ForDebug() && PropData.Size == P->GetSize(), TEXT("Property (%s) size changes"), *PropData.PropertyName.ToString());
-		}
-	}
-}
-#endif // #if DEBUG_CONTROLRIG_PROPERTYCHANGE
-
 void UControlRig::BeginDestroy()
 {
-#if DEBUG_CONTROLRIG_PROPERTYCHANGE
-	ValidateDebugClassData(); 
-#endif // #if DEBUG_CONTROLRIG_PROPERTYCHANGE
 	Super::BeginDestroy();
 	InitializedEvent.Clear();
 	ExecutedEvent.Clear();
@@ -139,7 +75,7 @@ void UControlRig::BeginDestroy()
 
 UWorld* UControlRig::GetWorld() const
 {
-	if(ObjectBinding.IsValid())
+	if (ObjectBinding.IsValid())
 	{
 		AActor* HostingActor = ObjectBinding->GetHostingActor();
 		if (HostingActor)
@@ -175,44 +111,18 @@ void UControlRig::Initialize(bool bInitRigUnits)
 	}
 
 	InitializeFromCDO();
-
-	InstantiateOperatorsFromGeneratedClass();
-	ResolvePropertyPaths();
-
-#if WITH_EDITORONLY_DATA
-	// initialize rig unit cached names
-	UControlRigBlueprintGeneratedClass* GeneratedClass = Cast<UControlRigBlueprintGeneratedClass>(GetClass());
-	if (GeneratedClass)
-	{
-		for (FStructProperty* UnitProperty : GeneratedClass->RigUnitProperties)
-		{
-			FRigUnit* RigUnit = UnitProperty->ContainerPtrToValuePtr<FRigUnit>(this);
-			RigUnit->RigUnitName = UnitProperty->GetFName();
-			RigUnit->RigUnitStructName = UnitProperty->Struct->GetFName();
-#if 0 
-			FRigUnit* RigUnitDefault = UnitProperty->ContainerPtrToValuePtr<FRigUnit>(Class->GetDefaultObject());
-			ensure(RigUnitDefault != nullptr);
-#endif // DEBUG only
-		}
-	}
-#endif // WITH_EDITORONLY_DATA
+	InstantiateVMFromCDO();
 
 	// should refresh mapping 
 	Hierarchy.Initialize();
 
-	// resolve IO properties
-	ResolveInputOutputProperties();
-
 	if (bInitRigUnits)
 	{
-		// execute rig units with init state
-		Execute(EControlRigState::Init);
+		RequestInit();
 	}
 
-	// cache requested inputs
-	// go through find requested inputs
-
 #if WITH_EDITOR
+	Hierarchy.ControlHierarchy.OnControlSelected.RemoveAll(this);
 	Hierarchy.ControlHierarchy.OnControlSelected.AddUObject(this, &UControlRig::HandleOnControlSelected);
 #endif
 }
@@ -222,9 +132,17 @@ void UControlRig::InitializeFromCDO()
 	DECLARE_SCOPE_HIERARCHICAL_COUNTER_FUNC()
 
 	// copy CDO property you need to here
-	UControlRig* CDO = GetClass()->GetDefaultObject<UControlRig>();
-	// copy operation
-	Hierarchy = CDO->Hierarchy;
+	if (!HasAnyFlags(RF_ClassDefaultObject))
+	{
+		UControlRig* CDO = GetClass()->GetDefaultObject<UControlRig>();
+
+		// copy hierarchy
+		Hierarchy = CDO->Hierarchy;
+		Hierarchy.Initialize();
+
+		// copy draw container
+		DrawContainer = CDO->DrawContainer;
+	}
 }
 
 void UControlRig::Evaluate_AnyThread()
@@ -240,7 +158,7 @@ FText UControlRig::GetCategory() const
 	return LOCTEXT("DefaultControlRigCategory", "Animation|ControlRigs");
 }
 
-FText UControlRig::GetTooltipText() const
+FText UControlRig::GetToolTipText() const
 {
 	return LOCTEXT("DefaultControlRigTooltip", "ControlRig");
 }
@@ -251,49 +169,43 @@ void UControlRig::SetDeltaTime(float InDeltaTime)
 	DeltaTime = InDeltaTime;
 }
 
-float UControlRig::GetDeltaTime() const
-{
-	return DeltaTime;
-}
-
-void UControlRig::InstantiateOperatorsFromGeneratedClass()
+void UControlRig::InstantiateVMFromCDO()
 {
 	DECLARE_SCOPE_HIERARCHICAL_COUNTER_FUNC()
 
-	UControlRigBlueprintGeneratedClass* GeneratedClass = Cast<UControlRigBlueprintGeneratedClass>(GetClass());
-	if (GeneratedClass)
+	if (VM == nullptr || VM->GetOuter() != this)
 	{
-		Operators.Reset(GeneratedClass->Operators.Num());
-		for (const FControlRigOperator& Op : GeneratedClass->Operators)
-		{
-			Operators.Add(FControlRigOperator::MakeUnresolvedCopy(Op));
-		}
+		VM = NewObject<URigVM>(this, TEXT("VM"));
 	}
-	else
-	{
-		Operators.Reset();
-	}
-}
 
-void UControlRig::ResolvePropertyPaths()
-{
-	DECLARE_SCOPE_HIERARCHICAL_COUNTER_FUNC()
-
-	for (int32 Index = 0; Index < Operators.Num(); ++Index)
+	if (!HasAnyFlags(RF_ClassDefaultObject))
 	{
-		if (!Operators[Index].Resolve(this))
+		UControlRig* CDO = GetClass()->GetDefaultObject<UControlRig>();
+		if (VM && CDO && CDO->VM)
 		{
-			UE_LOG(LogControlRig, Error, TEXT("Operator '%s' cannot be resolved."), *Operators[Index].ToString());
+			VM->CopyFrom(CDO->VM);
+		}
+		else
+		{
+			VM->Reset();
 		}
 	}
+	RequestInit();
 }
 
 void UControlRig::Execute(const EControlRigState InState)
 {
 	DECLARE_SCOPE_HIERARCHICAL_COUNTER_FUNC()
 
+	if (bRequiresInitExecution)
+	{
+		bRequiresInitExecution = false;
+		Execute(EControlRigState::Init);
+	}
+
 	FRigUnitContext Context;
 	Context.DrawInterface = DrawInterface;
+	Context.DrawContainer = &DrawContainer;
 
 	if (InState == EControlRigState::Init)
 	{
@@ -312,41 +224,42 @@ void UControlRig::Execute(const EControlRigState InState)
 	}
 #endif
 
-	ControlRigVM::Execute(this, Context, Operators, ExecutionType);
+	// execute units
+	ExecuteUnits(Context);
 
 #if WITH_EDITOR
 	if (ControlRigLog != nullptr && bEnableControlRigLogging)
 	{
 		for (const FControlRigLog::FLogEntry& Entry : ControlRigLog->Entries)
 		{
-			if (Entry.Unit == NAME_None || Entry.Message.IsEmpty())
+			if (Entry.OperatorName == NAME_None || Entry.InstructionIndex == INDEX_NONE || Entry.Message.IsEmpty())
 			{
 				continue;
 			}
 
 			switch (Entry.Severity)
 			{
-				case EMessageSeverity::CriticalError:
-				case EMessageSeverity::Error:
-				{
-					UE_LOG(LogControlRig, Error, TEXT("Unit '%s': '%s'"), *Entry.Unit.ToString(), *Entry.Message);
-					break;
-				}
-				case EMessageSeverity::PerformanceWarning:
-				case EMessageSeverity::Warning:
-				{
-					UE_LOG(LogControlRig, Warning, TEXT("Unit '%s': '%s'"), *Entry.Unit.ToString(), *Entry.Message);
-					break;
-				}
-				case EMessageSeverity::Info:
-				{
-					UE_LOG(LogControlRig, Display, TEXT("Unit '%s': '%s'"), *Entry.Unit.ToString(), *Entry.Message);
-					break;
-				}
-				default:
-				{
-					break;
-				}
+			case EMessageSeverity::CriticalError:
+			case EMessageSeverity::Error:
+			{
+				UE_LOG(LogControlRig, Error, TEXT("Operator[%d] '%s': '%s'"), Entry.InstructionIndex, *Entry.OperatorName.ToString(), *Entry.Message);
+				break;
+			}
+			case EMessageSeverity::PerformanceWarning:
+			case EMessageSeverity::Warning:
+			{
+				UE_LOG(LogControlRig, Warning, TEXT("Operator[%d] '%s': '%s'"), Entry.InstructionIndex, *Entry.OperatorName.ToString(), *Entry.Message);
+				break;
+			}
+			case EMessageSeverity::Info:
+			{
+				UE_LOG(LogControlRig, Display, TEXT("Operator[%d] '%s': '%s'"), Entry.InstructionIndex, *Entry.OperatorName.ToString(), *Entry.Message);
+				break;
+			}
+			default:
+			{
+				break;
+			}
 			}
 		}
 	}
@@ -361,11 +274,190 @@ void UControlRig::Execute(const EControlRigState InState)
 	}
 	else if (InState == EControlRigState::Update)
 	{
+		DeltaTime = 0.f;
+
 		if (ExecutedEvent.IsBound())
 		{
 			ExecutedEvent.Broadcast(this, EControlRigState::Update);
 		}
 	}
+
+	if (Context.DrawInterface && Context.DrawContainer)
+	{
+		Context.DrawInterface->Instructions.Append(Context.DrawContainer->Instructions);
+
+		for (const FRigControl& Control : Hierarchy.ControlHierarchy)
+		{
+			if (Control.bGizmoEnabled && !Control.bIsTransientControl && Control.bDrawLimits && (Control.bLimitTranslation || Control.bLimitRotation || Control.bLimitScale))
+			{
+				// for now we don't draw rotational limits
+				if(!Control.bLimitTranslation)
+				{
+					continue;
+				}
+
+				FTransform Transform = Hierarchy.ControlHierarchy.GetParentTransform(Control.Index);
+				FControlRigDrawInstruction Instruction(EControlRigDrawSettings::Lines, Control.GizmoColor, 0.f, Transform);
+
+				switch (Control.ControlType)
+				{
+					case ERigControlType::Float:
+					{
+						FVector MinPos = FVector::ZeroVector;
+						FVector MaxPos = FVector::ZeroVector;
+
+						switch (Control.PrimaryAxis)
+						{
+							case ERigControlAxis::X:
+							{
+								MinPos.X = Control.MinimumValue.Get<float>();
+								MaxPos.X = Control.MaximumValue.Get<float>();
+								break;
+							}
+							case ERigControlAxis::Y:
+							{
+								MinPos.Y = Control.MinimumValue.Get<float>();
+								MaxPos.Y = Control.MaximumValue.Get<float>();
+								break;
+							}
+							case ERigControlAxis::Z:
+							{
+								MinPos.Z = Control.MinimumValue.Get<float>();
+								MaxPos.Z = Control.MaximumValue.Get<float>();
+								break;
+							}
+						}
+
+						Instruction.Positions.Add(MinPos);
+						Instruction.Positions.Add(MaxPos);
+						break;
+					}
+					case ERigControlType::Vector2D:
+					{
+						Instruction.PrimitiveType = EControlRigDrawSettings::LineStrip;
+						FVector2D MinPos = Control.MinimumValue.Get<FVector2D>();
+						FVector2D MaxPos = Control.MaximumValue.Get<FVector2D>();
+
+						switch (Control.PrimaryAxis)
+						{
+							case ERigControlAxis::X:
+							{
+								Instruction.Positions.Add(FVector(0.f, MinPos.X, MinPos.Y));
+								Instruction.Positions.Add(FVector(0.f, MaxPos.X, MinPos.Y));
+								Instruction.Positions.Add(FVector(0.f, MaxPos.X, MaxPos.Y));
+								Instruction.Positions.Add(FVector(0.f, MinPos.X, MaxPos.Y));
+								Instruction.Positions.Add(FVector(0.f, MinPos.X, MinPos.Y));
+								break;
+							}
+							case ERigControlAxis::Y:
+							{
+								Instruction.Positions.Add(FVector(MinPos.X, 0.f, MinPos.Y));
+								Instruction.Positions.Add(FVector(MaxPos.X, 0.f, MinPos.Y));
+								Instruction.Positions.Add(FVector(MaxPos.X, 0.f, MaxPos.Y));
+								Instruction.Positions.Add(FVector(MinPos.X, 0.f, MaxPos.Y));
+								Instruction.Positions.Add(FVector(MinPos.X, 0.f, MinPos.Y));
+								break;
+							}
+							case ERigControlAxis::Z:
+							{
+								Instruction.Positions.Add(FVector(MinPos.X, MinPos.Y, 0.f));
+								Instruction.Positions.Add(FVector(MaxPos.X, MinPos.Y, 0.f));
+								Instruction.Positions.Add(FVector(MaxPos.X, MaxPos.Y, 0.f));
+								Instruction.Positions.Add(FVector(MinPos.X, MaxPos.Y, 0.f));
+								Instruction.Positions.Add(FVector(MinPos.X, MinPos.Y, 0.f));
+								break;
+							}
+						}
+						break;
+					}
+					case ERigControlType::Position:
+					case ERigControlType::Scale:
+					case ERigControlType::Transform:
+					case ERigControlType::TransformNoScale:
+					{
+						FVector MinPos, MaxPos;
+
+						switch (Control.ControlType)
+						{
+							case ERigControlType::Position:
+							case ERigControlType::Scale:
+							{
+								MinPos = Control.MinimumValue.Get<FVector>();
+								MaxPos = Control.MaximumValue.Get<FVector>();
+								break;
+							}
+							case ERigControlType::Transform:
+							{
+								MinPos = Control.MinimumValue.Get<FTransform>().GetLocation();
+								MaxPos = Control.MaximumValue.Get<FTransform>().GetLocation();
+								break;
+							}
+							case ERigControlType::TransformNoScale:
+							{
+								MinPos = Control.MinimumValue.Get<FTransformNoScale>().Location;
+								MaxPos = Control.MaximumValue.Get<FTransformNoScale>().Location;
+								break;
+							}
+						}
+
+						Instruction.Positions.Add(FVector(MinPos.X, MinPos.Y, MinPos.Z));
+						Instruction.Positions.Add(FVector(MaxPos.X, MinPos.Y, MinPos.Z));
+						Instruction.Positions.Add(FVector(MinPos.X, MaxPos.Y, MinPos.Z));
+						Instruction.Positions.Add(FVector(MaxPos.X, MaxPos.Y, MinPos.Z));
+						Instruction.Positions.Add(FVector(MinPos.X, MinPos.Y, MaxPos.Z));
+						Instruction.Positions.Add(FVector(MaxPos.X, MinPos.Y, MaxPos.Z));
+						Instruction.Positions.Add(FVector(MinPos.X, MaxPos.Y, MaxPos.Z));
+						Instruction.Positions.Add(FVector(MaxPos.X, MaxPos.Y, MaxPos.Z));
+
+						Instruction.Positions.Add(FVector(MinPos.X, MinPos.Y, MinPos.Z));
+						Instruction.Positions.Add(FVector(MinPos.X, MaxPos.Y, MinPos.Z));
+						Instruction.Positions.Add(FVector(MaxPos.X, MinPos.Y, MinPos.Z));
+						Instruction.Positions.Add(FVector(MaxPos.X, MaxPos.Y, MinPos.Z));
+						Instruction.Positions.Add(FVector(MinPos.X, MinPos.Y, MaxPos.Z));
+						Instruction.Positions.Add(FVector(MinPos.X, MaxPos.Y, MaxPos.Z));
+						Instruction.Positions.Add(FVector(MaxPos.X, MinPos.Y, MaxPos.Z));
+						Instruction.Positions.Add(FVector(MaxPos.X, MaxPos.Y, MaxPos.Z));
+
+						Instruction.Positions.Add(FVector(MinPos.X, MinPos.Y, MinPos.Z));
+						Instruction.Positions.Add(FVector(MinPos.X, MinPos.Y, MaxPos.Z));
+						Instruction.Positions.Add(FVector(MaxPos.X, MinPos.Y, MinPos.Z));
+						Instruction.Positions.Add(FVector(MaxPos.X, MinPos.Y, MaxPos.Z));
+						Instruction.Positions.Add(FVector(MinPos.X, MaxPos.Y, MinPos.Z));
+						Instruction.Positions.Add(FVector(MinPos.X, MaxPos.Y, MaxPos.Z));
+						Instruction.Positions.Add(FVector(MaxPos.X, MaxPos.Y, MinPos.Z));
+						Instruction.Positions.Add(FVector(MaxPos.X, MaxPos.Y, MaxPos.Z));
+						break;
+					}
+				}
+
+				if (Instruction.Positions.Num() > 0)
+				{
+					DrawInterface->Instructions.Add(Instruction);
+				}
+			}
+		}
+	}
+}
+
+void UControlRig::ExecuteUnits(FRigUnitContext& InOutContext)
+{
+	if (VM)
+	{
+		FRigVMMemoryContainer* LocalMemory[] = { &VM->WorkMemory, &VM->LiteralMemory };
+		TArray<void*> AdditionalArguments;
+		AdditionalArguments.Add(&InOutContext);
+		VM->Execute(FRigVMMemoryContainerPtrArray(LocalMemory, 2), AdditionalArguments);
+	}
+}
+
+URigVM* UControlRig::GetVM()
+{
+	if (VM == nullptr)
+	{
+		Initialize(true);
+		check(VM);
+	}
+	return VM;
 }
 
 FTransform UControlRig::GetGlobalTransform(const FName& BoneName) const
@@ -436,98 +528,7 @@ void UControlRig::GetMappableNodeData(TArray<FName>& OutNames, TArray<FNodeItem>
 	}
 }
 
-void UControlRig::ResolveInputOutputProperties()
-{
-	DECLARE_SCOPE_HIERARCHICAL_COUNTER_FUNC()
-
-	for (auto Iter = InputProperties.CreateIterator(); Iter; ++Iter)
-	{
-		FCachedPropertyPath& PropPath = Iter.Value();
-		PropPath.Resolve(this);
-	}
-
-	for (auto Iter = OutputProperties.CreateIterator(); Iter; ++Iter)
-	{
-		FCachedPropertyPath& PropPath = Iter.Value();
-		PropPath.Resolve(this);
-	}
-}
-
-bool UControlRig::GetInOutPropertyPath(bool bInput, const FName& InPropertyPath, FCachedPropertyPath& OutCachedPath)
-{
-	DECLARE_SCOPE_HIERARCHICAL_COUNTER_FUNC()
-
-	// if root properties, use this
-	TMap<FName, FCachedPropertyPath>& Properties = (bInput) ? InputProperties: OutputProperties;
-
-	FCachedPropertyPath* CachedPath = Properties.Find(InPropertyPath);
-	if (CachedPath)
-	{
-		if (!CachedPath->IsResolved())
-		{
-			CachedPath->Resolve(this);
-		}
-
-		if (CachedPath->IsResolved())
-		{
-			OutCachedPath = *CachedPath;
-			return true;
-		}
-	}
-
-	return false;
-}
 #if WITH_EDITORONLY_DATA
-FName UControlRig::GetRigClassNameFromRigUnit(const FRigUnit* InRigUnit) const
-{
-	DECLARE_SCOPE_HIERARCHICAL_COUNTER_FUNC()
-
-	if (InRigUnit)
-	{
-		UControlRigBlueprintGeneratedClass* GeneratedClass = Cast<UControlRigBlueprintGeneratedClass>(GetClass());
-		for (FStructProperty* UnitProperty : GeneratedClass->RigUnitProperties)
-		{
-			if (UnitProperty->ContainerPtrToValuePtr<FRigUnit>(this) == InRigUnit)
-			{
-				return UnitProperty->Struct->GetFName();
-			}
-		}
-	}
-
-	return NAME_None;
-}
-
-FRigUnit_Control* UControlRig::GetControlRigUnitFromName(const FName& PropertyName) 
-{
-	DECLARE_SCOPE_HIERARCHICAL_COUNTER_FUNC()
-
-	UControlRigBlueprintGeneratedClass* GeneratedClass = Cast<UControlRigBlueprintGeneratedClass>(GetClass());
-	for (FStructProperty* ControlProperty : GeneratedClass->ControlUnitProperties)
-	{
-		if (ControlProperty->GetFName() == PropertyName)
-		{
-			return ControlProperty->ContainerPtrToValuePtr<FRigUnit_Control>(this);
-		}
-	}
-
-	return nullptr;
-}
-
-FRigUnit* UControlRig::GetRigUnitFromName(const FName& PropertyName) 
-{
-	DECLARE_SCOPE_HIERARCHICAL_COUNTER_FUNC()
-
-	UControlRigBlueprintGeneratedClass* GeneratedClass = Cast<UControlRigBlueprintGeneratedClass>(GetClass());
-	for (FStructProperty* UnitProperty : GeneratedClass->RigUnitProperties)
-	{
-		if (UnitProperty->GetFName() == PropertyName)
-		{
-			return UnitProperty->ContainerPtrToValuePtr<FRigUnit>(this);
-		}
-	}
-
-	return nullptr;
-}
 
 void UControlRig::PostReinstanceCallback(const UControlRig* Old)
 {
@@ -535,16 +536,8 @@ void UControlRig::PostReinstanceCallback(const UControlRig* Old)
 
 	ObjectBinding = Old->ObjectBinding;
 	Initialize();
-
-	// propagate values from the CDO to the instance
-	UControlRigBlueprintGeneratedClass* GeneratedClass = Cast<UControlRigBlueprintGeneratedClass>(GetClass());
-	for (FStructProperty* UnitProperty : GeneratedClass->RigUnitProperties)
-	{
-		FRigUnit* Source = UnitProperty->ContainerPtrToValuePtr<FRigUnit>(GeneratedClass->ClassDefaultObject);
-		FRigUnit* Dest = UnitProperty->ContainerPtrToValuePtr<FRigUnit>(this);
-		UnitProperty->CopyCompleteValue(Dest, Source);
-	}
 }
+
 #endif // WITH_EDITORONLY_DATA
 
 void UControlRig::AddReferencedObjects(UObject* InThis, FReferenceCollector& Collector)
@@ -552,21 +545,13 @@ void UControlRig::AddReferencedObjects(UObject* InThis, FReferenceCollector& Col
 	DECLARE_SCOPE_HIERARCHICAL_COUNTER_FUNC()
 
 	Super::AddReferencedObjects(InThis, Collector);
-#if WITH_EDITOR
-	UControlRig* This = CastChecked<UControlRig>(InThis);
-	for (auto Iter = This->RigUnitEditorObjects.CreateIterator(); Iter; ++Iter)
-	{
-		Collector.AddReferencedObject(Iter.Value());
-	}
-#endif // WITH_EDITOR
 }
+
 #if WITH_EDITOR
 //undo will clear out the transient Operators, need to recreate them
 void UControlRig::PostEditUndo()
 {
 	Super::PostEditUndo();
-	InstantiateOperatorsFromGeneratedClass();
-	ResolvePropertyPaths();
 }
 #endif // WITH_EDITOR
 
@@ -577,53 +562,6 @@ void UControlRig::Serialize(FArchive& Ar)
 	Super::Serialize(Ar);
 
 	Ar.UsingCustomVersion(FControlRigObjectVersion::GUID);
-
-	if (Ar.IsLoading())
-	{
-		if (Ar.CustomVer(FControlRigObjectVersion::GUID) < FControlRigObjectVersion::OperatorsStoringPropertyPaths)
-		{
-			// create the cached paths based on the deprecated string
-			for (FControlRigOperator& Operator : Operators)
-			{
-				if (!Operator.PropertyPath1_DEPRECATED.IsEmpty())
-				{
-					Operator.CachedPropertyPath1 = FCachedPropertyPath(Operator.PropertyPath1_DEPRECATED);
-				}
-				if (!Operator.PropertyPath2_DEPRECATED.IsEmpty())
-				{
-					Operator.CachedPropertyPath2 = FCachedPropertyPath(Operator.PropertyPath2_DEPRECATED);
-				}
-			}
-		}
-	}
-}
-
-bool UControlRig::IsValidIOVariables(bool bInput, const FName& PropertyName) const
-{
-	DECLARE_SCOPE_HIERARCHICAL_COUNTER_FUNC()
-
-	const TMap<FName, FCachedPropertyPath>& Properties = (bInput) ? InputProperties : OutputProperties;
-	return Properties.Contains(PropertyName);
-}
-
-void UControlRig::QueryIOVariables(bool bInput, TArray<FControlRigIOVariable>& OutVars) const
-{
-	DECLARE_SCOPE_HIERARCHICAL_COUNTER_FUNC()
-
-	const TMap<FName, FCachedPropertyPath>& Properties = (bInput) ? InputProperties : OutputProperties;
-	OutVars.Reset(Properties.Num());
-
-	for (auto Iter = Properties.CreateConstIterator(); Iter; ++Iter)
-	{
-		const FCachedPropertyPath& PropPath = Iter.Value();
-		FProperty* Property = PropPath.GetFProperty();
-		FControlRigIOVariable NewVar;
-		NewVar.PropertyPath = PropPath.ToString();
-		NewVar.PropertyType = FControlRigIOHelper::GetFriendlyTypeName(Property);
-		NewVar.Size = Property->GetSize();
-
-		OutVars.Add(NewVar);
-	}
 }
 
 // BEGIN IControlRigManipulatable interface
@@ -657,6 +595,12 @@ bool UControlRig::SetSpaceGlobalTransform(const FName& InSpaceName, const FTrans
 
 const TArray<FRigControl>& UControlRig::AvailableControls() const
 {
+#if WITH_EDITOR
+	if (AvailableControlsOverride.Num() > 0)
+	{
+		return AvailableControlsOverride;
+	}
+#endif
 	return Hierarchy.ControlHierarchy.GetControls();
 }
 
@@ -668,11 +612,165 @@ FRigControl* UControlRig::FindControl(const FName& InControlName)
 		return &Hierarchy.ControlHierarchy[ControlIndex];
 	}
 
+#if WITH_EDITOR
+	if (TransientControls.Num() > 0)
+	{
+		for (int32 Index = 0; Index < TransientControls.Num(); Index++)
+		{
+			if (TransientControls[Index].Name == InControlName)
+			{
+				return &TransientControls[Index];
+			}
+		}
+	}
+#endif
+
 	return nullptr;
 }
 
+void UControlRig::CreateRigControlsForCurveContainer()
+{
+	const bool bCreateFloatControls = CVarControlRigCreateFloatControlsForCurves->GetInt() == 0 ? false : true;
+	if(bCreateFloatControls)
+	{
+		FRigCurveContainer& CurveContainer = GetCurveContainer();
+		FString CTRLName(TEXT("CTRL_"));
+		for (FRigCurve& Curve : CurveContainer)
+		{
+			FString Name = Curve.Name.ToString();
+			if (Name.Contains(CTRLName) && Hierarchy.ControlHierarchy.GetIndex(Curve.Name) == INDEX_NONE)
+			{
+				FRigControlValue Value;
+				Value.Set<float>(Curve.Value);
+				FRigControl& Control = Hierarchy.ControlHierarchy.Add(Curve.Name, ERigControlType::Float, NAME_None, NAME_None, Value
+	#if WITH_EDITORONLY_DATA
+					, Curve.Name,
+					FTransform::Identity,
+					FLinearColor::Red
+	#endif
+				);
+				Control.bIsCurve = true;
+			}
+		}
+		ControlModified().AddUObject(this, &UControlRig::HandleOnControlModified);
+	}
+}
+
+void UControlRig::HandleOnControlModified(IControlRigManipulatable* Subject, const FRigControl& Control, EControlRigSetKey InSetKey)
+{
+	if (Control.bIsCurve)
+	{
+		SetCurveValue(Control.Name, Control.Value.Get<float>());
+	}	
+}
+
+bool UControlRig::IsCurveControl(const FRigControl* InControl) const
+{
+	return InControl->bIsCurve;
+}
+
 FTransform UControlRig::GetControlGlobalTransform(const FName& InControlName) const
 {
+#if WITH_EDITOR
+	if (TransientControls.Num() > 0)
+	{
+		for (int32 Index = 0; Index < TransientControls.Num(); Index++)
+		{
+			const FRigControl& Control = TransientControls[Index];
+			if (Control.Name == InControlName)
+			{
+				FTransform Transform = FTransform::Identity;
+
+				FTransform ParentTransform = FTransform::Identity;
+				if (Control.ParentIndex != INDEX_NONE)
+				{
+					ParentTransform = Hierarchy.BoneHierarchy.GetGlobalTransform(Control.ParentIndex);
+				}
+
+				switch (Control.ControlType)
+				{
+					case ERigControlType::Bool:
+					{
+						// not sure how to extract this from global transform
+						break;
+					}
+					case ERigControlType::Float:
+					{
+						float Value = Control.Value.Get<float>();
+						switch (Control.PrimaryAxis)
+						{
+							case ERigControlAxis::X:
+							{
+								Transform.SetLocation(FVector(Value, 0.f, 0.f));
+								break;
+							}
+							case ERigControlAxis::Y:
+							{
+								Transform.SetLocation(FVector(0.f, Value, 0.f));
+								break;
+							}
+							case ERigControlAxis::Z:
+							{
+								Transform.SetLocation(FVector(0.f, 0.f, Value));
+								break;
+							}
+						}
+						break;
+					}
+					case ERigControlType::Vector2D:
+					{
+						FVector2D Value = Control.Value.Get<FVector2D>();
+						switch (Control.PrimaryAxis)
+						{
+							case ERigControlAxis::X:
+							{
+								Transform.SetLocation(FVector(0.f, Value.X, Value.Y));
+								break;
+							}
+							case ERigControlAxis::Y:
+							{
+								Transform.SetLocation(FVector(Value.X, 0.f, Value.Y));
+								break;
+							}
+							case ERigControlAxis::Z:
+							{
+								Transform.SetLocation(FVector(Value.X, Value.Y, 0.f));
+								break;
+							}
+						}
+						break;
+					}
+					case ERigControlType::Position:
+					{
+						Transform.SetLocation(Control.Value.Get<FVector>());
+						break;
+					}
+					case ERigControlType::Scale:
+					{
+						Transform.SetScale3D(Control.Value.Get<FVector>());
+						break;
+					}
+					case ERigControlType::Rotator:
+					{
+						Transform.SetRotation(FQuat(Control.Value.Get<FRotator>()));
+						break;
+					}
+					case ERigControlType::Transform:
+					{
+						Transform = Control.Value.Get<FTransform>();
+						break;
+					}
+					case ERigControlType::TransformNoScale:
+					{
+						Transform = Control.Value.Get<FTransformNoScale>();
+						break;
+					}
+				}
+				return Transform * ParentTransform;
+			}
+		}
+	}
+#endif
 	return Hierarchy.ControlHierarchy.GetGlobalTransform(InControlName);
 }
 
@@ -680,54 +778,104 @@ FRigControlValue UControlRig::GetControlValueFromGlobalTransform(const FName& In
 {
 	FRigControlValue RetVal;
 
-	const int32 ControlIndex = Hierarchy.ControlHierarchy.GetIndex(InControlName);
-	if (ControlIndex != INDEX_NONE)
+	if (FRigControl* Control = FindControl(InControlName))
 	{
-		FRigControl& Control = Hierarchy.ControlHierarchy[ControlIndex];
+		RetVal = Control->Value;
 
-		RetVal = Control.Value;
+		FTransform ParentTransform = FTransform::Identity;
+		if ((Control->ParentIndex != INDEX_NONE || Control->SpaceIndex != INDEX_NONE) && Control->bIsTransientControl)
+		{
+			if (Control->ParentIndex != INDEX_NONE)
+			{
+				ParentTransform = Hierarchy.BoneHierarchy.GetGlobalTransform(Control->ParentIndex);
+			}
+		}
+		else if(Control->Index != INDEX_NONE)
+		{
+			ParentTransform = Hierarchy.ControlHierarchy.GetParentTransform(Control->Index);
+		}
 
-		FTransform ParentTransform = Hierarchy.ControlHierarchy.GetParentTransform(ControlIndex);
-		switch (Control.ControlType)
+		FTransform Transform = InGlobalTransform.GetRelativeTransform(ParentTransform);
+		switch (Control->ControlType)
 		{
 			case ERigControlType::Bool:
+			{
+				// not sure how to extract this from global transform
+				break;
+			}
 			case ERigControlType::Float:
+			{
+				switch (Control->PrimaryAxis)
+				{
+					case ERigControlAxis::X:
+					{
+						RetVal.Set<float>(Transform.GetLocation().X);
+						break;
+					}
+					case ERigControlAxis::Y:
+					{
+						RetVal.Set<float>(Transform.GetLocation().Y);
+						break;
+					}
+					case ERigControlAxis::Z:
+					{
+						RetVal.Set<float>(Transform.GetLocation().Z);
+						break;
+					}
+				}
+				break;
+			}
 			case ERigControlType::Vector2D:
 			{
-				// not sure how to extract this from global transform
+				FVector Location = Transform.GetLocation();
+				switch (Control->PrimaryAxis)
+				{
+					case ERigControlAxis::X:
+					{
+						RetVal.Set<FVector2D>(FVector2D(Location.Y, Location.Z));
+						break;
+					}
+					case ERigControlAxis::Y:
+					{
+						RetVal.Set<FVector2D>(FVector2D(Location.X, Location.Z));
+						break;
+					}
+					case ERigControlAxis::Z:
+					{
+						RetVal.Set<FVector2D>(FVector2D(Location.X, Location.Y));
+						break;
+					}
+				}
 				break;
 			}
 			case ERigControlType::Position:
 			{
-				FTransform Transform = InGlobalTransform.GetRelativeTransform(ParentTransform);
 				RetVal.Set<FVector>(Transform.GetLocation());
 				break;
 			}
 			case ERigControlType::Scale:
 			{
-				FTransform Transform = InGlobalTransform.GetRelativeTransform(ParentTransform);
 				RetVal.Set<FVector>(Transform.GetScale3D());
 				break;
 			}
-			case ERigControlType::Quat:
-			{
-				FTransform Transform = InGlobalTransform.GetRelativeTransform(ParentTransform);
-				RetVal.Set<FQuat>(Transform.GetRotation());
-				break;
-			}
 			case ERigControlType::Rotator:
 			{
-				FTransform Transform = InGlobalTransform.GetRelativeTransform(ParentTransform);
 				RetVal.Set<FRotator>(Transform.GetRotation().Rotator());
 				break;
 			}
 			case ERigControlType::Transform:
 			{
-				FTransform Transform = InGlobalTransform.GetRelativeTransform(ParentTransform);
 				RetVal.Set<FTransform>(Transform);
 				break;
 			}
+			case ERigControlType::TransformNoScale:
+			{
+				RetVal.Set<FTransformNoScale>(Transform);
+				break;
+			}
 		}
+
+		Control->ApplyLimits(RetVal);
 	}
 
 	return RetVal;
@@ -792,6 +940,265 @@ void UControlRig::HandleOnControlSelected(FRigHierarchyContainer* InContainer, c
 		}
 	}
 }
+
+void UControlRig::UpdateAvailableControls()
+{
+	AvailableControlsOverride = Hierarchy.ControlHierarchy.GetControls();
+	AvailableControlsOverride.Append(TransientControls);
+}
+
+FName UControlRig::AddTransientControl(URigVMPin* InPin, FName InSpaceName)
+{
+	if (InPin == nullptr)
+	{
+		return NAME_None;
+	}
+
+	if (InPin->GetCPPType() != TEXT("FVector") &&
+		InPin->GetCPPType() != TEXT("FQuat") &&
+		InPin->GetCPPType() != TEXT("FTransform"))
+	{
+		return NAME_None;
+	}
+
+	RemoveTransientControl(InPin);
+
+	FName ControlName = *InPin->GetPinPath();
+
+	FRigControl NewControl;
+	NewControl.Name = ControlName;
+	NewControl.bIsTransientControl = true;
+	NewControl.GizmoTransform.SetScale3D(FVector::ZeroVector);
+
+	URigVMPin* PinForLink = InPin->GetPinForLink();
+	if (URigVMPin* ColorPin = PinForLink->GetNode()->FindPin(TEXT("Color")))
+	{
+		if (ColorPin->GetCPPType() == TEXT("FLinearColor"))
+		{
+			FRigControlValue Value;
+			NewControl.GizmoColor = Value.SetFromString<FLinearColor>(ColorPin->GetDefaultValue());
+		}
+	}
+
+	if(!InSpaceName.IsNone())
+	{
+		NewControl.ParentIndex = Hierarchy.BoneHierarchy.GetIndex(InSpaceName);
+		if (NewControl.ParentIndex != INDEX_NONE)
+		{
+			NewControl.ParentName = Hierarchy.BoneHierarchy[NewControl.ParentIndex].Name;
+		}
+	}
+
+	TransientControls.Add(NewControl);
+
+	SetTransientControlValue(InPin);
+	UpdateAvailableControls();
+
+	return ControlName;
+}
+
+bool UControlRig::SetTransientControlValue(URigVMPin* InPin)
+{
+	FString PinPath = InPin->GetPinForLink()->GetPinPath();
+	FString OriginalPinPath = InPin->GetOriginalPinFromInjectedNode()->GetPinPath();
+	for (FRigControl& Control : TransientControls)
+	{
+		FString ControlName = Control.Name.ToString();
+		if (ControlName == PinPath || ControlName == OriginalPinPath)
+		{
+			FString DefaultValue = InPin->GetPinForLink()->GetDefaultValue();
+			if (!DefaultValue.IsEmpty())
+			{
+				if (InPin->GetCPPType() == TEXT("FVector"))
+				{
+					Control.ControlType = ERigControlType::Position;
+					FVector Value = Control.Value.SetFromString<FVector>(DefaultValue);
+					Control.InitialValue.Set<FVector>(Value);
+				}
+				else if (InPin->GetCPPType() == TEXT("FQuat"))
+				{
+					Control.ControlType = ERigControlType::Rotator;
+					FQuat Value = Control.Value.SetFromString<FQuat>(DefaultValue);
+					Control.Value.Set<FRotator>(Value.Rotator());
+					Control.InitialValue.Set<FRotator>(Value.Rotator());
+				}
+				else
+				{
+					Control.ControlType = ERigControlType::Transform;
+					FTransform Value = Control.Value.SetFromString<FTransform>(DefaultValue);
+					Control.InitialValue.Set<FTransform>(Value);
+				}
+			}
+			return true;
+		}
+	}
+	return false;
+}
+
+FName UControlRig::RemoveTransientControl(URigVMPin* InPin)
+{
+	if (InPin == nullptr)
+	{
+		return NAME_None;
+	}
+
+	FName ControlName = *InPin->GetPinPath();
+	for (int32 Index = 0; Index < TransientControls.Num(); Index++)
+	{
+		if (TransientControls[Index].Name == ControlName)
+		{
+			FRigControl ControlToRemove = TransientControls[Index];
+			TransientControls.RemoveAt(Index);
+			UpdateAvailableControls();
+			return ControlToRemove.Name;
+		}
+	}
+
+	return NAME_None;
+}
+
+FName UControlRig::AddTransientControl(const FRigElementKey& InElement)
+{
+	if (!InElement.IsValid())
+	{
+		return NAME_None;
+	}
+
+	RemoveTransientControl(InElement);
+
+	int32 ElementIndex = Hierarchy.GetIndex(InElement);
+	if (ElementIndex == INDEX_NONE)
+	{
+		return NAME_None;
+	}
+
+	FName ControlName = GetNameForTransientControl(InElement);
+
+	FRigControl NewControl;
+	NewControl.Name = ControlName;
+	NewControl.bIsTransientControl = true;
+	NewControl.GizmoTransform.SetScale3D(FVector::ZeroVector);
+
+	switch (InElement.Type)
+	{
+		case ERigElementType::Bone:
+		{
+			NewControl.ParentIndex = Hierarchy.BoneHierarchy[ElementIndex].ParentIndex;
+			NewControl.ParentName = Hierarchy.BoneHierarchy[ElementIndex].ParentName;
+			break;
+		}
+		case ERigElementType::Space:
+		{
+			NewControl.SpaceIndex = ElementIndex;
+			NewControl.SpaceName = InElement.Name;
+			break;
+		}
+	}
+
+	TransientControls.Add(NewControl);
+
+	if (InElement.Type == ERigElementType::Bone)
+	{
+		if (PreviewInstance)
+		{
+			PreviewInstance->ModifyBone(InElement.Name);
+		}
+	}
+
+	SetTransientControlValue(InElement);
+	UpdateAvailableControls();
+
+	return ControlName;
+}
+
+bool UControlRig::SetTransientControlValue(const FRigElementKey& InElement)
+{
+	if (!InElement.IsValid())
+	{
+		return false;
+	}
+
+	FName ControlName = GetNameForTransientControl(InElement);
+	for (FRigControl& Control : TransientControls)
+	{
+		if (Control.Name == ControlName)
+		{
+			Control.ControlType = ERigControlType::Transform;
+
+			if (InElement.Type == ERigElementType::Bone)
+			{
+				FTransform LocalTransform = Hierarchy.GetLocalTransform(InElement);
+				Control.Value.Set<FTransform>(LocalTransform);
+				Control.InitialValue.Set<FTransform>(LocalTransform);
+
+				if (PreviewInstance)
+				{
+					if (FAnimNode_ModifyBone* Modify = PreviewInstance->FindModifiedBone(InElement.Name))
+					{
+						Modify->Translation = LocalTransform.GetTranslation();
+						Modify->Rotation = LocalTransform.GetRotation().Rotator();
+						Modify->TranslationSpace = EBoneControlSpace::BCS_ParentBoneSpace;
+						Modify->RotationSpace = EBoneControlSpace::BCS_ParentBoneSpace;
+					}
+				}
+			}
+			else if (InElement.Type == ERigElementType::Space)
+			{
+				FTransform GlobalTransform = Hierarchy.GetGlobalTransform(InElement);
+				Control.Value.Set<FTransform>(GlobalTransform);
+				Control.InitialValue.Set<FTransform>(GlobalTransform);
+			}
+
+			return true;
+		}
+	}
+	return false;
+}
+
+FName UControlRig::RemoveTransientControl(const FRigElementKey& InElement)
+{
+	if (!InElement.IsValid())
+	{
+		return NAME_None;
+	}
+
+	FName ControlName = GetNameForTransientControl(InElement);
+	for (int32 Index = 0; Index < TransientControls.Num(); Index++)
+	{
+		if (TransientControls[Index].Name == ControlName)
+		{
+			FRigControl ControlToRemove = TransientControls[Index];
+			TransientControls.RemoveAt(Index);
+			UpdateAvailableControls();
+			return ControlToRemove.Name;
+		}
+	}
+
+	return NAME_None;
+
+}
+
+FName UControlRig::GetNameForTransientControl(const FRigElementKey& InElement)
+{
+	if (InElement.Type == ERigElementType::Control)
+	{
+		return InElement.Name;
+	}
+
+	FName EnumName = StaticEnum<ERigElementType>()->GetNameByValue((int64)InElement.Type);
+	return *FString::Printf(TEXT("ControlForRigElement_%s_%s"), *EnumName.ToString(), *InElement.Name.ToString());
+}
+
+void UControlRig::ClearTransientControls()
+{
+	for (int32 Index = TransientControls.Num() - 1; Index >= 0; Index--)
+	{
+		FRigControl ControlToRemove = TransientControls[Index];
+		TransientControls.RemoveAt(Index);
+		UpdateAvailableControls();
+	}
+}
+
 #endif
 
 #undef LOCTEXT_NAMESPACE
diff --git a/Engine/Plugins/Experimental/ControlRig/Source/ControlRig/Private/ControlRigBlueprintGeneratedClass.cpp b/Engine/Plugins/Experimental/ControlRig/Source/ControlRig/Private/ControlRigBlueprintGeneratedClass.cpp
--- a/Engine/Plugins/Experimental/ControlRig/Source/ControlRig/Private/ControlRigBlueprintGeneratedClass.cpp
+++ b/Engine/Plugins/Experimental/ControlRig/Source/ControlRig/Private/ControlRigBlueprintGeneratedClass.cpp
@@ -2,50 +2,62 @@
 
 #include "ControlRigBlueprintGeneratedClass.h"
 #include "Units/Control/RigUnit_Control.h"
+#include "ControlRigObjectVersion.h"
+#include "ControlRig.h"
 
-UControlRigBlueprintGeneratedClass::UControlRigBlueprintGeneratedClass()
+UControlRigBlueprintGeneratedClass::UControlRigBlueprintGeneratedClass(const FObjectInitializer& ObjectInitializer)
+	: Super(ObjectInitializer)
 {
 }
 
-void UControlRigBlueprintGeneratedClass::Link(FArchive& Ar, bool bRelinkExistingProperties)
+uint8* UControlRigBlueprintGeneratedClass::GetPersistentUberGraphFrame(UObject* Obj, UFunction* FuncToCheck) const
 {
-	Super::Link(Ar, bRelinkExistingProperties);
-
-#if WITH_EDITORONLY_DATA
-	ControlUnitProperties.Empty();
-	RigUnitProperties.Empty();
-
-	for (TFieldIterator<FProperty> It(this); It; ++It)
+	if(!IsInGameThread())
 	{
-		if (FStructProperty* StructProp = CastField<FStructProperty>(*It))
-		{
-			if (StructProp->Struct->IsChildOf(FRigUnit::StaticStruct()))
-			{
-				RigUnitProperties.Add(StructProp);
-
-				if (StructProp->Struct->IsChildOf(FRigUnit_Control::StaticStruct()))
-				{
-					ControlUnitProperties.Add(StructProp);
-				}
-			}
-		}
+		// we cant use the persistent frame if we are executing in parallel (as we could potentially thunk to BP)
+		return nullptr;
 	}
-#endif
+	return Super::GetPersistentUberGraphFrame(Obj, FuncToCheck);
 }
 
-void UControlRigBlueprintGeneratedClass::PurgeClass(bool bRecompilingOnLoad)
+void UControlRigBlueprintGeneratedClass::Serialize(FArchive& Ar)
 {
-	Super::PurgeClass(bRecompilingOnLoad);
-}
+	DECLARE_SCOPE_HIERARCHICAL_COUNTER_FUNC()
 
-uint8* UControlRigBlueprintGeneratedClass::GetPersistentUberGraphFrame(UObject* Obj, UFunction* FuncToCheck) const
-{
-	if(!IsInGameThread())
+	Super::Serialize(Ar);
+
+	Ar.UsingCustomVersion(FControlRigObjectVersion::GUID);
+
+	if (Ar.CustomVer(FControlRigObjectVersion::GUID) < FControlRigObjectVersion::SwitchedToRigVM)
 	{
-		// we cant use the persistent frame if we are executing in parallel (as we could potentially thunk to BP)
-		return nullptr;
+		return;
 	}
 
-	return Super::GetPersistentUberGraphFrame(Obj, FuncToCheck);
-}
+	URigVM* VM = NewObject<URigVM>(GetTransientPackage());
+
+	if (UControlRig* CDO = Cast<UControlRig>(GetDefaultObject(true)))
+	{
+		if (Ar.IsSaving() && CDO->VM)
+		{
+			VM->CopyFrom(CDO->VM);
+		}
+	}
+	
+	VM->Serialize(Ar);
 
+	if (UControlRig* CDO = Cast<UControlRig>(GetDefaultObject(false)))
+	{
+		if (Ar.IsLoading())
+		{
+			if (CDO->VM == nullptr)
+			{
+				CDO->VM = NewObject<URigVM>(CDO, TEXT("VM"));
+			}
+			if (CDO->VM->GetOuter() != CDO)
+			{
+				CDO->VM = NewObject<URigVM>(CDO, TEXT("VM"));
+			}
+			CDO->VM->CopyFrom(VM);
+		}
+	}
+}
\ No newline at end of file
diff --git a/Engine/Plugins/Experimental/ControlRig/Source/ControlRig/Private/ControlRigDefines.cpp b/Engine/Plugins/Experimental/ControlRig/Source/ControlRig/Private/ControlRigDefines.cpp
--- a/Engine/Plugins/Experimental/ControlRig/Source/ControlRig/Private/ControlRigDefines.cpp
+++ b/Engine/Plugins/Experimental/ControlRig/Source/ControlRig/Private/ControlRigDefines.cpp
@@ -3,30 +3,3 @@
 #include "ControlRigDefines.h"
 #include "PropertyPathHelpers.h"
 
-FControlRigOperator FControlRigOperator::MakeUnresolvedCopy(const FControlRigOperator& ToCopy)
-{
-	FControlRigOperator Op;
-	Op.OpCode = ToCopy.OpCode;
-	Op.CachedPropertyPath1 = FCachedPropertyPath::MakeUnresolvedCopy(ToCopy.CachedPropertyPath1);
-	Op.CachedPropertyPath2 = FCachedPropertyPath::MakeUnresolvedCopy(ToCopy.CachedPropertyPath2);
-	return Op;
-}
-
-bool FControlRigOperator::Resolve(UObject* OuterObject)
-{
-	if (CachedPropertyPath1.IsValid() && !CachedPropertyPath1.IsFullyResolved())
-	{
-		if (!CachedPropertyPath1.Resolve(OuterObject))
-		{
-			return false;
-		}
-	}
-	if (CachedPropertyPath2.IsValid() && !CachedPropertyPath2.IsFullyResolved())
-	{
-		if (!CachedPropertyPath2.Resolve(OuterObject))
-		{
-			return false;
-		}
-	}
-	return true;
-}
diff --git a/Engine/Plugins/Experimental/ControlRig/Source/ControlRig/Private/ControlRigLog.cpp b/Engine/Plugins/Experimental/ControlRig/Source/ControlRig/Private/ControlRigLog.cpp
--- a/Engine/Plugins/Experimental/ControlRig/Source/ControlRig/Private/ControlRigLog.cpp
+++ b/Engine/Plugins/Experimental/ControlRig/Source/ControlRig/Private/ControlRigLog.cpp
@@ -2,9 +2,9 @@
 
 #include "ControlRigLog.h"
 
-void FControlRigLog::Report(EMessageSeverity::Type InSeverity, const FName& InUnit, const FString& InMessage)
+void FControlRigLog::Report(EMessageSeverity::Type InSeverity, const FName& InOperatorName, int32 InInstructionIndex, const FString& InMessage)
 {
 #if WITH_EDITOR
-	Entries.Add(FLogEntry(InSeverity, InUnit, InMessage));
+	Entries.Add(FLogEntry(InSeverity, InOperatorName, InInstructionIndex, InMessage));
 #endif
 }
diff --git a/Engine/Plugins/Experimental/ControlRig/Source/ControlRig/Private/Drawing/ControlRigDrawContainer.cpp b/Engine/Plugins/Experimental/ControlRig/Source/ControlRig/Private/Drawing/ControlRigDrawContainer.cpp
new file mode 100644
--- /dev/null
+++ b/Engine/Plugins/Experimental/ControlRig/Source/ControlRig/Private/Drawing/ControlRigDrawContainer.cpp
@@ -0,0 +1,20 @@
+// Copyright Epic Games, Inc. All Rights Reserved.
+
+#include "Drawing/ControlRigDrawContainer.h"
+
+int32 FControlRigDrawContainer::GetIndex(const FName& InName) const
+{
+	for (int32 Index = 0; Index < Instructions.Num(); ++Index)
+	{
+		if (Instructions[Index].Name == InName)
+		{
+			return Index;
+		}
+	}
+	return INDEX_NONE;
+}
+
+void FControlRigDrawContainer::Reset()
+{
+	Instructions.Reset();
+}
diff --git a/Engine/Plugins/Experimental/ControlRig/Source/ControlRig/Private/Drawing/ControlRigDrawInstruction.cpp b/Engine/Plugins/Experimental/ControlRig/Source/ControlRig/Private/Drawing/ControlRigDrawInstruction.cpp
new file mode 100644
--- /dev/null
+++ b/Engine/Plugins/Experimental/ControlRig/Source/ControlRig/Private/Drawing/ControlRigDrawInstruction.cpp
@@ -0,0 +1,3 @@
+// Copyright Epic Games, Inc. All Rights Reserved.
+
+#include "Drawing/ControlRigDrawInstruction.h"
diff --git a/Engine/Plugins/Experimental/ControlRig/Source/ControlRig/Private/Drawing/ControlRigDrawInterface.cpp b/Engine/Plugins/Experimental/ControlRig/Source/ControlRig/Private/Drawing/ControlRigDrawInterface.cpp
--- a/Engine/Plugins/Experimental/ControlRig/Source/ControlRig/Private/Drawing/ControlRigDrawInterface.cpp
+++ b/Engine/Plugins/Experimental/ControlRig/Source/ControlRig/Private/Drawing/ControlRigDrawInterface.cpp
@@ -5,83 +5,73 @@
 
 void FControlRigDrawInterface::DrawPoint(const FTransform& WorldOffset, const FVector& Position, float Size, const FLinearColor& Color)
 {
-	FDrawIntruction Instruction(EDrawType_Point, Color, Size);
-	Instruction.Positions.Add(WorldOffset.TransformPosition(Position));
-	DrawInstructions.Add(Instruction);
+	FControlRigDrawInstruction Instruction(EControlRigDrawSettings::Points, Color, Size, WorldOffset);
+	Instruction.Positions.Add(Position);
+	Instructions.Add(Instruction);
 }
 
 void FControlRigDrawInterface::DrawPoints(const FTransform& WorldOffset, const TArrayView<FVector>& Points, float Size, const FLinearColor& Color)
 {
-	FDrawIntruction Instruction(EDrawType_Point, Color, Size);
-	for(const FVector& Point : Points)
-	{
-		Instruction.Positions.Add(WorldOffset.TransformPosition(Point));
-	}
-	DrawInstructions.Add(Instruction);
+	FControlRigDrawInstruction Instruction(EControlRigDrawSettings::Points, Color, Size, WorldOffset);
+	Instruction.Positions.Append(Points.GetData(), Points.Num());
+	Instructions.Add(Instruction);
 }
 
 void FControlRigDrawInterface::DrawLine(const FTransform& WorldOffset, const FVector& LineStart, const FVector& LineEnd, const FLinearColor& Color, float Thickness)
 {
-	FDrawIntruction Instruction(EDrawType_Lines, Color, Thickness);
-	Instruction.Positions.Add(WorldOffset.TransformPosition(LineStart));
-	Instruction.Positions.Add(WorldOffset.TransformPosition(LineEnd));
-	DrawInstructions.Add(Instruction);
+	FControlRigDrawInstruction Instruction(EControlRigDrawSettings::Lines, Color, Thickness, WorldOffset);
+	Instruction.Positions.Add(LineStart);
+	Instruction.Positions.Add(LineEnd);
+	Instructions.Add(Instruction);
 }
 
 void FControlRigDrawInterface::DrawLines(const FTransform& WorldOffset, const TArrayView<FVector>& Positions, const FLinearColor& Color, float Thickness)
 {
-	FDrawIntruction Instruction(EDrawType_Lines, Color, Thickness);
-	for (const FVector& Point : Positions)
-	{
-		Instruction.Positions.Add(WorldOffset.TransformPosition(Point));
-	}
-	DrawInstructions.Add(Instruction);
+	FControlRigDrawInstruction Instruction(EControlRigDrawSettings::Lines, Color, Thickness, WorldOffset);
+	Instruction.Positions.Append(Positions.GetData(), Positions.Num());
 }
 
 void FControlRigDrawInterface::DrawLineStrip(const FTransform& WorldOffset, const TArrayView<FVector>& Positions, const FLinearColor& Color, float Thickness)
 {
-	FDrawIntruction Instruction(EDrawType_LineStrip, Color, Thickness);
-	for (const FVector& Point : Positions)
-	{
-		Instruction.Positions.Add(WorldOffset.TransformPosition(Point));
-	}
-	DrawInstructions.Add(Instruction);
+	FControlRigDrawInstruction Instruction(EControlRigDrawSettings::LineStrip, Color, Thickness, WorldOffset);
+	Instruction.Positions.Append(Positions.GetData(), Positions.Num());
+	Instructions.Add(Instruction);
 }
 
 void FControlRigDrawInterface::DrawBox(const FTransform& WorldOffset, const FTransform& Transform, const FLinearColor& Color, float Thickness)
 {
 	FTransform DrawTransform = Transform * WorldOffset;
 
-	FDrawIntruction Instruction(EDrawType_Lines, Color, Thickness);
-
-	Instruction.Positions.Add(DrawTransform.TransformPosition(FVector(0.5f, 0.5f, 0.5f)));
-	Instruction.Positions.Add(DrawTransform.TransformPosition(FVector(0.5f, -0.5f, 0.5f)));
-	Instruction.Positions.Add(DrawTransform.TransformPosition(FVector(0.5f, -0.5f, 0.5f)));
-	Instruction.Positions.Add(DrawTransform.TransformPosition(FVector(-0.5f, -0.5f, 0.5f)));
-	Instruction.Positions.Add(DrawTransform.TransformPosition(FVector(-0.5f, -0.5f, 0.5f)));
-	Instruction.Positions.Add(DrawTransform.TransformPosition(FVector(-0.5f, 0.5f, 0.5f)));
-	Instruction.Positions.Add(DrawTransform.TransformPosition(FVector(-0.5f, 0.5f, 0.5f)));
-	Instruction.Positions.Add(DrawTransform.TransformPosition(FVector(0.5f, 0.5f, 0.5f)));
-
-	Instruction.Positions.Add(DrawTransform.TransformPosition(FVector(0.5f, 0.5f, -0.5f)));
-	Instruction.Positions.Add(DrawTransform.TransformPosition(FVector(0.5f, -0.5f, -0.5f)));
-	Instruction.Positions.Add(DrawTransform.TransformPosition(FVector(0.5f, -0.5f, -0.5f)));
-	Instruction.Positions.Add(DrawTransform.TransformPosition(FVector(-0.5f, -0.5f, -0.5f)));
-	Instruction.Positions.Add(DrawTransform.TransformPosition(FVector(-0.5f, -0.5f, -0.5f)));
-	Instruction.Positions.Add(DrawTransform.TransformPosition(FVector(-0.5f, 0.5f, -0.5f)));
-	Instruction.Positions.Add(DrawTransform.TransformPosition(FVector(-0.5f, 0.5f, -0.5f)));
-	Instruction.Positions.Add(DrawTransform.TransformPosition(FVector(0.5f, 0.5f, -0.5f)));
-
-	Instruction.Positions.Add(DrawTransform.TransformPosition(FVector(0.5f, 0.5f, 0.5f)));
-	Instruction.Positions.Add(DrawTransform.TransformPosition(FVector(0.5f, 0.5f, -0.5f)));
-	Instruction.Positions.Add(DrawTransform.TransformPosition(FVector(0.5f, -0.5f, 0.5f)));
-	Instruction.Positions.Add(DrawTransform.TransformPosition(FVector(0.5f, -0.5f, -0.5f)));
-	Instruction.Positions.Add(DrawTransform.TransformPosition(FVector(-0.5f, -0.5f, 0.5f)));
-	Instruction.Positions.Add(DrawTransform.TransformPosition(FVector(-0.5f, -0.5f, -0.5f)));
-	Instruction.Positions.Add(DrawTransform.TransformPosition(FVector(-0.5f, 0.5f, 0.5f)));
-	Instruction.Positions.Add(DrawTransform.TransformPosition(FVector(-0.5f, 0.5f, -0.5f)));
-
-	DrawInstructions.Add(Instruction);
+	FControlRigDrawInstruction Instruction(EControlRigDrawSettings::Lines, Color, Thickness, DrawTransform);
+
+	Instruction.Positions.Add(FVector(0.5f, 0.5f, 0.5f));
+	Instruction.Positions.Add(FVector(0.5f, -0.5f, 0.5f));
+	Instruction.Positions.Add(FVector(0.5f, -0.5f, 0.5f));
+	Instruction.Positions.Add(FVector(-0.5f, -0.5f, 0.5f));
+	Instruction.Positions.Add(FVector(-0.5f, -0.5f, 0.5f));
+	Instruction.Positions.Add(FVector(-0.5f, 0.5f, 0.5f));
+	Instruction.Positions.Add(FVector(-0.5f, 0.5f, 0.5f));
+	Instruction.Positions.Add(FVector(0.5f, 0.5f, 0.5f));
+
+	Instruction.Positions.Add(FVector(0.5f, 0.5f, -0.5f));
+	Instruction.Positions.Add(FVector(0.5f, -0.5f, -0.5f));
+	Instruction.Positions.Add(FVector(0.5f, -0.5f, -0.5f));
+	Instruction.Positions.Add(FVector(-0.5f, -0.5f, -0.5f));
+	Instruction.Positions.Add(FVector(-0.5f, -0.5f, -0.5f));
+	Instruction.Positions.Add(FVector(-0.5f, 0.5f, -0.5f));
+	Instruction.Positions.Add(FVector(-0.5f, 0.5f, -0.5f));
+	Instruction.Positions.Add(FVector(0.5f, 0.5f, -0.5f));
+
+	Instruction.Positions.Add(FVector(0.5f, 0.5f, 0.5f));
+	Instruction.Positions.Add(FVector(0.5f, 0.5f, -0.5f));
+	Instruction.Positions.Add(FVector(0.5f, -0.5f, 0.5f));
+	Instruction.Positions.Add(FVector(0.5f, -0.5f, -0.5f));
+	Instruction.Positions.Add(FVector(-0.5f, -0.5f, 0.5f));
+	Instruction.Positions.Add(FVector(-0.5f, -0.5f, -0.5f));
+	Instruction.Positions.Add(FVector(-0.5f, 0.5f, 0.5f));
+	Instruction.Positions.Add(FVector(-0.5f, 0.5f, -0.5f));
+
+	Instructions.Add(Instruction);
 }
 
 void FControlRigDrawInterface::DrawAxes(const FTransform& WorldOffset, const FTransform& Transform, float Size, float Thickness)
@@ -95,32 +85,32 @@ void FControlRigDrawInterface::DrawRectangle(const FTransform& WorldOffset, cons
 {
 	FTransform DrawTransform = Transform * WorldOffset;
 
-	FDrawIntruction Instruction(EDrawType_LineStrip, Color, Thickness);
+	FControlRigDrawInstruction Instruction(EControlRigDrawSettings::LineStrip, Color, Thickness, DrawTransform);
 
 	float Extent = Size * 0.5f;
 	Instruction.Positions.Reserve(5);
-	Instruction.Positions.Add(DrawTransform.TransformPosition(FVector(-Extent, -Extent, 0.f)));
-	Instruction.Positions.Add(DrawTransform.TransformPosition(FVector(-Extent, Extent, 0.f)));
-	Instruction.Positions.Add(DrawTransform.TransformPosition(FVector(Extent, Extent, 0.f)));
-	Instruction.Positions.Add(DrawTransform.TransformPosition(FVector(Extent, -Extent, 0.f)));
-	Instruction.Positions.Add(DrawTransform.TransformPosition(FVector(-Extent, -Extent, 0.f)));
+	Instruction.Positions.Add(FVector(-Extent, -Extent, 0.f));
+	Instruction.Positions.Add(FVector(-Extent, Extent, 0.f));
+	Instruction.Positions.Add(FVector(Extent, Extent, 0.f));
+	Instruction.Positions.Add(FVector(Extent, -Extent, 0.f));
+	Instruction.Positions.Add(FVector(-Extent, -Extent, 0.f));
 
-	DrawInstructions.Add(Instruction);
+	Instructions.Add(Instruction);
 }
 
 void FControlRigDrawInterface::DrawArc(const FTransform& WorldOffset, const FTransform& Transform, float Radius, float MinimumAngle, float MaximumAngle, const FLinearColor& Color, float Thickness, int32 Detail)
 {
 	int32 Count = FMath::Clamp<int32>(Detail, 4, 32);
 	
-	FDrawIntruction Instruction(EDrawType_LineStrip, Color, Thickness);
-	Instruction.Positions.Reserve(Count);
-
 	FTransform DrawTransform = Transform * WorldOffset;
 
+	FControlRigDrawInstruction Instruction(EControlRigDrawSettings::LineStrip, Color, Thickness, DrawTransform);
+	Instruction.Positions.Reserve(Count);
+
 	FVector V = FVector(Radius, 0.f, 0.f);
 	FQuat Rotation(FVector(0.f, 0.f, 1.f), MinimumAngle);
 	V = Rotation.RotateVector(V);
-	Instruction.Positions.Add(DrawTransform.TransformPosition(V));
+	Instruction.Positions.Add(V);
 	float StepAngle = (MaximumAngle - MinimumAngle) / float(Count);
 	if (FMath::Abs<float>(MaximumAngle - MinimumAngle) >= PI * 2.f - SMALL_NUMBER)
 	{
@@ -131,34 +121,28 @@ void FControlRigDrawInterface::DrawArc(const FTransform& WorldOffset, const FTra
 	for(int32 Index=1;Index<Count;Index++)
 	{
 		V = Rotation.RotateVector(V);
-		Instruction.Positions.Add(DrawTransform.TransformPosition(V));
+		Instruction.Positions.Add(V);
 	}
 
-	DrawInstructions.Add(Instruction);
+	Instructions.Add(Instruction);
 }
 
 void FControlRigDrawInterface::DrawBezier(const FTransform& WorldOffset, const FCRFourPointBezier& InBezier, float MinimumU, float MaximumU, const FLinearColor& Color, float Thickness, int32 Detail)
 {
 	int32 Count = FMath::Clamp<int32>(Detail, 4, 64);
-	FDrawIntruction Instruction(EDrawType_LineStrip, Color, Thickness);
+	FControlRigDrawInstruction Instruction(EControlRigDrawSettings::LineStrip, Color, Thickness, WorldOffset);
 	Instruction.Positions.SetNumUninitialized(Count);
 
-	FCRFourPointBezier Bezier = InBezier;
-	Bezier.A = WorldOffset.TransformPosition(Bezier.A);
-	Bezier.B = WorldOffset.TransformPosition(Bezier.B);
-	Bezier.C = WorldOffset.TransformPosition(Bezier.C);
-	Bezier.D = WorldOffset.TransformPosition(Bezier.D);
-
 	float T = MinimumU;
 	float Step = (MaximumU - MinimumU) / float(Detail-1);
 	for(int32 Index=0;Index<Count;Index++)
 	{
 		FVector Tangent;
-		FControlRigMathLibrary::FourPointBezier(Bezier, T, Instruction.Positions[Index], Tangent);
+		FControlRigMathLibrary::FourPointBezier(InBezier, T, Instruction.Positions[Index], Tangent);
 		T += Step;
 	}
 
-	DrawInstructions.Add(Instruction);
+	Instructions.Add(Instruction);
 }
 
 void FControlRigDrawInterface::DrawHierarchy(const FTransform& WorldOffset, const FRigBoneHierarchy& Hierarchy, EControlRigDrawHierarchyMode::Type Mode, float Scale, const FLinearColor& Color, float Thickness)
@@ -167,10 +151,10 @@ void FControlRigDrawInterface::DrawHierarchy(const FTransform& WorldOffset, cons
 	{
 		case EControlRigDrawHierarchyMode::Axes:
 		{
-			FDrawIntruction InstructionX(EDrawType_Lines, FLinearColor::Red, Thickness);
-			FDrawIntruction InstructionY(EDrawType_Lines, FLinearColor::Green, Thickness);
-			FDrawIntruction InstructionZ(EDrawType_Lines, FLinearColor::Blue, Thickness);
-			FDrawIntruction InstructionParent(EDrawType_Lines, Color, Thickness);
+			FControlRigDrawInstruction InstructionX(EControlRigDrawSettings::Lines, FLinearColor::Red, Thickness, WorldOffset);
+			FControlRigDrawInstruction InstructionY(EControlRigDrawSettings::Lines, FLinearColor::Green, Thickness, WorldOffset);
+			FControlRigDrawInstruction InstructionZ(EControlRigDrawSettings::Lines, FLinearColor::Blue, Thickness, WorldOffset);
+			FControlRigDrawInstruction InstructionParent(EControlRigDrawSettings::Lines, Color, Thickness, WorldOffset);
 			InstructionX.Positions.Reserve(Hierarchy.Num() * 2);
 			InstructionY.Positions.Reserve(Hierarchy.Num() * 2);
 			InstructionZ.Positions.Reserve(Hierarchy.Num() * 2);
@@ -178,7 +162,7 @@ void FControlRigDrawInterface::DrawHierarchy(const FTransform& WorldOffset, cons
 
 			for (const FRigBone& Bone : Hierarchy)
 			{
-				FTransform Transform = Bone.GlobalTransform * WorldOffset;
+				FTransform Transform = Bone.GlobalTransform;
 				FVector P0 = Transform.GetLocation();
 				FVector PX = Transform.TransformPosition(FVector(Scale, 0.f, 0.f));
 				FVector PY = Transform.TransformPosition(FVector(0.f, Scale, 0.f));
@@ -192,17 +176,17 @@ void FControlRigDrawInterface::DrawHierarchy(const FTransform& WorldOffset, cons
 
 				if (Bone.ParentIndex != INDEX_NONE)
 				{
-					FTransform ParentTransform = Hierarchy[Bone.ParentIndex].GlobalTransform * WorldOffset;
+					FTransform ParentTransform = Hierarchy[Bone.ParentIndex].GlobalTransform;
 					FVector P1 = ParentTransform.GetLocation();
 					InstructionParent.Positions.Add(P0);
 					InstructionParent.Positions.Add(P1);
 				}
 			}
 
-			DrawInstructions.Add(InstructionX);
-			DrawInstructions.Add(InstructionY);
-			DrawInstructions.Add(InstructionZ);
-			DrawInstructions.Add(InstructionParent);
+			Instructions.Add(InstructionX);
+			Instructions.Add(InstructionY);
+			Instructions.Add(InstructionZ);
+			Instructions.Add(InstructionParent);
 			break;
 		}
 	}
@@ -210,20 +194,20 @@ void FControlRigDrawInterface::DrawHierarchy(const FTransform& WorldOffset, cons
 
 void FControlRigDrawInterface::DrawPointSimulation(const FTransform& WorldOffset, const FCRSimPointContainer& Simulation, const FLinearColor& Color, float Thickness, float PrimitiveSize, bool bDrawPointsAsSphere)
 {
-	FDrawIntruction PointsInstruction(EDrawType_Point, Color, Thickness * 6.f);
-	FDrawIntruction SpringsInstruction(EDrawType_Lines, Color * FLinearColor(0.55f, 0.55f, 0.55f, 1.f), Thickness);
-	FDrawIntruction VolumesMinInstruction(EDrawType_Lines, Color * FLinearColor(0.25f, 0.25f, 0.25f, 1.f), Thickness);
-	FDrawIntruction VolumesMaxInstruction(EDrawType_Lines, Color * FLinearColor(0.75f, 0.75f, 0.75f, 1.f) + FLinearColor(0.25f, 0.25f, 0.25f, 0.f), Thickness);
+	FControlRigDrawInstruction PointsInstruction(EControlRigDrawSettings::Points, Color, Thickness * 6.f, WorldOffset);
+	FControlRigDrawInstruction SpringsInstruction(EControlRigDrawSettings::Lines, Color * FLinearColor(0.55f, 0.55f, 0.55f, 1.f), Thickness, WorldOffset);
+	FControlRigDrawInstruction VolumesMinInstruction(EControlRigDrawSettings::Lines, Color * FLinearColor(0.25f, 0.25f, 0.25f, 1.f), Thickness, WorldOffset);
+	FControlRigDrawInstruction VolumesMaxInstruction(EControlRigDrawSettings::Lines, Color * FLinearColor(0.75f, 0.75f, 0.75f, 1.f) + FLinearColor(0.25f, 0.25f, 0.25f, 0.f), Thickness, WorldOffset);
 
 	if (bDrawPointsAsSphere)
 	{
-		PointsInstruction.DrawType = EDrawType_Lines;
+		PointsInstruction.PrimitiveType = EControlRigDrawSettings::Lines;
 		PointsInstruction.Thickness = Thickness * 2.f;
 
 		for (int32 PointIndex = 0; PointIndex < Simulation.Points.Num(); PointIndex++)
 		{
 			FCRSimPoint Point = Simulation.GetPointInterpolated(PointIndex);
-			FTransform Transform = FTransform(Point.Position) * WorldOffset;
+			FTransform Transform = FTransform(Point.Position);
 			static const int32 Subdivision = 8;
 			FVector MinV = Transform.TransformVector(FVector(Point.Size, 0.f, 0.f));
 			FQuat Q = FQuat(Transform.TransformVectorNoScale(FVector(0.f, 1.f, 0.f)), 2.f * PI / float(Subdivision));
@@ -257,7 +241,7 @@ void FControlRigDrawInterface::DrawPointSimulation(const FTransform& WorldOffset
 		for (int32 PointIndex = 0; PointIndex < Simulation.Points.Num(); PointIndex++)
 		{
 			FCRSimPoint Point = Simulation.GetPointInterpolated(PointIndex);
-			PointsInstruction.Positions.Add(WorldOffset.TransformPosition(Point.Position));
+			PointsInstruction.Positions.Add(Point.Position);
 		}
 	}
 
@@ -272,26 +256,26 @@ void FControlRigDrawInterface::DrawPointSimulation(const FTransform& WorldOffset
 		{
 			continue;
 		}
-		SpringsInstruction.Positions.Add(WorldOffset.TransformPosition(Simulation.GetPointInterpolated(Spring.SubjectA).Position));
-		SpringsInstruction.Positions.Add(WorldOffset.TransformPosition(Simulation.GetPointInterpolated(Spring.SubjectB).Position));
+		SpringsInstruction.Positions.Add(Simulation.GetPointInterpolated(Spring.SubjectA).Position);
+		SpringsInstruction.Positions.Add(Simulation.GetPointInterpolated(Spring.SubjectB).Position);
 	}
 
 	if (PrimitiveSize > SMALL_NUMBER)
 	{
 		for (const FCRSimSoftCollision& Volume : Simulation.CollisionVolumes)
 		{
-			FTransform Transform = Volume.Transform * WorldOffset;
+			FTransform Transform = Volume.Transform;
 			switch (Volume.ShapeType)
 			{
 				case ECRSimSoftCollisionType::Plane:
 				{
-					VolumesMinInstruction.DrawType = EDrawType_LineStrip;
+					VolumesMinInstruction.PrimitiveType = EControlRigDrawSettings::LineStrip;
 					VolumesMinInstruction.Positions.Add(Transform.TransformPosition(FVector(PrimitiveSize, PrimitiveSize, Volume.MinimumDistance) * 0.5f));
 					VolumesMinInstruction.Positions.Add(Transform.TransformPosition(FVector(-PrimitiveSize, PrimitiveSize, Volume.MinimumDistance) * 0.5f));
 					VolumesMinInstruction.Positions.Add(Transform.TransformPosition(FVector(-PrimitiveSize, -PrimitiveSize, Volume.MinimumDistance) * 0.5f));
 					VolumesMinInstruction.Positions.Add(Transform.TransformPosition(FVector(PrimitiveSize, -PrimitiveSize, Volume.MinimumDistance) * 0.5f));
 					VolumesMinInstruction.Positions.Add(Transform.TransformPosition(FVector(PrimitiveSize, PrimitiveSize, Volume.MinimumDistance) * 0.5f));
-					VolumesMaxInstruction.DrawType = EDrawType_LineStrip;
+					VolumesMaxInstruction.PrimitiveType = EControlRigDrawSettings::LineStrip;
 					VolumesMaxInstruction.Positions.Add(Transform.TransformPosition(FVector(PrimitiveSize, PrimitiveSize, Volume.MaximumDistance) * 0.5f));
 					VolumesMaxInstruction.Positions.Add(Transform.TransformPosition(FVector(-PrimitiveSize, PrimitiveSize, Volume.MaximumDistance) * 0.5f));
 					VolumesMaxInstruction.Positions.Add(Transform.TransformPosition(FVector(-PrimitiveSize, -PrimitiveSize, Volume.MaximumDistance) * 0.5f));
@@ -373,14 +357,14 @@ void FControlRigDrawInterface::DrawPointSimulation(const FTransform& WorldOffset
 		}
 	}
 
-	DrawInstructions.Add(PointsInstruction);
+	Instructions.Add(PointsInstruction);
 	if (SpringsInstruction.Positions.Num() > 0)
 	{
-		DrawInstructions.Add(SpringsInstruction);
+		Instructions.Add(SpringsInstruction);
 	}
 	if (VolumesMinInstruction.Positions.Num() > 0)
 	{
-		DrawInstructions.Add(VolumesMinInstruction);
-		DrawInstructions.Add(VolumesMaxInstruction);
+		Instructions.Add(VolumesMinInstruction);
+		Instructions.Add(VolumesMaxInstruction);
 	}
 }
diff --git a/Engine/Plugins/Experimental/ControlRig/Source/ControlRig/Private/Manipulatable/IControlRigManipulatable.cpp b/Engine/Plugins/Experimental/ControlRig/Source/ControlRig/Private/Manipulatable/IControlRigManipulatable.cpp
--- a/Engine/Plugins/Experimental/ControlRig/Source/ControlRig/Private/Manipulatable/IControlRigManipulatable.cpp
+++ b/Engine/Plugins/Experimental/ControlRig/Source/ControlRig/Private/Manipulatable/IControlRigManipulatable.cpp
@@ -12,7 +12,7 @@ IControlRigManipulatable::IControlRigManipulatable()
 	bManipulationEnabled = false;
 }
 
-bool IControlRigManipulatable::SetControlGlobalTransform(const FName& InControlName, const FTransform& InGlobalTransform)
+bool IControlRigManipulatable::SetControlGlobalTransform(const FName& InControlName, const FTransform& InGlobalTransform, EControlRigSetKey InSetKey)
 {
 	FRigControlValue Value = GetControlValueFromGlobalTransform(InControlName, InGlobalTransform);
 	if (OnFilterControl.IsBound())
@@ -23,6 +23,6 @@ bool IControlRigManipulatable::SetControlGlobalTransform(const FName& InControlN
 			OnFilterControl.Broadcast(this, *Control, Value);
 		}
 	}
-	SetControlValue(InControlName, Value, true /* notify */);
+	SetControlValue(InControlName, Value, true /* notify */, InSetKey);
 	return true;
 }
\ No newline at end of file
diff --git a/Engine/Plugins/Experimental/ControlRig/Source/ControlRig/Private/Rigs/AdditiveControlRig.cpp b/Engine/Plugins/Experimental/ControlRig/Source/ControlRig/Private/Rigs/AdditiveControlRig.cpp
new file mode 100644
--- /dev/null
+++ b/Engine/Plugins/Experimental/ControlRig/Source/ControlRig/Private/Rigs/AdditiveControlRig.cpp
@@ -0,0 +1,123 @@
+// Copyright Epic Games, Inc. All Rights Reserved.
+
+#include "Rigs/AdditiveControlRig.h"
+#include "Animation/SmartName.h"
+#include "Engine/SkeletalMesh.h"
+#include "IControlRigObjectBinding.h"
+#include "Components/SkeletalMeshComponent.h"
+
+#define LOCTEXT_NAMESPACE "AdditiveControlRig"
+
+UAdditiveControlRig::UAdditiveControlRig(const FObjectInitializer& ObjectInitializer)
+	: Super(ObjectInitializer)
+{
+
+}
+
+FName UAdditiveControlRig::GetControlName(const FName& InBoneName)
+{
+	if (InBoneName != NAME_None)
+	{
+		return FName(*(InBoneName.ToString() + TEXT("_CONTROL")));
+	}
+
+	// if bone name is coming as none, we don't append
+	return NAME_None;
+}
+
+FName UAdditiveControlRig::GetSpaceName(const FName& InBoneName)
+{
+	if (InBoneName != NAME_None)
+	{
+		return FName(*(InBoneName.ToString() + TEXT("_SPACE")));
+	}
+
+	// if bone name is coming as none, we don't append
+	return NAME_None;
+}
+
+void UAdditiveControlRig::ExecuteUnits(FRigUnitContext& InOutContext)
+{
+	FRigControlHierarchy& ControlHierarchy = GetControlHierarchy();
+	for (FRigUnit_AddBoneTransform& Unit : AddBoneRigUnits)
+	{
+		FName ControlName = GetControlName(Unit.Bone);
+		const int32 Index = ControlHierarchy.GetIndex(ControlName);
+		Unit.Transform = ControlHierarchy.GetLocalTransform(Index);
+		Unit.ExecuteContext.Hierarchy = GetHierarchy();
+		Unit.Execute(InOutContext);
+	}
+}
+
+void UAdditiveControlRig::Initialize(bool bInitRigUnits /*= true*/)
+{
+	Super::Initialize(bInitRigUnits);
+
+	// we do this after Initialize because Initialize will copy from CDO. 
+	// create hierarchy from the incoming skeleton
+	if (USkeletalMeshComponent* SkeletalMeshComponent = Cast<USkeletalMeshComponent>(GetObjectBinding()->GetBoundObject()))
+	{
+		CreateRigElements(SkeletalMeshComponent->SkeletalMesh);
+	}
+
+	// add units and initialize
+	AddBoneRigUnits.Reset();
+	FRigBoneHierarchy& BoneHierarchy = GetBoneHierarchy();
+
+	for (int32 BoneIndex = 0; BoneIndex < BoneHierarchy.Num(); ++BoneIndex)
+	{
+		FRigUnit_AddBoneTransform NewUnit;
+		NewUnit.Bone = BoneHierarchy[BoneIndex].Name;
+		NewUnit.bPropagateToChildren = true;
+		AddBoneRigUnits.Add(NewUnit);
+	}
+
+	// execute init
+	Execute(EControlRigState::Init);
+}
+
+void UAdditiveControlRig::CreateRigElements(const FReferenceSkeleton& InReferenceSkeleton, const FSmartNameMapping* InSmartNameMapping)
+{
+	FRigHierarchyContainer* Container = GetHierarchy();
+	Container->Reset();
+	FRigBoneHierarchy& BoneHierarchy = Container->BoneHierarchy;
+	BoneHierarchy.ImportSkeleton(InReferenceSkeleton, NAME_None, false, false, true, false);
+
+	if (InSmartNameMapping)
+	{
+		FRigCurveContainer& CurveContainer = Container->CurveContainer;
+		TArray<FName> NameArray;
+		InSmartNameMapping->FillNameArray(NameArray);
+		for (int32 Index = 0; Index < NameArray.Num(); ++Index)
+		{
+			CurveContainer.Add(NameArray[Index]);
+		}
+	}
+
+	// add control for all bone hierarchy 
+	for (int32 BoneIndex = 0; BoneIndex < BoneHierarchy.Num(); ++BoneIndex)
+	{
+		const FRigBone& RigBone = BoneHierarchy[BoneIndex];
+		FName BoneName = RigBone.Name;
+		FName ParentName = RigBone.ParentName;
+		FName SpaceName = GetSpaceName(BoneName);// name conflict?
+		FName ControlName = GetControlName(BoneName); // name conflict?
+		Container->SpaceHierarchy.Add(SpaceName, ERigSpaceType::Bone, BoneName);
+		Container->ControlHierarchy.Add(ControlName, ERigControlType::Transform, GetControlName(ParentName), SpaceName);
+	}
+
+	Container->Initialize(true);
+}
+
+void UAdditiveControlRig::CreateRigElements(const USkeletalMesh* InReferenceMesh)
+{
+	if (InReferenceMesh)
+	{
+		USkeleton* Skeleton = InReferenceMesh->Skeleton;
+		CreateRigElements(InReferenceMesh->RefSkeleton, (Skeleton) ? Skeleton->GetSmartNameContainer(USkeleton::AnimCurveMappingName) : nullptr);
+	}
+}
+
+#undef LOCTEXT_NAMESPACE
+
+
diff --git a/Engine/Plugins/Experimental/ControlRig/Source/ControlRig/Private/Rigs/RigBoneHierarchy.cpp b/Engine/Plugins/Experimental/ControlRig/Source/ControlRig/Private/Rigs/RigBoneHierarchy.cpp
--- a/Engine/Plugins/Experimental/ControlRig/Source/ControlRig/Private/Rigs/RigBoneHierarchy.cpp
+++ b/Engine/Plugins/Experimental/ControlRig/Source/ControlRig/Private/Rigs/RigBoneHierarchy.cpp
@@ -10,17 +10,21 @@
 ////////////////////////////////////////////////////////////////////////////////
 
 FRigBoneHierarchy::FRigBoneHierarchy()
-	:Container(nullptr)
+	: Container(nullptr)
+	, bSuspendNotifications(false)
 {
 }
 
 FRigBoneHierarchy& FRigBoneHierarchy::operator= (const FRigBoneHierarchy &InOther)
 {
 #if WITH_EDITOR
-	for (int32 Index = Num() - 1; Index >= 0; Index--)
+	if (!bSuspendNotifications)
 	{
-		FRigBone BoneToRemove = Bones[Index];
-		OnBoneRemoved.Broadcast(Container, FRigElementKey(BoneToRemove.Name, ERigElementType::Bone));
+		for (int32 Index = Num() - 1; Index >= 0; Index--)
+		{
+			FRigBone BoneToRemove = Bones[Index];
+			OnBoneRemoved.Broadcast(Container, FRigElementKey(BoneToRemove.Name, ERigElementType::Bone));
+		}
 	}
 #endif
 
@@ -30,9 +34,12 @@ FRigBoneHierarchy& FRigBoneHierarchy::operator= (const FRigBoneHierarchy &InOthe
 	RefreshMapping();
 
 #if WITH_EDITOR
-	for (const FRigBone& BoneAdded : Bones)
+	if (!bSuspendNotifications)
 	{
-		OnBoneAdded.Broadcast(Container, FRigElementKey(BoneAdded.Name, ERigElementType::Bone));
+		for (const FRigBone& BoneAdded : Bones)
+		{
+			OnBoneAdded.Broadcast(Container, FRigElementKey(BoneAdded.Name, ERigElementType::Bone));
+		}
 	}
 #endif
 
@@ -66,7 +73,10 @@ FRigBone& FRigBoneHierarchy::Add(const FName& InNewName, const FName& InParentNa
 	RefreshMapping();
 
 #if WITH_EDITOR
-	OnBoneAdded.Broadcast(Container, NewBone.GetElementKey());
+	if (!bSuspendNotifications)
+	{
+		OnBoneAdded.Broadcast(Container, NewBone.GetElementKey());
+	}
 #endif
 
 	int32 Index = GetIndex(NewBoneName);
@@ -116,7 +126,10 @@ bool FRigBoneHierarchy::Reparent(const FName& InName, const FName& InNewParentNa
 #if WITH_EDITOR
 		if (OldParentName != NewParentName)
 		{
-			OnBoneReparented.Broadcast(Container, FRigElementKey(InName, RigElementType()), OldParentName, NewParentName);
+			if (!bSuspendNotifications)
+			{
+				OnBoneReparented.Broadcast(Container, FRigElementKey(InName, RigElementType()), OldParentName, NewParentName);
+			}
 		}
 #endif
 		return Bones[GetIndex(InName)].ParentName == InNewParentName;
@@ -156,11 +169,14 @@ FRigBone FRigBoneHierarchy::Remove(const FName& InNameToRemove)
 	RefreshMapping();
 
 #if WITH_EDITOR
-	for (const FName& RemovedChildBone : RemovedChildBones)
+	if (!bSuspendNotifications)
 	{
-		OnBoneRemoved.Broadcast(Container, FRigElementKey(RemovedChildBone, RigElementType()));
+		for (const FName& RemovedChildBone : RemovedChildBones)
+		{
+			OnBoneRemoved.Broadcast(Container, FRigElementKey(RemovedChildBone, RigElementType()));
+		}
+		OnBoneRemoved.Broadcast(Container, RemovedBone.GetElementKey());
 	}
-	OnBoneRemoved.Broadcast(Container, RemovedBone.GetElementKey());
 #endif
 
 	return RemovedBone;
@@ -290,7 +306,7 @@ void FRigBoneHierarchy::SetInitialTransform(int32 InIndex, const FTransform& InT
 		FRigBone& Bone = Bones[InIndex];
 		Bone.InitialTransform = InTransform;
 		Bone.InitialTransform.NormalizeRotation();
-		RecalculateLocalTransform(Bone);
+		//RecalculateLocalTransform(Bone);
 	}
 }
 
@@ -358,10 +374,13 @@ FName FRigBoneHierarchy::Rename(const FName& InOldName, const FName& InNewName)
 			RefreshMapping();
 
 #if WITH_EDITOR
-			OnBoneRenamed.Broadcast(Container, RigElementType(), InOldName, NewName);
-			for (const FName& ReparentedBone : ReparentedBones)
+			if (!bSuspendNotifications)
 			{
-				OnBoneReparented.Broadcast(Container, FRigElementKey(ReparentedBone, RigElementType()), InOldName, NewName);
+				OnBoneRenamed.Broadcast(Container, RigElementType(), InOldName, NewName);
+				for (const FName& ReparentedBone : ReparentedBones)
+				{
+					OnBoneReparented.Broadcast(Container, FRigElementKey(ReparentedBone, RigElementType()), InOldName, NewName);
+				}
 			}
 			if (bWasSelected)
 			{
@@ -465,7 +484,7 @@ void FRigBoneHierarchy::RefreshMapping()
 	}
 }
 
-void FRigBoneHierarchy::Initialize()
+void FRigBoneHierarchy::Initialize(bool bResetTransforms)
 {
 	RefreshMapping();
 
@@ -477,8 +496,11 @@ void FRigBoneHierarchy::Initialize()
 	// initialize transform
 	for (int32 Index = 0; Index < Bones.Num(); ++Index)
 	{
-		Bones[Index].GlobalTransform = Bones[Index].InitialTransform;
-		RecalculateLocalTransform(Bones[Index]);
+		if (bResetTransforms)
+		{
+			Bones[Index].GlobalTransform = Bones[Index].InitialTransform;
+			RecalculateLocalTransform(Bones[Index]);
+		}
 
 		// update children
 		GetChildren(Index, Bones[Index].Dependents, false);
@@ -567,7 +589,10 @@ bool FRigBoneHierarchy::Select(const FName& InName, bool bSelect)
 		Selection.Remove(InName);
 	}
 
-	OnBoneSelected.Broadcast(Container, FRigElementKey(InName, RigElementType()), bSelect);
+	if (!bSuspendNotifications)
+	{
+		OnBoneSelected.Broadcast(Container, FRigElementKey(InName, RigElementType()), bSelect);
+	}
 
 	return true;
 }
@@ -595,8 +620,12 @@ bool FRigBoneHierarchy::IsSelected(const FName& InName) const
 	return Selection.Contains(InName);
 }
 
-TArray<FRigElementKey> FRigBoneHierarchy::ImportSkeleton(const FReferenceSkeleton& InSkeleton, const FName& InNameSpace, bool bReplaceExistingBones, bool bRemoveObsoleteBones, bool bSelectBones)
+#endif // WITH_EDITOR
+
+TArray<FRigElementKey> FRigBoneHierarchy::ImportSkeleton(const FReferenceSkeleton& InSkeleton, const FName& InNameSpace, bool bReplaceExistingBones, bool bRemoveObsoleteBones, bool bSelectBones, bool bNotify)
 {
+	TGuardValue<bool> SuspendNotification(bSuspendNotifications, !bNotify);
+
 	TArray<FRigElementKey> AddedBones;
 	TArray<FName> BoneNamesToSelect;
 	TMap<FName, FName> BoneNameMap;
@@ -606,13 +635,13 @@ TArray<FRigElementKey> FRigBoneHierarchy::ImportSkeleton(const FReferenceSkeleto
 
 	struct Local
 	{
-		static FName DetermineBoneName(const FName& InBoneName, const FName& InNameSpace)
+		static FName DetermineBoneName(const FName& InBoneName, const FName& InLocalNameSpace)
 		{
-			if (InNameSpace == NAME_None || InBoneName == NAME_None)
+			if (InLocalNameSpace == NAME_None || InBoneName == NAME_None)
 			{
 				return InBoneName;
 			}
-			return *FString::Printf(TEXT("%s_%s"), *InNameSpace.ToString(), *InBoneName.ToString());
+			return *FString::Printf(TEXT("%s_%s"), *InLocalNameSpace.ToString(), *InBoneName.ToString());
 		}
 	};
 
@@ -730,6 +759,7 @@ TArray<FRigElementKey> FRigBoneHierarchy::ImportSkeleton(const FReferenceSkeleto
 		}
 	}
 
+#if WITH_EDITOR
 	if (bSelectBones)
 	{
 		ClearSelection();
@@ -738,11 +768,12 @@ TArray<FRigElementKey> FRigBoneHierarchy::ImportSkeleton(const FReferenceSkeleto
 			Select(BoneNameToSelect);
 		}
 	}
+#endif // WITH_EDITOR
 
 	Initialize();
 
+
+
 	return AddedBones;
 }
 
-
-#endif
\ No newline at end of file
