COPY from //Dev-Anim to //Dev-Main

diff --git a/Engine/Plugins/Experimental/ControlRig/Source/ControlRig/Private/Rigs/RigControlHierarchy.cpp b/Engine/Plugins/Experimental/ControlRig/Source/ControlRig/Private/Rigs/RigControlHierarchy.cpp
--- a/Engine/Plugins/Experimental/ControlRig/Source/ControlRig/Private/Rigs/RigControlHierarchy.cpp
+++ b/Engine/Plugins/Experimental/ControlRig/Source/ControlRig/Private/Rigs/RigControlHierarchy.cpp
@@ -5,6 +5,151 @@
 #include "HelperUtil.h"
 #include "Settings/ControlRigSettings.h"
 
+////////////////////////////////////////////////////////////////////////////////
+// FRigControl
+////////////////////////////////////////////////////////////////////////////////
+
+void FRigControl::ApplyLimits(FRigControlValue& InOutValue)
+{
+	if (!bLimitTranslation && !bLimitRotation && !bLimitScale)
+	{
+		return;
+	}
+
+	switch(ControlType)
+	{
+		case ERigControlType::Float:
+		{
+			if (bLimitTranslation)
+			{
+				float& ValueRef = InOutValue.GetRef<float>();
+				ValueRef = Clamp(ValueRef, MinimumValue.Get<float>(), MaximumValue.Get<float>());
+			}
+			break;
+		}
+		case ERigControlType::Vector2D:
+		{
+			if (bLimitTranslation)
+			{
+				FVector2D& ValueRef = InOutValue.GetRef<FVector2D>();
+				const FVector2D& Min = MinimumValue.GetRef<FVector2D>();
+				const FVector2D& Max = MaximumValue.GetRef<FVector2D>();
+				ValueRef.X = Clamp(ValueRef.X, Min.X, Max.X);
+				ValueRef.Y = Clamp(ValueRef.Y, Min.Y, Max.Y);
+			}
+			break;
+		}
+		case ERigControlType::Position:
+		{
+			if (bLimitTranslation)
+			{
+				FVector& ValueRef = InOutValue.GetRef<FVector>();
+				const FVector& Min = MinimumValue.GetRef<FVector>();
+				const FVector& Max = MaximumValue.GetRef<FVector>();
+				ValueRef.X = Clamp(ValueRef.X, Min.X, Max.X);
+				ValueRef.Y = Clamp(ValueRef.Y, Min.Y, Max.Y);
+				ValueRef.Z = Clamp(ValueRef.Z, Min.Z, Max.Z);
+			}
+			break;
+		}
+		case ERigControlType::Scale:
+		{
+			if (bLimitScale)
+			{
+				FVector& ValueRef = InOutValue.GetRef<FVector>();
+				const FVector& Min = MinimumValue.GetRef<FVector>();
+				const FVector& Max = MaximumValue.GetRef<FVector>();
+				ValueRef.X = Clamp(ValueRef.X, Min.X, Max.X);
+				ValueRef.Y = Clamp(ValueRef.Y, Min.Y, Max.Y);
+				ValueRef.Z = Clamp(ValueRef.Z, Min.Z, Max.Z);
+			}
+			break;
+		}
+		case ERigControlType::Rotator:
+		{
+			if (bLimitRotation)
+			{
+				FRotator& ValueRef = InOutValue.GetRef<FRotator>();
+				const FRotator& Min = MinimumValue.GetRef<FRotator>();
+				const FRotator& Max = MaximumValue.GetRef<FRotator>();
+				ValueRef.Pitch = Clamp(ValueRef.Pitch, Min.Pitch, Max.Pitch);
+				ValueRef.Yaw = Clamp(ValueRef.Yaw, Min.Yaw, Max.Yaw);
+				ValueRef.Roll = Clamp(ValueRef.Roll, Min.Roll, Max.Roll);
+			}
+			break;
+		}
+		case ERigControlType::Transform:
+		{
+			FTransform& ValueRef = InOutValue.GetRef<FTransform>();
+			const FTransform& Min = MinimumValue.GetRef<FTransform>();
+			const FTransform& Max = MaximumValue.GetRef<FTransform>();
+
+			if (bLimitTranslation)
+			{
+				ValueRef.SetLocation(FVector(
+					Clamp(ValueRef.GetLocation().X, Min.GetLocation().X, Max.GetLocation().X),
+					Clamp(ValueRef.GetLocation().Y, Min.GetLocation().Y, Max.GetLocation().Y),
+					Clamp(ValueRef.GetLocation().Z, Min.GetLocation().Z, Max.GetLocation().Z)
+				));
+			}
+			if (bLimitRotation)
+			{
+				FRotator Rotator = ValueRef.GetRotation().Rotator();
+				FRotator MinRotator = Min.GetRotation().Rotator();
+				FRotator MaxRotator = Max.GetRotation().Rotator();
+
+				ValueRef.SetRotation(FQuat(FRotator(
+					Clamp(Rotator.Pitch, MinRotator.Pitch, MaxRotator.Pitch),
+					Clamp(Rotator.Yaw, MinRotator.Yaw, MaxRotator.Yaw),
+					Clamp(Rotator.Roll, MinRotator.Roll, MaxRotator.Roll)
+				)));
+			}
+			if (bLimitScale)
+			{
+				ValueRef.SetScale3D(FVector(
+					Clamp(ValueRef.GetScale3D().X, Min.GetScale3D().X, Max.GetScale3D().X),
+					Clamp(ValueRef.GetScale3D().Y, Min.GetScale3D().Y, Max.GetScale3D().Y),
+					Clamp(ValueRef.GetScale3D().Z, Min.GetScale3D().Z, Max.GetScale3D().Z)
+				));
+			}
+			break;
+		}
+		case ERigControlType::TransformNoScale:
+		{
+			FTransformNoScale& ValueRef = InOutValue.GetRef<FTransformNoScale>();
+			const FTransformNoScale& Min = MinimumValue.GetRef<FTransformNoScale>();
+			const FTransformNoScale& Max = MaximumValue.GetRef<FTransformNoScale>();
+
+			if (bLimitTranslation)
+			{
+				ValueRef.Location = FVector(
+					Clamp(ValueRef.Location.X, Min.Location.X, Max.Location.X),
+					Clamp(ValueRef.Location.Y, Min.Location.Y, Max.Location.Y),
+					Clamp(ValueRef.Location.Z, Min.Location.Z, Max.Location.Z)
+				);
+			}
+			if (bLimitRotation)
+			{
+				FRotator Rotator = ValueRef.Rotation.Rotator();
+				FRotator MinRotator = Min.Rotation.Rotator();
+				FRotator MaxRotator = Max.Rotation.Rotator();
+
+				ValueRef.Rotation = FQuat(FRotator(
+					Clamp(Rotator.Pitch, MinRotator.Pitch, MaxRotator.Pitch),
+					Clamp(Rotator.Yaw, MinRotator.Yaw, MaxRotator.Yaw),
+					Clamp(Rotator.Roll, MinRotator.Roll, MaxRotator.Roll)
+				));
+			}
+			break;
+		}
+		case ERigControlType::Bool:
+		default:
+		{
+			break;
+		}
+	}
+}
+
 ////////////////////////////////////////////////////////////////////////////////
 // FRigControlHierarchy
 ////////////////////////////////////////////////////////////////////////////////
@@ -55,12 +200,10 @@ FRigControl& FRigControlHierarchy::Add(
 	ERigControlType InControlType,
 	const FName& InParentName,
 	const FName& InSpaceName,
-	const FRigControlValue& InValue
-#if WITH_EDITORONLY_DATA
-	,const FName& InGizmoName,
+	const FRigControlValue& InValue,
+	const FName& InGizmoName,
 	const FTransform& InGizmoTransform,
 	const FLinearColor& InGizmoColor
-#endif
 )
 {
 	FRigControl NewControl;
@@ -71,12 +214,10 @@ FRigControl& FRigControlHierarchy::Add(
 	NewControl.SpaceIndex = INDEX_NONE;
 	NewControl.SpaceName = NAME_None;
 	NewControl.InitialValue = InValue;
-	NewControl.Value = InValue;
-#if WITH_EDITORONLY_DATA
+	NewControl.Value = FRigControlValue();
 	NewControl.GizmoName = InGizmoName;
 	NewControl.GizmoTransform = InGizmoTransform;
 	NewControl.GizmoColor = InGizmoColor;
-#endif
 
 	FName NewControlName = NewControl.Name;
 	Controls.Add(NewControl);
@@ -297,84 +438,209 @@ FTransform FRigControlHierarchy::GetGlobalTransform(int32 InIndex) const
 
 	if (Controls.IsValidIndex(InIndex))
 	{
-		const FRigControl& Control = Controls[InIndex];
-
 		FTransform ParentTransform = GetParentTransform(InIndex);
-		switch (Control.ControlType)
-		{
-			case ERigControlType::Bool:
-			case ERigControlType::Float:
-			case ERigControlType::Vector2D:
-			{
-				return ParentTransform;
-			}
-			case ERigControlType::Position:
-			{
-				FTransform Transform;
-				Transform.SetLocation(Control.Value.Get<FVector>());
-				Transform = Transform * ParentTransform;
-				Transform.NormalizeRotation();
-				return Transform;
-			}
-			case ERigControlType::Scale:
-			{
-				FTransform Transform;
-				Transform.SetScale3D(Control.Value.Get<FVector>());
-				Transform = Transform * ParentTransform;
-				Transform.NormalizeRotation();
-				return Transform;
-			}
-			case ERigControlType::Quat:
-			{
-				FTransform Transform;
-				Transform.SetRotation(Control.Value.Get<FQuat>());
-				Transform = Transform * ParentTransform;
-				Transform.NormalizeRotation();
-				return Transform;
-			}
-			case ERigControlType::Rotator:
-			{
-				FTransform Transform;
-				Transform.SetRotation(FQuat(Control.Value.Get<FRotator>()));
-				Transform = Transform * ParentTransform;
-				Transform.NormalizeRotation();
-				return Transform;
-			}
-			case ERigControlType::Transform:
-			{
-				FTransform Transform = Control.Value.Get<FTransform>() * ParentTransform;
-				Transform.NormalizeRotation();
-				return Transform;
-			}
-			default:
-			{
-				ensure(false);
-				break;
-			}
-		}
+		FTransform Transform = GetLocalTransform(InIndex) * ParentTransform;
+		Transform.NormalizeRotation();
+		return Transform;
 	}
 
 	return FTransform::Identity;
 }
 
-void FRigControlHierarchy::SetLocalTransform(const FName& InName, const FTransform& InTransform)
+void FRigControlHierarchy::SetLocalTransform(const FName& InName, const FTransform& InTransform, ERigControlValueType InValueType)
 {
-	SetValue<FTransform>(InName, InTransform);
+	SetLocalTransform(GetIndex(InName), InTransform, InValueType);
 }
 
-void FRigControlHierarchy::SetLocalTransform(int32 InIndex, const FTransform& InTransform)
+void FRigControlHierarchy::SetLocalTransform(int32 InIndex, const FTransform& InTransform, ERigControlValueType InValueType)
 {
-	SetValue<FTransform>(InIndex, InTransform);
+	const FRigControl& Control = Controls[InIndex];
+	switch (Control.ControlType)
+	{
+		case ERigControlType::Bool:
+		{
+			SetValue<bool>(InIndex, InTransform.GetLocation().X > SMALL_NUMBER, InValueType);
+			break;
+		}
+		case ERigControlType::Float:
+		{
+			switch (Control.PrimaryAxis)
+			{
+				case ERigControlAxis::X:
+				{
+					SetValue<float>(InIndex, InTransform.GetLocation().X, InValueType);
+					break;
+				}
+				case ERigControlAxis::Y:
+				{
+					SetValue<float>(InIndex, InTransform.GetLocation().Y, InValueType);
+					break;
+				}
+				case ERigControlAxis::Z:
+				{
+					SetValue<float>(InIndex, InTransform.GetLocation().Z, InValueType);
+					break;
+				}
+			}
+			break;
+		}
+		case ERigControlType::Vector2D:
+		{
+			FVector Location = InTransform.GetLocation();
+			switch (Control.PrimaryAxis)
+			{
+				case ERigControlAxis::X:
+				{
+					SetValue<FVector2D>(InIndex, FVector2D(Location.Y, Location.Z), InValueType);
+					break;
+				}
+				case ERigControlAxis::Y:
+				{
+					SetValue<FVector2D>(InIndex, FVector2D(Location.X, Location.Z), InValueType);
+					break;
+				}
+				case ERigControlAxis::Z:
+				{
+					SetValue<FVector2D>(InIndex, FVector2D(Location.X, Location.Y), InValueType);
+					break;
+				}
+			}
+			break;
+		}
+		case ERigControlType::Position:
+		{
+			SetValue<FVector>(InIndex, InTransform.GetLocation(), InValueType);
+			break;
+		}
+		case ERigControlType::Scale:
+		{
+			SetValue<FVector>(InIndex, InTransform.GetScale3D(), InValueType);
+			break;
+		}
+		case ERigControlType::Rotator:
+		{
+			SetValue<FRotator>(InIndex, InTransform.GetRotation().Rotator(), InValueType);
+			break;
+		}
+		case ERigControlType::Transform:
+		{
+			SetValue<FTransform>(InIndex, InTransform, InValueType);
+			break;
+		}
+		case ERigControlType::TransformNoScale:
+		{
+			FTransformNoScale NoScale = InTransform;
+			SetValue<FTransformNoScale>(InIndex, NoScale, InValueType);
+			break;
+		}
+		default:
+		{
+			ensure(false);
+			break;
+		}
+	}
 }
 
-FTransform FRigControlHierarchy::GetLocalTransform(const FName& InName) const
+FTransform FRigControlHierarchy::GetLocalTransform(const FName& InName, ERigControlValueType InValueType) const
 {
-	return GetValue<FTransform>(InName);
+	return GetLocalTransform(GetIndex(InName), InValueType);
 }
 
-FTransform FRigControlHierarchy::GetLocalTransform(int32 InIndex) const
+FTransform FRigControlHierarchy::GetLocalTransform(int32 InIndex, ERigControlValueType InValueType) const
 {
-	return GetValue<FTransform>(InIndex);
+	const FRigControl& Control = Controls[InIndex];
+	switch (Control.ControlType)
+	{
+		case ERigControlType::Bool:
+		{
+			FTransform Transform;
+			Transform.SetLocation(FVector(Control.GetValue(InValueType).Get<bool>() ? 1.f : 0.f, 0.f, 0.f));
+			return Transform;
+		}
+		case ERigControlType::Float:
+		{
+			float Value = Control.GetValue(InValueType).Get<float>();
+			FTransform Transform = FTransform::Identity;
+			switch (Control.PrimaryAxis)
+			{
+				case ERigControlAxis::X:
+				{
+					Transform.SetLocation(FVector(Value, 0.f, 0.f));
+					break;
+				}
+				case ERigControlAxis::Y:
+				{
+					Transform.SetLocation(FVector(0.f, Value, 0.f));
+					break;
+				}
+				case ERigControlAxis::Z:
+				{
+					Transform.SetLocation(FVector(0.f, 0.f, Value));
+					break;
+				}
+			}
+			return Transform;
+		}
+		case ERigControlType::Vector2D:
+		{
+			FVector2D Value = Control.GetValue(InValueType).Get<FVector2D>();
+			FTransform Transform = FTransform::Identity;
+			switch (Control.PrimaryAxis)
+			{
+				case ERigControlAxis::X:
+				{
+					Transform.SetLocation(FVector(0.f, Value.X, Value.Y));
+					break;
+				}
+				case ERigControlAxis::Y:
+				{
+					Transform.SetLocation(FVector(Value.X, 0.f, Value.Y));
+					break;
+				}
+				case ERigControlAxis::Z:
+				{
+					Transform.SetLocation(FVector(Value.X, Value.Y, 0.f));
+					break;
+				}
+			}
+			return Transform;
+		}
+		case ERigControlType::Position:
+		{
+			FTransform Transform;
+			Transform.SetLocation(Control.GetValue(InValueType).Get<FVector>());
+			return Transform;
+		}
+		case ERigControlType::Scale:
+		{
+			FTransform Transform;
+			Transform.SetScale3D(Control.GetValue(InValueType).Get<FVector>());
+			return Transform;
+		}
+		case ERigControlType::Rotator:
+		{
+			FTransform Transform;
+			Transform.SetRotation(FQuat(Control.GetValue(InValueType).Get<FRotator>()));
+			return Transform;
+		}
+		case ERigControlType::Transform:
+		{
+			return Control.GetValue(InValueType).Get<FTransform>();
+		}
+		case ERigControlType::TransformNoScale:
+		{
+			FTransformNoScale TransformNoScale = Control.GetValue(InValueType).Get<FTransformNoScale>();
+			FTransform Transform = TransformNoScale;
+			Transform.NormalizeRotation();
+			return Transform;
+		}
+		default:
+		{
+			ensure(false);
+			break;
+		}
+	}
+	return FTransform::Identity;
 }
 
 FTransform FRigControlHierarchy::GetParentTransform(int32 InIndex) const
@@ -391,81 +657,50 @@ FTransform FRigControlHierarchy::GetParentTransform(int32 InIndex) const
 		ParentTransform = GetGlobalTransform(Control.ParentIndex);
 	}
 
-	return ParentTransform;
+	return GetLocalTransform(InIndex, ERigControlValueType::Initial) * ParentTransform;
 }
 
-void FRigControlHierarchy::SetValue(const FName& InName, const FRigControlValue& InValue)
+void FRigControlHierarchy::SetValue(const FName& InName, const FRigControlValue& InValue, ERigControlValueType InValueType)
 {
-	SetValue(GetIndex(InName), InValue);
+	SetValue(GetIndex(InName), InValue, InValueType);
 }
 
-void FRigControlHierarchy::SetValue(int32 InIndex, const FRigControlValue& InValue)
+void FRigControlHierarchy::SetValue(int32 InIndex, const FRigControlValue& InValue, ERigControlValueType InValueType)
 {
 	if (Controls.IsValidIndex(InIndex))
 	{
 		FRigControl& Control = Controls[InIndex];
-		Control.Value = InValue;
 
-		if (Control.ControlType == ERigControlType::Transform)
-		{
-			Control.Value.GetRef<FTransform>().NormalizeRotation();
-		}
-		else if (Control.ControlType == ERigControlType::Quat)
+		Control.GetValue(InValueType) = InValue;
+
+		if (InValueType == ERigControlValueType::Current)
 		{
-			Control.Value.GetRef<FQuat>().Normalize();
+			Control.ApplyLimits(Control.GetValue(InValueType));
 		}
-	}
-}
-
-FRigControlValue FRigControlHierarchy::GetValue(const FName& InName) const
-{
-	return GetValue(GetIndex(InName));
-}
-
-FRigControlValue FRigControlHierarchy::GetValue(int32 InIndex) const
-{
-	if (Controls.IsValidIndex(InIndex))
-	{
-		return Controls[InIndex].Value;
-	}
-
-	return FRigControlValue();
-}
-
-void FRigControlHierarchy::SetInitialValue(const FName& InName, const FRigControlValue& InValue)
-{
-	SetInitialValue(GetIndex(InName), InValue);
-}
-
-void FRigControlHierarchy::SetInitialValue(int32 InIndex, const FRigControlValue& InValue)
-{
-	if (Controls.IsValidIndex(InIndex))
-	{
-		FRigControl& Control = Controls[InIndex];
-		Control.InitialValue = InValue;
 
 		if (Control.ControlType == ERigControlType::Transform)
 		{
-			Control.InitialValue.GetRef<FTransform>().NormalizeRotation();
+			Control.GetValue(InValueType).GetRef<FTransform>().NormalizeRotation();
 		}
-		else if (Control.ControlType == ERigControlType::Quat)
+		else if (Control.ControlType == ERigControlType::TransformNoScale)
 		{
-			Control.InitialValue.GetRef<FQuat>().Normalize();
+			Control.GetValue(InValueType).GetRef<FTransformNoScale>().Rotation.Normalize();
 		}
 	}
 }
 
-FRigControlValue FRigControlHierarchy::GetInitialValue(const FName& InName) const
+FRigControlValue FRigControlHierarchy::GetValue(const FName& InName, ERigControlValueType InValueType) const
 {
-	return GetInitialValue(GetIndex(InName));
+	return GetValue(GetIndex(InName), InValueType);
 }
 
-FRigControlValue FRigControlHierarchy::GetInitialValue(int32 InIndex) const
+FRigControlValue FRigControlHierarchy::GetValue(int32 InIndex, ERigControlValueType InValueType) const
 {
 	if (Controls.IsValidIndex(InIndex))
 	{
-		return Controls[InIndex].InitialValue;
+		return Controls[InIndex].GetValue(InValueType);
 	}
+
 	return FRigControlValue();
 }
 
@@ -527,32 +762,72 @@ FTransform FRigControlHierarchy::GetInitialGlobalTransform(int32 InIndex) const
 		switch (Control.ControlType)
 		{
 			case ERigControlType::Bool:
+			{
+				break;
+			}
 			case ERigControlType::Float:
+			{
+				float Value = Control.InitialValue.Get<float>();
+				FTransform Transform = FTransform::Identity;
+				switch (Control.PrimaryAxis)
+				{
+					case ERigControlAxis::X:
+					{
+						Transform.SetLocation(FVector(Value, 0.f, 0.f));
+						break;
+					}
+					case ERigControlAxis::Y:
+					{
+						Transform.SetLocation(FVector(0.f, Value, 0.f));
+						break;
+					}
+					case ERigControlAxis::Z:
+					{
+						Transform.SetLocation(FVector(0.f, 0.f, Value));
+						break;
+					}
+				}
+				return Transform * ParentTransform;
+			}
 			case ERigControlType::Vector2D:
 			{
-				return ParentTransform;
+				FVector2D Value = Control.InitialValue.Get<FVector2D>();
+				FTransform Transform = FTransform::Identity;
+				switch (Control.PrimaryAxis)
+				{
+					case ERigControlAxis::X:
+					{
+						Transform.SetLocation(FVector(0.f, Value.X, Value.Y));
+						break;
+					}
+					case ERigControlAxis::Y:
+					{
+						Transform.SetLocation(FVector(Value.X, 0.f, Value.Y));
+						break;
+					}
+					case ERigControlAxis::Z:
+					{
+						Transform.SetLocation(FVector(Value.X, Value.Y, 0.f));
+						break;
+					}
+				}
+				return Transform * ParentTransform;
 			}
 			case ERigControlType::Position:
 			{
-				FTransform Transform;
+				FTransform Transform = FTransform::Identity;
 				Transform.SetLocation(Control.InitialValue.Get<FVector>());
 				return Transform * ParentTransform;
 			}
 			case ERigControlType::Scale:
 			{
-				FTransform Transform;
+				FTransform Transform = FTransform::Identity;
 				Transform.SetScale3D(Control.InitialValue.Get<FVector>());
 				return Transform * ParentTransform;
 			}
-			case ERigControlType::Quat:
-			{
-				FTransform Transform;
-				Transform.SetRotation(Control.InitialValue.Get<FQuat>());
-				return Transform * ParentTransform;
-			}
 			case ERigControlType::Rotator:
 			{
-				FTransform Transform;
+				FTransform Transform = FTransform::Identity;
 				Transform.SetRotation(FQuat(Control.InitialValue.Get<FRotator>()));
 				return Transform * ParentTransform;
 			}
@@ -560,6 +835,10 @@ FTransform FRigControlHierarchy::GetInitialGlobalTransform(int32 InIndex) const
 			{
 				return Control.InitialValue.Get<FTransform>() * ParentTransform;
 			}
+			case ERigControlType::TransformNoScale:
+			{
+				return Control.InitialValue.Get<FTransformNoScale>().ToFTransform() * ParentTransform;
+			}
 		}
 	}
 	return FTransform::Identity;
@@ -629,7 +908,7 @@ void FRigControlHierarchy::RefreshMapping()
 	}
 }
 
-void FRigControlHierarchy::Initialize()
+void FRigControlHierarchy::Initialize(bool bResetTransforms)
 {
 	RefreshMapping();
 
@@ -646,7 +925,17 @@ void FRigControlHierarchy::Initialize()
 	// initialize transform
 	for (int32 Index = 0; Index < Controls.Num(); ++Index)
 	{
-		Controls[Index].Value = Controls[Index].InitialValue;
+		if (bResetTransforms)
+		{
+			if (Controls[Index].ControlType == ERigControlType::Bool)
+			{
+				Controls[Index].Value = Controls[Index].InitialValue;
+			}
+			else
+			{
+				SetLocalTransform(Index, FTransform::Identity, ERigControlValueType::Current);
+			}
+		}
 
 		// update children
 		GetChildren(Index, Controls[Index].Dependents, false);
@@ -663,7 +952,14 @@ void FRigControlHierarchy::ResetValues()
 	// initialize transform
 	for (int32 Index = 0; Index < Controls.Num(); ++Index)
 	{
-		Controls[Index].Value = Controls[Index].InitialValue;
+		if (Controls[Index].ControlType == ERigControlType::Bool)
+		{
+			Controls[Index].Value = Controls[Index].InitialValue;
+		}
+		else
+		{
+			SetLocalTransform(Index, FTransform::Identity, ERigControlValueType::Current);
+		}
 	}
 }
 
diff --git a/Engine/Plugins/Experimental/ControlRig/Source/ControlRig/Private/Rigs/RigCurveContainer.cpp b/Engine/Plugins/Experimental/ControlRig/Source/ControlRig/Private/Rigs/RigCurveContainer.cpp
--- a/Engine/Plugins/Experimental/ControlRig/Source/ControlRig/Private/Rigs/RigCurveContainer.cpp
+++ b/Engine/Plugins/Experimental/ControlRig/Source/ControlRig/Private/Rigs/RigCurveContainer.cpp
@@ -3,23 +3,28 @@
 #include "Rigs/RigCurveContainer.h"
 #include "ControlRig.h"
 #include "HelperUtil.h"
+#include "Animation/Skeleton.h"
 
 ////////////////////////////////////////////////////////////////////////////////
 // FRigCurveContainer
 ////////////////////////////////////////////////////////////////////////////////
 
 FRigCurveContainer::FRigCurveContainer()
-	:Container(nullptr)
+	: Container(nullptr)
+	, bSuspendNotifications(false)
 {
 }
 
 FRigCurveContainer& FRigCurveContainer::operator= (const FRigCurveContainer &InOther)
 {
 #if WITH_EDITOR
-	for (int32 Index = Num() - 1; Index >= 0; Index--)
+	if (!bSuspendNotifications)
 	{
-		FRigCurve CurveToRemove = Curves[Index];
-		OnCurveRemoved.Broadcast(Container, FRigElementKey(CurveToRemove.Name, ERigElementType::Curve));
+		for (int32 Index = Num() - 1; Index >= 0; Index--)
+		{
+			FRigCurve CurveToRemove = Curves[Index];
+			OnCurveRemoved.Broadcast(Container, FRigElementKey(CurveToRemove.Name, ERigElementType::Curve));
+		}
 	}
 #endif
 
@@ -29,9 +34,12 @@ FRigCurveContainer& FRigCurveContainer::operator= (const FRigCurveContainer &InO
 	RefreshMapping();
 
 #if WITH_EDITOR
-	for (const FRigCurve& CurveAdded : Curves)
+	if (!bSuspendNotifications)
 	{
-		OnCurveAdded.Broadcast(Container, FRigElementKey(CurveAdded.Name, ERigElementType::Curve));
+		for (const FRigCurve& CurveAdded : Curves)
+		{
+			OnCurveAdded.Broadcast(Container, FRigElementKey(CurveAdded.Name, ERigElementType::Curve));
+		}
 	}
 #endif
 
@@ -61,7 +69,10 @@ FRigCurve& FRigCurveContainer::Add(const FName& InNewName, float InValue)
 	RefreshMapping();
 
 #if WITH_EDITOR
-	OnCurveAdded.Broadcast(Container, NewCurve.GetElementKey());
+	if (!bSuspendNotifications)
+	{
+		OnCurveAdded.Broadcast(Container, NewCurve.GetElementKey());
+	}
 #endif
 
 	int32 Index = GetIndex(NewCurveName);
@@ -84,7 +95,10 @@ FRigCurve FRigCurveContainer::Remove(const FName& InName)
 	RefreshMapping();
 
 #if WITH_EDITOR
-	OnCurveRemoved.Broadcast(Container, RemovedCurve.GetElementKey());
+	if (!bSuspendNotifications)
+	{
+		OnCurveRemoved.Broadcast(Container, RemovedCurve.GetElementKey());
+	}
 #endif
 
 	return RemovedCurve;
@@ -165,7 +179,10 @@ FName FRigCurveContainer::Rename(const FName& InOldName, const FName& InNewName)
 			RefreshMapping();
 
 #if WITH_EDITOR
-			OnCurveRenamed.Broadcast(Container, RigElementType(), InOldName, NewName);
+			if (!bSuspendNotifications)
+			{
+				OnCurveRenamed.Broadcast(Container, RigElementType(), InOldName, NewName);
+			}
 			if(bWasSelected)
 			{
 				Select(NewName, true);
@@ -246,7 +263,10 @@ bool FRigCurveContainer::Select(const FName& InName, bool bSelect)
 		Selection.Remove(InName);
 	}
 
-	OnCurveSelected.Broadcast(Container, FRigElementKey(InName, RigElementType()), bSelect);
+	if (!bSuspendNotifications)
+	{
+		OnCurveSelected.Broadcast(Container, FRigElementKey(InName, RigElementType()), bSelect);
+	}
 
 	return true;
 }
@@ -274,4 +294,32 @@ bool FRigCurveContainer::IsSelected(const FName& InName) const
 	return Selection.Contains(InName);
 }
 
+TArray<FRigElementKey> FRigCurveContainer::ImportCurvesFromSkeleton(const USkeleton* InSkeleton, const FName& InNameSpace, bool bRemoveObsoleteCurves, bool bSelectCurves, bool bNotify)
+{
+	check(InSkeleton);
+
+	TGuardValue<bool> SuspendNotification(bSuspendNotifications, !bNotify);
+
+	TArray<FRigElementKey> Keys;
+
+	const FSmartNameMapping* SmartNameMapping = InSkeleton->GetSmartNameContainer(USkeleton::AnimCurveMappingName);
+
+	TArray<FName> NameArray;
+	SmartNameMapping->FillNameArray(NameArray);
+	for (int32 Index = 0; Index < NameArray.Num(); ++Index)
+	{
+		FName Name = NameArray[Index];
+		if (!InNameSpace.IsNone())
+		{
+			Name = *FString::Printf(TEXT("%s::%s"), *InNameSpace.ToString(), *Name.ToString());
+		}
+		Add(Name);
+		Select(Name, true);
+		Keys.Add(FRigElementKey(Name, ERigElementType::Curve));
+	}
+
+	return Keys;
+}
+
+
 #endif
\ No newline at end of file
diff --git a/Engine/Plugins/Experimental/ControlRig/Source/ControlRig/Private/Rigs/RigHierarchyContainer.cpp b/Engine/Plugins/Experimental/ControlRig/Source/ControlRig/Private/Rigs/RigHierarchyContainer.cpp
--- a/Engine/Plugins/Experimental/ControlRig/Source/ControlRig/Private/Rigs/RigHierarchyContainer.cpp
+++ b/Engine/Plugins/Experimental/ControlRig/Source/ControlRig/Private/Rigs/RigHierarchyContainer.cpp
@@ -23,7 +23,7 @@ FRigHierarchyContainer& FRigHierarchyContainer::operator= (const FRigHierarchyCo
 	return *this;
 }
 
-void FRigHierarchyContainer::Initialize()
+void FRigHierarchyContainer::Initialize(bool bResetTransforms)
 {
 	BoneHierarchy.Container = this;
 	SpaceHierarchy.Container = this;
@@ -90,12 +90,15 @@ void FRigHierarchyContainer::Initialize()
 
 #endif
 
-	BoneHierarchy.Initialize();
-	SpaceHierarchy.Initialize();
-	ControlHierarchy.Initialize();
+	BoneHierarchy.Initialize(bResetTransforms);
+	SpaceHierarchy.Initialize(bResetTransforms);
+	ControlHierarchy.Initialize(bResetTransforms);
 	CurveContainer.Initialize();
 
-	ResetTransforms();
+	if (bResetTransforms)
+	{
+		ResetTransforms();
+	}
 }
 
 void FRigHierarchyContainer::Reset()
@@ -712,6 +715,18 @@ TArray<FRigElementKey> FRigHierarchyContainer::ImportFromText(const FString& InC
 						}
 
 						FRigControl& NewElement = ControlHierarchy.Add(Element.Name, Element.ControlType, ParentName, SpaceName, Element.InitialValue, Element.GizmoName, Element.GizmoTransform, Element.GizmoColor);
+
+						// copy additional members
+						NewElement.bAnimatable = Element.bAnimatable;
+						NewElement.PrimaryAxis = Element.PrimaryAxis;
+						NewElement.bLimitTranslation = Element.bLimitTranslation;
+						NewElement.bLimitRotation = Element.bLimitRotation;
+						NewElement.bLimitScale= Element.bLimitScale;
+						NewElement.MinimumValue = Element.MinimumValue;
+						NewElement.MaximumValue = Element.MaximumValue;
+						NewElement.bDrawLimits = Element.bDrawLimits;
+						NewElement.bGizmoEnabled = Element.bGizmoEnabled;
+
 						ElementMap.FindOrAdd(Element.GetElementKey()) = NewElement.GetElementKey();
 						PastedKeys.Add(NewElement.GetElementKey());
 
@@ -771,11 +786,21 @@ TArray<FRigElementKey> FRigHierarchyContainer::ImportFromText(const FString& InC
 				{
 					Data.LocalTransforms[Index].NormalizeRotation();
 					SetLocalTransform(Selection[Index], Data.LocalTransforms[Index]);
+
+					if (Selection[Index].Type == ERigElementType::Space)
+					{
+						SpaceHierarchy.SetInitialTransform(Selection[Index].Name, Data.LocalTransforms[Index]);
+					}
 				}
 				else
 				{
 					Data.GlobalTransforms[Index].NormalizeRotation();
 					SetGlobalTransform(Selection[Index], Data.GlobalTransforms[Index]);
+
+					if (Selection[Index].Type == ERigElementType::Space)
+					{
+						SpaceHierarchy.SetInitialGlobalTransform(Selection[Index].Name, Data.GlobalTransforms[Index]);
+					}
 				}
 				PastedKeys.Add(Selection[Index]);
 			}
diff --git a/Engine/Plugins/Experimental/ControlRig/Source/ControlRig/Private/Rigs/RigSpaceHierarchy.cpp b/Engine/Plugins/Experimental/ControlRig/Source/ControlRig/Private/Rigs/RigSpaceHierarchy.cpp
--- a/Engine/Plugins/Experimental/ControlRig/Source/ControlRig/Private/Rigs/RigSpaceHierarchy.cpp
+++ b/Engine/Plugins/Experimental/ControlRig/Source/ControlRig/Private/Rigs/RigSpaceHierarchy.cpp
@@ -428,18 +428,21 @@ void FRigSpaceHierarchy::RefreshMapping()
 	}
 }
 
-void FRigSpaceHierarchy::Initialize()
+void FRigSpaceHierarchy::Initialize(bool bResetTransforms)
 {
 	RefreshMapping();
 
 	// initialize transform
 	for (int32 Index = 0; Index < Spaces.Num(); ++Index)
 	{
+		if (bResetTransforms)
+		{
+			Spaces[Index].LocalTransform = Spaces[Index].InitialTransform;
+		}
 		if (Container)
 		{
 			Spaces[Index].ParentIndex = Container->GetIndex(Spaces[Index].GetParentElementKey());
 		}
-		Spaces[Index].LocalTransform = Spaces[Index].InitialTransform;
 	}
 }
 
diff --git a/Engine/Plugins/Experimental/ControlRig/Source/ControlRig/Private/Sequencer/ControlRigBindingTemplate.cpp b/Engine/Plugins/Experimental/ControlRig/Source/ControlRig/Private/Sequencer/ControlRigBindingTemplate.cpp
--- a/Engine/Plugins/Experimental/ControlRig/Source/ControlRig/Private/Sequencer/ControlRigBindingTemplate.cpp
+++ b/Engine/Plugins/Experimental/ControlRig/Source/ControlRig/Private/Sequencer/ControlRigBindingTemplate.cpp
@@ -7,7 +7,8 @@
 #include "IMovieScenePlayer.h"
 #include "ControlRig.h"
 #include "Components/SkeletalMeshComponent.h"
-#include "Sequencer/ControlRigSequencerAnimInstance.h"
+#include "AnimCustomInstanceHelper.h"
+#include "Sequencer/ControlRigLayerInstance.h"
 #include "Sequencer/MovieSceneControlRigInstanceData.h"
 #include "IControlRigObjectBinding.h"
 
@@ -38,13 +39,13 @@ struct FControlRigPreAnimatedTokenProducer : IMovieScenePreAnimatedTokenProducer
 				{
 					if (USkeletalMeshComponent* SkeletalMeshComponent = Cast<USkeletalMeshComponent>(ControlRig->GetObjectBinding()->GetBoundObject()))
 					{
-						if (UControlRigSequencerAnimInstance* AnimInstance = Cast<UControlRigSequencerAnimInstance>(SkeletalMeshComponent->GetAnimInstance()))
+						if (UControlRigLayerInstance* AnimInstance = Cast<UControlRigLayerInstance>(SkeletalMeshComponent->GetAnimInstance()))
 						{
 							// Force us to zero weight before we despawn, as the graph could persist
 							AnimInstance->ResetNodes();
 							AnimInstance->RecalcRequiredBones();
 						}
-						UAnimSequencerInstance::UnbindFromSkeletalMeshComponent(SkeletalMeshComponent);
+						FAnimCustomInstanceHelper::UnbindFromSkeletalMeshComponent<UControlRigLayerInstance>(SkeletalMeshComponent);
 					}
 
 					ControlRig->GetObjectBinding()->UnbindFromObject();
@@ -78,25 +79,30 @@ struct FBindControlRigObjectToken : IMovieSceneExecutionToken
 		, bSpawned(bInSpawned)
 	{}
 
-	void BindToSequencerInstance(UControlRig* ControlRig)
+	void BindToSequencerInstance(UControlRig* ControlRig, uint32 SequenceID)
 	{
 		check(ControlRig);
 		if (USkeletalMeshComponent* SkeletalMeshComponent = Cast<USkeletalMeshComponent>(ControlRig->GetObjectBinding()->GetBoundObject()))
 		{
 			bool bWasCreated = false;
-			if (UControlRigSequencerAnimInstance* AnimInstance = UAnimCustomInstance::BindToSkeletalMeshComponent<UControlRigSequencerAnimInstance>(SkeletalMeshComponent,bWasCreated))
+			if (UControlRigLayerInstance* AnimInstance = FAnimCustomInstanceHelper::BindToSkeletalMeshComponent<UControlRigLayerInstance>(SkeletalMeshComponent,bWasCreated))
 			{
+				if (bWasCreated)
+				{
+					AnimInstance->AddControlRigTrack(SequenceID, ControlRig);
+				}
+				
 				AnimInstance->RecalcRequiredBones();
 			}
 		}
 	}
 
-	void UnBindFromSequencerInstance(UControlRig* ControlRig)
+	void UnBindFromSequencerInstance(UControlRig* ControlRig, uint32 SequenceID)
 	{
 		check(ControlRig);
 		if (USkeletalMeshComponent* SkeletalMeshComponent = Cast<USkeletalMeshComponent>(ControlRig->GetObjectBinding()->GetBoundObject()))
 		{
-			UAnimCustomInstance::UnbindFromSkeletalMeshComponent(SkeletalMeshComponent);
+			FAnimCustomInstanceHelper::UnbindFromSkeletalMeshComponent<UControlRigLayerInstance>(SkeletalMeshComponent);
 		}
 	}
 
@@ -140,7 +146,7 @@ struct FBindControlRigObjectToken : IMovieSceneExecutionToken
 							UObject* OuterBoundObject = OuterBoundObjects[0].Get();
 							if (OuterBoundObject && !ControlRig->GetObjectBinding()->IsBoundToObject(OuterBoundObject))
 							{
-								UnBindFromSequencerInstance(ControlRig);
+								UnBindFromSequencerInstance(ControlRig, Operand.SequenceID.GetInternalValue());
 								ControlRig->GetObjectBinding()->UnbindFromObject();
 								ControlRig->GetObjectBinding()->BindToObject(OuterBoundObject);
 							}
@@ -149,12 +155,12 @@ struct FBindControlRigObjectToken : IMovieSceneExecutionToken
 #if WITH_EDITORONLY_DATA
 					else if (ObjectBinding.IsValid() && !ControlRig->GetObjectBinding()->IsBoundToObject(ObjectBinding.Get()))
 					{
-						UnBindFromSequencerInstance(ControlRig);
+						UnBindFromSequencerInstance(ControlRig, Operand.SequenceID.GetInternalValue());
 						ControlRig->GetObjectBinding()->UnbindFromObject();
 						ControlRig->GetObjectBinding()->BindToObject(ObjectBinding.Get());
 					}
 #endif
-					BindToSequencerInstance(ControlRig);
+					BindToSequencerInstance(ControlRig, Operand.SequenceID.GetInternalValue());
 				}
 			}
 			else
@@ -164,12 +170,12 @@ struct FBindControlRigObjectToken : IMovieSceneExecutionToken
 #if WITH_EDITORONLY_DATA
 				if (ObjectBinding.IsValid() && !ControlRig->GetObjectBinding()->IsBoundToObject(ObjectBinding.Get()))
 				{
-					UnBindFromSequencerInstance(ControlRig);
+					UnBindFromSequencerInstance(ControlRig, Operand.SequenceID.GetInternalValue());
 					ControlRig->GetObjectBinding()->UnbindFromObject();
 					ControlRig->GetObjectBinding()->BindToObject(ObjectBinding.Get());
 				}
 
-				BindToSequencerInstance(ControlRig);
+				BindToSequencerInstance(ControlRig, Operand.SequenceID.GetInternalValue());
 #endif
 			}
 
@@ -178,14 +184,10 @@ struct FBindControlRigObjectToken : IMovieSceneExecutionToken
 			{
 				if (USkeletalMeshComponent* SkeletalMeshComponent = Cast<USkeletalMeshComponent>(ControlRig->GetObjectBinding()->GetBoundObject()))
 				{
-					if (UControlRigSequencerAnimInstance* AnimInstance = Cast<UControlRigSequencerAnimInstance>(SkeletalMeshComponent->GetAnimInstance()))
+					if (UControlRigLayerInstance* AnimInstance = Cast<UControlRigLayerInstance>(SkeletalMeshComponent->GetAnimInstance()))
 					{
 						// @todo: let's think about the input here, I'm disabling handling of inputs but allowing it to execute when ticks
-						bool bStructureChanged = AnimInstance->UpdateControlRig(ControlRig, Operand.SequenceID.GetInternalValue(), bAdditive, bApplyBoneFilter, *BoneFilter, Weight, false, true);
-						if (bStructureChanged)
-						{
-							AnimInstance->RecalcRequiredBones();
-						}
+						AnimInstance->UpdateControlRigTrack(Operand.SequenceID.GetInternalValue(), Weight, FControlRigIOSettings::MakeDisabled(), true);
 					}
 				}
 			}
@@ -207,13 +209,14 @@ struct FBindControlRigObjectToken : IMovieSceneExecutionToken
 				{
 					if(USkeletalMeshComponent* SkeletalMeshComponent = Cast<USkeletalMeshComponent>(ControlRig->GetObjectBinding()->GetBoundObject()))
 					{
-						if (UControlRigSequencerAnimInstance* AnimInstance = Cast<UControlRigSequencerAnimInstance>(SkeletalMeshComponent->GetAnimInstance()))
+						if (UControlRigLayerInstance* AnimInstance = Cast<UControlRigLayerInstance>(SkeletalMeshComponent->GetAnimInstance()))
 						{
 							// Force us to zero weight before we despawn, as the graph could persist
-							AnimInstance->UpdateControlRig(ControlRig, Operand.SequenceID.GetInternalValue(), bAdditive, bApplyBoneFilter, *BoneFilter, 0.0f, true, true);
+							// @todo take a look and see why this would be 
+							AnimInstance->UpdateControlRigTrack(Operand.SequenceID.GetInternalValue(), 0.0f, FControlRigIOSettings::MakeEnabled(), true);
 							AnimInstance->RecalcRequiredBones();
 						}
-						UAnimSequencerInstance::UnbindFromSkeletalMeshComponent(SkeletalMeshComponent);
+						FAnimCustomInstanceHelper::UnbindFromSkeletalMeshComponent<UControlRigLayerInstance>(SkeletalMeshComponent);
 					}
 
 					ControlRig->GetObjectBinding()->UnbindFromObject();
diff --git a/Engine/Plugins/Experimental/ControlRig/Source/ControlRig/Private/Sequencer/ControlRigLayerInstanceProxy.h b/Engine/Plugins/Experimental/ControlRig/Source/ControlRig/Private/Sequencer/ControlRigLayerInstanceProxy.h
new file mode 100644
--- /dev/null
+++ b/Engine/Plugins/Experimental/ControlRig/Source/ControlRig/Private/Sequencer/ControlRigLayerInstanceProxy.h
@@ -0,0 +1,136 @@
+// Copyright Epic Games, Inc. All Rights Reserved.
+
+#pragma once
+
+#include "Animation/AnimInstanceProxy.h"
+#include "AnimNode_ControlRig_ExternalSource.h"
+#include "ControlRigLayerInstanceProxy.generated.h"
+
+class UControlRig;
+class UAnimSequencerInstance;
+
+/** Custom internal Input Pose node that handles any AnimInstance */
+USTRUCT()
+struct FAnimNode_ControlRigInputPose : public FAnimNode_Base
+{
+	GENERATED_BODY()
+
+	FAnimNode_ControlRigInputPose()
+		: InputProxy(nullptr)
+	{
+	}
+
+	/** Input pose, optionally linked dynamically to another graph */
+	UPROPERTY()
+	FPoseLink InputPose;
+
+	// FAnimNode_Base interface
+	virtual void Initialize_AnyThread(const FAnimationInitializeContext& Context) override;
+	virtual void CacheBones_AnyThread(const FAnimationCacheBonesContext& Context) override;
+	virtual void Update_AnyThread(const FAnimationUpdateContext& Context) override;
+	virtual void Evaluate_AnyThread(FPoseContext& Output) override;
+	virtual void GatherDebugData(FNodeDebugData& DebugData) override;
+	// End of FAnimNode_Base interface
+
+	/** Called by linked instance nodes to dynamically link this to an outer graph */
+	void Link(UAnimInstance* InInputInstance, FAnimInstanceProxy* InInputProxy);
+
+	/** Called by linked instance nodes to dynamically unlink this to an outer graph */
+	void Unlink();
+
+private:
+	/** The proxy to use when getting inputs, set when dynamically linked */
+	FAnimInstanceProxy* InputProxy;
+	UAnimInstance*		InputAnimInstance;
+};
+
+/** Proxy override for this UAnimInstance-derived class */
+USTRUCT()
+struct CONTROLRIG_API FControlRigLayerInstanceProxy : public FAnimInstanceProxy
+{
+	GENERATED_BODY()
+
+public:
+	FControlRigLayerInstanceProxy()
+		: CurrentRoot(nullptr)
+		, CurrentSourceAnimInstance(nullptr)
+	{
+	}
+
+	FControlRigLayerInstanceProxy(UAnimInstance* InAnimInstance)
+		: FAnimInstanceProxy(InAnimInstance)
+		, CurrentRoot(nullptr)
+		, CurrentSourceAnimInstance(nullptr)
+	{
+	}
+
+	virtual ~FControlRigLayerInstanceProxy();
+
+	// FAnimInstanceProxy interface
+	virtual void Initialize(UAnimInstance* InAnimInstance) override;
+	virtual bool Evaluate(FPoseContext& Output) override;
+	virtual void CacheBones() override;
+	virtual void UpdateAnimationNode(const FAnimationUpdateContext& InContext) override;
+
+	/** Anim Instance Source info - created externally and used here */
+	void SetSourceAnimInstance(UAnimInstance* SourceAnimInstance, FAnimInstanceProxy* SourceAnimInputProxy);
+	UAnimInstance* GetSourceAnimInstance() const { return CurrentSourceAnimInstance; }
+
+	/** ControlRig related support */
+	void AddControlRigTrack(int32 ControlRigID, UControlRig* InControlRig);
+	void UpdateControlRigTrack(int32 ControlRigID, float Weight, const FControlRigIOSettings& InputSettings, bool bExecute);
+	void RemoveControlRigTrack(int32 ControlRigID);
+	void ResetControlRigTracks();
+
+	/** Sequencer AnimInstance Interface */
+	void AddAnimation(int32 SequenceId, UAnimSequenceBase* InAnimSequence);
+	void UpdateAnimTrack(UAnimSequenceBase* InAnimSequence, int32 SequenceId, float InPosition, float Weight, bool bFireNotifies);
+	void UpdateAnimTrack(UAnimSequenceBase* InAnimSequence, int32 SequenceId, float InFromPosition, float InToPosition, float Weight, bool bFireNotifies);
+	void RemoveAnimation(int32 SequenceId);
+
+	/** Reset all nodes in this instance */
+	virtual void ResetNodes();
+	/** Reset the pose in this instance*/
+	virtual void ResetPose();
+
+	/** return first available control rig from the node it has */
+	UControlRig* GetFirstAvailableControlRig() const;
+
+	virtual void AddReferencedObjects(UAnimInstance* InAnimInstance, FReferenceCollector& Collector) override;
+
+	static void InitializeCustomProxy(FAnimInstanceProxy* InputProxy, UAnimInstance* InAnimInstance);
+	static void GatherCustomProxyDebugData(FAnimInstanceProxy* InputProxy, FNodeDebugData& DebugData);
+	static void CacheBonesCustomProxy(FAnimInstanceProxy* InputProxy);
+	static void UpdateCustomProxy(FAnimInstanceProxy* InputProxy, const FAnimationUpdateContext& Context);
+	static void EvaluateCustomProxy(FAnimInstanceProxy* InputProxy, FPoseContext& Output);
+
+	// this doesn't work because this instance continuously change root
+	// this will invalidate the evaluation
+// 	virtual FAnimNode_Base* GetCustomRootNode() 
+
+	friend struct FAnimNode_ControlRigInputPose;
+protected:
+	/** Find ControlRig node of the */
+	FAnimNode_ControlRig_ExternalSource* FindControlRigNode(int32 ControlRigID) const;
+
+	/** Construct and link the base part of the blend tree */
+	void ConstructNodes();
+
+	/** Input pose anim node */
+	FAnimNode_ControlRigInputPose InputPose;
+
+	/** Cuyrrent Root node - this changes whenever track changes */
+	FAnimNode_Base* CurrentRoot;
+
+	/** ControlRig Nodes */
+	TArray<FAnimNode_ControlRig_ExternalSource> ControlRigNodes;
+
+	/** mapping from sequencer index to internal player index */
+	TMap<int32, FAnimNode_ControlRig_ExternalSource*> SequencerToControlRigNodeMap;
+
+	/** Source Anim Instance */
+	UAnimInstance* CurrentSourceAnimInstance;
+
+	/** getter for Sequencer AnimInstance. It will return null if it's using AnimBP */
+	UAnimSequencerInstance* GetSequencerAnimInstance();
+};
diff --git a/Engine/Plugins/Experimental/ControlRig/Source/ControlRig/Private/Sequencer/ControlRigSequencerAnimInstance.cpp b/Engine/Plugins/Experimental/ControlRig/Source/ControlRig/Private/Sequencer/ControlRigSequencerAnimInstance.cpp
--- a/Engine/Plugins/Experimental/ControlRig/Source/ControlRig/Private/Sequencer/ControlRigSequencerAnimInstance.cpp
+++ b/Engine/Plugins/Experimental/ControlRig/Source/ControlRig/Private/Sequencer/ControlRigSequencerAnimInstance.cpp
@@ -15,10 +15,10 @@ FAnimInstanceProxy* UControlRigSequencerAnimInstance::CreateAnimInstanceProxy()
 	return new FControlRigSequencerAnimInstanceProxy(this);
 }
 
-bool UControlRigSequencerAnimInstance::UpdateControlRig(UControlRig* InControlRig, uint32 SequenceId, bool bAdditive, bool bApplyBoneFilter, const FInputBlendPose& BoneFilter, float Weight, bool bUpdateInput, bool bExecute)
+bool UControlRigSequencerAnimInstance::UpdateControlRig(UControlRig* InControlRig, uint32 SequenceId, bool bAdditive, bool bApplyBoneFilter, const FInputBlendPose& BoneFilter, float Weight, const FControlRigIOSettings& InputSettings, bool bExecute)
 {
 	CachedControlRig = InControlRig;
-	return GetProxyOnGameThread<FControlRigSequencerAnimInstanceProxy>().UpdateControlRig(InControlRig, SequenceId, bAdditive, bApplyBoneFilter, BoneFilter, Weight, bUpdateInput, bExecute);
+	return GetProxyOnGameThread<FControlRigSequencerAnimInstanceProxy>().UpdateControlRig(InControlRig, SequenceId, bAdditive, bApplyBoneFilter, BoneFilter, Weight, InputSettings, bExecute);
 }
 
 bool UControlRigSequencerAnimInstance::SetAnimationAsset(class UAnimationAsset* NewAsset)
diff --git a/Engine/Plugins/Experimental/ControlRig/Source/ControlRig/Private/Sequencer/ControlRigSequencerAnimInstanceProxy.cpp b/Engine/Plugins/Experimental/ControlRig/Source/ControlRig/Private/Sequencer/ControlRigSequencerAnimInstanceProxy.cpp
--- a/Engine/Plugins/Experimental/ControlRig/Source/ControlRig/Private/Sequencer/ControlRigSequencerAnimInstanceProxy.cpp
+++ b/Engine/Plugins/Experimental/ControlRig/Source/ControlRig/Private/Sequencer/ControlRigSequencerAnimInstanceProxy.cpp
@@ -137,7 +137,7 @@ void FControlRigSequencerAnimInstanceProxy::InitControlRigTrack(UControlRig* InC
 	}
 }
 
-bool FControlRigSequencerAnimInstanceProxy::UpdateControlRig(UControlRig* InControlRig, uint32 SequenceId, bool bAdditive, bool bApplyBoneFilter, const FInputBlendPose& BoneFilter, float Weight, bool bUpdateInput, bool bExecute)
+bool FControlRigSequencerAnimInstanceProxy::UpdateControlRig(UControlRig* InControlRig, uint32 SequenceId, bool bAdditive, bool bApplyBoneFilter, const FInputBlendPose& BoneFilter, float Weight, const FControlRigIOSettings& InputSettings, bool bExecute)
 {
 	bool bCreated = EnsureControlRigTrack(InControlRig, bAdditive, bApplyBoneFilter, BoneFilter, SequenceId);
 
@@ -153,7 +153,7 @@ bool FControlRigSequencerAnimInstanceProxy::UpdateControlRig(UControlRig* InCont
 		BlendNode.DesiredAlphas[PlayerState->PoseIndex] = Weight;
 	}
 
-	PlayerState->ControlRigNode.bUpdateInput = bUpdateInput;
+	PlayerState->ControlRigNode.InputSettings = InputSettings;
 	PlayerState->ControlRigNode.bExecute = bExecute;
 
 	return bCreated;
diff --git a/Engine/Plugins/Experimental/ControlRig/Source/ControlRig/Private/Sequencer/ControlRigSequencerAnimInstanceProxy.h b/Engine/Plugins/Experimental/ControlRig/Source/ControlRig/Private/Sequencer/ControlRigSequencerAnimInstanceProxy.h
--- a/Engine/Plugins/Experimental/ControlRig/Source/ControlRig/Private/Sequencer/ControlRigSequencerAnimInstanceProxy.h
+++ b/Engine/Plugins/Experimental/ControlRig/Source/ControlRig/Private/Sequencer/ControlRigSequencerAnimInstanceProxy.h
@@ -48,7 +48,7 @@ public:
 	// FAnimSequencerInstanceProxy interface
 	virtual void ResetNodes() override;
 
-	bool UpdateControlRig(UControlRig* InControlRig, uint32 SequenceId, bool bAdditive, bool bApplyBoneFilter, const FInputBlendPose& BoneFilter, float Weight, bool bUpdateInput, bool bExecute);
+	bool UpdateControlRig(UControlRig* InControlRig, uint32 SequenceId, bool bAdditive, bool bApplyBoneFilter, const FInputBlendPose& BoneFilter, float Weight, const FControlRigIOSettings& InputSettings, bool bExecute);
 	bool SetAnimationAsset(class UAnimationAsset* NewAsset);
 
 private:
diff --git a/Engine/Plugins/Experimental/ControlRig/Source/ControlRig/Private/Sequencer/ControlRigSortedControls.cpp b/Engine/Plugins/Experimental/ControlRig/Source/ControlRig/Private/Sequencer/ControlRigSortedControls.cpp
new file mode 100644
--- /dev/null
+++ b/Engine/Plugins/Experimental/ControlRig/Source/ControlRig/Private/Sequencer/ControlRigSortedControls.cpp
@@ -0,0 +1,106 @@
+// Copyright Epic Games, Inc. All Rights Reserved.
+
+#include "Sequencer/ControlRigSortedControls.h"
+
+void FControlRigSortedControls::CreateRigControlsRecursive(TSharedPtr<FRigControlTreeElement>& Element, const TArray<FRigControl>& Controls, TArray<FRigControl>& OutControls)
+{
+	if (Element->Key.Type == ERigElementType::Control)
+	{
+		check(Element->Index >= 0);
+		const FRigControl&  RigControl = Controls[Element->Index];
+		FRigControl NewControl(RigControl);
+		OutControls.Add(NewControl);
+	}
+	for (int32 ChildIndex = 0; ChildIndex < Element->Children.Num(); ++ChildIndex)
+	{
+		FControlRigSortedControls::CreateRigControlsRecursive(Element->Children[ChildIndex], Controls, OutControls);
+	}
+}
+
+void FControlRigSortedControls::GetControlsInOrder(IControlRigManipulatable* Manip, TArray<FRigControl>& SortedControls)
+{
+	TMap<FRigElementKey, TSharedPtr<FRigControlTreeElement>> ElementMap;
+	TArray<TSharedPtr<FRigControlTreeElement>> RootElements;
+	const TArray<FRigControl>& Controls = Manip->AvailableControls();
+	const TArray<FRigSpace>& Spaces = Manip->AvailableSpaces();
+
+	int32 Index = 0;
+	for (const FRigControl& RigControl : Controls)
+	{
+		FControlRigSortedControls::AddControlElement(RigControl, Index, Controls, Spaces, ElementMap, RootElements);
+		++Index;
+	}
+
+	for (Index = 0; Index < RootElements.Num(); ++Index)
+	{
+		FControlRigSortedControls::CreateRigControlsRecursive(RootElements[Index], Controls, SortedControls);
+	}
+}
+
+void FControlRigSortedControls::AddControlElement(FRigControl InControl, int32 Index, const TArray<FRigControl>& Controls, const TArray<FRigSpace>& Spaces,
+	TMap<FRigElementKey, TSharedPtr<FRigControlTreeElement>>& ElementMap,
+	TArray<TSharedPtr<FRigControlTreeElement>>& RootElements)
+{
+	FRigElementKey ParentKey;
+	if (InControl.SpaceIndex != INDEX_NONE)
+	{
+		FControlRigSortedControls::AddSpaceElement(Spaces[InControl.SpaceIndex], InControl.SpaceIndex, Controls, Spaces, ElementMap, RootElements);
+		ParentKey = Spaces[InControl.SpaceIndex].GetElementKey();
+	}
+	else if (InControl.ParentIndex != INDEX_NONE)
+	{
+		FControlRigSortedControls::AddControlElement(Controls[InControl.ParentIndex], InControl.ParentIndex, Controls, Spaces, ElementMap, RootElements);
+		ParentKey = Controls[InControl.ParentIndex].GetElementKey();
+	}
+	FControlRigSortedControls::AddElement(InControl.GetElementKey(), Index, ParentKey, ElementMap, RootElements);
+}
+
+void FControlRigSortedControls::AddSpaceElement(FRigSpace InSpace, int32 Index, const TArray<FRigControl>& Controls, const TArray<FRigSpace>& Spaces,
+	TMap<FRigElementKey, TSharedPtr<FRigControlTreeElement>>& ElementMap, TArray<TSharedPtr<FRigControlTreeElement>>& RootElements)
+{
+	FRigElementKey ParentKey;
+	if (InSpace.ParentIndex != INDEX_NONE)
+	{
+		switch (InSpace.SpaceType)
+		{
+		case ERigSpaceType::Control:
+		{
+			FControlRigSortedControls::AddControlElement(Controls[InSpace.ParentIndex], InSpace.ParentIndex, Controls, Spaces, ElementMap, RootElements);
+			ParentKey = Controls[InSpace.ParentIndex].GetElementKey();
+			break;
+		}
+		case ERigSpaceType::Space:
+		{
+			FControlRigSortedControls::AddSpaceElement(Spaces[InSpace.ParentIndex], InSpace.ParentIndex, Controls, Spaces, ElementMap, RootElements);
+			ParentKey = Spaces[InSpace.ParentIndex].GetElementKey();
+			break;
+		}
+		default:
+		{
+			break;
+		}
+		}
+	}
+	FControlRigSortedControls::AddElement(InSpace.GetElementKey(), Index, ParentKey, ElementMap, RootElements);
+}
+
+void FControlRigSortedControls::AddElement(FRigElementKey InKey, int32 Index, FRigElementKey InParentKey, TMap<FRigElementKey, TSharedPtr<FRigControlTreeElement>>& ElementMap,
+	TArray<TSharedPtr<FRigControlTreeElement>>& RootElements)
+{
+	if (ElementMap.Contains(InKey))
+	{
+		return;
+	}
+	TSharedPtr<FRigControlTreeElement> NewItem = MakeShared<FRigControlTreeElement>(InKey, Index);
+	ElementMap.Add(InKey, NewItem);
+	if (InParentKey)
+	{
+		TSharedPtr<FRigControlTreeElement>* FoundItem = ElementMap.Find(InParentKey);
+		check(FoundItem);
+		FoundItem->Get()->Children.Add(NewItem);
+	}
+	else
+	{
+		RootElements.Add(NewItem);
+	}
+}
diff --git a/Engine/Plugins/Experimental/ControlRig/Source/ControlRig/Private/Sequencer/MovieSceneControlRigParameterSection.cpp b/Engine/Plugins/Experimental/ControlRig/Source/ControlRig/Private/Sequencer/MovieSceneControlRigParameterSection.cpp
--- a/Engine/Plugins/Experimental/ControlRig/Source/ControlRig/Private/Sequencer/MovieSceneControlRigParameterSection.cpp
+++ b/Engine/Plugins/Experimental/ControlRig/Source/ControlRig/Private/Sequencer/MovieSceneControlRigParameterSection.cpp
@@ -8,15 +8,312 @@
 #include "Sequencer/ControlRigBindingTemplate.h"
 #include "Sequencer/MovieSceneControlRigInstanceData.h"
 #include "Channels/MovieSceneChannelProxy.h"
+#include "Sequencer/ControlRigSortedControls.h"
 
 #define LOCTEXT_NAMESPACE "MovieSceneControlParameterRigSection"
 
 #if WITH_EDITOR
 
-//mz todo need other types
+struct FParameterFloatChannelEditorData
+{
+	FParameterFloatChannelEditorData(UControlRig *InControlRig, const FName& InName, bool bEnabledOverride, const FText& GroupName, int SortStartIndex)
+	{
+		ControlRig = InControlRig;
+		ParameterName = InName;
+		FString NameAsString = InName.ToString();
+		FText Group = FText::FromString(NameAsString);
+		{
+			MetaData.SetIdentifiers(InName, FText::FromName(InName),FText());
+			MetaData.Group = Group;
+			MetaData.bEnabled = bEnabledOverride;
+			MetaData.SortOrder = SortStartIndex++;
+			MetaData.bCanCollapseToTrack = false;
+		}
+
+		ExternalValues.OnGetExternalValue = [InControlRig, InName](UObject& InObject, FTrackInstancePropertyBindings* Bindings) { return GetValue(InControlRig, InName,InObject, Bindings); };
+		
+		ExternalValues.OnGetCurrentValueAndWeight = [InName](UObject* Object, UMovieSceneSection*  SectionToKey, FFrameNumber KeyTime, FFrameRate TickResolution, FMovieSceneRootEvaluationTemplateInstance& RootTemplate,
+			float& OutValue, float& OutWeight) { GetChannelValueAndWeight(InName, Object, SectionToKey, KeyTime, TickResolution, RootTemplate, OutValue, OutWeight); };
+		
+	}
+
+	static TOptional<float> GetValue(UControlRig* ControlRig, FName ParameterName, UObject& InObject, FTrackInstancePropertyBindings* Bindings)
+	{
+		if (ControlRig)
+		{
+			FRigControl* RigControl = ControlRig->FindControl(ParameterName);
+			if (RigControl)
+			{
+				float Val = RigControl->Value.Get<float>();
+				return Val;
+		
+			}
+		}
+		return TOptional<float>();
+	}
+	
+	static void GetChannelValueAndWeight(FName ParameterName, UObject* Object, UMovieSceneSection*  SectionToKey, FFrameNumber KeyTime, FFrameRate TickResolution, FMovieSceneRootEvaluationTemplateInstance& RootTemplate,
+		float& OutValue, float& OutWeight)
+	{
+		OutValue = 0.0f;
+		OutWeight = 1.0f;
+
+		UMovieSceneTrack* Track = SectionToKey->GetTypedOuter<UMovieSceneTrack>();
+
+		if (Track)
+		{
+			FMovieSceneEvaluationTrack EvalTrack = Track->GenerateTrackTemplate();
+			FMovieSceneInterrogationData InterrogationData;
+			RootTemplate.CopyActuators(InterrogationData.GetAccumulator());
+
+			FMovieSceneContext Context(FMovieSceneEvaluationRange(KeyTime, TickResolution));
+			EvalTrack.Interrogate(Context, InterrogationData, Object);
+
+			float Val = 0.0f;
+			for (const FFloatInterrogationData& InVector : InterrogationData.Iterate<FFloatInterrogationData>(UMovieSceneControlRigParameterSection::GetFloatInterrogationKey()))
+			{
+				if (InVector.ParameterName == ParameterName)
+				{
+					Val = InVector.Val;
+					break;
+				}
+			}
+			OutValue = Val;
+		}
+		OutWeight = MovieSceneHelpers::CalculateWeightForBlending(SectionToKey, KeyTime);
+	}
+
+	FText							GroupName;
+	FMovieSceneChannelMetaData      MetaData;
+	TMovieSceneExternalValue<float> ExternalValues;
+	FName ParameterName;
+	UControlRig *ControlRig;
+};
+
+//Set up with all 4 Channels so it can be used by all vector types.
+struct FParameterVectorChannelEditorData
+{
+	FParameterVectorChannelEditorData(UControlRig *InControlRig, const FName& InName, bool bEnabledOverride, const FText& GroupName, int SortStartIndex, int32 NumChannels)
+	{
+		ControlRig = InControlRig;
+		ParameterName = InName;
+		FString NameAsString = InName.ToString();
+		FString TotalName = NameAsString;
+		FText Group = FText::FromString(NameAsString);
+
+		{
+			TotalName += ".X";
+			MetaData[0].SetIdentifiers(FName(*TotalName), FCommonChannelData::ChannelX);
+			TotalName = NameAsString;
+			MetaData[0].Group = Group;
+			MetaData[0].bEnabled = bEnabledOverride;
+			MetaData[0].SortOrder = SortStartIndex++;
+			MetaData[0].bCanCollapseToTrack = false;
+		}
+		{
+			TotalName += ".Y";
+			MetaData[1].SetIdentifiers(FName(*TotalName), FCommonChannelData::ChannelY);
+			TotalName = NameAsString;
+			MetaData[1].Group = Group;
+			MetaData[1].bEnabled = bEnabledOverride;
+			MetaData[1].SortOrder = SortStartIndex++;
+			MetaData[1].bCanCollapseToTrack = false;
+		}
+		{
+			TotalName += ".Z";
+			MetaData[2].SetIdentifiers(FName(*TotalName), FCommonChannelData::ChannelZ);
+			TotalName = NameAsString;
+			MetaData[2].Group = Group;
+			MetaData[2].bEnabled = bEnabledOverride;
+			MetaData[2].SortOrder = SortStartIndex++;
+			MetaData[2].bCanCollapseToTrack = false;
+		}
+		{
+			TotalName += ".W";
+			MetaData[3].SetIdentifiers(FName(*TotalName), FCommonChannelData::ChannelW);
+			TotalName = NameAsString;
+			MetaData[3].Group = Group;
+			MetaData[3].bEnabled = bEnabledOverride;
+			MetaData[3].SortOrder = SortStartIndex++;
+			MetaData[3].bCanCollapseToTrack = false;
+		}
+
+		ExternalValues[0].OnGetExternalValue = [InControlRig, InName,NumChannels](UObject& InObject, FTrackInstancePropertyBindings* Bindings) { return ExtractChannelX(InObject, Bindings, InControlRig, InName, NumChannels); };
+		ExternalValues[1].OnGetExternalValue = [InControlRig, InName,NumChannels](UObject& InObject, FTrackInstancePropertyBindings* Bindings) { return ExtractChannelY(InObject, Bindings, InControlRig, InName, NumChannels); };
+		ExternalValues[2].OnGetExternalValue = [InControlRig, InName,NumChannels](UObject& InObject, FTrackInstancePropertyBindings* Bindings) { return ExtractChannelZ(InObject, Bindings, InControlRig, InName, NumChannels); };
+		ExternalValues[3].OnGetExternalValue = [InControlRig, InName,NumChannels](UObject& InObject, FTrackInstancePropertyBindings* Bindings) { return ExtractChannelW(InObject, Bindings,InControlRig, InName, NumChannels); };
+
+		ExternalValues[0].OnGetCurrentValueAndWeight = [InName, NumChannels](UObject* Object, UMovieSceneSection*  SectionToKey, FFrameNumber KeyTime, FFrameRate TickResolution, FMovieSceneRootEvaluationTemplateInstance& RootTemplate,
+			float& OutValue, float& OutWeight) { GetChannelValueAndWeight(InName, NumChannels, 0, Object, SectionToKey, KeyTime, TickResolution, RootTemplate, OutValue, OutWeight); };
+		ExternalValues[1].OnGetCurrentValueAndWeight = [InName, NumChannels](UObject* Object, UMovieSceneSection*  SectionToKey, FFrameNumber KeyTime, FFrameRate TickResolution, FMovieSceneRootEvaluationTemplateInstance& RootTemplate,
+			float& OutValue, float& OutWeight) { GetChannelValueAndWeight(InName, NumChannels, 1, Object, SectionToKey, KeyTime, TickResolution, RootTemplate, OutValue, OutWeight); };
+		ExternalValues[2].OnGetCurrentValueAndWeight = [InName, NumChannels](UObject* Object, UMovieSceneSection*  SectionToKey, FFrameNumber KeyTime, FFrameRate TickResolution, FMovieSceneRootEvaluationTemplateInstance& RootTemplate,
+			float& OutValue, float& OutWeight) { GetChannelValueAndWeight(InName, NumChannels, 2, Object, SectionToKey, KeyTime, TickResolution, RootTemplate, OutValue, OutWeight); };
+		ExternalValues[3].OnGetCurrentValueAndWeight = [InName, NumChannels](UObject* Object, UMovieSceneSection*  SectionToKey, FFrameNumber KeyTime, FFrameRate TickResolution, FMovieSceneRootEvaluationTemplateInstance& RootTemplate,
+			float& OutValue, float& OutWeight) { GetChannelValueAndWeight(InName, NumChannels, 3, Object, SectionToKey, KeyTime, TickResolution, RootTemplate, OutValue, OutWeight); };
+
+	}
+
+	static FVector4 GetPropertyValue(UControlRig* ControlRig, FName ParameterName, UObject& InObject, FTrackInstancePropertyBindings& Bindings, int32 NumChannels)
+	{
+		if (ControlRig)
+		{
+			FRigControl* RigControl = ControlRig->FindControl(ParameterName);
+			if (RigControl)
+			{
+		
+				if (NumChannels == 2)
+				{
+					FVector2D Vector = RigControl->Value.Get<FVector2D>();
+					return FVector4(Vector.X, Vector.Y, 0.f, 0.f);
+				}
+				else if (NumChannels == 3)
+				{
+					FVector Vector = RigControl->Value.Get<FVector>();
+					return FVector4(Vector.X, Vector.Y, Vector.Z, 0.f);
+				}
+				else
+				{
+					return RigControl->Value.Get<FVector4>();
+				}
+			}
+		}
+		return FVector4();
+	}
+
+	static TOptional<float> ExtractChannelX(UObject& InObject, FTrackInstancePropertyBindings* Bindings, UControlRig* ControlRig, FName ParameterName, int32 NumChannels)
+	{
+		return Bindings ? GetPropertyValue(ControlRig, ParameterName, InObject, *Bindings, NumChannels).X : TOptional<float>();
+	}
+	static TOptional<float> ExtractChannelY(UObject& InObject, FTrackInstancePropertyBindings* Bindings, UControlRig* ControlRig, FName ParameterName, int32 NumChannels)
+	{
+		return Bindings ? GetPropertyValue(ControlRig, ParameterName, InObject, *Bindings, NumChannels).Y : TOptional<float>();
+	}
+	static TOptional<float> ExtractChannelZ(UObject& InObject, FTrackInstancePropertyBindings* Bindings, UControlRig* ControlRig, FName ParameterName, int32 NumChannels)
+	{
+		return Bindings ? GetPropertyValue(ControlRig, ParameterName, InObject, *Bindings, NumChannels).Z : TOptional<float>();
+	}
+	static TOptional<float> ExtractChannelW(UObject& InObject, FTrackInstancePropertyBindings* Bindings, UControlRig* ControlRig, FName ParameterName, int32 NumChannels)
+	{
+		return Bindings ? GetPropertyValue(ControlRig, ParameterName, InObject, *Bindings, NumChannels).W : TOptional<float>();
+	}
+
+	static void GetChannelValueAndWeight(FName ParameterName, int32 NumChannels, int32 Index, UObject* Object, UMovieSceneSection*  SectionToKey, FFrameNumber KeyTime, FFrameRate TickResolution, FMovieSceneRootEvaluationTemplateInstance& RootTemplate,
+		float& OutValue, float& OutWeight)
+	{
+		OutValue = 0.0f;
+		OutWeight = 1.0f;
+		if (Index >= NumChannels)
+		{
+			return;
+		}
+
+		UMovieSceneTrack* Track = SectionToKey->GetTypedOuter<UMovieSceneTrack>();
+
+		if (Track)
+		{
+			FMovieSceneEvaluationTrack EvalTrack = Track->GenerateTrackTemplate();
+			FMovieSceneInterrogationData InterrogationData;
+			RootTemplate.CopyActuators(InterrogationData.GetAccumulator());
+
+			FMovieSceneContext Context(FMovieSceneEvaluationRange(KeyTime, TickResolution));
+			EvalTrack.Interrogate(Context, InterrogationData, Object);
+
+			switch (NumChannels)
+			{
+			case 2:
+			{
+				FVector2D Val(0.0f, 0.0f);
+				for (const FVector2DInterrogationData& InVector : InterrogationData.Iterate<FVector2DInterrogationData>(UMovieSceneControlRigParameterSection::GetVector2DInterrogationKey()))
+				{
+					if (InVector.ParameterName == ParameterName)
+					{
+						Val = InVector.Val;
+						break;
+					}
+				}
+				switch (Index)
+				{
+				case 0:
+					OutValue = Val.X;
+					break;
+				case 1:
+					OutValue = Val.Y;
+					break;
+				}
+			}
+			break;
+			case 3:
+			{
+				FVector Val(0.0f, 0.0f, 0.0f);
+				for (const FVectorInterrogationData& InVector : InterrogationData.Iterate<FVectorInterrogationData>(UMovieSceneControlRigParameterSection::GetVectorInterrogationKey()))
+				{
+					if (InVector.ParameterName == ParameterName)
+					{
+						Val = InVector.Val;
+						break;
+					}
+				}
+				switch (Index)
+				{
+				case 0:
+					OutValue = Val.X;
+					break;
+				case 1:
+					OutValue = Val.Y;
+					break;
+				case 2:
+					OutValue = Val.Z;
+					break;
+				}
+			}
+			break;
+			case 4:
+			{
+				/* No Interrogation for Vector4, todo if we do add later
+				FVector4 Val(0.0f, 0.0f, 0.0f, 0.0f);
+				for (const FVector4InterrogationData& InVector : InterrogationData.Iterate<FVector4InterrogationData>(UMovieSceneControlRigParameterSection::GetVector4InterrogationKey()))
+				{
+					if (InVector.ParameterName == ParameterName)
+					{
+						Val = InVector.Val;
+						break;
+					}
+				}
+				switch (Index)
+				{
+				case 0:
+					OutValue = Val.X;
+					break;
+				case 1:
+					OutValue = Val.Y;
+					break;
+				case 2:
+					OutValue = Val.Z;
+					break;
+				case 3:
+					OutValue = Val.W;
+					break;
+				}
+				*/
+			}
+			
+			break;
+			}
+		}
+		OutWeight = MovieSceneHelpers::CalculateWeightForBlending(SectionToKey, KeyTime);
+	}
+	FText							GroupName;
+	FMovieSceneChannelMetaData      MetaData[4];
+	TMovieSceneExternalValue<float> ExternalValues[4];
+	FName ParameterName;
+	UControlRig *ControlRig;
+};
+
 struct FParameterTransformChannelEditorData
 {
-	FParameterTransformChannelEditorData(UControlRig *InControlRig, const FName& InName, bool bEnabledOverride, EMovieSceneTransformChannel Mask, const FText& GroupName)
+	FParameterTransformChannelEditorData(UControlRig *InControlRig, const FName& InName, bool bEnabledOverride, EMovieSceneTransformChannel Mask, 
+		const FText& GroupName, int SortStartIndex)
 	{
 		ControlRig = InControlRig;
 		ParameterName = InName;
@@ -36,7 +333,7 @@ struct FParameterTransformChannelEditorData
 
 			MetaData[0].bEnabled = bEnabledOverride && EnumHasAllFlags(Mask, EMovieSceneTransformChannel::TranslationX);
 			MetaData[0].Color = FCommonChannelData::RedChannelColor;
-			MetaData[0].SortOrder = 0;
+			MetaData[0].SortOrder = SortStartIndex++;
 			MetaData[0].bCanCollapseToTrack = false;
 
 			//MetaData[1].SetIdentifiers("Location.Y", FCommonChannelData::ChannelY, LocationGroup);
@@ -46,7 +343,7 @@ struct FParameterTransformChannelEditorData
 
 			MetaData[1].bEnabled = bEnabledOverride && EnumHasAllFlags(Mask, EMovieSceneTransformChannel::TranslationY);
 			MetaData[1].Color = FCommonChannelData::GreenChannelColor;
-			MetaData[1].SortOrder = 1;
+			MetaData[1].SortOrder = SortStartIndex++;
 			MetaData[1].bCanCollapseToTrack = false;
 
 			//MetaData[2].SetIdentifiers("Location.Z", FCommonChannelData::ChannelZ, LocationGroup);
@@ -56,7 +353,7 @@ struct FParameterTransformChannelEditorData
 
 			MetaData[2].bEnabled = bEnabledOverride && EnumHasAllFlags(Mask, EMovieSceneTransformChannel::TranslationZ);
 			MetaData[2].Color = FCommonChannelData::BlueChannelColor;
-			MetaData[2].SortOrder = 2;
+			MetaData[2].SortOrder = SortStartIndex++;
 			MetaData[2].bCanCollapseToTrack = false;
 		}
 		{
@@ -67,7 +364,7 @@ struct FParameterTransformChannelEditorData
 
 			MetaData[3].bEnabled = bEnabledOverride && EnumHasAllFlags(Mask, EMovieSceneTransformChannel::RotationX);
 			MetaData[3].Color = FCommonChannelData::RedChannelColor;
-			MetaData[3].SortOrder = 3;
+			MetaData[3].SortOrder = SortStartIndex++;
 			MetaData[3].bCanCollapseToTrack = false;
 
 			//MetaData[4].SetIdentifiers("Rotation.Y", NSLOCTEXT("MovieSceneTransformSection", "RotationY", "Pitch"), RotationGroup);
@@ -77,7 +374,7 @@ struct FParameterTransformChannelEditorData
 
 			MetaData[4].bEnabled = bEnabledOverride && EnumHasAllFlags(Mask, EMovieSceneTransformChannel::RotationY);
 			MetaData[4].Color = FCommonChannelData::GreenChannelColor;
-			MetaData[4].SortOrder = 4;
+			MetaData[4].SortOrder = SortStartIndex++;
 			MetaData[4].bCanCollapseToTrack = false;
 
 			//MetaData[5].SetIdentifiers("Rotation.Z", NSLOCTEXT("MovieSceneTransformSection", "RotationZ", "Yaw"), RotationGroup);
@@ -87,7 +384,7 @@ struct FParameterTransformChannelEditorData
 
 			MetaData[5].bEnabled = bEnabledOverride && EnumHasAllFlags(Mask, EMovieSceneTransformChannel::RotationZ);
 			MetaData[5].Color = FCommonChannelData::BlueChannelColor;
-			MetaData[5].SortOrder = 5;
+			MetaData[5].SortOrder = SortStartIndex++;
 			MetaData[5].bCanCollapseToTrack = false;
 		}
 		{
@@ -98,7 +395,7 @@ struct FParameterTransformChannelEditorData
 
 			MetaData[6].bEnabled = bEnabledOverride && EnumHasAllFlags(Mask, EMovieSceneTransformChannel::ScaleX);
 			MetaData[6].Color = FCommonChannelData::RedChannelColor;
-			MetaData[6].SortOrder = 6;
+			MetaData[6].SortOrder = SortStartIndex++;
 			MetaData[6].bCanCollapseToTrack = false;
 
 			//MetaData[7].SetIdentifiers("Scale.Y", FCommonChannelData::ChannelY, ScaleGroup);
@@ -108,7 +405,7 @@ struct FParameterTransformChannelEditorData
 
 			MetaData[7].bEnabled = bEnabledOverride && EnumHasAllFlags(Mask, EMovieSceneTransformChannel::ScaleY);
 			MetaData[7].Color = FCommonChannelData::GreenChannelColor;
-			MetaData[7].SortOrder = 7;
+			MetaData[7].SortOrder = SortStartIndex++;
 			MetaData[7].bCanCollapseToTrack = false;
 
 			//MetaData[8].SetIdentifiers("Scale.Z", FCommonChannelData::ChannelZ, ScaleGroup);
@@ -118,7 +415,7 @@ struct FParameterTransformChannelEditorData
 
 			MetaData[8].bEnabled = bEnabledOverride && EnumHasAllFlags(Mask, EMovieSceneTransformChannel::ScaleZ);
 			MetaData[8].Color = FCommonChannelData::BlueChannelColor;
-			MetaData[8].SortOrder = 8;
+			MetaData[8].SortOrder = SortStartIndex++;
 			MetaData[8].bCanCollapseToTrack = false;
 		}
 		{
@@ -228,8 +525,17 @@ struct FParameterTransformChannelEditorData
 			FRigControl* RigControl = ControlRig->FindControl(ParameterName);
 			if (RigControl)
 			{
-				FTransform Transform = RigControl->Value.Get<FTransform>();
-				return Transform.GetTranslation();
+				if (RigControl->ControlType == ERigControlType::Transform)
+				{
+					FTransform Transform = RigControl->Value.Get<FTransform>();
+					return Transform.GetTranslation();
+				}
+				else if  (RigControl->ControlType == ERigControlType::TransformNoScale)
+				{
+					FTransformNoScale NoScale = RigControl->Value.Get<FTransformNoScale>();
+					FTransform Transform = NoScale;
+					return Transform.GetTranslation();
+				}
 			}
 		}
 		return TOptional<FVector>();
@@ -243,11 +549,19 @@ struct FParameterTransformChannelEditorData
 			FRigControl* RigControl = ControlRig->FindControl(ParameterName);
 			if (RigControl)
 			{
-				FTransform Transform = RigControl->Value.Get<FTransform>();
-				return Transform.GetRotation().Rotator();
+				if (RigControl->ControlType == ERigControlType::Transform)
+				{
+					FTransform Transform = RigControl->Value.Get<FTransform>();
+					return Transform.GetRotation().Rotator();
+				}
+				else if (RigControl->ControlType == ERigControlType::TransformNoScale)
+				{
+					FTransformNoScale NoScale = RigControl->Value.Get<FTransformNoScale>();
+					FTransform Transform = NoScale;
+					return Transform.GetRotation().Rotator();
+				}
 			}
 		}
-
 		return TOptional<FRotator>();
 	}
 
@@ -333,10 +647,8 @@ public:
 	UControlRig *ControlRig;
 };
 
-
 #endif // WITH_EDITOR
 
-
 UMovieSceneControlRigParameterSection::UMovieSceneControlRigParameterSection() :bDoNotKey(false)
 {
 	// Section template relies on always restoring state for objects when they are no longer animating. This is how it releases animation control.
@@ -368,6 +680,11 @@ void UMovieSceneControlRigParameterSection::PostEditImport()
 	Super::PostEditImport();
 }
 
+void UMovieSceneControlRigParameterSection::PostLoad()
+{
+	Super::PostLoad();
+
+}
 bool UMovieSceneControlRigParameterSection::HasScalarParameter(FName InParameterName) const
 {
 	for (const FScalarParameterNameAndCurve& ScalarParameterNameAndCurve : ScalarParameterNamesAndCurves)
@@ -380,6 +697,30 @@ bool UMovieSceneControlRigParameterSection::HasScalarParameter(FName InParameter
 	return false;
 }
 
+bool UMovieSceneControlRigParameterSection::HasBoolParameter(FName InParameterName) const
+{
+	for (const FBoolParameterNameAndCurve& BoolParameterNameAndCurve : BoolParameterNamesAndCurves)
+	{
+		if (BoolParameterNameAndCurve.ParameterName == InParameterName)
+		{
+			return true;
+		}
+	}
+	return false;
+}
+
+bool UMovieSceneControlRigParameterSection::HasVector2DParameter(FName InParameterName) const
+{
+	for (const FVector2DParameterNameAndCurves& Vector2DParameterNameAndCurve : Vector2DParameterNamesAndCurves)
+	{
+		if (Vector2DParameterNameAndCurve.ParameterName == InParameterName)
+		{
+			return true;
+		}
+	}
+	return false;
+}
+
 bool UMovieSceneControlRigParameterSection::HasVectorParameter(FName InParameterName) const
 {
 	for (const FVectorParameterNameAndCurves& VectorParameterNameAndCurve : VectorParameterNamesAndCurves)
@@ -416,8 +757,7 @@ bool UMovieSceneControlRigParameterSection::HasTransformParameter(FName InParame
 	return false;
 }
 
-
-void UMovieSceneControlRigParameterSection::AddScalarParameter(FName InParameterName, TOptional<float> DefaultValue)
+void UMovieSceneControlRigParameterSection::AddScalarParameter(FName InParameterName, TOptional<float> DefaultValue, bool bReconstructChannel)
 {
 	FMovieSceneFloatChannel* ExistingChannel = nullptr;
 	if (!HasScalarParameter(InParameterName))
@@ -432,11 +772,63 @@ void UMovieSceneControlRigParameterSection::AddScalarParameter(FName InParameter
 		{
 			ExistingChannel->SetDefault(0.0f);
 		}
-		ReconstructChannelProxy();
+		if (bReconstructChannel)
+		{
+			ReconstructChannelProxy(true);
+		}
+	}
+}
+
+
+void UMovieSceneControlRigParameterSection::AddBoolParameter(FName InParameterName, TOptional<bool> DefaultValue, bool bReconstructChannel)
+{
+	FMovieSceneBoolChannel* ExistingChannel = nullptr;
+	if (!HasBoolParameter(InParameterName))
+	{
+		const int32 NewIndex = BoolParameterNamesAndCurves.Add(FBoolParameterNameAndCurve(InParameterName));
+		ExistingChannel = &BoolParameterNamesAndCurves[NewIndex].ParameterCurve;
+		if (DefaultValue.IsSet())
+		{
+			ExistingChannel->SetDefault(DefaultValue.GetValue());
+		}
+		else
+		{
+			ExistingChannel->SetDefault(false);
+		}
+		if (bReconstructChannel)
+		{
+			ReconstructChannelProxy(true);
+		}
+	}
+}
+
+
+void UMovieSceneControlRigParameterSection::AddVector2DParameter(FName InParameterName, TOptional<FVector2D> DefaultValue, bool bReconstructChannel)
+{
+	FVector2DParameterNameAndCurves* ExistingCurves = nullptr;
+
+	if (!HasVector2DParameter(InParameterName))
+	{
+		int32 NewIndex = Vector2DParameterNamesAndCurves.Add(FVector2DParameterNameAndCurves(InParameterName));
+		ExistingCurves = &Vector2DParameterNamesAndCurves[NewIndex];
+		if (DefaultValue.IsSet())
+		{
+			ExistingCurves->XCurve.SetDefault(DefaultValue.GetValue().X);
+			ExistingCurves->YCurve.SetDefault(DefaultValue.GetValue().Y);
+		}
+		else
+		{
+			ExistingCurves->XCurve.SetDefault(0.0f);
+			ExistingCurves->YCurve.SetDefault(0.0f);
+		}
+		if (bReconstructChannel)
+		{
+			ReconstructChannelProxy(true);
+		}
 	}
 }
 
-void UMovieSceneControlRigParameterSection::AddVectorParameter(FName InParameterName, TOptional<FVector> DefaultValue)
+void UMovieSceneControlRigParameterSection::AddVectorParameter(FName InParameterName, TOptional<FVector> DefaultValue, bool bReconstructChannel)
 {
 	FVectorParameterNameAndCurves* ExistingCurves = nullptr;
 
@@ -457,11 +849,14 @@ void UMovieSceneControlRigParameterSection::AddVectorParameter(FName InParameter
 			ExistingCurves->YCurve.SetDefault(0.0f);
 			ExistingCurves->ZCurve.SetDefault(0.0f);
 		}
-		ReconstructChannelProxy();
+		if (bReconstructChannel)
+		{
+			ReconstructChannelProxy(true);
+		}
 	}
 }
 
-void UMovieSceneControlRigParameterSection::AddColorParameter(FName InParameterName, TOptional<FLinearColor> DefaultValue)
+void UMovieSceneControlRigParameterSection::AddColorParameter(FName InParameterName, TOptional<FLinearColor> DefaultValue, bool bReconstructChannel)
 {
 	FColorParameterNameAndCurves* ExistingCurves = nullptr;
 
@@ -483,11 +878,14 @@ void UMovieSceneControlRigParameterSection::AddColorParameter(FName InParameterN
 			ExistingCurves->BlueCurve.SetDefault(0.0f);
 			ExistingCurves->AlphaCurve.SetDefault(0.0f);
 		}
-		ReconstructChannelProxy();
+		if (bReconstructChannel)
+		{
+			ReconstructChannelProxy(true);
+		}
 	}
 }
 
-void UMovieSceneControlRigParameterSection::AddTransformParameter(FName InParameterName, TOptional<FTransform> DefaultValue)
+void UMovieSceneControlRigParameterSection::AddTransformParameter(FName InParameterName, TOptional<FTransform> DefaultValue, bool bReconstructChannel)
 {
 	FTransformParameterNameAndCurves* ExistingCurves = nullptr;
 
@@ -528,123 +926,308 @@ void UMovieSceneControlRigParameterSection::AddTransformParameter(FName InParame
 			ExistingCurves->Scale[1].SetDefault(0.0f);
 			ExistingCurves->Scale[2].SetDefault(0.0f);
 		}
-		ReconstructChannelProxy();
+		if(bReconstructChannel)
+		{
+			ReconstructChannelProxy(true);
+		}
 	}
 }
 
-void UMovieSceneControlRigParameterSection::ReconstructChannelProxy()
+
+void UMovieSceneControlRigParameterSection::ReconstructChannelProxy(bool bForce)
 {
 	FMovieSceneChannelProxyData Channels;
-	/** Unserialized mask that defines the mask of the current channel proxy so we don't needlessly re-create it on post-undo */
-	EMovieSceneTransformChannel ProxyChannels = TransformMask.GetChannels();
-
-#if WITH_EDITOR
-
-	int32 Index = 0;
-	for (FScalarParameterNameAndCurve& Scalar : GetScalarParameterNamesAndCurves())
-	{
-		bool bEnabled = ControlsMask[Index++];
-		FMovieSceneChannelMetaData MetaData(Scalar.ParameterName, FText::FromName(Scalar.ParameterName), FText(), bEnabled);
-		// Prevent single channels from collapsing to the track node
-		MetaData.bCanCollapseToTrack = false;
-		Channels.Add(Scalar.ParameterCurve, MetaData, TMovieSceneExternalValue<float>());
-	}
-	for (FVectorParameterNameAndCurves& Vector : GetVectorParameterNamesAndCurves())
-	{
-		bool bEnabled = ControlsMask[Index++];
-		FString ParameterString = Vector.ParameterName.ToString();
-		FText Group = FText::FromString(ParameterString);
-		Channels.Add(Vector.XCurve, FMovieSceneChannelMetaData(*(ParameterString + TEXT(".X")), FCommonChannelData::ChannelX, Group, bEnabled), TMovieSceneExternalValue<float>());
-		Channels.Add(Vector.YCurve, FMovieSceneChannelMetaData(*(ParameterString + TEXT(".Y")), FCommonChannelData::ChannelY, Group, bEnabled), TMovieSceneExternalValue<float>());
-		Channels.Add(Vector.ZCurve, FMovieSceneChannelMetaData(*(ParameterString + TEXT(".Z")), FCommonChannelData::ChannelZ, Group, bEnabled), TMovieSceneExternalValue<float>());
-	}
-	for (FColorParameterNameAndCurves& Color : GetColorParameterNamesAndCurves())
+	if(bForce || ControlsMask!= OldControlsMask)
 	{
-		bool bEnabled = ControlsMask[Index++];
-		FString ParameterString = Color.ParameterName.ToString();
-		FText Group = FText::FromString(ParameterString);
+		ControlChannelMap.Empty();
 
-		FMovieSceneChannelMetaData MetaData_R(*(ParameterString + TEXT("R")), FCommonChannelData::ChannelR, Group, bEnabled);
-		MetaData_R.SortOrder = 0;
-		MetaData_R.Color = FCommonChannelData::RedChannelColor;
-
-		FMovieSceneChannelMetaData MetaData_G(*(ParameterString + TEXT("G")), FCommonChannelData::ChannelG, Group, bEnabled);
-		MetaData_G.SortOrder = 1;
-		MetaData_G.Color = FCommonChannelData::GreenChannelColor;
-
-		FMovieSceneChannelMetaData MetaData_B(*(ParameterString + TEXT("B")), FCommonChannelData::ChannelB, Group, bEnabled);
-		MetaData_B.SortOrder = 2;
-		MetaData_B.Color = FCommonChannelData::BlueChannelColor;
-
-		FMovieSceneChannelMetaData MetaData_A(*(ParameterString + TEXT("A")), FCommonChannelData::ChannelA, Group, bEnabled);
-		MetaData_A.SortOrder = 3;
-
-		Channels.Add(Color.RedCurve, MetaData_R, TMovieSceneExternalValue<float>());
-		Channels.Add(Color.GreenCurve, MetaData_G, TMovieSceneExternalValue<float>());
-		Channels.Add(Color.BlueCurve, MetaData_B, TMovieSceneExternalValue<float>());
-		Channels.Add(Color.AlphaCurve, MetaData_A, TMovieSceneExternalValue<float>());
-	}
-
-	for (FTransformParameterNameAndCurves& Transform : GetTransformParameterNamesAndCurves())
-	{
-		bool bEnabled = ControlsMask[Index++];
-		FString ParameterString = Transform.ParameterName.ToString();
-		FText Group = FText::FromString(ParameterString);
-
-		FParameterTransformChannelEditorData EditorData(ControlRig, Transform.ParameterName, bEnabled, TransformMask.GetChannels(), Group);
-
-		Channels.Add(Transform.Translation[0], EditorData.MetaData[0], EditorData.ExternalValues[0]);
-		Channels.Add(Transform.Translation[1], EditorData.MetaData[1], EditorData.ExternalValues[1]);
-		Channels.Add(Transform.Translation[2], EditorData.MetaData[2], EditorData.ExternalValues[2]);
-
-		Channels.Add(Transform.Rotation[0], EditorData.MetaData[3], EditorData.ExternalValues[3]);
-		Channels.Add(Transform.Rotation[1], EditorData.MetaData[4], EditorData.ExternalValues[4]);
-		Channels.Add(Transform.Rotation[2], EditorData.MetaData[5], EditorData.ExternalValues[5]);
-
-		Channels.Add(Transform.Scale[0], EditorData.MetaData[6], EditorData.ExternalValues[6]);
-		Channels.Add(Transform.Scale[1], EditorData.MetaData[7], EditorData.ExternalValues[7]);
-		Channels.Add(Transform.Scale[2], EditorData.MetaData[8], EditorData.ExternalValues[8]);
-
-	}
+		OldControlsMask = ControlsMask;
+		// Need to create the channels in sorted orders
+		if (ControlRig)
+		{
+			TArray<FRigControl> SortedControls;
+			FControlRigSortedControls::GetControlsInOrder(ControlRig, SortedControls);
+			int32 Index = 0;
+			int32 ChannelIndex = 0;
+			for (const FRigControl& RigControl : SortedControls)
+			{
+#if WITH_EDITOR
+				switch (RigControl.ControlType)
+				{
+				case ERigControlType::Float:
+				{
+					for (FScalarParameterNameAndCurve& Scalar : GetScalarParameterNamesAndCurves())
+					{
+						if (RigControl.Name == Scalar.ParameterName)
+						{
+							ControlChannelMap.Add(Scalar.ParameterName, FChannelMapInfo(Index, ChannelIndex));
+							bool bEnabled = ControlsMask[Index++];
+							FString ParameterString = Scalar.ParameterName.ToString();
+							FText Group = FText::FromString(ParameterString);
+							FParameterFloatChannelEditorData EditorData(ControlRig, Scalar.ParameterName, bEnabled, Group, ChannelIndex);
+							Channels.Add(Scalar.ParameterCurve, EditorData.MetaData, EditorData.ExternalValues);
+							ChannelIndex += 1;
+							break;
+						}
+					}
+					break;
+				}
+				case ERigControlType::Bool:
+				{
+					for (FBoolParameterNameAndCurve& Bool : GetBoolParameterNamesAndCurves())
+					{
+						if (RigControl.Name == Bool.ParameterName)
+						{
+							ControlChannelMap.Add(Bool.ParameterName, FChannelMapInfo(Index, ChannelIndex));
+							bool bEnabled = ControlsMask[Index++];
+							FMovieSceneChannelMetaData MetaData(Bool.ParameterName, FText::FromName(Bool.ParameterName), FText::FromName(Bool.ParameterName), bEnabled);
+							MetaData.SortOrder = ChannelIndex++;
+							// Prevent single channels from collapsing to the track node
+							MetaData.bCanCollapseToTrack = false;
+							Channels.Add(Bool.ParameterCurve, MetaData, TMovieSceneExternalValue<bool>());
+							break;
+						}
+					}
+					break;
+				}
+				case ERigControlType::Vector2D:
+				{
+					for (FVector2DParameterNameAndCurves& Vector2D : GetVector2DParameterNamesAndCurves())
+					{
+						if (RigControl.Name == Vector2D.ParameterName)
+						{
+							ControlChannelMap.Add(Vector2D.ParameterName, FChannelMapInfo(Index, ChannelIndex));
+							bool bEnabled = ControlsMask[Index++];
+							FString ParameterString = Vector2D.ParameterName.ToString();
+							FText Group = FText::FromString(ParameterString);
+							FParameterVectorChannelEditorData EditorData(ControlRig, Vector2D.ParameterName, bEnabled, Group, ChannelIndex, 2);
+							Channels.Add(Vector2D.XCurve, EditorData.MetaData[0], EditorData.ExternalValues[0]);
+							Channels.Add(Vector2D.YCurve, EditorData.MetaData[1], EditorData.ExternalValues[1]);
+							ChannelIndex += 2;
+							break;
+						}
+					}
+					break;
+				}
+				case ERigControlType::Position:
+				case ERigControlType::Scale:
+				case ERigControlType::Rotator:
+				{
+					for (FVectorParameterNameAndCurves& Vector : GetVectorParameterNamesAndCurves())
+					{
+						if (RigControl.Name == Vector.ParameterName)
+						{
+							ControlChannelMap.Add(Vector.ParameterName, FChannelMapInfo(Index, ChannelIndex));
+							bool bEnabled = ControlsMask[Index++];
+							FString ParameterString = Vector.ParameterName.ToString();
+							FText Group = FText::FromString(ParameterString);
+							FParameterVectorChannelEditorData EditorData(ControlRig, Vector.ParameterName, bEnabled, Group, ChannelIndex, 3);
+							Channels.Add(Vector.XCurve, EditorData.MetaData[0], EditorData.ExternalValues[0]);
+							Channels.Add(Vector.YCurve, EditorData.MetaData[1], EditorData.ExternalValues[1]);
+							Channels.Add(Vector.ZCurve, EditorData.MetaData[2], EditorData.ExternalValues[2]);
+							ChannelIndex += 3;
+							break;
+						}
+					}
+					break;
+				}
+				/*No colors or Vector4's yet
+				case ERigControlType::Color:
+				{
+					for (FColorParameterNameAndCurves& Color : GetColorParameterNamesAndCurves())
+					{
+						if(RigControl.Name == Color.ParameterName)
+						{
+							ControlChannelMap.Add(Color.ParameterName, ChannelIndex);
+
+							bool bEnabled = ControlsMask[Index++];
+							FString ParameterString = Color.ParameterName.ToString();
+							FText Group = FText::FromString(ParameterString);
+
+							FMovieSceneChannelMetaData MetaData_R(*(ParameterString + TEXT("R")), FCommonChannelData::ChannelR, Group, bEnabled);
+							MetaData_R.SortOrder = ChannelIndex++;
+							MetaData_R.Color = FCommonChannelData::RedChannelColor;
+
+							FMovieSceneChannelMetaData MetaData_G(*(ParameterString + TEXT("G")), FCommonChannelData::ChannelG, Group, bEnabled);
+							MetaData_G.SortOrder = ChannelIndex++;
+							MetaData_G.Color = FCommonChannelData::GreenChannelColor;
+
+							FMovieSceneChannelMetaData MetaData_B(*(ParameterString + TEXT("B")), FCommonChannelData::ChannelB, Group, bEnabled);
+							MetaData_B.SortOrder = ChannelIndex++;
+							MetaData_B.Color = FCommonChannelData::BlueChannelColor;
+
+							FMovieSceneChannelMetaData MetaData_A(*(ParameterString + TEXT("A")), FCommonChannelData::ChannelA, Group, bEnabled);
+							MetaData_A.SortOrder = ChannelIndex++;
+
+							Channels.Add(Color.RedCurve, MetaData_R, TMovieSceneExternalValue<float>());
+							Channels.Add(Color.GreenCurve, MetaData_G, TMovieSceneExternalValue<float>());
+							Channels.Add(Color.BlueCurve, MetaData_B, TMovieSceneExternalValue<float>());
+							Channels.Add(Color.AlphaCurve, MetaData_A, TMovieSceneExternalValue<float>());
+							break;
+						}
+
+					}
+					break;
+				}
+				*/
+				case ERigControlType::TransformNoScale:
+				case ERigControlType::Transform:
+				{
+					for (FTransformParameterNameAndCurves& Transform : GetTransformParameterNamesAndCurves())
+					{
+						if (RigControl.Name == Transform.ParameterName)
+						{
+							ControlChannelMap.Add(Transform.ParameterName, FChannelMapInfo(Index, ChannelIndex));
+							bool bEnabled = ControlsMask[Index++];
+							FString ParameterString = Transform.ParameterName.ToString();
+							FText Group = FText::FromString(ParameterString);
+
+							FParameterTransformChannelEditorData EditorData(ControlRig, Transform.ParameterName, bEnabled, TransformMask.GetChannels(), Group, ChannelIndex);
+
+							Channels.Add(Transform.Translation[0], EditorData.MetaData[0], EditorData.ExternalValues[0]);
+							Channels.Add(Transform.Translation[1], EditorData.MetaData[1], EditorData.ExternalValues[1]);
+							Channels.Add(Transform.Translation[2], EditorData.MetaData[2], EditorData.ExternalValues[2]);
+
+							Channels.Add(Transform.Rotation[0], EditorData.MetaData[3], EditorData.ExternalValues[3]);
+							Channels.Add(Transform.Rotation[1], EditorData.MetaData[4], EditorData.ExternalValues[4]);
+							Channels.Add(Transform.Rotation[2], EditorData.MetaData[5], EditorData.ExternalValues[5]);
+
+							if (RigControl.ControlType == ERigControlType::Transform)
+							{
+								Channels.Add(Transform.Scale[0], EditorData.MetaData[6], EditorData.ExternalValues[6]);
+								Channels.Add(Transform.Scale[1], EditorData.MetaData[7], EditorData.ExternalValues[7]);
+								Channels.Add(Transform.Scale[2], EditorData.MetaData[8], EditorData.ExternalValues[8]);
+								ChannelIndex += 9;
+							}
+							else
+							{
+								ChannelIndex += 6;
+							}
+							break;
+						}
+					}
+				}
+				default:
+					break;
+				}
 #else
-
-	for (FScalarParameterNameAndCurve& Scalar : GetScalarParameterNamesAndCurves())
-	{
-		Channels.Add(Scalar.ParameterCurve);
-	}
-	for (FVectorParameterNameAndCurves& Vector : GetVectorParameterNamesAndCurves())
-	{
-		Channels.Add(Vector.XCurve);
-		Channels.Add(Vector.YCurve);
-		Channels.Add(Vector.ZCurve);
-	}
-	for (FColorParameterNameAndCurves& Color : GetColorParameterNamesAndCurves())
-	{
-		Channels.Add(Color.RedCurve);
-		Channels.Add(Color.GreenCurve);
-		Channels.Add(Color.BlueCurve);
-		Channels.Add(Color.AlphaCurve);
-	}
-
-	for (FTransformParameterNameAndCurves& Transform : GetTransformParameterNamesAndCurves())
-	{
-		Channels.Add(Transform.Translation[0]);
-		Channels.Add(Transform.Translation[1]);
-		Channels.Add(Transform.Translation[2]);
-
-		Channels.Add(Transform.Rotation[0]);
-		Channels.Add(Transform.Rotation[1]);
-		Channels.Add(Transform.Rotation[2]);
-
-		Channels.Add(Transform.Scale[0]);
-		Channels.Add(Transform.Scale[1]);
-		Channels.Add(Transform.Scale[2]);
-
-	}
-
+				switch (RigControl.ControlType)
+				{
+				case ERigControlType::Float:
+				{
+					for (FScalarParameterNameAndCurve& Scalar : GetScalarParameterNamesAndCurves())
+					{
+						if (RigControl.Name == Scalar.ParameterName)
+						{
+							ControlChannelMap.Add(Scalar.ParameterName, FChannelMapInfo(Index, ChannelIndex));
+							Channels.Add(Scalar.ParameterCurve);
+							ChannelIndex += 1;
+							Index++;
+							break;
+						}
+					}
+					break;
+				}
+				case ERigControlType::Bool:
+				{
+					for (FBoolParameterNameAndCurve& Bool : GetBoolParameterNamesAndCurves())
+					{
+						if (RigControl.Name == Bool.ParameterName)
+						{
+							ControlChannelMap.Add(Bool.ParameterName, FChannelMapInfo(Index, ChannelIndex));
+							Channels.Add(Bool.ParameterCurve);
+							ChannelIndex += 1;
+							Index++;
+							break;
+						}
+					}
+					break;
+				}
+				case ERigControlType::Vector2D:
+				{
+					for (FVector2DParameterNameAndCurves& Vector2D : GetVector2DParameterNamesAndCurves())
+					{
+						if (RigControl.Name == Vector2D.ParameterName)
+						{
+							ControlChannelMap.Add(Vector2D.ParameterName, FChannelMapInfo(Index, ChannelIndex));
+							Channels.Add(Vector2D.XCurve);
+							Channels.Add(Vector2D.YCurve);
+							ChannelIndex += 2;
+							Index++;
+							break;
+						}
+					}
+					break;
+				}
+				case ERigControlType::Position:
+				case ERigControlType::Scale:
+				case ERigControlType::Rotator:
+				{
+					for (FVectorParameterNameAndCurves& Vector : GetVectorParameterNamesAndCurves())
+					{
+						if (RigControl.Name == Vector.ParameterName)
+						{
+							ControlChannelMap.Add(Vector.ParameterName, FChannelMapInfo(Index, ChannelIndex));
+							Channels.Add(Vector.XCurve);
+							Channels.Add(Vector.YCurve);
+							Channels.Add(Vector.ZCurve);
+							ChannelIndex += 3;
+							Index++;
+							break;
+						}
+					}
+					break;
+				}
+				/*
+				for (FColorParameterNameAndCurves& Color : GetColorParameterNamesAndCurves())
+				{
+					Channels.Add(Color.RedCurve);
+					Channels.Add(Color.GreenCurve);
+					Channels.Add(Color.BlueCurve);
+					Channels.Add(Color.AlphaCurve);
+					break
+				}
+				*/
+				case ERigControlType::TransformNoScale:
+				case ERigControlType::Transform:
+				{
+					for (FTransformParameterNameAndCurves& Transform : GetTransformParameterNamesAndCurves())
+					{
+						if (RigControl.Name == Transform.ParameterName)
+						{
+							ControlChannelMap.Add(Transform.ParameterName, FChannelMapInfo(Index, ChannelIndex));
+							Channels.Add(Transform.Translation[0]);
+							Channels.Add(Transform.Translation[1]);
+							Channels.Add(Transform.Translation[2]);
+
+							Channels.Add(Transform.Rotation[0]);
+							Channels.Add(Transform.Rotation[1]);
+							Channels.Add(Transform.Rotation[2]);
+
+							if (RigControl.ControlType == ERigControlType::Transform)
+							{
+								Channels.Add(Transform.Scale[0]);
+								Channels.Add(Transform.Scale[1]);
+								Channels.Add(Transform.Scale[2]);
+								ChannelIndex += 9;
+							}
+							else
+							{
+								ChannelIndex += 6;
+							}
+
+							Index++;
+							break;
+						}
+					}
+					break;
+				}
+				}
 #endif
-
-	ChannelProxy = MakeShared<FMovieSceneChannelProxy>(MoveTemp(Channels));
+			}
+		}
+		ChannelProxy = MakeShared<FMovieSceneChannelProxy>(MoveTemp(Channels));
+	}
 }
 
 FMovieSceneInterrogationKey UMovieSceneControlRigParameterSection::GetFloatInterrogationKey()
@@ -653,12 +1236,24 @@ FMovieSceneInterrogationKey UMovieSceneControlRigParameterSection::GetFloatInter
 	return TypeID;
 }
 
+FMovieSceneInterrogationKey UMovieSceneControlRigParameterSection::GetVector2DInterrogationKey()
+{
+	static FMovieSceneAnimTypeID TypeID = FMovieSceneAnimTypeID::Unique();
+	return TypeID;
+}
+
 FMovieSceneInterrogationKey UMovieSceneControlRigParameterSection::GetVectorInterrogationKey()
 {
 	static FMovieSceneAnimTypeID TypeID = FMovieSceneAnimTypeID::Unique();
 	return TypeID;
 }
 
+FMovieSceneInterrogationKey UMovieSceneControlRigParameterSection::GetVector4InterrogationKey()
+{
+	static FMovieSceneAnimTypeID TypeID = FMovieSceneAnimTypeID::Unique();
+	return TypeID;
+}
+
 FMovieSceneInterrogationKey UMovieSceneControlRigParameterSection::GetTransformInterrogationKey()
 {
 	static FMovieSceneAnimTypeID TypeID = FMovieSceneAnimTypeID::Unique();
