COPY from //Dev-Anim to //Dev-Main

diff --git a/Engine/Plugins/Experimental/ControlRig/Source/ControlRig/Private/Sequencer/MovieSceneControlRigParameterTemplate.cpp b/Engine/Plugins/Experimental/ControlRig/Source/ControlRig/Private/Sequencer/MovieSceneControlRigParameterTemplate.cpp
--- a/Engine/Plugins/Experimental/ControlRig/Source/ControlRig/Private/Sequencer/MovieSceneControlRigParameterTemplate.cpp
+++ b/Engine/Plugins/Experimental/ControlRig/Source/ControlRig/Private/Sequencer/MovieSceneControlRigParameterTemplate.cpp
@@ -7,13 +7,14 @@
 #include "Evaluation/MovieSceneEvaluation.h"
 #include "IMovieScenePlayer.h"
 #include "Components/SkeletalMeshComponent.h"
-#include "Sequencer/ControlRigSequencerAnimInstance.h"
+#include "AnimCustomInstanceHelper.h"
+#include "Sequencer/ControlRigLayerInstance.h"
 #include "IControlRigObjectBinding.h"
 #include "Animation/AnimData/BoneMaskFilter.h"
 #include "ControlRigSkeletalMeshBinding.h"
 #include "Evaluation/Blending/BlendableTokenStack.h"
 #include "Evaluation/Blending/MovieSceneBlendingActuatorID.h"
-
+#include "TransformNoScale.h"
 
 //#include "Particles/ParticleSystemComponent.h"
 
@@ -37,7 +38,8 @@ struct FScalarParameterStringAndValue
 	FScalarParameterStringAndValue(FName InParameterName, float InValue)
 	{
 		ParameterName = InParameterName;
-		ParameterString = InParameterName.ToString();		Value = InValue;
+		ParameterString = InParameterName.ToString();		
+		Value = InValue;
 	}
 
 	/** The name of the scalar parameter. */
@@ -47,6 +49,47 @@ struct FScalarParameterStringAndValue
 	float Value;
 };
 
+/**
+ * Structure representing the animated value of a scalar parameter.
+ */
+struct FBoolParameterStringAndValue
+{
+	/** Creates a new FBoolParameterAndValue with a parameter name and a value. */
+	FBoolParameterStringAndValue(FName InParameterName, bool InValue)
+	{
+		ParameterName = InParameterName;
+		ParameterString = InParameterName.ToString();		
+		Value = InValue;
+	}
+
+	/** The name of the bool parameter. */
+	FString ParameterString;
+	FName ParameterName;
+	/** The animated value of the bool parameter. */
+	bool Value;
+};
+
+/**
+ * Structure representing the animated value of a vector2D parameter.
+ */
+struct FVector2DParameterStringAndValue
+{
+	/** Creates a new FVector2DParameterAndValue with a parameter name and a value. */
+	FVector2DParameterStringAndValue(FName InParameterName, FVector2D InValue)
+	{
+		ParameterName = InParameterName;
+		ParameterString = InParameterName.ToString();
+		Value = InValue;
+	}
+
+	/** The name of the vector2D parameter. */
+	FString ParameterString;
+	FName ParameterName;
+
+	/** The animated value of the vector2D parameter. */
+	FVector2D Value;
+};
+
 /**
  * Structure representing the animated value of a vector parameter.
  */
@@ -125,6 +168,10 @@ struct FEvaluatedControlRigParameterSectionValues
 
 	/** Array of evaluated scalar values */
 	TArray<FScalarParameterStringAndValue, TInlineAllocator<2>> ScalarValues;
+	/** Array of evaluated scalar values */
+	TArray<FBoolParameterStringAndValue, TInlineAllocator<2>> BoolValues;
+	/** Array of evaluated scalar values */
+	TArray<FVector2DParameterStringAndValue, TInlineAllocator<2>> Vector2DValues;
 	/** Array of evaluated vector values */
 	TArray<FVectorParameterStringAndValue, TInlineAllocator<2>> VectorValues;
 	/** Array of evaluated color values */
@@ -146,6 +193,27 @@ struct FControlRigTrackTokenFloat
 
 };
 
+struct FControlRigTrackTokenBool
+{
+	FControlRigTrackTokenBool() {}
+
+	FControlRigTrackTokenBool(bool InValue)
+		:Value(InValue)
+	{}
+
+	bool Value;
+};
+
+struct FControlRigTrackTokenVector2D
+{
+	FControlRigTrackTokenVector2D() {}
+	FControlRigTrackTokenVector2D(FVector2D InValue)
+		:Value(InValue)
+	{}
+
+	FVector2D Value;
+};
+
 struct FControlRigTrackTokenVector
 {
 	FControlRigTrackTokenVector() {}
@@ -175,6 +243,18 @@ template<> FMovieSceneAnimTypeID GetBlendingDataType<FControlRigTrackTokenFloat>
 	return TypeId;
 }
 
+template<> FMovieSceneAnimTypeID GetBlendingDataType<FControlRigTrackTokenBool>()
+{
+	static FMovieSceneAnimTypeID TypeId = FMovieSceneAnimTypeID::Unique();
+	return TypeId;
+}
+
+template<> FMovieSceneAnimTypeID GetBlendingDataType<FControlRigTrackTokenVector2D>()
+{
+	static FMovieSceneAnimTypeID TypeId = FMovieSceneAnimTypeID::Unique();
+	return TypeId;
+}
+
 template<> FMovieSceneAnimTypeID GetBlendingDataType<FControlRigTrackTokenVector>()
 {
 	static FMovieSceneAnimTypeID TypeId = FMovieSceneAnimTypeID::Unique();
@@ -195,6 +275,16 @@ template<>  struct TBlendableTokenTraits<FControlRigTrackTokenFloat>
 	typedef MovieScene::TMaskedBlendable<float, 1> WorkingDataType;
 };
 
+template<>  struct TBlendableTokenTraits<FControlRigTrackTokenBool>
+{
+	typedef MovieScene::TMaskedBlendable<bool, 1> WorkingDataType;
+};
+
+template<> struct TBlendableTokenTraits<FControlRigTrackTokenVector2D>
+{
+	typedef MovieScene::TMaskedBlendable<float, 2> WorkingDataType;
+};
+
 template<> struct TBlendableTokenTraits<FControlRigTrackTokenVector>
 {
 	typedef MovieScene::TMaskedBlendable<float, 3> WorkingDataType;
@@ -221,6 +311,26 @@ namespace MovieScene
 		Out.Value = In[0];
 	}
 
+	void MultiChannelFromData(const FControlRigTrackTokenBool& In, TMultiChannelValue<bool, 1>& Out)
+	{
+		Out = { In.Value };
+	}
+
+	void ResolveChannelsToData(const TMultiChannelValue<bool, 1>& In, FControlRigTrackTokenBool& Out)
+	{
+		Out.Value = In[0];
+	}
+
+	void MultiChannelFromData(const FControlRigTrackTokenVector2D& In, TMultiChannelValue<float, 2>& Out)
+	{
+		Out = { In.Value.X, In.Value.Y };
+	}
+
+	void ResolveChannelsToData(const TMultiChannelValue<float, 2>& In, FControlRigTrackTokenVector2D& Out)
+	{
+		Out.Value = FVector2D(In[0], In[1]);
+	}
+
 	void MultiChannelFromData(const FControlRigTrackTokenVector& In, TMultiChannelValue<float, 3>& Out)
 	{
 		Out = { In.Value.X, In.Value.Y, In.Value.Z };
@@ -250,7 +360,44 @@ namespace MovieScene
 	}
 }
 
+namespace FControlRigBindingHelper
+{
+	void BindToSequencerInstance(UControlRig* ControlRig)
+	{
+		check(ControlRig);
+		if (USkeletalMeshComponent* SkeletalMeshComponent = Cast<USkeletalMeshComponent>(ControlRig->GetObjectBinding()->GetBoundObject()))
+		{
+			bool bWasCreated = false;
+			if (UControlRigLayerInstance* AnimInstance = FAnimCustomInstanceHelper::BindToSkeletalMeshComponent<UControlRigLayerInstance>(SkeletalMeshComponent, bWasCreated))
+			{
+				if (bWasCreated)
+				{
+					AnimInstance->RecalcRequiredBones();
+					AnimInstance->AddControlRigTrack(ControlRig->GetUniqueID(), ControlRig);
+					ControlRig->CreateRigControlsForCurveContainer();
+					ControlRig->Execute(EControlRigState::Init);
+				}
+			}
+		}
+	}
 
+	void UnBindFromSequencerInstance(UControlRig* ControlRig)
+	{
+		check(ControlRig);
+		if (USkeletalMeshComponent* SkeletalMeshComponent = Cast<USkeletalMeshComponent>(ControlRig->GetObjectBinding()->GetBoundObject()))
+		{
+			UControlRigLayerInstance* AnimInstance = Cast<UControlRigLayerInstance>(SkeletalMeshComponent->GetAnimInstance());
+			if (AnimInstance)
+			{
+				AnimInstance->ResetNodes();
+				AnimInstance->RecalcRequiredBones();
+				AnimInstance->RemoveControlRigTrack(ControlRig->GetUniqueID());
+			}
+
+			FAnimCustomInstanceHelper::UnbindFromSkeletalMeshComponent< UControlRigLayerInstance>(SkeletalMeshComponent);
+		}
+	}
+}
 
 struct FControlRigParameterPreAnimatedTokenProducer : IMovieScenePreAnimatedTokenProducer
 {
@@ -265,36 +412,56 @@ struct FControlRigParameterPreAnimatedTokenProducer : IMovieScenePreAnimatedToke
 		{
 			FToken(FMovieSceneSequenceIDRef InSequenceID)
 				: SequenceID(InSequenceID)
-			{}
+#if WITH_EDITOR
+				,bUpdateAnimationInEditor(true)
+#endif
+			{
+
+			}
+			void SetSkelMesh(USkeletalMeshComponent* InComponent)
+			{
+#if WITH_EDITOR
+				bUpdateAnimationInEditor = InComponent->GetUpdateAnimationInEditor();
+#endif
+			}
 
 			virtual void RestoreState(UObject& InObject, IMovieScenePlayer& Player) override
 			{
+
 				if (UControlRig* ControlRig = Cast<UControlRig>(&InObject))
 				{
 					if (ControlRig->GetObjectBinding())
 					{
 						if (USkeletalMeshComponent* SkeletalMeshComponent = Cast<USkeletalMeshComponent>(ControlRig->GetObjectBinding()->GetBoundObject()))
 						{
-							if (UControlRigSequencerAnimInstance* AnimInstance = Cast<UControlRigSequencerAnimInstance>(SkeletalMeshComponent->GetAnimInstance()))
-							{
-								AnimInstance->ResetNodes();
-								AnimInstance->RecalcRequiredBones();
-							}
-							UAnimSequencerInstance::UnbindFromSkeletalMeshComponent(SkeletalMeshComponent);
+#if WITH_EDITOR
+							SkeletalMeshComponent->SetUpdateAnimationInEditor(bUpdateAnimationInEditor);
+#endif
 						}
+						FControlRigBindingHelper::UnBindFromSequencerInstance(ControlRig);
 						for (TNameAndValue<float>& Value : ScalarValues)
 						{
-							ControlRig->SetControlValue<float>(Value.Name, Value.Value);
+							ControlRig->SetControlValue<float>(Value.Name, Value.Value, true, EControlRigSetKey::Never);
+						}
+
+						for (TNameAndValue<bool>& Value : BoolValues)
+						{
+							ControlRig->SetControlValue<bool>(Value.Name, Value.Value, true, EControlRigSetKey::Never);
+						}
+
+						for (TNameAndValue<FVector2D>& Value : Vector2DValues)
+						{
+							ControlRig->SetControlValue<FVector2D>(Value.Name, Value.Value,true, EControlRigSetKey::Never);
 						}
 
 						for (TNameAndValue<FVector>& Value : VectorValues)
 						{
-							ControlRig->SetControlValue<FVector>(Value.Name, Value.Value);
+							ControlRig->SetControlValue<FVector>(Value.Name, Value.Value, true, EControlRigSetKey::Never);
 						}
 
 						for (TNameAndValue<FTransform>& Value : TransformValues)
 						{
-							ControlRig->SetControlValue<FTransform>(Value.Name, Value.Value);
+							ControlRig->SetControlValue<FTransform>(Value.Name, Value.Value, true, EControlRigSetKey::Never);
 						}
 						ControlRig->GetObjectBinding()->UnbindFromObject();
 					}
@@ -303,8 +470,13 @@ struct FControlRigParameterPreAnimatedTokenProducer : IMovieScenePreAnimatedToke
 
 			FMovieSceneSequenceID SequenceID;
 			TArray< TNameAndValue<float> > ScalarValues;
+			TArray< TNameAndValue<bool> > BoolValues;
 			TArray< TNameAndValue<FVector> > VectorValues;
+			TArray< TNameAndValue<FVector2D> > Vector2DValues;
 			TArray< TNameAndValue<FTransform> > TransformValues;
+#if WITH_EDITOR
+			bool bUpdateAnimationInEditor;
+#endif
 		};
 
 
@@ -312,22 +484,21 @@ struct FControlRigParameterPreAnimatedTokenProducer : IMovieScenePreAnimatedToke
 
 		if (UControlRig* ControlRig = Cast<UControlRig>(&Object))
 		{
+			if (ControlRig->GetObjectBinding())
+			{
+				if (USkeletalMeshComponent* SkeletalMeshComponent = Cast<USkeletalMeshComponent>(ControlRig->GetObjectBinding()->GetBoundObject()))
+				{
+					if (SkeletalMeshComponent)
+					{
+						Token.SetSkelMesh(SkeletalMeshComponent);
+					}
+				}
+			}
+
 
 			/*
 			//mz todo the other types
-
-			UENUM()
-			enum class ERigControlType : uint8
-			{
-				Bool,
-				Float,
-				Vector2D,
-				Position,
-				Scale,
-				Quat,
-				Rotator,
-				Transform
-			};
+			just left is bool
 			*/
 
 			const TArray<FRigControl>& Controls = ControlRig->AvailableControls();
@@ -342,6 +513,12 @@ struct FControlRigParameterPreAnimatedTokenProducer : IMovieScenePreAnimatedToke
 					Token.ScalarValues.Add(TNameAndValue<float>{ RigControl.Name, Val });
 					break;
 				}
+				case ERigControlType::Vector2D:
+				{
+					FVector2D Val = RigControl.Value.Get<FVector2D>();
+					Token.Vector2DValues.Add(TNameAndValue<FVector2D>{ RigControl.Name, Val });
+					break;
+				}
 				case ERigControlType::Position:
 				case ERigControlType::Scale:
 				case ERigControlType::Rotator:
@@ -357,6 +534,13 @@ struct FControlRigParameterPreAnimatedTokenProducer : IMovieScenePreAnimatedToke
 					Token.TransformValues.Add(TNameAndValue<FTransform>{ RigControl.Name, Val });
 					break;
 				}
+				case ERigControlType::TransformNoScale:
+				{
+					FTransformNoScale NoScale = RigControl.Value.Get<FTransformNoScale>();
+					FTransform Val = NoScale;
+					Token.TransformValues.Add(TNameAndValue<FTransform>{ RigControl.Name, Val });
+					break;
+				}
 				}
 			}
 		}
@@ -365,7 +549,9 @@ struct FControlRigParameterPreAnimatedTokenProducer : IMovieScenePreAnimatedToke
 	}
 
 	FMovieSceneSequenceID SequenceID;
+	TArray< TNameAndValue<bool> > BoolValues;
 	TArray< TNameAndValue<float> > ScalarValues;
+	TArray< TNameAndValue<FVector2D> > Vector2DValues;
 	TArray< TNameAndValue<FVector> > VectorValues;
 	TArray< TNameAndValue<FTransform> > TransformValues;
 
@@ -386,28 +572,6 @@ struct FControlRigParameterExecutionToken : IMovieSceneExecutionToken
 	FControlRigParameterExecutionToken(const FControlRigParameterExecutionToken&) = delete;
 	FControlRigParameterExecutionToken& operator=(const FControlRigParameterExecutionToken&) = delete;
 
-	void BindToSequencerInstance(UControlRig* ControlRig)
-	{
-		check(ControlRig);
-		if (USkeletalMeshComponent* SkeletalMeshComponent = Cast<USkeletalMeshComponent>(ControlRig->GetObjectBinding()->GetBoundObject()))
-		{
-			bool bWasCreated = false;
-			if (UControlRigSequencerAnimInstance* AnimInstance = UAnimCustomInstance::BindToSkeletalMeshComponent<UControlRigSequencerAnimInstance>(SkeletalMeshComponent, bWasCreated))
-			{
-				AnimInstance->RecalcRequiredBones();
-			}
-		}
-	}
-
-	void UnBindFromSequencerInstance(UControlRig* ControlRig)
-	{
-		check(ControlRig);
-		if (USkeletalMeshComponent* SkeletalMeshComponent = Cast<USkeletalMeshComponent>(ControlRig->GetObjectBinding()->GetBoundObject()))
-		{
-			UAnimCustomInstance::UnbindFromSkeletalMeshComponent(SkeletalMeshComponent);
-		}
-	}
-
 	virtual void Execute(const FMovieSceneContext& Context, const FMovieSceneEvaluationOperand& Operand, FPersistentEvaluationData& PersistentData, IMovieScenePlayer& Player)
 	{
 		MOVIESCENE_DETAILED_SCOPE_CYCLE_COUNTER(MovieSceneEval_ControlRigParameterTrack_TokenExecute)
@@ -437,21 +601,21 @@ struct FControlRigParameterExecutionToken : IMovieSceneExecutionToken
 				if (!ControlRig->GetObjectBinding()->GetBoundObject())
 				{
 					ControlRig->GetObjectBinding()->BindToObject(BoundSkelMeshes[0].Get());
+					ControlRig->Execute(EControlRigState::Init);
 				}
-				BindToSequencerInstance(ControlRig);
+				FControlRigBindingHelper::BindToSequencerInstance(ControlRig);
 				//MZ TODO HANDLE BOOLS AND OTHER NON BLENDABLES
 				if (ControlRig->GetObjectBinding())
 				{
 					if (USkeletalMeshComponent* SkeletalMeshComponent = Cast<USkeletalMeshComponent>(ControlRig->GetObjectBinding()->GetBoundObject()))
 					{
-						if (UControlRigSequencerAnimInstance* AnimInstance = Cast<UControlRigSequencerAnimInstance>(SkeletalMeshComponent->GetAnimInstance()))
+						if (UControlRigLayerInstance* AnimInstance = Cast<UControlRigLayerInstance>(SkeletalMeshComponent->GetAnimInstance()))
 						{
 							float Weight = 1.0f;
-							bool bStructureChanged = AnimInstance->UpdateControlRig(ControlRig, Operand.SequenceID.GetInternalValue(), bAdditive, bApplyBoneFilter, *BoneFilter, Weight, true, true);
-							if (bStructureChanged)
-							{
-								AnimInstance->RecalcRequiredBones();
-							}
+							FControlRigIOSettings InputSettings;
+							InputSettings.bUpdateCurves = false;
+							InputSettings.bUpdatePose = true;
+							AnimInstance->UpdateControlRigTrack(ControlRig->GetUniqueID(), Weight, InputSettings, true);
 						}
 					}
 				}
@@ -516,7 +680,7 @@ struct TControlRigParameterActuatorFloat : TMovieSceneBlendingActuator<FControlR
 			FRigControl* RigControl = ControlRig->FindControl(ParameterName);
 			if (RigControl && RigControl->ControlType == ERigControlType::Float)
 			{
-				ControlRig->SetControlValue<float>(ParameterName, InFinalValue.Value);
+				ControlRig->SetControlValue<float>(ParameterName, InFinalValue.Value, true, EControlRigSetKey::Never);
 			}
 		}
 		if (Section)
@@ -541,6 +705,69 @@ struct TControlRigParameterActuatorFloat : TMovieSceneBlendingActuator<FControlR
 };
 
 
+struct TControlRigParameterActuatorVector2D : TMovieSceneBlendingActuator<FControlRigTrackTokenVector2D>
+{
+	TControlRigParameterActuatorVector2D(FMovieSceneAnimTypeID& InAnimID, UControlRig* InControlRig, const FName& InParameterName, const UMovieSceneControlRigParameterSection* InSection)
+		: TMovieSceneBlendingActuator<FControlRigTrackTokenVector2D>(FMovieSceneBlendingActuatorID(InAnimID))
+		, ControlRig(InControlRig)
+		, ParameterName(InParameterName)
+		, SectionData(InSection)
+	{}
+
+
+
+	FControlRigTrackTokenVector2D RetrieveCurrentValue(UObject* InObject, IMovieScenePlayer* Player) const
+	{
+		if (ControlRig)
+		{
+			FRigControl* RigControl = ControlRig->FindControl(ParameterName);
+			if (RigControl && (RigControl->ControlType == ERigControlType::Vector2D))
+			{
+				FVector2D Val = RigControl->Value.Get<FVector2D>();
+				return FControlRigTrackTokenVector2D(Val);
+			}
+		}
+		return FControlRigTrackTokenVector2D();
+	}
+
+	void Actuate(UObject* InObject, const FControlRigTrackTokenVector2D& InFinalValue, const TBlendableTokenStack<FControlRigTrackTokenVector2D>& OriginalStack, const FMovieSceneContext& Context, FPersistentEvaluationData& PersistentData, IMovieScenePlayer& Player)
+	{
+		const UMovieSceneControlRigParameterSection* Section = nullptr;
+		Section = SectionData.Get();
+
+		bool bWasDoNotKey = false;
+		if (Section)
+		{
+			bWasDoNotKey = Section->GetDoNotKey();
+			Section->SetDoNotKey(true);
+		}
+		if (ControlRig)
+		{
+			FRigControl* RigControl = ControlRig->FindControl(ParameterName);
+			if (RigControl && (RigControl->ControlType == ERigControlType::Vector2D))
+			{
+				ControlRig->SetControlValue<FVector2D>(ParameterName, InFinalValue.Value, true, EControlRigSetKey::Never);
+			}
+		}
+		if (Section)
+		{
+			Section->SetDoNotKey(bWasDoNotKey);
+		}
+	}
+	virtual void Actuate(FMovieSceneInterrogationData& InterrogationData, const FControlRigTrackTokenVector2D& InValue, const TBlendableTokenStack<FControlRigTrackTokenVector2D>& OriginalStack, const FMovieSceneContext& Context) const override
+	{
+		FVector2DInterrogationData Data;
+		Data.Val = InValue.Value;
+		Data.ParameterName = ParameterName;
+		InterrogationData.Add(FVector2DInterrogationData(Data), UMovieSceneControlRigParameterSection::GetVectorInterrogationKey());
+	}
+
+	UControlRig *ControlRig;
+	FName ParameterName;
+	TWeakObjectPtr<const UMovieSceneControlRigParameterSection> SectionData;
+
+};
+
 
 struct TControlRigParameterActuatorVector : TMovieSceneBlendingActuator<FControlRigTrackTokenVector>
 {
@@ -583,7 +810,7 @@ struct TControlRigParameterActuatorVector : TMovieSceneBlendingActuator<FControl
 			FRigControl* RigControl = ControlRig->FindControl(ParameterName);
 			if (RigControl && (RigControl->ControlType == ERigControlType::Position || RigControl->ControlType == ERigControlType::Scale || RigControl->ControlType == ERigControlType::Rotator))
 			{
-				ControlRig->SetControlValue<FVector>(ParameterName, InFinalValue.Value);
+				ControlRig->SetControlValue<FVector>(ParameterName, InFinalValue.Value, true, EControlRigSetKey::Never);
 			}
 		}
 		if (Section)
@@ -626,6 +853,12 @@ struct TControlRigParameterActuatorTransform : TMovieSceneBlendingActuator<FCont
 				FTransform Val = RigControl->Value.Get<FTransform>();
 				return FControlRigTrackTokenTransform(Val);
 			}
+			else if(RigControl && RigControl->ControlType == ERigControlType::TransformNoScale)
+			{
+				FTransformNoScale ValNoScale = RigControl->Value.Get<FTransformNoScale>();
+				FTransform Val = ValNoScale;
+				return FControlRigTrackTokenTransform(Val);
+			}
 		}
 		return FControlRigTrackTokenTransform();
 	}
@@ -655,7 +888,12 @@ struct TControlRigParameterActuatorTransform : TMovieSceneBlendingActuator<FCont
 			FRigControl* RigControl = ControlRig->FindControl(ParameterName);
 			if (RigControl && RigControl->ControlType == ERigControlType::Transform)
 			{
-				ControlRig->SetControlValue<FTransform>(ParameterName, InFinalValue.Value);
+				ControlRig->SetControlValue<FTransform>(ParameterName, InFinalValue.Value,true, EControlRigSetKey::Never);
+			}
+			else if (RigControl && RigControl->ControlType == ERigControlType::TransformNoScale)
+			{
+				FTransformNoScale NoScale = InFinalValue.Value;
+				ControlRig->SetControlValue<FTransformNoScale>(ParameterName, NoScale, true, EControlRigSetKey::Never);
 			}
 		}
 		if (Section)
@@ -683,11 +921,9 @@ void FMovieSceneControlRigParameterTemplate::Evaluate(const FMovieSceneEvaluatio
 
 	const FFrameTime Time = Context.GetTime();
 
-	const UMovieSceneControlRigParameterSection* Section = nullptr;
-	if (GetSourceSection())
+	const UMovieSceneControlRigParameterSection* Section = Cast<UMovieSceneControlRigParameterSection>(GetSourceSection());
+	if (Section)
 	{
-		Section = Cast<UMovieSceneControlRigParameterSection>(GetSourceSection());
-
 		//Do basic token
 		FControlRigParameterExecutionToken ExecutionToken(Section);
 		ExecutionTokens.Add(MoveTemp(ExecutionToken));
@@ -696,7 +932,9 @@ void FMovieSceneControlRigParameterTemplate::Evaluate(const FMovieSceneEvaluatio
 		FEvaluatedControlRigParameterSectionValues Values;
 		EvaluateCurvesWithMasks(Context, Values);
 		static TMovieSceneAnimTypeIDContainer<FString> ScalarAnimTypeIDsByName;
+		static TMovieSceneAnimTypeIDContainer<FString> BoolAnimTypeIDsByName;
 		static TMovieSceneAnimTypeIDContainer<FString> VectorAnimTypeIDsByName;
+		static TMovieSceneAnimTypeIDContainer<FString> Vector2DAnimTypeIDsByName;
 		static TMovieSceneAnimTypeIDContainer<FString> TransformAnimTypeIDsByName;
 
 
@@ -710,8 +948,29 @@ void FMovieSceneControlRigParameterTemplate::Evaluate(const FMovieSceneEvaluatio
 	//	Weight *= ManualWeight;
 	//}
 
+		//Do Bool straight up no blending
+		bool bWasDoNotKey = false;
+	
+		bWasDoNotKey = Section->GetDoNotKey();
+		Section->SetDoNotKey(true);
+		
+		if (Section->ControlRig)
+		{
+			for (const FBoolParameterStringAndValue& BoolNameAndValue : Values.BoolValues)
+			{
+				FRigControl* RigControl = Section->ControlRig->FindControl(BoolNameAndValue.ParameterName);
+				if (RigControl && RigControl->ControlType == ERigControlType::Bool)
+				{
+					Section->ControlRig->SetControlValue<bool>(BoolNameAndValue.ParameterName, BoolNameAndValue.Value, true, EControlRigSetKey::Never);
+				}
+			}
+		}
+		Section->SetDoNotKey(bWasDoNotKey);
+		
+		
 		uint32 OperandHash = GetTypeHash(Operand);
 		FString UniqueActuator(FString::FromInt((int32)OperandHash));
+
 		for (const FScalarParameterStringAndValue& ScalarNameAndValue : Values.ScalarValues)
 		{
 			FString NewString(ScalarNameAndValue.ParameterString);
@@ -727,7 +986,6 @@ void FMovieSceneControlRigParameterTemplate::Evaluate(const FMovieSceneEvaluatio
 		}
 
 		MovieScene::TMultiChannelValue<float, 3> VectorData;
-
 		for (const FVectorParameterStringAndValue& VectorNameAndValue : Values.VectorValues)
 		{
 			FString NewString(VectorNameAndValue.ParameterString);
@@ -745,8 +1003,26 @@ void FMovieSceneControlRigParameterTemplate::Evaluate(const FMovieSceneEvaluatio
 
 			ExecutionTokens.BlendToken(ActuatorTypeID, TBlendableToken<FControlRigTrackTokenVector>(VectorData, Section->GetBlendType().Get(), Weight));
 		}
-		MovieScene::TMultiChannelValue<float, 9> TransformData;
 
+		MovieScene::TMultiChannelValue<float, 2> Vector2DData;
+		for (const FVector2DParameterStringAndValue& Vector2DNameAndValue : Values.Vector2DValues)
+		{
+			FString NewString(Vector2DNameAndValue.ParameterString);
+			NewString.Append(UniqueActuator);
+			FMovieSceneAnimTypeID AnimTypeID = Vector2DAnimTypeIDsByName.GetAnimTypeID(NewString);
+			FMovieSceneBlendingActuatorID ActuatorTypeID(AnimTypeID);
+
+			if (!ExecutionTokens.GetBlendingAccumulator().FindActuator< FControlRigTrackTokenVector2D>(ActuatorTypeID))
+			{
+				ExecutionTokens.GetBlendingAccumulator().DefineActuator(ActuatorTypeID, MakeShared <TControlRigParameterActuatorVector2D>(AnimTypeID, Section->ControlRig, Vector2DNameAndValue.ParameterName, Section));
+			}
+			Vector2DData.Set(0, Vector2DNameAndValue.Value.X);
+			Vector2DData.Set(1, Vector2DNameAndValue.Value.Y);
+
+			ExecutionTokens.BlendToken(ActuatorTypeID, TBlendableToken<FControlRigTrackTokenVector2D>(Vector2DData, Section->GetBlendType().Get(), Weight));
+		}
+
+		MovieScene::TMultiChannelValue<float, 9> TransformData;
 		for (const FTransformParameterStringAndValue& TransformNameAndValue : Values.TransformValues)
 		{
 			FString NewString(TransformNameAndValue.ParameterString);
@@ -790,23 +1066,50 @@ void FMovieSceneControlRigParameterTemplate::EvaluateCurvesWithMasks(const FMovi
 		Section = Cast<UMovieSceneControlRigParameterSection>(GetSourceSection());
 		TArray<bool> ControlsMask = Section->GetControlsMask();
 
-		int32 Index = 0;
+		//mz todo optimize this, don't want to do this map search every tick, will cache on the NameAndCurve objects
+		const FChannelMapInfo* ChannelInfo = nullptr;
 		for (const FScalarParameterNameAndCurve& Scalar : Scalars)
 		{
 			float Value = 0;
-			if (ControlsMask[Index])
+			ChannelInfo = Section->ControlChannelMap.Find(Scalar.ParameterName);
+			if (!ChannelInfo || ControlsMask[ChannelInfo->ControlIndex])
 			{
 				Scalar.ParameterCurve.Evaluate(Time, Value);
 			}
 			
 			Values.ScalarValues.Emplace(Scalar.ParameterName, Value);
-			++Index;
+		}
+
+		for (const FBoolParameterNameAndCurve& Bool : Bools)
+		{
+			bool Value = false;
+			ChannelInfo = Section->ControlChannelMap.Find(Bool.ParameterName);
+			if (!ChannelInfo || ControlsMask[ChannelInfo->ControlIndex])
+			{
+				Bool.ParameterCurve.Evaluate(Time, Value);
+			}
+
+			Values.BoolValues.Emplace(Bool.ParameterName, Value);
+		}
+
+		for (const FVector2DParameterNameAndCurves& Vector2D : Vector2Ds)
+		{
+			FVector2D Value(ForceInitToZero);
+			ChannelInfo = Section->ControlChannelMap.Find(Vector2D.ParameterName);
+			if (!ChannelInfo || ControlsMask[ChannelInfo->ControlIndex])
+			{
+				Vector2D.XCurve.Evaluate(Time, Value.X);
+				Vector2D.YCurve.Evaluate(Time, Value.Y);
+			}
+
+			Values.Vector2DValues.Emplace(Vector2D.ParameterName, Value);
 		}
 
 		for (const FVectorParameterNameAndCurves& Vector : Vectors)
 		{
 			FVector Value(ForceInitToZero);
-			if (ControlsMask[Index])
+			ChannelInfo = Section->ControlChannelMap.Find(Vector.ParameterName);
+			if (!ChannelInfo || ControlsMask[ChannelInfo->ControlIndex])
 			{
 				Vector.XCurve.Evaluate(Time, Value.X);
 				Vector.YCurve.Evaluate(Time, Value.Y);
@@ -814,13 +1117,13 @@ void FMovieSceneControlRigParameterTemplate::EvaluateCurvesWithMasks(const FMovi
 			}
 
 			Values.VectorValues.Emplace(Vector.ParameterName, Value);
-			++Index;
 		}
 
 		for (const FColorParameterNameAndCurves& Color : Colors)
 		{
 			FLinearColor ColorValue = FLinearColor::White;
-			if (ControlsMask[Index])
+			ChannelInfo = Section->ControlChannelMap.Find(Color.ParameterName);
+			if (!ChannelInfo || ControlsMask[ChannelInfo->ControlIndex])
 			{
 				Color.RedCurve.Evaluate(Time, ColorValue.R);
 				Color.GreenCurve.Evaluate(Time, ColorValue.G);
@@ -829,7 +1132,6 @@ void FMovieSceneControlRigParameterTemplate::EvaluateCurvesWithMasks(const FMovi
 			}
 		
 			Values.ColorValues.Emplace(Color.ParameterName, ColorValue);
-			++Index;
 		}
 
 		for (const FTransformParameterNameAndCurves& Transform : Transforms)
@@ -837,8 +1139,9 @@ void FMovieSceneControlRigParameterTemplate::EvaluateCurvesWithMasks(const FMovi
 			FVector Translation(ForceInitToZero), Scale(FVector::OneVector);
 			FRotator Rotator(0.0f, 0.0f, 0.0f);
 			EMovieSceneTransformChannel ChannelMask = Section->GetTransformMask().GetChannels();
+			ChannelInfo = Section->ControlChannelMap.Find(Transform.ParameterName);
 
-			if (ControlsMask[Index])
+			if (!ChannelInfo || ControlsMask[ChannelInfo->ControlIndex])
 			{
 				if (EnumHasAllFlags(ChannelMask, EMovieSceneTransformChannel::TranslationX))
 				{
@@ -880,8 +1183,6 @@ void FMovieSceneControlRigParameterTemplate::EvaluateCurvesWithMasks(const FMovi
 			}
 			FTransformParameterStringAndValue NameAndValue(Transform.ParameterName, Translation, Rotator, Scale);
 			Values.TransformValues.Emplace(NameAndValue);
-			++Index;
-
 		}
 	}
 }
@@ -908,6 +1209,8 @@ void FMovieSceneControlRigParameterTemplate::Interrogate(const FMovieSceneContex
 		FEvaluatedControlRigParameterSectionValues Values;
 		EvaluateCurvesWithMasks(Context, Values);
 		static TMovieSceneAnimTypeIDContainer<FString> ScalarAnimTypeIDsByName;
+		static TMovieSceneAnimTypeIDContainer<FString> BoolAnimTypeIDsByName;
+		static TMovieSceneAnimTypeIDContainer<FString> Vector2DAnimTypeIDsByName;
 		static TMovieSceneAnimTypeIDContainer<FString> VectorAnimTypeIDsByName;
 		static TMovieSceneAnimTypeIDContainer<FString> TransformAnimTypeIDsByName;
 
@@ -935,8 +1238,23 @@ void FMovieSceneControlRigParameterTemplate::Interrogate(const FMovieSceneContex
 			Container.GetAccumulator().BlendToken(FMovieSceneEvaluationOperand(), ActuatorTypeID, FMovieSceneEvaluationScope(), Context,TBlendableToken<FControlRigTrackTokenFloat>(ScalarNameAndValue.Value, Section->GetBlendType().Get(), Weight));
 		}
 
-		MovieScene::TMultiChannelValue<float, 3> VectorData;
+		MovieScene::TMultiChannelValue<float, 2> Vector2DData;
+		for (const FVector2DParameterStringAndValue& Vector2DNameAndValue : Values.Vector2DValues)
+		{
+			FMovieSceneAnimTypeID AnimTypeID = Vector2DAnimTypeIDsByName.GetAnimTypeID(Vector2DNameAndValue.ParameterString);
+			FMovieSceneBlendingActuatorID ActuatorTypeID(AnimTypeID);
 
+			if (!Container.GetAccumulator().FindActuator< FControlRigTrackTokenVector>(ActuatorTypeID))
+			{
+				Container.GetAccumulator().DefineActuator(ActuatorTypeID, MakeShared <TControlRigParameterActuatorVector2D>(AnimTypeID, Section->ControlRig, Vector2DNameAndValue.ParameterName, Section));
+			}
+			Vector2DData.Set(0, Vector2DNameAndValue.Value.X);
+			Vector2DData.Set(1, Vector2DNameAndValue.Value.Y);
+
+			Container.GetAccumulator().BlendToken(FMovieSceneEvaluationOperand(), ActuatorTypeID, FMovieSceneEvaluationScope(), Context, TBlendableToken<FControlRigTrackTokenVector2D>(Vector2DData, Section->GetBlendType().Get(), Weight));
+		}
+
+		MovieScene::TMultiChannelValue<float, 3> VectorData;
 		for (const FVectorParameterStringAndValue& VectorNameAndValue : Values.VectorValues)
 		{
 			FMovieSceneAnimTypeID AnimTypeID = VectorAnimTypeIDsByName.GetAnimTypeID(VectorNameAndValue.ParameterString);
@@ -952,8 +1270,8 @@ void FMovieSceneControlRigParameterTemplate::Interrogate(const FMovieSceneContex
 
 			Container.GetAccumulator().BlendToken(FMovieSceneEvaluationOperand(), ActuatorTypeID, FMovieSceneEvaluationScope(),Context,TBlendableToken<FControlRigTrackTokenVector>(VectorData, Section->GetBlendType().Get(), Weight));
 		}
-		MovieScene::TMultiChannelValue<float, 9> TransformData;
 
+		MovieScene::TMultiChannelValue<float, 9> TransformData;
 		for (const FTransformParameterStringAndValue& TransformNameAndValue : Values.TransformValues)
 		{
 			uint32 ID = BindingOverride->GetUniqueID();
diff --git a/Engine/Plugins/Experimental/ControlRig/Source/ControlRig/Private/Sequencer/MovieSceneControlRigParameterTrack.cpp b/Engine/Plugins/Experimental/ControlRig/Source/ControlRig/Private/Sequencer/MovieSceneControlRigParameterTrack.cpp
--- a/Engine/Plugins/Experimental/ControlRig/Source/ControlRig/Private/Sequencer/MovieSceneControlRigParameterTrack.cpp
+++ b/Engine/Plugins/Experimental/ControlRig/Source/ControlRig/Private/Sequencer/MovieSceneControlRigParameterTrack.cpp
@@ -7,7 +7,7 @@
 #include "Sequencer/MovieSceneControlRigParameterTemplate.h"
 #include "MovieScene.h"
 #include "MovieSceneTimeHelpers.h"
-
+#include "Sequencer/ControlRigSortedControls.h"
 #define LOCTEXT_NAMESPACE "MovieSceneParameterControlRigTrack"
 
 
@@ -47,77 +47,94 @@ UMovieSceneSection* UMovieSceneControlRigParameterTrack::CreateNewSection()
 	{
 		NewSection->SetBlendType(EMovieSceneBlendType::Additive);
 	}
-	UMovieScene* OuterMovieScene = GetTypedOuter<UMovieScene>();
-	IControlRigManipulatable* Manip = GetManipulatableFromBinding(OuterMovieScene);
-	if (Manip)
+	if (ControlRig)
 	{
-		/*
-		UENUM()
-		enum class ERigControlType : uint8
-		{
-			Bool,
-			Float,
-			Vector2D,
-			Position,
-			Scale,
-			Quat,
-			Rotator,
-			Transform
-		};
-		*/
-
-		const TArray<FRigControl>& Controls = Manip->AvailableControls();
 		TArray<bool> OnArray;
+		const TArray<FRigControl>& Controls = ControlRig->AvailableControls();
 		OnArray.Init(true,Controls.Num());
 		NewSection->SetControlsMask(OnArray);
 
-
-		for (const FRigControl& RigControl : Controls)
+		TArray<FRigControl> SortedControls;
+		FControlRigSortedControls::GetControlsInOrder(ControlRig, SortedControls);
+		
+		for (const FRigControl& RigControl : SortedControls)
 		{
 			switch (RigControl.ControlType)
 			{
-			case ERigControlType::Float:
-			{
-				TOptional<float> DefaultValue;
-				if (bSetDefault)
+				case ERigControlType::Float:
 				{
-					//or use IntialValue?
-					DefaultValue = RigControl.Value.Get<float>();
+					TOptional<float> DefaultValue;
+					if (bSetDefault)
+					{
+						//or use IntialValue?
+						DefaultValue = RigControl.Value.Get<float>();
+					}
+					NewSection->AddScalarParameter(RigControl.Name,DefaultValue,false);
+					break;
 				}
-				NewSection->AddScalarParameter(RigControl.Name,DefaultValue);
-				break;
-			}
-			case ERigControlType::Position:
-			case ERigControlType::Scale:
-			case ERigControlType::Rotator:
-			{
-				TOptional<FVector> DefaultValue;
-				if (bSetDefault)
+				case ERigControlType::Bool:
 				{
-					//or use IntialValue?
-					DefaultValue = RigControl.Value.Get<FVector>();
+					TOptional<bool> DefaultValue;
+					if (bSetDefault)
+					{
+						//or use IntialValue?
+						DefaultValue = RigControl.Value.Get<bool>();
+					}
+					NewSection->AddBoolParameter(RigControl.Name, DefaultValue, false);
+					break;
 				}
-				NewSection->AddVectorParameter(RigControl.Name,DefaultValue);
-				//mz todo specify rotator special so we can do quat interps or no we have a quat type so...
-				break;
-			}
-			case ERigControlType::Transform:
-			{
-				TOptional<FTransform> DefaultValue;
-				if (bSetDefault)
+				case ERigControlType::Vector2D:
 				{
-					//or use IntialValue?
-					DefaultValue = RigControl.Value.Get<FTransform>();
+					TOptional<FVector2D> DefaultValue;
+					if (bSetDefault)
+					{
+						//or use IntialValue?
+						DefaultValue = RigControl.Value.Get<FVector2D>();
+					}
+					NewSection->AddVector2DParameter(RigControl.Name, DefaultValue, false);
+					break;
 				}
-				NewSection->AddTransformParameter(RigControl.Name,DefaultValue);
-				break;
-			}
-				//mz todo the other types
+				
+				case ERigControlType::Position:
+				case ERigControlType::Scale:
+				case ERigControlType::Rotator:
+				{
+					TOptional<FVector> DefaultValue;
+					if (bSetDefault)
+					{
+						//or use IntialValue?
+						DefaultValue = RigControl.Value.Get<FVector>();
+					}
+					NewSection->AddVectorParameter(RigControl.Name,DefaultValue,false);
+					//mz todo specify rotator special so we can do quat interps
+					break;
+				}
+				case ERigControlType::TransformNoScale:
+				case ERigControlType::Transform:
+				{
+					TOptional<FTransform> DefaultValue;
+					if (bSetDefault)
+					{
+						if (RigControl.ControlType == ERigControlType::Transform)
+						{
+							DefaultValue = RigControl.Value.Get<FTransform>();
+						}
+						else
+						{
+							FTransformNoScale NoScale = RigControl.Value.Get<FTransformNoScale>();
+							DefaultValue = NoScale;
+						}
+					}
+					NewSection->AddTransformParameter(RigControl.Name,DefaultValue,false);
+					break;
+				}
+				
+				default:
+					break;
 			}
 		}
-
-
-
+		NewSection->ReconstructChannelProxy(true);
+		
 	}
 	return  NewSection;
 }
@@ -383,4 +400,14 @@ UMovieSceneSection* UMovieSceneControlRigParameterTrack::GetSectionToKey() const
 	return SectionToKey;
 }
 
+void UMovieSceneControlRigParameterTrack::PostLoad()
+{
+	Super::PostLoad();
+	if (ControlRig)
+	{
+		ControlRig->Initialize();
+		ControlRig->CreateRigControlsForCurveContainer();
+	}
+}
+
 #undef LOCTEXT_NAMESPACE
diff --git a/Engine/Plugins/Experimental/ControlRig/Source/ControlRig/Private/Units/Animation/RigUnit_AnimEasing.cpp b/Engine/Plugins/Experimental/ControlRig/Source/ControlRig/Private/Units/Animation/RigUnit_AnimEasing.cpp
--- a/Engine/Plugins/Experimental/ControlRig/Source/ControlRig/Private/Units/Animation/RigUnit_AnimEasing.cpp
+++ b/Engine/Plugins/Experimental/ControlRig/Source/ControlRig/Private/Units/Animation/RigUnit_AnimEasing.cpp
@@ -3,6 +3,10 @@
 #include "Units/Animation/RigUnit_AnimEasing.h"
 #include "Units/RigUnitContext.h"
 
+FRigUnit_AnimEasingType_Execute()
+{
+}
+
 FRigUnit_AnimEasing_Execute()
 {
     DECLARE_SCOPE_HIERARCHICAL_COUNTER_RIGUNIT()
diff --git a/Engine/Plugins/Experimental/ControlRig/Source/ControlRig/Private/Units/Control/RigUnit_Control_StaticMesh.cpp b/Engine/Plugins/Experimental/ControlRig/Source/ControlRig/Private/Units/Control/RigUnit_Control_StaticMesh.cpp
--- a/Engine/Plugins/Experimental/ControlRig/Source/ControlRig/Private/Units/Control/RigUnit_Control_StaticMesh.cpp
+++ b/Engine/Plugins/Experimental/ControlRig/Source/ControlRig/Private/Units/Control/RigUnit_Control_StaticMesh.cpp
@@ -3,9 +3,10 @@
 #include "Units/Control/RigUnit_Control_StaticMesh.h"
 
 FRigUnit_Control_StaticMesh::FRigUnit_Control_StaticMesh()
-#if WITH_EDITORONLY_DATA
-	: StaticMesh(nullptr)
-	, MeshTransform(FTransform::Identity)
-#endif
 {
 }
+
+FRigUnit_Control_StaticMesh_Execute()
+{
+	FRigUnit_Control::StaticExecute(RigVMOperatorName, RigVMOperatorIndex, Transform, Base, InitTransform, Result, Filter, Context);
+}
\ No newline at end of file
diff --git a/Engine/Plugins/Experimental/ControlRig/Source/ControlRig/Private/Units/Debug/RigUnit_DebugBezier.h b/Engine/Plugins/Experimental/ControlRig/Source/ControlRig/Private/Units/Debug/RigUnit_DebugBezier.h
--- a/Engine/Plugins/Experimental/ControlRig/Source/ControlRig/Private/Units/Debug/RigUnit_DebugBezier.h
+++ b/Engine/Plugins/Experimental/ControlRig/Source/ControlRig/Private/Units/Debug/RigUnit_DebugBezier.h
@@ -44,12 +44,12 @@ struct FRigUnit_DebugBezier : public FRigUnit_DebugBaseMutable
 	UPROPERTY(meta = (Input))
 	int32 Detail;
 
-	UPROPERTY(meta = (Input, Constant, BoneName))
+	UPROPERTY(meta = (Input, Constant, CustomWidget = "BoneName"))
 	FName Space;
 
 	UPROPERTY(meta = (Input))
 	FTransform WorldOffset;
 
-	UPROPERTY(meta = (Input, Constant, BoneName))
+	UPROPERTY(meta = (Input, Constant))
 	bool bEnabled;
 };
diff --git a/Engine/Plugins/Experimental/ControlRig/Source/ControlRig/Private/Units/Debug/RigUnit_DebugHierarchy.h b/Engine/Plugins/Experimental/ControlRig/Source/ControlRig/Private/Units/Debug/RigUnit_DebugHierarchy.h
--- a/Engine/Plugins/Experimental/ControlRig/Source/ControlRig/Private/Units/Debug/RigUnit_DebugHierarchy.h
+++ b/Engine/Plugins/Experimental/ControlRig/Source/ControlRig/Private/Units/Debug/RigUnit_DebugHierarchy.h
@@ -39,6 +39,6 @@ struct FRigUnit_DebugHierarchy : public FRigUnit_DebugBaseMutable
 	UPROPERTY(meta = (Input))
 	FTransform WorldOffset;
 
-	UPROPERTY(meta = (Input, Constant, BoneName))
+	UPROPERTY(meta = (Input, Constant))
 	bool bEnabled;
 };
diff --git a/Engine/Plugins/Experimental/ControlRig/Source/ControlRig/Private/Units/Debug/RigUnit_DebugLine.h b/Engine/Plugins/Experimental/ControlRig/Source/ControlRig/Private/Units/Debug/RigUnit_DebugLine.h
--- a/Engine/Plugins/Experimental/ControlRig/Source/ControlRig/Private/Units/Debug/RigUnit_DebugLine.h
+++ b/Engine/Plugins/Experimental/ControlRig/Source/ControlRig/Private/Units/Debug/RigUnit_DebugLine.h
@@ -34,12 +34,12 @@ struct FRigUnit_DebugLine : public FRigUnit_DebugBaseMutable
 	UPROPERTY(meta = (Input))
 	float Thickness;
 
-	UPROPERTY(meta = (Input, Constant, BoneName))
+	UPROPERTY(meta = (Input, Constant, CustomWidget = "BoneName"))
 	FName Space;
 
 	UPROPERTY(meta = (Input))
 	FTransform WorldOffset;
 
-	UPROPERTY(meta = (Input, Constant, BoneName))
+	UPROPERTY(meta = (Input, Constant))
 	bool bEnabled;
 };
diff --git a/Engine/Plugins/Experimental/ControlRig/Source/ControlRig/Private/Units/Debug/RigUnit_DebugPoint.h b/Engine/Plugins/Experimental/ControlRig/Source/ControlRig/Private/Units/Debug/RigUnit_DebugPoint.h
--- a/Engine/Plugins/Experimental/ControlRig/Source/ControlRig/Private/Units/Debug/RigUnit_DebugPoint.h
+++ b/Engine/Plugins/Experimental/ControlRig/Source/ControlRig/Private/Units/Debug/RigUnit_DebugPoint.h
@@ -18,7 +18,7 @@ enum class ERigUnitDebugPointMode : uint8
 	Max UMETA(Hidden),
 };
 
-USTRUCT(meta=(DisplayName="Draw Point In Place", Keywords="Draw Vector"))
+USTRUCT(meta=(DisplayName="Draw Point In Place", Keywords="Draw Vector", Deprecated = "4.25.0"))
 struct FRigUnit_DebugPoint : public FRigUnit_DebugBase
 {
 	GENERATED_BODY()
@@ -52,17 +52,17 @@ struct FRigUnit_DebugPoint : public FRigUnit_DebugBase
 	UPROPERTY(meta = (Input))
 	float Thickness;
 
-	UPROPERTY(meta = (Input, Constant, BoneName))
+	UPROPERTY(meta = (Input, Constant, CustomWidget = "BoneName"))
 	FName Space;
 
 	UPROPERTY(meta = (Input))
 	FTransform WorldOffset;
 
-	UPROPERTY(meta = (Input, Constant, BoneName))
+	UPROPERTY(meta = (Input, Constant))
 	bool bEnabled;
 };
 
-USTRUCT(meta=(DisplayName="Draw Point", Keywords="Draw Vector"))
+USTRUCT(meta=(DisplayName="Draw Point", Keywords="Draw Vector", Deprecated = "4.25.0"))
 struct FRigUnit_DebugPointMutable : public FRigUnit_DebugBaseMutable
 {
 	GENERATED_BODY()
@@ -96,12 +96,12 @@ struct FRigUnit_DebugPointMutable : public FRigUnit_DebugBaseMutable
 	UPROPERTY(meta = (Input))
 	float Thickness;
 
-	UPROPERTY(meta = (Input, Constant, BoneName))
+	UPROPERTY(meta = (Input, Constant, CustomWidget = "BoneName"))
 	FName Space;
 
 	UPROPERTY(meta = (Input))
 	FTransform WorldOffset;
 
-	UPROPERTY(meta = (Input, Constant, BoneName))
+	UPROPERTY(meta = (Input, Constant))
 	bool bEnabled;
 };
\ No newline at end of file
diff --git a/Engine/Plugins/Experimental/ControlRig/Source/ControlRig/Private/Units/Debug/RigUnit_DebugPrimitives.h b/Engine/Plugins/Experimental/ControlRig/Source/ControlRig/Private/Units/Debug/RigUnit_DebugPrimitives.h
--- a/Engine/Plugins/Experimental/ControlRig/Source/ControlRig/Private/Units/Debug/RigUnit_DebugPrimitives.h
+++ b/Engine/Plugins/Experimental/ControlRig/Source/ControlRig/Private/Units/Debug/RigUnit_DebugPrimitives.h
@@ -35,13 +35,13 @@ struct FRigUnit_DebugRectangle : public FRigUnit_DebugBaseMutable
 	UPROPERTY(meta = (Input))
 	float Thickness;
 
-	UPROPERTY(meta = (Input, Constant, BoneName))
+	UPROPERTY(meta = (Input, Constant, CustomWidget = "BoneName"))
 	FName Space;
 
 	UPROPERTY(meta = (Input))
 	FTransform WorldOffset;
 
-	UPROPERTY(meta = (Input, Constant, BoneName))
+	UPROPERTY(meta = (Input, Constant))
 	bool bEnabled;
 };
 
@@ -86,12 +86,12 @@ struct FRigUnit_DebugArc : public FRigUnit_DebugBaseMutable
 	UPROPERTY(meta = (Input))
 	int32 Detail;
 
-	UPROPERTY(meta = (Input, Constant, BoneName))
+	UPROPERTY(meta = (Input, Constant, CustomWidget = "BoneName"))
 	FName Space;
 
 	UPROPERTY(meta = (Input))
 	FTransform WorldOffset;
 
-	UPROPERTY(meta = (Input, Constant, BoneName))
+	UPROPERTY(meta = (Input, Constant))
 	bool bEnabled;
 };
\ No newline at end of file
diff --git a/Engine/Plugins/Experimental/ControlRig/Source/ControlRig/Private/Units/Debug/RigUnit_DebugTransform.h b/Engine/Plugins/Experimental/ControlRig/Source/ControlRig/Private/Units/Debug/RigUnit_DebugTransform.h
--- a/Engine/Plugins/Experimental/ControlRig/Source/ControlRig/Private/Units/Debug/RigUnit_DebugTransform.h
+++ b/Engine/Plugins/Experimental/ControlRig/Source/ControlRig/Private/Units/Debug/RigUnit_DebugTransform.h
@@ -21,7 +21,7 @@ enum class ERigUnitDebugTransformMode : uint8
 	Max UMETA(Hidden),
 };
 
-USTRUCT(meta=(DisplayName="Draw Transform In Place"))
+USTRUCT(meta=(DisplayName="Draw Transform In Place", Deprecated = "4.25.0"))
 struct FRigUnit_DebugTransform : public FRigUnit_DebugBase
 {
 	GENERATED_BODY()
@@ -54,17 +54,17 @@ struct FRigUnit_DebugTransform : public FRigUnit_DebugBase
 	UPROPERTY(meta = (Input))
 	float Scale;
 
-	UPROPERTY(meta = (Input, Constant, BoneName))
+	UPROPERTY(meta = (Input, Constant, CustomWidget = "BoneName"))
 	FName Space;
 
 	UPROPERTY(meta = (Input))
 	FTransform WorldOffset;
 
-	UPROPERTY(meta = (Input, Constant, BoneName))
+	UPROPERTY(meta = (Input, Constant))
 	bool bEnabled;
 };
 
-USTRUCT(meta=(DisplayName="Draw Transform"))
+USTRUCT(meta=(DisplayName="Draw Transform", Deprecated = "4.25.0"))
 struct FRigUnit_DebugTransformMutable : public FRigUnit_DebugBaseMutable
 {
 	GENERATED_BODY()
@@ -97,13 +97,13 @@ struct FRigUnit_DebugTransformMutable : public FRigUnit_DebugBaseMutable
 	UPROPERTY(meta = (Input))
 	float Scale;
 
-	UPROPERTY(meta = (Input, Constant, BoneName))
+	UPROPERTY(meta = (Input, Constant, CustomWidget = "BoneName"))
 	FName Space;
 
 	UPROPERTY(meta = (Input))
 	FTransform WorldOffset;
 
-	UPROPERTY(meta = (Input, Constant, BoneName))
+	UPROPERTY(meta = (Input, Constant))
 	bool bEnabled;
 };
 
@@ -116,7 +116,7 @@ struct FRigUnit_DebugTransformArrayMutable_WorkData
 	TArray<FTransform> DrawTransforms;
 };
 
-USTRUCT(meta=(DisplayName="Draw Transform Array"))
+USTRUCT(meta=(DisplayName="Draw Transform Array", Deprecated = "4.25.0"))
 struct FRigUnit_DebugTransformArrayMutable : public FRigUnit_DebugBaseMutable
 {
  	GENERATED_BODY()
@@ -149,13 +149,13 @@ struct FRigUnit_DebugTransformArrayMutable : public FRigUnit_DebugBaseMutable
 	UPROPERTY(meta = (Input))
 	float Scale;
 
-	UPROPERTY(meta = (Input, Constant, BoneName))
+	UPROPERTY(meta = (Input, Constant, CustomWidget = "BoneName"))
 	FName Space;
 
 	UPROPERTY(meta = (Input))
 	FTransform WorldOffset;
 
-	UPROPERTY(meta = (Input, Constant, BoneName))
+	UPROPERTY(meta = (Input, Constant))
 	bool bEnabled;
 
 	UPROPERTY(transient)
diff --git a/Engine/Plugins/Experimental/ControlRig/Source/ControlRig/Private/Units/Debug/RigUnit_VisualDebug.cpp b/Engine/Plugins/Experimental/ControlRig/Source/ControlRig/Private/Units/Debug/RigUnit_VisualDebug.cpp
new file mode 100644
--- /dev/null
+++ b/Engine/Plugins/Experimental/ControlRig/Source/ControlRig/Private/Units/Debug/RigUnit_VisualDebug.cpp
@@ -0,0 +1,101 @@
+// Copyright Epic Games, Inc. All Rights Reserved.
+
+#include "Units/Debug/RigUnit_VisualDebug.h"
+#include "Units/RigUnitContext.h"
+
+FName FRigUnit_VisualDebugVector::DetermineSpaceForPin(const FString& InPinPath, void* InUserContext) const
+{
+	if (InPinPath == TEXT("Value"))
+	{
+		return BoneSpace;
+	}
+	return NAME_None;
+}
+
+FRigUnit_VisualDebugVector_Execute()
+{
+    DECLARE_SCOPE_HIERARCHICAL_COUNTER_RIGUNIT()
+	if (Context.State == EControlRigState::Init)
+	{
+		return;
+	}
+
+	if (Context.DrawInterface == nullptr || !bEnabled)
+	{
+		return;
+	}
+
+	FTransform WorldOffset = FTransform::Identity;
+	if (BoneSpace != NAME_None && Context.GetBones() != nullptr)
+	{
+		WorldOffset = Context.GetBones()->GetGlobalTransform(BoneSpace);
+	}
+
+	switch(Mode)
+	{
+		case ERigUnitVisualDebugPointMode::Point:
+		{
+			Context.DrawInterface->DrawPoint(WorldOffset, Value, Thickness, Color);
+			break;
+		}
+		case ERigUnitVisualDebugPointMode::Vector:
+		{
+			Context.DrawInterface->DrawLine(WorldOffset, FVector::ZeroVector, Value * Scale, Color, Thickness);
+			break;
+		}
+		default:
+		{
+			checkNoEntry();
+		}
+	}
+}
+
+FName FRigUnit_VisualDebugQuat::DetermineSpaceForPin(const FString& InPinPath, void* InUserContext) const
+{
+	if (InPinPath == TEXT("Value"))
+	{
+		return BoneSpace;
+	}
+	return NAME_None;
+}
+
+FRigUnit_VisualDebugQuat_Execute()
+{
+    DECLARE_SCOPE_HIERARCHICAL_COUNTER_RIGUNIT()
+
+    FTransform Transform = FTransform::Identity;
+    Transform.SetRotation(Value);
+
+	FRigUnit_VisualDebugTransform::StaticExecute(RigVMOperatorName, RigVMOperatorIndex, Transform, bEnabled, Thickness, Scale, BoneSpace, Context);
+}
+
+FName FRigUnit_VisualDebugTransform::DetermineSpaceForPin(const FString& InPinPath, void* InUserContext) const
+{
+	if (InPinPath == TEXT("Value"))
+	{
+		return BoneSpace;
+	}
+	return NAME_None;
+}
+
+FRigUnit_VisualDebugTransform_Execute()
+{
+    DECLARE_SCOPE_HIERARCHICAL_COUNTER_RIGUNIT()
+	if (Context.State == EControlRigState::Init)
+	{
+		return;
+	}
+
+	if (Context.DrawInterface == nullptr || !bEnabled)
+	{
+		return;
+	}
+
+	FTransform WorldOffset = FTransform::Identity;
+	if (BoneSpace != NAME_None && Context.GetBones() != nullptr)
+	{
+		WorldOffset = Context.GetBones()->GetGlobalTransform(BoneSpace);
+	}
+
+	Context.DrawInterface->DrawAxes(WorldOffset, Value, Scale, Thickness);
+}
diff --git a/Engine/Plugins/Experimental/ControlRig/Source/ControlRig/Private/Units/Debug/RigUnit_VisualDebug.h b/Engine/Plugins/Experimental/ControlRig/Source/ControlRig/Private/Units/Debug/RigUnit_VisualDebug.h
new file mode 100644
--- /dev/null
+++ b/Engine/Plugins/Experimental/ControlRig/Source/ControlRig/Private/Units/Debug/RigUnit_VisualDebug.h
@@ -0,0 +1,132 @@
+// Copyright Epic Games, Inc. All Rights Reserved.
+
+#pragma once
+
+#include "RigUnit_DebugBase.h"
+#include "RigUnit_VisualDebug.generated.h"
+
+UENUM()
+enum class ERigUnitVisualDebugPointMode : uint8
+{
+	/** Draw as point */
+	Point,
+
+	/** Draw as vector */
+	Vector,
+
+	/** MAX - invalid */
+	Max UMETA(Hidden),
+};
+
+USTRUCT(meta=(DisplayName = "Visual Debug Vector", PrototypeName = "VisualDebug", Keywords = "Draw,Point"))
+struct FRigUnit_VisualDebugVector : public FRigUnit_DebugBase
+{
+	GENERATED_BODY()
+
+	FRigUnit_VisualDebugVector()
+	{
+		Value = FVector::ZeroVector;
+		bEnabled = true;
+		Mode = ERigUnitVisualDebugPointMode::Point;
+		Color = FLinearColor::Red;
+		Thickness = 10.f;
+		Scale = 1.f;
+		BoneSpace = NAME_None;
+	}
+
+	virtual FName DetermineSpaceForPin(const FString& InPinPath, void* InUserContext) const override;
+
+	RIGVM_METHOD()
+	virtual void Execute(const FRigUnitContext& Context) override;
+
+	UPROPERTY(meta = (Input, Output))
+	FVector Value;
+
+	UPROPERTY(meta = (Input))
+	bool bEnabled;
+
+	UPROPERTY(meta = (Input, EditCondition = "bEnabled"))
+	ERigUnitVisualDebugPointMode Mode;
+
+	UPROPERTY(meta = (Input, EditCondition = "bEnabled"))
+	FLinearColor Color;
+
+	UPROPERTY(meta = (Input, EditCondition = "bEnabled"))
+	float Thickness;
+
+	UPROPERTY(meta = (Input, EditCondition = "bEnabled"))
+	float Scale;
+
+	UPROPERTY(meta = (Input, CustomWidget = "BoneName", EditCondition = "bEnabled"))
+	FName BoneSpace;
+};
+
+USTRUCT(meta = (DisplayName = "Visual Debug Quat", PrototypeName = "VisualDebug", Keywords = "Draw,Rotation"))
+struct FRigUnit_VisualDebugQuat : public FRigUnit_DebugBase
+{
+	GENERATED_BODY()
+
+	FRigUnit_VisualDebugQuat()
+	{
+		Value = FQuat::Identity;
+		bEnabled = true;
+		Thickness = 0.f;
+		Scale = 10.f;
+		BoneSpace = NAME_None;
+	}
+
+	virtual FName DetermineSpaceForPin(const FString& InPinPath, void* InUserContext) const override;
+
+	RIGVM_METHOD()
+	virtual void Execute(const FRigUnitContext& Context) override;
+
+	UPROPERTY(meta = (Input, Output))
+	FQuat Value;
+
+	UPROPERTY(meta = (Input))
+	bool bEnabled;
+
+	UPROPERTY(meta = (Input, EditCondition = "bEnabled"))
+	float Thickness;
+
+	UPROPERTY(meta = (Input, EditCondition = "bEnabled"))
+	float Scale;
+
+	UPROPERTY(meta = (Input, CustomWidget = "BoneName", EditCondition = "bEnabled"))
+	FName BoneSpace;
+};
+
+USTRUCT(meta=(DisplayName="Visual Debug Transform", PrototypeName = "VisualDebug", Keywords = "Draw,Axes"))
+struct FRigUnit_VisualDebugTransform : public FRigUnit_DebugBase
+{
+	GENERATED_BODY()
+
+	FRigUnit_VisualDebugTransform()
+	{
+		Value = FTransform::Identity;
+		bEnabled = true;
+		Thickness = 0.f;
+		Scale = 10.f;
+		BoneSpace = NAME_None;
+	}
+
+	virtual FName DetermineSpaceForPin(const FString& InPinPath, void* InUserContext) const override;
+
+	RIGVM_METHOD()
+	virtual void Execute(const FRigUnitContext& Context) override;
+
+	UPROPERTY(meta = (Input, Output))
+	FTransform Value;
+
+	UPROPERTY(meta = (Input))
+	bool bEnabled;
+
+	UPROPERTY(meta = (Input, EditCondition = "bEnabled"))
+	float Thickness;
+
+	UPROPERTY(meta = (Input, EditCondition = "bEnabled"))
+	float Scale;
+
+	UPROPERTY(meta = (Input, CustomWidget = "BoneName", EditCondition = "bEnabled"))
+	FName BoneSpace;
+};
diff --git a/Engine/Plugins/Experimental/ControlRig/Source/ControlRig/Private/Units/Deprecated/Math/RigUnit_Converter.h b/Engine/Plugins/Experimental/ControlRig/Source/ControlRig/Private/Units/Deprecated/Math/RigUnit_Converter.h
--- a/Engine/Plugins/Experimental/ControlRig/Source/ControlRig/Private/Units/Deprecated/Math/RigUnit_Converter.h
+++ b/Engine/Plugins/Experimental/ControlRig/Source/ControlRig/Private/Units/Deprecated/Math/RigUnit_Converter.h
@@ -17,10 +17,8 @@ struct FRigUnit_ConvertTransform : public FRigUnit
 	UPROPERTY(meta=(Output))
 	FEulerTransform Result;
 
-	virtual void Execute(const FRigUnitContext& Context) override
-	{
-		Result.FromFTransform(Input);
-	}
+	RIGVM_METHOD()
+	virtual void Execute(const FRigUnitContext& Context) override;
 };
 
 USTRUCT(meta = (DisplayName = "ConvertToTransform", Category = "Math|Convert", NodeColor = "0.1 0.1 0.7", Deprecated="4.23.0"))
@@ -34,10 +32,8 @@ struct FRigUnit_ConvertEulerTransform : public FRigUnit
 	UPROPERTY(meta=(Output))
 	FTransform Result;
 
-	virtual void Execute(const FRigUnitContext& Context) override
-	{
-		Result = Input.ToFTransform();
-	}
+	RIGVM_METHOD()
+	virtual void Execute(const FRigUnitContext& Context) override;
 };
 
 USTRUCT(meta = (DisplayName = "ConvertToQuaternion", Category = "Math|Convert", NodeColor = "0.1 0.1 0.7", Deprecated="4.23.0"))
@@ -51,10 +47,8 @@ struct FRigUnit_ConvertRotation : public FRigUnit
 	UPROPERTY(meta = (Output))
 	FQuat	Result;
 
-	virtual void Execute(const FRigUnitContext& Context) override
-	{
-		Result = Input.Quaternion();
-	}
+	RIGVM_METHOD()
+	virtual void Execute(const FRigUnitContext& Context) override;
 };
 
 USTRUCT(meta = (DisplayName = "ConvertToQuaternionDeprecated", Category = "Math|Convert", NodeColor = "0.1 0.1 0.7", Deprecated="4.23.0"))
@@ -74,10 +68,8 @@ struct FRigUnit_ConvertQuaternion: public FRigUnit
 	UPROPERTY(meta = (Output))
 	FRotator	Result;
 
-	virtual void Execute(const FRigUnitContext& Context) override
-	{
-		Result = Input.Rotator();
-	}
+	RIGVM_METHOD()
+	virtual void Execute(const FRigUnitContext& Context) override;
 };
 
 USTRUCT(meta = (DisplayName = "ConvertVectorToRotation", Category = "Math|Convert", NodeColor = "0.1 0.1 0.7", Deprecated="4.23.0"))
@@ -91,10 +83,8 @@ struct FRigUnit_ConvertVectorToRotation: public FRigUnit
 	UPROPERTY(meta = (Output))
 	FRotator	Result;
 
-	virtual void Execute(const FRigUnitContext& Context) override
-	{
-		Result = Input.Rotation();
-	}
+	RIGVM_METHOD()
+	virtual void Execute(const FRigUnitContext& Context) override;
 };
 
 USTRUCT(meta = (DisplayName = "ConvertVectorToQuaternion", Category = "Math|Convert", NodeColor = "0.1 0.1 0.7", Deprecated="4.23.0"))
@@ -108,11 +98,8 @@ struct FRigUnit_ConvertVectorToQuaternion: public FRigUnit
 	UPROPERTY(meta = (Output))
 	FQuat	Result;
 
-	virtual void Execute(const FRigUnitContext& Context) override
-	{
-		Result = Input.Rotation().Quaternion();
-		Result.Normalize();
-	}
+	RIGVM_METHOD()
+	virtual void Execute(const FRigUnitContext& Context) override;
 };
 
 
@@ -127,10 +114,8 @@ struct FRigUnit_ConvertRotationToVector: public FRigUnit
 	UPROPERTY(meta = (Output))
 	FVector Result;
 
-	virtual void Execute(const FRigUnitContext& Context) override
-	{
-		Result = Input.RotateVector(FVector(1.f, 0.f, 0.f));
-	}
+	RIGVM_METHOD()
+	virtual void Execute(const FRigUnitContext& Context) override;
 };
 
 USTRUCT(meta = (DisplayName = "ConvertQuaternionToVector", Category = "Math|Convert", NodeColor = "0.1 0.1 0.7", Deprecated="4.23.0"))
@@ -144,10 +129,8 @@ struct FRigUnit_ConvertQuaternionToVector: public FRigUnit
 	UPROPERTY(meta = (Output))
 	FVector Result;
 
-	virtual void Execute(const FRigUnitContext& Context) override
-	{
-		Result = Input.RotateVector(FVector(1.f, 0.f, 0.f));
-	}
+	RIGVM_METHOD()
+	virtual void Execute(const FRigUnitContext& Context) override;
 };
 
 USTRUCT(meta = (DisplayName = "ToSwingAndTwist", Category = "Math|Transform", NodeColor = "0.1 0.1 0.7", Deprecated="4.23.0"))
@@ -175,12 +158,6 @@ struct FRigUnit_ToSwingAndTwist : public FRigUnit
 	{
 	}
 
-	virtual void Execute(const FRigUnitContext& Context) override
-	{
-		if (!TwistAxis.IsZero())
-		{
-			FVector NormalizedAxis = TwistAxis.GetSafeNormal();
-			Input.ToSwingTwist(TwistAxis, Swing, Twist);
-		}
-	}
+	RIGVM_METHOD()
+	virtual void Execute(const FRigUnitContext& Context) override;
 };
diff --git a/Engine/Plugins/Experimental/ControlRig/Source/ControlRig/Private/Units/Deprecated/Math/RigUnit_Float.cpp b/Engine/Plugins/Experimental/ControlRig/Source/ControlRig/Private/Units/Deprecated/Math/RigUnit_Float.cpp
--- a/Engine/Plugins/Experimental/ControlRig/Source/ControlRig/Private/Units/Deprecated/Math/RigUnit_Float.cpp
+++ b/Engine/Plugins/Experimental/ControlRig/Source/ControlRig/Private/Units/Deprecated/Math/RigUnit_Float.cpp
@@ -1,3 +1,33 @@
 // Copyright Epic Games, Inc. All Rights Reserved.
 
-#include "RigUnit_Float.h"
\ No newline at end of file
+#include "RigUnit_Float.h"
+
+FRigUnit_Multiply_FloatFloat_Execute()
+{
+	Result = FRigMathLibrary::Multiply(Argument0, Argument1);
+}
+
+FRigUnit_Add_FloatFloat_Execute()
+{
+	Result = FRigMathLibrary::Add(Argument0, Argument1);
+}
+
+FRigUnit_Subtract_FloatFloat_Execute()
+{
+	Result = FRigMathLibrary::Subtract(Argument0, Argument1);
+}
+
+FRigUnit_Divide_FloatFloat_Execute()
+{
+	Result = FRigMathLibrary::Divide(Argument0, Argument1);
+}
+
+FRigUnit_Clamp_Float_Execute()
+{
+	Result = FMath::Clamp(Value, Min, Max);
+}
+
+FRigUnit_MapRange_Float_Execute()
+{
+	Result = FMath::GetMappedRangeValueClamped(FVector2D(MinIn, MaxIn), FVector2D(MinOut, MaxOut), Value);
+}
diff --git a/Engine/Plugins/Experimental/ControlRig/Source/ControlRig/Private/Units/Deprecated/Math/RigUnit_Float.h b/Engine/Plugins/Experimental/ControlRig/Source/ControlRig/Private/Units/Deprecated/Math/RigUnit_Float.h
--- a/Engine/Plugins/Experimental/ControlRig/Source/ControlRig/Private/Units/Deprecated/Math/RigUnit_Float.h
+++ b/Engine/Plugins/Experimental/ControlRig/Source/ControlRig/Private/Units/Deprecated/Math/RigUnit_Float.h
@@ -27,10 +27,8 @@ struct FRigUnit_Multiply_FloatFloat : public FRigUnit_BinaryFloatOp
 {
 	GENERATED_BODY()
 
-	virtual void Execute(const FRigUnitContext& Context) override
-	{
-		Result = FRigMathLibrary::Multiply(Argument0, Argument1);
-	}
+	RIGVM_METHOD()
+	virtual void Execute(const FRigUnitContext& Context) override;
 };
 
 USTRUCT(meta=(DisplayName="Add", Category="Math|Float", Keywords = "+,Sum", Deprecated="4.23.0"))
@@ -38,10 +36,8 @@ struct FRigUnit_Add_FloatFloat : public FRigUnit_BinaryFloatOp
 {
 	GENERATED_BODY()
 
-	virtual void Execute(const FRigUnitContext& Context) override
-	{
-		Result = FRigMathLibrary::Add(Argument0, Argument1);
-	}
+	RIGVM_METHOD()
+	virtual void Execute(const FRigUnitContext& Context) override;
 };
 
 USTRUCT(meta=(DisplayName="Subtract", Category="Math|Float", Keywords = "-", Deprecated="4.23.0"))
@@ -49,10 +45,8 @@ struct FRigUnit_Subtract_FloatFloat : public FRigUnit_BinaryFloatOp
 {
 	GENERATED_BODY()
 
-	virtual void Execute(const FRigUnitContext& Context) override
-	{
-		Result = FRigMathLibrary::Subtract(Argument0, Argument1);
-	}
+	RIGVM_METHOD()
+	virtual void Execute(const FRigUnitContext& Context) override;
 };
 
 USTRUCT(meta=(DisplayName="Divide", Category="Math|Float", Keywords = "/", Deprecated="4.23.0"))
@@ -60,10 +54,8 @@ struct FRigUnit_Divide_FloatFloat : public FRigUnit_BinaryFloatOp
 {
 	GENERATED_BODY()
 
-	virtual void Execute(const FRigUnitContext& Context) override
-	{
-		Result = FRigMathLibrary::Divide(Argument0, Argument1);
-	}
+	RIGVM_METHOD()
+	virtual void Execute(const FRigUnitContext& Context) override;
 };
 
 /** Two args and a result of float type */
@@ -84,10 +76,8 @@ struct FRigUnit_Clamp_Float: public FRigUnit
 	UPROPERTY(meta = (Output))
 	float Result;
 
-	virtual void Execute(const FRigUnitContext& Context) override
-	{
-		Result = FMath::Clamp(Value, Min, Max);
-	}
+	RIGVM_METHOD()
+	virtual void Execute(const FRigUnitContext& Context) override;
 };
 
 /** Two args and a result of float type */
@@ -114,8 +104,6 @@ struct FRigUnit_MapRange_Float: public FRigUnit
 	UPROPERTY(meta = (Output))
 	float Result;
 
-	virtual void Execute(const FRigUnitContext& Context) override
-	{
-		Result = FMath::GetMappedRangeValueClamped(FVector2D(MinIn, MaxIn), FVector2D(MinOut, MaxOut), Value);
-	}
+	RIGVM_METHOD()
+	virtual void Execute(const FRigUnitContext& Context) override;
 };
\ No newline at end of file
diff --git a/Engine/Plugins/Experimental/ControlRig/Source/ControlRig/Private/Units/Deprecated/Math/RigUnit_Quaternion.cpp b/Engine/Plugins/Experimental/ControlRig/Source/ControlRig/Private/Units/Deprecated/Math/RigUnit_Quaternion.cpp
--- a/Engine/Plugins/Experimental/ControlRig/Source/ControlRig/Private/Units/Deprecated/Math/RigUnit_Quaternion.cpp
+++ b/Engine/Plugins/Experimental/ControlRig/Source/ControlRig/Private/Units/Deprecated/Math/RigUnit_Quaternion.cpp
@@ -2,7 +2,32 @@
 
 #include "RigUnit_Quaternion.h"
 
-void FRigUnit_QuaternionToAngle::Execute(const FRigUnitContext& Context)
+FRigUnit_MultiplyQuaternion_Execute()
+{
+	Result = Argument0*Argument1;
+	Result.Normalize();
+}
+
+FRigUnit_InverseQuaterion_Execute()
+{
+	Result = Argument.Inverse();
+	Result.Normalize();
+}
+
+FRigUnit_QuaternionToAxisAndAngle_Execute()
+{
+	FVector NewAxis = Axis.GetSafeNormal();
+	Argument.ToAxisAndAngle(NewAxis, Angle);
+	Angle = FMath::RadiansToDegrees(Angle);
+}
+
+FRigUnit_QuaternionFromAxisAndAngle_Execute()
+{
+	FVector NewAxis = Axis.GetSafeNormal();
+	Result = FQuat(NewAxis, FMath::DegreesToRadians(Angle));
+}
+
+FRigUnit_QuaternionToAngle_Execute()
 {
     DECLARE_SCOPE_HIERARCHICAL_COUNTER_RIGUNIT()
 	FQuat Swing, Twist;
diff --git a/Engine/Plugins/Experimental/ControlRig/Source/ControlRig/Private/Units/Deprecated/Math/RigUnit_Quaternion.h b/Engine/Plugins/Experimental/ControlRig/Source/ControlRig/Private/Units/Deprecated/Math/RigUnit_Quaternion.h
--- a/Engine/Plugins/Experimental/ControlRig/Source/ControlRig/Private/Units/Deprecated/Math/RigUnit_Quaternion.h
+++ b/Engine/Plugins/Experimental/ControlRig/Source/ControlRig/Private/Units/Deprecated/Math/RigUnit_Quaternion.h
@@ -33,11 +33,8 @@ struct FRigUnit_MultiplyQuaternion : public FRigUnit_BinaryQuaternionOp
 {
 	GENERATED_BODY()
 
-	virtual void Execute(const FRigUnitContext& Context) override
- 	{
- 		Result = Argument0*Argument1;
- 		Result.Normalize();
- 	}
+	RIGVM_METHOD()
+	virtual void Execute(const FRigUnitContext& Context) override;
 };
 
 /** Two args and a result of Quaternion type */
@@ -58,11 +55,8 @@ struct FRigUnit_InverseQuaterion: public FRigUnit_UnaryQuaternionOp
 {
 	GENERATED_BODY()
 
-	virtual void Execute(const FRigUnitContext& Context) override
-	{
-		Result = Argument.Inverse();
-		Result.Normalize();
-	}
+	RIGVM_METHOD()
+	virtual void Execute(const FRigUnitContext& Context) override;
 };
 
 USTRUCT(meta = (DisplayName = "To Axis And Angle(Quaternion)", Category = "Math|Quaternion", NodeColor = "0.1 0.7 0.1", Deprecated="4.23.0"))
@@ -79,12 +73,8 @@ struct FRigUnit_QuaternionToAxisAndAngle : public FRigUnit
 	UPROPERTY(meta = (Output))
 	float Angle;
 
-	virtual void Execute(const FRigUnitContext& Context) override
-	{
-		FVector NewAxis = Axis.GetSafeNormal();
-		Argument.ToAxisAndAngle(NewAxis, Angle);
-		Angle = FMath::RadiansToDegrees(Angle);
-	}
+	RIGVM_METHOD()
+	virtual void Execute(const FRigUnitContext& Context) override;
 };
 
 USTRUCT(meta = (DisplayName = "From Axis And Angle(Quaternion)", Category = "Math|Quaternion", NodeColor = "0.1 0.7 0.1", Deprecated="4.23.0"))
@@ -107,11 +97,8 @@ struct FRigUnit_QuaternionFromAxisAndAngle : public FRigUnit
 	UPROPERTY(meta = (Output))
 	FQuat Result;
 
-	virtual void Execute(const FRigUnitContext& Context) override
-	{
-		FVector NewAxis = Axis.GetSafeNormal();
-		Result = FQuat(NewAxis, FMath::DegreesToRadians(Angle));
-	}
+	RIGVM_METHOD()
+	virtual void Execute(const FRigUnitContext& Context) override;
 };
 
 USTRUCT(meta = (DisplayName = "Get Angle Around Axis", Category = "Math|Quaternion", NodeColor = "0.1 0.7 0.1", Deprecated="4.23.0"))
@@ -134,6 +121,7 @@ struct FRigUnit_QuaternionToAngle : public FRigUnit
 	UPROPERTY(meta = (Output))
 	float Angle;
 
+	RIGVM_METHOD()
 	virtual void Execute(const FRigUnitContext& Context) override;
 };
 
diff --git a/Engine/Plugins/Experimental/ControlRig/Source/ControlRig/Private/Units/Deprecated/Math/RigUnit_Transform.h b/Engine/Plugins/Experimental/ControlRig/Source/ControlRig/Private/Units/Deprecated/Math/RigUnit_Transform.h
--- a/Engine/Plugins/Experimental/ControlRig/Source/ControlRig/Private/Units/Deprecated/Math/RigUnit_Transform.h
+++ b/Engine/Plugins/Experimental/ControlRig/Source/ControlRig/Private/Units/Deprecated/Math/RigUnit_Transform.h
@@ -27,10 +27,8 @@ struct FRigUnit_MultiplyTransform : public FRigUnit_BinaryTransformOp
 {
 	GENERATED_BODY()
 
-	virtual void Execute(const FRigUnitContext& Context) override
-	{
-		Result = Argument0*Argument1;
-	}
+	RIGVM_METHOD()
+	virtual void Execute(const FRigUnitContext& Context) override;
 };
 
 USTRUCT(meta = (DisplayName = "GetRelativeTransform", Category = "Math|Transform", Deprecated="4.23.0"))
@@ -38,9 +36,7 @@ struct FRigUnit_GetRelativeTransform : public FRigUnit_BinaryTransformOp
 {
 	GENERATED_BODY()
 
-	virtual void Execute(const FRigUnitContext& Context) override
-	{
-		Result = Argument0.GetRelativeTransform(Argument1);
-	}
+	RIGVM_METHOD()
+	virtual void Execute(const FRigUnitContext& Context) override;
 };
 
diff --git a/Engine/Plugins/Experimental/ControlRig/Source/ControlRig/Private/Units/Deprecated/Math/RigUnit_Vector.h b/Engine/Plugins/Experimental/ControlRig/Source/ControlRig/Private/Units/Deprecated/Math/RigUnit_Vector.h
--- a/Engine/Plugins/Experimental/ControlRig/Source/ControlRig/Private/Units/Deprecated/Math/RigUnit_Vector.h
+++ b/Engine/Plugins/Experimental/ControlRig/Source/ControlRig/Private/Units/Deprecated/Math/RigUnit_Vector.h
@@ -27,10 +27,8 @@ struct FRigUnit_Multiply_VectorVector : public FRigUnit_BinaryVectorOp
 {
 	GENERATED_BODY()
 
-	virtual void Execute(const FRigUnitContext& Context) override
-	{
-		Result = FRigMathLibrary::Multiply(Argument0, Argument1);
-	}
+	RIGVM_METHOD()
+	virtual void Execute(const FRigUnitContext& Context) override;
 };
 
 USTRUCT(meta=(DisplayName="Add(Vector)", Category="Math|Vector", Deprecated="4.23.0"))
@@ -38,10 +36,8 @@ struct FRigUnit_Add_VectorVector : public FRigUnit_BinaryVectorOp
 {
 	GENERATED_BODY()
 
-	virtual void Execute(const FRigUnitContext& Context) override
-	{
-		Result = FRigMathLibrary::Add(Argument0, Argument1);
-	}
+	RIGVM_METHOD()
+	virtual void Execute(const FRigUnitContext& Context) override;
 };
 
 USTRUCT(meta=(DisplayName="Subtract(Vector)", Category="Math|Vector", Deprecated="4.23.0"))
@@ -49,10 +45,8 @@ struct FRigUnit_Subtract_VectorVector : public FRigUnit_BinaryVectorOp
 {
 	GENERATED_BODY()
 
-	virtual void Execute(const FRigUnitContext& Context) override
-	{
-		Result = FRigMathLibrary::Subtract(Argument0, Argument1);
-	}
+	RIGVM_METHOD()
+	virtual void Execute(const FRigUnitContext& Context) override;
 };
 
 USTRUCT(meta=(DisplayName="Divide(Vector)", Category="Math|Vector", Deprecated="4.23.0"))
@@ -60,10 +54,8 @@ struct FRigUnit_Divide_VectorVector : public FRigUnit_BinaryVectorOp
 {
 	GENERATED_BODY()
 
-	virtual void Execute(const FRigUnitContext& Context) override
-	{
-		Result = FRigMathLibrary::Divide(Argument0, Argument1);
-	}
+	RIGVM_METHOD()
+	virtual void Execute(const FRigUnitContext& Context) override;
 };
 
 USTRUCT(meta = (DisplayName = "Distance", Category = "Math|Vector", Deprecated="4.23.0", NodeColor = "0.1 0.7 0.1"))
@@ -80,8 +72,6 @@ struct FRigUnit_Distance_VectorVector : public FRigUnit
 	UPROPERTY(meta=(Output))
 	float Result;
 
-	virtual void Execute(const FRigUnitContext& Context) override
-	{
-		Result = (Argument0 - Argument1).Size();
-	}
+	RIGVM_METHOD()
+	virtual void Execute(const FRigUnitContext& Context) override;
 };
\ No newline at end of file
diff --git a/Engine/Plugins/Experimental/ControlRig/Source/ControlRig/Private/Units/Deprecated/RigUnit_AimConstraint.cpp b/Engine/Plugins/Experimental/ControlRig/Source/ControlRig/Private/Units/Deprecated/RigUnit_AimConstraint.cpp
--- a/Engine/Plugins/Experimental/ControlRig/Source/ControlRig/Private/Units/Deprecated/RigUnit_AimConstraint.cpp
+++ b/Engine/Plugins/Experimental/ControlRig/Source/ControlRig/Private/Units/Deprecated/RigUnit_AimConstraint.cpp
@@ -37,10 +37,6 @@ void FRigUnit_AimConstraint::Execute(const FRigUnitContext& Context)
 				}
 			}
 		}
-		else
-		{
-			UnitLogHelpers::PrintMissingHierarchy(RigUnitName);
-		}
 	}
 	else if (Context.State == EControlRigState::Update)
 	{
diff --git a/Engine/Plugins/Experimental/ControlRig/Source/ControlRig/Private/Units/Deprecated/RigUnit_ApplyFK.cpp b/Engine/Plugins/Experimental/ControlRig/Source/ControlRig/Private/Units/Deprecated/RigUnit_ApplyFK.cpp
--- a/Engine/Plugins/Experimental/ControlRig/Source/ControlRig/Private/Units/Deprecated/RigUnit_ApplyFK.cpp
+++ b/Engine/Plugins/Experimental/ControlRig/Source/ControlRig/Private/Units/Deprecated/RigUnit_ApplyFK.cpp
@@ -4,17 +4,13 @@
 #include "Units/RigUnitContext.h"
 #include "HelperUtil.h"
 
-void FRigUnit_ApplyFK::Execute(const FRigUnitContext& Context)
+FRigUnit_ApplyFK_Execute()
 {
     DECLARE_SCOPE_HIERARCHICAL_COUNTER_RIGUNIT()
 	
 	if (Context.State == EControlRigState::Init)
 	{
-		FRigBoneHierarchy* Hierarchy = ExecuteContext.GetBones();
-		if (!Hierarchy)
-		{
-			UnitLogHelpers::PrintMissingHierarchy(RigUnitName);
-		}
+		return;
 	}
 	else if (Context.State == EControlRigState::Update)
 	{
@@ -28,12 +24,19 @@ void FRigUnit_ApplyFK::Execute(const FRigUnitContext& Context)
 				FTransform InputTransform = Transform;
 				Filter.FilterTransform(InputTransform);
 
+				FTransform InputBaseTransform = UtilityHelpers::GetBaseTransformByMode(
+					ApplyTransformSpace,
+					[Hierarchy](const FName& BoneName) { return Hierarchy->GetGlobalTransform(BoneName); },
+					(*Hierarchy)[Index].ParentName,
+					BaseJoint,
+					BaseTransform
+				);
+
 				// now get override or additive
 				// whether I'd like to apply whole thing or not
 				if (ApplyTransformMode == EApplyTransformMode::Override)
 				{
 					// get base transform
-					FTransform InputBaseTransform = GetBaseTransform(Index, Hierarchy);
 					FTransform ApplyTransform = InputTransform * InputBaseTransform;
 					Hierarchy->SetGlobalTransform(Index, ApplyTransform);
 				}
@@ -41,7 +44,6 @@ void FRigUnit_ApplyFK::Execute(const FRigUnitContext& Context)
 				{
 					// if additive, we get current transform and calculate base transform and apply in their local space
 					FTransform CurrentTransform = Hierarchy->GetGlobalTransform(Index);
-					FTransform InputBaseTransform = GetBaseTransform(Index, Hierarchy);
 					FTransform LocalTransform = InputTransform * CurrentTransform.GetRelativeTransform(InputBaseTransform);
 					// apply additive
 					Hierarchy->SetGlobalTransform(Index, LocalTransform * InputBaseTransform);
@@ -51,12 +53,6 @@ void FRigUnit_ApplyFK::Execute(const FRigUnitContext& Context)
 	}
 }
 
-FTransform FRigUnit_ApplyFK::GetBaseTransform(int32 BoneIndex, const FRigBoneHierarchy* CurrentHierarchy) const
-{
-	return UtilityHelpers::GetBaseTransformByMode(ApplyTransformSpace, [CurrentHierarchy](const FName& BoneName) { return CurrentHierarchy->GetGlobalTransform(BoneName); },
-		(*CurrentHierarchy)[BoneIndex].ParentName, BaseJoint, BaseTransform);
-}
-
 #if WITH_DEV_AUTOMATION_TESTS
 #include "Units/RigUnitTest.h"
 
diff --git a/Engine/Plugins/Experimental/ControlRig/Source/ControlRig/Private/Units/Deprecated/RigUnit_ApplyFK.h b/Engine/Plugins/Experimental/ControlRig/Source/ControlRig/Private/Units/Deprecated/RigUnit_ApplyFK.h
--- a/Engine/Plugins/Experimental/ControlRig/Source/ControlRig/Private/Units/Deprecated/RigUnit_ApplyFK.h
+++ b/Engine/Plugins/Experimental/ControlRig/Source/ControlRig/Private/Units/Deprecated/RigUnit_ApplyFK.h
@@ -30,6 +30,7 @@ struct FRigUnit_ApplyFK : public FRigUnitMutable
 {
 	GENERATED_BODY()
 
+	RIGVM_METHOD()
 	virtual void Execute(const FRigUnitContext& Context) override;
 
 	UPROPERTY(EditAnywhere, Category = "ApplyFK", meta = (Input))
@@ -55,7 +56,4 @@ struct FRigUnit_ApplyFK : public FRigUnitMutable
 	// Transform op option. Use if ETransformSpace is BaseJoint
 	UPROPERTY(EditAnywhere, Category = "ApplyFK", meta = (Input))
 	FName BaseJoint;
-
-private:
-	FTransform GetBaseTransform(int32 JointIndex, const FRigBoneHierarchy* CurrentHierarchy) const;
 };
diff --git a/Engine/Plugins/Experimental/ControlRig/Source/ControlRig/Private/Units/Deprecated/RigUnit_BlendTransform.cpp b/Engine/Plugins/Experimental/ControlRig/Source/ControlRig/Private/Units/Deprecated/RigUnit_BlendTransform.cpp
--- a/Engine/Plugins/Experimental/ControlRig/Source/ControlRig/Private/Units/Deprecated/RigUnit_BlendTransform.cpp
+++ b/Engine/Plugins/Experimental/ControlRig/Source/ControlRig/Private/Units/Deprecated/RigUnit_BlendTransform.cpp
@@ -4,7 +4,7 @@
 #include "Units/RigUnitContext.h"
 #include "AnimationRuntime.h"
 
-void FRigUnit_BlendTransform::Execute(const FRigUnitContext& Context)
+FRigUnit_BlendTransform_Execute()
 {
     DECLARE_SCOPE_HIERARCHICAL_COUNTER_RIGUNIT()
 	if (Targets.Num() > 0)
diff --git a/Engine/Plugins/Experimental/ControlRig/Source/ControlRig/Private/Units/Deprecated/RigUnit_GetJointTransform.cpp b/Engine/Plugins/Experimental/ControlRig/Source/ControlRig/Private/Units/Deprecated/RigUnit_GetJointTransform.cpp
--- a/Engine/Plugins/Experimental/ControlRig/Source/ControlRig/Private/Units/Deprecated/RigUnit_GetJointTransform.cpp
+++ b/Engine/Plugins/Experimental/ControlRig/Source/ControlRig/Private/Units/Deprecated/RigUnit_GetJointTransform.cpp
@@ -4,7 +4,7 @@
 #include "Units/RigUnitContext.h"
 #include "HelperUtil.h"
 
-void FRigUnit_GetJointTransform::Execute(const FRigUnitContext& Context)
+FRigUnit_GetJointTransform_Execute()
 {
     DECLARE_SCOPE_HIERARCHICAL_COUNTER_RIGUNIT()
 
@@ -19,38 +19,22 @@ void FRigUnit_GetJointTransform::Execute(const FRigUnitContext& Context)
 			{
 			case ETransformGetterType::Current:
 				{
-					Output = Hierarchy->GetGlobalTransform(Index).GetRelativeTransform(GetBaseTransform(Index, Hierarchy, false));
+					FTransform ComputedBaseTransform = UtilityHelpers::GetBaseTransformByMode(TransformSpace, [Hierarchy](const FName& JointName) { return Hierarchy->GetGlobalTransform(JointName); },
+					(*Hierarchy)[Index].ParentName, BaseJoint, BaseTransform);
+
+					Output = Hierarchy->GetGlobalTransform(Index).GetRelativeTransform(ComputedBaseTransform);
 					break;
 				}
 			case ETransformGetterType::Initial:
 			default:
 				{
-					Output = Hierarchy->GetInitialTransform(Index).GetRelativeTransform(GetBaseTransform(Index, Hierarchy, true));
+				FTransform ComputedBaseTransform = UtilityHelpers::GetBaseTransformByMode(TransformSpace, [Hierarchy](const FName& JointName) { return Hierarchy->GetInitialTransform(JointName); },
+					(*Hierarchy)[Index].ParentName, BaseJoint, BaseTransform);
+
+					Output = Hierarchy->GetInitialTransform(Index).GetRelativeTransform(ComputedBaseTransform);
 					break;
 				}
 			}
 		}
 	}
-	else
-	{
-		if (Context.State == EControlRigState::Init)
-		{
-			UnitLogHelpers::PrintMissingHierarchy(RigUnitName);
-		}
-	}
 }
-
-FTransform FRigUnit_GetJointTransform::GetBaseTransform(int32 JointIndex, const FRigBoneHierarchy* CurrentHierarchy, bool bUseInitial) const
-{
-	if (bUseInitial)
-	{
-		return UtilityHelpers::GetBaseTransformByMode(TransformSpace, [CurrentHierarchy](const FName& JointName) { return CurrentHierarchy->GetInitialTransform(JointName); },
-			(*CurrentHierarchy)[JointIndex].ParentName, BaseJoint, BaseTransform);
-	}
-	else
-	{
-		return UtilityHelpers::GetBaseTransformByMode(TransformSpace, [CurrentHierarchy](const FName& JointName) { return CurrentHierarchy->GetGlobalTransform(JointName); },
-			(*CurrentHierarchy)[JointIndex].ParentName, BaseJoint, BaseTransform);
-	}
-
-}
\ No newline at end of file
diff --git a/Engine/Plugins/Experimental/ControlRig/Source/ControlRig/Private/Units/Deprecated/RigUnit_GetJointTransform.h b/Engine/Plugins/Experimental/ControlRig/Source/ControlRig/Private/Units/Deprecated/RigUnit_GetJointTransform.h
--- a/Engine/Plugins/Experimental/ControlRig/Source/ControlRig/Private/Units/Deprecated/RigUnit_GetJointTransform.h
+++ b/Engine/Plugins/Experimental/ControlRig/Source/ControlRig/Private/Units/Deprecated/RigUnit_GetJointTransform.h
@@ -18,6 +18,7 @@ struct FRigUnit_GetJointTransform : public FRigUnitMutable
 		, TransformSpace(ETransformSpaceMode::GlobalSpace)
 	{}
 
+	RIGVM_METHOD()
 	virtual void Execute(const FRigUnitContext& Context) override;
 
 	UPROPERTY(meta = (Input))
@@ -40,7 +41,4 @@ struct FRigUnit_GetJointTransform : public FRigUnitMutable
 	// possibly space, relative transform so on can be input
 	UPROPERTY(meta=(Output))
 	FTransform Output;
-
-private:
-	FTransform GetBaseTransform(int32 JointIndex, const FRigBoneHierarchy* CurrentHierarchy, bool bUseInitial) const;
 };
diff --git a/Engine/Plugins/Experimental/ControlRig/Source/ControlRig/Private/Units/Deprecated/RigUnit_TwoBoneIKFK.cpp b/Engine/Plugins/Experimental/ControlRig/Source/ControlRig/Private/Units/Deprecated/RigUnit_TwoBoneIKFK.cpp
--- a/Engine/Plugins/Experimental/ControlRig/Source/ControlRig/Private/Units/Deprecated/RigUnit_TwoBoneIKFK.cpp
+++ b/Engine/Plugins/Experimental/ControlRig/Source/ControlRig/Private/Units/Deprecated/RigUnit_TwoBoneIKFK.cpp
@@ -5,7 +5,7 @@
 #include "HelperUtil.h"
 #include "TwoBoneIK.h"
 
-void FRigUnit_TwoBoneIKFK::Execute(const FRigUnitContext& Context)
+FRigUnit_TwoBoneIKFK_Execute()
 {
     DECLARE_SCOPE_HIERARCHICAL_COUNTER_RIGUNIT()
 
@@ -51,10 +51,6 @@ void FRigUnit_TwoBoneIKFK::Execute(const FRigUnitContext& Context)
 				}
 			}
 		}
-		else
-		{
-			UnitLogHelpers::PrintMissingHierarchy(RigUnitName);
-		}
 	}
 	else  if (Context.State == EControlRigState::Update)
 	{
@@ -82,7 +78,8 @@ void FRigUnit_TwoBoneIKFK::Execute(const FRigUnitContext& Context)
 				MidJointIKTransform = Hierarchy->GetGlobalTransform(MidJointIndex);
 				EndJointIKTransform = Hierarchy->GetGlobalTransform(EndJointIndex);
 
-				SolveIK();
+				AnimationCore::SolveTwoBoneIK(StartJointIKTransform, MidJointIKTransform, EndJointIKTransform, PoleTarget, EndEffector.GetLocation(), UpperLimbLength, LowerLimbLength, false, 1.0f, 1.05f);
+				EndJointIKTransform.SetRotation(EndEffector.GetRotation());
 
 				StartJointTransform = StartJointIKTransform;
 				MidJointTransform = MidJointIKTransform;
@@ -98,7 +95,8 @@ void FRigUnit_TwoBoneIKFK::Execute(const FRigUnitContext& Context)
 				MidJointIKTransform = Hierarchy->GetGlobalTransform(MidJointIndex);
 				EndJointIKTransform = Hierarchy->GetGlobalTransform(EndJointIndex);
 
-				SolveIK();
+				AnimationCore::SolveTwoBoneIK(StartJointIKTransform, MidJointIKTransform, EndJointIKTransform, PoleTarget, EndEffector.GetLocation(), UpperLimbLength, LowerLimbLength, false, 1.0f, 1.05f);
+				EndJointIKTransform.SetRotation(EndEffector.GetRotation());
 
 				StartJointTransform.Blend(StartJointFKTransform, StartJointIKTransform, IKBlend);
 				MidJointTransform.Blend(MidJointFKTransform, MidJointIKTransform, IKBlend);
@@ -116,10 +114,3 @@ void FRigUnit_TwoBoneIKFK::Execute(const FRigUnitContext& Context)
 	}
 }
 
-void FRigUnit_TwoBoneIKFK::SolveIK()
-{
-	AnimationCore::SolveTwoBoneIK(StartJointIKTransform, MidJointIKTransform, EndJointIKTransform, PoleTarget, EndEffector.GetLocation(), UpperLimbLength, LowerLimbLength, false, 1.0f, 1.05f);
-	// set end effector rotation to current rotation
-	EndJointIKTransform.SetRotation(EndEffector.GetRotation());
-}
-
diff --git a/Engine/Plugins/Experimental/ControlRig/Source/ControlRig/Private/Units/Drawing/RigUnit_DrawContainer.cpp b/Engine/Plugins/Experimental/ControlRig/Source/ControlRig/Private/Units/Drawing/RigUnit_DrawContainer.cpp
new file mode 100644
--- /dev/null
+++ b/Engine/Plugins/Experimental/ControlRig/Source/ControlRig/Private/Units/Drawing/RigUnit_DrawContainer.cpp
@@ -0,0 +1,106 @@
+// Copyright Epic Games, Inc. All Rights Reserved.
+
+#include "Units/Drawing/RigUnit_DrawContainer.h"
+#include "Units/RigUnitContext.h"
+
+FRigUnit_DrawContainerGetInstruction_Execute()
+{
+    DECLARE_SCOPE_HIERARCHICAL_COUNTER_RIGUNIT()
+
+    if(Context.DrawContainer == nullptr)
+    {
+    	return;
+    }
+	
+	switch(Context.State)
+	{
+		case EControlRigState::Update:
+		{
+			int32 Index = Context.DrawContainer->GetIndex(InstructionName);
+			if (Index != INDEX_NONE)
+			{
+				const FControlRigDrawInstruction& Instruction = (*Context.DrawContainer)[Index];
+				Color = Instruction.Color;
+				Transform = Instruction.Transform;
+			}
+			else
+			{
+				Color = FLinearColor::Red;
+				Transform = FTransform::Identity;
+			}
+			break;
+		}
+	}
+}
+
+FRigUnit_DrawContainerSetThickness_Execute()
+{
+	DECLARE_SCOPE_HIERARCHICAL_COUNTER_RIGUNIT()
+
+	if (Context.DrawContainer == nullptr)
+	{
+		return;
+	}
+
+	switch (Context.State)
+	{
+		case EControlRigState::Update:
+		{
+			int32 Index = Context.DrawContainer->GetIndex(InstructionName);
+			if (Index != INDEX_NONE)
+			{
+				FControlRigDrawInstruction& Instruction = (*Context.DrawContainer)[Index];
+				Instruction.Thickness = Thickness;
+			}
+			break;
+		}
+	}
+}
+
+FRigUnit_DrawContainerSetColor_Execute()
+{
+	DECLARE_SCOPE_HIERARCHICAL_COUNTER_RIGUNIT()
+
+	if (Context.DrawContainer == nullptr)
+	{
+		return;
+	}
+
+	switch (Context.State)
+	{
+		case EControlRigState::Update:
+		{
+			int32 Index = Context.DrawContainer->GetIndex(InstructionName);
+			if (Index != INDEX_NONE)
+			{
+				FControlRigDrawInstruction& Instruction = (*Context.DrawContainer)[Index];
+				Instruction.Color = Color;
+			}
+			break;
+		}
+	}
+}
+
+FRigUnit_DrawContainerSetTransform_Execute()
+{
+	DECLARE_SCOPE_HIERARCHICAL_COUNTER_RIGUNIT()
+
+	if (Context.DrawContainer == nullptr)
+	{
+		return;
+	}
+
+	switch (Context.State)
+	{
+		case EControlRigState::Update:
+		{
+			int32 Index = Context.DrawContainer->GetIndex(InstructionName);
+			if (Index != INDEX_NONE)
+			{
+				FControlRigDrawInstruction& Instruction = (*Context.DrawContainer)[Index];
+				Instruction.Transform = Transform;
+			}
+			break;
+		}
+	}
+}
\ No newline at end of file
diff --git a/Engine/Plugins/Experimental/ControlRig/Source/ControlRig/Private/Units/Drawing/RigUnit_DrawContainer.h b/Engine/Plugins/Experimental/ControlRig/Source/ControlRig/Private/Units/Drawing/RigUnit_DrawContainer.h
new file mode 100644
--- /dev/null
+++ b/Engine/Plugins/Experimental/ControlRig/Source/ControlRig/Private/Units/Drawing/RigUnit_DrawContainer.h
@@ -0,0 +1,94 @@
+// Copyright Epic Games, Inc. All Rights Reserved.
+
+#pragma once
+
+#include "Units/RigUnit.h"
+#include "RigUnit_DrawContainer.generated.h"
+
+USTRUCT(meta=(DisplayName="Get DrawInstruction", Category="Drawing", NodeColor = "0.25 0.25 0.05"))
+struct FRigUnit_DrawContainerGetInstruction : public FRigUnit
+{
+	GENERATED_BODY()
+
+	FRigUnit_DrawContainerGetInstruction()
+	{
+		InstructionName = NAME_None;
+		Transform = FTransform::Identity;
+		Color = FLinearColor::Red;
+	}
+
+	RIGVM_METHOD()
+	virtual void Execute(const FRigUnitContext& Context) override;
+
+	UPROPERTY(meta = (Input, Constant, CustomWidget = "DrawingName"))
+	FName InstructionName;
+
+	UPROPERTY(meta = (Output))
+	FLinearColor Color;
+
+	UPROPERTY(meta = (Output))
+	FTransform Transform;
+};
+
+USTRUCT(meta = (DisplayName = "Set Draw Color", Category = "Drawing", NodeColor = "0.25 0.25 0.05"))
+struct FRigUnit_DrawContainerSetColor: public FRigUnitMutable
+{
+	GENERATED_BODY()
+
+	FRigUnit_DrawContainerSetColor()
+	{
+		InstructionName = NAME_None;
+		Color = FLinearColor::Red;
+	}
+
+	RIGVM_METHOD()
+	virtual void Execute(const FRigUnitContext& Context) override;
+
+	UPROPERTY(meta = (Input, Constant, CustomWidget = "DrawingName"))
+	FName InstructionName;
+
+	UPROPERTY(meta = (Input))
+	FLinearColor Color;
+};
+
+USTRUCT(meta = (DisplayName = "Set Draw Thickness", Category = "Drawing", NodeColor = "0.25 0.25 0.05"))
+struct FRigUnit_DrawContainerSetThickness : public FRigUnitMutable
+{
+	GENERATED_BODY()
+
+	FRigUnit_DrawContainerSetThickness()
+	{
+		InstructionName = NAME_None;
+		Thickness = 0.f;
+	}
+
+	RIGVM_METHOD()
+	virtual void Execute(const FRigUnitContext& Context) override;
+
+	UPROPERTY(meta = (Input, Constant, CustomWidget = "DrawingName"))
+	FName InstructionName;
+
+	UPROPERTY(meta = (Input))
+	float Thickness;
+};
+
+USTRUCT(meta = (DisplayName = "Set Draw Transform", Category = "Drawing", NodeColor = "0.25 0.25 0.05"))
+struct FRigUnit_DrawContainerSetTransform : public FRigUnitMutable
+{
+	GENERATED_BODY()
+
+	FRigUnit_DrawContainerSetTransform()
+	{
+		InstructionName = NAME_None;
+		Transform = FTransform::Identity;
+	}
+
+	RIGVM_METHOD()
+	virtual void Execute(const FRigUnitContext& Context) override;
+
+	UPROPERTY(meta = (Input, Constant, CustomWidget = "DrawingName"))
+	FName InstructionName;
+
+	UPROPERTY(meta = (Input))
+	FTransform Transform;
+};
\ No newline at end of file
diff --git a/Engine/Plugins/Experimental/ControlRig/Source/ControlRig/Private/Units/Hierarchy/RigUnit_AddBoneTransform.cpp b/Engine/Plugins/Experimental/ControlRig/Source/ControlRig/Private/Units/Hierarchy/RigUnit_AddBoneTransform.cpp
--- a/Engine/Plugins/Experimental/ControlRig/Source/ControlRig/Private/Units/Hierarchy/RigUnit_AddBoneTransform.cpp
+++ b/Engine/Plugins/Experimental/ControlRig/Source/ControlRig/Private/Units/Hierarchy/RigUnit_AddBoneTransform.cpp
@@ -1,6 +1,6 @@
 // Copyright Epic Games, Inc. All Rights Reserved.
 
-#include "RigUnit_AddBoneTransform.h"
+#include "Units/Hierarchy/RigUnit_AddBoneTransform.h"
 #include "Units/RigUnitContext.h"
 
 FString FRigUnit_AddBoneTransform::GetUnitLabel() const
@@ -19,6 +19,10 @@ FRigUnit_AddBoneTransform_Execute()
 			case EControlRigState::Init:
 			{
 				CachedBoneIndex = Hierarchy->GetIndex(Bone);
+				if (CachedBoneIndex == INDEX_NONE)
+				{
+					UE_CONTROLRIG_RIGUNIT_REPORT_WARNING(TEXT("Bone is not set."));
+				}
 				break;
 			}
 			case EControlRigState::Update:
diff --git a/Engine/Plugins/Experimental/ControlRig/Source/ControlRig/Private/Units/Hierarchy/RigUnit_BoneName.cpp b/Engine/Plugins/Experimental/ControlRig/Source/ControlRig/Private/Units/Hierarchy/RigUnit_BoneName.cpp
--- a/Engine/Plugins/Experimental/ControlRig/Source/ControlRig/Private/Units/Hierarchy/RigUnit_BoneName.cpp
+++ b/Engine/Plugins/Experimental/ControlRig/Source/ControlRig/Private/Units/Hierarchy/RigUnit_BoneName.cpp
@@ -10,4 +10,22 @@ FString FRigUnit_BoneName::GetUnitLabel() const
 
 FRigUnit_BoneName_Execute()
 {
+}
+
+FString FRigUnit_SpaceName::GetUnitLabel() const
+{
+	return FString::Printf(TEXT("%s Name"), *Space.ToString());
+}
+
+FRigUnit_SpaceName_Execute()
+{
+}
+
+FString FRigUnit_ControlName::GetUnitLabel() const
+{
+	return FString::Printf(TEXT("%s Name"), *Control.ToString());
+}
+
+FRigUnit_ControlName_Execute()
+{
 }
\ No newline at end of file
diff --git a/Engine/Plugins/Experimental/ControlRig/Source/ControlRig/Private/Units/Hierarchy/RigUnit_BoneName.h b/Engine/Plugins/Experimental/ControlRig/Source/ControlRig/Private/Units/Hierarchy/RigUnit_BoneName.h
--- a/Engine/Plugins/Experimental/ControlRig/Source/ControlRig/Private/Units/Hierarchy/RigUnit_BoneName.h
+++ b/Engine/Plugins/Experimental/ControlRig/Source/ControlRig/Private/Units/Hierarchy/RigUnit_BoneName.h
@@ -24,6 +24,52 @@ struct FRigUnit_BoneName : public FRigUnit
 	/**
 	 * The name of the Bone
 	 */
-	UPROPERTY(meta = (Input, Output, BoneName, Constant))
+	UPROPERTY(meta = (Input, Output, CustomWidget = "BoneName", Constant))
 	FName Bone;
 };
+
+/**
+ * SpaceName is used to represent a Space name in the graph
+ */
+USTRUCT(meta=(DisplayName="Space Name", Category="Hierarchy", DocumentationPolicy = "Strict"))
+struct FRigUnit_SpaceName : public FRigUnit
+{
+	GENERATED_BODY()
+
+	FRigUnit_SpaceName()
+	{}
+
+	virtual FString GetUnitLabel() const override;
+
+	RIGVM_METHOD()
+	virtual void Execute(const FRigUnitContext& Context) override;
+
+	/**
+	 * The name of the Space
+	 */
+	UPROPERTY(meta = (Input, Output, CustomWidget = "SpaceName", Constant))
+	FName Space;
+};
+
+/**
+ * ControlName is used to represent a Control name in the graph
+ */
+USTRUCT(meta=(DisplayName="Control Name", Category="Hierarchy", DocumentationPolicy = "Strict"))
+struct FRigUnit_ControlName : public FRigUnit
+{
+	GENERATED_BODY()
+
+	FRigUnit_ControlName()
+	{}
+
+	virtual FString GetUnitLabel() const override;
+
+	RIGVM_METHOD()
+	virtual void Execute(const FRigUnitContext& Context) override;
+
+	/**
+	 * The name of the Control
+	 */
+	UPROPERTY(meta = (Input, Output, CustomWidget = "ControlName", Constant))
+	FName Control;
+};
diff --git a/Engine/Plugins/Experimental/ControlRig/Source/ControlRig/Private/Units/Hierarchy/RigUnit_GetBoneTransform.h b/Engine/Plugins/Experimental/ControlRig/Source/ControlRig/Private/Units/Hierarchy/RigUnit_GetBoneTransform.h
--- a/Engine/Plugins/Experimental/ControlRig/Source/ControlRig/Private/Units/Hierarchy/RigUnit_GetBoneTransform.h
+++ b/Engine/Plugins/Experimental/ControlRig/Source/ControlRig/Private/Units/Hierarchy/RigUnit_GetBoneTransform.h
@@ -14,18 +14,36 @@ struct FRigUnit_GetBoneTransform : public FRigUnit
 	GENERATED_BODY()
 
 	FRigUnit_GetBoneTransform()
-		: Space(EBoneGetterSetterMode::LocalSpace)
+		: Space(EBoneGetterSetterMode::GlobalSpace)
 		, CachedBoneIndex(INDEX_NONE)
 	{}
 
 	virtual FString GetUnitLabel() const override;
+
+	virtual FName DetermineSpaceForPin(const FString& InPinPath, void* InUserContext) const override
+	{
+		if (InPinPath.StartsWith(TEXT("Transform")) && Space == EBoneGetterSetterMode::LocalSpace)
+		{
+			if (const FRigHierarchyContainer* Container = (const FRigHierarchyContainer*)InUserContext)
+			{
+				int32 BoneIndex = Container->BoneHierarchy.GetIndex(Bone);
+				if (BoneIndex != INDEX_NONE)
+				{
+					return Container->BoneHierarchy[BoneIndex].ParentName;
+				}
+
+			}
+		}
+		return NAME_None;
+	}
+
 	RIGVM_METHOD()
 	virtual void Execute(const FRigUnitContext& Context) override;
 
 	/**
 	 * The name of the Bone to retrieve the transform for.
 	 */
-	UPROPERTY(meta = (Input, BoneName, Constant))
+	UPROPERTY(meta = (Input, CustomWidget = "BoneName", Constant))
 	FName Bone;
 
 	/**
diff --git a/Engine/Plugins/Experimental/ControlRig/Source/ControlRig/Private/Units/Hierarchy/RigUnit_GetControlInitialTransform.h b/Engine/Plugins/Experimental/ControlRig/Source/ControlRig/Private/Units/Hierarchy/RigUnit_GetControlInitialTransform.h
--- a/Engine/Plugins/Experimental/ControlRig/Source/ControlRig/Private/Units/Hierarchy/RigUnit_GetControlInitialTransform.h
+++ b/Engine/Plugins/Experimental/ControlRig/Source/ControlRig/Private/Units/Hierarchy/RigUnit_GetControlInitialTransform.h
@@ -26,7 +26,7 @@ struct FRigUnit_GetControlInitialTransform : public FRigUnit
 	/**
 	 * The name of the Control to retrieve the transform for.
 	 */
-	UPROPERTY(meta = (Input, ControlName, Constant))
+	UPROPERTY(meta = (Input, CustomWidget = "ControlName", Constant))
 	FName Control;
 
 	/**
diff --git a/Engine/Plugins/Experimental/ControlRig/Source/ControlRig/Private/Units/Hierarchy/RigUnit_GetControlTransform.cpp b/Engine/Plugins/Experimental/ControlRig/Source/ControlRig/Private/Units/Hierarchy/RigUnit_GetControlTransform.cpp
--- a/Engine/Plugins/Experimental/ControlRig/Source/ControlRig/Private/Units/Hierarchy/RigUnit_GetControlTransform.cpp
+++ b/Engine/Plugins/Experimental/ControlRig/Source/ControlRig/Private/Units/Hierarchy/RigUnit_GetControlTransform.cpp
@@ -3,9 +3,193 @@
 #include "RigUnit_GetControlTransform.h"
 #include "Units/RigUnitContext.h"
 
-FString FRigUnit_GetControlTransform::GetUnitLabel() const
+FRigUnit_GetControlBool_Execute()
 {
-	return FString::Printf(TEXT("Get Transform %s"), *Control.ToString());
+    DECLARE_SCOPE_HIERARCHICAL_COUNTER_RIGUNIT()
+	const FRigControlHierarchy* Hierarchy = Context.GetControls();
+	if (Hierarchy)
+	{
+		switch (Context.State)
+		{
+			case EControlRigState::Init:
+			{
+				CachedControlIndex = Hierarchy->GetIndex(Control);
+			}
+			case EControlRigState::Update:
+			{
+				if (CachedControlIndex != INDEX_NONE)
+				{
+					BoolValue = Hierarchy->GetValue(CachedControlIndex).Get<bool>();
+				}
+			}
+			default:
+			{
+				break;
+			}
+		}
+	}
+}
+
+FRigUnit_GetControlFloat_Execute()
+{
+    DECLARE_SCOPE_HIERARCHICAL_COUNTER_RIGUNIT()
+	const FRigControlHierarchy* Hierarchy = Context.GetControls();
+	if (Hierarchy)
+	{
+		switch (Context.State)
+		{
+			case EControlRigState::Init:
+			{
+				CachedControlIndex = Hierarchy->GetIndex(Control);
+			}
+			case EControlRigState::Update:
+			{
+				if (CachedControlIndex != INDEX_NONE)
+				{
+					FloatValue = Hierarchy->GetValue(CachedControlIndex).Get<float>();
+					Minimum = Hierarchy->GetValue(CachedControlIndex, ERigControlValueType::Minimum).Get<float>();
+					Maximum = Hierarchy->GetValue(CachedControlIndex, ERigControlValueType::Maximum).Get<float>();
+				}
+			}
+			default:
+			{
+				break;
+			}
+		}
+	}
+}
+
+FRigUnit_GetControlVector2D_Execute()
+{
+    DECLARE_SCOPE_HIERARCHICAL_COUNTER_RIGUNIT()
+	const FRigControlHierarchy* Hierarchy = Context.GetControls();
+	if (Hierarchy)
+	{
+		switch (Context.State)
+		{
+			case EControlRigState::Init:
+			{
+				CachedControlIndex = Hierarchy->GetIndex(Control);
+			}
+			case EControlRigState::Update:
+			{
+				if (CachedControlIndex != INDEX_NONE)
+				{
+					Vector = Hierarchy->GetValue(CachedControlIndex).Get<FVector2D>();
+					Minimum = Hierarchy->GetValue(CachedControlIndex, ERigControlValueType::Minimum).Get<FVector2D>();
+					Maximum = Hierarchy->GetValue(CachedControlIndex, ERigControlValueType::Maximum).Get<FVector2D>();
+				}
+			}
+			default:
+			{
+				break;
+			}
+		}
+	}
+}
+
+FRigUnit_GetControlVector_Execute()
+{
+    DECLARE_SCOPE_HIERARCHICAL_COUNTER_RIGUNIT()
+	const FRigControlHierarchy* Hierarchy = Context.GetControls();
+	if (Hierarchy)
+	{
+		switch (Context.State)
+		{
+			case EControlRigState::Init:
+			{
+				CachedControlIndex = Hierarchy->GetIndex(Control);
+			}
+			case EControlRigState::Update:
+			{
+				if (CachedControlIndex != INDEX_NONE)
+				{
+					FTransform Transform = FTransform::Identity;
+					switch (Space)
+					{
+						case EBoneGetterSetterMode::GlobalSpace:
+						{
+							Transform = Hierarchy->GetGlobalTransform(CachedControlIndex);
+							break;
+						}
+						case EBoneGetterSetterMode::LocalSpace:
+						{
+							Transform = Hierarchy->GetLocalTransform(CachedControlIndex);
+							break;
+						}
+						default:
+						{
+							break;
+						}
+					}
+
+					if((*Hierarchy)[CachedControlIndex].ControlType == ERigControlType::Position)
+					{
+						Vector = Transform.GetLocation();
+					}
+					else if((*Hierarchy)[CachedControlIndex].ControlType == ERigControlType::Scale)
+					{
+						Vector = Transform.GetScale3D();
+					}
+
+					Minimum = Hierarchy->GetValue(CachedControlIndex, ERigControlValueType::Minimum).Get<FVector>();
+					Maximum = Hierarchy->GetValue(CachedControlIndex, ERigControlValueType::Maximum).Get<FVector>();
+				}
+			}
+			default:
+			{
+				break;
+			}
+		}
+	}
+}
+
+FRigUnit_GetControlRotator_Execute()
+{
+    DECLARE_SCOPE_HIERARCHICAL_COUNTER_RIGUNIT()
+	const FRigControlHierarchy* Hierarchy = Context.GetControls();
+	if (Hierarchy)
+	{
+		switch (Context.State)
+		{
+			case EControlRigState::Init:
+			{
+				CachedControlIndex = Hierarchy->GetIndex(Control);
+			}
+			case EControlRigState::Update:
+			{
+				if (CachedControlIndex != INDEX_NONE)
+				{
+					FTransform Transform = FTransform::Identity;
+					switch (Space)
+					{
+						case EBoneGetterSetterMode::GlobalSpace:
+						{
+							Transform = Hierarchy->GetGlobalTransform(CachedControlIndex);
+							break;
+						}
+						case EBoneGetterSetterMode::LocalSpace:
+						{
+							Transform = Hierarchy->GetLocalTransform(CachedControlIndex);
+							break;
+						}
+						default:
+						{
+							break;
+						}
+					}
+
+					Rotator = Transform.GetRotation().Rotator();
+					Minimum = Hierarchy->GetValue(CachedControlIndex, ERigControlValueType::Minimum).Get<FRotator>();
+					Maximum = Hierarchy->GetValue(CachedControlIndex, ERigControlValueType::Maximum).Get<FRotator>();
+				}
+			}
+			default:
+			{
+				break;
+			}
+		}
+	}
 }
 
 FRigUnit_GetControlTransform_Execute()
@@ -19,7 +203,6 @@ FRigUnit_GetControlTransform_Execute()
 			case EControlRigState::Init:
 			{
 				CachedControlIndex = Hierarchy->GetIndex(Control);
-				break;				
 			}
 			case EControlRigState::Update:
 			{
@@ -42,6 +225,8 @@ FRigUnit_GetControlTransform_Execute()
 							break;
 						}
 					}
+					Minimum = Hierarchy->GetValue(CachedControlIndex, ERigControlValueType::Minimum).Get<FTransform>();
+					Maximum = Hierarchy->GetValue(CachedControlIndex, ERigControlValueType::Maximum).Get<FTransform>();
 				}
 			}
 			default:
@@ -75,7 +260,7 @@ IMPLEMENT_RIGUNIT_AUTOMATION_TEST(FRigUnit_GetControlTransform)
 	AddErrorIfFalse(Unit.Transform.GetTranslation().Equals(FVector(1.f, 0.f, 0.f)), TEXT("unexpected global transform"));
 	Unit.Space = EBoneGetterSetterMode::LocalSpace;
 	InitAndExecute();
-	AddErrorIfFalse(Unit.Transform.GetTranslation().Equals(FVector(1.f, 0.f, 0.f)), TEXT("unexpected local transform"));
+	AddErrorIfFalse(Unit.Transform.GetTranslation().Equals(FVector(0.f, 0.f, 0.f)), TEXT("unexpected local transform"));
 
 	Unit.Control = TEXT("ControlA");
 	Unit.Space = EBoneGetterSetterMode::GlobalSpace;
@@ -83,7 +268,7 @@ IMPLEMENT_RIGUNIT_AUTOMATION_TEST(FRigUnit_GetControlTransform)
 	AddErrorIfFalse(Unit.Transform.GetTranslation().Equals(FVector(2.f, 2.f, 3.f)), TEXT("unexpected global transform"));
 	Unit.Space = EBoneGetterSetterMode::LocalSpace;
 	InitAndExecute();
-	AddErrorIfFalse(Unit.Transform.GetTranslation().Equals(FVector(1.f, 2.f, 3.f)), TEXT("unexpected local transform"));
+	AddErrorIfFalse(Unit.Transform.GetTranslation().Equals(FVector(0.f, 0.f, 0.f)), TEXT("unexpected local transform"));
 
 	return true;
 }
