COPY from //Dev-Anim to //Dev-Main

diff --git a/Engine/Plugins/Experimental/ControlRig/Source/ControlRig/Private/Units/Hierarchy/RigUnit_GetControlTransform.h b/Engine/Plugins/Experimental/ControlRig/Source/ControlRig/Private/Units/Hierarchy/RigUnit_GetControlTransform.h
--- a/Engine/Plugins/Experimental/ControlRig/Source/ControlRig/Private/Units/Hierarchy/RigUnit_GetControlTransform.h
+++ b/Engine/Plugins/Experimental/ControlRig/Source/ControlRig/Private/Units/Hierarchy/RigUnit_GetControlTransform.h
@@ -5,27 +5,241 @@
 #include "Units/RigUnit.h"
 #include "RigUnit_GetControlTransform.generated.h"
 
+/**
+ * GetControlBool is used to retrieve a single Bool from a hierarchy.
+ */
+USTRUCT(meta=(DisplayName="Get Control Bool", Category="Controls", DocumentationPolicy = "Strict", Keywords="GetControlBool", PrototypeName = "GetControlValue"))
+struct FRigUnit_GetControlBool : public FRigUnit
+{
+	GENERATED_BODY()
+
+	FRigUnit_GetControlBool()
+		: BoolValue(false)
+		, CachedControlIndex(INDEX_NONE)
+	{}
+
+	RIGVM_METHOD()
+	virtual void Execute(const FRigUnitContext& Context) override;
+
+	/**
+	 * The name of the Control to retrieve the bool for.
+	 */
+	UPROPERTY(meta = (Input, CustomWidget = "ControlName", Constant))
+	FName Control;
+
+	// The current bool of the given bone - or identity in case it wasn't found.
+	UPROPERTY(meta=(Output))
+	bool BoolValue;
+
+	// Used to cache the internally used bone index
+	UPROPERTY()
+	int32 CachedControlIndex;
+};
+
+/**
+ * GetControlFloat is used to retrieve a single Float from a hierarchy.
+ */
+USTRUCT(meta=(DisplayName="Get Control Float", Category="Controls", DocumentationPolicy = "Strict", Keywords="GetControlFloat", PrototypeName = "GetControlValue"))
+struct FRigUnit_GetControlFloat : public FRigUnit
+{
+	GENERATED_BODY()
+
+	FRigUnit_GetControlFloat()
+		: FloatValue(0.f)
+		, Minimum(0.f)
+		, Maximum(0.f)
+		, CachedControlIndex(INDEX_NONE)
+	{}
+
+	RIGVM_METHOD()
+	virtual void Execute(const FRigUnitContext& Context) override;
+
+	/**
+	 * The name of the Control to retrieve the Float for.
+	 */
+	UPROPERTY(meta = (Input, CustomWidget = "ControlName", Constant))
+	FName Control;
+
+	// The current value of the control.
+	UPROPERTY(meta=(Output, UIMin = "0.0", UIMax = "1.0"))
+	float FloatValue;
+
+	// The minimum value of the control.
+	UPROPERTY(meta = (Output, UIMin = "0.0", UIMax = "1.0"))
+	float Minimum;
+
+	// The maximum value of the control.
+	UPROPERTY(meta = (Output, UIMin = "0.0", UIMax = "1.0"))
+	float Maximum;
+
+	// Used to cache the internally used bone index
+	UPROPERTY()
+	int32 CachedControlIndex;
+};
+
+/**
+ * GetControlVector2D is used to retrieve a single Vector2D from a hierarchy.
+ */
+USTRUCT(meta=(DisplayName="Get Control Vector2D", Category="Controls", DocumentationPolicy = "Strict", Keywords="GetControlVector2D", PrototypeName = "GetControlValue"))
+struct FRigUnit_GetControlVector2D : public FRigUnit
+{
+	GENERATED_BODY()
+
+	FRigUnit_GetControlVector2D()
+		: Vector(FVector2D::ZeroVector)
+		, Minimum(FVector2D::ZeroVector)
+		, Maximum(FVector2D::ZeroVector)
+		, CachedControlIndex(INDEX_NONE)
+	{}
+
+	RIGVM_METHOD()
+	virtual void Execute(const FRigUnitContext& Context) override;
+
+	/**
+	 * The name of the Control to retrieve the Vector2D for.
+	 */
+	UPROPERTY(meta = (Input, CustomWidget = "ControlName", Constant))
+	FName Control;
+
+	// The current value of the control.
+	UPROPERTY(meta=(Output))
+	FVector2D Vector;
+
+	// The minimum value of the control.
+	UPROPERTY(meta = (Output))
+	FVector2D Minimum;
+
+	// The maximum value of the control.
+	UPROPERTY(meta = (Output))
+	FVector2D Maximum;
+
+	// Used to cache the internally used bone index
+	UPROPERTY()
+	int32 CachedControlIndex;
+};
+
+/**
+ * GetControlVector is used to retrieve a single Vector from a hierarchy.
+ */
+USTRUCT(meta=(DisplayName="Get Control Vector", Category="Controls", DocumentationPolicy = "Strict", Keywords="GetControlVector", PrototypeName = "GetControlValue"))
+struct FRigUnit_GetControlVector : public FRigUnit
+{
+	GENERATED_BODY()
+
+	FRigUnit_GetControlVector()
+		: Space(EBoneGetterSetterMode::GlobalSpace)
+		, Vector(FVector::ZeroVector)
+		, Minimum(FVector::ZeroVector)
+		, Maximum(FVector::ZeroVector)
+		, CachedControlIndex(INDEX_NONE)
+	{}
+
+	RIGVM_METHOD()
+	virtual void Execute(const FRigUnitContext& Context) override;
+
+	/**
+	 * The name of the Control to retrieve the Vector for.
+	 */
+	UPROPERTY(meta = (Input, CustomWidget = "ControlName", Constant))
+	FName Control;
+
+	/**
+	 * Defines if the Control's Vector should be retrieved
+	 * in local or global space.
+	 */ 
+	UPROPERTY(meta = (Input))
+	EBoneGetterSetterMode Space;
+
+	// The current value of the control.
+	UPROPERTY(meta=(Output))
+	FVector Vector;
+
+	// The minimum value of the control.
+	UPROPERTY(meta = (Output))
+	FVector Minimum;
+
+	// The maximum value of the control.
+	UPROPERTY(meta = (Output))
+	FVector Maximum;
+
+	// Used to cache the internally used bone index
+	UPROPERTY()
+	int32 CachedControlIndex;
+};
+
+/**
+ * GetControlRotator is used to retrieve a single Rotator from a hierarchy.
+ */
+USTRUCT(meta=(DisplayName="Get Control Rotator", Category="Controls", DocumentationPolicy = "Strict", Keywords="GetControlRotator", PrototypeName = "GetControlValue"))
+struct FRigUnit_GetControlRotator : public FRigUnit
+{
+	GENERATED_BODY()
+
+	FRigUnit_GetControlRotator()
+		: Space(EBoneGetterSetterMode::GlobalSpace)
+		, Rotator(FRotator::ZeroRotator)
+		, Minimum(FRotator::ZeroRotator)
+		, Maximum(FRotator::ZeroRotator)
+		, CachedControlIndex(INDEX_NONE)
+	{
+	}
+
+	RIGVM_METHOD()
+	virtual void Execute(const FRigUnitContext& Context) override;
+
+	/**
+	 * The name of the Control to retrieve the Rotator for.
+	 */
+	UPROPERTY(meta = (Input, CustomWidget = "ControlName", Constant))
+	FName Control;
+
+	/**
+	 * Defines if the Control's Rotator should be retrieved
+	 * in local or global space.
+	 */ 
+	UPROPERTY(meta = (Input))
+	EBoneGetterSetterMode Space;
+
+	// The current value of the control.
+	UPROPERTY(meta=(Output))
+	FRotator Rotator;
+
+	// The minimum value of the control.
+	UPROPERTY(meta = (Output))
+	FRotator Minimum;
+
+	// The maximum value of the control.
+	UPROPERTY(meta = (Output))
+	FRotator Maximum;
+
+	// Used to cache the internally used bone index
+	UPROPERTY()
+	int32 CachedControlIndex;
+};
+
 /**
  * GetControlTransform is used to retrieve a single transform from a hierarchy.
  */
-USTRUCT(meta=(DisplayName="Get Control Transform", Category="Controls", DocumentationPolicy = "Strict", Keywords="GetControlTransform"))
+USTRUCT(meta=(DisplayName="Get Control Transform", Category="Controls", DocumentationPolicy = "Strict", Keywords="GetControlTransform", PrototypeName = "GetControlValue"))
 struct FRigUnit_GetControlTransform : public FRigUnit
 {
 	GENERATED_BODY()
 
 	FRigUnit_GetControlTransform()
-		: Space(EBoneGetterSetterMode::LocalSpace)
+		: Space(EBoneGetterSetterMode::GlobalSpace)
+		, Transform(FTransform::Identity)
+		, Minimum(FTransform::Identity)
+		, Maximum(FTransform::Identity)
 		, CachedControlIndex(INDEX_NONE)
 	{}
 
-	virtual FString GetUnitLabel() const override;
 	RIGVM_METHOD()
 	virtual void Execute(const FRigUnitContext& Context) override;
 
 	/**
 	 * The name of the Control to retrieve the transform for.
 	 */
-	UPROPERTY(meta = (Input, ControlName, Constant))
+	UPROPERTY(meta = (Input, CustomWidget = "ControlName", Constant))
 	FName Control;
 
 	/**
@@ -35,10 +249,18 @@ struct FRigUnit_GetControlTransform : public FRigUnit
 	UPROPERTY(meta = (Input))
 	EBoneGetterSetterMode Space;
 
-	// The current transform of the given bone - or identity in case it wasn't found.
+	// The current value of the control.
 	UPROPERTY(meta=(Output))
 	FTransform Transform;
 
+	// The minimum value of the control.
+	UPROPERTY(meta = (Output))
+	FTransform Minimum;
+
+	// The maximum value of the control.
+	UPROPERTY(meta = (Output))
+	FTransform Maximum;
+
 	// Used to cache the internally used bone index
 	UPROPERTY()
 	int32 CachedControlIndex;
diff --git a/Engine/Plugins/Experimental/ControlRig/Source/ControlRig/Private/Units/Hierarchy/RigUnit_GetCurveValue.h b/Engine/Plugins/Experimental/ControlRig/Source/ControlRig/Private/Units/Hierarchy/RigUnit_GetCurveValue.h
--- a/Engine/Plugins/Experimental/ControlRig/Source/ControlRig/Private/Units/Hierarchy/RigUnit_GetCurveValue.h
+++ b/Engine/Plugins/Experimental/ControlRig/Source/ControlRig/Private/Units/Hierarchy/RigUnit_GetCurveValue.h
@@ -26,7 +26,7 @@ struct FRigUnit_GetCurveValue : public FRigUnit
 	/**
 	 * The name of the Curve to retrieve the transform for.
 	 */
-	UPROPERTY(meta = (Input, CurveName, Constant))
+	UPROPERTY(meta = (Input, CustomWidget = "CurveName", Constant))
 	FName Curve;
 
 	// The current transform of the given Curve - or identity in case it wasn't found.
diff --git a/Engine/Plugins/Experimental/ControlRig/Source/ControlRig/Private/Units/Hierarchy/RigUnit_GetInitialBoneTransform.cpp b/Engine/Plugins/Experimental/ControlRig/Source/ControlRig/Private/Units/Hierarchy/RigUnit_GetInitialBoneTransform.cpp
--- a/Engine/Plugins/Experimental/ControlRig/Source/ControlRig/Private/Units/Hierarchy/RigUnit_GetInitialBoneTransform.cpp
+++ b/Engine/Plugins/Experimental/ControlRig/Source/ControlRig/Private/Units/Hierarchy/RigUnit_GetInitialBoneTransform.cpp
@@ -19,6 +19,10 @@ FRigUnit_GetInitialBoneTransform_Execute()
 			case EControlRigState::Init:
 			{
 				CachedBoneIndex = Hierarchy->GetIndex(Bone);
+				if (CachedBoneIndex == INDEX_NONE)
+				{
+					UE_CONTROLRIG_RIGUNIT_REPORT_WARNING(TEXT("Bone is not set."));
+				}
 			}
 			case EControlRigState::Update:
 			{
diff --git a/Engine/Plugins/Experimental/ControlRig/Source/ControlRig/Private/Units/Hierarchy/RigUnit_GetInitialBoneTransform.h b/Engine/Plugins/Experimental/ControlRig/Source/ControlRig/Private/Units/Hierarchy/RigUnit_GetInitialBoneTransform.h
--- a/Engine/Plugins/Experimental/ControlRig/Source/ControlRig/Private/Units/Hierarchy/RigUnit_GetInitialBoneTransform.h
+++ b/Engine/Plugins/Experimental/ControlRig/Source/ControlRig/Private/Units/Hierarchy/RigUnit_GetInitialBoneTransform.h
@@ -14,18 +14,36 @@ struct FRigUnit_GetInitialBoneTransform : public FRigUnit
 	GENERATED_BODY()
 
 	FRigUnit_GetInitialBoneTransform()
-		: Space(EBoneGetterSetterMode::LocalSpace)
+		: Space(EBoneGetterSetterMode::GlobalSpace)
 		, CachedBoneIndex(INDEX_NONE)
 	{}
 
 	virtual FString GetUnitLabel() const override;
+
+	virtual FName DetermineSpaceForPin(const FString& InPinPath, void* InUserContext) const override
+	{
+		if (InPinPath.StartsWith(TEXT("Transform")) && Space == EBoneGetterSetterMode::LocalSpace)
+		{
+			if (const FRigHierarchyContainer* Container = (const FRigHierarchyContainer*)InUserContext)
+			{
+				int32 BoneIndex = Container->BoneHierarchy.GetIndex(Bone);
+				if (BoneIndex != INDEX_NONE)
+				{
+					return Container->BoneHierarchy[BoneIndex].ParentName;
+				}
+
+			}
+		}
+		return NAME_None;
+	}
+
 	RIGVM_METHOD()
 	virtual void Execute(const FRigUnitContext& Context) override;
 
 	/**
 	 * The name of the Bone to retrieve the transform for.
 	 */
-	UPROPERTY(meta = (Input, BoneName, Constant))
+	UPROPERTY(meta = (Input, CustomWidget = "BoneName", Constant))
 	FName Bone;
 
 	/**
diff --git a/Engine/Plugins/Experimental/ControlRig/Source/ControlRig/Private/Units/Hierarchy/RigUnit_GetRelativeBoneTransform.cpp b/Engine/Plugins/Experimental/ControlRig/Source/ControlRig/Private/Units/Hierarchy/RigUnit_GetRelativeBoneTransform.cpp
--- a/Engine/Plugins/Experimental/ControlRig/Source/ControlRig/Private/Units/Hierarchy/RigUnit_GetRelativeBoneTransform.cpp
+++ b/Engine/Plugins/Experimental/ControlRig/Source/ControlRig/Private/Units/Hierarchy/RigUnit_GetRelativeBoneTransform.cpp
@@ -20,6 +20,14 @@ FRigUnit_GetRelativeBoneTransform_Execute()
 			{
 				CachedBoneIndex = Hierarchy->GetIndex(Bone);
 				CachedSpaceIndex = Hierarchy->GetIndex(Space);
+				if (CachedBoneIndex == INDEX_NONE)
+				{
+					UE_CONTROLRIG_RIGUNIT_REPORT_WARNING(TEXT("Bone is not set."));
+				}
+				if (CachedSpaceIndex == INDEX_NONE)
+				{
+					UE_CONTROLRIG_RIGUNIT_REPORT_WARNING(TEXT("Space is not set."));
+				}
 			}
 			case EControlRigState::Update:
 			{
diff --git a/Engine/Plugins/Experimental/ControlRig/Source/ControlRig/Private/Units/Hierarchy/RigUnit_GetRelativeBoneTransform.h b/Engine/Plugins/Experimental/ControlRig/Source/ControlRig/Private/Units/Hierarchy/RigUnit_GetRelativeBoneTransform.h
--- a/Engine/Plugins/Experimental/ControlRig/Source/ControlRig/Private/Units/Hierarchy/RigUnit_GetRelativeBoneTransform.h
+++ b/Engine/Plugins/Experimental/ControlRig/Source/ControlRig/Private/Units/Hierarchy/RigUnit_GetRelativeBoneTransform.h
@@ -19,19 +19,29 @@ struct FRigUnit_GetRelativeBoneTransform : public FRigUnit
 	{}
 
 	virtual FString GetUnitLabel() const override;
+
+	virtual FName DetermineSpaceForPin(const FString& InPinPath, void* InUserContext) const override
+	{
+		if (InPinPath.StartsWith(TEXT("Transform")))
+		{
+			return Space;
+		}
+		return NAME_None;
+	}
+
 	RIGVM_METHOD()
 	virtual void Execute(const FRigUnitContext& Context) override;
 
 	/**
 	 * The name of the Bone to retrieve the transform for.
 	 */
-	UPROPERTY(meta = (Input, BoneName, Constant))
+	UPROPERTY(meta = (Input, CustomWidget = "BoneName", Constant))
 	FName Bone;
 
 	/**
 	 * The name of the Bone to retrieve the transform relative within.
 	 */
-	UPROPERTY(meta = (Input, BoneName, Constant))
+	UPROPERTY(meta = (Input, CustomWidget = "BoneName", Constant))
 	FName Space;
 
 	// The current transform of the given bone - or identity in case it wasn't found.
diff --git a/Engine/Plugins/Experimental/ControlRig/Source/ControlRig/Private/Units/Hierarchy/RigUnit_GetSpaceTransform.h b/Engine/Plugins/Experimental/ControlRig/Source/ControlRig/Private/Units/Hierarchy/RigUnit_GetSpaceTransform.h
--- a/Engine/Plugins/Experimental/ControlRig/Source/ControlRig/Private/Units/Hierarchy/RigUnit_GetSpaceTransform.h
+++ b/Engine/Plugins/Experimental/ControlRig/Source/ControlRig/Private/Units/Hierarchy/RigUnit_GetSpaceTransform.h
@@ -14,7 +14,7 @@ struct FRigUnit_GetSpaceTransform : public FRigUnit
 	GENERATED_BODY()
 
 	FRigUnit_GetSpaceTransform()
-		: SpaceType(EBoneGetterSetterMode::LocalSpace)
+		: SpaceType(EBoneGetterSetterMode::GlobalSpace)
 		, CachedSpaceIndex(INDEX_NONE)
 	{}
 
@@ -25,7 +25,7 @@ struct FRigUnit_GetSpaceTransform : public FRigUnit
 	/**
 	 * The name of the Space to retrieve the transform for.
 	 */
-	UPROPERTY(meta = (Input, SpaceName, Constant))
+	UPROPERTY(meta = (Input, CustomWidget = "SpaceName", Constant))
 	FName Space;
 
 	/**
diff --git a/Engine/Plugins/Experimental/ControlRig/Source/ControlRig/Private/Units/Hierarchy/RigUnit_SetBoneRotation.cpp b/Engine/Plugins/Experimental/ControlRig/Source/ControlRig/Private/Units/Hierarchy/RigUnit_SetBoneRotation.cpp
--- a/Engine/Plugins/Experimental/ControlRig/Source/ControlRig/Private/Units/Hierarchy/RigUnit_SetBoneRotation.cpp
+++ b/Engine/Plugins/Experimental/ControlRig/Source/ControlRig/Private/Units/Hierarchy/RigUnit_SetBoneRotation.cpp
@@ -19,6 +19,10 @@ FRigUnit_SetBoneRotation_Execute()
 			case EControlRigState::Init:
 			{
 				CachedBoneIndex = Hierarchy->GetIndex(Bone);
+				if (CachedBoneIndex == INDEX_NONE)
+				{
+					UE_CONTROLRIG_RIGUNIT_REPORT_WARNING(TEXT("Bone is not set."));
+				}
 				// fall through to update
 			}
 			case EControlRigState::Update:
diff --git a/Engine/Plugins/Experimental/ControlRig/Source/ControlRig/Private/Units/Hierarchy/RigUnit_SetBoneRotation.h b/Engine/Plugins/Experimental/ControlRig/Source/ControlRig/Private/Units/Hierarchy/RigUnit_SetBoneRotation.h
--- a/Engine/Plugins/Experimental/ControlRig/Source/ControlRig/Private/Units/Hierarchy/RigUnit_SetBoneRotation.h
+++ b/Engine/Plugins/Experimental/ControlRig/Source/ControlRig/Private/Units/Hierarchy/RigUnit_SetBoneRotation.h
@@ -14,7 +14,7 @@ struct FRigUnit_SetBoneRotation : public FRigUnitMutable
 {
 	GENERATED_BODY()
 
-	FRigUnit_SetBoneRotation()
+		FRigUnit_SetBoneRotation()
 		: Rotation(FQuat::Identity)
 		, Space(EBoneGetterSetterMode::LocalSpace)
 		, Weight(1.f)
@@ -23,13 +23,31 @@ struct FRigUnit_SetBoneRotation : public FRigUnitMutable
 	{}
 
 	virtual FString GetUnitLabel() const override;
+
+	virtual FName DetermineSpaceForPin(const FString& InPinPath, void* InUserContext) const override
+	{
+		if (InPinPath.StartsWith(TEXT("Rotation")) && Space == EBoneGetterSetterMode::LocalSpace)
+		{
+			if (const FRigHierarchyContainer* Container = (const FRigHierarchyContainer*)InUserContext)
+			{
+				int32 BoneIndex = Container->BoneHierarchy.GetIndex(Bone);
+				if (BoneIndex != INDEX_NONE)
+				{
+					return Container->BoneHierarchy[BoneIndex].ParentName;
+				}
+
+			}
+		}
+		return NAME_None;
+	}
+
 	RIGVM_METHOD()
 	virtual void Execute(const FRigUnitContext& Context) override;
 
 	/**
 	 * The name of the Bone to set the rotation for.
 	 */
-	UPROPERTY(meta = (Input, BoneName, Constant))
+	UPROPERTY(meta = (Input, CustomWidget = "BoneName", Constant))
 	FName Bone;
 
 	/**
@@ -56,7 +74,7 @@ struct FRigUnit_SetBoneRotation : public FRigUnitMutable
 	 * of this bone will be recalculated based on their local transforms.
 	 * Note: This is computationally more expensive than turning it off.
 	 */
-	UPROPERTY(meta = (Input))
+	UPROPERTY(meta = (Input, Constant))
 	bool bPropagateToChildren;
 
 	// Used to cache the internally used bone index
diff --git a/Engine/Plugins/Experimental/ControlRig/Source/ControlRig/Private/Units/Hierarchy/RigUnit_SetBoneTransform.h b/Engine/Plugins/Experimental/ControlRig/Source/ControlRig/Private/Units/Hierarchy/RigUnit_SetBoneTransform.h
--- a/Engine/Plugins/Experimental/ControlRig/Source/ControlRig/Private/Units/Hierarchy/RigUnit_SetBoneTransform.h
+++ b/Engine/Plugins/Experimental/ControlRig/Source/ControlRig/Private/Units/Hierarchy/RigUnit_SetBoneTransform.h
@@ -23,13 +23,30 @@ struct FRigUnit_SetBoneTransform : public FRigUnitMutable
 
 	virtual FString GetUnitLabel() const override;
 
+	virtual FName DetermineSpaceForPin(const FString& InPinPath, void* InUserContext) const override
+	{
+		if (InPinPath.StartsWith(TEXT("Transform")) && Space == EBoneGetterSetterMode::LocalSpace)
+		{
+			if (const FRigHierarchyContainer* Container = (const FRigHierarchyContainer*)InUserContext)
+			{
+				int32 BoneIndex = Container->BoneHierarchy.GetIndex(Bone);
+				if (BoneIndex != INDEX_NONE)
+				{
+					return Container->BoneHierarchy[BoneIndex].ParentName;
+				}
+
+			}
+		}
+		return NAME_None;
+	}
+
 	RIGVM_METHOD()
 	virtual void Execute(const FRigUnitContext& Context) override;
 
 	/**
 	 * The name of the Bone to set the transform for.
 	 */
-	UPROPERTY(meta = (Input, BoneName, Constant))
+	UPROPERTY(meta = (Input, CustomWidget = "BoneName", Constant))
 	FName Bone;
 
 	/**
@@ -56,7 +73,7 @@ struct FRigUnit_SetBoneTransform : public FRigUnitMutable
 	 * of this bone will be recalculated based on their local transforms.
 	 * Note: This is computationally more expensive than turning it off.
 	 */
-	UPROPERTY(meta = (Input))
+	UPROPERTY(meta = (Input, Constant))
 	bool bPropagateToChildren;
 
 	// Used to cache the internally used bone index
diff --git a/Engine/Plugins/Experimental/ControlRig/Source/ControlRig/Private/Units/Hierarchy/RigUnit_SetBoneTranslation.h b/Engine/Plugins/Experimental/ControlRig/Source/ControlRig/Private/Units/Hierarchy/RigUnit_SetBoneTranslation.h
--- a/Engine/Plugins/Experimental/ControlRig/Source/ControlRig/Private/Units/Hierarchy/RigUnit_SetBoneTranslation.h
+++ b/Engine/Plugins/Experimental/ControlRig/Source/ControlRig/Private/Units/Hierarchy/RigUnit_SetBoneTranslation.h
@@ -23,13 +23,31 @@ struct FRigUnit_SetBoneTranslation : public FRigUnitMutable
 	{}
 
 	virtual FString GetUnitLabel() const override;
+
+	virtual FName DetermineSpaceForPin(const FString& InPinPath, void* InUserContext) const override
+	{
+		if (InPinPath.StartsWith(TEXT("Translation")) && Space == EBoneGetterSetterMode::LocalSpace)
+		{
+			if (const FRigHierarchyContainer* Container = (const FRigHierarchyContainer*)InUserContext)
+			{
+				int32 BoneIndex = Container->BoneHierarchy.GetIndex(Bone);
+				if (BoneIndex != INDEX_NONE)
+				{
+					return Container->BoneHierarchy[BoneIndex].ParentName;
+				}
+
+			}
+		}
+		return NAME_None;
+	}
+
 	RIGVM_METHOD()
 	virtual void Execute(const FRigUnitContext& Context) override;
 
 	/**
 	 * The name of the Bone to set the Translation for.
 	 */
-	UPROPERTY(meta = (Input, BoneName, Constant))
+	UPROPERTY(meta = (Input, CustomWidget = "BoneName", Constant))
 	FName Bone;
 
 	/**
@@ -56,7 +74,7 @@ struct FRigUnit_SetBoneTranslation : public FRigUnitMutable
 	 * of this bone will be recalculated based on their local transforms.
 	 * Note: This is computationally more expensive than turning it off.
 	 */
-	UPROPERTY(meta = (Input))
+	UPROPERTY(meta = (Input, Constant))
 	bool bPropagateToChildren;
 
 	// Used to cache the internally used bone index
diff --git a/Engine/Plugins/Experimental/ControlRig/Source/ControlRig/Private/Units/Hierarchy/RigUnit_SetControlColor.cpp b/Engine/Plugins/Experimental/ControlRig/Source/ControlRig/Private/Units/Hierarchy/RigUnit_SetControlColor.cpp
new file mode 100644
--- /dev/null
+++ b/Engine/Plugins/Experimental/ControlRig/Source/ControlRig/Private/Units/Hierarchy/RigUnit_SetControlColor.cpp
@@ -0,0 +1,34 @@
+// Copyright Epic Games, Inc. All Rights Reserved.
+
+#include "RigUnit_SetControlColor.h"
+#include "Units/RigUnitContext.h"
+#include "Math/ControlRigMathLibrary.h"
+
+FRigUnit_SetControlColor_Execute()
+{
+    DECLARE_SCOPE_HIERARCHICAL_COUNTER_RIGUNIT()
+	FRigControlHierarchy* Hierarchy = ExecuteContext.GetControls();
+	if (Hierarchy)
+	{
+		switch (Context.State)
+		{
+			case EControlRigState::Init:
+			{
+				CachedControlIndex = Hierarchy->GetIndex(Control);
+				break;
+			}
+			case EControlRigState::Update:
+			{
+				if (CachedControlIndex != INDEX_NONE)
+				{
+					(*Hierarchy)[CachedControlIndex].GizmoColor = Color;
+				}
+				break;
+			}
+			default:
+			{
+				break;
+			}
+		}
+	}
+}
diff --git a/Engine/Plugins/Experimental/ControlRig/Source/ControlRig/Private/Units/Hierarchy/RigUnit_SetControlColor.h b/Engine/Plugins/Experimental/ControlRig/Source/ControlRig/Private/Units/Hierarchy/RigUnit_SetControlColor.h
new file mode 100644
--- /dev/null
+++ b/Engine/Plugins/Experimental/ControlRig/Source/ControlRig/Private/Units/Hierarchy/RigUnit_SetControlColor.h
@@ -0,0 +1,40 @@
+// Copyright Epic Games, Inc. All Rights Reserved.
+
+#pragma once
+
+#include "Units/RigUnit.h"
+#include "RigUnit_SetControlColor.generated.h"
+
+/**
+ * SetControlTransform is used to change the gizmo color on a control at runtime
+ */
+USTRUCT(meta=(DisplayName="Set Control Color", Category="Hierarchy", DocumentationPolicy="Strict", Keywords = "SetControlColor", PrototypeName = "SetControlColor"))
+struct FRigUnit_SetControlColor : public FRigUnitMutable
+{
+	GENERATED_BODY()
+
+	FRigUnit_SetControlColor()
+		: Control(NAME_None)
+		, Color(FLinearColor::Black)
+		, CachedControlIndex(INDEX_NONE)
+	{}
+
+	RIGVM_METHOD()
+	virtual void Execute(const FRigUnitContext& Context) override;
+
+	/**
+	 * The name of the Control to set the transform for.
+	 */
+	UPROPERTY(meta = (Input, CustomWidget = "ControlName", Constant))
+	FName Control;
+
+	/**
+	 * The color to set for the control
+	 */
+	UPROPERTY(meta = (Input))
+	FLinearColor Color;
+
+	// Used to cache the internally used bone index
+	UPROPERTY()
+	int32 CachedControlIndex;
+};
diff --git a/Engine/Plugins/Experimental/ControlRig/Source/ControlRig/Private/Units/Hierarchy/RigUnit_SetControlTransform.cpp b/Engine/Plugins/Experimental/ControlRig/Source/ControlRig/Private/Units/Hierarchy/RigUnit_SetControlTransform.cpp
--- a/Engine/Plugins/Experimental/ControlRig/Source/ControlRig/Private/Units/Hierarchy/RigUnit_SetControlTransform.cpp
+++ b/Engine/Plugins/Experimental/ControlRig/Source/ControlRig/Private/Units/Hierarchy/RigUnit_SetControlTransform.cpp
@@ -2,13 +2,187 @@
 
 #include "RigUnit_SetControlTransform.h"
 #include "Units/RigUnitContext.h"
+#include "Math/ControlRigMathLibrary.h"
 
-FString FRigUnit_SetControlTransform::GetUnitLabel() const
+FRigUnit_SetControlBool_Execute()
 {
-	return FString::Printf(TEXT("Set Control %s"), *Control.ToString());
+    DECLARE_SCOPE_HIERARCHICAL_COUNTER_RIGUNIT()
+	FRigControlHierarchy* Hierarchy = ExecuteContext.GetControls();
+	if (Hierarchy)
+	{
+		switch (Context.State)
+		{
+			case EControlRigState::Init:
+			{
+				CachedControlIndex = Hierarchy->GetIndex(Control);
+				break;
+			}
+			case EControlRigState::Update:
+			{
+				if (CachedControlIndex != INDEX_NONE)
+				{
+					Hierarchy->SetValue(CachedControlIndex, FRigControlValue::Make<bool>(BoolValue));
+				}
+				break;
+			}
+			default:
+			{
+				break;
+			}
+		}
+	}
 }
 
-FRigUnit_SetControlTransform_Execute()
+FRigUnit_SetControlFloat_Execute()
+{
+    DECLARE_SCOPE_HIERARCHICAL_COUNTER_RIGUNIT()
+	FRigControlHierarchy* Hierarchy = ExecuteContext.GetControls();
+	if (Hierarchy)
+	{
+		switch (Context.State)
+		{
+			case EControlRigState::Init:
+			{
+				CachedControlIndex = Hierarchy->GetIndex(Control);
+				break;
+			}
+			case EControlRigState::Update:
+			{
+				if (CachedControlIndex != INDEX_NONE)
+				{
+					if(FMath::IsNearlyEqual(Weight, 1.f))
+					{
+						Hierarchy->SetValue(CachedControlIndex, FRigControlValue::Make<float>(FloatValue));
+					}
+					else
+					{
+						float PreviousValue = Hierarchy->GetValue(CachedControlIndex).Get<float>();
+						Hierarchy->SetValue(CachedControlIndex, FRigControlValue::Make<float>(FMath::Lerp<float>(PreviousValue, FloatValue, FMath::Clamp<float>(Weight, 0.f, 1.f))));
+					}
+				}
+				break;
+			}
+			default:
+			{
+				break;
+			}
+		}
+	}
+}
+
+FRigUnit_SetControlVector2D_Execute()
+{
+    DECLARE_SCOPE_HIERARCHICAL_COUNTER_RIGUNIT()
+	FRigControlHierarchy* Hierarchy = ExecuteContext.GetControls();
+	if (Hierarchy)
+	{
+		switch (Context.State)
+		{
+			case EControlRigState::Init:
+			{
+				CachedControlIndex = Hierarchy->GetIndex(Control);
+				break;
+			}
+			case EControlRigState::Update:
+			{
+				if (CachedControlIndex != INDEX_NONE)
+				{
+					if(FMath::IsNearlyEqual(Weight, 1.f))
+					{
+						Hierarchy->SetValue(CachedControlIndex, FRigControlValue::Make<FVector2D>(Vector));
+					}
+					else
+					{
+						FVector2D PreviousValue = Hierarchy->GetValue(CachedControlIndex).Get<FVector2D>();
+						Hierarchy->SetValue(CachedControlIndex, FRigControlValue::Make<FVector2D>(FMath::Lerp<FVector2D>(PreviousValue, Vector, FMath::Clamp<float>(Weight, 0.f, 1.f))));
+					}
+				}
+				break;
+			}
+			default:
+			{
+				break;
+			}
+		}
+	}
+}
+
+FRigUnit_SetControlVector_Execute()
+{
+    DECLARE_SCOPE_HIERARCHICAL_COUNTER_RIGUNIT()
+	FRigControlHierarchy* Hierarchy = ExecuteContext.GetControls();
+	if (Hierarchy)
+	{
+		switch (Context.State)
+		{
+			case EControlRigState::Init:
+			{
+				CachedControlIndex = Hierarchy->GetIndex(Control);
+				break;
+			}
+			case EControlRigState::Update:
+			{
+				if (CachedControlIndex != INDEX_NONE)
+				{
+					FTransform Transform = FTransform::Identity;
+					if (Space == EBoneGetterSetterMode::GlobalSpace)
+					{
+						Transform = Hierarchy->GetGlobalTransform(CachedControlIndex);
+					}
+
+					if ((*Hierarchy)[CachedControlIndex].ControlType == ERigControlType::Position)
+					{
+						if(FMath::IsNearlyEqual(Weight, 1.f))
+						{
+							Transform.SetLocation(Vector);
+						}
+						else
+						{
+							FVector PreviousValue = Transform.GetLocation();
+							Transform.SetLocation(FMath::Lerp<FVector>(PreviousValue, Vector, FMath::Clamp<float>(Weight, 0.f, 1.f)));
+						}
+					}
+					else if ((*Hierarchy)[CachedControlIndex].ControlType == ERigControlType::Scale)
+					{
+						if(FMath::IsNearlyEqual(Weight, 1.f))
+						{
+							Transform.SetScale3D(Vector);
+						}
+						else
+						{
+							FVector PreviousValue = Transform.GetScale3D();
+							Transform.SetScale3D(FMath::Lerp<FVector>(PreviousValue, Vector, FMath::Clamp<float>(Weight, 0.f, 1.f)));
+						}
+					}
+
+					switch (Space)
+					{
+						case EBoneGetterSetterMode::GlobalSpace:
+						{
+							Hierarchy->SetGlobalTransform(CachedControlIndex, Transform);
+							break;
+						}
+						case EBoneGetterSetterMode::LocalSpace:
+						{
+							Hierarchy->SetLocalTransform(CachedControlIndex, Transform);
+							break;
+						}
+						default:
+						{
+							break;
+						}
+					}
+				}
+			}
+			default:
+			{
+				break;
+			}
+		}
+	}
+}
+
+FRigUnit_SetControlRotator_Execute()
 {
     DECLARE_SCOPE_HIERARCHICAL_COUNTER_RIGUNIT()
 	FRigControlHierarchy* Hierarchy = ExecuteContext.GetControls();
@@ -25,6 +199,24 @@ FRigUnit_SetControlTransform_Execute()
 			{
 				if (CachedControlIndex != INDEX_NONE)
 				{
+					FTransform Transform = FTransform::Identity;
+					if (Space == EBoneGetterSetterMode::GlobalSpace)
+					{
+						Transform = Hierarchy->GetGlobalTransform(CachedControlIndex);
+					}
+
+					FQuat Quat = FQuat(Rotator);
+					if (FMath::IsNearlyEqual(Weight, 1.f))
+					{
+						Transform.SetRotation(Quat);
+					}
+					else
+					{
+						FQuat PreviousValue = Transform.GetRotation();
+						Transform.SetRotation(FQuat::Slerp(PreviousValue, Quat, FMath::Clamp<float>(Weight, 0.f, 1.f)));
+					}
+					Transform.NormalizeRotation();
+
 					switch (Space)
 					{
 						case EBoneGetterSetterMode::GlobalSpace:
@@ -51,3 +243,63 @@ FRigUnit_SetControlTransform_Execute()
 		}
 	}
 }
+
+FRigUnit_SetControlTransform_Execute()
+{
+    DECLARE_SCOPE_HIERARCHICAL_COUNTER_RIGUNIT()
+	FRigControlHierarchy* Hierarchy = ExecuteContext.GetControls();
+	if (Hierarchy)
+	{
+		switch (Context.State)
+		{
+			case EControlRigState::Init:
+			{
+				CachedControlIndex = Hierarchy->GetIndex(Control);
+				break;
+			}
+			case EControlRigState::Update:
+			{
+				if (CachedControlIndex != INDEX_NONE)
+				{
+					switch (Space)
+					{
+						case EBoneGetterSetterMode::GlobalSpace:
+						{
+							if(FMath::IsNearlyEqual(Weight, 1.f))
+							{
+								Hierarchy->SetGlobalTransform(CachedControlIndex, Transform);
+							}
+							else
+							{
+								FTransform PreviousTransform = Hierarchy->GetGlobalTransform(CachedControlIndex);
+								Hierarchy->SetGlobalTransform(CachedControlIndex, FControlRigMathLibrary::LerpTransform(PreviousTransform, Transform, FMath::Clamp<float>(Weight, 0.f, 1.f)));
+							}
+							break;
+						}
+						case EBoneGetterSetterMode::LocalSpace:
+						{
+							if(FMath::IsNearlyEqual(Weight, 1.f))
+							{
+								Hierarchy->SetLocalTransform(CachedControlIndex, Transform);
+							}
+							else
+							{
+								FTransform PreviousTransform = Hierarchy->GetLocalTransform(CachedControlIndex);
+								Hierarchy->SetLocalTransform(CachedControlIndex, FControlRigMathLibrary::LerpTransform(PreviousTransform, Transform, FMath::Clamp<float>(Weight, 0.f, 1.f)));
+							}
+							break;
+						}
+						default:
+						{
+							break;
+						}
+					}
+				}
+			}
+			default:
+			{
+				break;
+			}
+		}
+	}
+}
diff --git a/Engine/Plugins/Experimental/ControlRig/Source/ControlRig/Private/Units/Hierarchy/RigUnit_SetControlTransform.h b/Engine/Plugins/Experimental/ControlRig/Source/ControlRig/Private/Units/Hierarchy/RigUnit_SetControlTransform.h
--- a/Engine/Plugins/Experimental/ControlRig/Source/ControlRig/Private/Units/Hierarchy/RigUnit_SetControlTransform.h
+++ b/Engine/Plugins/Experimental/ControlRig/Source/ControlRig/Private/Units/Hierarchy/RigUnit_SetControlTransform.h
@@ -5,31 +5,244 @@
 #include "Units/RigUnit.h"
 #include "RigUnit_SetControlTransform.generated.h"
 
+/**
+ * SetControlBool is used to perform a change in the hierarchy by setting a single bone's Bool.
+ */
+USTRUCT(meta=(DisplayName="Set Control Bool", Category="Hierarchy", DocumentationPolicy="Strict", Keywords = "SetControlBool", PrototypeName = "SetControlValue"))
+struct FRigUnit_SetControlBool : public FRigUnitMutable
+{
+	GENERATED_BODY()
+
+	FRigUnit_SetControlBool()
+		: BoolValue(false)
+		, CachedControlIndex(INDEX_NONE)
+	{}
+
+	RIGVM_METHOD()
+	virtual void Execute(const FRigUnitContext& Context) override;
+
+	/**
+	 * The name of the Control to set the bool for.
+	 */
+	UPROPERTY(meta = (Input, CustomWidget = "ControlName", Constant))
+	FName Control;
+
+	/**
+	 * The bool value to set for the given Control.
+	 */
+	UPROPERTY(meta = (Input, Output))
+	bool BoolValue;
+
+	// Used to cache the internally used bone index
+	UPROPERTY()
+	int32 CachedControlIndex;
+};
+
+/**
+ * SetControlFloat is used to perform a change in the hierarchy by setting a single bone's Float.
+ */
+USTRUCT(meta=(DisplayName="Set Control Float", Category="Hierarchy", DocumentationPolicy="Strict", Keywords = "SetControlFloat", PrototypeName = "SetControlValue"))
+struct FRigUnit_SetControlFloat : public FRigUnitMutable
+{
+	GENERATED_BODY()
+
+	FRigUnit_SetControlFloat()
+		: Weight(1.f)
+		, FloatValue(0.f)
+		, CachedControlIndex(INDEX_NONE)
+	{}
+
+	RIGVM_METHOD()
+	virtual void Execute(const FRigUnitContext& Context) override;
+
+	/**
+	 * The name of the Control to set the transform for.
+	 */
+	UPROPERTY(meta = (Input, CustomWidget = "ControlName", Constant))
+	FName Control;
+
+	/**
+	 * The weight of the change - how much the change should be applied
+	 */
+	UPROPERTY(meta = (Input, UIMin = "0.0", UIMax = "1.0"))
+	float Weight;
+
+	/**
+	 * The transform value to set for the given Control.
+	 */
+	UPROPERTY(meta = (Input, Output, UIMin = "0.0", UIMax = "1.0"))
+	float FloatValue;
+
+	// Used to cache the internally used bone index
+	UPROPERTY()
+	int32 CachedControlIndex;
+};
+
+/**
+ * SetControlVector2D is used to perform a change in the hierarchy by setting a single bone's Vector2D.
+ */
+USTRUCT(meta=(DisplayName="Set Control Vector2D", Category="Hierarchy", DocumentationPolicy="Strict", Keywords = "SetControlVector2D", PrototypeName = "SetControlValue"))
+struct FRigUnit_SetControlVector2D : public FRigUnitMutable
+{
+	GENERATED_BODY()
+
+	FRigUnit_SetControlVector2D()
+		: Weight(1.f)
+		, Vector(FVector2D::ZeroVector)
+		, CachedControlIndex(INDEX_NONE)
+	{}
+
+	RIGVM_METHOD()
+	virtual void Execute(const FRigUnitContext& Context) override;
+
+	/**
+	 * The name of the Control to set the transform for.
+	 */
+	UPROPERTY(meta = (Input, CustomWidget = "ControlName", Constant))
+	FName Control;
+
+	/**
+	 * The weight of the change - how much the change should be applied
+	 */
+	UPROPERTY(meta = (Input, UIMin = "0.0", UIMax = "1.0"))
+	float Weight;
+
+	/**
+	 * The transform value to set for the given Control.
+	 */
+	UPROPERTY(meta = (Input, Output))
+	FVector2D Vector;
+
+	// Used to cache the internally used bone index
+	UPROPERTY()
+	int32 CachedControlIndex;
+};
+
+/**
+ * SetControlVector is used to perform a change in the hierarchy by setting a single bone's Vector.
+ */
+USTRUCT(meta=(DisplayName="Set Control Vector", Category="Hierarchy", DocumentationPolicy="Strict", Keywords = "SetControlVector", PrototypeName = "SetControlValue"))
+struct FRigUnit_SetControlVector : public FRigUnitMutable
+{
+	GENERATED_BODY()
+
+	FRigUnit_SetControlVector()
+		: Weight(1.f)
+		, Vector(FVector::OneVector)
+		, Space(EBoneGetterSetterMode::GlobalSpace)
+		, CachedControlIndex(INDEX_NONE)
+	{}
+
+	RIGVM_METHOD()
+	virtual void Execute(const FRigUnitContext& Context) override;
+
+	/**
+	 * The name of the Control to set the transform for.
+	 */
+	UPROPERTY(meta = (Input, CustomWidget = "ControlName", Constant))
+	FName Control;
+
+	/**
+	 * The weight of the change - how much the change should be applied
+	 */
+	UPROPERTY(meta = (Input, UIMin = "0.0", UIMax = "1.0"))
+	float Weight;
+
+	/**
+	 * The transform value to set for the given Control.
+	 */
+	UPROPERTY(meta = (Input, Output))
+	FVector Vector;
+
+	/**
+	 * Defines if the bone's transform should be set
+	 * in local or global space.
+	 */
+	UPROPERTY(meta = (Input))
+	EBoneGetterSetterMode Space;
+
+	// Used to cache the internally used bone index
+	UPROPERTY()
+	int32 CachedControlIndex;
+};
+
+/**
+ * SetControlRotator is used to perform a change in the hierarchy by setting a single bone's Rotator.
+ */
+USTRUCT(meta=(DisplayName="Set Control Rotator", Category="Hierarchy", DocumentationPolicy="Strict", Keywords = "SetControlRotator", PrototypeName = "SetControlValue"))
+struct FRigUnit_SetControlRotator : public FRigUnitMutable
+{
+	GENERATED_BODY()
+
+	FRigUnit_SetControlRotator()
+		: Weight(1.f)
+		, Rotator(FRotator::ZeroRotator)
+		, Space(EBoneGetterSetterMode::GlobalSpace)
+		, CachedControlIndex(INDEX_NONE)
+	{}
+
+	RIGVM_METHOD()
+	virtual void Execute(const FRigUnitContext& Context) override;
+
+	/**
+	 * The name of the Control to set the transform for.
+	 */
+	UPROPERTY(meta = (Input, CustomWidget = "ControlName", Constant))
+	FName Control;
+
+	/**
+	 * The weight of the change - how much the change should be applied
+	 */
+	UPROPERTY(meta = (Input, UIMin = "0.0", UIMax = "1.0"))
+	float Weight;
+
+	/**
+	 * The transform value to set for the given Control.
+	 */
+	UPROPERTY(meta = (Input, Output))
+	FRotator Rotator;
+
+	/**
+	 * Defines if the bone's transform should be set
+	 * in local or global space.
+	 */
+	UPROPERTY(meta = (Input))
+	EBoneGetterSetterMode Space;
+
+	// Used to cache the internally used bone index
+	UPROPERTY()
+	int32 CachedControlIndex;
+};
 
 /**
- * SetControlTransform is used to perform a change in the hierarchy by setting a single bone's transform.
+ * SetControlTransform is used to perform a change in the hierarchy by setting a single control's transform.
  */
-USTRUCT(meta=(DisplayName="Set Control", Category="Hierarchy", DocumentationPolicy="Strict", Keywords = "SetControlTransform"))
+USTRUCT(meta=(DisplayName="Set Control Transform", Category="Hierarchy", DocumentationPolicy="Strict", Keywords = "SetControlTransform", PrototypeName = "SetControlValue"))
 struct FRigUnit_SetControlTransform : public FRigUnitMutable
 {
 	GENERATED_BODY()
 
 	FRigUnit_SetControlTransform()
-		: Space(EBoneGetterSetterMode::LocalSpace)
+		: Weight(1.f)
+		, Space(EBoneGetterSetterMode::GlobalSpace)
 		, CachedControlIndex(INDEX_NONE)
 	{}
 
-	virtual FString GetUnitLabel() const override;
-
 	RIGVM_METHOD()
 	virtual void Execute(const FRigUnitContext& Context) override;
 
 	/**
 	 * The name of the Control to set the transform for.
 	 */
-	UPROPERTY(meta = (Input, ControlName, Constant))
+	UPROPERTY(meta = (Input, CustomWidget = "ControlName", Constant))
 	FName Control;
 
+	/**
+	 * The weight of the change - how much the change should be applied
+	 */
+	UPROPERTY(meta = (Input, UIMin = "0.0", UIMax = "1.0"))
+	float Weight;
+
 	/**
 	 * The transform value to set for the given Control.
 	 */
diff --git a/Engine/Plugins/Experimental/ControlRig/Source/ControlRig/Private/Units/Hierarchy/RigUnit_SetCurveValue.h b/Engine/Plugins/Experimental/ControlRig/Source/ControlRig/Private/Units/Hierarchy/RigUnit_SetCurveValue.h
--- a/Engine/Plugins/Experimental/ControlRig/Source/ControlRig/Private/Units/Hierarchy/RigUnit_SetCurveValue.h
+++ b/Engine/Plugins/Experimental/ControlRig/Source/ControlRig/Private/Units/Hierarchy/RigUnit_SetCurveValue.h
@@ -26,7 +26,7 @@ struct FRigUnit_SetCurveValue : public FRigUnitMutable
 	/**
 	 * The name of the Curve to set the Value for.
 	 */
-	UPROPERTY(meta = (Input, CurveName, Constant))
+	UPROPERTY(meta = (Input, CustomWidget = "CurveName", Constant))
 	FName Curve;
 
 	/**
diff --git a/Engine/Plugins/Experimental/ControlRig/Source/ControlRig/Private/Units/Hierarchy/RigUnit_SetRelativeBoneTransform.cpp b/Engine/Plugins/Experimental/ControlRig/Source/ControlRig/Private/Units/Hierarchy/RigUnit_SetRelativeBoneTransform.cpp
--- a/Engine/Plugins/Experimental/ControlRig/Source/ControlRig/Private/Units/Hierarchy/RigUnit_SetRelativeBoneTransform.cpp
+++ b/Engine/Plugins/Experimental/ControlRig/Source/ControlRig/Private/Units/Hierarchy/RigUnit_SetRelativeBoneTransform.cpp
@@ -20,6 +20,14 @@ FRigUnit_SetRelativeBoneTransform_Execute()
 			{
 				CachedBoneIndex = Hierarchy->GetIndex(Bone);
 				CachedSpaceIndex = Hierarchy->GetIndex(Space);
+				if (CachedBoneIndex == INDEX_NONE)
+				{
+					UE_CONTROLRIG_RIGUNIT_REPORT_WARNING(TEXT("Bone is not set."));
+				}
+				if (CachedSpaceIndex == INDEX_NONE)
+				{
+					UE_CONTROLRIG_RIGUNIT_REPORT_WARNING(TEXT("Space is not set."));
+				}
 				// fall through to update
 			}
 			case EControlRigState::Update:
diff --git a/Engine/Plugins/Experimental/ControlRig/Source/ControlRig/Private/Units/Hierarchy/RigUnit_SetRelativeBoneTransform.h b/Engine/Plugins/Experimental/ControlRig/Source/ControlRig/Private/Units/Hierarchy/RigUnit_SetRelativeBoneTransform.h
--- a/Engine/Plugins/Experimental/ControlRig/Source/ControlRig/Private/Units/Hierarchy/RigUnit_SetRelativeBoneTransform.h
+++ b/Engine/Plugins/Experimental/ControlRig/Source/ControlRig/Private/Units/Hierarchy/RigUnit_SetRelativeBoneTransform.h
@@ -22,19 +22,29 @@ struct FRigUnit_SetRelativeBoneTransform : public FRigUnitMutable
 	{}
 
 	virtual FString GetUnitLabel() const override;
+
+	virtual FName DetermineSpaceForPin(const FString& InPinPath, void* InUserContext) const override
+	{
+		if (InPinPath.StartsWith(TEXT("Transform")))
+		{
+			return Space;
+		}
+		return NAME_None;
+	}
+
 	RIGVM_METHOD()
 	virtual void Execute(const FRigUnitContext& Context) override;
 
 	/**
 	 * The name of the Bone to set the transform for.
 	 */
-	UPROPERTY(meta = (Input, BoneName, Constant))
+	UPROPERTY(meta = (Input, CustomWidget = "BoneName", Constant))
 	FName Bone;
 
 	/**
 	 * The name of the Bone to set the transform relative within.
 	 */
-	UPROPERTY(meta = (Input, BoneName, Constant))
+	UPROPERTY(meta = (Input, CustomWidget = "BoneName", Constant))
 	FName Space;
 
 	/**
@@ -54,7 +64,7 @@ struct FRigUnit_SetRelativeBoneTransform : public FRigUnitMutable
 	 * of this bone will be recalculated based on their local transforms.
 	 * Note: This is computationally more expensive than turning it off.
 	 */
-	UPROPERTY(meta = (Input))
+	UPROPERTY(meta = (Input, Constant))
 	bool bPropagateToChildren;
 
 	// Used to cache the internally used bone index
diff --git a/Engine/Plugins/Experimental/ControlRig/Source/ControlRig/Private/Units/Hierarchy/RigUnit_SetSpaceTransform.cpp b/Engine/Plugins/Experimental/ControlRig/Source/ControlRig/Private/Units/Hierarchy/RigUnit_SetSpaceTransform.cpp
--- a/Engine/Plugins/Experimental/ControlRig/Source/ControlRig/Private/Units/Hierarchy/RigUnit_SetSpaceTransform.cpp
+++ b/Engine/Plugins/Experimental/ControlRig/Source/ControlRig/Private/Units/Hierarchy/RigUnit_SetSpaceTransform.cpp
@@ -2,6 +2,7 @@
 
 #include "RigUnit_SetSpaceTransform.h"
 #include "Units/RigUnitContext.h"
+#include "Math/ControlRigMathLibrary.h"
 
 FString FRigUnit_SetSpaceTransform::GetUnitLabel() const
 {
@@ -29,12 +30,28 @@ FRigUnit_SetSpaceTransform_Execute()
 					{
 						case EBoneGetterSetterMode::GlobalSpace:
 						{
-							Hierarchy->SetGlobalTransform(CachedSpaceIndex, Transform);
+							if(FMath::IsNearlyEqual(Weight, 1.f))
+							{
+								Hierarchy->SetGlobalTransform(CachedSpaceIndex, Transform);
+							}
+							else
+							{
+								FTransform PreviousTransform = Hierarchy->GetGlobalTransform(CachedSpaceIndex);
+								Hierarchy->SetGlobalTransform(CachedSpaceIndex, FControlRigMathLibrary::LerpTransform(PreviousTransform, Transform, FMath::Clamp<float>(Weight, 0.f, 1.f)));
+							}
 							break;
 						}
 						case EBoneGetterSetterMode::LocalSpace:
 						{
-							Hierarchy->SetLocalTransform(CachedSpaceIndex, Transform);
+							if(FMath::IsNearlyEqual(Weight, 1.f))
+							{
+								Hierarchy->SetLocalTransform(CachedSpaceIndex, Transform);
+							}
+							else
+							{
+								FTransform PreviousTransform = Hierarchy->GetLocalTransform(CachedSpaceIndex);
+								Hierarchy->SetLocalTransform(CachedSpaceIndex, FControlRigMathLibrary::LerpTransform(PreviousTransform, Transform, FMath::Clamp<float>(Weight, 0.f, 1.f)));
+							}
 							break;
 						}
 						default:
diff --git a/Engine/Plugins/Experimental/ControlRig/Source/ControlRig/Private/Units/Hierarchy/RigUnit_SetSpaceTransform.h b/Engine/Plugins/Experimental/ControlRig/Source/ControlRig/Private/Units/Hierarchy/RigUnit_SetSpaceTransform.h
--- a/Engine/Plugins/Experimental/ControlRig/Source/ControlRig/Private/Units/Hierarchy/RigUnit_SetSpaceTransform.h
+++ b/Engine/Plugins/Experimental/ControlRig/Source/ControlRig/Private/Units/Hierarchy/RigUnit_SetSpaceTransform.h
@@ -7,7 +7,7 @@
 
 
 /**
- * SetSpaceTransform is used to perform a change in the hierarchy by setting a single bone's transform.
+ * SetSpaceTransform is used to perform a change in the hierarchy by setting a single space's transform.
  */
 USTRUCT(meta=(DisplayName="Set Space", Category="Hierarchy", DocumentationPolicy="Strict", Keywords = "SetSpaceTransform"))
 struct FRigUnit_SetSpaceTransform : public FRigUnitMutable
@@ -15,7 +15,8 @@ struct FRigUnit_SetSpaceTransform : public FRigUnitMutable
 	GENERATED_BODY()
 
 	FRigUnit_SetSpaceTransform()
-		: SpaceType(EBoneGetterSetterMode::LocalSpace)
+		: Weight(1.f)
+		, SpaceType(EBoneGetterSetterMode::GlobalSpace)
 		, CachedSpaceIndex(INDEX_NONE)
 	{}
 
@@ -27,9 +28,15 @@ struct FRigUnit_SetSpaceTransform : public FRigUnitMutable
 	/**
 	 * The name of the Space to set the transform for.
 	 */
-	UPROPERTY(meta = (Input, SpaceName, Constant))
+	UPROPERTY(meta = (Input, CustomWidget = "SpaceName", Constant))
 	FName Space;
 
+	/**
+	 * The weight of the change - how much the change should be applied
+	 */
+	UPROPERTY(meta = (Input, UIMin = "0.0", UIMax = "1.0"))
+	float Weight;
+
 	/**
 	 * The transform value to set for the given Space.
 	 */
diff --git a/Engine/Plugins/Experimental/ControlRig/Source/ControlRig/Private/Units/Highlevel/Harmonics/RigUnit_BoneHarmonics.cpp b/Engine/Plugins/Experimental/ControlRig/Source/ControlRig/Private/Units/Highlevel/Harmonics/RigUnit_BoneHarmonics.cpp
--- a/Engine/Plugins/Experimental/ControlRig/Source/ControlRig/Private/Units/Highlevel/Harmonics/RigUnit_BoneHarmonics.cpp
+++ b/Engine/Plugins/Experimental/ControlRig/Source/ControlRig/Private/Units/Highlevel/Harmonics/RigUnit_BoneHarmonics.cpp
@@ -24,7 +24,12 @@ FRigUnit_BoneHarmonics_Execute()
 		{
 			for (const FRigUnit_BoneHarmonics_BoneTarget& Bone : Bones)
 			{
-				BoneIndices.Add(Hierarchy->GetIndex(Bone.Bone));
+				int32 BoneIndex = Hierarchy->GetIndex(Bone.Bone);
+				if(BoneIndex == INDEX_NONE)
+				{
+					UE_CONTROLRIG_RIGUNIT_REPORT_WARNING(TEXT("Bone not found."));
+				}
+				BoneIndices.Add(BoneIndex);
 			}
 		}
 	}
diff --git a/Engine/Plugins/Experimental/ControlRig/Source/ControlRig/Private/Units/Highlevel/Harmonics/RigUnit_BoneHarmonics.h b/Engine/Plugins/Experimental/ControlRig/Source/ControlRig/Private/Units/Highlevel/Harmonics/RigUnit_BoneHarmonics.h
--- a/Engine/Plugins/Experimental/ControlRig/Source/ControlRig/Private/Units/Highlevel/Harmonics/RigUnit_BoneHarmonics.h
+++ b/Engine/Plugins/Experimental/ControlRig/Source/ControlRig/Private/Units/Highlevel/Harmonics/RigUnit_BoneHarmonics.h
@@ -20,7 +20,7 @@ struct FRigUnit_BoneHarmonics_BoneTarget
 	/**
 	 * The name of the bone to drive
 	 */
-	UPROPERTY(meta = (Input, Constant, BoneName))
+	UPROPERTY(meta = (Input, Constant, CustomWidget = "BoneName"))
 	FName Bone;
 
 	/**
@@ -110,7 +110,7 @@ struct FRigUnit_BoneHarmonics : public FRigUnit_HighlevelBaseMutable
 	 * of this bone will be recalculated based on their local transforms.
 	 * Note: This is computationally more expensive than turning it off.
 	 */
-	UPROPERTY(meta = (Input))
+	UPROPERTY(meta = (Input, Constant))
 	bool bPropagateToChildren;
 
 	UPROPERTY(transient)
diff --git a/Engine/Plugins/Experimental/ControlRig/Source/ControlRig/Private/Units/Highlevel/Harmonics/RigUnit_ChainHarmonics.h b/Engine/Plugins/Experimental/ControlRig/Source/ControlRig/Private/Units/Highlevel/Harmonics/RigUnit_ChainHarmonics.h
--- a/Engine/Plugins/Experimental/ControlRig/Source/ControlRig/Private/Units/Highlevel/Harmonics/RigUnit_ChainHarmonics.h
+++ b/Engine/Plugins/Experimental/ControlRig/Source/ControlRig/Private/Units/Highlevel/Harmonics/RigUnit_ChainHarmonics.h
@@ -200,7 +200,7 @@ struct FRigUnit_ChainHarmonics : public FRigUnit_HighlevelBaseMutable
 	RIGVM_METHOD()
 	virtual void Execute(const FRigUnitContext& Context) override;
 
-	UPROPERTY(meta = (Input, Constant, BoneName))
+	UPROPERTY(meta = (Input, Constant, CustomWidget = "BoneName"))
 	FName ChainRoot;
 
 	UPROPERTY(meta = (Input))
diff --git a/Engine/Plugins/Experimental/ControlRig/Source/ControlRig/Private/Units/Highlevel/Hierarchy/RigUnit_AimBone.cpp b/Engine/Plugins/Experimental/ControlRig/Source/ControlRig/Private/Units/Highlevel/Hierarchy/RigUnit_AimBone.cpp
--- a/Engine/Plugins/Experimental/ControlRig/Source/ControlRig/Private/Units/Highlevel/Hierarchy/RigUnit_AimBone.cpp
+++ b/Engine/Plugins/Experimental/ControlRig/Source/ControlRig/Private/Units/Highlevel/Hierarchy/RigUnit_AimBone.cpp
@@ -24,7 +24,7 @@ FRigUnit_AimBone_Execute()
 		return;
 	}
 
-	if (Primary.Weight <= SMALL_NUMBER && Secondary.Weight <= SMALL_NUMBER)
+	if ((Weight <= SMALL_NUMBER) || (Primary.Weight <= SMALL_NUMBER && Secondary.Weight <= SMALL_NUMBER))
 	{
 		return;
 	}
@@ -83,9 +83,10 @@ FRigUnit_AimBone_Execute()
 		{
 			Target = Target.GetSafeNormal();
 			FVector Axis = Transform.TransformVectorNoScale(Primary.Axis).GetSafeNormal();
-			if (Primary.Weight < 1.f - SMALL_NUMBER)
+			float T = Primary.Weight * Weight;
+			if (T < 1.f - SMALL_NUMBER)
 			{
-				Target = FMath::Lerp<FVector>(Axis, Target, Primary.Weight).GetSafeNormal();
+				Target = FMath::Lerp<FVector>(Axis, Target, T).GetSafeNormal();
 			}
 			FQuat Rotation = FQuat::FindBetweenNormals(Axis, Target);
 			Transform.SetRotation((Rotation * Transform.GetRotation()).GetNormalized());
@@ -156,9 +157,10 @@ FRigUnit_AimBone_Execute()
 			Target = Target.GetSafeNormal();
 
 			FVector Axis = Transform.TransformVectorNoScale(Secondary.Axis).GetSafeNormal();
-			if (Secondary.Weight < 1.f - SMALL_NUMBER)
+			float T = Secondary.Weight * Weight;
+			if (T < 1.f - SMALL_NUMBER)
 			{
-				Target = FMath::Lerp<FVector>(Axis, Target, Secondary.Weight).GetSafeNormal();
+				Target = FMath::Lerp<FVector>(Axis, Target, T).GetSafeNormal();
 			}
 			FQuat Rotation = FQuat::FindBetweenNormals(Axis, Target);
 			Transform.SetRotation((Rotation * Transform.GetRotation()).GetNormalized());
diff --git a/Engine/Plugins/Experimental/ControlRig/Source/ControlRig/Private/Units/Highlevel/Hierarchy/RigUnit_AimBone.h b/Engine/Plugins/Experimental/ControlRig/Source/ControlRig/Private/Units/Highlevel/Hierarchy/RigUnit_AimBone.h
--- a/Engine/Plugins/Experimental/ControlRig/Source/ControlRig/Private/Units/Highlevel/Hierarchy/RigUnit_AimBone.h
+++ b/Engine/Plugins/Experimental/ControlRig/Source/ControlRig/Private/Units/Highlevel/Hierarchy/RigUnit_AimBone.h
@@ -22,31 +22,31 @@ struct FRigUnit_AimBone_Target
 	/**
 	 * The amount of aim rotation to apply on this target.
 	 */
-	UPROPERTY(meta = (Input))
+	UPROPERTY(EditAnywhere, meta = (Input), Category = "AimTarget")
 	float Weight;
 
 	/**
 	 * The axis to align with the aim on this target
 	 */
-	UPROPERTY(meta = (Input))
+	UPROPERTY(EditAnywhere, meta = (Input, EditCondition = "Weight > 0.0"), Category = "AimTarget")
 	FVector Axis;
 
 	/**
 	 * The target to aim at - can be a direction or location based on the Kind setting
 	 */
-	UPROPERTY(meta = (Input))
+	UPROPERTY(EditAnywhere, meta = (Input, EditCondition = "Weight > 0.0"), Category = "AimTarget")
 	FVector Target;
 
 	/**
 	 * The kind of target this is representing - can be a direction or a location
 	 */
-	UPROPERTY(meta = (Input))
+	UPROPERTY(EditAnywhere, meta = (Input, EditCondition = "Weight > 0.0"), Category = "AimTarget")
 	EControlRigVectorKind Kind;
 
 	/**
 	 * The space in which the target is expressed
 	 */
-	UPROPERTY(meta = (Input, BoneName, Constant))
+	UPROPERTY(EditAnywhere, meta = (Input, EditCondition = "Weight > 0.0", CustomWidget = "BoneName", Constant), Category = "AimTarget")
 	FName Space;
 };
 
@@ -65,19 +65,19 @@ struct FRigUnit_AimBone_DebugSettings
 	/**
 	 * If enabled debug information will be drawn 
 	 */
-	UPROPERTY(meta = (Input))
+	UPROPERTY(EditAnywhere, meta = (Input), Category = "DebugSettings")
 	bool bEnabled;
 
 	/**
 	 * The size of the debug drawing information
 	 */
-	UPROPERTY(meta = (Input))
+	UPROPERTY(EditAnywhere, meta = (Input, EditCondition = "bEnabled"), Category = "DebugSettings")
 	float Scale;
 
 	/**
 	 * The offset at which to draw the debug information in the world
 	 */
-	UPROPERTY(meta = (Input))
+	UPROPERTY(EditAnywhere, meta = (Input, EditCondition = "bEnabled"), Category = "DebugSettings")
 	FTransform WorldOffset;
 };
 
@@ -97,6 +97,7 @@ struct FRigUnit_AimBone : public FRigUnit_HighlevelBaseMutable
 		Secondary = FRigUnit_AimBone_Target();
 		Primary.Axis = FVector(1.f, 0.f, 0.f);
 		Secondary.Axis = FVector(0.f, 0.f, 1.f);
+		Weight = 1.f;
 		bPropagateToChildren = false;
 		DebugSettings = FRigUnit_AimBone_DebugSettings();
 		BoneIndex = INDEX_NONE;
@@ -106,13 +107,26 @@ struct FRigUnit_AimBone : public FRigUnit_HighlevelBaseMutable
 		SecondaryCachedSpaceIndex = INDEX_NONE;
 	}
 
+	virtual FName DetermineSpaceForPin(const FString& InPinPath, void* InUserContext) const override
+	{
+		if (InPinPath.StartsWith(TEXT("Primary.Target")))
+		{
+			return Primary.Space;
+		}
+		if (InPinPath.StartsWith(TEXT("Secondary.Target")))
+		{
+			return Secondary.Space;
+		}
+		return NAME_None;
+	}
+
 	RIGVM_METHOD()
 	virtual void Execute(const FRigUnitContext& Context) override;
 
 	/** 
 	 * The name of the bone to align
 	 */
-	UPROPERTY(meta = (Input, Constant, BoneName))
+	UPROPERTY(meta = (Input, Constant, CustomWidget = "BoneName"))
 	FName Bone;
 
 	/**
@@ -127,12 +141,18 @@ struct FRigUnit_AimBone : public FRigUnit_HighlevelBaseMutable
 	UPROPERTY(meta = (Input))
 	FRigUnit_AimBone_Target Secondary;
 
+	/**
+	 * The weight of the change - how much the change should be applied
+	 */
+	UPROPERTY(meta = (Input, UIMin = "0.0", UIMax = "1.0"))
+	float Weight;
+
 	/**
 	 * If set to true all of the global transforms of the children
 	 * of this bone will be recalculated based on their local transforms.
 	 * Note: This is computationally more expensive than turning it off.
 	 */
-	UPROPERTY(meta = (Input))
+	UPROPERTY(meta = (Input, Constant))
 	bool bPropagateToChildren;
 
 	/** The debug setting for the node */
diff --git a/Engine/Plugins/Experimental/ControlRig/Source/ControlRig/Private/Units/Highlevel/Hierarchy/RigUnit_CCDIK.cpp b/Engine/Plugins/Experimental/ControlRig/Source/ControlRig/Private/Units/Highlevel/Hierarchy/RigUnit_CCDIK.cpp
--- a/Engine/Plugins/Experimental/ControlRig/Source/ControlRig/Private/Units/Highlevel/Hierarchy/RigUnit_CCDIK.cpp
+++ b/Engine/Plugins/Experimental/ControlRig/Source/ControlRig/Private/Units/Highlevel/Hierarchy/RigUnit_CCDIK.cpp
@@ -50,6 +50,10 @@ FRigUnit_CCDIK_Execute()
 
 			Chain.Reserve(BoneIndices.Num());
 		}
+		else
+		{
+			return;
+		}
 
 		int32 RootParentIndex = (*Hierarchy)[RootIndex].ParentIndex;
 		if (RootParentIndex != INDEX_NONE)
@@ -66,6 +70,11 @@ FRigUnit_CCDIK_Execute()
 			BoneIndex = BoneIndices.Find(BoneIndex);
 			RotationLimitIndex.Add(BoneIndex);
 		}
+
+		if(BoneIndices.Num() < 2)
+		{
+			UE_CONTROLRIG_RIGUNIT_REPORT_WARNING(TEXT("No bones found."));
+		}
 	}
 	else  if (Context.State == EControlRigState::Update)
 	{
diff --git a/Engine/Plugins/Experimental/ControlRig/Source/ControlRig/Private/Units/Highlevel/Hierarchy/RigUnit_CCDIK.h b/Engine/Plugins/Experimental/ControlRig/Source/ControlRig/Private/Units/Highlevel/Hierarchy/RigUnit_CCDIK.h
--- a/Engine/Plugins/Experimental/ControlRig/Source/ControlRig/Private/Units/Highlevel/Hierarchy/RigUnit_CCDIK.h
+++ b/Engine/Plugins/Experimental/ControlRig/Source/ControlRig/Private/Units/Highlevel/Hierarchy/RigUnit_CCDIK.h
@@ -20,7 +20,7 @@ struct FRigUnit_CCDIK_RotationLimit
 	/**
 	 * The name of the bone to apply the rotation limit to.
 	 */
-	UPROPERTY(meta = (Input, Constant, BoneName))
+	UPROPERTY(meta = (Input, Constant, CustomWidget = "BoneName"))
 	FName Bone;
 
 	/**
@@ -83,13 +83,13 @@ struct FRigUnit_CCDIK : public FRigUnit_HighlevelBaseMutable
 	/**
 	 * The first bone in the chain to solve
 	 */
-	UPROPERTY(meta = (Input, Constant, BoneName))
+	UPROPERTY(meta = (Input, Constant, CustomWidget = "BoneName"))
 	FName StartBone;
 
 	/**
 	 * The last bone in the chain to solve - the effector
 	 */
-	UPROPERTY(meta = (Input, Constant, BoneName))
+	UPROPERTY(meta = (Input, Constant, CustomWidget = "BoneName"))
 	FName EffectorBone;
 
 	/**
@@ -139,7 +139,7 @@ struct FRigUnit_CCDIK : public FRigUnit_HighlevelBaseMutable
 	 * of this bone will be recalculated based on their local transforms.
 	 * Note: This is computationally more expensive than turning it off.
 	 */
-	UPROPERTY(meta = (Input))
+	UPROPERTY(meta = (Input, Constant))
 	bool bPropagateToChildren;
 
 	UPROPERTY(transient)
diff --git a/Engine/Plugins/Experimental/ControlRig/Source/ControlRig/Private/Units/Highlevel/Hierarchy/RigUnit_DistributeRotation.h b/Engine/Plugins/Experimental/ControlRig/Source/ControlRig/Private/Units/Highlevel/Hierarchy/RigUnit_DistributeRotation.h
--- a/Engine/Plugins/Experimental/ControlRig/Source/ControlRig/Private/Units/Highlevel/Hierarchy/RigUnit_DistributeRotation.h
+++ b/Engine/Plugins/Experimental/ControlRig/Source/ControlRig/Private/Units/Highlevel/Hierarchy/RigUnit_DistributeRotation.h
@@ -68,19 +68,28 @@ struct FRigUnit_DistributeRotation : public FRigUnit_HighlevelBaseMutable
 		bPropagateToChildren = false;
 	}
 
+	virtual FName DetermineSpaceForPin(const FString& InPinPath, void* InUserContext) const override
+	{
+		if (InPinPath.StartsWith(TEXT("Rotations")))
+		{
+			return StartBone;
+		}
+		return NAME_None;
+	}
+
 	RIGVM_METHOD()
 	virtual void Execute(const FRigUnitContext& Context) override;
 
 	/** 
 	 * The name of the first bone to align
 	 */
-	UPROPERTY(meta = (Input, Constant, BoneName))
+	UPROPERTY(meta = (Input, Constant, CustomWidget = "BoneName"))
 	FName StartBone;
 
 	/** 
 	 * The name of the last bone to align
 	 */
-	UPROPERTY(meta = (Input, Constant, BoneName))
+	UPROPERTY(meta = (Input, Constant, CustomWidget = "BoneName"))
 	FName EndBone;
 
 	/** 
@@ -100,7 +109,7 @@ struct FRigUnit_DistributeRotation : public FRigUnit_HighlevelBaseMutable
 	 * of this bone will be recalculated based on their local transforms.
 	 * Note: This is computationally more expensive than turning it off.
 	 */
-	UPROPERTY(meta = (Input))
+	UPROPERTY(meta = (Input, Constant))
 	bool bPropagateToChildren;
 
 	UPROPERTY(transient)
diff --git a/Engine/Plugins/Experimental/ControlRig/Source/ControlRig/Private/Units/Highlevel/Hierarchy/RigUnit_FABRIK.h b/Engine/Plugins/Experimental/ControlRig/Source/ControlRig/Private/Units/Highlevel/Hierarchy/RigUnit_FABRIK.h
--- a/Engine/Plugins/Experimental/ControlRig/Source/ControlRig/Private/Units/Highlevel/Hierarchy/RigUnit_FABRIK.h
+++ b/Engine/Plugins/Experimental/ControlRig/Source/ControlRig/Private/Units/Highlevel/Hierarchy/RigUnit_FABRIK.h
@@ -51,13 +51,13 @@ struct FRigUnit_FABRIK : public FRigUnit_HighlevelBaseMutable
 	/**
 	 * The first bone in the chain to solve
 	 */
-	UPROPERTY(meta = (Input, Constant, BoneName))
+	UPROPERTY(meta = (Input, Constant, CustomWidget = "BoneName"))
 	FName StartBone;
 
 	/**
 	 * The last bone in the chain to solve - the effector
 	 */
-	UPROPERTY(meta = (Input, Constant, BoneName))
+	UPROPERTY(meta = (Input, Constant, CustomWidget = "BoneName"))
 	FName EffectorBone;
 
 	/**
@@ -83,7 +83,7 @@ struct FRigUnit_FABRIK : public FRigUnit_HighlevelBaseMutable
 	 * of this bone will be recalculated based on their local transforms.
 	 * Note: This is computationally more expensive than turning it off.
 	 */
-	UPROPERTY(meta = (Input))
+	UPROPERTY(meta = (Input, Constant))
 	bool bPropagateToChildren;
 
 	/**
diff --git a/Engine/Plugins/Experimental/ControlRig/Source/ControlRig/Private/Units/Highlevel/Hierarchy/RigUnit_FitChainToCurve.h b/Engine/Plugins/Experimental/ControlRig/Source/ControlRig/Private/Units/Highlevel/Hierarchy/RigUnit_FitChainToCurve.h
--- a/Engine/Plugins/Experimental/ControlRig/Source/ControlRig/Private/Units/Highlevel/Hierarchy/RigUnit_FitChainToCurve.h
+++ b/Engine/Plugins/Experimental/ControlRig/Source/ControlRig/Private/Units/Highlevel/Hierarchy/RigUnit_FitChainToCurve.h
@@ -27,13 +27,13 @@ struct FRigUnit_FitChainToCurve_Rotation
 	/**
 	 * The rotation to be applied
 	 */
-	UPROPERTY(meta = (Input))
+	UPROPERTY(EditAnywhere, meta = (Input), Category = "Rotation")
 	FQuat Rotation;
 
 	/**
 	 * The ratio of where this rotation sits along the chain
 	 */
-	UPROPERTY(meta = (Input, Constant))
+	UPROPERTY(EditAnywhere, meta = (Input, Constant), Category = "Rotation")
 	float Ratio;
 };
 
@@ -54,31 +54,31 @@ struct FRigUnit_FitChainToCurve_DebugSettings
 	/**
 	 * If enabled debug information will be drawn
 	 */
-	UPROPERTY(meta = (Input))
+	UPROPERTY(EditAnywhere, meta = (Input), Category = "DebugSettings")
 	bool bEnabled;
 
 	/**
 	 * The size of the debug drawing information
 	 */
-	UPROPERTY(meta = (Input))
+	UPROPERTY(EditAnywhere, meta = (Input, EditCondition = "bEnabled"), Category = "DebugSettings")
 	float Scale;
 
 	/**
 	 * The color to use for debug drawing
 	 */
-	UPROPERTY(meta = (Input))
+	UPROPERTY(EditAnywhere, meta = (Input, EditCondition = "bEnabled"), Category = "DebugSettings")
 	FLinearColor CurveColor;
 
 	/**
 	 * The color to use for debug drawing
 	 */
-	UPROPERTY(meta = (Input))
+	UPROPERTY(EditAnywhere, meta = (Input, EditCondition = "bEnabled"), Category = "DebugSettings")
 	FLinearColor SegmentsColor;
 
 	/**
 	 * The offset at which to draw the debug information in the world
 	 */
-	UPROPERTY(meta = (Input))
+	UPROPERTY(EditAnywhere, meta = (Input, EditCondition = "bEnabled"), Category = "DebugSettings")
 	FTransform WorldOffset;
 };
 
@@ -153,13 +153,13 @@ struct FRigUnit_FitChainToCurve : public FRigUnit_HighlevelBaseMutable
 	/** 
 	 * The name of the first bone to align
 	 */
-	UPROPERTY(meta = (Input, Constant, BoneName))
+	UPROPERTY(meta = (Input, Constant, CustomWidget = "BoneName"))
 	FName StartBone;
 
 	/** 
 	 * The name of the last bone to align
 	 */
-	UPROPERTY(meta = (Input, Constant, BoneName))
+	UPROPERTY(meta = (Input, Constant, CustomWidget = "BoneName"))
 	FName EndBone;
 
 	/** 
@@ -229,7 +229,7 @@ struct FRigUnit_FitChainToCurve : public FRigUnit_HighlevelBaseMutable
 	 * of this bone will be recalculated based on their local transforms.
 	 * Note: This is computationally more expensive than turning it off.
 	 */
-	UPROPERTY(meta = (Input))
+	UPROPERTY(meta = (Input, Constant))
 	bool bPropagateToChildren;
 
 	UPROPERTY(meta = (Input))
diff --git a/Engine/Plugins/Experimental/ControlRig/Source/ControlRig/Private/Units/Highlevel/Hierarchy/RigUnit_ModifyBoneTransforms.h b/Engine/Plugins/Experimental/ControlRig/Source/ControlRig/Private/Units/Highlevel/Hierarchy/RigUnit_ModifyBoneTransforms.h
--- a/Engine/Plugins/Experimental/ControlRig/Source/ControlRig/Private/Units/Highlevel/Hierarchy/RigUnit_ModifyBoneTransforms.h
+++ b/Engine/Plugins/Experimental/ControlRig/Source/ControlRig/Private/Units/Highlevel/Hierarchy/RigUnit_ModifyBoneTransforms.h
@@ -44,13 +44,13 @@ struct FRigUnit_ModifyBoneTransforms_PerBone
 	/**
 	 * The name of the Bone to set the transform for.
 	 */
-	UPROPERTY(meta = (Input, BoneName, Constant))
+	UPROPERTY(EditAnywhere, meta = (Input, CustomWidget = "BoneName", Constant), Category = FRigUnit_ModifyBoneTransforms_PerBone)
 	FName Bone;
 
 	/**
 	 * The transform value to set for the given Bone.
 	 */
-	UPROPERTY(meta = (Input))
+	UPROPERTY(EditAnywhere, meta = (Input), Category = FRigUnit_ModifyBoneTransforms_PerBone)
 	FTransform Transform;
 };
 
@@ -64,7 +64,7 @@ struct FRigUnit_ModifyBoneTransforms_WorkData
 };
 
 /**
- * SetBoneTransform is used to perform a change in the hierarchy by setting a single bone's transform.
+ * ModifyBonetransforms is used to perform a change in the hierarchy by setting one or more bones' transforms.
  */
 USTRUCT(meta=(DisplayName="Modify Transforms", Category="Hierarchy", DocumentationPolicy="Strict", Keywords = "ModifyBone"))
 struct FRigUnit_ModifyBoneTransforms : public FRigUnit_HighlevelBaseMutable
@@ -76,7 +76,46 @@ struct FRigUnit_ModifyBoneTransforms : public FRigUnit_HighlevelBaseMutable
 		, WeightMinimum(0.f)
 		, WeightMaximum(1.f)
 		, Mode(EControlRigModifyBoneMode::AdditiveLocal)
-	{}
+	{
+		BoneToModify.Add(FRigUnit_ModifyBoneTransforms_PerBone());
+	}
+
+	virtual FName DetermineSpaceForPin(const FString& InPinPath, void* InUserContext) const override
+	{
+		if (InPinPath.StartsWith(TEXT("BoneToModify")))
+		{
+			int32 Index = INDEX_NONE;
+			FString Left, Middle, Right;
+			if (InPinPath.Replace(TEXT("["), TEXT(".")).Split(TEXT("."), &Left, &Middle))
+			{
+				if (Middle.Replace(TEXT("]"), TEXT(".")).Split(TEXT("."), &Left, &Right))
+				{
+					Index = FCString::Atoi(*Left);
+				}
+			}
+
+			if (BoneToModify.IsValidIndex(Index))
+			{
+				if (Mode == EControlRigModifyBoneMode::AdditiveLocal)
+				{
+					return BoneToModify[Index].Bone;
+				}
+
+				if (Mode == EControlRigModifyBoneMode::OverrideLocal)
+				{
+					if (const FRigHierarchyContainer* Container = (const FRigHierarchyContainer*)InUserContext)
+					{
+						int32 BoneIndex = Container->BoneHierarchy.GetIndex(BoneToModify[Index].Bone);
+						if (BoneIndex != INDEX_NONE)
+						{
+							return Container->BoneHierarchy[BoneIndex].ParentName;
+						}
+					}
+				}
+			}
+		}
+		return NAME_None;
+	}
 
 	RIGVM_METHOD()
 	virtual void Execute(const FRigUnitContext& Context) override;
diff --git a/Engine/Plugins/Experimental/ControlRig/Source/ControlRig/Private/Units/Highlevel/Hierarchy/RigUnit_SlideChain.h b/Engine/Plugins/Experimental/ControlRig/Source/ControlRig/Private/Units/Highlevel/Hierarchy/RigUnit_SlideChain.h
--- a/Engine/Plugins/Experimental/ControlRig/Source/ControlRig/Private/Units/Highlevel/Hierarchy/RigUnit_SlideChain.h
+++ b/Engine/Plugins/Experimental/ControlRig/Source/ControlRig/Private/Units/Highlevel/Hierarchy/RigUnit_SlideChain.h
@@ -53,13 +53,13 @@ struct FRigUnit_SlideChain: public FRigUnit_HighlevelBaseMutable
 	/** 
 	 * The name of the first bone to slide
 	 */
-	UPROPERTY(meta = (Input, Constant, BoneName))
+	UPROPERTY(meta = (Input, Constant, CustomWidget = "BoneName"))
 	FName StartBone;
 
 	/** 
 	 * The name of the last bone to slide
 	 */
-	UPROPERTY(meta = (Input, Constant, BoneName))
+	UPROPERTY(meta = (Input, Constant, CustomWidget = "BoneName"))
 	FName EndBone;
 
 	/** 
@@ -73,7 +73,7 @@ struct FRigUnit_SlideChain: public FRigUnit_HighlevelBaseMutable
 	 * of this bone will be recalculated based on their local transforms.
 	 * Note: This is computationally more expensive than turning it off.
 	 */
-	UPROPERTY(meta = (Input))
+	UPROPERTY(meta = (Input, Constant))
 	bool bPropagateToChildren;
 
 	UPROPERTY(transient)
diff --git a/Engine/Plugins/Experimental/ControlRig/Source/ControlRig/Private/Units/Highlevel/Hierarchy/RigUnit_SpringIK.h b/Engine/Plugins/Experimental/ControlRig/Source/ControlRig/Private/Units/Highlevel/Hierarchy/RigUnit_SpringIK.h
--- a/Engine/Plugins/Experimental/ControlRig/Source/ControlRig/Private/Units/Highlevel/Hierarchy/RigUnit_SpringIK.h
+++ b/Engine/Plugins/Experimental/ControlRig/Source/ControlRig/Private/Units/Highlevel/Hierarchy/RigUnit_SpringIK.h
@@ -22,25 +22,25 @@ struct FRigUnit_SpringIK_DebugSettings
 	/**
 	 * If enabled debug information will be drawn 
 	 */
-	UPROPERTY(meta = (Input))
+	UPROPERTY(EditAnywhere, meta = (Input), Category = "DebugSettings")
 	bool bEnabled;
 
 	/**
 	 * The size of the debug drawing information
 	 */
-	UPROPERTY(meta = (Input))
+	UPROPERTY(EditAnywhere, meta = (Input, EditCondition = "bEnabled"), Category = "DebugSettings")
 	float Scale;
 
 	/**
 	 * The color to use for debug drawing
 	 */
-	UPROPERTY(meta = (Input))
+	UPROPERTY(EditAnywhere, meta = (Input, EditCondition = "bEnabled"), Category = "DebugSettings")
 	FLinearColor Color;
 
 	/**
 	 * The offset at which to draw the debug information in the world
 	 */
-	UPROPERTY(meta = (Input))
+	UPROPERTY(EditAnywhere, meta = (Input, EditCondition = "bEnabled"), Category = "DebugSettings")
 	FTransform WorldOffset;
 };
 
@@ -96,19 +96,29 @@ struct FRigUnit_SpringIK : public FRigUnit_HighlevelBaseMutable
 		DebugSettings = FRigUnit_SpringIK_DebugSettings();
 	}
 
+
+	virtual FName DetermineSpaceForPin(const FString& InPinPath, void* InUserContext) const override
+	{
+		if (InPinPath.StartsWith(TEXT("PoleVector")))
+		{
+			return PoleVectorSpace;
+		}
+		return NAME_None;
+	}
+
 	RIGVM_METHOD()
 	virtual void Execute(const FRigUnitContext& Context) override;
 
 	/** 
 	 * The name of the first bone to solve
 	 */
-	UPROPERTY(meta = (Input, Constant, BoneName))
+	UPROPERTY(meta = (Input, Constant, CustomWidget = "BoneName"))
 	FName StartBone;
 
 	/**
 	 * The name of the second bone to solve
 	 */
-	UPROPERTY(meta = (Input, Constant, BoneName))
+	UPROPERTY(meta = (Input, Constant, CustomWidget = "BoneName"))
 	FName EndBone;
 
 	/**
@@ -171,7 +181,7 @@ struct FRigUnit_SpringIK : public FRigUnit_HighlevelBaseMutable
 	/**
 	 * The space in which the pole vector is expressed
 	 */
-	UPROPERTY(meta = (Input, Constant, BoneName))
+	UPROPERTY(meta = (Input, Constant, CustomWidget = "BoneName"))
 	FName PoleVectorSpace;
 
 	/**
@@ -211,7 +221,7 @@ struct FRigUnit_SpringIK : public FRigUnit_HighlevelBaseMutable
 	 * of this bone will be recalculated based on their local transforms.
 	 * Note: This is computationally more expensive than turning it off.
 	 */
-	UPROPERTY(meta = (Input))
+	UPROPERTY(meta = (Input, Constant))
 	bool bPropagateToChildren;
 
 	/** The debug setting for the node */
diff --git a/Engine/Plugins/Experimental/ControlRig/Source/ControlRig/Private/Units/Highlevel/Hierarchy/RigUnit_TransformConstraint.h b/Engine/Plugins/Experimental/ControlRig/Source/ControlRig/Private/Units/Highlevel/Hierarchy/RigUnit_TransformConstraint.h
--- a/Engine/Plugins/Experimental/ControlRig/Source/ControlRig/Private/Units/Highlevel/Hierarchy/RigUnit_TransformConstraint.h
+++ b/Engine/Plugins/Experimental/ControlRig/Source/ControlRig/Private/Units/Highlevel/Hierarchy/RigUnit_TransformConstraint.h
@@ -58,10 +58,34 @@ struct FRigUnit_TransformConstraint : public FRigUnit_HighlevelBaseMutable
 		: BaseTransformSpace(ETransformSpaceMode::GlobalSpace)
 	{}
 
+	virtual FName DetermineSpaceForPin(const FString& InPinPath, void* InUserContext) const override
+	{
+		if (InPinPath.StartsWith(TEXT("Targets")))
+		{
+			if (BaseTransformSpace == ETransformSpaceMode::BaseJoint)
+			{
+				return BaseBone;
+			}
+
+			if (BaseTransformSpace == ETransformSpaceMode::LocalSpace)
+			{
+				if (const FRigHierarchyContainer* Container = (const FRigHierarchyContainer*)InUserContext)
+				{
+					int32 BoneIndex = Container->BoneHierarchy.GetIndex(Bone);
+					if (BoneIndex != INDEX_NONE)
+					{
+						return Container->BoneHierarchy[BoneIndex].ParentName;
+					}
+				}
+			}
+		}
+		return NAME_None;
+	}
+
 	RIGVM_METHOD()
 	virtual void Execute(const FRigUnitContext& Context) override;
 
-	UPROPERTY(EditAnywhere, Category = "Constraint", meta = (Input, Constant, BoneName))
+	UPROPERTY(EditAnywhere, Category = "Constraint", meta = (Input, Constant, CustomWidget = "BoneName"))
 	FName Bone;
 
 	UPROPERTY(EditAnywhere, Category = "Constraint", meta = (Input))
@@ -72,7 +96,7 @@ struct FRigUnit_TransformConstraint : public FRigUnit_HighlevelBaseMutable
 	FTransform BaseTransform;
 
 	// Transform op option. Use if ETransformSpace is BaseJoint
-	UPROPERTY(EditAnywhere, Category = "Constraint", meta = (Input, Constant, BoneName))
+	UPROPERTY(EditAnywhere, Category = "Constraint", meta = (Input, Constant, CustomWidget = "BoneName"))
 	FName BaseBone;
 
 	UPROPERTY(EditAnywhere, Category = "Constraint", meta = (Input, ExpandByDefault, DefaultArraySize = 1))
diff --git a/Engine/Plugins/Experimental/ControlRig/Source/ControlRig/Private/Units/Highlevel/Hierarchy/RigUnit_TwoBoneIKSimple.h b/Engine/Plugins/Experimental/ControlRig/Source/ControlRig/Private/Units/Highlevel/Hierarchy/RigUnit_TwoBoneIKSimple.h
--- a/Engine/Plugins/Experimental/ControlRig/Source/ControlRig/Private/Units/Highlevel/Hierarchy/RigUnit_TwoBoneIKSimple.h
+++ b/Engine/Plugins/Experimental/ControlRig/Source/ControlRig/Private/Units/Highlevel/Hierarchy/RigUnit_TwoBoneIKSimple.h
@@ -20,19 +20,19 @@ struct FRigUnit_TwoBoneIKSimple_DebugSettings
 	/**
 	 * If enabled debug information will be drawn 
 	 */
-	UPROPERTY(meta = (Input))
+	UPROPERTY(EditAnywhere, meta = (Input), Category = "DebugSettings")
 	bool bEnabled;
 
 	/**
 	 * The size of the debug drawing information
 	 */
-	UPROPERTY(meta = (Input))
+	UPROPERTY(EditAnywhere, meta = (Input, EditCondition = "bEnabled"), Category = "DebugSettings")
 	float Scale;
 
 	/**
 	 * The offset at which to draw the debug information in the world
 	 */
-	UPROPERTY(meta = (Input))
+	UPROPERTY(EditAnywhere, meta = (Input, EditCondition = "bEnabled"), Category = "DebugSettings")
 	FTransform WorldOffset;
 };
 
@@ -68,25 +68,34 @@ struct FRigUnit_TwoBoneIKSimple : public FRigUnit_HighlevelBaseMutable
 		PoleVectorSpaceIndex = INDEX_NONE;
 	}
 
+	virtual FName DetermineSpaceForPin(const FString& InPinPath, void* InUserContext) const override
+	{
+		if (InPinPath.StartsWith(TEXT("PoleVector")))
+		{
+			return PoleVectorSpace;
+		}
+		return NAME_None;
+	}
+
 	RIGVM_METHOD()
 	virtual void Execute(const FRigUnitContext& Context) override;
 
 	/** 
 	 * The name of first bone
 	 */
-	UPROPERTY(meta = (Input, Constant, BoneName))
+	UPROPERTY(meta = (Input, Constant, CustomWidget = "BoneName"))
 	FName BoneA;
 
 	/** 
 	 * The name of second bone
 	 */
-	UPROPERTY(meta = (Input, Constant, BoneName))
+	UPROPERTY(meta = (Input, Constant, CustomWidget = "BoneName"))
 	FName BoneB;
 
 	/** 
 	 * The name of the effector bone (if exists)
 	 */
-	UPROPERTY(meta = (Input, Constant, BoneName))
+	UPROPERTY(meta = (Input, Constant, CustomWidget = "BoneName"))
 	FName EffectorBone;
 
 	/** 
@@ -129,7 +138,7 @@ struct FRigUnit_TwoBoneIKSimple : public FRigUnit_HighlevelBaseMutable
 	/** 
 	 * The space in which the pole vector is expressed
 	 */
-	UPROPERTY(meta = (Input, Constant, BoneName))
+	UPROPERTY(meta = (Input, Constant, CustomWidget = "BoneName"))
 	FName PoleVectorSpace;
 
 	/**
@@ -175,7 +184,7 @@ struct FRigUnit_TwoBoneIKSimple : public FRigUnit_HighlevelBaseMutable
 	 * of this bone will be recalculated based on their local transforms.
 	 * Note: This is computationally more expensive than turning it off.
 	 */
-	UPROPERTY(meta = (Input))
+	UPROPERTY(meta = (Input, Constant))
 	bool bPropagateToChildren;
 
 	/** 
diff --git a/Engine/Plugins/Experimental/ControlRig/Source/ControlRig/Private/Units/Math/RigUnit_MathColor.cpp b/Engine/Plugins/Experimental/ControlRig/Source/ControlRig/Private/Units/Math/RigUnit_MathColor.cpp
new file mode 100644
--- /dev/null
+++ b/Engine/Plugins/Experimental/ControlRig/Source/ControlRig/Private/Units/Math/RigUnit_MathColor.cpp
@@ -0,0 +1,35 @@
+// Copyright Epic Games, Inc. All Rights Reserved.
+
+#include "Units/Math/RigUnit_MathColor.h"
+#include "Units/RigUnitContext.h"
+#include "Math/ControlRigMathLibrary.h"
+
+FRigUnit_MathColorFromFloat_Execute()
+{
+    DECLARE_SCOPE_HIERARCHICAL_COUNTER_RIGUNIT()
+	Result = FLinearColor(Value, Value, Value);
+}
+
+FRigUnit_MathColorAdd_Execute()
+{
+    DECLARE_SCOPE_HIERARCHICAL_COUNTER_RIGUNIT()
+	Result = A + B;
+}
+
+FRigUnit_MathColorSub_Execute()
+{
+    DECLARE_SCOPE_HIERARCHICAL_COUNTER_RIGUNIT()
+	Result = A - B;
+}
+
+FRigUnit_MathColorMul_Execute()
+{
+    DECLARE_SCOPE_HIERARCHICAL_COUNTER_RIGUNIT()
+	Result = A * B;
+}
+
+FRigUnit_MathColorLerp_Execute()
+{
+    DECLARE_SCOPE_HIERARCHICAL_COUNTER_RIGUNIT()
+	Result = FMath::Lerp<FLinearColor>(A, B, T);
+}
diff --git a/Engine/Plugins/Experimental/ControlRig/Source/ControlRig/Private/Units/Math/RigUnit_MathColor.h b/Engine/Plugins/Experimental/ControlRig/Source/ControlRig/Private/Units/Math/RigUnit_MathColor.h
new file mode 100644
--- /dev/null
+++ b/Engine/Plugins/Experimental/ControlRig/Source/ControlRig/Private/Units/Math/RigUnit_MathColor.h
@@ -0,0 +1,129 @@
+// Copyright Epic Games, Inc. All Rights Reserved.
+
+#pragma once
+
+#include "Units/Math/RigUnit_MathBase.h"
+#include "Math/ControlRigMathLibrary.h"
+#include "RigUnit_MathColor.generated.h"
+
+USTRUCT(meta=(Abstract, Category="Math|Color", MenuDescSuffix="(Color)"))
+struct FRigUnit_MathColorBase : public FRigUnit_MathBase
+{
+	GENERATED_BODY()
+};
+
+USTRUCT(meta=(Abstract))
+struct FRigUnit_MathColorBinaryOp : public FRigUnit_MathColorBase
+{
+	GENERATED_BODY()
+
+	FRigUnit_MathColorBinaryOp()
+	{
+		A = B = Result = FLinearColor::Black;
+	}
+
+	UPROPERTY(meta=(Input))
+	FLinearColor A;
+
+	UPROPERTY(meta=(Input))
+	FLinearColor B;
+
+	UPROPERTY(meta=(Output))
+	FLinearColor Result;
+};
+
+/**
+ * Makes a vector from a single float
+ */
+USTRUCT(meta=(DisplayName="From Float", PrototypeName="FromFloat", Keywords="Make,Construct"))
+struct FRigUnit_MathColorFromFloat : public FRigUnit_MathColorBase
+{
+	GENERATED_BODY()
+	RIGVM_METHOD()
+	virtual void Execute(const FRigUnitContext& Context) override;
+
+	FRigUnit_MathColorFromFloat()
+	{
+		Value = 0.f;
+		Result = FLinearColor::Black;
+	}
+
+	UPROPERTY(meta=(Input))
+	float Value;
+
+	UPROPERTY(meta=(Output))
+	FLinearColor Result;
+};
+
+/**
+ * Returns the sum of the two values
+ */
+USTRUCT(meta=(DisplayName="Add", PrototypeName="Add", Keywords="Sum,+"))
+struct FRigUnit_MathColorAdd : public FRigUnit_MathColorBinaryOp
+{
+	GENERATED_BODY()
+
+	RIGVM_METHOD()
+	virtual void Execute(const FRigUnitContext& Context) override;
+};
+
+/**
+ * Returns the difference of the two values
+ */
+USTRUCT(meta=(DisplayName="Subtract", PrototypeName="Subtract", Keywords="-"))
+struct FRigUnit_MathColorSub : public FRigUnit_MathColorBinaryOp
+{
+	GENERATED_BODY()
+
+	RIGVM_METHOD()
+	virtual void Execute(const FRigUnitContext& Context) override;
+};
+
+/**
+ * Returns the product of the two values
+ */
+USTRUCT(meta=(DisplayName="Multiply", PrototypeName="Multiply", Keywords="Product,*"))
+struct FRigUnit_MathColorMul : public FRigUnit_MathColorBinaryOp
+{
+	GENERATED_BODY()
+
+	FRigUnit_MathColorMul()
+	{
+		A = B = FLinearColor::White;
+	}
+
+	RIGVM_METHOD()
+	virtual void Execute(const FRigUnitContext& Context) override;
+};
+
+/**
+ * Linearly interpolates between A and B using the ratio T
+ */
+USTRUCT(meta=(DisplayName="Interpolate", PrototypeName="Interpolate", Keywords="Lerp,Mix,Blend"))
+struct FRigUnit_MathColorLerp : public FRigUnit_MathColorBase
+{
+	GENERATED_BODY()
+
+	FRigUnit_MathColorLerp()
+	{
+		A = Result = FLinearColor::Black;
+		B = FLinearColor::White;
+		T = 0.f;
+	}
+
+	RIGVM_METHOD()
+	virtual void Execute(const FRigUnitContext& Context) override;
+
+	UPROPERTY(meta=(Input))
+	FLinearColor A;
+
+	UPROPERTY(meta=(Input))
+	FLinearColor B;
+
+	UPROPERTY(meta=(Input))
+	float T;
+
+	UPROPERTY(meta=(Output))
+	FLinearColor Result;
+};
+
diff --git a/Engine/Plugins/Experimental/ControlRig/Source/ControlRig/Private/Units/Math/RigUnit_MathFloat.cpp b/Engine/Plugins/Experimental/ControlRig/Source/ControlRig/Private/Units/Math/RigUnit_MathFloat.cpp
--- a/Engine/Plugins/Experimental/ControlRig/Source/ControlRig/Private/Units/Math/RigUnit_MathFloat.cpp
+++ b/Engine/Plugins/Experimental/ControlRig/Source/ControlRig/Private/Units/Math/RigUnit_MathFloat.cpp
@@ -3,6 +3,21 @@
 #include "Units/Math/RigUnit_MathFloat.h"
 #include "Units/RigUnitContext.h"
 
+FRigUnit_MathFloatConstPi_Execute()
+{
+	Value = PI;
+}
+
+FRigUnit_MathFloatConstHalfPi_Execute()
+{
+	Value = HALF_PI;
+}
+
+FRigUnit_MathFloatConstTwoPi_Execute()
+{
+	Value = PI * 2.f;
+}
+
 FRigUnit_MathFloatAdd_Execute()
 {
     DECLARE_SCOPE_HIERARCHICAL_COUNTER_RIGUNIT()
diff --git a/Engine/Plugins/Experimental/ControlRig/Source/ControlRig/Private/Units/Math/RigUnit_MathFloat.h b/Engine/Plugins/Experimental/ControlRig/Source/ControlRig/Private/Units/Math/RigUnit_MathFloat.h
--- a/Engine/Plugins/Experimental/ControlRig/Source/ControlRig/Private/Units/Math/RigUnit_MathFloat.h
+++ b/Engine/Plugins/Experimental/ControlRig/Source/ControlRig/Private/Units/Math/RigUnit_MathFloat.h
@@ -21,7 +21,7 @@ struct FRigUnit_MathFloatConstant : public FRigUnit_MathFloatBase
 		Value = 0.f;
 	}
 
-	UPROPERTY(meta=(Output, Constant))
+	UPROPERTY(meta=(Output))
 	float Value;
 };
 
@@ -70,10 +70,8 @@ struct FRigUnit_MathFloatConstPi : public FRigUnit_MathFloatConstant
 {
 	GENERATED_BODY()
 	
-	FRigUnit_MathFloatConstPi()
-	{
-		Value = PI;
-	}
+	RIGVM_METHOD()
+	virtual void Execute(const FRigUnitContext& Context) override;
 };
 
 /**
@@ -84,10 +82,8 @@ struct FRigUnit_MathFloatConstHalfPi : public FRigUnit_MathFloatConstant
 {
 	GENERATED_BODY()
 
-	FRigUnit_MathFloatConstHalfPi()
-	{
-		Value = HALF_PI;
-	}
+	RIGVM_METHOD()
+	virtual void Execute(const FRigUnitContext& Context) override;
 };
 
 /**
@@ -97,11 +93,9 @@ USTRUCT(meta=(DisplayName="Two Pi"))
 struct FRigUnit_MathFloatConstTwoPi : public FRigUnit_MathFloatConstant
 {
 	GENERATED_BODY()
-	
-	FRigUnit_MathFloatConstTwoPi()
-	{
-		Value = PI * 2.f;
-	}
+
+	RIGVM_METHOD()
+	virtual void Execute(const FRigUnitContext& Context) override;
 };
 
 /**
diff --git a/Engine/Plugins/Experimental/ControlRig/Source/ControlRig/Private/Units/Math/RigUnit_MathTransform.cpp b/Engine/Plugins/Experimental/ControlRig/Source/ControlRig/Private/Units/Math/RigUnit_MathTransform.cpp
--- a/Engine/Plugins/Experimental/ControlRig/Source/ControlRig/Private/Units/Math/RigUnit_MathTransform.cpp
+++ b/Engine/Plugins/Experimental/ControlRig/Source/ControlRig/Private/Units/Math/RigUnit_MathTransform.cpp
@@ -80,7 +80,7 @@ FRigUnit_MathTransformClampSpatially_Execute()
 {
 	DECLARE_SCOPE_HIERARCHICAL_COUNTER_RIGUNIT()
 	FVector Position;
-	FRigUnit_MathVectorClampSpatially::StaticExecute(Value.GetTranslation(), Axis, Type, Minimum, Maximum, Space, bDrawDebug, DebugColor, DebugThickness, Position, RigUnitName, RigUnitStructName, ExecutionType, Context);
+	FRigUnit_MathVectorClampSpatially::StaticExecute(RigVMOperatorName, RigVMOperatorIndex, Value.GetTranslation(), Axis, Type, Minimum, Maximum, Space, bDrawDebug, DebugColor, DebugThickness, Position, Context);
 	Result = Value;
 	Result.SetTranslation(Position);
 }
\ No newline at end of file
diff --git a/Engine/Plugins/Experimental/ControlRig/Source/ControlRig/Private/Units/RigUnit.cpp b/Engine/Plugins/Experimental/ControlRig/Source/ControlRig/Private/Units/RigUnit.cpp
--- a/Engine/Plugins/Experimental/ControlRig/Source/ControlRig/Private/Units/RigUnit.cpp
+++ b/Engine/Plugins/Experimental/ControlRig/Source/ControlRig/Private/Units/RigUnit.cpp
@@ -4,16 +4,3 @@
 
 DEFINE_LOG_CATEGORY_STATIC(LogRigUnit, Log, All);
 
-// this will have to change in the future and move to editor, I assume the errors will be saved in the rig unit and it will print fromthe editor module
-namespace UnitLogHelpers
-{
-	void PrintMissingHierarchy(const FName& InputName)
-	{
-		//UE_LOG(LogRigUnit, Warning, TEXT("%s: Input Hierarch Link is missing"), *InputName.ToString());
-	}
-
-	void PrintUnimplemented(const FName& InputName)
-	{
-		UE_LOG(LogRigUnit, Warning, TEXT("%s: Not implemented"), *InputName.ToString());
-	}
-}
diff --git a/Engine/Plugins/Experimental/ControlRig/Source/ControlRig/Private/Units/Simulation/RigUnit_AlphaInterp.cpp b/Engine/Plugins/Experimental/ControlRig/Source/ControlRig/Private/Units/Simulation/RigUnit_AlphaInterp.cpp
--- a/Engine/Plugins/Experimental/ControlRig/Source/ControlRig/Private/Units/Simulation/RigUnit_AlphaInterp.cpp
+++ b/Engine/Plugins/Experimental/ControlRig/Source/ControlRig/Private/Units/Simulation/RigUnit_AlphaInterp.cpp
@@ -12,18 +12,91 @@ FRigUnit_AlphaInterp_Execute()
 		ScaleBiasClamp.bMapRange = bMapRange;
 		ScaleBiasClamp.bClampResult = bClampResult;
 		ScaleBiasClamp.bInterpResult = bInterpResult;
+	}
+	else
+	{
 		ScaleBiasClamp.InRange = InRange;
 		ScaleBiasClamp.OutRange = OutRange;
 		ScaleBiasClamp.ClampMin = ClampMin;
 		ScaleBiasClamp.ClampMax = ClampMax;
+		ScaleBiasClamp.Scale = Scale;
+		ScaleBiasClamp.Bias = Bias;
+		ScaleBiasClamp.InterpSpeedIncreasing = InterpSpeedIncreasing;
+		ScaleBiasClamp.InterpSpeedDecreasing = InterpSpeedDecreasing;
+
+		Result = ScaleBiasClamp.ApplyTo(Value, Context.DeltaTime);
+	}
+}
+
+FString FRigUnit_AlphaInterp::ProcessPinLabelForInjection(const FString& InLabel) const
+{
+	FString Formula;
+	if (bMapRange)
+	{
+		Formula += FString::Printf(TEXT(" Map(%.02f, %.02f, %.02f, %.02f)"), InRange.Min, InRange.Max, OutRange.Min, OutRange.Max);
+	}
+	if (bInterpResult)
+	{
+		Formula += FString::Printf(TEXT(" Interp(%.02f, %.02f)"), InterpSpeedIncreasing, InterpSpeedDecreasing);
+	}
+	if (bClampResult)
+	{
+		Formula += FString::Printf(TEXT(" Clamp(%.02f, %.02f)"), ClampMin, ClampMax);
+	}
+
+	if (Formula.IsEmpty())
+	{
+		return InLabel;
+	}
+	return FString::Printf(TEXT("%s: %s"), *InLabel, *Formula);
+}
+
+FRigUnit_AlphaInterpVector_Execute()
+{
+    DECLARE_SCOPE_HIERARCHICAL_COUNTER_RIGUNIT()
+	if (Context.State == EControlRigState::Init)
+	{
+		ScaleBiasClamp.Reinitialize();
+		ScaleBiasClamp.bMapRange = bMapRange;
+		ScaleBiasClamp.bClampResult = bClampResult;
+		ScaleBiasClamp.bInterpResult = bInterpResult;
 	}
 	else
 	{
+		ScaleBiasClamp.InRange = InRange;
+		ScaleBiasClamp.OutRange = OutRange;
+		ScaleBiasClamp.ClampMin = ClampMin;
+		ScaleBiasClamp.ClampMax = ClampMax;
 		ScaleBiasClamp.Scale = Scale;
 		ScaleBiasClamp.Bias = Bias;
 		ScaleBiasClamp.InterpSpeedIncreasing = InterpSpeedIncreasing;
 		ScaleBiasClamp.InterpSpeedDecreasing = InterpSpeedDecreasing;
 
-		Result = ScaleBiasClamp.ApplyTo(Value, Context.DeltaTime);
+		Result.X = ScaleBiasClamp.ApplyTo(Value.X, Context.DeltaTime);
+		Result.Y = ScaleBiasClamp.ApplyTo(Value.Y, Context.DeltaTime);
+		Result.Z = ScaleBiasClamp.ApplyTo(Value.Z, Context.DeltaTime);
 	}
 }
+
+FString FRigUnit_AlphaInterpVector::ProcessPinLabelForInjection(const FString& InLabel) const
+{
+	FString Formula;
+	if (bMapRange)
+	{
+		Formula += FString::Printf(TEXT(" Map(%.02f, %.02f, %.02f, %.02f)"), InRange.Min, InRange.Max, OutRange.Min, OutRange.Max);
+	}
+	if (bInterpResult)
+	{
+		Formula += FString::Printf(TEXT(" Interp(%.02f, %.02f)"), InterpSpeedIncreasing, InterpSpeedDecreasing);
+	}
+	if (bClampResult)
+	{
+		Formula += FString::Printf(TEXT(" Clamp(%.02f, %.02f)"), ClampMin, ClampMax);
+	}
+
+	if (Formula.IsEmpty())
+	{
+		return InLabel;
+	}
+	return FString::Printf(TEXT("%s: %s"), *InLabel, *Formula);
+}
\ No newline at end of file
diff --git a/Engine/Plugins/Experimental/ControlRig/Source/ControlRig/Private/Units/Simulation/RigUnit_AlphaInterp.h b/Engine/Plugins/Experimental/ControlRig/Source/ControlRig/Private/Units/Simulation/RigUnit_AlphaInterp.h
--- a/Engine/Plugins/Experimental/ControlRig/Source/ControlRig/Private/Units/Simulation/RigUnit_AlphaInterp.h
+++ b/Engine/Plugins/Experimental/ControlRig/Source/ControlRig/Private/Units/Simulation/RigUnit_AlphaInterp.h
@@ -2,14 +2,14 @@
 
 #pragma once
 
-#include "Units/Simulation/RigUnit_SimBase.h"
+#include "RigUnit_SimBase.h"
 #include "Animation/InputScaleBias.h"
 #include "RigUnit_AlphaInterp.generated.h"
 
 /**
  * Adds a value over time over and over again
  */
-USTRUCT(meta=(DisplayName="Alpha Interpolate", Keywords="Alpha,Lerp,LinearInterpolate"))
+USTRUCT(meta=(DisplayName="Interpolate", Keywords="Alpha,Lerp,LinearInterpolate", PrototypeName = "AlphaInterp"))
 struct FRigUnit_AlphaInterp : public FRigUnit_SimBase
 {
 	GENERATED_BODY()
@@ -34,44 +34,119 @@ struct FRigUnit_AlphaInterp : public FRigUnit_SimBase
 	RIGVM_METHOD()
 	virtual void Execute(const FRigUnitContext& Context) override;
 
+	virtual FString ProcessPinLabelForInjection(const FString& InLabel) const override;
+
 	UPROPERTY(meta=(Input))
 	float Value;
 
+	UPROPERTY(meta = (Input))
+	float Scale;
+
+	UPROPERTY(meta = (Input))
+	float Bias;
+
 	UPROPERTY(meta=(Input, Constant))
 	bool bMapRange;
 
-	UPROPERTY(meta=(Input, Constant))
+	UPROPERTY(meta=(Input, EditCondition = "bMapRange"))
+	FInputRange InRange;
+
+	UPROPERTY(meta=(Input, EditCondition = "bMapRange"))
+	FInputRange OutRange;
+
+	UPROPERTY(meta = (Input, Constant))
 	bool bClampResult;
 
-	UPROPERTY(meta=(Input, Constant))
+	UPROPERTY(meta=(Input, EditCondition = "bClampResult"))
+	float ClampMin;
+
+	UPROPERTY(meta=(Input, EditCondition = "bClampResult"))
+	float ClampMax;
+
+	UPROPERTY(meta = (Input, Constant))
 	bool bInterpResult;
 
-	UPROPERTY(meta=(Input, Constant))
-	FInputRange InRange;
+	UPROPERTY(meta=(Input, EditCondition = "bInterpResult"))
+	float InterpSpeedIncreasing;
 
-	UPROPERTY(meta=(Input, Constant))
-	FInputRange OutRange;
+	UPROPERTY(meta=(Input, EditCondition = "bInterpResult"))
+	float InterpSpeedDecreasing;
+
+	UPROPERTY(meta=(Output))
+	float Result;
+
+	UPROPERTY()
+	FInputScaleBiasClamp ScaleBiasClamp;
+};
+
+/**
+ * Adds a value over time over and over again
+ */
+USTRUCT(meta=(DisplayName="Interpolate", Keywords="Alpha,Lerp,LinearInterpolate", PrototypeName = "AlphaInterp", MenuDescSuffix = "(Vector)"))
+struct FRigUnit_AlphaInterpVector : public FRigUnit_SimBase
+{
+	GENERATED_BODY()
+	
+	FRigUnit_AlphaInterpVector()
+	{
+		Value = Result = FVector::ZeroVector;
+		ScaleBiasClamp = FInputScaleBiasClamp();
+		bMapRange = ScaleBiasClamp.bMapRange;
+		bClampResult = ScaleBiasClamp.bClampResult;
+		bInterpResult = ScaleBiasClamp.bInterpResult;
+		InRange = ScaleBiasClamp.InRange;
+		OutRange = ScaleBiasClamp.OutRange;
+		Scale = ScaleBiasClamp.Scale;
+		Bias = ScaleBiasClamp.Bias;
+		ClampMin = ScaleBiasClamp.ClampMin;
+		ClampMax = ScaleBiasClamp.ClampMax;
+		InterpSpeedIncreasing = ScaleBiasClamp.InterpSpeedIncreasing;
+		InterpSpeedDecreasing = ScaleBiasClamp.InterpSpeedDecreasing;
+	}
+
+	RIGVM_METHOD()
+	virtual void Execute(const FRigUnitContext& Context) override;
+
+	virtual FString ProcessPinLabelForInjection(const FString& InLabel) const override;
 
 	UPROPERTY(meta=(Input))
+	FVector Value;
+
+	UPROPERTY(meta = (Input))
 	float Scale;
 
-	UPROPERTY(meta=(Input))
+	UPROPERTY(meta = (Input))
 	float Bias;
 
-	UPROPERTY(meta=(Input, Constant))
+	UPROPERTY(meta = (Input, Constant))
+	bool bMapRange;
+
+	UPROPERTY(meta = (Input, EditCondition = "bMapRange"))
+	FInputRange InRange;
+
+	UPROPERTY(meta = (Input, EditCondition = "bMapRange"))
+	FInputRange OutRange;
+
+	UPROPERTY(meta = (Input, Constant))
+	bool bClampResult;
+
+	UPROPERTY(meta = (Input, EditCondition = "bClampResult"))
 	float ClampMin;
 
-	UPROPERTY(meta=(Input, Constant))
+	UPROPERTY(meta = (Input, EditCondition = "bClampResult"))
 	float ClampMax;
 
-	UPROPERTY(meta=(Input))
+	UPROPERTY(meta = (Input, Constant))
+	bool bInterpResult;
+
+	UPROPERTY(meta = (Input, EditCondition = "bInterpResult"))
 	float InterpSpeedIncreasing;
 
-	UPROPERTY(meta=(Input))
+	UPROPERTY(meta = (Input, EditCondition = "bInterpResult"))
 	float InterpSpeedDecreasing;
 
 	UPROPERTY(meta=(Output))
-	float Result;
+	FVector Result;
 
 	UPROPERTY()
 	FInputScaleBiasClamp ScaleBiasClamp;
diff --git a/Engine/Plugins/Experimental/ControlRig/Source/ControlRig/Private/Units/Simulation/RigUnit_PointSimulation.h b/Engine/Plugins/Experimental/ControlRig/Source/ControlRig/Private/Units/Simulation/RigUnit_PointSimulation.h
--- a/Engine/Plugins/Experimental/ControlRig/Source/ControlRig/Private/Units/Simulation/RigUnit_PointSimulation.h
+++ b/Engine/Plugins/Experimental/ControlRig/Source/ControlRig/Private/Units/Simulation/RigUnit_PointSimulation.h
@@ -29,37 +29,37 @@ struct FRigUnit_PointSimulation_DebugSettings
 	/**
 	 * If enabled debug information will be drawn 
 	 */
-	UPROPERTY(meta = (Input))
+	UPROPERTY(EditAnywhere, meta = (Input), Category = "DebugSettings")
 	bool bEnabled;
 
 	/**
 	 * The size of the debug drawing information
 	 */
-	UPROPERTY(meta = (Input))
+	UPROPERTY(EditAnywhere, meta = (Input, EditCondition = "bEnabled"), Category = "DebugSettings")
 	float Scale;
 
 	/**
      * The size of the debug drawing information
      */
-	UPROPERTY(meta = (Input))
+	UPROPERTY(EditAnywhere, meta = (Input, EditCondition = "bEnabled"), Category = "DebugSettings")
 	float CollisionScale;
 
 	/**
 	 * If set to true points will be drawn as spheres with their sizes reflected
 	 */
-	UPROPERTY(meta = (Input))
+	UPROPERTY(EditAnywhere, meta = (Input, EditCondition = "bEnabled"), Category = "DebugSettings")
 	bool bDrawPointsAsSpheres;
 
 	/**
 	 * The color to use for debug drawing
 	 */
-	UPROPERTY(meta = (Input))
+	UPROPERTY(EditAnywhere, meta = (Input, EditCondition = "bEnabled"), Category = "DebugSettings")
 	FLinearColor Color;
 
 	/**
 	 * The offset at which to draw the debug information in the world
 	 */
-	UPROPERTY(meta = (Input))
+	UPROPERTY(EditAnywhere, meta = (Input, EditCondition = "bEnabled"), Category = "DebugSettings")
 	FTransform WorldOffset;
 };
 
@@ -77,27 +77,27 @@ struct FRigUnit_PointSimulation_BoneTarget
 	/**
 	 * The name of the bone to map
 	 */
-	UPROPERTY(meta = (Input, Constant, BoneName))
+	UPROPERTY(EditAnywhere, meta = (Input, Constant, CustomWidget = "BoneName"), Category = "BoneTarget")
 	FName Bone;
 
 	/**
 	 * The index of the point to use for translation
 	 */
-	UPROPERTY(meta = (Input, Constant))
+	UPROPERTY(EditAnywhere, meta = (Input, Constant), Category = "BoneTarget")
 	int32 TranslationPoint;
 
 	/**
 	 * The index of the point to use for aiming the primary axis.
 	 * Use -1 to indicate that you don't want to aim the bone.
 	 */
-	UPROPERTY(meta = (Input, Constant))
+	UPROPERTY(EditAnywhere, meta = (Input, Constant), Category = "BoneTarget")
 	int32 PrimaryAimPoint;
 
 	/**
 	 * The index of the point to use for aiming the secondary axis.
 	 * Use -1 to indicate that you don't want to aim the bone.
 	 */
-	UPROPERTY(meta = (Input, Constant))
+	UPROPERTY(EditAnywhere, meta = (Input, Constant), Category = "BoneTarget")
 	int32 SecondaryAimPoint;
 };
 
@@ -182,7 +182,7 @@ struct FRigUnit_PointSimulation : public FRigUnit_SimBaseMutable
 	 * of this bone will be recalculated based on their local transforms.
 	 * Note: This is computationally more expensive than turning it off.
 	 */
-	UPROPERTY(meta = (Input))
+	UPROPERTY(meta = (Input, Constant))
 	bool bPropagateToChildren;
 
 	/** The primary axis to use for the aim */
diff --git a/Engine/Plugins/Experimental/ControlRig/Source/ControlRig/Public/AnimNode_ControlRig.h b/Engine/Plugins/Experimental/ControlRig/Source/ControlRig/Public/AnimNode_ControlRig.h
--- a/Engine/Plugins/Experimental/ControlRig/Source/ControlRig/Public/AnimNode_ControlRig.h
+++ b/Engine/Plugins/Experimental/ControlRig/Source/ControlRig/Public/AnimNode_ControlRig.h
@@ -32,6 +32,10 @@ struct CONTROLRIG_API FAnimNode_ControlRig : public FAnimNode_ControlRigBase
 
 	void SetIOMapping(bool bInput, const FName& SourceProperty, const FName& TargetCurve);
 	FName GetIOMapping(bool bInput, const FName& SourceProperty) const;
+
+	virtual void InitializeProperties(const UObject* InSourceInstance, UClass* InTargetClass) override;
+	virtual void PropagateInputProperties(const UObject* InSourceInstance) override;
+
 private:
 
 	/** Cached ControlRig */
@@ -74,15 +78,18 @@ private:
 
 	TMap<FName, FName> InputTypes;
 	TMap<FName, FName> OutputTypes;
+	TArray<uint8*> DestParameters;
 
 #if WITH_EDITOR
 	void OnObjectsReplaced(const TMap<UObject*, UObject*>& OldToNewInstanceMap);
 #endif // WITH_EDITOR
 protected:
-	virtual UClass* GetTargetClass() const override { return *ControlRigClass; } 
+	virtual UClass* GetTargetClass() const override { return *ControlRigClass; }
 	virtual void UpdateInput(UControlRig* InControlRig, const FPoseContext& InOutput) override;
 	virtual void UpdateOutput(UControlRig* InControlRig, FPoseContext& InOutput) override;
+
 public:
+
 	void PostSerialize(const FArchive& Ar);
 
 	friend class UAnimGraphNode_ControlRig;
diff --git a/Engine/Plugins/Experimental/ControlRig/Source/ControlRig/Public/AnimNode_ControlRigBase.h b/Engine/Plugins/Experimental/ControlRig/Source/ControlRig/Public/AnimNode_ControlRigBase.h
--- a/Engine/Plugins/Experimental/ControlRig/Source/ControlRig/Public/AnimNode_ControlRigBase.h
+++ b/Engine/Plugins/Experimental/ControlRig/Source/ControlRig/Public/AnimNode_ControlRigBase.h
@@ -7,6 +7,37 @@
 
 class UControlRig;
 class UNodeMappingContainer;
+
+/** Struct defining the settings to override when driving a control rig */
+USTRUCT()
+struct CONTROLRIG_API FControlRigIOSettings
+{
+	GENERATED_BODY()
+
+	FControlRigIOSettings()
+		: bUpdatePose(true)
+		, bUpdateCurves(true)
+	{}
+
+	FORCEINLINE static FControlRigIOSettings MakeEnabled()
+	{
+		return FControlRigIOSettings();
+	}
+
+	FORCEINLINE static FControlRigIOSettings MakeDisabled()
+	{
+		FControlRigIOSettings Settings;
+		Settings.bUpdatePose = Settings.bUpdateCurves = false;
+		return Settings;
+	}
+
+	UPROPERTY()
+	bool bUpdatePose;
+
+	UPROPERTY()
+	bool bUpdateCurves;
+};
+
 /**
  * Animation node that allows animation ControlRig output to be used in an animation graph
  */
@@ -50,7 +81,10 @@ protected:
 	TWeakObjectPtr<UNodeMappingContainer> NodeMappingContainer;
 
 	UPROPERTY(transient)
-	bool bUpdateInput;
+	FControlRigIOSettings InputSettings;
+
+	UPROPERTY(transient)
+	FControlRigIOSettings OutputSettings;
 
 	UPROPERTY(transient)
 	bool bExecute;
@@ -67,5 +101,6 @@ protected:
 	void ExecuteControlRig(FPoseContext& InOutput);
 
 	friend struct FControlRigSequencerAnimInstanceProxy;
+	friend struct FControlRigLayerInstanceProxy;
 };
 
