COPY from //Dev-Anim to //Dev-Main

diff --git a/Engine/Plugins/Experimental/ControlRig/Source/ControlRig/Public/ControlRig.h b/Engine/Plugins/Experimental/ControlRig/Source/ControlRig/Public/ControlRig.h
--- a/Engine/Plugins/Experimental/ControlRig/Source/ControlRig/Public/ControlRig.h
+++ b/Engine/Plugins/Experimental/ControlRig/Source/ControlRig/Public/ControlRig.h
@@ -14,26 +14,30 @@
 #include "Units/RigUnit.h"
 #include "Units/Control/RigUnit_Control.h"
 #include "Manipulatable/IControlRigManipulatable.h"
+#include "RigVMCore/RigVM.h"
+
+#if WITH_EDITOR
+#include "RigVMModel/RigVMPin.h"
+#endif
+
+#if WITH_EDITOR
+#include "AnimPreviewInstance.h"
+#endif 
+
 #include "ControlRig.generated.h"
 
 class IControlRigObjectBinding;
 class UScriptStruct;
 struct FRigUnit;
-struct FControlRigIOVariable;
+struct FRigControl;
 
 DECLARE_LOG_CATEGORY_EXTERN(LogControlRig, Log, All);
 
-/** Delegate used to optionally gather inputs before evaluating a ControlRig */
-DECLARE_DELEGATE_OneParam(FPreEvaluateGatherInput, UControlRig*);
-DECLARE_DELEGATE_OneParam(FPostEvaluateQueryOutput, UControlRig*);
-
-#define DEBUG_CONTROLRIG_PROPERTYCHANGE !(UE_BUILD_SHIPPING || UE_BUILD_TEST)
-
 /** Runs logic for mapping input data to transforms (the "Rig") */
 UCLASS(Blueprintable, Abstract, editinlinenew)
 class CONTROLRIG_API UControlRig : public UObject, public INodeMappingProviderInterface, public IControlRigManipulatable
 {
-	GENERATED_BODY()
+	GENERATED_UCLASS_BODY()
 
 	friend class UControlRigComponent;
 	friend class SControlRigStackView;
@@ -47,20 +51,18 @@ public:
 	static const FName DisplayNameMetaName;
 	static const FName MenuDescSuffixMetaName;
 	static const FName ShowVariableNameInTitleMetaName;
-	static const FName BoneNameMetaName;
-	static const FName ControlNameMetaName;
-	static const FName SpaceNameMetaName;
-	static const FName CurveNameMetaName;
+	static const FName CustomWidgetMetaName;
 	static const FName ConstantMetaName;
 	static const FName TitleColorMetaName;
 	static const FName NodeColorMetaName;
 	static const FName KeywordsMetaName;
 	static const FName PrototypeNameMetaName;
-	static const FName AnimationInputMetaName;
-	static const FName AnimationOutputMetaName;
 	static const FName ExpandPinByDefaultMetaName;
 	static const FName DefaultArraySizeMetaName;
 
+	static const FName OwnerComponent;
+
+
 private:
 	/** Current delta time */
 	float DeltaTime;
@@ -73,10 +75,6 @@ public:
 	virtual void PostEditUndo() override;
 #endif
 	
-	/** Get the current delta time */
-	UFUNCTION(BlueprintPure, Category = "Animation")
-	float GetDeltaTime() const;
-
 	/** Set the current delta time */
 	void SetDeltaTime(float InDeltaTime);
 
@@ -85,7 +83,7 @@ public:
 	virtual FText GetCategory() const;
 
 	/** Get the tooltip text to display for this node (displayed in graphs and from context menus) */
-	virtual FText GetTooltipText() const;
+	virtual FText GetToolTipText() const;
 #endif
 
 	/** UObject interface */
@@ -97,6 +95,22 @@ public:
 	/** Evaluate at Any Thread */
 	virtual void Evaluate_AnyThread();
 
+	/** input output handling */
+	FORCEINLINE const TArray<FRigVMParameter>& GetParameters() const
+	{
+		return VM->GetParameters();
+	}
+	template<class T>
+	FORCEINLINE T GetParameterValue(const FName& InParameterName)
+	{
+		return VM->GetParameterValue<T>(InParameterName);
+	}
+	template<class T>
+	FORCEINLINE void SetParameterValue(const FName& InParameterName, const T& InValue)
+	{
+		VM->SetParameterValue<T>(InParameterName, InValue);
+	}
+
 	/** Setup bindings to a runtime object (or clear by passing in nullptr). */
 	virtual void SetObjectBinding(TSharedPtr<IControlRigObjectBinding> InObjectBinding) override
 	{
@@ -164,37 +178,8 @@ public:
 	/** Evaluate another animation ControlRig */
 	void SetCurveValue(const int32 CurveIndex, const float CurveValue);
 
-	/* 
-	 * Query input output variables
-	 *
-	 * @param bInput - True if it's input. False if you want to query Output
-	 * @param OutVars - Output array of variables
-	 */
-	void QueryIOVariables(bool bInput, TArray<FControlRigIOVariable>& OutVars) const;
-	/*
-	 * Return true if the PropertyName is in IO
-	 *
-	 * @param bInput - True if it's input. False if you want to query Output
-	 * @return true if it's valid
-	 */
-	bool IsValidIOVariables(bool bInput, const FName& PropertyName) const;
-
-	/*
-	 * Get IO property path
-	 *
-	 * @param bInput - True if it's input. False if you want to query Output
-	 * @param InPropertyPath - Property path to query
-	 * @param OutCachedPath - output of cached path
-	 * @return true if succeed
-	 */
-	bool GetInOutPropertyPath(bool bInput, const FName& InPropertyPath, FCachedPropertyPath& OutCachedPath);
-
 #if WITH_EDITOR
-	// get class name of rig unit that is owned by this rig
-	FName GetRigClassNameFromRigUnit(const FRigUnit* InRigUnit) const;
-	FRigUnit_Control* GetControlRigUnitFromName(const FName& PropertyName);
-	FRigUnit* GetRigUnitFromName(const FName& PropertyName);
-	
+
 	// called after post reinstance when compilng blueprint by Sequencer
 	void PostReinstanceCallback(const UControlRig* Old);
 
@@ -208,12 +193,23 @@ public:
 	UPROPERTY(transient)
 	ERigExecutionType ExecutionType;
 
-	/** Execute the rig unit */
+	/** Execute */
 	void Execute(const EControlRigState State);
 
+	/** ExecuteUnits */
+	virtual void ExecuteUnits(FRigUnitContext& InOutContext);
+
+	/** Requests to perform an init during the next execution */
+	void RequestInit() { bRequiresInitExecution = true;  }
+
+	URigVM* GetVM();
+
 	/** INodeMappingInterface implementation */
 	virtual void GetMappableNodeData(TArray<FName>& OutNames, TArray<FNodeItem>& OutNodeItems) const override;
 
+	/** Data Source Registry Getter */
+	UAnimationDataSourceRegistry* GetDataSourceRegistry() { return DataSourceRegistry; }
+
 	// BEGIN IControlRigManipulatable interface
 	virtual const TArray<FRigSpace>& AvailableSpaces() const override;
 	virtual FRigSpace* FindSpace(const FName& InSpaceName) override;
@@ -225,6 +221,9 @@ public:
 	virtual FRigControlValue GetControlValueFromGlobalTransform(const FName& InControlName, const FTransform& InGlobalTransform) override;
 	virtual bool SetControlSpace(const FName& InControlName, const FName& InSpaceName) override;
 	virtual UControlRigGizmoLibrary* GetGizmoLibrary() const override;
+	virtual void CreateRigControlsForCurveContainer() override;
+
+
 #if WITH_EDITOR
 	virtual void SelectControl(const FName& InControlName, bool bSelect = true) override;
 	virtual bool ClearControlSelection() override;
@@ -233,27 +232,20 @@ public:
 #endif
 	// END IControlRigManipulatable interface
 
+	// Not in IControlRigManipulatable *, but maybe should
+	bool IsCurveControl(const FRigControl* InRigControl) const;
+
 	DECLARE_EVENT_TwoParams(UControlRig, FControlRigExecuteEvent, class UControlRig*, const EControlRigState);
-	FControlRigExecuteEvent& OnInitialized() { return InitializedEvent; }
-	FControlRigExecuteEvent& OnExecuted() { return ExecutedEvent; }
+	FControlRigExecuteEvent& OnInitialized_AnyThread() { return InitializedEvent; }
+	FControlRigExecuteEvent& OnExecuted_AnyThread() { return ExecutedEvent; }
 
 private:
-	UPROPERTY(VisibleDefaultsOnly, Category = "Hierarchy")
-	FRigHierarchyContainer Hierarchy;
 
-#if WITH_EDITORONLY_DATA
-	/** The properties of source accessible <target, source local path> when source -> target
-	 * For example, if you have property RigUnitA.B->RigUnitB.C, this will save as <RigUnitB.C, RigUnitA.B> */
-	UPROPERTY()
-	TMap<FName, FString> AllowSourceAccessProperties;
+	UPROPERTY(VisibleAnywhere, Category = "VM")
+	URigVM* VM;
 
-	/** Cached editor object reference by rig unit */
-	TMap<FRigUnit*, UObject*> RigUnitEditorObjects;
-#endif // WITH_EDITOR
-
-	/** list of operators. */
-	UPROPERTY(Transient)
-	TArray<FControlRigOperator> Operators;
+	UPROPERTY(VisibleDefaultsOnly, Category = "Hierarchy")
+	FRigHierarchyContainer Hierarchy;
 
 	UPROPERTY()
 	TAssetPtr<UControlRigGizmoLibrary> GizmoLibrary;
@@ -265,15 +257,24 @@ private:
 	FControlRigLog* ControlRigLog;
 	bool bEnableControlRigLogging;
 #endif
+
 	// you either go Input or Output, currently if you put it in both place, Output will override
 	UPROPERTY()
-	TMap<FName, FCachedPropertyPath> InputProperties;
+	TMap<FName, FCachedPropertyPath> InputProperties_DEPRECATED;
 
 	UPROPERTY()
-	TMap<FName, FCachedPropertyPath> OutputProperties;
+	TMap<FName, FCachedPropertyPath> OutputProperties_DEPRECATED;
+
+private:
+	// Controls for the container
+	void HandleOnControlModified(IControlRigManipulatable* Subject, const FRigControl& Control, EControlRigSetKey InSetKey);
+
 
 private:
 
+	UPROPERTY()
+	FControlRigDrawContainer DrawContainer;
+
 	/** The draw interface for the units to use */
 	FControlRigDrawInterface* DrawInterface;
 
@@ -281,30 +282,9 @@ private:
 	UPROPERTY(transient)
 	UAnimationDataSourceRegistry* DataSourceRegistry;
 
-#if DEBUG_CONTROLRIG_PROPERTYCHANGE
-	// This is to debug class size when constructed and destroyed to verify match
-	// if this size changes, that implies more problem, where properties have been changed and layout has been modified
-	// also it caches if destructor and property has been chagned
-	// the name can be destroyed but we should make sure we have proper properties size/offset is linked
-	int32 DebugClassSize;
-	TArray<UScriptStruct*> Destructors;
-	struct FPropertyData
-	{
-		int32 Offset;
-		int32 Size;
-		FName PropertyName;
-	};
-	TArray<FPropertyData> PropertyData;
-	void ValidateDebugClassData();
-	void CacheDebugClassData();
-#endif // 	DEBUG_CONTROLRIG_PROPERTYCHANGE
-
-	/** Copy the operators from the generated class */
-	void InstantiateOperatorsFromGeneratedClass();
+	/** Copy the VM from the default object */
+	void InstantiateVMFromCDO();
 	
-	/** Re-resolve operator property paths */
-	void ResolvePropertyPaths();
-
 	/** Broadcasts a notification whenever the controlrig is initialized. */
 	FControlRigExecuteEvent InitializedEvent;
 
@@ -314,15 +294,44 @@ private:
 #if WITH_EDITOR
 	/** Handle a Control Being Selected */
 	void HandleOnControlSelected(FRigHierarchyContainer* InContainer, const FRigElementKey& InKey, bool bSelected);
+
+	/** Update the available controls within the control rig editor */
+	void UpdateAvailableControls();
+
+	/** Remove a transient / temporary control used to interact with a pin */
+	FName AddTransientControl(URigVMPin* InPin, FName InSpaceName = NAME_None);
+
+	/** Sets the value of a transient control based on a pin */
+	bool SetTransientControlValue(URigVMPin* InPin);
+
+	/** Remove a transient / temporary control used to interact with a pin */
+	FName RemoveTransientControl(URigVMPin* InPin);
+
+	FName AddTransientControl(const FRigElementKey& InElement);
+
+	/** Sets the value of a transient control based on a bone */
+	bool SetTransientControlValue(const FRigElementKey& InElement);
+
+	/** Remove a transient / temporary control used to interact with a bone */
+	FName RemoveTransientControl(const FRigElementKey& InElement);
+
+	/** Removes all  transient / temporary control used to interact with pins */
+	void ClearTransientControls();
+
+	TArray<FRigControl> AvailableControlsOverride;
+	TArray<FRigControl> TransientControls;
+	UAnimPreviewInstance* PreviewInstance;
+
 #endif
-	void ResolveInputOutputProperties();
 
 	void InitializeFromCDO();
 
+	static FName GetNameForTransientControl(const FRigElementKey& InElement);
+
+	bool bRequiresInitExecution;
+
 	friend class FControlRigBlueprintCompilerContext;
 	friend struct FRigHierarchyRef;
-	friend class UControlRigEditorLibrary;
-	friend class URigUnitEditor_Base;
 	friend class FControlRigEditor;
 	friend class SRigCurveContainer;
 	friend class SRigHierarchy;
@@ -330,4 +339,5 @@ private:
  	friend class FControlRigEditMode;
 	friend class FControlRigIOHelper;
 	friend class UControlRigBlueprint;
+	friend class UControlRigBlueprintGeneratedClass;
 };
\ No newline at end of file
diff --git a/Engine/Plugins/Experimental/ControlRig/Source/ControlRig/Public/ControlRigBlueprintGeneratedClass.h b/Engine/Plugins/Experimental/ControlRig/Source/ControlRig/Public/ControlRigBlueprintGeneratedClass.h
--- a/Engine/Plugins/Experimental/ControlRig/Source/ControlRig/Public/ControlRigBlueprintGeneratedClass.h
+++ b/Engine/Plugins/Experimental/ControlRig/Source/ControlRig/Public/ControlRigBlueprintGeneratedClass.h
@@ -6,34 +6,19 @@
 #include "UObject/ObjectMacros.h"
 #include "Engine/BlueprintGeneratedClass.h"
 #include "ControlRigDefines.h"
+#include "RigVMCore/RigVM.h"
 #include "ControlRigBlueprintGeneratedClass.generated.h"
 
 UCLASS()
 class CONTROLRIG_API UControlRigBlueprintGeneratedClass : public UBlueprintGeneratedClass
 {
-	GENERATED_BODY()
+	GENERATED_UCLASS_BODY()
 
 public:
-	UControlRigBlueprintGeneratedClass();
-
-	// UStruct interface
-	virtual void Link(FArchive& Ar, bool bRelinkExistingProperties) override;
 
 	// UClass interface
-	virtual void PurgeClass(bool bRecompilingOnLoad) override;
 	virtual uint8* GetPersistentUberGraphFrame(UObject* Obj, UFunction* FuncToCheck) const override;
 
-public:
-
-	/** list of operators. Visible for debug purpose for now */
-	UPROPERTY(VisibleAnywhere, Category = "Links")
-	TArray<FControlRigOperator> Operators;
-
-#if WITH_EDITORONLY_DATA
-	/** The properties of all of the control units */
-	TArray<FStructProperty*> ControlUnitProperties;
-
-	/** The properties of all the rig units */
-	TArray<FStructProperty*> RigUnitProperties;
-#endif
+	// UObject interface
+	void Serialize(FArchive& Ar);
 };
\ No newline at end of file
diff --git a/Engine/Plugins/Experimental/ControlRig/Source/ControlRig/Public/ControlRigDefines.h b/Engine/Plugins/Experimental/ControlRig/Source/ControlRig/Public/ControlRigDefines.h
--- a/Engine/Plugins/Experimental/ControlRig/Source/ControlRig/Public/ControlRigDefines.h
+++ b/Engine/Plugins/Experimental/ControlRig/Source/ControlRig/Public/ControlRigDefines.h
@@ -6,18 +6,20 @@
 #include "Rigs/RigHierarchyContainer.h"
 #include "Rigs/RigCurveContainer.h"
 #include "Stats/StatsHierarchical.h"
+#include "RigVMCore/RigVMExecuteContext.h"
 #include "ControlRigDefines.generated.h"
 
 #define DECLARE_SCOPE_HIERARCHICAL_COUNTER_RIGUNIT()
 //	DECLARE_SCOPE_HIERARCHICAL_COUNTER_FUNC()
 
 USTRUCT()
-struct FControlRigExecuteContext
+struct FControlRigExecuteContext : public FRigVMExecuteContext
 {
 	GENERATED_BODY()
 
 	FControlRigExecuteContext()
-		: Hierarchy(nullptr)
+		: FRigVMExecuteContext()
+		, Hierarchy(nullptr)
 	{
 	}
 		
@@ -111,67 +113,6 @@ enum class EBoneGetterSetterMode : uint8
 	Max UMETA(Hidden),
 };
 
-UENUM()
-enum class EControlRigOpCode : uint8
-{
-	Done,
-	Copy,
-	Exec,
-	Invalid,
-};
-
-USTRUCT()
-struct FControlRigOperator
-{
-	GENERATED_BODY()
-
-	UPROPERTY(VisibleAnywhere, Category = "FControlRigBlueprintOperator")
-	EControlRigOpCode OpCode;
-
-	/** Path to the property we are linking from */
-	UPROPERTY()
-	FString PropertyPath1_DEPRECATED;
-
-	/** Path to the property we are linking to */
-	UPROPERTY()
-	FString PropertyPath2_DEPRECATED;
-
-	/** Path to the property we are linking from */
-	UPROPERTY(VisibleAnywhere, Category = "FControlRigBlueprintOperator")
-	FCachedPropertyPath CachedPropertyPath1;
-
-	/** Path to the property we are linking to */
-	UPROPERTY(VisibleAnywhere, Category = "FControlRigBlueprintOperator")
-	FCachedPropertyPath CachedPropertyPath2;
-
-	FControlRigOperator(EControlRigOpCode Op = EControlRigOpCode::Invalid)
-		: OpCode(Op)
-	{
-	}
-
-	FControlRigOperator(EControlRigOpCode Op, const FCachedPropertyPath& InProperty1, const FCachedPropertyPath& InProperty2)
-		: OpCode(Op)
-		, CachedPropertyPath1(InProperty1)
-		, CachedPropertyPath2(InProperty2)
-	{
-
-	}
-
-	static FControlRigOperator MakeUnresolvedCopy(const FControlRigOperator& ToCopy);
-
-	bool Resolve(UObject* OuterObject);
-
-	FString ToString()
-	{
-		TArray<FStringFormatArg> Arguments;
-		Arguments.Add(FStringFormatArg((int32)OpCode));
-		Arguments.Add(CachedPropertyPath1.ToString());
-		Arguments.Add(CachedPropertyPath1.ToString());
-
-		return FString::Format(TEXT("Opcode {0} : Property1 {1}, Property2 {2}"), Arguments);
-	}
-};
-
 // thought of mixing this with execution on
 // the problem is execution on is transient state, and 
 // this execution type is something to be set per rig
diff --git a/Engine/Plugins/Experimental/ControlRig/Source/ControlRig/Public/ControlRigLog.h b/Engine/Plugins/Experimental/ControlRig/Source/ControlRig/Public/ControlRigLog.h
--- a/Engine/Plugins/Experimental/ControlRig/Source/ControlRig/Public/ControlRigLog.h
+++ b/Engine/Plugins/Experimental/ControlRig/Source/ControlRig/Public/ControlRigLog.h
@@ -15,18 +15,20 @@ public:
 #if WITH_EDITOR
 	struct FLogEntry
 	{
-		FLogEntry(EMessageSeverity::Type InSeverity, const FName& InUnit, const FString& InMessage)
+		FLogEntry(EMessageSeverity::Type InSeverity, const FName& InOperatorName, int32 InInstructionIndex, const FString& InMessage)
 		: Severity(InSeverity)
-		, Unit(InUnit)
+		, OperatorName(InOperatorName)
+		, InstructionIndex(InInstructionIndex)
 		, Message(InMessage)
 		{}
 		
 		EMessageSeverity::Type Severity;
-		FName Unit;
+		FName OperatorName;
+		int32 InstructionIndex;
 		FString Message;
 	};
 	TArray<FLogEntry> Entries;
 #endif
 
-	virtual void Report(EMessageSeverity::Type InSeverity, const FName& InUnit, const FString& InMessage);
+	virtual void Report(EMessageSeverity::Type InSeverity, const FName& InOperatorName, int32 InInstructionIndex, const FString& InMessage);
 };
diff --git a/Engine/Plugins/Experimental/ControlRig/Source/ControlRig/Public/ControlRigObjectVersion.h b/Engine/Plugins/Experimental/ControlRig/Source/ControlRig/Public/ControlRigObjectVersion.h
--- a/Engine/Plugins/Experimental/ControlRig/Source/ControlRig/Public/ControlRigObjectVersion.h
+++ b/Engine/Plugins/Experimental/ControlRig/Source/ControlRig/Public/ControlRigObjectVersion.h
@@ -18,6 +18,9 @@ struct CONTROLRIG_API FControlRigObjectVersion
 		// Refactored operators to store FCachedPropertyPath instead of string
 		OperatorsStoringPropertyPaths,
 
+		// Introduced new RigVM as a backend
+		SwitchedToRigVM,
+
 		// -----<new versions can be added above this line>-------------------------------------------------
 		VersionPlusOne,
 		LatestVersion = VersionPlusOne - 1
diff --git a/Engine/Plugins/Experimental/ControlRig/Source/ControlRig/Public/Drawing/ControlRigDrawContainer.h b/Engine/Plugins/Experimental/ControlRig/Source/ControlRig/Public/Drawing/ControlRigDrawContainer.h
new file mode 100644
--- /dev/null
+++ b/Engine/Plugins/Experimental/ControlRig/Source/ControlRig/Public/Drawing/ControlRigDrawContainer.h
@@ -0,0 +1,30 @@
+// Copyright Epic Games, Inc. All Rights Reserved.
+
+#pragma once
+
+#include "CoreMinimal.h"
+#include "ControlRigDrawInstruction.h"
+#include "ControlRigDrawContainer.generated.h"
+
+USTRUCT()
+struct CONTROLRIG_API FControlRigDrawContainer
+{
+	GENERATED_BODY()
+
+	FORCEINLINE int32 Num() const { return Instructions.Num(); }
+	int32 GetIndex(const FName& InName) const;
+	FORCEINLINE const FControlRigDrawInstruction& operator[](int32 InIndex) const { return Instructions[InIndex]; }
+	FORCEINLINE FControlRigDrawInstruction& operator[](int32 InIndex) { return Instructions[InIndex]; }
+	FORCEINLINE const FControlRigDrawInstruction& operator[](const FName& InName) const { return Instructions[GetIndex(InName)]; }
+	FORCEINLINE FControlRigDrawInstruction& operator[](const FName& InName) { return Instructions[GetIndex(InName)]; }
+
+	FORCEINLINE TArray<FControlRigDrawInstruction>::RangedForIteratorType      begin() { return Instructions.begin(); }
+	FORCEINLINE TArray<FControlRigDrawInstruction>::RangedForConstIteratorType begin() const { return Instructions.begin(); }
+	FORCEINLINE TArray<FControlRigDrawInstruction>::RangedForIteratorType      end() { return Instructions.end(); }
+	FORCEINLINE TArray<FControlRigDrawInstruction>::RangedForConstIteratorType end() const { return Instructions.end(); }
+
+	void Reset();
+
+	UPROPERTY(EditAnywhere, Category = "DrawContainer")
+	TArray<FControlRigDrawInstruction> Instructions;
+};
diff --git a/Engine/Plugins/Experimental/ControlRig/Source/ControlRig/Public/Drawing/ControlRigDrawInstruction.h b/Engine/Plugins/Experimental/ControlRig/Source/ControlRig/Public/Drawing/ControlRigDrawInstruction.h
new file mode 100644
--- /dev/null
+++ b/Engine/Plugins/Experimental/ControlRig/Source/ControlRig/Public/Drawing/ControlRigDrawInstruction.h
@@ -0,0 +1,57 @@
+// Copyright Epic Games, Inc. All Rights Reserved.
+
+#pragma once
+
+#include "CoreMinimal.h"
+#include "ControlRigDrawInstruction.generated.h"
+
+UENUM()
+namespace EControlRigDrawSettings
+{
+	enum Primitive
+	{
+		Points,
+		Lines,
+		LineStrip
+	};
+}
+
+USTRUCT()
+struct CONTROLRIG_API FControlRigDrawInstruction
+{
+	GENERATED_BODY()
+
+	UPROPERTY(EditAnywhere, Category = "DrawInstruction")
+	FName Name;
+
+	UPROPERTY(EditAnywhere, Category = "DrawInstruction")
+	TEnumAsByte<EControlRigDrawSettings::Primitive> PrimitiveType;
+
+	UPROPERTY(EditAnywhere, Category = "DrawInstruction")
+	TArray<FVector> Positions;
+
+	UPROPERTY(EditAnywhere, Category = "DrawInstruction")
+	FLinearColor Color;
+
+	UPROPERTY(EditAnywhere, Category = "DrawInstruction")
+	float Thickness;
+
+	UPROPERTY(EditAnywhere, Category = "DrawInstruction")
+	FTransform Transform;
+
+	FControlRigDrawInstruction()
+		: Name(NAME_None)
+		, PrimitiveType(EControlRigDrawSettings::Points)
+		, Color(FLinearColor::Red)
+		, Thickness(0.f)
+		, Transform(FTransform::Identity)
+	{}
+
+	FControlRigDrawInstruction(EControlRigDrawSettings::Primitive InPrimitiveType, const FLinearColor& InColor, float InThickness = 0.f, FTransform InTransform = FTransform::Identity)
+		: Name(NAME_None)
+		, PrimitiveType(InPrimitiveType)
+		, Color(InColor)
+		, Thickness(InThickness)
+		, Transform(InTransform)
+	{}
+};
diff --git a/Engine/Plugins/Experimental/ControlRig/Source/ControlRig/Public/Drawing/ControlRigDrawInterface.h b/Engine/Plugins/Experimental/ControlRig/Source/ControlRig/Public/Drawing/ControlRigDrawInterface.h
--- a/Engine/Plugins/Experimental/ControlRig/Source/ControlRig/Public/Drawing/ControlRigDrawInterface.h
+++ b/Engine/Plugins/Experimental/ControlRig/Source/ControlRig/Public/Drawing/ControlRigDrawInterface.h
@@ -5,6 +5,8 @@
 #include "CoreMinimal.h"
 #include "Rigs/RigHierarchyContainer.h"
 #include "Math/Simulation/CRSimPointContainer.h"
+#include "Drawing/ControlRigDrawInstruction.h"
+#include "Drawing/ControlRigDrawContainer.h"
 #include "ControlRigDrawInterface.generated.h"
 
 UENUM()
@@ -20,10 +22,13 @@ namespace EControlRigDrawHierarchyMode
 	};
 }
 
-struct CONTROLRIG_API FControlRigDrawInterface
+USTRUCT()
+struct CONTROLRIG_API FControlRigDrawInterface : public FControlRigDrawContainer
 {
 public:
 
+	GENERATED_BODY()
+
 	FControlRigDrawInterface() {}
 	virtual ~FControlRigDrawInterface() {}
 
@@ -43,35 +48,6 @@ public:
 
 private:
 
-	enum EDrawType
-	{
-		EDrawType_Point,
-		EDrawType_Lines,
-		EDrawType_LineStrip
-	};
-
-	struct FDrawIntruction
-	{
-		EDrawType DrawType;
-		TArray<FVector> Positions;
-		FLinearColor Color;
-		float Thickness;
-
-		FDrawIntruction()
-			: DrawType(EDrawType_Lines)
-			, Color(FLinearColor::Red)
-			, Thickness(0.f)
-		{}
-
-		FDrawIntruction(EDrawType InDrawType, const FLinearColor& InColor, float InThickness = 0.f)
-			: DrawType(InDrawType)
-			, Color(InColor)
-			, Thickness(InThickness)
-		{}
-	};
-
-	TArray<FDrawIntruction> DrawInstructions;
-
 	friend class FControlRigEditMode;
 	friend class UControlRig;
 };
diff --git a/Engine/Plugins/Experimental/ControlRig/Source/ControlRig/Public/IControlRigModule.h b/Engine/Plugins/Experimental/ControlRig/Source/ControlRig/Public/IControlRigModule.h
--- a/Engine/Plugins/Experimental/ControlRig/Source/ControlRig/Public/IControlRigModule.h
+++ b/Engine/Plugins/Experimental/ControlRig/Source/ControlRig/Public/IControlRigModule.h
@@ -6,8 +6,6 @@
 #include "Modules/ModuleInterface.h"
 #include "Templates/SubclassOf.h"
 
-class URigUnitEditor_Base;
-
 class IControlRigModule : public IModuleInterface
 {
 public:
diff --git a/Engine/Plugins/Experimental/ControlRig/Source/ControlRig/Public/Manipulatable/IControlRigManipulatable.h b/Engine/Plugins/Experimental/ControlRig/Source/ControlRig/Public/Manipulatable/IControlRigManipulatable.h
--- a/Engine/Plugins/Experimental/ControlRig/Source/ControlRig/Public/Manipulatable/IControlRigManipulatable.h
+++ b/Engine/Plugins/Experimental/ControlRig/Source/ControlRig/Public/Manipulatable/IControlRigManipulatable.h
@@ -10,6 +10,14 @@
 #include "UObject/Interface.h"
 #include "IControlRigManipulatable.generated.h"
 
+/** When setting control values what to do with regards to setting key.*/
+enum class EControlRigSetKey : uint8
+{
+	DoNotCare = 0x0,    //Don't care if a key is set or not, may get set, say if auto key is on somewhere.
+	Always,				//Always set a key here
+	Never				//Never set a key here.
+};
+
 /**
  *
  * IControlRigManipulatable provides an interface for subjects who desire to be 
@@ -39,7 +47,7 @@ public:
 	DECLARE_EVENT_ThreeParams(IControlRigManipulatable, FFilterControlEvent, IControlRigManipulatable*, const FRigControl&, FRigControlValue&);
 
 	/** Bindable event for external objects to be notified of Control changes */
-	DECLARE_EVENT_TwoParams(IControlRigManipulatable, FControlModifiedEvent, IControlRigManipulatable*, const FRigControl&);
+	DECLARE_EVENT_ThreeParams(IControlRigManipulatable, FControlModifiedEvent, IControlRigManipulatable*, const FRigControl&, EControlRigSetKey);
 #if WITH_EDITOR
 	/** Bindable event for external objects to be notified that a Control is Selected */
 	DECLARE_EVENT_ThreeParams(IControlRigManipulatable, FControlSelectedEvent, IControlRigManipulatable*, const FRigControl&,  bool);
@@ -93,23 +101,27 @@ public:
 	}
 
 	// Sets the relative value of a Control
-	FORCEINLINE void SetControlValue(const FName& InControlName, const FRigControlValue& InValue, bool bNotify = true)
+	FORCEINLINE void SetControlValue(const FName& InControlName, const FRigControlValue& InValue, bool bNotify = true,
+		EControlRigSetKey InSetKey = EControlRigSetKey::DoNotCare)
 	{
 		FRigControl* Control = FindControl(InControlName);
 		check(Control);
+
 		Control->Value = InValue;
+		Control->ApplyLimits(Control->Value);
 
 		if (bNotify && OnControlModified.IsBound())
 		{
-			OnControlModified.Broadcast(this, *Control);
+			OnControlModified.Broadcast(this, *Control, InSetKey);
 		}
 	}
 
 	// Sets the relative value of a Control
 	template<class T>
-	FORCEINLINE void SetControlValue(const FName& InControlName, T InValue, bool bNotify = true)
+	FORCEINLINE void SetControlValue(const FName& InControlName, T InValue, bool bNotify = true,
+		EControlRigSetKey InSetKey = EControlRigSetKey::DoNotCare)
 	{
-		SetControlValue(InControlName, FRigControlValue::Make<T>(InValue), bNotify);
+		SetControlValue(InControlName, FRigControlValue::Make<T>(InValue), bNotify, InSetKey);
 	}
 
 	// Returns the global / world transform of a Control
@@ -117,7 +129,8 @@ public:
 
 	// Sets the global / world transform of a Control. This should be called from the interaction
 	// layer / edit mode. Returns true when successful.
-	bool SetControlGlobalTransform(const FName& InControlName, const FTransform& InGlobalTransform);
+	bool SetControlGlobalTransform(const FName& InControlName, const FTransform& InGlobalTransform, 
+		EControlRigSetKey InSetKey = EControlRigSetKey::DoNotCare);
 
 	// Returns the value given a global transform
 	virtual FRigControlValue GetControlValueFromGlobalTransform(const FName& InControlName, const FTransform& InGlobalTransform) = 0;
@@ -163,6 +176,8 @@ public:
 	//Get bindings to a runtime object 
 	virtual TSharedPtr<IControlRigObjectBinding> GetObjectBinding() const {	return nullptr;}
 
+	//Create RigControls for Curves, they will get added to the AvailableControls
+	virtual void CreateRigControlsForCurveContainer()  {};
 
 private:
 
diff --git a/Engine/Plugins/Experimental/ControlRig/Source/ControlRig/Public/Rigs/AdditiveControlRig.h b/Engine/Plugins/Experimental/ControlRig/Source/ControlRig/Public/Rigs/AdditiveControlRig.h
new file mode 100644
--- /dev/null
+++ b/Engine/Plugins/Experimental/ControlRig/Source/ControlRig/Public/Rigs/AdditiveControlRig.h
@@ -0,0 +1,35 @@
+// Copyright Epic Games, Inc. All Rights Reserved.
+
+#pragma once
+#include "ControlRig.h"
+#include "Units/Hierarchy/RigUnit_AddBoneTransform.h"
+#include "AdditiveControlRig.generated.h"
+
+class USkeletalMesh;
+struct FReferenceSkeleton;
+struct FSmartNameMapping;
+
+/** Rig that allows additive layer editing per joint */
+UCLASS(NotBlueprintable)
+class CONTROLRIG_API UAdditiveControlRig : public UControlRig
+{
+	GENERATED_UCLASS_BODY()
+
+public: 
+	// BEGIN ControlRig
+	virtual void Initialize(bool bInitRigUnits = true) override;
+	virtual void ExecuteUnits(FRigUnitContext& InOutContext) override;
+	// END ControlRig
+
+	// utility function to 
+	static FName GetControlName(const FName& InBoneName);
+	static FName GetSpaceName(const FName& InBoneName);
+
+private:
+	// custom units that are used to execute this rig
+	TArray<FRigUnit_AddBoneTransform> AddBoneRigUnits;
+
+	/** Create RigElements - bone hierarchy and curves - from incoming skeleton */
+	void CreateRigElements(const USkeletalMesh* InReferenceMesh);
+	void CreateRigElements(const FReferenceSkeleton& InReferenceSkeleton, const FSmartNameMapping* InSmartNameMapping);
+};
\ No newline at end of file
diff --git a/Engine/Plugins/Experimental/ControlRig/Source/ControlRig/Public/Rigs/RigBoneHierarchy.h b/Engine/Plugins/Experimental/ControlRig/Source/ControlRig/Public/Rigs/RigBoneHierarchy.h
--- a/Engine/Plugins/Experimental/ControlRig/Source/ControlRig/Public/Rigs/RigBoneHierarchy.h
+++ b/Engine/Plugins/Experimental/ControlRig/Source/ControlRig/Public/Rigs/RigBoneHierarchy.h
@@ -9,7 +9,7 @@
 
 class UControlRig;
 
-UENUM()
+UENUM(BlueprintType)
 enum class ERigBoneType : uint8
 {
 	Imported,
@@ -34,20 +34,20 @@ struct CONTROLRIG_API FRigBone: public FRigElement
 	}
 	virtual ~FRigBone() {}
 
-	UPROPERTY(VisibleAnywhere, Category = FRigElement)
+	UPROPERTY(BlueprintReadOnly, VisibleAnywhere, Category = FRigElement)
 	FName ParentName;
 
 	UPROPERTY(transient)
 	int32 ParentIndex;
 
 	/* Initial global transform that is saved in this rig */
-	UPROPERTY(EditAnywhere, Category = FRigElement)
+	UPROPERTY(BlueprintReadOnly, EditAnywhere, Category = FRigElement)
 	FTransform InitialTransform;
 
-	UPROPERTY(transient, VisibleAnywhere, Category = FRigElement)
+	UPROPERTY(BlueprintReadOnly, transient, EditAnywhere, Category = FRigElement)
 	FTransform GlobalTransform;
 
-	UPROPERTY(transient, VisibleAnywhere, Category = FRigElement)
+	UPROPERTY(BlueprintReadOnly, transient, EditAnywhere, Category = FRigElement)
 	FTransform LocalTransform;
 
 	/** dependent list - direct dependent for child or anything that needs to update due to this */
@@ -55,7 +55,7 @@ struct CONTROLRIG_API FRigBone: public FRigElement
 	TArray<int32> Dependents;
 
 	/** the source of the bone to differentiate procedurally generated, imported etc */
-	UPROPERTY(VisibleAnywhere, Category = FRigElement)
+	UPROPERTY(BlueprintReadWrite, VisibleAnywhere, Category = FRigElement)
 	ERigBoneType Type;
 
 	FORCEINLINE virtual ERigElementType GetElementType() const override
@@ -69,7 +69,7 @@ struct CONTROLRIG_API FRigBone: public FRigElement
 	}
 };
 
-USTRUCT()
+USTRUCT(BlueprintType)
 struct CONTROLRIG_API FRigBoneHierarchy
 {
 	GENERATED_BODY()
@@ -152,7 +152,7 @@ struct CONTROLRIG_API FRigBoneHierarchy
 	FTransform GetInitialTransform(int32 InIndex) const;
 
 	// updates all of the internal caches
-	void Initialize();
+	void Initialize(bool bResetTransforms = true);
 
 	// clears the hierarchy and removes all content
 	void Reset();
@@ -160,6 +160,9 @@ struct CONTROLRIG_API FRigBoneHierarchy
 	// resets all of the transforms back to the initial transform
 	void ResetTransforms();
 
+	// import skeleton
+	TArray<FRigElementKey> ImportSkeleton(const FReferenceSkeleton& InSkeleton, const FName& InNameSpace, bool bReplaceExistingBones, bool bRemoveObsoleteBones, bool bSelectBones, bool bNotify);
+
 #if WITH_EDITOR
 
 	bool Select(const FName& InName, bool bSelect = true);
@@ -167,8 +170,6 @@ struct CONTROLRIG_API FRigBoneHierarchy
 	TArray<FName> CurrentSelection() const;
 	bool IsSelected(const FName& InName) const;
 
-	TArray<FRigElementKey> ImportSkeleton(const FReferenceSkeleton& InSkeleton, const FName& InNameSpace, bool bReplaceExistingBones, bool bRemoveObsoleteBones, bool bSelectBones);
-
 	FRigElementAdded OnBoneAdded;
 	FRigElementRemoved OnBoneRemoved;
 	FRigElementRenamed OnBoneRenamed;
@@ -209,6 +210,9 @@ private:
 	int32 GetChildrenRecursive(const int32 InIndex, TArray<int32>& OutChildren, bool bRecursively) const;
 
 	void PropagateTransform(int32 InIndex);
-	
+
+	bool bSuspendNotifications;
+
 	friend struct FRigHierarchyContainer;
+	friend class UControlRigHierarchyModifier;
 };
diff --git a/Engine/Plugins/Experimental/ControlRig/Source/ControlRig/Public/Rigs/RigControlHierarchy.h b/Engine/Plugins/Experimental/ControlRig/Source/ControlRig/Public/Rigs/RigControlHierarchy.h
--- a/Engine/Plugins/Experimental/ControlRig/Source/ControlRig/Public/Rigs/RigControlHierarchy.h
+++ b/Engine/Plugins/Experimental/ControlRig/Source/ControlRig/Public/Rigs/RigControlHierarchy.h
@@ -4,11 +4,12 @@
 
 #include "CoreMinimal.h"
 #include "RigHierarchyDefines.h"
+#include "TransformNoScale.h"
 #include "RigControlHierarchy.generated.h"
 
 class UControlRig;
 
-UENUM()
+UENUM(BlueprintType)
 enum class ERigControlType : uint8
 {
 	Bool,
@@ -16,12 +17,29 @@ enum class ERigControlType : uint8
 	Vector2D,
 	Position,
 	Scale,
-	Quat,
 	Rotator,
-	Transform
+	Transform,
+	TransformNoScale
 };
 
-USTRUCT()
+UENUM(BlueprintType)
+enum class ERigControlValueType : uint8
+{
+	Initial,
+	Current,
+	Minimum,
+	Maximum
+};
+
+UENUM(BlueprintType)
+enum class ERigControlAxis : uint8
+{
+	X,
+	Y,
+	Z
+};
+
+USTRUCT(BlueprintType)
 struct CONTROLRIG_API FRigControlValue
 {
 	GENERATED_BODY()
@@ -34,33 +52,50 @@ public:
 	}
 
 	template<class T>
-	FORCEINLINE T Get() const
+	FORCEINLINE_DEBUGGABLE T Get() const
 	{
 		return GetRef<T>();
 	}
 
 	template<class T>
-	FORCEINLINE T& GetRef()
+	FORCEINLINE_DEBUGGABLE T& GetRef()
 	{
 		ensure(sizeof(T) <= sizeof(FTransform));
 		return *(T*)&Storage;
 	}
 
 	template<class T>
-	FORCEINLINE const T& GetRef() const
+	FORCEINLINE_DEBUGGABLE const T& GetRef() const
 	{
 		ensure(sizeof(T) <= sizeof(FTransform));
 		return *(T*)&Storage;
 	}
 
 	template<class T>
-	FORCEINLINE void Set(T InValue)
+	FORCEINLINE_DEBUGGABLE void Set(T InValue)
 	{
 		GetRef<T>() = InValue;
 	}
 
 	template<class T>
-	FORCEINLINE static FRigControlValue Make(T InValue)
+	FORCEINLINE_DEBUGGABLE FString ToString() const
+	{
+		FString Result;
+		TBaseStructure<T>::Get()->ExportText(Result, &GetRef<T>(), nullptr, nullptr, PPF_None, nullptr);
+		return Result;
+	}
+
+	template<class T>
+	FORCEINLINE_DEBUGGABLE T SetFromString(const FString& InString)
+	{
+		T Value;
+		TBaseStructure<T>::Get()->ImportText(*InString, &Value, nullptr, PPF_None, nullptr, TBaseStructure<T>::Get()->GetName());
+		Set<T>(Value);
+		return Value;
+	}
+
+	template<class T>
+	FORCEINLINE_DEBUGGABLE static FRigControlValue Make(T InValue)
 	{
 		FRigControlValue Value;
 		Value.Set<T>(InValue);
@@ -73,15 +108,24 @@ private:
 	FTransform Storage;
 };
 
+template<>
+FORCEINLINE_DEBUGGABLE FQuat FRigControlValue::SetFromString<FQuat>(const FString& InString)
+{
+	FQuat Value;
+	TBaseStructure<FQuat>::Get()->ImportText(*InString, &Value, nullptr, PPF_None, nullptr, TBaseStructure<FQuat>::Get()->GetName());
+	Set<FRotator>(Value.Rotator());
+	return Value;
+}
+
 
 class UStaticMesh;
 
-USTRUCT()
+USTRUCT(BlueprintType)
 struct CONTROLRIG_API FRigControl : public FRigElement
 {
 	GENERATED_BODY()
 
-	FRigControl()
+		FRigControl()
 		: FRigElement()
 		, ControlType(ERigControlType::Transform)
 		, ParentName(NAME_None)
@@ -90,70 +134,179 @@ struct CONTROLRIG_API FRigControl : public FRigElement
 		, SpaceIndex(INDEX_NONE)
 		, InitialValue()
 		, Value()
-#if WITH_EDITORONLY_DATA
+		, PrimaryAxis(ERigControlAxis::X)
+		, bIsCurve(false)
+		, bAnimatable(true)
+		, bLimitTranslation(false)
+		, bLimitRotation(false)
+		, bLimitScale(false)
+		, bDrawLimits(false)
+		, MinimumValue()
+		, MaximumValue()
+		, bGizmoEnabled(true)
 		, GizmoName(TEXT("Gizmo"))
 		, GizmoTransform(FTransform::Identity)
 		, GizmoColor(FLinearColor::Red)
-#endif
 		, Dependents()
+		, bIsTransientControl(false)
 	{
 	}
 	virtual ~FRigControl() {}
 
-	UPROPERTY(VisibleAnywhere, Category = FRigElement)
+	UPROPERTY(BlueprintReadWrite, EditAnywhere, Category = Control)
 	ERigControlType ControlType;
 
-	UPROPERTY(VisibleAnywhere, Category = FRigElement)
+	UPROPERTY(BlueprintReadOnly, VisibleAnywhere, Category = Control)
 	FName ParentName;
 
-	UPROPERTY(transient)
+	UPROPERTY(BlueprintReadOnly, transient, Category = Control)
 	int32 ParentIndex;
 
-	UPROPERTY(VisibleAnywhere, Category = FRigElement)
+	UPROPERTY(BlueprintReadOnly, VisibleAnywhere, Category = Control)
 	FName SpaceName;
 
-	UPROPERTY(transient)
+	UPROPERTY(BlueprintReadOnly, transient, Category = Control)
 	int32 SpaceIndex;
 
-	UPROPERTY(VisibleAnywhere, Category = FRigElement)
+	UPROPERTY(BlueprintReadOnly, VisibleAnywhere, Category = Control)
 	FRigControlValue InitialValue;
 
-	UPROPERTY(transient, VisibleAnywhere, Category = FRigElement)
+	UPROPERTY(BlueprintReadOnly, transient, VisibleAnywhere, Category = Control)
 	FRigControlValue Value;
 
-#if WITH_EDITORONLY_DATA
+	/** the primary axis to use for float controls */
+	UPROPERTY(BlueprintReadWrite, EditAnywhere, Category = Control)
+	ERigControlAxis PrimaryAxis;
+
+	/** If Created from a Curve  Container*/
+	UPROPERTY(transient)
+	bool bIsCurve;
+
+	/** If the control is animatable in sequencer */
+	UPROPERTY(BlueprintReadWrite, EditAnywhere, Category = Control)
+	bool bAnimatable;
+
+	/** True if the control has to obey translation limits. */
+	UPROPERTY(BlueprintReadWrite, EditAnywhere, Category = Limits)
+	bool bLimitTranslation;
+
+	/** True if the control has to obey rotation limits. */
+	UPROPERTY(BlueprintReadWrite, EditAnywhere, Category = Limits)
+	bool bLimitRotation;
+
+	/** True if the control has to obey scale limits. */
+	UPROPERTY(BlueprintReadWrite, EditAnywhere, Category = Limits)
+	bool bLimitScale;
+
+	/** True if the limits should be drawn in debug. */
+	UPROPERTY(BlueprintReadWrite, EditAnywhere, Category = Limits, meta = (EditCondition = "bLimitTranslation || bLimitRotation || bLimitScale"))
+	bool bDrawLimits;
+
+	/** The storage for all minimum values */
+	UPROPERTY(BlueprintReadWrite, EditAnywhere, Category = Limits, meta = (EditCondition = "bLimitTranslation || bLimitRotation || bLimitScale"))
+	FRigControlValue MinimumValue;
+
+	/** The storage for all maximum values */
+	UPROPERTY(BlueprintReadWrite, EditAnywhere, Category = Limits, meta = (EditCondition = "bLimitTranslation || bLimitRotation || bLimitScale"))
+	FRigControlValue MaximumValue;
+
+	/** Set to true if the gizmo is enabled in 3d */
+	UPROPERTY(BlueprintReadWrite, EditAnywhere, Category = Gizmo)
+	bool bGizmoEnabled;
+
 	/* This is optional UI setting - this doesn't mean this is always used, but it is optional for manipulation layer to use this*/
-	UPROPERTY(EditAnywhere, Category = FRigElement)
+	UPROPERTY(BlueprintReadWrite, EditAnywhere, Category = Gizmo, meta = (EditCondition = "bGizmoEnabled"))
 	FName GizmoName;
 
-	UPROPERTY(EditAnywhere, Category = FRigElement)
+	UPROPERTY(BlueprintReadWrite, EditAnywhere, Category = Gizmo, meta = (EditCondition = "bGizmoEnabled"))
 	FTransform GizmoTransform;
 
-	UPROPERTY(EditAnywhere, Category = FRigElement)
+	UPROPERTY(BlueprintReadWrite, EditAnywhere, Category = Gizmo, meta = (EditCondition = "bGizmoEnabled"))
 	FLinearColor GizmoColor;
-#endif // WITH_EDITORONLY_DATA
 
 	/** dependent list - direct dependent for child or anything that needs to update due to this */
 	UPROPERTY(transient)
 	TArray<int32> Dependents;
 
-	FORCEINLINE virtual ERigElementType GetElementType() const override
+	/** If the control is transient and only visible in the control rig editor */
+	UPROPERTY(BlueprintReadWrite, EditAnywhere, Category = Control)
+	bool bIsTransientControl;
+
+	FORCEINLINE_DEBUGGABLE virtual ERigElementType GetElementType() const override
 	{
 		return ERigElementType::Control;
 	}
 
-	FORCEINLINE virtual FRigElementKey GetParentElementKey() const
+	FORCEINLINE_DEBUGGABLE virtual FRigElementKey GetParentElementKey() const
 	{
 		return FRigElementKey(ParentName, GetElementType());
 	}
 
-	FORCEINLINE virtual FRigElementKey GetSpaceElementKey() const
+	FORCEINLINE_DEBUGGABLE virtual FRigElementKey GetSpaceElementKey() const
 	{
 		return FRigElementKey(SpaceName, ERigElementType::Space);
 	}
+
+	FORCEINLINE_DEBUGGABLE const FRigControlValue& GetValue(ERigControlValueType InValueType = ERigControlValueType::Current) const
+	{
+		switch(InValueType)
+		{
+			case ERigControlValueType::Initial:
+			{
+				return InitialValue;
+			}
+			case ERigControlValueType::Minimum:
+			{
+				return MinimumValue;
+			}
+			case ERigControlValueType::Maximum:
+			{
+				return MaximumValue;
+			}
+			default:
+			{
+				break;
+			}
+		}
+		return Value;
+	}
+
+	FORCEINLINE_DEBUGGABLE FRigControlValue& GetValue(ERigControlValueType InValueType = ERigControlValueType::Current)
+	{
+		switch(InValueType)
+		{
+			case ERigControlValueType::Initial:
+			{
+				return InitialValue;
+			}
+			case ERigControlValueType::Minimum:
+			{
+				return MinimumValue;
+			}
+			case ERigControlValueType::Maximum:
+			{
+				return MaximumValue;
+			}
+			default:
+			{
+				break;
+			}
+		}
+		return Value;
+	}
+
+	void ApplyLimits(FRigControlValue& InOutValue);
+	FORCEINLINE_DEBUGGABLE static float Clamp(float Value, float Minimum, float Maximum)
+	{
+		if (Minimum < Maximum)
+		{
+			return FMath::Clamp<float>(Value, Minimum, Maximum);
+		}
+		return FMath::Clamp<float>(Value, Maximum, Minimum);
+	}
 };
 
-USTRUCT()
+USTRUCT(BlueprintType)
 struct CONTROLRIG_API FRigControlHierarchy
 {
 	GENERATED_BODY()
@@ -161,22 +314,22 @@ struct CONTROLRIG_API FRigControlHierarchy
 	FRigControlHierarchy();
 	FRigControlHierarchy& operator= (const FRigControlHierarchy &InOther);
 
-	FORCEINLINE ERigElementType RigElementType() const { return ERigElementType::Control; }
+	FORCEINLINE_DEBUGGABLE ERigElementType RigElementType() const { return ERigElementType::Control; }
 
-	FORCEINLINE int32 Num() const { return Controls.Num(); }
-	FORCEINLINE const FRigControl& operator[](int32 InIndex) const { return Controls[InIndex]; }
-	FORCEINLINE FRigControl& operator[](int32 InIndex) { return Controls[InIndex]; }
-	FORCEINLINE const FRigControl& operator[](const FName& InName) const { return Controls[GetIndex(InName)]; }
-	FORCEINLINE FRigControl& operator[](const FName& InName) { return Controls[GetIndex(InName)]; }
+	FORCEINLINE_DEBUGGABLE int32 Num() const { return Controls.Num(); }
+	FORCEINLINE_DEBUGGABLE const FRigControl& operator[](int32 InIndex) const { return Controls[InIndex]; }
+	FORCEINLINE_DEBUGGABLE FRigControl& operator[](int32 InIndex) { return Controls[InIndex]; }
+	FORCEINLINE_DEBUGGABLE const FRigControl& operator[](const FName& InName) const { return Controls[GetIndex(InName)]; }
+	FORCEINLINE_DEBUGGABLE FRigControl& operator[](const FName& InName) { return Controls[GetIndex(InName)]; }
 
-	FORCEINLINE const TArray<FRigControl>& GetControls() const { return Controls; }
+	FORCEINLINE_DEBUGGABLE const TArray<FRigControl>& GetControls() const { return Controls; }
 
-	FORCEINLINE TArray<FRigControl>::RangedForIteratorType      begin()       { return Controls.begin(); }
-	FORCEINLINE TArray<FRigControl>::RangedForConstIteratorType begin() const { return Controls.begin(); }
-	FORCEINLINE TArray<FRigControl>::RangedForIteratorType      end()         { return Controls.end();   }
-	FORCEINLINE TArray<FRigControl>::RangedForConstIteratorType end() const   { return Controls.end();   }
+	FORCEINLINE_DEBUGGABLE TArray<FRigControl>::RangedForIteratorType      begin()       { return Controls.begin(); }
+	FORCEINLINE_DEBUGGABLE TArray<FRigControl>::RangedForConstIteratorType begin() const { return Controls.begin(); }
+	FORCEINLINE_DEBUGGABLE TArray<FRigControl>::RangedForIteratorType      end()         { return Controls.end();   }
+	FORCEINLINE_DEBUGGABLE TArray<FRigControl>::RangedForConstIteratorType end() const   { return Controls.end();   }
 
-	FORCEINLINE bool IsNameAvailable(const FName& InPotentialNewName) const { return GetIndex(InPotentialNewName) == INDEX_NONE; }
+	FORCEINLINE_DEBUGGABLE bool IsNameAvailable(const FName& InPotentialNewName) const { return GetIndex(InPotentialNewName) == INDEX_NONE; }
 
 	FName GetSafeNewName(const FName& InPotentialNewName) const;
 
@@ -185,12 +338,10 @@ struct CONTROLRIG_API FRigControlHierarchy
 		ERigControlType InControlType = ERigControlType::Transform,
 		const FName& InParentName = NAME_None,
 		const FName& InSpaceName = NAME_None,
-		const FRigControlValue& InValue = FRigControlValue()
-#if WITH_EDITORONLY_DATA
-		, const FName& InGizmoName = TEXT("Gizmo"),
+		const FRigControlValue& InValue = FRigControlValue(),
+		const FName& InGizmoName = TEXT("Gizmo"),
 		const FTransform& InGizmoTransform = FTransform::Identity,
 		const FLinearColor& InGizmoColor = FLinearColor::Red
-#endif
 	);
 
 	FRigControl Remove(const FName& InNameToRemove);
@@ -203,7 +354,7 @@ struct CONTROLRIG_API FRigControlHierarchy
 
 	FName GetName(int32 InIndex) const;
 
-	FORCEINLINE int32 GetIndex(const FName& InName) const
+	FORCEINLINE_DEBUGGABLE int32 GetIndex(const FName& InName) const
 	{
 		if(NameToIndexMapping.Num() != Controls.Num())
 		{
@@ -232,17 +383,17 @@ struct CONTROLRIG_API FRigControlHierarchy
 
 	FTransform GetGlobalTransform(int32 InIndex) const;
 
-	void SetLocalTransform(const FName& InName, const FTransform& InTransform);
+	void SetLocalTransform(const FName& InName, const FTransform& InTransform, ERigControlValueType InValueType = ERigControlValueType::Current);
 
-	void SetLocalTransform(int32 InIndex, const FTransform& InTransform);
+	void SetLocalTransform(int32 InIndex, const FTransform& InTransform, ERigControlValueType InValueType = ERigControlValueType::Current);
 
-	FTransform GetLocalTransform(const FName& InName) const;
+	FTransform GetLocalTransform(const FName& InName, ERigControlValueType InValueType = ERigControlValueType::Current) const;
 
-	FTransform GetLocalTransform(int32 InIndex) const;
+	FTransform GetLocalTransform(int32 InIndex, ERigControlValueType InValueType = ERigControlValueType::Current) const;
 
 	FTransform GetParentTransform(int32 InIndex) const;
 
-	// this is only valid if you have trasnform as type
+	// this is only valid if you have transform as type
 	// setting initial transform from global transform (control transform is based on parent)
 	void SetInitialGlobalTransform(const FName& InName, const FTransform& GlobalTransform);
 	
@@ -252,72 +403,172 @@ struct CONTROLRIG_API FRigControlHierarchy
 	
 	FTransform GetInitialGlobalTransform(int32 InIndex) const;
 	
-	void SetValue(const FName& InName, const FRigControlValue& InValue);
+	void SetValue(const FName& InName, const FRigControlValue& InValue, ERigControlValueType InValueType = ERigControlValueType::Current);
 
-	void SetValue(int32 InIndex, const FRigControlValue& InValue);
+	void SetValue(int32 InIndex, const FRigControlValue& InValue, ERigControlValueType InValueType = ERigControlValueType::Current);
 
 	template<class T>
-	FORCEINLINE void SetValue(const FName& InName, T InValue)
+	FORCEINLINE_DEBUGGABLE void SetValue(const FName& InName, T InValue, ERigControlValueType InValueType = ERigControlValueType::Current)
 	{
-		SetValue(InName, FRigControlValue::Make<T>(InValue));
+		SetValue(InName, FRigControlValue::Make<T>(InValue), InValueType);
 	}
 
 	template<class T>
-	FORCEINLINE void SetValue(int32 InIndex, T InValue)
+	FORCEINLINE_DEBUGGABLE void SetValue(int32 InIndex, T InValue, ERigControlValueType InValueType = ERigControlValueType::Current)
 	{
-		SetValue(InIndex, FRigControlValue::Make<T>(InValue));
+		SetValue(InIndex, FRigControlValue::Make<T>(InValue), InValueType);
 	}
 
-	FRigControlValue GetValue(const FName& InName) const;
+	FRigControlValue GetValue(const FName& InName, ERigControlValueType InValueType = ERigControlValueType::Current) const;
 
-	FRigControlValue GetValue(int32 InIndex) const;
+	FRigControlValue GetValue(int32 InIndex, ERigControlValueType InValueType = ERigControlValueType::Current) const;
 
 	template<class T>
-	FORCEINLINE T GetValue(const FName& InName) const
+	FORCEINLINE_DEBUGGABLE T GetValue(const FName& InName, ERigControlValueType InValueType = ERigControlValueType::Current) const
 	{
-		return GetValue(InName).Get<T>();
+		return GetValue(InName, InValueType).Get<T>();
 	}
 
 	template<class T>
-	FORCEINLINE T GetValue(int32 InIndex) const
+	FORCEINLINE_DEBUGGABLE T GetValue(int32 InIndex, ERigControlValueType InValueType = ERigControlValueType::Current) const
 	{
-		return GetValue(InIndex).Get<T>();
+		return GetValue(InIndex, InValueType).Get<T>();
 	}
 
-	void SetInitialValue(const FName& InName, const FRigControlValue& InValue);
+	FORCEINLINE_DEBUGGABLE void SetInitialValue(const FName& InName, const FRigControlValue& InValue)
+	{
+		SetValue(InName, InValue, ERigControlValueType::Initial);
+	}
 
-	void SetInitialValue(int32 InIndex, const FRigControlValue& InValue);
+	FORCEINLINE_DEBUGGABLE void SetInitialValue(int32 InIndex, const FRigControlValue& InValue)
+	{
+		SetValue(InIndex, InValue, ERigControlValueType::Initial);
+	}
 
 	template<class T>
-	FORCEINLINE void SetInitialValue(const FName& InName, T InValue)
+	FORCEINLINE_DEBUGGABLE void SetInitialValue(const FName& InName, T InValue)
 	{
 		SetInitialValue(InName, FRigControlValue::Make<T>(InValue));
 	}
 
 	template<class T>
-	FORCEINLINE void SetInitialValue(int32 InIndex, T InValue)
+	FORCEINLINE_DEBUGGABLE void SetInitialValue(int32 InIndex, T InValue)
 	{
 		SetInitialValue(InIndex, FRigControlValue::Make<T>(InValue));
 	}
 
-	FRigControlValue GetInitialValue(const FName& InName) const;
+	FORCEINLINE_DEBUGGABLE FRigControlValue GetInitialValue(const FName& InName) const
+	{
+		return GetValue(InName, ERigControlValueType::Initial);
+	}
 
-	FRigControlValue GetInitialValue(int32 InIndex) const;
+	FORCEINLINE_DEBUGGABLE FRigControlValue GetInitialValue(int32 InIndex) const
+	{
+		return GetValue(InIndex, ERigControlValueType::Initial);
+	}
 
 	template<class T>
-	FORCEINLINE T GetInitialValue(const FName& InName) const
+	FORCEINLINE_DEBUGGABLE T GetInitialValue(const FName& InName) const
 	{
 		return GetInitialValue(InName).Get<T>();
 	}
 
 	template<class T>
-	FORCEINLINE T GetInitialValue(int32 InIndex) const
+	FORCEINLINE_DEBUGGABLE T GetInitialValue(int32 InIndex) const
 	{
 		return GetInitialValue(InIndex).Get<T>();
 	}
 
+	FORCEINLINE_DEBUGGABLE void SetMinimumValue(const FName& InName, const FRigControlValue& InValue)
+	{
+		SetValue(InName, InValue, ERigControlValueType::Minimum);
+	}
+
+	FORCEINLINE_DEBUGGABLE void SetMinimumValue(int32 InIndex, const FRigControlValue& InValue)
+	{
+		SetValue(InIndex, InValue, ERigControlValueType::Minimum);
+	}
+
+	template<class T>
+	FORCEINLINE_DEBUGGABLE void SetMinimumValue(const FName& InName, T InValue)
+	{
+		SetMinimumValue(InName, FRigControlValue::Make<T>(InValue));
+	}
+
+	template<class T>
+	FORCEINLINE_DEBUGGABLE void SetMinimumValue(int32 InIndex, T InValue)
+	{
+		SetMinimumValue(InIndex, FRigControlValue::Make<T>(InValue));
+	}
+
+	FORCEINLINE_DEBUGGABLE FRigControlValue GetMinimumValue(const FName& InName) const
+	{
+		return GetValue(InName, ERigControlValueType::Minimum);
+	}
+
+	FORCEINLINE_DEBUGGABLE FRigControlValue GetMinimumValue(int32 InIndex) const
+	{
+		return GetValue(InIndex, ERigControlValueType::Minimum);
+	}
+
+	template<class T>
+	FORCEINLINE_DEBUGGABLE T GetMinimumValue(const FName& InName) const
+	{
+		return GetMinimumValue(InName).Get<T>();
+	}
+
+	template<class T>
+	FORCEINLINE_DEBUGGABLE T GetMinimumValue(int32 InIndex) const
+	{
+		return GetMinimumValue(InIndex).Get<T>();
+	}
+
+	FORCEINLINE_DEBUGGABLE void SetMaximumValue(const FName& InName, const FRigControlValue& InValue)
+	{
+		return SetValue(InName, InValue, ERigControlValueType::Maximum);
+	}
+
+	FORCEINLINE_DEBUGGABLE void SetMaximumValue(int32 InIndex, const FRigControlValue& InValue)
+	{
+		return SetValue(InIndex, InValue, ERigControlValueType::Maximum);
+	}
+
+	template<class T>
+	FORCEINLINE_DEBUGGABLE void SetMaximumValue(const FName& InName, T InValue)
+	{
+		SetMaximumValue(InName, FRigControlValue::Make<T>(InValue));
+	}
+
+	template<class T>
+	FORCEINLINE_DEBUGGABLE void SetMaximumValue(int32 InIndex, T InValue)
+	{
+		SetMaximumValue(InIndex, FRigControlValue::Make<T>(InValue));
+	}
+
+	FORCEINLINE_DEBUGGABLE FRigControlValue GetMaximumValue(const FName& InName) const
+	{
+		return GetValue(InName, ERigControlValueType::Maximum);
+	}
+
+	FORCEINLINE_DEBUGGABLE FRigControlValue GetMaximumValue(int32 InIndex) const
+	{
+		return GetValue(InIndex, ERigControlValueType::Maximum);
+	}
+
+	template<class T>
+	FORCEINLINE_DEBUGGABLE T GetMaximumValue(const FName& InName) const
+	{
+		return GetMaximumValue(InName).Get<T>();
+	}
+
+	template<class T>
+	FORCEINLINE_DEBUGGABLE T GetMaximumValue(int32 InIndex) const
+	{
+		return GetMaximumValue(InIndex).Get<T>();
+	}
+
 	// updates all of the internal caches
-	void Initialize();
+	void Initialize(bool bResetTransforms = true);
 
 	// clears the hierarchy and removes all content
 	void Reset();
@@ -373,4 +624,5 @@ private:
 	int32 GetChildrenRecursive(const int32 InIndex, TArray<int32>& OutChildren, bool bRecursively) const;
 
 	friend struct FRigHierarchyContainer;
+	friend class UControlRigHierarchyModifier;
 };
diff --git a/Engine/Plugins/Experimental/ControlRig/Source/ControlRig/Public/Rigs/RigCurveContainer.h b/Engine/Plugins/Experimental/ControlRig/Source/ControlRig/Public/Rigs/RigCurveContainer.h
--- a/Engine/Plugins/Experimental/ControlRig/Source/ControlRig/Public/Rigs/RigCurveContainer.h
+++ b/Engine/Plugins/Experimental/ControlRig/Source/ControlRig/Public/Rigs/RigCurveContainer.h
@@ -8,6 +8,7 @@
 #include "RigCurveContainer.generated.h"
 
 class UControlRig;
+class USkeleton;
 
 USTRUCT(BlueprintType)
 struct FRigCurve : public FRigElement
@@ -21,7 +22,7 @@ struct FRigCurve : public FRigElement
 	}
 	virtual ~FRigCurve() {}
 
-	UPROPERTY(VisibleAnywhere, Category = FRigElement)
+	UPROPERTY(BlueprintReadWrite, VisibleAnywhere, Category = FRigElement)
 	float Value;
 
 	FORCEINLINE virtual ERigElementType GetElementType() const override
@@ -30,7 +31,7 @@ struct FRigCurve : public FRigElement
 	}
 };
 
-USTRUCT()
+USTRUCT(BlueprintType)
 struct CONTROLRIG_API FRigCurveContainer
 {
 	GENERATED_BODY()
@@ -112,6 +113,8 @@ public:
 	FRigElementRenamed OnCurveRenamed;
 	FRigElementSelected OnCurveSelected;
 
+	TArray<FRigElementKey> ImportCurvesFromSkeleton(const USkeleton* InSkeleton, const FName& InNameSpace, bool bRemoveObsoleteCurves, bool bSelectCurves, bool bNotify);
+
 #endif
 
 private:
@@ -134,6 +137,9 @@ private:
 
 	void RefreshMapping();
 
+	bool bSuspendNotifications;
+
 	friend struct FRigHierarchyContainer;
+	friend class UControlRigHierarchyModifier;
 };
 
diff --git a/Engine/Plugins/Experimental/ControlRig/Source/ControlRig/Public/Rigs/RigHierarchyContainer.h b/Engine/Plugins/Experimental/ControlRig/Source/ControlRig/Public/Rigs/RigHierarchyContainer.h
--- a/Engine/Plugins/Experimental/ControlRig/Source/ControlRig/Public/Rigs/RigHierarchyContainer.h
+++ b/Engine/Plugins/Experimental/ControlRig/Source/ControlRig/Public/Rigs/RigHierarchyContainer.h
@@ -13,7 +13,7 @@
 class UControlRig;
 struct FRigHierarchyContainer;
 
-UENUM()
+UENUM(BlueprintType)
 enum class ERigHierarchyImportMode : uint8
 {
 	Append,
@@ -46,10 +46,11 @@ struct CONTROLRIG_API FRigHierarchyCopyPasteContent
 	TArray<FTransform> GlobalTransforms;
 };
 
-USTRUCT()
+USTRUCT(BlueprintType)
 struct CONTROLRIG_API FRigHierarchyContainer
 {
 public:
+
 	GENERATED_BODY()
 
 	FRigHierarchyContainer();
@@ -67,7 +68,7 @@ public:
 	UPROPERTY()
 	FRigCurveContainer CurveContainer;
 
-	void Initialize();
+	void Initialize(bool bResetTransforms = true);
 	void Reset();
 	void ResetTransforms();
 
diff --git a/Engine/Plugins/Experimental/ControlRig/Source/ControlRig/Public/Rigs/RigHierarchyDefines.h b/Engine/Plugins/Experimental/ControlRig/Source/ControlRig/Public/Rigs/RigHierarchyDefines.h
--- a/Engine/Plugins/Experimental/ControlRig/Source/ControlRig/Public/Rigs/RigHierarchyDefines.h
+++ b/Engine/Plugins/Experimental/ControlRig/Source/ControlRig/Public/Rigs/RigHierarchyDefines.h
@@ -11,7 +11,7 @@ struct FRigHierarchyContainer;
  * This is rig element types that we support
  * This can be used as a mask so supported as a bitfield
  */
-UENUM()
+UENUM(BlueprintType)
 enum class ERigElementType : uint8
 {
 	None UMETA(Hidden),
@@ -48,9 +48,15 @@ namespace FRigElementTypeHelper
 	}
 }
 
+USTRUCT(BlueprintType)
 struct FRigElementKey
 {
+	GENERATED_BODY()
+
+	UPROPERTY(BlueprintReadWrite, EditAnywhere, Category = "Hierarchy")
 	FName Name;
+
+	UPROPERTY(BlueprintReadWrite, EditAnywhere, Category = "Hierarchy")
 	ERigElementType Type;
 
 	FRigElementKey()
@@ -118,10 +124,10 @@ struct CONTROLRIG_API FRigElement
 	{}
 	virtual ~FRigElement() {}
 	
-	UPROPERTY(EditAnywhere, Category = FRigElement)
+	UPROPERTY(BlueprintReadOnly, EditAnywhere, Category = FRigElement)
 	FName Name;
 
-	UPROPERTY(VisibleAnywhere, Category = FRigElement)
+	UPROPERTY(BlueprintReadOnly, VisibleAnywhere, Category = FRigElement)
 	int32 Index;
 
 	virtual ERigElementType GetElementType() const
diff --git a/Engine/Plugins/Experimental/ControlRig/Source/ControlRig/Public/Rigs/RigSpaceHierarchy.h b/Engine/Plugins/Experimental/ControlRig/Source/ControlRig/Public/Rigs/RigSpaceHierarchy.h
--- a/Engine/Plugins/Experimental/ControlRig/Source/ControlRig/Public/Rigs/RigSpaceHierarchy.h
+++ b/Engine/Plugins/Experimental/ControlRig/Source/ControlRig/Public/Rigs/RigSpaceHierarchy.h
@@ -8,7 +8,7 @@
 
 class UControlRig;
 
-UENUM()
+UENUM(BlueprintType)
 enum class ERigSpaceType : uint8
 {
 	/** Not attached to anything */
@@ -24,7 +24,7 @@ enum class ERigSpaceType : uint8
 	Space
 };
 
-USTRUCT()
+USTRUCT(BlueprintType)
 struct CONTROLRIG_API FRigSpace : public FRigElement
 {
 	GENERATED_BODY()
@@ -40,19 +40,19 @@ struct CONTROLRIG_API FRigSpace : public FRigElement
 	}
 	virtual ~FRigSpace() {}
 
-	UPROPERTY(VisibleAnywhere, Category = FRigElement)
+	UPROPERTY(BlueprintReadWrite, VisibleAnywhere, Category = FRigElement)
 	ERigSpaceType SpaceType;
 
-	UPROPERTY(VisibleAnywhere, Category = FRigElement)
+	UPROPERTY(BlueprintReadOnly, VisibleAnywhere, Category = FRigElement)
 	FName ParentName;
 
-	UPROPERTY(transient)
+	UPROPERTY(BlueprintReadOnly, transient, Category = FRigElement)
 	int32 ParentIndex;
 
-	UPROPERTY(EditAnywhere, Category = FRigElement)
+	UPROPERTY(BlueprintReadOnly, EditAnywhere, Category = FRigElement)
 	FTransform InitialTransform;
 
-	UPROPERTY(transient, EditAnywhere, Category = FRigElement)
+	UPROPERTY(BlueprintReadOnly, transient, EditAnywhere, Category = FRigElement)
 	FTransform LocalTransform;
 
 	FORCEINLINE virtual ERigElementType GetElementType() const override
@@ -85,7 +85,7 @@ struct CONTROLRIG_API FRigSpace : public FRigElement
 	}
 };
 
-USTRUCT()
+USTRUCT(BlueprintType)
 struct CONTROLRIG_API FRigSpaceHierarchy
 {
 	GENERATED_BODY()
@@ -171,7 +171,7 @@ struct CONTROLRIG_API FRigSpaceHierarchy
 	FTransform GetInitialGlobalTransform(int32 InIndex) const;
 
 	// updates all of the internal caches
-	void Initialize();
+	void Initialize(bool bResetTransforms = true);
 
 	// clears the hierarchy and removes all content
 	void Reset();
@@ -223,4 +223,5 @@ private:
 	void RefreshMapping();
 
 	friend struct FRigHierarchyContainer;
+	friend class UControlRigHierarchyModifier;
 };
diff --git a/Engine/Plugins/Experimental/ControlRig/Source/ControlRig/Public/Sequencer/ControlRigLayerInstance.h b/Engine/Plugins/Experimental/ControlRig/Source/ControlRig/Public/Sequencer/ControlRigLayerInstance.h
new file mode 100644
--- /dev/null
+++ b/Engine/Plugins/Experimental/ControlRig/Source/ControlRig/Public/Sequencer/ControlRigLayerInstance.h
@@ -0,0 +1,55 @@
+// Copyright Epic Games, Inc. All Rights Reserved.
+
+/**
+ *
+ * This Instance only contains one AnimationAsset, and produce poses
+ * Used by Preview in AnimGraph, Playing single animation in Kismet2 and etc
+ */
+
+#pragma once
+#include "Animation/AnimInstance.h"
+#include "SequencerAnimationSupport.h"
+#include "AnimNode_ControlRigBase.h"
+#include "ControlRigLayerInstance.generated.h"
+
+class UControlRig;
+
+UCLASS(transient, NotBlueprintable)
+class CONTROLRIG_API UControlRigLayerInstance : public UAnimInstance, public ISequencerAnimationSupport
+{
+	GENERATED_UCLASS_BODY()
+
+public:
+	/** ControlRig related support */
+	void AddControlRigTrack(int32 ControlRigID, UControlRig* InControlRig);
+	void UpdateControlRigTrack(int32 ControlRigID, float Weight, const FControlRigIOSettings& InputSettings, bool bExecute);
+	void RemoveControlRigTrack(int32 ControlRigID);
+	void ResetControlRigTracks();
+
+	/** Sequencer AnimInstance Interface */
+	void AddAnimation(int32 SequenceId, UAnimSequenceBase* InAnimSequence);
+	virtual void UpdateAnimTrack(UAnimSequenceBase* InAnimSequence, int32 SequenceId, float InPosition, float Weight, bool bFireNotifies) override;
+	virtual void UpdateAnimTrack(UAnimSequenceBase* InAnimSequence, int32 SequenceId, float InFromPosition, float InToPosition, float Weight, bool bFireNotifies) override;
+	void RemoveAnimation(int32 SequenceId);
+
+	/** Reset all nodes in this instance */
+	virtual void ResetNodes() override;
+	/** Reset the pose in this instance*/
+	virtual void ResetPose() override;
+	/** Saved the named pose to restore after */
+	virtual void SavePose() override;
+
+	/** Return the first available control rig */
+	UControlRig* GetFirstAvailableControlRig() const;
+
+	virtual UAnimInstance* GetSourceAnimInstance() override;
+	virtual void SetSourceAnimInstance(UAnimInstance* SourceAnimInstance) override;
+	virtual bool DoesSupportDifferentSourceAnimInstance() const override { return true; }
+
+protected:
+	// UAnimInstance interface
+	virtual FAnimInstanceProxy* CreateAnimInstanceProxy() override;
+
+public:
+	static const FName SequencerPoseName;
+};
diff --git a/Engine/Plugins/Experimental/ControlRig/Source/ControlRig/Public/Sequencer/ControlRigSequencerAnimInstance.h b/Engine/Plugins/Experimental/ControlRig/Source/ControlRig/Public/Sequencer/ControlRigSequencerAnimInstance.h
--- a/Engine/Plugins/Experimental/ControlRig/Source/ControlRig/Public/Sequencer/ControlRigSequencerAnimInstance.h
+++ b/Engine/Plugins/Experimental/ControlRig/Source/ControlRig/Public/Sequencer/ControlRigSequencerAnimInstance.h
@@ -3,6 +3,7 @@
 #pragma once
 
 #include "AnimSequencerInstance.h"
+#include "AnimNode_ControlRigBase.h"
 #include "ControlRigSequencerAnimInstance.generated.h"
 
 class UControlRig;
@@ -15,7 +16,7 @@ class CONTROLRIG_API UControlRigSequencerAnimInstance : public UAnimSequencerIns
 
 public:
 	/** Update an animation ControlRig in this sequence */
-	bool UpdateControlRig(UControlRig* InControlRig, uint32 SequenceId, bool bAdditive, bool bApplyBoneFilter, const FInputBlendPose& BoneFilter, float Weight, bool bUpdateInput, bool bExecute);
+	bool UpdateControlRig(UControlRig* InControlRig, uint32 SequenceId, bool bAdditive, bool bApplyBoneFilter, const FInputBlendPose& BoneFilter, float Weight, const FControlRigIOSettings& InputSettings, bool bExecute);
 
 	virtual bool SetAnimationAsset(class UAnimationAsset* NewAsset);
 
diff --git a/Engine/Plugins/Experimental/ControlRig/Source/ControlRig/Public/Sequencer/ControlRigSortedControls.h b/Engine/Plugins/Experimental/ControlRig/Source/ControlRig/Public/Sequencer/ControlRigSortedControls.h
new file mode 100644
--- /dev/null
+++ b/Engine/Plugins/Experimental/ControlRig/Source/ControlRig/Public/Sequencer/ControlRigSortedControls.h
@@ -0,0 +1,40 @@
+// Copyright Epic Games, Inc. All Rights Reserved.
+
+#pragma once
+#include "CoreMinimal.h"
+#include "Manipulatable/IControlRigManipulatable.h"
+
+/**
+*  Class to hold sorted order of control/space elements so we sort them correctly in sequencer
+*/
+class CONTROLRIG_API FRigControlTreeElement : public TSharedFromThis<FRigControlTreeElement>
+{
+public:
+	FRigControlTreeElement(const FRigElementKey& InKey, int32 InIndex) : Key(InKey), Index(InIndex) {};
+public:
+	FRigElementKey Key;
+	int32 Index; //location in hierarchy array
+	TArray<TSharedPtr<FRigControlTreeElement>> Children;
+
+};
+
+/**
+* Functions to set up the control rig controls so they are in hierarchical order
+*/
+class CONTROLRIG_API FControlRigSortedControls
+{
+public:
+	//Given an Array of Available Controls, return a set of them sorted.
+	static void GetControlsInOrder(IControlRigManipulatable* ControlRig, TArray<FRigControl>& SortedControls);
+private:
+	//Helper functions
+	static void AddElement(FRigElementKey InKey, int32 InIndex, FRigElementKey InParentKey, TMap<FRigElementKey, TSharedPtr<FRigControlTreeElement>>& ElementMap,
+		TArray<TSharedPtr<FRigControlTreeElement>>& RootElements);
+	static void AddSpaceElement(FRigSpace InSpace, int32 InIndex, const TArray<FRigControl>& Controls, const TArray<FRigSpace>& Spaces,
+		TMap<FRigElementKey, TSharedPtr<FRigControlTreeElement>>& ElementMap, TArray<TSharedPtr<FRigControlTreeElement>>& RootElements);
+	static void AddControlElement(FRigControl InControl, int32 InIndex, const TArray<FRigControl>& Controls, const TArray<FRigSpace>& Spaces,
+		TMap<FRigElementKey, TSharedPtr<FRigControlTreeElement>>& ElementMap,
+		TArray<TSharedPtr<FRigControlTreeElement>>& RootElements);
+	static void CreateRigControlsRecursive(TSharedPtr<FRigControlTreeElement>& Element, const TArray<FRigControl>& Controls, TArray<FRigControl>& OutControls);
+
+};
diff --git a/Engine/Plugins/Experimental/ControlRig/Source/ControlRig/Public/Sequencer/MovieSceneControlRigParameterSection.h b/Engine/Plugins/Experimental/ControlRig/Source/ControlRig/Public/Sequencer/MovieSceneControlRigParameterSection.h
--- a/Engine/Plugins/Experimental/ControlRig/Source/ControlRig/Public/Sequencer/MovieSceneControlRigParameterSection.h
+++ b/Engine/Plugins/Experimental/ControlRig/Source/ControlRig/Public/Sequencer/MovieSceneControlRigParameterSection.h
@@ -24,6 +24,17 @@ struct FFloatInterrogationData
 	FName ParameterName;
 };
 
+struct FBoolInterrogationData
+{
+	bool Val;
+	FName ParameterName;
+};
+
+struct FVector2DInterrogationData
+{
+	FVector2D Val;
+	FName ParameterName;
+};
 
 struct FVectorInterrogationData
 {
@@ -37,6 +48,19 @@ struct FTransformInterrogationData
 	FName ParameterName;
 };
 
+USTRUCT()
+struct CONTROLRIG_API FChannelMapInfo
+{
+	GENERATED_USTRUCT_BODY()
+
+	FChannelMapInfo() = default;
+
+	FChannelMapInfo(int32 InControlIndex, int32 InChannelIndex) : ControlIndex(InControlIndex), ChannelIndex(InChannelIndex) {};
+	UPROPERTY()
+	int32 ControlIndex;
+	UPROPERTY()
+	int32 ChannelIndex;
+};
 
 /**
  * Movie scene section that controls animation controller animation
@@ -76,6 +100,14 @@ public:
 	UPROPERTY()
 	FMovieSceneFloatChannel Weight;
 
+	/** Map from the control name to where it starts as a channel*/
+	UPROPERTY()
+	TMap<FName, FChannelMapInfo> ControlChannelMap;
+
+private:
+	/** Copy of Mask for controls, checked when reconstructing*/
+	TArray<bool> OldControlsMask;
+
 public:
 
 	UMovieSceneControlRigParameterSection();
@@ -97,7 +129,7 @@ public:
 	void SetControlsMask(const TArray<bool>& InMask)
 	{
 		ControlsMask = InMask;
-		ReconstructChannelProxy();
+		ReconstructChannelProxy(true);
 	}
 
 	void SetControlsMask(int32 Index, bool Val)
@@ -106,14 +138,14 @@ public:
 		{
 			ControlsMask[Index] = Val;
 		}
-		ReconstructChannelProxy();
+		ReconstructChannelProxy(true);
 	}
 
 
 	void FillControlsMask(bool Val)
 	{
 		ControlsMask.Init(Val, ControlsMask.Num());
-		ReconstructChannelProxy();
+		ReconstructChannelProxy(true);
 
 	}
 
@@ -131,7 +163,7 @@ public:
 	void SetTransformMask(FMovieSceneTransformMask NewMask)
 	{
 		TransformMask = NewMask;
-		ReconstructChannelProxy();
+		ReconstructChannelProxy(true);
 	}
 
 public:
@@ -144,44 +176,57 @@ public:
 	/**  Whether or not this section his scalar*/
 	bool HasScalarParameter(FName InParameterName) const;
 
+	/**  Whether or not this section his bool*/
+	bool HasBoolParameter(FName InParameterName) const;
+
+	/**  Whether or not this section his scalar*/
+	bool HasVector2DParameter(FName InParameterName) const;
+
 	/**  Whether or not this section his scalar*/
 	bool HasVectorParameter(FName InParameterName) const;
 
 	/**  Whether or not this section his scalar*/
 	bool HasColorParameter(FName InParameterName) const;
 
-		/**  Whether or not this section his scalar*/
+	/**  Whether or not this section his scalar*/
 	bool HasTransformParameter(FName InParameterName) const;
 
 	/** Adds specified scalar parameter. */
-	void AddScalarParameter(FName InParameterName,  TOptional<float> DefaultValue);
+	void AddScalarParameter(FName InParameterName,  TOptional<float> DefaultValue, bool bReconstructChannel);
+
+	/** Adds specified bool parameter. */
+	void AddBoolParameter(FName InParameterName, TOptional<bool> DefaultValue, bool bReconstructChannel);
 
 	/** Adds a a key for a specific vector parameter. */
-	void AddVectorParameter(FName InParameterName, TOptional<FVector> DefaultValue);
+	void AddVectorParameter(FName InParameterName, TOptional<FVector> DefaultValue, bool bReconstructChannel);
+
+	/** Adds a a key for a specific vector2D parameter. */
+	void AddVector2DParameter(FName InParameterName, TOptional<FVector2D> DefaultValue, bool bReconstructChannel);
 
 	/** Adds a a key for a specific color parameter. */
-	void AddColorParameter(FName InParameterName, TOptional<FLinearColor> DefaultValue);
+	void AddColorParameter(FName InParameterName, TOptional<FLinearColor> DefaultValue, bool bReconstructChannel);
 
 	/** Adds a a key for a specific color parameter*/
-	void AddTransformParameter(FName InParameterName, TOptional<FTransform> DefaultValue);
+	void AddTransformParameter(FName InParameterName, TOptional<FTransform> DefaultValue, bool bReconstructChannel);
 
 public:
 	/**
 	* Access the interrogation key for control rig data 
 	*/
 	 static FMovieSceneInterrogationKey GetFloatInterrogationKey();
+	 static FMovieSceneInterrogationKey GetVector2DInterrogationKey();
+	 static FMovieSceneInterrogationKey GetVector4InterrogationKey();
 	 static FMovieSceneInterrogationKey GetVectorInterrogationKey();
 	 static FMovieSceneInterrogationKey GetTransformInterrogationKey();
 
-
-protected:
-	virtual void ReconstructChannelProxy();
+	virtual void ReconstructChannelProxy(bool bForce ) override;
 
 protected:
 
 	//~ UMovieSceneSection interface
 	virtual void Serialize(FArchive& Ar) override;
 	virtual void PostEditImport() override;
+	virtual void PostLoad() override;
 
 	// When true we do not set a key on the section, since it will be set because we changed the value
 	// We need this because control rig notifications are set on every change even when just changing sequencer time
diff --git a/Engine/Plugins/Experimental/ControlRig/Source/ControlRig/Public/Sequencer/MovieSceneControlRigParameterTrack.h b/Engine/Plugins/Experimental/ControlRig/Source/ControlRig/Public/Sequencer/MovieSceneControlRigParameterTrack.h
--- a/Engine/Plugins/Experimental/ControlRig/Source/ControlRig/Public/Sequencer/MovieSceneControlRigParameterTrack.h
+++ b/Engine/Plugins/Experimental/ControlRig/Source/ControlRig/Public/Sequencer/MovieSceneControlRigParameterTrack.h
@@ -35,7 +35,8 @@ public:
 	virtual bool IsEmpty() const override;
 	virtual const TArray<UMovieSceneSection*>& GetAllSections() const override;
 	virtual FName GetTrackName() const override { return TrackName; }
-
+	// UObject
+	virtual void PostLoad() override;
 #if WITH_EDITORONLY_DATA
 	virtual FText GetDefaultDisplayName() const override;
 #endif
diff --git a/Engine/Plugins/Experimental/ControlRig/Source/ControlRig/Public/Units/Control/RigUnit_Control_StaticMesh.h b/Engine/Plugins/Experimental/ControlRig/Source/ControlRig/Public/Units/Control/RigUnit_Control_StaticMesh.h
--- a/Engine/Plugins/Experimental/ControlRig/Source/ControlRig/Public/Units/Control/RigUnit_Control_StaticMesh.h
+++ b/Engine/Plugins/Experimental/ControlRig/Source/ControlRig/Public/Units/Control/RigUnit_Control_StaticMesh.h
@@ -16,17 +16,10 @@ struct CONTROLRIG_API FRigUnit_Control_StaticMesh : public FRigUnit_Control
 
 	FRigUnit_Control_StaticMesh();
 
-#if WITH_EDITORONLY_DATA
-	/** The static mesh to use to display this control */
-	UPROPERTY(meta=(Input))
-	UStaticMesh* StaticMesh;
-
-	/** The override materials we use to display this control */
-	UPROPERTY(meta=(Input))
-	TArray<UMaterialInterface*> Materials;
+	RIGVM_METHOD()
+	virtual void Execute(const FRigUnitContext& Context) override;
 
 	/** The the transform the mesh will be rendered with (applied on top of the control's transform in the viewport) */
 	UPROPERTY(meta=(Input))
 	FTransform MeshTransform;
-#endif
 };
\ No newline at end of file
diff --git a/Engine/Plugins/Experimental/ControlRig/Source/ControlRig/Public/Units/Deprecated/RigUnit_TwoBoneIKFK.h b/Engine/Plugins/Experimental/ControlRig/Source/ControlRig/Public/Units/Deprecated/RigUnit_TwoBoneIKFK.h
--- a/Engine/Plugins/Experimental/ControlRig/Source/ControlRig/Public/Units/Deprecated/RigUnit_TwoBoneIKFK.h
+++ b/Engine/Plugins/Experimental/ControlRig/Source/ControlRig/Public/Units/Deprecated/RigUnit_TwoBoneIKFK.h
@@ -16,6 +16,7 @@ struct CONTROLRIG_API FRigUnit_TwoBoneIKFK : public FRigUnitMutable
 {
 	GENERATED_BODY()
 
+	RIGVM_METHOD()
 	virtual void Execute(const FRigUnitContext& Context) override;
 
 	FRigUnit_TwoBoneIKFK()
@@ -52,33 +53,44 @@ struct CONTROLRIG_API FRigUnit_TwoBoneIKFK : public FRigUnitMutable
 	UPROPERTY(EditAnywhere, Category = "TwoBoneIKFK", meta = (Input, ClampMin = "0", AllowSourceAccess))
 	float IKBlend; 
 
+private:
+
 	// Transform : Start Joint FK Transform         # Transform for the Start Joint when in FK mode
 	// Transform: Mid Joint FK Transform           # Transform for the Mid Joint when in FK mode
 	// Transform : End Joint FK Transform          # Transform for the End Joint when in FK mode
-	UPROPERTY(EditAnywhere, Category = "TwoBoneIKFK", meta = (Input, AllowSourceAccess))
+	UPROPERTY(EditAnywhere, Category = "TwoBoneIKFK", meta = (AllowSourceAccess))
 	FTransform StartJointFKTransform;
 
-	UPROPERTY(EditAnywhere, Category = "TwoBoneIKFK", meta = (Input, AllowSourceAccess))
+	UPROPERTY(EditAnywhere, Category = "TwoBoneIKFK", meta = (AllowSourceAccess))
 	FTransform MidJointFKTransform;
 
-	UPROPERTY(EditAnywhere, Category = "TwoBoneIKFK", meta = (Input, AllowSourceAccess))
+	UPROPERTY(EditAnywhere, Category = "TwoBoneIKFK", meta = (AllowSourceAccess))
 	FTransform EndJointFKTransform;
 
-private:
+	UPROPERTY(transient)
 	float PreviousFKIKBlend;
 
+	UPROPERTY(transient)
 	FTransform StartJointIKTransform;
+	
+	UPROPERTY(transient)
 	FTransform MidJointIKTransform;
+
+	UPROPERTY(transient)
 	FTransform EndJointIKTransform;
 
+	UPROPERTY(transient)
 	int32 StartJointIndex;
+
+	UPROPERTY(transient)
 	int32 MidJointIndex;
+
+	UPROPERTY(transient)
 	int32 EndJointIndex;
 
+	UPROPERTY(transient)
 	float UpperLimbLength;
+	
+	UPROPERTY(transient)
 	float LowerLimbLength;
-
-	void SolveIK();
-
-	friend class URigUnitEditor_TwoBoneIKFK;
 };
diff --git a/Engine/Plugins/Experimental/ControlRig/Source/ControlRig/Public/Units/Hierarchy/RigUnit_AddBoneTransform.h b/Engine/Plugins/Experimental/ControlRig/Source/ControlRig/Public/Units/Hierarchy/RigUnit_AddBoneTransform.h
new file mode 100644
--- /dev/null
+++ b/Engine/Plugins/Experimental/ControlRig/Source/ControlRig/Public/Units/Hierarchy/RigUnit_AddBoneTransform.h
@@ -0,0 +1,75 @@
+// Copyright Epic Games, Inc. All Rights Reserved.
+
+#pragma once
+
+#include "Units/RigUnit.h"
+#include "RigUnit_AddBoneTransform.generated.h"
+
+
+/**
+ * Offset Transform is used to perform a change in the hierarchy by setting a single bone's transform.
+ */
+USTRUCT(meta=(DisplayName="Offset Transform", Category="Hierarchy", DocumentationPolicy="Strict", Keywords="Offset,AddToBoneTransform"))
+struct FRigUnit_AddBoneTransform : public FRigUnitMutable
+{
+	GENERATED_BODY()
+
+	FRigUnit_AddBoneTransform()
+		: Weight(1.f)
+		, bPostMultiply(false)
+		, bPropagateToChildren(false)
+		, CachedBoneIndex(INDEX_NONE)
+	{}
+
+	virtual FString GetUnitLabel() const override;
+
+	virtual FName DetermineSpaceForPin(const FString& InPinPath, void* InUserContext) const override
+	{
+		if (InPinPath.StartsWith(TEXT("Transform")))
+		{
+			return Bone;
+		}
+		return NAME_None;
+	}
+
+	RIGVM_METHOD()
+	virtual void Execute(const FRigUnitContext& Context) override;
+
+	/**
+	 * The name of the Bone to set the transform for.
+	 */
+	UPROPERTY(meta = (Input, CustomWidget = "BoneName", Constant))
+	FName Bone;
+
+	/**
+	 * The transform value to set for the given Bone.
+	 */
+	UPROPERTY(meta = (Input))
+	FTransform Transform;
+
+	/**
+	 * The weight of the change - how much the change should be applied
+	 */
+	UPROPERTY(meta = (Input, UIMin = "0.0", UIMax = "1.0"))
+	float Weight;
+
+	/**
+	 * If set to true the transform will be post multiplied, otherwise pre multiplied.
+	 * Post multiplying means that the transform is understood as a parent space change,
+	 * while pre multiplying means that the transform is understood as a child space change.
+	 */
+	UPROPERTY(meta = (Input))
+	bool bPostMultiply;
+
+	/**
+	 * If set to true all of the global transforms of the children 
+	 * of this bone will be recalculated based on their local transforms.
+	 * Note: This is computationally more expensive than turning it off.
+	 */
+	UPROPERTY(meta = (Input, Constant))
+	bool bPropagateToChildren;
+
+	// Used to cache the internally used bone index
+	UPROPERTY()
+	int32 CachedBoneIndex;
+};
diff --git a/Engine/Plugins/Experimental/ControlRig/Source/ControlRig/Public/Units/RigUnit.h b/Engine/Plugins/Experimental/ControlRig/Source/ControlRig/Public/Units/RigUnit.h
--- a/Engine/Plugins/Experimental/ControlRig/Source/ControlRig/Public/Units/RigUnit.h
+++ b/Engine/Plugins/Experimental/ControlRig/Source/ControlRig/Public/Units/RigUnit.h
@@ -3,32 +3,19 @@
 #pragma once
 
 #include "ControlRigDefines.h"
+#include "RigVMCore/RigVMStruct.h"
 #include "RigVMCore/RigVMRegistry.h"
 #include "RigUnit.generated.h"
 
 struct FRigUnitContext;
 
-/** 
- * Current state of execution type
- */
-UENUM()
-enum class EUnitExecutionType : uint8
-{
-	Always,
-	InEditingTime, // in control rig editor
-	Disable, // disable completely - good for debugging
-	Initialize, // only during init
-	Max UMETA(Hidden),
-};
-
 /** Base class for all rig units */
 USTRUCT(BlueprintType, meta=(Abstract, NodeColor = "0.1 0.1 0.1"))
-struct CONTROLRIG_API FRigUnit
+struct CONTROLRIG_API FRigUnit : public FRigVMStruct
 {
 	GENERATED_BODY()
 
 	FRigUnit()
-		:ExecutionType(EUnitExecutionType::Always)
 	{}
 
 	/** Virtual destructor */
@@ -39,21 +26,6 @@ struct CONTROLRIG_API FRigUnit
 
 	/** Execute logic for this rig unit */
 	virtual void Execute(const FRigUnitContext& Context) {}
-
-	/* 
-	 * This is property name given by ControlRig as transient when initialized, so only available in run-time
-	 */
-	UPROPERTY(BlueprintReadOnly, transient, Category=FRigUnit, meta = (Constant))
-	FName RigUnitName;
-	
-	/* 
-	 * This is struct name given by ControlRig as transient when initialized, so only available in run-time
-	 */
-	UPROPERTY(BlueprintReadOnly, transient, Category=FRigUnit, meta = (Constant))
-	FName RigUnitStructName;
-
-	UPROPERTY(EditAnywhere, Category = FRigUnit, meta = (Constant))
-	EUnitExecutionType ExecutionType;
 };
 
 /** Base class for all rig units that can change data */
@@ -72,11 +44,3 @@ struct CONTROLRIG_API FRigUnitMutable : public FRigUnit
 	UPROPERTY(DisplayName = "Execute", Transient, meta = (Input, Output))
 	FControlRigExecuteContext ExecuteContext;
 };
-
-// this will have to change in the future and move to editor, I assume the errors will be saved in the rig unit and it will print fromthe editor module
-namespace UnitLogHelpers
-{
-	CONTROLRIG_API void PrintMissingHierarchy(const FName& InputName);
-	CONTROLRIG_API void PrintUnimplemented(const FName& InputName);
-}
- 
\ No newline at end of file
diff --git a/Engine/Plugins/Experimental/ControlRig/Source/ControlRig/Public/Units/RigUnitContext.h b/Engine/Plugins/Experimental/ControlRig/Source/ControlRig/Public/Units/RigUnitContext.h
--- a/Engine/Plugins/Experimental/ControlRig/Source/ControlRig/Public/Units/RigUnitContext.h
+++ b/Engine/Plugins/Experimental/ControlRig/Source/ControlRig/Public/Units/RigUnitContext.h
@@ -26,6 +26,7 @@ struct FRigUnitContext
 	/** default constructor */
 	FRigUnitContext()
 		: DrawInterface(nullptr)
+		, DrawContainer(nullptr)
 		, DataSourceRegistry(nullptr)
 		, DeltaTime(0.f)
 		, State(EControlRigState::Invalid)
@@ -34,12 +35,14 @@ struct FRigUnitContext
 		, Log(nullptr)
 #endif
 	{
-
 	}
 
 	/** The draw interface for the units to use */
 	FControlRigDrawInterface* DrawInterface;
 
+	/** The auxiliary draw container for the units to use */
+	FControlRigDrawContainer* DrawContainer;
+
 	/** The registry to access data source */
 	const UAnimationDataSourceRegistry* DataSourceRegistry;
 
@@ -114,7 +117,7 @@ struct FRigUnitContext
 #define UE_CONTROLRIG_RIGUNIT_REPORT(Severity, Format, ...) \
 if(Context.Log != nullptr) \
 { \
-	Context.Log->Report(EMessageSeverity::Severity, RigUnitName, FString::Printf((Format), ##__VA_ARGS__)); \
+	Context.Log->Report(EMessageSeverity::Severity, RigVMOperatorName, RigVMOperatorIndex, FString::Printf((Format), ##__VA_ARGS__)); \
 }
 #define UE_CONTROLRIG_RIGUNIT_LOG_MESSAGE(Format, ...) UE_CONTROLRIG_RIGUNIT_REPORT(Info, (Format), ##__VA_ARGS__)
 #define UE_CONTROLRIG_RIGUNIT_REPORT_WARNING(Format, ...) UE_CONTROLRIG_RIGUNIT_REPORT(Warning, (Format), ##__VA_ARGS__)
diff --git a/Engine/Plugins/Experimental/ControlRig/Source/ControlRigDeveloper/ControlRigDeveloper.Build.cs b/Engine/Plugins/Experimental/ControlRig/Source/ControlRigDeveloper/ControlRigDeveloper.Build.cs
--- a/Engine/Plugins/Experimental/ControlRig/Source/ControlRigDeveloper/ControlRigDeveloper.Build.cs
+++ b/Engine/Plugins/Experimental/ControlRig/Source/ControlRigDeveloper/ControlRigDeveloper.Build.cs
@@ -28,8 +28,9 @@ namespace UnrealBuildTool.Rules
                     "SlateCore",
                     "InputCore",
                     "TimeManagement",
-					"Persona",
+					"EditorWidgets",
 					"MessageLog",
+                    "RigVM",
                 }
             );
 
@@ -50,6 +51,7 @@ namespace UnrealBuildTool.Rules
                         "AnimGraph",
                         "BlueprintGraph",
                         "PropertyEditor",
+                        "RigVMDeveloper",
                     }
                 );
 
diff --git a/Engine/Plugins/Experimental/ControlRig/Source/ControlRigDeveloper/Private/AnimGraphNode_ControlRig.cpp b/Engine/Plugins/Experimental/ControlRig/Source/ControlRigDeveloper/Private/AnimGraphNode_ControlRig.cpp
--- a/Engine/Plugins/Experimental/ControlRig/Source/ControlRigDeveloper/Private/AnimGraphNode_ControlRig.cpp
+++ b/Engine/Plugins/Experimental/ControlRig/Source/ControlRigDeveloper/Private/AnimGraphNode_ControlRig.cpp
@@ -7,12 +7,15 @@
 #include "Widgets/Text/STextBlock.h"
 #include "Widgets/Input/SCheckBox.h"
 #include "DetailWidgetRow.h"
-#include "ControlRigVariables.h"
 #include "SVariableMappingWidget.h"
 #include "ScopedTransaction.h"
 #include "Kismet2/BlueprintEditorUtils.h"
+#include "AnimationGraphSchema.h"
+#include "ControlRigBlueprintGeneratedClass.h"
 #include "ControlRigBlueprint.h"
 
+PRAGMA_DISABLE_OPTIMIZATION
+
 #define LOCTEXT_NAMESPACE "AnimGraphNode_ControlRig"
 
 UAnimGraphNode_ControlRig::UAnimGraphNode_ControlRig(const FObjectInitializer& ObjectInitializer)
@@ -32,43 +35,99 @@ FText UAnimGraphNode_ControlRig::GetTooltipText() const
 	return LOCTEXT("AnimGraphNode_ControlRig_Tooltip", "Evaluates a control rig");
 }
 
-void UAnimGraphNode_ControlRig::GetExposableProperties(TArray<FProperty*>& OutExposableProperties) const
+void UAnimGraphNode_ControlRig::ReallocatePinsDuringReconstruction(TArray<UEdGraphPin*>& OldPins)
 {
 	DECLARE_SCOPE_HIERARCHICAL_COUNTER_FUNC()
 
-	// we only need inputs
-	TMap<FName, FControlRigIOVariable> LocalInputVariables;
-	GetIOProperties(true, LocalInputVariables);
+	// we do this to refresh input variables
+	RebuildExposedProperties();
+	
+	// skip the parent class's implementation and go directly to UAnimGraphNode_Base
+	UAnimGraphNode_Base::ReallocatePinsDuringReconstruction(OldPins);
 
-	UClass* TargetClass = GetTargetClass();
-	// if we have target class, see if we can get 
-	if (TargetClass)
+	// Grab the SKELETON class here as when we are reconstructed during during BP compilation
+	// the full generated class is not yet present built.
+	UClass* TargetClass = GetTargetSkeletonClass();
+
+	if (!TargetClass)
 	{
-		UControlRig* ControlRig = TargetClass->GetDefaultObject<UControlRig>();
-		if (ControlRig)
+		// Nothing to search for properties
+		return;
+	}
+
+	// Need the schema to extract pin types
+	const UEdGraphSchema_K2* Schema = CastChecked<UEdGraphSchema_K2>(GetSchema());
+
+	// Default anim schema for util funcions
+	const UAnimationGraphSchema* AnimGraphDefaultSchema = GetDefault<UAnimationGraphSchema>();
+
+	// We'll track the names we encounter by removing from this list, if anything remains the properties
+	// have been removed from the target class and we should remove them too
+	TArray<FName> BeginExposableNames = KnownExposableProperties;
+
+	for (const TPair<FName, FRigVMParameter>& ParameterPair : InputParameters)
+	{
+		FName PropertyName = ParameterPair.Key;
+		BeginExposableNames.Remove(PropertyName);
+
+		if (!KnownExposableProperties.Contains(PropertyName))
 		{
-			for (auto Iter = LocalInputVariables.CreateIterator(); Iter; ++Iter)
+			// New property added to the target class
+			KnownExposableProperties.Add(PropertyName);
+		}
+
+		if (ExposedPropertyNames.Contains(PropertyName))
+		{
+			const FRigVMParameter& Parameter = ParameterPair.Value;
+
+			FEdGraphPinType PinType;
+
+			if (Parameter.GetScriptStruct())
 			{
-				FControlRigIOVariable& Var = Iter.Value();
-				FCachedPropertyPath CachePath;
-				ControlRig->GetInOutPropertyPath(true, Iter.Key(), CachePath);
-				if (CachePath.IsResolved())
-				{
-					OutExposableProperties.Add(CachePath.GetFProperty());
-				}
+				PinType.PinCategory = UEdGraphSchema_K2::PC_Struct;
+				PinType.PinSubCategoryObject = Parameter.GetScriptStruct();
+			}
+			else if (Parameter.GetCPPType() == TEXT("bool"))
+			{
+				PinType.PinCategory = UEdGraphSchema_K2::PC_Boolean;
+			}
+			else if (Parameter.GetCPPType() == TEXT("float"))
+			{
+				PinType.PinCategory = UEdGraphSchema_K2::PC_Float;
+			}
+			else if (Parameter.GetCPPType() == TEXT("int32"))
+			{
+				PinType.PinCategory = UEdGraphSchema_K2::PC_Int;
+			}
+			else if (Parameter.GetCPPType() == TEXT("FName"))
+			{
+				PinType.PinCategory = UEdGraphSchema_K2::PC_Name;
+			}
+			else if (Parameter.GetCPPType() == TEXT("FString"))
+			{
+				PinType.PinCategory = UEdGraphSchema_K2::PC_String;
+			}
+			else
+			{
+				ensure(false);
 			}
+
+			UEdGraphPin* NewPin = CreatePin(EEdGraphPinDirection::EGPD_Input, PinType, PropertyName);
+			NewPin->PinFriendlyName = FText::FromName(PropertyName);
+
+			// We cant interrogate CDO here as we may be mid-compile, so we can only really
+			// reset to the autogenerated default.
+			AnimGraphDefaultSchema->ResetPinToAutogeneratedDefaultValue(NewPin, false);
+
+			CustomizePinData(NewPin, PropertyName, INDEX_NONE);
 		}
 	}
-}
 
-void UAnimGraphNode_ControlRig::ReallocatePinsDuringReconstruction(TArray<UEdGraphPin*>& OldPins)
-{
-	DECLARE_SCOPE_HIERARCHICAL_COUNTER_FUNC()
-
-	// we do this to refresh input variables
-	RebuildExposedProperties();
-	// we avoid CustomProperty, it only allows "the direct child"
-	Super::ReallocatePinsDuringReconstruction(OldPins);
+	// Remove any properties that no longer exist on the target class
+	for (FName& RemovedPropertyName : BeginExposableNames)
+	{
+		KnownExposableProperties.Remove(RemovedPropertyName);
+	}
 }
 
 void UAnimGraphNode_ControlRig::ValidateAnimNodeDuringCompilation(USkeleton* ForSkeleton, FCompilerResultsLog& MessageLog)
@@ -113,13 +172,13 @@ void UAnimGraphNode_ControlRig::RebuildExposedProperties()
 	KnownExposableProperties.Empty();
 
 	// go through exposed properties, and clean up
-	GetIOProperties(true, InputVariables);
+	GetIOParameters(true, InputParameters);
 	// we still update OUtputvariables
 	// we don't want output to be exposed
-	GetIOProperties(false, OutputVariables);
+	GetIOParameters(false, OutputParameters);
 
 	// clear IO variables that don't exist anymore
-	auto ClearInvalidMapping = [](const TMap<FName, FControlRigIOVariable>& InVariables, TMap<FName, FName>& InOutMapping)
+	auto ClearInvalidMapping = [](const TMap<FName, FRigVMParameter>& InParameters, TMap<FName, FName>& InOutMapping)
 	{
 		TArray<FName> KeyArray;
 		InOutMapping.GenerateKeyArray(KeyArray);
@@ -127,17 +186,17 @@ void UAnimGraphNode_ControlRig::RebuildExposedProperties()
 		for (int32 Index=0; Index<KeyArray.Num(); ++Index)
 		{
 			// if this input doesn't exist anymore
-			if (!InVariables.Contains(KeyArray[Index]))
+			if (!InParameters.Contains(KeyArray[Index]))
 			{
 				InOutMapping.Remove(KeyArray[Index]);
 			}
 		}
 	};
 
-	ClearInvalidMapping(InputVariables, Node.InputMapping);
-	ClearInvalidMapping(OutputVariables, Node.OutputMapping);
+	ClearInvalidMapping(InputParameters, Node.InputMapping);
+	ClearInvalidMapping(OutputParameters, Node.OutputMapping);
 
-	for (auto Iter = InputVariables.CreateConstIterator(); Iter; ++Iter)
+	for (auto Iter = InputParameters.CreateConstIterator(); Iter; ++Iter)
 	{
 		KnownExposableProperties.Add(Iter.Key());
 	}
@@ -157,7 +216,7 @@ bool UAnimGraphNode_ControlRig::IsInputProperty(const FName& PropertyName) const
 {
 	DECLARE_SCOPE_HIERARCHICAL_COUNTER_FUNC()
 
-	return InputVariables.Contains(PropertyName);
+	return InputParameters.Contains(PropertyName);
 }
 
 bool UAnimGraphNode_ControlRig::IsAvailableToMapToCurve(const FName& PropertyName, bool bInput) const
@@ -166,10 +225,10 @@ bool UAnimGraphNode_ControlRig::IsAvailableToMapToCurve(const FName& PropertyNam
 
 	// find if input or output
 	// ensure it could convert to float
-	const FControlRigIOVariable* Variable = (bInput) ? InputVariables.Find(PropertyName) : OutputVariables.Find(PropertyName);
-	if (ensure(Variable))
+	const FRigVMParameter* Parameter = (bInput) ? InputParameters.Find(PropertyName) : OutputParameters.Find(PropertyName);
+	if (ensure(Parameter))
 	{
-		return FControlRigIOHelper::CanConvert(FName(*Variable->PropertyType), FControlRigIOTypes::GetTypeString<float>());
+		return Parameter->GetCPPType().Equals(TEXT("float"));
 	}
 
 	return false;
@@ -297,30 +356,30 @@ void UAnimGraphNode_ControlRig::CustomizeDetails(IDetailLayoutBuilder& DetailBui
 	// alpha property blending support END
 }
 
-void UAnimGraphNode_ControlRig::GetIOProperties(bool bInput, TMap<FName, FControlRigIOVariable>& OutVars) const
+void UAnimGraphNode_ControlRig::GetIOParameters(bool bInput, TMap<FName, FRigVMParameter>& OutParameters) const
 {
 	DECLARE_SCOPE_HIERARCHICAL_COUNTER_FUNC()
 
-	OutVars.Reset();
+	OutParameters.Reset();
 
-	UClass* TargetClass = GetTargetClass();
-	// if we have target class, see if we can get 
-	if (TargetClass)
+	if (UControlRigBlueprintGeneratedClass* TargetClass = Cast<UControlRigBlueprintGeneratedClass>(GetTargetClass()))
 	{
-		UControlRig* ControlRig = TargetClass->GetDefaultObject<UControlRig>();
-		if (ControlRig)
+		if (UControlRigBlueprint* RigBlueprint = Cast<UControlRigBlueprint>(TargetClass->ClassGeneratedBy))
 		{
-			TArray<FControlRigIOVariable> RigIOVars;
-			// we use CDO to do this because we don't have instance
-			ControlRig->QueryIOVariables(bInput, RigIOVars);
+			RigBlueprint->InitializeModelIfRequired();
 
-			// now we have IO variables, but need to fill up property array
-			for (int32 Index = 0; Index < RigIOVars.Num(); ++Index)
+			UControlRig* ControlRig = TargetClass->GetDefaultObject<UControlRig>();
+			if (ControlRig)
 			{
-				// this sub property name doesn't exist - ARGGG
-				FName PropertyName = FName(*RigIOVars[Index].PropertyPath);
-
-				OutVars.Add(PropertyName, RigIOVars[Index]);
+				const TArray<FRigVMParameter>& Parameters = ControlRig->GetParameters();
+				for (const FRigVMParameter& Parameter : Parameters)
+				{
+					if ((bInput && Parameter.GetType() == ERigVMParameterType::Input) ||
+						(!bInput && Parameter.GetType() == ERigVMParameterType::Output))
+					{
+						OutParameters.Add(Parameter.GetName(), Parameter);
+					}
+				}
 			}
 		}
 	}
@@ -391,8 +450,8 @@ void UAnimGraphNode_ControlRig::CreateVariableMapping(const FString& FilteredTex
 	OutArray.Reset();
 
 	bool bDoFiltering = !FilteredText.IsEmpty();
-	const TMap<FName, FControlRigIOVariable>& Variables = (bInput)? InputVariables : OutputVariables;
-	for (auto Iter = Variables.CreateConstIterator(); Iter; ++Iter)
+	const TMap<FName, FRigVMParameter>& Parameters = (bInput)? InputParameters : OutputParameters;
+	for (auto Iter = Parameters.CreateConstIterator(); Iter; ++Iter)
 	{
 		const FName& Name = Iter.Key();
 		const FString& DisplayName = Name.ToString();
@@ -503,3 +562,5 @@ void UAnimGraphNode_ControlRig::CustomizePinData(UEdGraphPin* Pin, FName SourceP
 	}
 }
 #undef LOCTEXT_NAMESPACE
+
+PRAGMA_ENABLE_OPTIMIZATION
\ No newline at end of file
diff --git a/Engine/Plugins/Experimental/ControlRig/Source/ControlRigDeveloper/Private/AnimGraphNode_ControlRig.h b/Engine/Plugins/Experimental/ControlRig/Source/ControlRigDeveloper/Private/AnimGraphNode_ControlRig.h
--- a/Engine/Plugins/Experimental/ControlRig/Source/ControlRigDeveloper/Private/AnimGraphNode_ControlRig.h
+++ b/Engine/Plugins/Experimental/ControlRig/Source/ControlRigDeveloper/Private/AnimGraphNode_ControlRig.h
@@ -5,7 +5,6 @@
 #include "AnimGraphNode_CustomProperty.h"
 #include "AnimNode_ControlRig.h"
 #include "SSearchableComboBox.h"
-#include "ControlRigVariables.h"
 #include "AnimGraphNode_ControlRig.generated.h"
 
 struct FVariableMappingInfo;
@@ -31,18 +30,19 @@ private:
 	virtual const FAnimNode_CustomProperty* GetCustomPropertyNode() const override { return &Node; }
 
 	// property related things
-	void GetIOProperties(bool bInput, TMap<FName, FControlRigIOVariable>& OutVars) const;
+	void GetIOParameters(bool bInput, TMap<FName, FRigVMParameter>& OutParameters) const;
+
+	TMap<FName, FRigVMParameter> InputParameters;
+	TMap<FName, FRigVMParameter> OutputParameters;
+
 	// we have to override both of it
 	// Rebuild is about rebuilding internal data structre
 	// Getter is about getting only properties, so that it can reconstruct node
-	virtual void GetExposableProperties(TArray<FProperty*>& OutExposableProperties) const override;
+	virtual void GetExposableProperties(TArray<FProperty*>& OutExposableProperties) const override {}
 	virtual void RebuildExposedProperties() override;
 	virtual void ReallocatePinsDuringReconstruction(TArray<UEdGraphPin*>& OldPins) override;
 	virtual void ValidateAnimNodeDuringCompilation(USkeleton* ForSkeleton, FCompilerResultsLog& MessageLog) override;
 
-	TMap<FName, FControlRigIOVariable> InputVariables;
-	TMap<FName, FControlRigIOVariable> OutputVariables;
-
 	// pin option related
 	bool IsPropertyExposeEnabled(FName PropertyName) const;
 	virtual ECheckBoxState IsPropertyExposed(FName PropertyName) const override;
