COPY from //Dev-Anim to //Dev-Main

diff --git a/Engine/Plugins/Experimental/ControlRig/Source/ControlRigDeveloper/Private/ControlRigBlueprint.cpp b/Engine/Plugins/Experimental/ControlRig/Source/ControlRigDeveloper/Private/ControlRigBlueprint.cpp
--- a/Engine/Plugins/Experimental/ControlRig/Source/ControlRigDeveloper/Private/ControlRigBlueprint.cpp
+++ b/Engine/Plugins/Experimental/ControlRig/Source/ControlRigDeveloper/Private/ControlRigBlueprint.cpp
@@ -11,6 +11,13 @@
 #include "Graph/ControlRigGraph.h"
 #include "Graph/ControlRigGraphNode.h"
 #include "UObject/UObjectGlobals.h"
+#include "ControlRigObjectVersion.h"
+#include "ControlRigDeveloper.h"
+#include "Curves/CurveFloat.h"
+#include "BlueprintCompilationManager.h"
+#include "RigVMCompiler/RigVMCompiler.h"
+#include "RigVMCore/RigVMRegistry.h"
+#include "Units/Execution/RigUnit_BeginExecution.h"
 
 #if WITH_EDITOR
 #include "IControlRigEditorModule.h"
@@ -21,33 +28,56 @@
 
 #define LOCTEXT_NAMESPACE "ControlRigBlueprint"
 
-UControlRigBlueprint::UControlRigBlueprint()
+TArray<UControlRigBlueprint*> UControlRigBlueprint::sCurrentlyOpenedRigBlueprints;
+
+UControlRigBlueprint::UControlRigBlueprint(const FObjectInitializer& ObjectInitializer)
 {
 	bSuspendModelNotificationsForSelf = false;
 	bSuspendModelNotificationsForOthers = false;
-	Model = nullptr;
-	ModelController = nullptr;
 
 #if WITH_EDITORONLY_DATA
 	GizmoLibrary = UControlRigSettings::Get()->DefaultGizmoLibrary;
 #endif
+
+	bRecompileOnLoad = 0;
+	bAutoRecompileVM = true;
+	bVMRecompilationRequired = false;
+	VMRecompilationBracket = 0;
+
+	Model = ObjectInitializer.CreateDefaultSubobject<URigVMGraph>(this, TEXT("Model"));
+	Controller = nullptr;
 }
 
-void UControlRigBlueprint::InitializeModel()
+void UControlRigBlueprint::InitializeModelIfRequired()
 {
 	DECLARE_SCOPE_HIERARCHICAL_COUNTER_FUNC()
 
-	if (Model == nullptr || ModelController == nullptr)
+	if (Controller == nullptr)
 	{
-		Model = NewObject<UControlRigModel>(this);
-		ModelController = NewObject<UControlRigController>(this);
-		ModelController->SetModel(Model);
-		ModelController->OnModified().AddUObject(this, &UControlRigBlueprint::HandleModelModified);
+		Controller = NewObject<URigVMController>(this);
+		Controller->SetExecuteContextStruct(FControlRigExecuteContext::StaticStruct());
+		Controller->SetGraph(Model);
+		Controller->OnModified().AddUObject(this, &UControlRigBlueprint::HandleModifiedEvent);
+
+		Controller->UnfoldStructDelegate.BindLambda([](const UStruct* InStruct) -> bool {
+
+			if (InStruct == TBaseStructure<FQuat>::Get())
+			{
+				return false;
+			}
+			if (InStruct == FRuntimeFloatCurve::StaticStruct())
+			{
+				return false;
+			}
+			return true;
+		});
 
 		for (int32 i = 0; i < UbergraphPages.Num(); ++i)
 		{
 			if (UControlRigGraph* Graph = Cast<UControlRigGraph>(UbergraphPages[i]))
 			{
+				PopulateModelFromGraphForBackwardsCompatibility(Graph);
+				RecompileVM();
 				Graph->Initialize(this);
 			}
 		}
@@ -81,13 +111,6 @@ void UControlRigBlueprint::LoadModulesRequiredForCompilation()
 {
 }
 
-void UControlRigBlueprint::MakePropertyLink(const FString& InSourcePropertyPath, const FString& InDestPropertyPath, int32 InSourceLinkIndex, int32 InDestLinkIndex)
-{
-	DECLARE_SCOPE_HIERARCHICAL_COUNTER_FUNC()
-
-	PropertyLinks.AddUnique(FControlRigBlueprintPropertyLink(InSourcePropertyPath, InDestPropertyPath, InSourceLinkIndex, InDestLinkIndex));
-}
-
 USkeletalMesh* UControlRigBlueprint::GetPreviewMesh() const
 {
 	DECLARE_SCOPE_HIERARCHICAL_COUNTER_FUNC()
@@ -110,54 +133,165 @@ void UControlRigBlueprint::SetPreviewMesh(USkeletalMesh* PreviewMesh, bool bMark
 	PreviewSkeletalMesh = PreviewMesh;
 }
 
-void UControlRigBlueprint::GetTypeActions(FBlueprintActionDatabaseRegistrar& ActionRegistrar) const
+void UControlRigBlueprint::PostLoad()
 {
-	DECLARE_SCOPE_HIERARCHICAL_COUNTER_FUNC()
+	Super::PostLoad();
 
-	IControlRigEditorModule::Get().GetTypeActions(this, ActionRegistrar);
+	// remove all non-controlrig-graphs
+	TArray<UEdGraph*> NewUberGraphPages;
+	for (UEdGraph* Graph : UbergraphPages)
+	{
+		UControlRigGraph* RigGraph = Cast<UControlRigGraph>(Graph);
+		if (RigGraph)
+		{
+			NewUberGraphPages.Add(RigGraph);
+		}
+		else
+		{
+			Graph->MarkPendingKill();
+			Graph->Rename(nullptr, GetTransientPackage(), REN_ForceNoResetLoaders);
+		}
+	}
+	UbergraphPages = NewUberGraphPages;
+
+	InitializeModelIfRequired();
+
+#if WITH_EDITOR
+
+	Controller->DetachLinksFromPinObjects();
+	for (URigVMNode* Node : Model->GetNodes())
+	{
+		Controller->RepopulatePinsOnNode(Node);
+	}
+	Controller->ReattachLinksToPinObjects();
+
+	RecompileVM();
+	RequestControlRigInit();
+
+#endif
 }
 
-void UControlRigBlueprint::GetInstanceActions(FBlueprintActionDatabaseRegistrar& ActionRegistrar) const
+void UControlRigBlueprint::RecompileVM()
 {
-	DECLARE_SCOPE_HIERARCHICAL_COUNTER_FUNC()
+	UControlRigBlueprintGeneratedClass* RigClass = GetControlRigBlueprintGeneratedClass();
+	UControlRig* CDO = Cast<UControlRig>(RigClass->GetDefaultObject(true /* create if needed */));
+	if (CDO->VM != nullptr)
+	{
+		if (CDO->VM->GetOuter() != CDO)
+		{
+			CDO->VM = NewObject<URigVM>(CDO, TEXT("VM"));
+		}
+
+		CDO->Modify();
+		CDO->VM->Empty();
+
+		FRigUnitContext Context;
+		Context.State = EControlRigState::Init;
+		Context.Hierarchy = &CDO->Hierarchy;
+		void* ContextPtr = &Context;
+		FRigVMUserDataArray UserData = FRigVMUserDataArray(&ContextPtr, 1);
 
-	IControlRigEditorModule::Get().GetInstanceActions(this, ActionRegistrar);
+		URigVMCompiler* Compiler = URigVMCompiler::StaticClass()->GetDefaultObject<URigVMCompiler>();
+		Compiler->Settings = VMCompileSettings;
+		Compiler->Compile(Model, CDO->VM, UserData, &PinToOperandMap);
+
+		CDO->Execute(EControlRigState::Init); // need to clarify if we actually need this
+		Statistics = CDO->VM->GetStatistics();
+
+		TArray<UObject*> ArchetypeInstances;
+		CDO->GetArchetypeInstances(ArchetypeInstances);
+		for (UObject* Instance : ArchetypeInstances)
+		{
+			if (UControlRig* InstanceRig = Cast<UControlRig>(Instance))
+			{
+				InstanceRig->Hierarchy.Initialize(false);
+				InstanceRig->InstantiateVMFromCDO();
+			}
+		}
+
+		bVMRecompilationRequired = false;
+		VMRecompilationBracket = 0;
+		VMCompiledEvent.Broadcast(this, CDO->VM);
+	}
 }
 
-void UControlRigBlueprint::SetObjectBeingDebugged(UObject* NewObject)
+void UControlRigBlueprint::RecompileVMIfRequired()
 {
-	UControlRig* PreviousRigBeingDebugged = Cast<UControlRig>(GetObjectBeingDebugged());
-	if (PreviousRigBeingDebugged && PreviousRigBeingDebugged != NewObject)
+	if (bVMRecompilationRequired)
 	{
-		PreviousRigBeingDebugged->DrawInterface = nullptr;
-		PreviousRigBeingDebugged->ControlRigLog = nullptr;
+		RecompileVM();
 	}
-
-	Super::SetObjectBeingDebugged(NewObject);
 }
 
-void UControlRigBlueprint::PostLoad()
+void UControlRigBlueprint::RequestAutoVMRecompilation()
 {
-	DECLARE_SCOPE_HIERARCHICAL_COUNTER_FUNC()
+	bVMRecompilationRequired = true;
+	if (bAutoRecompileVM && VMRecompilationBracket == 0)
+	{
+		VMCompileSettings.ASTSettings = FRigVMParserASTSettings::Fast();
+		RecompileVMIfRequired();
+	}
+}
 
-	Super::PostLoad();
+void UControlRigBlueprint::IncrementVMRecompileBracket()
+{
+	VMRecompilationBracket++;
+}
 
-	// remove all non-controlrig-graphs
-	TArray<UEdGraph*> NewUberGraphPages;
-	for (UEdGraph* Graph : UbergraphPages)
+void UControlRigBlueprint::DecrementVMRecompileBracket()
+{
+	if (VMRecompilationBracket == 1)
 	{
-		UControlRigGraph* RigGraph = Cast<UControlRigGraph>(Graph);
-		if (RigGraph)
+		if (bAutoRecompileVM)
 		{
-			NewUberGraphPages.Add(RigGraph);
+			if (bVMRecompilationRequired)
+			{
+				VMCompileSettings.ASTSettings = FRigVMParserASTSettings::Fast();
+			}
+			RecompileVMIfRequired();
 		}
-		else
+		VMRecompilationBracket = 0;
+	}
+	else if (VMRecompilationBracket > 0)
+	{
+		VMRecompilationBracket--;
+	}
+}
+
+void UControlRigBlueprint::RequestControlRigInit()
+{
+	UControlRigBlueprintGeneratedClass* RigClass = GetControlRigBlueprintGeneratedClass();
+	UControlRig* CDO = Cast<UControlRig>(RigClass->GetDefaultObject(true /* create if needed */));
+	CDO->RequestInit();
+
+	TArray<UObject*> ArchetypeInstances;
+	CDO->GetArchetypeInstances(ArchetypeInstances);
+	for (UObject* Instance : ArchetypeInstances)
+	{
+		if (UControlRig* InstanceRig = Cast<UControlRig>(Instance))
 		{
-			Graph->MarkPendingKill();
-			Graph->Rename(nullptr, GetTransientPackage(), REN_ForceNoResetLoaders);
+			InstanceRig->RequestInit();
 		}
 	}
-	UbergraphPages = NewUberGraphPages;
+}
+
+void UControlRigBlueprint::GetTypeActions(FBlueprintActionDatabaseRegistrar& ActionRegistrar) const
+{
+	DECLARE_SCOPE_HIERARCHICAL_COUNTER_FUNC()
+
+	IControlRigEditorModule::Get().GetTypeActions(this, ActionRegistrar);
+}
+
+void UControlRigBlueprint::SetObjectBeingDebugged(UObject* NewObject)
+{
+	UControlRig* PreviousRigBeingDebugged = Cast<UControlRig>(GetObjectBeingDebugged());
+	if (PreviousRigBeingDebugged && PreviousRigBeingDebugged != NewObject)
+	{
+		PreviousRigBeingDebugged->DrawInterface = nullptr;
+		PreviousRigBeingDebugged->ControlRigLog = nullptr;
+	}
+
+	Super::SetObjectBeingDebugged(NewObject);
 }
 
 void UControlRigBlueprint::PostTransacted(const FTransactionObjectEvent& TransactionEvent)
@@ -169,7 +303,16 @@ void UControlRigBlueprint::PostTransacted(const FTransactionObjectEvent& Transac
 	{
 		if (TransactionEvent.GetChangedProperties().Contains(TEXT("HierarchyContainer")))
 		{
-			PropagateHierarchyFromBPToInstances(true);
+			int32 TransactionIndex = GEditor->Trans->FindTransactionIndex(TransactionEvent.GetTransactionId());
+			const FTransaction* Transaction = GEditor->Trans->GetTransaction(TransactionIndex);
+
+			if (Transaction->GenerateDiff().TransactionTitle == TEXT("Transform Gizmo"))
+			{
+				PropagatePoseFromBPToInstances();
+				return;
+			}
+
+			PropagateHierarchyFromBPToInstances(true, true);
 			HierarchyContainer.OnElementChanged.Broadcast(&HierarchyContainer, FRigElementKey());
 
 			// make sure the bone name list is up 2 date for the editor graph
@@ -180,50 +323,332 @@ void UControlRigBlueprint::PostTransacted(const FTransactionObjectEvent& Transac
 				{
 					continue;
 				}
-				RigGraph->CacheNameLists(&HierarchyContainer);
+				RigGraph->CacheNameLists(&HierarchyContainer, &DrawContainer);
+			}
+
+			RequestAutoVMRecompilation();
+			MarkPackageDirty();
+		}
+
+		else if (TransactionEvent.GetChangedProperties().Contains(TEXT("DrawContainer")))
+		{
+			PropagateDrawInstructionsFromBPToInstances();
+		}
+	}
+}
+
+FRigVMGraphModifiedEvent& UControlRigBlueprint::OnModified()
+{
+	return ModifiedEvent;
+}
+
+
+FOnVMCompiledEvent& UControlRigBlueprint::OnVMCompiled()
+{
+	return VMCompiledEvent;
+}
+
+TArray<UControlRigBlueprint*> UControlRigBlueprint::GetCurrentlyOpenRigBlueprints()
+{
+	return sCurrentlyOpenedRigBlueprints;
+}
+
+TArray<UStruct*> UControlRigBlueprint::GetAvailableRigUnits()
+{
+	const TArray<FRigVMFunction>& Functions = FRigVMRegistry::Get().GetFunctions();
+
+	TArray<UStruct*> Structs;
+	UStruct* BaseStruct = FRigUnit::StaticStruct();
+
+	for (const FRigVMFunction& Function : Functions)
+	{
+		if (Function.Struct)
+		{
+			if (Function.Struct->IsChildOf(BaseStruct))
+			{
+				Structs.Add(Function.Struct);
+			}
+		}
+	}
+
+	return Structs;
+}
+
+UControlRigHierarchyModifier* UControlRigBlueprint::GetHierarchyModifier()
+{
+	if (HierarchyModifier == nullptr)
+	{
+		HierarchyModifier = NewObject<UControlRigHierarchyModifier>(this, TEXT("HierarchyModifier"));
+		HierarchyModifier->Container = &HierarchyContainer;
+	}
+	return HierarchyModifier;
+}
+
+#if WITH_EDITOR
+
+FName UControlRigBlueprint::AddTransientControl(URigVMPin* InPin)
+{
+	// for now we only allow one pin control at the same time
+	ClearTransientControls();
+
+	UControlRigBlueprintGeneratedClass* RigClass = GetControlRigBlueprintGeneratedClass();
+	UControlRig* CDO = Cast<UControlRig>(RigClass->GetDefaultObject(true /* create if needed */));
+
+	FName SpaceName = NAME_None;
+	if (URigVMStructNode* StructNode = Cast<URigVMStructNode>(InPin->GetPinForLink()->GetNode()))
+	{
+		if (TSharedPtr<FStructOnScope> DefaultStructScope = StructNode->ConstructStructInstance())
+		{
+			FRigVMStruct* DefaultStruct = (FRigVMStruct*)DefaultStructScope->GetStructMemory();
+
+			FString PinPath = InPin->GetPinForLink()->GetPinPath();
+			FString Left, Right;
+
+			if (URigVMPin::SplitPinPathAtStart(PinPath, Left, Right))
+			{
+				SpaceName = DefaultStruct->DetermineSpaceForPin(Right, &HierarchyContainer);
+			}
+		}
+	}
+
+	FName ReturnName = NAME_None;
+	TArray<UObject*> ArchetypeInstances;
+	CDO->GetArchetypeInstances(ArchetypeInstances);
+	for (UObject* ArchetypeInstance : ArchetypeInstances)
+	{
+		UControlRig* InstancedControlRig = Cast<UControlRig>(ArchetypeInstance);
+		if (InstancedControlRig)
+		{
+			FName ControlName = InstancedControlRig->AddTransientControl(InPin, SpaceName);
+			if (ReturnName == NAME_None)
+			{
+				ReturnName = ControlName;
+			}
+		}
+	}
+
+	if (ReturnName != NAME_None)
+	{
+		HierarchyContainer.OnElementAdded.Broadcast(&HierarchyContainer, FRigElementKey(ReturnName, ERigElementType::Control));
+		HierarchyContainer.OnElementSelected.Broadcast(&HierarchyContainer, FRigElementKey(ReturnName, ERigElementType::Control), true);
+	}
+
+	return ReturnName;
+}
+
+FName UControlRigBlueprint::RemoveTransientControl(URigVMPin* InPin)
+{
+	UControlRigBlueprintGeneratedClass* RigClass = GetControlRigBlueprintGeneratedClass();
+	UControlRig* CDO = Cast<UControlRig>(RigClass->GetDefaultObject(true /* create if needed */));
+
+	FName RemovedName = NAME_None;
+	TArray<UObject*> ArchetypeInstances;
+	CDO->GetArchetypeInstances(ArchetypeInstances);
+	for (UObject* ArchetypeInstance : ArchetypeInstances)
+	{
+		UControlRig* InstancedControlRig = Cast<UControlRig>(ArchetypeInstance);
+		if (InstancedControlRig)
+		{
+			FName Name = InstancedControlRig->RemoveTransientControl(InPin);
+			if (RemovedName == NAME_None)
+			{
+				RemovedName = Name;
+			}
+		}
+	}
+
+	if (RemovedName != NAME_None)
+	{
+		HierarchyContainer.OnElementSelected.Broadcast(&HierarchyContainer, FRigElementKey(RemovedName, ERigElementType::Control), false);
+		HierarchyContainer.OnElementRemoved.Broadcast(&HierarchyContainer, FRigElementKey(RemovedName, ERigElementType::Control));
+	}
+	return RemovedName;
+}
+
+FName UControlRigBlueprint::AddTransientControl(const FRigElementKey& InElement)
+{
+	// for now we only allow one pin control at the same time
+	ClearTransientControls();
+
+	UControlRigBlueprintGeneratedClass* RigClass = GetControlRigBlueprintGeneratedClass();
+	UControlRig* CDO = Cast<UControlRig>(RigClass->GetDefaultObject(true /* create if needed */));
+
+	FName ReturnName = NAME_None;
+	TArray<UObject*> ArchetypeInstances;
+	CDO->GetArchetypeInstances(ArchetypeInstances);
+	for (UObject* ArchetypeInstance : ArchetypeInstances)
+	{
+		UControlRig* InstancedControlRig = Cast<UControlRig>(ArchetypeInstance);
+		if (InstancedControlRig)
+		{
+			FName ControlName = InstancedControlRig->AddTransientControl(InElement);
+			if (ReturnName == NAME_None)
+			{
+				ReturnName = ControlName;
 			}
 		}
 	}
+
+	if (ReturnName != NAME_None)
+	{
+		HierarchyContainer.OnElementAdded.Broadcast(&HierarchyContainer, FRigElementKey(ReturnName, ERigElementType::Control));
+		HierarchyContainer.OnElementSelected.Broadcast(&HierarchyContainer, FRigElementKey(ReturnName, ERigElementType::Control), true);
+	}
+
+	return ReturnName;
+
 }
 
+FName UControlRigBlueprint::RemoveTransientControl(const FRigElementKey& InElement)
+{
+	UControlRigBlueprintGeneratedClass* RigClass = GetControlRigBlueprintGeneratedClass();
+	UControlRig* CDO = Cast<UControlRig>(RigClass->GetDefaultObject(true /* create if needed */));
 
-UControlRigModel::FModifiedEvent& UControlRigBlueprint::OnModified()
+	FName RemovedName = NAME_None;
+	TArray<UObject*> ArchetypeInstances;
+	CDO->GetArchetypeInstances(ArchetypeInstances);
+	for (UObject* ArchetypeInstance : ArchetypeInstances)
+	{
+		UControlRig* InstancedControlRig = Cast<UControlRig>(ArchetypeInstance);
+		if (InstancedControlRig)
+		{
+			FName Name = InstancedControlRig->RemoveTransientControl(InElement);
+			if (RemovedName == NAME_None)
+			{
+				RemovedName = Name;
+			}
+		}
+	}
+
+	if (RemovedName != NAME_None)
+	{
+		HierarchyContainer.OnElementSelected.Broadcast(&HierarchyContainer, FRigElementKey(RemovedName, ERigElementType::Control), false);
+		HierarchyContainer.OnElementRemoved.Broadcast(&HierarchyContainer, FRigElementKey(RemovedName, ERigElementType::Control));
+	}
+	return RemovedName;
+}
+
+void UControlRigBlueprint::ClearTransientControls()
 {
-	return _ModifiedEvent;
+	UControlRigBlueprintGeneratedClass* RigClass = GetControlRigBlueprintGeneratedClass();
+	UControlRig* CDO = Cast<UControlRig>(RigClass->GetDefaultObject(true /* create if needed */));
+
+	TArray<FRigControl> PreviousControls;
+	TArray<UObject*> ArchetypeInstances;
+	CDO->GetArchetypeInstances(ArchetypeInstances);
+	for (UObject* ArchetypeInstance : ArchetypeInstances)
+	{
+		UControlRig* InstancedControlRig = Cast<UControlRig>(ArchetypeInstance);
+		if (InstancedControlRig)
+		{
+			if (PreviousControls.Num() == 0)
+			{
+				PreviousControls = InstancedControlRig->TransientControls;
+			}
+			InstancedControlRig->ClearTransientControls();
+		}
+	}
+
+	for(const FRigControl& RemovedControl : PreviousControls)
+	{
+		HierarchyContainer.OnElementSelected.Broadcast(&HierarchyContainer, FRigElementKey(RemovedControl.Name, ERigElementType::Control), false);
+		HierarchyContainer.OnElementRemoved.Broadcast(&HierarchyContainer, FRigElementKey(RemovedControl.Name, ERigElementType::Control));
+	}
 }
 
-void UControlRigBlueprint::PopulateModelFromGraph(const UControlRigGraph* InGraph)
+#endif
+
+void UControlRigBlueprint::PopulateModelFromGraphForBackwardsCompatibility(UControlRigGraph* InGraph)
 {
 	DECLARE_SCOPE_HIERARCHICAL_COUNTER_FUNC()
 
-	if (Model != nullptr)
+	int32 LinkerVersion = GetLinkerCustomVersion(FControlRigObjectVersion::GUID);
+	if (LinkerVersion >= FControlRigObjectVersion::SwitchedToRigVM)
 	{
 		return;
 	}
 
+	if (LinkerVersion < FControlRigObjectVersion::RemovalOfHierarchyRefPins)
+	{
+		UE_LOG(LogControlRigDeveloper, Warning, TEXT("Control Rig is too old (prior 4.23) - cannot automatically upgrade. Clearing graph."));
+		RebuildGraphFromModel();
+		return;
+	}
+
 	TGuardValue<bool> ReentrantGuardSelf(bSuspendModelNotificationsForSelf, true);
 	{
 		TGuardValue<bool> ReentrantGuardOthers(bSuspendModelNotificationsForOthers, true);
 
-		InitializeModel();
-		ModelController->Clear();
-		NodeToParameterType.Reset();
+		struct LocalHelpers
+		{
+			static FString FixUpPinPath(const FString& InPinPath)
+			{
+				FString PinPath = InPinPath;
+				if (!PinPath.Contains(TEXT(".")))
+				{
+					PinPath += TEXT(".Value");
+				}
+
+				PinPath = PinPath.Replace(TEXT("["), TEXT("."), ESearchCase::IgnoreCase);
+				PinPath = PinPath.Replace(TEXT("]"), TEXT(""), ESearchCase::IgnoreCase);
 
-		for (const UEdGraphNode* Node : InGraph->Nodes)
+				return PinPath;
+			}
+		};
+
+		for (UEdGraphNode* Node : InGraph->Nodes)
 		{
-			if (const UControlRigGraphNode* RigNode = Cast<UControlRigGraphNode>(Node))
+			if (UControlRigGraphNode* RigNode = Cast<UControlRigGraphNode>(Node))
 			{
-				FName NodeName = RigNode->GetPropertyName();
+				FName PropertyName = RigNode->PropertyName_DEPRECATED;
 				FVector2D NodePosition = FVector2D((float)RigNode->NodePosX, (float)RigNode->NodePosY);
-				UScriptStruct* UnitStruct = RigNode->GetUnitScriptStruct();
-				if (UnitStruct)
+				FString StructPath = RigNode->StructPath_DEPRECATED;
+
+				if (StructPath.IsEmpty() && PropertyName != NAME_None)
 				{
-					FName TypeName = UnitStruct->GetFName();
-					ModelController->AddNode(TypeName, NodePosition, NodeName, false);
+					if(FStructProperty* StructProperty = CastField<FStructProperty>(GetControlRigBlueprintGeneratedClass()->FindPropertyByName(PropertyName)))
+					{
+						StructPath = StructProperty->Struct->GetPathName();
+					}
+				}
+
+				URigVMNode* ModelNode = nullptr;
+
+				UScriptStruct* UnitStruct = URigVMPin::FindObjectFromCPPTypeObjectPath<UScriptStruct>(StructPath);
+				if (UnitStruct && UnitStruct->IsChildOf(FRigVMStruct::StaticStruct()))
+				{ 
+					ModelNode = Controller->AddStructNode(UnitStruct, TEXT("Execute"), NodePosition, PropertyName.ToString(), false);
 				}
-				else if (NodeName != NAME_None) // check if this is a variable
+				else if (PropertyName != NAME_None) // check if this is a variable
 				{
-					FName DataType = RigNode->PinType.PinCategory;
+					bool bHasInputLinks = false;
+					bool bHasOutputLinks = false;
+					FString DefaultValue;
+
+					FEdGraphPinType PinType = RigNode->PinType_DEPRECATED;
+					if (RigNode->Pins.Num() > 0)
+					{
+						for (UEdGraphPin* Pin : RigNode->Pins)
+						{
+							if (!Pin->GetName().Contains(TEXT(".")))
+							{
+								PinType = Pin->PinType;
+
+								if (Pin->Direction == EGPD_Input)
+								{
+									bHasInputLinks = Pin->LinkedTo.Num() > 0;
+									DefaultValue = Pin->DefaultValue;
+								}
+								else if (Pin->Direction == EGPD_Output)
+								{
+									bHasOutputLinks = Pin->LinkedTo.Num() > 0;
+								}
+							}
+						}
+					}
+
+					FName DataType = PinType.PinCategory;
+					UObject* DataTypeObject = nullptr;
 					if (DataType == NAME_None)
 					{
 						continue;
@@ -231,310 +656,318 @@ void UControlRigBlueprint::PopulateModelFromGraph(const UControlRigGraph* InGrap
 					if (DataType == UEdGraphSchema_K2::PC_Struct)
 					{
 						DataType = NAME_None;
-
-						UScriptStruct* DataStruct = Cast<UScriptStruct>(RigNode->PinType.PinSubCategoryObject);
-						if (DataStruct != nullptr)
+						if (UScriptStruct* DataStruct = Cast<UScriptStruct>(PinType.PinSubCategoryObject))
 						{
-							DataType = DataStruct->GetFName();
+							DataTypeObject = DataStruct;
+							DataType = *DataStruct->GetStructCPPName();
 						}
 					}
 
-					EControlRigModelParameterType ParameterType = (EControlRigModelParameterType)RigNode->ParameterType;
-					if (ParameterType == EControlRigModelParameterType::None)
+					if (DataType == TEXT("int"))
+					{
+						DataType = TEXT("int32");
+					}
+					else if (DataType == TEXT("name"))
+					{
+						DataType = TEXT("FName");
+					}
+					else if (DataType == TEXT("string"))
+					{
+						DataType = TEXT("FString");
+					}
+
+					FProperty* ParameterProperty = GetControlRigBlueprintGeneratedClass()->FindPropertyByName(PropertyName);
+					if(ParameterProperty)
 					{
-						ParameterType = EControlRigModelParameterType::Hidden;
+						bool bIsInput = true;
+
+						if (ParameterProperty->HasMetaData(TEXT("AnimationInput")) || bHasOutputLinks)
+						{
+							bIsInput = true;
+						}
+						else if (ParameterProperty->HasMetaData(TEXT("AnimationOutput")))
+						{
+							bIsInput = false;
+						}
+
+						ModelNode = Controller->AddParameterNode(PropertyName, DataType.ToString(), DataTypeObject, bIsInput, FString(), NodePosition, PropertyName.ToString(), false);
 					}
-					ModelController->AddParameter(NodeName, DataType, ParameterType, NodePosition, false);
-					NodeToParameterType.Add(NodeName, ParameterType);
 				}
 				else
 				{
 					continue;
 				}
 
-				for (UEdGraphPin* Pin : RigNode->Pins)
+				if (ModelNode)
 				{
-					FString Left, Right;
-					Model->SplitPinPath(Pin->GetName(), Left, Right);
-
-					if (Pin->Direction == EGPD_Input && Pin->PinType.ContainerType == EPinContainerType::Array)
-					{
-						int32 ArraySize = Pin->SubPins.Num();
-						ModelController->SetArrayPinSize(NodeName, *Right, ArraySize, FString(), false);
-					}
-					if (RigNode->IsPinExpanded(Pin->GetName()))
+					for (UEdGraphPin* Pin : RigNode->Pins)
 					{
-						ModelController->ExpandPin(NodeName, *Right, Pin->Direction == EGPD_Input, true, false);
-					}
-					if (!Pin->DefaultValue.IsEmpty() && Pin->Direction == EGPD_Input)
-					{
-						ModelController->SetPinDefaultValue(NodeName, *Right, Pin->DefaultValue, false, false);
+						FString PinPath = LocalHelpers::FixUpPinPath(Pin->GetName());
+
+						// check the material + mesh pins for deprecated control nodes
+						if (URigVMStructNode* ModelStructNode = Cast<URigVMStructNode>(ModelNode))
+						{
+							if (ModelStructNode->GetScriptStruct()->IsChildOf(FRigUnit_Control::StaticStruct()))
+							{
+								if (Pin->GetName().EndsWith(TEXT(".StaticMesh")) || Pin->GetName().EndsWith(TEXT(".Materials")))
+								{
+									continue;
+								}
+							}
+						}
+
+						if (Pin->Direction == EGPD_Input && Pin->PinType.ContainerType == EPinContainerType::Array)
+						{
+							int32 ArraySize = Pin->SubPins.Num();
+							Controller->SetArrayPinSize(PinPath, ArraySize, FString(), false);
+						}
+
+						if (RigNode->ExpandedPins_DEPRECATED.Find(Pin->GetName()) != INDEX_NONE)
+						{
+							Controller->SetPinExpansion(PinPath, true, false);
+						}
+
+						if (Pin->SubPins.Num() == 0 && !Pin->DefaultValue.IsEmpty() && Pin->Direction == EGPD_Input)
+						{
+							Controller->SetPinDefaultValue(PinPath, Pin->DefaultValue, false, false, false);
+						}
 					}
 				}
+
+				const int32 VarIndex = FBlueprintEditorUtils::FindNewVariableIndex(this, PropertyName);
+				if (VarIndex != INDEX_NONE)
+				{
+					NewVariables.RemoveAt(VarIndex);
+					FBlueprintEditorUtils::RemoveVariableNodes(this, PropertyName);
+				}
 			}
 			else if (const UEdGraphNode_Comment* CommentNode = Cast<UEdGraphNode_Comment>(Node))
 			{
 				FVector2D NodePosition = FVector2D((float)CommentNode->NodePosX, (float)CommentNode->NodePosY);
 				FVector2D NodeSize = FVector2D((float)CommentNode->NodeWidth, (float)CommentNode->NodeHeight);
-				ModelController->AddComment(CommentNode->GetFName(), CommentNode->NodeComment, NodePosition, NodeSize, CommentNode->CommentColor, false);
+				Controller->AddCommentNode(CommentNode->NodeComment, NodePosition, NodeSize, CommentNode->CommentColor, CommentNode->GetName(), false);
 			}
 		}
 
-		for (const UEdGraphNode* Node : InGraph->Nodes)
+		for (UEdGraphNode* Node : InGraph->Nodes)
 		{
-			if (const UControlRigGraphNode* RigNode = Cast<UControlRigGraphNode>(Node))
+			if (UControlRigGraphNode* RigNode = Cast<UControlRigGraphNode>(Node))
 			{
-				for (const UEdGraphPin* Pin : RigNode->Pins)
+				for (UEdGraphPin* Pin : RigNode->Pins)
 				{
 					if (Pin->Direction == EGPD_Input)
 					{
 						continue;
 					}
 
-					FString Left, Right;
-					Model->SplitPinPath(Pin->GetName(), Left, Right);
-					for (const UEdGraphPin* LinkedPin : Pin->LinkedTo)
+					for (UEdGraphPin* LinkedPin : Pin->LinkedTo)
 					{
-						const UControlRigGraphNode* LinkedRigNode = Cast<UControlRigGraphNode>(LinkedPin->GetOwningNode());
+						UControlRigGraphNode* LinkedRigNode = Cast<UControlRigGraphNode>(LinkedPin->GetOwningNode());
 						if (LinkedRigNode != nullptr)
 						{
-							FString LinkedLeft, LinkedRight;
-							Model->SplitPinPath(LinkedPin->GetName(), LinkedLeft, LinkedRight);
-							ModelController->MakeLink(RigNode->PropertyName, *Right, LinkedRigNode->PropertyName, *LinkedRight, nullptr, false);
+							FString SourcePinPath = LocalHelpers::FixUpPinPath(Pin->GetName());
+							FString TargetPinPath = LocalHelpers::FixUpPinPath(LinkedPin->GetName());
+							Controller->AddLink(SourcePinPath, TargetPinPath, false);
 						}
 					}
 				}
 			}
 		}
 	}
+
+	RebuildGraphFromModel();
 }
 
 void UControlRigBlueprint::RebuildGraphFromModel()
 {
 	DECLARE_SCOPE_HIERARCHICAL_COUNTER_FUNC()
 
-	TGuardValue<bool> SelfGuard(bSuspendModelNotificationsForSelf, true);
-	check(ModelController);
-	ModelController->ResendAllNotifications();
+	TGuardValue<bool> SelfGuard(bSuspendModelNotificationsForSelf, true);
+	check(Controller);
+
+	for (UEdGraph* Graph : UbergraphPages)
+	{
+		TArray<UEdGraphNode*> Nodes = Graph->Nodes;
+		for (UEdGraphNode* Node : Nodes)
+		{
+			Graph->RemoveNode(Node);
+		}
+	}
+
+	Controller->ResendAllNotifications();
+}
+void UControlRigBlueprint::Notify(ERigVMGraphNotifType InNotifType, UObject* InSubject)
+{
+	Controller->Notify(InNotifType, InSubject);
 }
 
-void UControlRigBlueprint::HandleModelModified(const UControlRigModel* InModel, EControlRigModelNotifType InType, const void* InPayload)
+void UControlRigBlueprint::HandleModifiedEvent(ERigVMGraphNotifType InNotifType, URigVMGraph* InGraph, UObject* InSubject)
 {
 	DECLARE_SCOPE_HIERARCHICAL_COUNTER_FUNC()
 
-	if (!bSuspendModelNotificationsForSelf)
-	{
-
 #if WITH_EDITOR
 
-		switch (InType)
+	if (!bSuspendModelNotificationsForSelf)
+	{
+		switch (InNotifType)
 		{
-			case EControlRigModelNotifType::ModelCleared:
+			case ERigVMGraphNotifType::InteractionBracketOpened:
 			{
-				LastNameFromNotification = NAME_None;
-				NodeToParameterType.Reset();
-
-				for (const FControlRigModelNode& Node : InModel->Nodes())
-				{
-					FBlueprintEditorUtils::RemoveMemberVariable(this, Node.Name);
-				}
+				IncrementVMRecompileBracket();
 				break;
 			}
-			case EControlRigModelNotifType::NodeAdded:
+			case ERigVMGraphNotifType::InteractionBracketClosed:
+			case ERigVMGraphNotifType::InteractionBracketCanceled:
 			{
-				LastNameFromNotification = NAME_None;
-
-				const FControlRigModelNode* Node = (const FControlRigModelNode*)InPayload;
-				bool bValidNode = false;
-				if (Node != nullptr)
+				DecrementVMRecompileBracket();
+				break;
+			}
+			case ERigVMGraphNotifType::PinDefaultValueChanged:
+			{
+				if (URigVMPin* Pin = Cast<URigVMPin>(InSubject))
 				{
-					LastNameFromNotification = Node->Name;
-					switch (Node->NodeType)
+					bool bRequiresRecompile = false;
+
+					URigVMPin* RootPin = Pin->GetRootPin();
+					if (const FRigVMOperand* Operand = PinToOperandMap.Find(RootPin->GetPinPath()))
 					{
-						case EControlRigModelNodeType::Parameter:
-						{
-							FControlRigBlueprintUtils::AddPropertyMember(this, Node->Pins[0].Type, *Node->Name.ToString());
-							HandleModelModified(InModel, EControlRigModelNotifType::NodeChanged, InPayload);
-							NodeToParameterType.Add(Node->Name, Node->ParameterType);
-							bValidNode = true;
-							break;
-						}
-						case EControlRigModelNodeType::Function:
-						{
-							FControlRigBlueprintUtils::AddUnitMember(this, Node->UnitStruct(), Node->Name);
-							bValidNode = true;
-							break;
-						}
-						default:
+						if(const FRigVMExprAST* Expression = InGraph->GetRuntimeAST()->GetExprForSubject(RootPin))
 						{
-							break;
+							bRequiresRecompile = Expression->NumParents() > 1;
 						}
-					}
-				}
-
-				if (bValidNode)
-				{
-					FBlueprintEditorUtils::MarkBlueprintAsModified(this);
-
-					if (Node != nullptr)
-					{
-						if (Node->IsParameter())
+						else
 						{
-							UpdateParametersOnControlRig();
+							bRequiresRecompile = true;
 						}
-					}
-				}
-				break;
-			}
-			case EControlRigModelNotifType::NodeRemoved:
-			{
-				LastNameFromNotification = NAME_None;
-					
-				const FControlRigModelNode* Node = (const FControlRigModelNode*)InPayload;
-				if (Node != nullptr)
-				{
-					LastNameFromNotification = Node->Name;
-					FBlueprintEditorUtils::RemoveMemberVariable(this, Node->Name);
-					FBlueprintEditorUtils::MarkBlueprintAsModified(this);
-					WatchedPins.Reset();
-					NodeToParameterType.Remove(Node->Name);
-				}
-				break;
-			}
-			case EControlRigModelNotifType::NodeChanged:
-			{
-				LastNameFromNotification = NAME_None;
-				
-				const FControlRigModelNode* Node = (const FControlRigModelNode*)InPayload;
-				if (Node != nullptr)
-				{
-					LastNameFromNotification = Node->Name;
 
-					if (Node->IsParameter())
-					{
-						FProperty* Property = GeneratedClass->FindPropertyByName(Node->Name);
-						if (Property != nullptr)
+						if(!bRequiresRecompile)
 						{
-							EControlRigModelParameterType WasParameterType = EControlRigModelParameterType::Hidden;
-							if (NodeToParameterType.Contains(Node->Name))
+							TArray<FString> DefaultValues;
+							if (RootPin->IsArray())
 							{
-								WasParameterType = NodeToParameterType.FindChecked(Node->Name);
+								for (URigVMPin* ArrayElementPin : RootPin->GetSubPins())
+								{
+									DefaultValues.Add(ArrayElementPin->GetDefaultValue());
+								}
 							}
 							else
 							{
-								bool bWasInput = Property->HasMetaData(UControlRig::AnimationInputMetaName);
-								bool bWasOutput = Property->HasMetaData(UControlRig::AnimationOutputMetaName);
-								if (bWasInput)
-								{
-									WasParameterType = EControlRigModelParameterType::Input;
-								}
-								else if (bWasOutput)
-								{
-									WasParameterType = EControlRigModelParameterType::Output;
-								}
-								NodeToParameterType.Add(Node->Name, WasParameterType);
+								DefaultValues.Add(RootPin->GetDefaultValue());
 							}
 
-							if (WasParameterType != Node->ParameterType)
+							UControlRigBlueprintGeneratedClass* RigClass = GetControlRigBlueprintGeneratedClass();
+							UControlRig* CDO = Cast<UControlRig>(RigClass->GetDefaultObject(true /* create if needed */));
+							if (CDO->VM != nullptr)
 							{
-								Property->RemoveMetaData(UControlRig::AnimationInputMetaName);
-								Property->RemoveMetaData(UControlRig::AnimationOutputMetaName);
+								CDO->VM->SetRegisterValueFromString(*Operand, RootPin->GetCPPType(), RootPin->GetCPPTypeObject(), DefaultValues);
+							}
 
-								if (Node->ParameterType == EControlRigModelParameterType::Input)
-								{
-									Property->SetMetaData(UControlRig::AnimationInputMetaName, TEXT("True"));
-								}
-								else if (Node->ParameterType == EControlRigModelParameterType::Output)
+							TArray<UObject*> ArchetypeInstances;
+							CDO->GetArchetypeInstances(ArchetypeInstances);
+							for (UObject* ArchetypeInstance : ArchetypeInstances)
+							{
+								UControlRig* InstancedControlRig = Cast<UControlRig>(ArchetypeInstance);
+								if (InstancedControlRig)
 								{
-									Property->SetMetaData(UControlRig::AnimationOutputMetaName, TEXT("True"));
+									if (InstancedControlRig->VM)
+									{
+										InstancedControlRig->VM->SetRegisterValueFromString(*Operand, RootPin->GetCPPType(), RootPin->GetCPPTypeObject(), DefaultValues);
+									}
 								}
+							}
 
-								NodeToParameterType[Node->Name] = Node->ParameterType;
-								FBlueprintEditorUtils::MarkBlueprintAsModified(this);
-								UpdateParametersOnControlRig();
+							if (Pin->IsConstant() || Pin->GetRootPin()->IsConstant())
+							{
+								// re-init the rigs
+								RequestControlRigInit();
 							}
 						}
 					}
-				}
+					else
+					{
+						bRequiresRecompile = true;
+					}
+					
+					if(bRequiresRecompile)
+					{
+						RequestAutoVMRecompilation();
+					}
 
-				break;
-			}
-			case EControlRigModelNotifType::NodeRenamed:
-			{
-				FBlueprintEditorUtils::MarkBlueprintAsModified(this);
-				break;
-			}
-			case EControlRigModelNotifType::PinAdded:
-			case EControlRigModelNotifType::PinRemoved:
-			{
-				const FControlRigModelPin* Pin = (const FControlRigModelPin*)InPayload;
-				if (Pin)
-				{
-					if (Pin->ParentIndex != INDEX_NONE)
+					// check if this pin is part of an injected node, and if it is a visual debug node,
+					// we might need to recreate the control pin
+					if (UClass* MyControlRigClass = GeneratedClass)
 					{
-						const FControlRigModelNode& Node = InModel->Nodes()[Pin->Node];
-						const FControlRigModelPin& ParentPin = Node.Pins[Pin->ParentIndex];
-						if (ParentPin.IsArray())
+						if (UControlRig* DefaultObject = Cast<UControlRig>(MyControlRigClass->GetDefaultObject(false)))
 						{
-							FString PinPath = InModel->GetPinPath(ParentPin.GetPair(), true);
-
-							if (InType == EControlRigModelNotifType::PinAdded)
-							{
-								PerformArrayOperation(PinPath, [](FScriptArrayHelper& InArrayHelper, int32 InArrayIndex)
-								{
-									InArrayHelper.AddValue();
-									return true;
-								}, true, true);
-							}
-							else
+							TArray<UObject*> ArchetypeInstances;
+							DefaultObject->GetArchetypeInstances(ArchetypeInstances);
+							for (UObject* ArchetypeInstance : ArchetypeInstances)
 							{
-								PerformArrayOperation(PinPath, [](FScriptArrayHelper& InArrayHelper, int32 InArrayIndex)
+								if (UControlRig* InstanceRig = Cast<UControlRig>(ArchetypeInstance))
 								{
-									// for now let's remove the last one
-									InArrayHelper.RemoveValues(InArrayHelper.Num() - 1, 1);
-									//InArrayHelper.RemoveValues(InArrayIndex, 0);
-									return true;
-								}, true, true);
+									for (const FRigControl& Control : InstanceRig->TransientControls)
+									{
+										if (URigVMPin* ControlledPin = Model->FindPin(Control.Name.ToString()))
+										{
+											URigVMPin* ControlledPinForLink = ControlledPin->GetPinForLink();
+
+											if(ControlledPin->GetRootPin() == Pin->GetRootPin() ||
+											   ControlledPinForLink->GetRootPin() == Pin->GetRootPin())
+											{
+												InstanceRig->SetTransientControlValue(ControlledPin->GetPinForLink());
+											}
+											else if (ControlledPin->GetNode() == Pin->GetNode() ||
+													 ControlledPinForLink->GetNode() == Pin->GetNode())
+											{
+												InstanceRig->ClearTransientControls();
+												InstanceRig->AddTransientControl(ControlledPin);
+											}
+											break;
+										}
+									}
+								}
 							}
 						}
 					}
 				}
+				MarkPackageDirty();
 				break;
 			}
-			case EControlRigModelNotifType::PinChanged:
+			case ERigVMGraphNotifType::NodeAdded:
+			case ERigVMGraphNotifType::NodeRemoved:
+			case ERigVMGraphNotifType::LinkAdded:
+			case ERigVMGraphNotifType::LinkRemoved:
+			case ERigVMGraphNotifType::PinArraySizeChanged:
+			case ERigVMGraphNotifType::PinDirectionChanged:
 			{
-				const FControlRigModelPin* Pin = (const FControlRigModelPin*)InPayload;
-				if (Pin)
+				ClearTransientControls();
+				RequestAutoVMRecompilation();
+				MarkPackageDirty();
+				break;
+			}
+			case ERigVMGraphNotifType::PinWatchedChanged:
+			case ERigVMGraphNotifType::PinTypeChanged:
+			{
+				if (URigVMPin* ModelPin = Cast<URigVMPin>(InSubject))
 				{
-					if (UClass* MyControlRigClass = GeneratedClass)
+					for (int32 i = 0; i < UbergraphPages.Num(); ++i)
 					{
-						if (UObject* DefaultObject = MyControlRigClass->GetDefaultObject(false))
+						if (UControlRigGraph* Graph = Cast<UControlRigGraph>(UbergraphPages[i]))
 						{
-							DefaultObject->SetFlags(RF_Transactional);
-
-							FString PinPath = InModel->GetPinPath(Pin->GetPair(), true);
-							FString DefaultValueString = Pin->DefaultValue;
-							if (DefaultValueString.Len() > 0)
+							if (UEdGraphNode* EdNode = Graph->FindNodeForModelNodeName(ModelPin->GetNode()->GetFName()))
 							{
-								FCachedPropertyPath PropertyPath(PinPath);
-								if (PropertyPathHelpers::SetPropertyValueFromString(DefaultObject, PropertyPath, DefaultValueString))
+								if (UEdGraphPin* EdPin = EdNode->FindPin(*ModelPin->GetPinPath()))
 								{
-									if (Pin->bIsConstant)
+									if (ModelPin->RequiresWatch())
 									{
-										FBlueprintEditorUtils::MarkBlueprintAsModified(this);
+										WatchedPins.AddUnique(EdPin);
 									}
 									else
 									{
-										MarkPackageDirty();
+										WatchedPins.Remove(EdPin);
 									}
-								}
-
-								TArray<UObject*> ArchetypeInstances;
-								DefaultObject->GetArchetypeInstances(ArchetypeInstances);
-
-								for (UObject* ArchetypeInstance : ArchetypeInstances)
-								{
-									PropertyPathHelpers::SetPropertyValueFromString(ArchetypeInstance, PropertyPath, DefaultValueString);
+									RequestAutoVMRecompilation();
+									MarkPackageDirty();
 								}
 							}
 						}
@@ -542,10 +975,12 @@ void UControlRigBlueprint::HandleModelModified(const UControlRigModel* InModel,
 				}
 				break;
 			}
-			case EControlRigModelNotifType::LinkAdded:
-			case EControlRigModelNotifType::LinkRemoved:
+			case ERigVMGraphNotifType::ParameterAdded:
+			case ERigVMGraphNotifType::ParameterRemoved:
+			case ERigVMGraphNotifType::ParameterRenamed:
 			{
-				FBlueprintEditorUtils::MarkBlueprintAsModified(this);
+				RequestAutoVMRecompilation();
+				MarkPackageDirty();
 				break;
 			}
 			default:
@@ -557,180 +992,298 @@ void UControlRigBlueprint::HandleModelModified(const UControlRigModel* InModel,
 
 	if (!bSuspendModelNotificationsForOthers)
 	{
-		if (_ModifiedEvent.IsBound())
+		if (ModifiedEvent.IsBound())
 		{
-			_ModifiedEvent.Broadcast(InModel, InType, InPayload);
+			ModifiedEvent.Broadcast(InNotifType, InGraph, InSubject);
 		}
 	}
-
 #endif
 }
 
-bool UControlRigBlueprint::UpdateParametersOnControlRig(UControlRig* InRig)
+void UControlRigBlueprint::CleanupBoneHierarchyDeprecated()
 {
-	DECLARE_SCOPE_HIERARCHICAL_COUNTER_FUNC()
-
-	TArray<UObject*> ArchetypeInstances;
+	if (Hierarchy_DEPRECATED.Num() > 0)
+	{
+		HierarchyContainer.BoneHierarchy = Hierarchy_DEPRECATED;
+		Hierarchy_DEPRECATED.Reset();
+	}
 
-	if (InRig == nullptr)
+	if (CurveContainer_DEPRECATED.Num() > 0)
 	{
-		if (UClass* MyControlRigClass = GeneratedClass)
-		{
-			InRig = Cast<UControlRig>(MyControlRigClass->GetDefaultObject(false));
-			if (InRig == nullptr)
-			{
-				return false;
-			}
-			InRig->Modify();
-			InRig->GetArchetypeInstances(ArchetypeInstances);
-		}
+		HierarchyContainer.CurveContainer = CurveContainer_DEPRECATED;
+		CurveContainer_DEPRECATED.Reset();
 	}
 
-	check(InRig);
+}
 
-	InRig->InputProperties.Reset();
-	InRig->OutputProperties.Reset();
+void UControlRigBlueprint::PropagatePoseFromInstanceToBP(UControlRig* InControlRig)
+{
+	check(InControlRig);
 
-	if (Model == nullptr)
+	for (const FRigBone& InputBone : InControlRig->Hierarchy.BoneHierarchy)
 	{
-		return false;
+		FRigBone& OutputBone = HierarchyContainer.BoneHierarchy[InputBone.Name];
+		OutputBone.InitialTransform = InputBone.InitialTransform;
+		OutputBone.LocalTransform = InputBone.LocalTransform;
+		OutputBone.GlobalTransform = InputBone.GlobalTransform;
 	}
 
-	for(const FControlRigModelNode& Node : Model->Nodes())
+	for (const FRigSpace& InputSpace: InControlRig->Hierarchy.SpaceHierarchy)
 	{
-		if (!Node.IsParameter())
-		{
-			continue;
-		}
-
-		FCachedPropertyPath NewCachedProperty(Node.Name.ToString());
-		if (!NewCachedProperty.Resolve(InRig))
-		{
-			continue;
-		}
-
-		ensure(NewCachedProperty.IsFullyResolved());
-
-		if (Node.ParameterType == EControlRigModelParameterType::Input)
-		{
-			InRig->InputProperties.Add(Node.Name, NewCachedProperty);
-		}
-		else if (Node.ParameterType == EControlRigModelParameterType::Output)
-		{
-			InRig->OutputProperties.Add(Node.Name, NewCachedProperty);
-		}
+		FRigSpace& OutputSpace = HierarchyContainer.SpaceHierarchy[InputSpace.Name];
+		OutputSpace.InitialTransform = InputSpace.InitialTransform;
+		OutputSpace.LocalTransform = InputSpace.LocalTransform;
 	}
 
-	InRig->ResolveInputOutputProperties();
+	for (const FRigControl& InputControl : InControlRig->Hierarchy.ControlHierarchy)
+	{
+		FRigControl& OutputControl = HierarchyContainer.ControlHierarchy[InputControl.Name];
+		OutputControl.InitialValue = InputControl.InitialValue;
+		OutputControl.Value = InputControl.Value;
+	}
+}
 
-	for (UObject* ArchetypeInstance : ArchetypeInstances)
+void UControlRigBlueprint::PropagatePoseFromBPToInstances()
+{
+	if (UClass* MyControlRigClass = GeneratedClass)
 	{
-		UControlRig* InstancedControlRig = Cast<UControlRig>(ArchetypeInstance);
-		if (InstancedControlRig)
+		if (UControlRig* DefaultObject = Cast<UControlRig>(MyControlRigClass->GetDefaultObject(false)))
 		{
-			UpdateParametersOnControlRig(InstancedControlRig);
+			TArray<UObject*> ArchetypeInstances;
+			DefaultObject->GetArchetypeInstances(ArchetypeInstances);
+			for (UObject* ArchetypeInstance : ArchetypeInstances)
+			{
+				if (UControlRig* InstanceRig = Cast<UControlRig>(ArchetypeInstance))
+				{
+					for (const FRigBone& InputBone : HierarchyContainer.BoneHierarchy)
+					{
+						FRigBone& OutputBone = InstanceRig->Hierarchy.BoneHierarchy[InputBone.Name];
+						OutputBone.InitialTransform = InputBone.InitialTransform;
+						OutputBone.LocalTransform = InputBone.LocalTransform;
+						OutputBone.GlobalTransform = InputBone.GlobalTransform;
+					}
+
+					for (const FRigSpace& InputSpace : HierarchyContainer.SpaceHierarchy)
+					{
+						FRigSpace& OutputSpace = InstanceRig->Hierarchy.SpaceHierarchy[InputSpace.Name];
+						OutputSpace.InitialTransform = InputSpace.InitialTransform;
+						OutputSpace.LocalTransform = InputSpace.LocalTransform;
+					}
+
+					for (const FRigControl& InputControl : HierarchyContainer.ControlHierarchy)
+					{
+						FRigControl& OutputControl = InstanceRig->Hierarchy.ControlHierarchy[InputControl.Name];
+						OutputControl.InitialValue = InputControl.InitialValue;
+						OutputControl.Value = InputControl.Value;
+					}
+				}
+			}
 		}
 	}
-
-	return true;
 }
 
-bool UControlRigBlueprint::PerformArrayOperation(const FString& InPropertyPath, TFunctionRef<bool(FScriptArrayHelper&,int32)> InOperation, bool bCallModify, bool bPropagateToInstances)
+void UControlRigBlueprint::PropagateHierarchyFromBPToInstances(bool bInitializeContainer, bool bInitializeRigs)
 {
-	DECLARE_SCOPE_HIERARCHICAL_COUNTER_FUNC()
-
 	if (UClass* MyControlRigClass = GeneratedClass)
 	{
-		if(UObject* DefaultObject = MyControlRigClass->GetDefaultObject(false))
+		if (UControlRig* DefaultObject = Cast<UControlRig>(MyControlRigClass->GetDefaultObject(false)))
 		{
-			if(bCallModify)
+			if (bInitializeContainer)
 			{
-				DefaultObject->SetFlags(RF_Transactional);
-				DefaultObject->Modify();
+				HierarchyContainer.Initialize();
+				HierarchyContainer.ResetTransforms();
 			}
 
-			FCachedPropertyPath CachedPropertyPath(InPropertyPath);
-			if(PropertyPathHelpers::PerformArrayOperation(DefaultObject, CachedPropertyPath, InOperation))
+			DefaultObject->Hierarchy = HierarchyContainer;
+			if (bInitializeRigs)
 			{
-				if(bCallModify)
-				{
-					FBlueprintEditorUtils::MarkBlueprintAsModified(this);
+				DefaultObject->Initialize(true);
+			}
+			else
+			{
+				DefaultObject->Hierarchy.Initialize(false);
+			}
 
-					if(bPropagateToInstances)
+			TArray<UObject*> ArchetypeInstances;
+			DefaultObject->GetArchetypeInstances(ArchetypeInstances);
+			for (UObject* ArchetypeInstance : ArchetypeInstances)
+			{
+				if (UControlRig* InstanceRig = Cast<UControlRig>(ArchetypeInstance))
+				{
+					InstanceRig->Hierarchy = HierarchyContainer;
+					if (bInitializeRigs)
 					{
-						TArray<UObject*> ArchetypeInstances;
-						DefaultObject->GetArchetypeInstances(ArchetypeInstances);
-							
-						for (UObject* ArchetypeInstance : ArchetypeInstances)
-						{
-							PropertyPathHelpers::PerformArrayOperation(ArchetypeInstance, CachedPropertyPath, InOperation);
-						}
+						InstanceRig->Initialize(true);
+					}
+					else
+					{
+						InstanceRig->Hierarchy.Initialize(false);
 					}
 				}
-				return true;
 			}
 		}
 	}
-
-	return false;
 }
 
-void UControlRigBlueprint::CleanupBoneHierarchyDeprecated()
+void UControlRigBlueprint::PropagateDrawInstructionsFromBPToInstances()
 {
-	if (Hierarchy_DEPRECATED.Num() > 0)
+	if (UClass* MyControlRigClass = GeneratedClass)
 	{
-		HierarchyContainer.BoneHierarchy = Hierarchy_DEPRECATED;
-		Hierarchy_DEPRECATED.Reset();
-	}
+		if (UControlRig* DefaultObject = Cast<UControlRig>(MyControlRigClass->GetDefaultObject(false)))
+		{
+			DefaultObject->DrawContainer = DrawContainer;
 
-	if (CurveContainer_DEPRECATED.Num() > 0)
-	{
-		HierarchyContainer.CurveContainer = CurveContainer_DEPRECATED;
-		CurveContainer_DEPRECATED.Reset();
+			TArray<UObject*> ArchetypeInstances;
+			DefaultObject->GetArchetypeInstances(ArchetypeInstances);
+
+			for (UObject* ArchetypeInstance : ArchetypeInstances)
+			{
+				if (UControlRig* InstanceRig = Cast<UControlRig>(ArchetypeInstance))
+				{
+					InstanceRig->DrawContainer = DrawContainer;
+				}
+			}
+		}
 	}
 
-}
 
-void UControlRigBlueprint::PropagatePoseFromInstanceToBP()
-{
-	if (UClass* MyControlRigClass = GeneratedClass)
+	// make sure the bone name list is up 2 date for the editor graph
+	for (UEdGraph* Graph : UbergraphPages)
 	{
-		if (UControlRig* DebuggedRig = Cast<UControlRig>(GetObjectBeingDebugged()))
+		UControlRigGraph* RigGraph = Cast<UControlRigGraph>(Graph);
+		if (RigGraph == nullptr)
 		{
-			/* todo */
+			continue;
 		}
+		RigGraph->CacheNameLists(&HierarchyContainer, &DrawContainer);
 	}
 }
 
-void UControlRigBlueprint::PropagateHierarchyFromBPToInstances(bool bInitialize)
+void UControlRigBlueprint::PropagatePropertyFromBPToInstances(FRigElementKey InRigElement, const FProperty* InProperty)
 {
+	int32 ElementIndex = HierarchyContainer.GetIndex(InRigElement);
+	ensure(ElementIndex != INDEX_NONE);
+	check(InProperty);
+
 	if (UClass* MyControlRigClass = GeneratedClass)
 	{
 		if (UControlRig* DefaultObject = Cast<UControlRig>(MyControlRigClass->GetDefaultObject(false)))
 		{
-			if (bInitialize)
-			{
-				HierarchyContainer.Initialize();
-				HierarchyContainer.ResetTransforms();
-			}
-
-			DefaultObject->Hierarchy = HierarchyContainer;
-			DefaultObject->Initialize(true);
-
 			TArray<UObject*> ArchetypeInstances;
 			DefaultObject->GetArchetypeInstances(ArchetypeInstances);
-			for (UObject* ArchetypeInstance : ArchetypeInstances)
+
+			int32 PropertyOffset = InProperty->GetOffset_ReplaceWith_ContainerPtrToValuePtr();
+			int32 PropertySize = InProperty->GetSize();
+
+			switch (InRigElement.Type)
 			{
-				if (UControlRig* InstanceRig = Cast<UControlRig>(ArchetypeInstance))
+				case ERigElementType::Bone:
 				{
-					InstanceRig->Hierarchy = HierarchyContainer;
-					InstanceRig->Initialize(true);
+					uint8* Source = ((uint8*)&HierarchyContainer.BoneHierarchy[ElementIndex]) + PropertyOffset;
+					for (UObject* ArchetypeInstance : ArchetypeInstances)
+					{
+						if (UControlRig* InstanceRig = Cast<UControlRig>(ArchetypeInstance))
+						{
+							uint8* Dest = ((uint8*)&InstanceRig->Hierarchy.BoneHierarchy[ElementIndex]) + PropertyOffset;
+							FMemory::Memcpy(Dest, Source, PropertySize);
+						}
+					}
+					break;
+				}
+				case ERigElementType::Space:
+				{
+					uint8* Source = ((uint8*)&HierarchyContainer.SpaceHierarchy[ElementIndex]) + PropertyOffset;
+					for (UObject* ArchetypeInstance : ArchetypeInstances)
+					{
+						if (UControlRig* InstanceRig = Cast<UControlRig>(ArchetypeInstance))
+						{
+							uint8* Dest = ((uint8*)&InstanceRig->Hierarchy.SpaceHierarchy[ElementIndex]) + PropertyOffset;
+							FMemory::Memcpy(Dest, Source, PropertySize);
+						}
+					}
+					break;
+				}
+				case ERigElementType::Control:
+				{
+					uint8* Source = ((uint8*)&HierarchyContainer.ControlHierarchy[ElementIndex]) + PropertyOffset;
+					for (UObject* ArchetypeInstance : ArchetypeInstances)
+					{
+						if (UControlRig* InstanceRig = Cast<UControlRig>(ArchetypeInstance))
+						{
+							uint8* Dest = ((uint8*)&InstanceRig->Hierarchy.ControlHierarchy[ElementIndex]) + PropertyOffset;
+							FMemory::Memcpy(Dest, Source, PropertySize);
+						}
+					}
+					break;
+				}
+				case ERigElementType::Curve:
+				{
+					uint8* Source = ((uint8*)&HierarchyContainer.CurveContainer[ElementIndex]) + PropertyOffset;
+					for (UObject* ArchetypeInstance : ArchetypeInstances)
+					{
+						if (UControlRig* InstanceRig = Cast<UControlRig>(ArchetypeInstance))
+						{
+							uint8* Dest = ((uint8*)&InstanceRig->Hierarchy.CurveContainer[ElementIndex]) + PropertyOffset;
+							FMemory::Memcpy(Dest, Source, PropertySize);
+						}
+					}
+					break;
+				}
+				default:
+				{
+					break;
 				}
 			}
 		}
 	}
 }
 
+void UControlRigBlueprint::PropagatePropertyFromInstanceToBP(FRigElementKey InRigElement, const FProperty* InProperty, UControlRig* InInstance)
+{
+	int32 ElementIndex = HierarchyContainer.GetIndex(InRigElement);
+	ensure(ElementIndex != INDEX_NONE);
+	check(InProperty);
+
+	int32 PropertyOffset = InProperty->GetOffset_ReplaceWith_ContainerPtrToValuePtr();
+	int32 PropertySize = InProperty->GetSize();
+
+	switch (InRigElement.Type)
+	{
+		case ERigElementType::Bone:
+		{
+			uint8* Source = ((uint8*)&InInstance->Hierarchy.BoneHierarchy[ElementIndex]) + PropertyOffset;
+			uint8* Dest = ((uint8*)&HierarchyContainer.BoneHierarchy[ElementIndex]) + PropertyOffset;
+			FMemory::Memcpy(Dest, Source, PropertySize);
+			break;
+		}
+		case ERigElementType::Space:
+		{
+			uint8* Source = ((uint8*)&InInstance->Hierarchy.SpaceHierarchy[ElementIndex]) + PropertyOffset;
+			uint8* Dest = ((uint8*)&HierarchyContainer.SpaceHierarchy[ElementIndex]) + PropertyOffset;
+			FMemory::Memcpy(Dest, Source, PropertySize);
+			break;
+		}
+		case ERigElementType::Control:
+		{
+			uint8* Source = ((uint8*)&InInstance->Hierarchy.ControlHierarchy[ElementIndex]) + PropertyOffset;
+			uint8* Dest = ((uint8*)&HierarchyContainer.ControlHierarchy[ElementIndex]) + PropertyOffset;
+			FMemory::Memcpy(Dest, Source, PropertySize);
+			break;
+		}
+		case ERigElementType::Curve:
+		{
+			uint8* Source = ((uint8*)&InInstance->Hierarchy.CurveContainer[ElementIndex]) + PropertyOffset;
+			uint8* Dest = ((uint8*)&HierarchyContainer.CurveContainer[ElementIndex]) + PropertyOffset;
+			FMemory::Memcpy(Dest, Source, PropertySize);
+			break;
+		}
+		default:
+		{
+			break;
+		}
+	}
+}
+
 #if WITH_EDITOR
 
 void UControlRigBlueprint::HandleOnElementAdded(FRigHierarchyContainer* InContainer, const FRigElementKey& InKey)
@@ -755,6 +1308,19 @@ void UControlRigBlueprint::HandleOnElementReparented(FRigHierarchyContainer* InC
 
 void UControlRigBlueprint::HandleOnElementSelected(FRigHierarchyContainer* InContainer, const FRigElementKey& InKey, bool bSelected)
 {
+	if (InKey.Type == ERigElementType::Control)
+	{
+		if (UControlRig* RigBeingDebugged = Cast<UControlRig>(GetObjectBeingDebugged()))
+		{
+			if (FRigControl* Control = RigBeingDebugged->FindControl(InKey.Name))
+			{
+				if (!Control->bIsTransientControl)
+				{
+					ClearTransientControls();
+				}
+			}
+		}
+	}
 }
 
 #endif
diff --git a/Engine/Plugins/Experimental/ControlRig/Source/ControlRigDeveloper/Private/ControlRigBlueprintCompiler.cpp b/Engine/Plugins/Experimental/ControlRig/Source/ControlRigDeveloper/Private/ControlRigBlueprintCompiler.cpp
--- a/Engine/Plugins/Experimental/ControlRig/Source/ControlRigDeveloper/Private/ControlRigBlueprintCompiler.cpp
+++ b/Engine/Plugins/Experimental/ControlRig/Source/ControlRigDeveloper/Private/ControlRigBlueprintCompiler.cpp
@@ -13,8 +13,6 @@
 #include "Kismet2/KismetReinstanceUtilities.h"
 #include "Widgets/Notifications/SNotificationList.h"
 #include "Framework/Notifications/NotificationManager.h"
-#include "ControlRigGraphTraverser.h"
-#include "ControlRigDAG.h"
 
 DEFINE_LOG_CATEGORY_STATIC(LogControlRigCompiler, Log, All);
 #define LOCTEXT_NAMESPACE "ControlRigBlueprintCompiler"
@@ -63,65 +61,6 @@ void FControlRigBlueprintCompilerContext::MarkCompilationFailed(const FString& M
 	}
 }
 
-void FControlRigBlueprintCompilerContext::BuildPropertyLinks()
-{
-	DECLARE_SCOPE_HIERARCHICAL_COUNTER_FUNC()
-
-	UControlRigBlueprint* ControlRigBlueprint = Cast<UControlRigBlueprint>(Blueprint);
-	if (ControlRigBlueprint)
-	{
-		// remove all property links
-		ControlRigBlueprint->PropertyLinks.Reset();
-
-		// Build property links from pin links
-		for(UEdGraph* Graph : Blueprint->UbergraphPages)
-		{
-			if (Graph->GetFName() != UControlRigGraphSchema::GraphName_ControlRig)
-			{
-				continue;
-			}
-			UControlRigGraph* RigGraph = Cast<UControlRigGraph>(Graph);
-			if (RigGraph)
-			{
-				if (ControlRigBlueprint->Model == nullptr)
-				{
-					ControlRigBlueprint->PopulateModelFromGraph(RigGraph);
-				}
-
-				FControlRigGraphTraverser Traverser(ControlRigBlueprint->Model);
-				Traverser.TraverseAndBuildPropertyLinks(ControlRigBlueprint);
-
-				bool bEncounteredChange = false;
-				for (UEdGraphNode* Node : RigGraph->Nodes)
-				{
-					UControlRigGraphNode* RigNode = Cast<UControlRigGraphNode>(Node);
-					if (RigNode)
-					{
-						bool bDisplayAsDisabled = !Traverser.IsWiredToExecution(RigNode->PropertyName);
-						if (bDisplayAsDisabled != RigNode->IsDisplayAsDisabledForced())
-						{
-							RigNode->SetForceDisplayAsDisabled(bDisplayAsDisabled);
-							bEncounteredChange = true;
-						}
-					}
-				}
-
-				if (bEncounteredChange)
-				{
-					Graph->NotifyGraphChanged();
-				}
-			}
-		}
-	}
-}
-
-void FControlRigBlueprintCompilerContext::MergeUbergraphPagesIn(UEdGraph* Ubergraph)
-{
-	DECLARE_SCOPE_HIERARCHICAL_COUNTER_FUNC()
-
-	BuildPropertyLinks();
-}
-
 void FControlRigBlueprintCompilerContext::PostCompile()
 {
 	DECLARE_SCOPE_HIERARCHICAL_COUNTER_FUNC()
@@ -129,284 +68,10 @@ void FControlRigBlueprintCompilerContext::PostCompile()
 	UControlRigBlueprint* ControlRigBlueprint = Cast<UControlRigBlueprint>(Blueprint);
 	if (ControlRigBlueprint)
 	{
-		// create sorted graph
-		TArray<FControlRigOperator>& Operators = ControlRigBlueprint->GetControlRigBlueprintGeneratedClass()->Operators;
-		const TArray<FControlRigBlueprintPropertyLink>& PropertyLinks = ControlRigBlueprint->PropertyLinks;
-
-		TArray<FName> UnitNames;
-		TMap<FName, int32> UnitNameToIndex;
-		{
-			DECLARE_SCOPE_HIERARCHICAL_COUNTER(CreateUnitNameMap)
-			for (const FControlRigBlueprintPropertyLink& Link : PropertyLinks)
-			{
-				FName SourceUnitName(*Link.GetSourceUnitName());
-				if (!UnitNameToIndex.Contains(SourceUnitName))
-				{
-					UnitNames.Add(SourceUnitName);
-					UnitNameToIndex.Add(SourceUnitName, UnitNameToIndex.Num());
-				}
-				FName DestUnitName(*Link.GetDestUnitName());
-				if (!UnitNameToIndex.Contains(DestUnitName))
-				{
-					UnitNames.Add(DestUnitName);
-					UnitNameToIndex.Add(DestUnitName, UnitNameToIndex.Num());
-				}
-			}
-		}
-
-		// for UE-4.23.0:
-		// determine if this control rig has been built with a previous version.
-		// to do that we check if we had any mutable units - but no begin execution unit.
-		// given that the current traverser is based on the begin execution unit we know
-		// that the operator stack is from a previous version.
-		bool bIsFromVersionBeforeBeginExecution = false;
-		if (Operators.Num() > 1) // 1 is the "done" operator
-		{
-			bIsFromVersionBeforeBeginExecution = true;
-			for (const FControlRigOperator& Operator : Operators)
-			{
-				FName UnitName = *Operator.CachedPropertyPath1.ToString();
-				FStructProperty* StructProperty = CastField<FStructProperty>(ControlRigBlueprint->GeneratedClass->FindPropertyByName(UnitName));
-				if (StructProperty)
-				{
-					if (StructProperty->Struct->IsChildOf(FRigUnit_BeginExecution::StaticStruct()))
-					{
-						bIsFromVersionBeforeBeginExecution = false;
-						break;
-					}
-				}
-			}
-		}
-
-		int32 PreviousOperatorCount = Operators.Num();
-		Operators.Reset();
-
-		if (UnitNames.Num() > 0)
-		{
-			DECLARE_SCOPE_HIERARCHICAL_COUNTER(CycleDetection)
-
-			// add all of the nodes
-			FControlRigDAG SortGraph;
-			for (int32 UnitIndex = 0; UnitIndex < UnitNames.Num(); UnitIndex++)
-			{
-				bool IsMutableUnit = false;
-				bool IsOutputParameter = false;
-				FStructProperty* StructProperty = CastField<FStructProperty>(ControlRigBlueprint->GeneratedClass->FindPropertyByName(UnitNames[UnitIndex]));
-				if (StructProperty)
-				{
-					if (StructProperty->Struct->IsChildOf(FRigUnitMutable::StaticStruct()))
-					{
-						IsMutableUnit = true;
-					}
-					else if (StructProperty->Struct->IsChildOf(FRigUnit_BeginExecution::StaticStruct()))
-					{
-						IsMutableUnit = true;
-					}
-				}
-
-				const FControlRigModelNode* UnitNode = ControlRigBlueprint->Model->FindNode(UnitNames[UnitIndex]);
-				if (UnitNode)
-				{
-					IsOutputParameter = UnitNode->IsParameter() && UnitNode->ParameterType == EControlRigModelParameterType::Output;
-				}
-				SortGraph.AddNode(IsMutableUnit, IsOutputParameter, UnitNames[UnitIndex]);
-			}
-
-			// add all of the links
-			for (const FControlRigBlueprintPropertyLink& Link : PropertyLinks)
-			{
-				FName SourceUnitName(*Link.GetSourceUnitName());
-				FName DestUnitName(*Link.GetDestUnitName());
-				int32 SourceUnitIndex = UnitNameToIndex.FindChecked(SourceUnitName);
-				int32 DestUnitIndex = UnitNameToIndex.FindChecked(DestUnitName);
-				SortGraph.AddLink(SourceUnitIndex, DestUnitIndex, Link.GetSourceLinkIndex(), Link.GetDestLinkIndex());
-			}
-
-			// enable this for creating a new test case
-			//SortGraph.DumpDag();
-
-			TArray<FControlRigDAG::FNode> UnitOrder, UnitCycle;
-			if (!SortGraph.TopologicalSort(UnitOrder, UnitCycle))
-			{
-#if WITH_EDITORONLY_DATA
-				TSet<FName> UnitNamesInCycle;
-				for (const FControlRigDAG::FNode& NodeInCycle : UnitCycle)
-				{
-					UnitNamesInCycle.Add(NodeInCycle.Name);
-				}
-
-				for (UEdGraph* UbergraphPage : Blueprint->UbergraphPages)
-				{
-					if (UControlRigGraph* ControlRigGraph = Cast<UControlRigGraph>(UbergraphPage))
-					{
-						for (UEdGraphNode* Node : ControlRigGraph->Nodes)
-						{
-							if (UControlRigGraphNode* RigNode = Cast<UControlRigGraphNode>(Node))
-							{
-								if (FStructProperty* Property = RigNode->GetUnitProperty())
-								{
-									if (UnitNamesInCycle.Contains(Property->GetFName()))
-									{
-										RigNode->ErrorMsg = TEXT("The node is part of a cycle.");
-										RigNode->ErrorType = EMessageSeverity::Error;
-										RigNode->bHasCompilerMessage = true;
-
-										MessageLog.Error(*FString::Printf(TEXT("Node '%s' is part of a cycle."), *Property->GetName()));
-									}
-								}
-							}
-						}
-					}
-				}
-#endif
-
-				// we found a cycle - so mark all nodes with errors
-				MarkCompilationFailed(TEXT("The Control Rig compiler detected a cycle in the graph."));
-				return;
-			}
-
-#if WITH_EDITORONLY_DATA
-			// clear the errors on the graph
-			{
-				DECLARE_SCOPE_HIERARCHICAL_COUNTER(ClearNodeErrors)
-
-				for (UEdGraph* UbergraphPage : Blueprint->UbergraphPages)
-				{
-					if (UControlRigGraph* ControlRigGraph = Cast<UControlRigGraph>(UbergraphPage))
-					{
-						for (UEdGraphNode* Node : ControlRigGraph->Nodes)
-						{
-							if (Node->ErrorType < (int32)EMessageSeverity::Info + 1)
-							{
-								Node->ErrorMsg.Reset();
-								Node->ErrorType = (int32)EMessageSeverity::Info + 1;
-								Node->bHasCompilerMessage = false;
-							}
-						}
-					}
-				}
-			}
-#endif
-
-			{
-				DECLARE_SCOPE_HIERARCHICAL_COUNTER(CreateOperatorsForVM)
-
-				for (const FControlRigDAG::FNode Node : UnitOrder)
-				{
-					// copy all properties not originating from a unit
-					for (const FControlRigDAG::FPin& Pin : Node.Outputs)
-					{
-						const int32 PropertyLinkIndex = Pin.Link;
-						const FControlRigBlueprintPropertyLink& Link = PropertyLinks[PropertyLinkIndex];
-						FString SourceUnitName = Link.GetSourceUnitName();
-						FStructProperty* StructProperty = CastField<FStructProperty>(ControlRigBlueprint->GeneratedClass->FindPropertyByName(Node.Name));
-						if (StructProperty)
-						{
-							if (StructProperty->Struct->IsChildOf(FRigUnit::StaticStruct()))
-							{
-								continue;
-							}
-						}
-
-						Operators.Add(FControlRigOperator(EControlRigOpCode::Copy, Link.GetSourcePropertyPath(), Link.GetDestPropertyPath()));
-					}
-
-					FStructProperty* StructProperty = CastField<FStructProperty>(ControlRigBlueprint->GeneratedClass->FindPropertyByName(Node.Name));
-					if (StructProperty)
-					{
-						if (StructProperty->Struct->IsChildOf(FRigUnit::StaticStruct()))
-						{
-							Operators.Add(FControlRigOperator(EControlRigOpCode::Exec, Node.Name.ToString(), FCachedPropertyPath()));
-						}
-					}
-
-					// copy all properties to outputs
-					for (const FControlRigDAG::FPin& Pin : Node.Outputs)
-					{
-						const int32 PropertyLinkIndex = Pin.Link;
-						const FControlRigBlueprintPropertyLink& Link = PropertyLinks[PropertyLinkIndex];
-						Operators.Add(FControlRigOperator(EControlRigOpCode::Copy, Link.GetSourcePropertyPath(), Link.GetDestPropertyPath()));
-					}
-				}
-			}
-
-			Operators.Add(FControlRigOperator(EControlRigOpCode::Done));
-		}
-
-		// guard against the control rig failing due to serialization changes
-		if (PreviousOperatorCount > 1 && Operators.Num() == 1 && bIsFromVersionBeforeBeginExecution)
+		if (!ControlRigBlueprint->bIsRegeneratingOnLoad)
 		{
-			FString Message = FString::Printf(TEXT("The ControlRig '%s' needs to be recompiled in the editor."), *ControlRigBlueprint->GetOuter()->GetPathName());
-			MarkCompilationFailed(Message);
-			return;
-		}
-
-		// update allow source access properties
-		{
-			DECLARE_SCOPE_HIERARCHICAL_COUNTER(AllowAccessToSourceProperties)
-
-			TArray<FName> SourcePropertyLinkArray;
-			TArray<FName> DestPropertyLinkArray;
-
-			auto GetPartialPath = [](FString Input) -> FString
-			{
-				const TCHAR SearchChar = TCHAR('.');
-				int32 ParseIndex = 0;
-				if (Input.FindChar(SearchChar, ParseIndex))
-				{
-					FString Root = Input.Left(ParseIndex+1);
-					FString Child = Input.RightChop(ParseIndex+1);
-					int32 SubParseIndex = 0;
-					
-					if (Child.FindChar(SearchChar, SubParseIndex))
-					{
-						return (Root + Child.Left(SubParseIndex));
-					}
-				}
-
-				return Input;
-			};
-
-			//@todo: think about using ordered properties at the end
-			for (const FControlRigBlueprintPropertyLink& PropertyLink: PropertyLinks)
-			{
-				SourcePropertyLinkArray.Add(FName(*GetPartialPath(PropertyLink.GetSourcePropertyPath())));
-				DestPropertyLinkArray.Add(FName(*GetPartialPath(PropertyLink.GetDestPropertyPath())));
-			}
-
-			ControlRigBlueprint->AllowSourceAccessProperties.Reset();
-
-			TArray<FName> PropertyList;
-			UClass* MyClass = ControlRigBlueprint->GeneratedClass;
-			for (TFieldIterator<FProperty> It(MyClass); It; ++It)
-			{
-				if (FStructProperty* StructProperty = CastField<FStructProperty>(*It))
-				{
-					for (TFieldIterator<FProperty> It2(StructProperty->Struct); It2; ++It2)
-					{
-						if ((*It2)->HasMetaData(TEXT("AllowSourceAccess")))
-						{
-							FString PartialPropertyPath = StructProperty->GetName() + TEXT(".") + (*It2)->GetName();
-							PropertyList.Add(FName(*PartialPropertyPath));
-						}
-					}
-				}
-			}
-
-			// this is prototype code and really slow
-			for (int32 Index = 0; Index < PropertyList.Num(); ++Index)
-			{
-				// find source brute force
- 				const FName& PropertyToSearch = PropertyList[Index];
-				int32 DestIndex = DestPropertyLinkArray.Find(PropertyToSearch);
- 				FName LastSource;
-				if (DestIndex != INDEX_NONE)
-				{
-					FString& Value = ControlRigBlueprint->AllowSourceAccessProperties.Add(PropertyToSearch);
-					Value = SourcePropertyLinkArray[DestIndex].ToString();
-				}
-
-			}
+			// todo: can we find out somehow if we are cooking?
+			ControlRigBlueprint->RecompileVM();
 		}
 	}
 
@@ -414,50 +79,6 @@ void FControlRigBlueprintCompilerContext::PostCompile()
 		DECLARE_SCOPE_HIERARCHICAL_COUNTER(FKismetCompilerContext::PostCompile)
 		FKismetCompilerContext::PostCompile();
 	}
-
-	// ask the model to update all defaults
-	if (ControlRigBlueprint)
-	{
-		if (ControlRigBlueprint->ModelController)
-		{
-			DECLARE_SCOPE_HIERARCHICAL_COUNTER(UpdateModelAfterCompilation)
-
-			// ensure that blueprint storage arrays have the right size.
-			// they might get out of sync due to compilation order.
-			for (const FControlRigModelNode& Node : ControlRigBlueprint->Model->Nodes())
-			{
-				for (const FControlRigModelPin& Pin : Node.Pins)
-				{
-					if (Pin.Direction != EGPD_Input)
-					{
-						continue;
-					}
-					if (!Pin.IsArray())
-					{
-						continue;
-					}
-
-					int32 ArraySize = Pin.ArraySize();
-					FString PinPath = ControlRigBlueprint->Model->GetPinPath(Pin.GetPair());
-					ControlRigBlueprint->PerformArrayOperation(PinPath, [ArraySize](FScriptArrayHelper& InArrayHelper, int32 InArrayIndex)
-					{
-						while (InArrayHelper.Num() < ArraySize)
-						{
-							InArrayHelper.AddValue();
-						}
-						while (InArrayHelper.Num() > ArraySize)
-						{
-							InArrayHelper.RemoveValues(InArrayHelper.Num() - 1);
-						}
-						return true;
-					}, true, true);
-				}
-			}
-
-			TGuardValue<bool> GuardNotifs(ControlRigBlueprint->bSuspendModelNotificationsForOthers, true);
-			ControlRigBlueprint->ModelController->ResendAllPinDefaultNotifications();
-		}
-	}
 }
 
 void FControlRigBlueprintCompilerContext::CopyTermDefaultsToDefaultObject(UObject* DefaultObject)
@@ -473,9 +94,7 @@ void FControlRigBlueprintCompilerContext::CopyTermDefaultsToDefaultObject(UObjec
 
 		UControlRig* ControlRig = CastChecked<UControlRig>(DefaultObject);
 		ControlRig->Hierarchy = ControlRigBlueprint->HierarchyContainer;
-		// copy available rig units info, so that control rig can do things with it
-		ControlRig->AllowSourceAccessProperties = ControlRigBlueprint->AllowSourceAccessProperties;
-		ControlRigBlueprint->UpdateParametersOnControlRig(ControlRig);
+		ControlRig->DrawContainer = ControlRigBlueprint->DrawContainer;
 	}
 }
 
@@ -523,10 +142,6 @@ void FControlRigBlueprintCompilerContext::CleanAndSanitizeClass(UBlueprintGenera
 
 	// Make sure our typed pointer is set
 	check(ClassToClean == NewClass && NewControlRigBlueprintGeneratedClass == NewClass);
-
-	// Reser cached unit properties
-	NewControlRigBlueprintGeneratedClass->ControlUnitProperties.Empty();
-	NewControlRigBlueprintGeneratedClass->RigUnitProperties.Empty();
 }
 
 #undef LOCTEXT_NAMESPACE
diff --git a/Engine/Plugins/Experimental/ControlRig/Source/ControlRigDeveloper/Private/ControlRigBlueprintUtils.cpp b/Engine/Plugins/Experimental/ControlRig/Source/ControlRigDeveloper/Private/ControlRigBlueprintUtils.cpp
--- a/Engine/Plugins/Experimental/ControlRig/Source/ControlRigDeveloper/Private/ControlRigBlueprintUtils.cpp
+++ b/Engine/Plugins/Experimental/ControlRig/Source/ControlRigDeveloper/Private/ControlRigBlueprintUtils.cpp
@@ -11,54 +11,6 @@
 
 #define LOCTEXT_NAMESPACE "ControlRigBlueprintUtils"
 
-FName FControlRigBlueprintUtils::GetNewUnitMemberName(UBlueprint* InBlueprint, const UStruct* InStructTemplate)
-{
-	DECLARE_SCOPE_HIERARCHICAL_COUNTER_FUNC()
-
-	FString VariableBaseName = InStructTemplate->GetName();
-	VariableBaseName.RemoveFromStart(TEXT("RigUnit_"));
-	return FBlueprintEditorUtils::FindUniqueKismetName(InBlueprint, VariableBaseName);
-}
-
-FName FControlRigBlueprintUtils::AddUnitMember(UBlueprint* InBlueprint, const UStruct* InStructTemplate, const FName& InName)
-{
-	DECLARE_SCOPE_HIERARCHICAL_COUNTER_FUNC()
-
-	FName VarName = InName == NAME_None ? FControlRigBlueprintUtils::GetNewUnitMemberName(InBlueprint, InStructTemplate) : InName;
-
-	UScriptStruct* ScriptStruct = FindObjectChecked<UScriptStruct>(ANY_PACKAGE, *InStructTemplate->GetName());
-	UEdGraphSchema_K2 const* K2Schema = GetDefault<UEdGraphSchema_K2>();
-	if(FBlueprintEditorUtils::AddMemberVariable(InBlueprint, VarName, FEdGraphPinType(UEdGraphSchema_K2::PC_Struct, InStructTemplate->GetFName(), ScriptStruct, EPinContainerType::None, false, FEdGraphTerminalType())))
-	{
-		FBPVariableDescription& Variable = InBlueprint->NewVariables.Last();
-		Variable.Category = LOCTEXT("UnitsCategory", "Units");
-
-		FBlueprintEditorUtils::MarkBlueprintAsStructurallyModified(InBlueprint);	
-
-		return Variable.VarName;
-	}
-
-	return NAME_None;
-}
-
-FName FControlRigBlueprintUtils::GetNewPropertyMemberName(UBlueprint* InBlueprint, const FString& InVariableDesc)
-{
-	DECLARE_SCOPE_HIERARCHICAL_COUNTER_FUNC()
-
-	return FBlueprintEditorUtils::FindUniqueKismetName(InBlueprint, InVariableDesc);
-}
-
-FName FControlRigBlueprintUtils::AddPropertyMember(UBlueprint* InBlueprint, const FEdGraphPinType& InPinType, const FString& InVariableDesc)
-{
-	DECLARE_SCOPE_HIERARCHICAL_COUNTER_FUNC()
-
-	if(FBlueprintEditorUtils::AddMemberVariable(InBlueprint, *InVariableDesc, InPinType))
-	{
-		return InBlueprint->NewVariables.Last().VarName;
-	}
-	return NAME_None;
-}
-
 FName FControlRigBlueprintUtils::ValidateName(UBlueprint* InBlueprint, const FString& InName)
 {
 	DECLARE_SCOPE_HIERARCHICAL_COUNTER_FUNC()
@@ -115,77 +67,7 @@ FName FControlRigBlueprintUtils::ValidateName(UBlueprint* InBlueprint, const FSt
 	return *Name;
 }
 
-UControlRigGraphNode* FControlRigBlueprintUtils::InstantiateGraphNodeForProperty(UEdGraph* InGraph, const FName& InPropertyName, const FVector2D& InLocation, const FEdGraphPinType& InPinType)
-{
-	DECLARE_SCOPE_HIERARCHICAL_COUNTER_FUNC()
-
-	check(InGraph);
-
-	InGraph->Modify();
-
-	UControlRigGraphNode* NewNode = NewObject<UControlRigGraphNode>(InGraph);
-	NewNode->SetPropertyName(InPropertyName);
-	NewNode->PinType = InPinType;
-
-	InGraph->AddNode(NewNode, true);
-
-	NewNode->CreateNewGuid();
-	NewNode->PostPlacedNewNode();
-	NewNode->AllocateDefaultPins();
-
-	NewNode->NodePosX = InLocation.X;
-	NewNode->NodePosY = InLocation.Y;
-
-	NewNode->SetFlags(RF_Transactional);
-
-	return NewNode;
-}
-
-UControlRigGraphNode* FControlRigBlueprintUtils::InstantiateGraphNodeForStructPath(UEdGraph* InGraph, const FName& InPropertyName, const FVector2D& InLocation, const FString& InStructPath)
-{
-	DECLARE_SCOPE_HIERARCHICAL_COUNTER_FUNC()
-
-	check(InGraph);
-
-	InGraph->Modify();
-
-	UControlRigGraphNode* NewNode = NewObject<UControlRigGraphNode>(InGraph);
-	NewNode->SetPropertyName(InPropertyName);
-	NewNode->StructPath = InStructPath;
-
-	InGraph->AddNode(NewNode, true);
-
-	NewNode->CreateNewGuid();
-	NewNode->PostPlacedNewNode();
-	NewNode->AllocateDefaultPins();
-
-	NewNode->NodePosX = InLocation.X;
-	NewNode->NodePosY = InLocation.Y;
-
-	NewNode->SetFlags(RF_Transactional);
-
-	return NewNode;
-}
-
-bool FControlRigBlueprintUtils::CanInstantiateGraphNodeForProperty(UEdGraph* InGraph, const FName& InPropertyName)
-{
-	DECLARE_SCOPE_HIERARCHICAL_COUNTER_FUNC()
-
-	for(UEdGraphNode* Node : InGraph->Nodes)
-	{
-		if(UControlRigGraphNode* ControlRigGraphNode = Cast<UControlRigGraphNode>(Node))
-		{
-			if(ControlRigGraphNode->GetPropertyName() == InPropertyName)
-			{
-				return false;
-			}
-		}
-	}
-
-	return true;
-}
-
-void FControlRigBlueprintUtils::ForAllRigUnits(TFunction<void(UStruct*)> InFunction)
+void FControlRigBlueprintUtils::ForAllRigUnits(TFunction<void(UScriptStruct*)> InFunction)
 {
 	DECLARE_SCOPE_HIERARCHICAL_COUNTER_FUNC()
 
@@ -194,7 +76,10 @@ void FControlRigBlueprintUtils::ForAllRigUnits(TFunction<void(UStruct*)> InFunct
 	{
 		if(StructIt->IsChildOf(FRigUnit::StaticStruct()) && !StructIt->HasMetaData(UControlRig::AbstractMetaName))
 		{
-			InFunction(*StructIt);
+			if (UScriptStruct* ScriptStruct = Cast<UScriptStruct>(*StructIt))
+			{
+				InFunction(ScriptStruct);
+			}
 		}
 	}
 }
@@ -203,16 +88,7 @@ void FControlRigBlueprintUtils::HandleReconstructAllNodes(UBlueprint* InBlueprin
 {
 	DECLARE_SCOPE_HIERARCHICAL_COUNTER_FUNC()
 
-	if(InBlueprint->IsA<UControlRigBlueprint>())
-	{
-		// TArray<UControlRigGraphNode*> AllNodes;
-		// FBlueprintEditorUtils::GetAllNodesOfClass(InBlueprint, AllNodes);
-
-		// for(UControlRigGraphNode* Node : AllNodes)
-		// {
-		// 	Node->ReconstructNode();
-		// }
-	}
+	return HandleRefreshAllNodes(InBlueprint);
 }
 
 void FControlRigBlueprintUtils::HandleRefreshAllNodes(UBlueprint* InBlueprint)
@@ -231,43 +107,13 @@ void FControlRigBlueprintUtils::HandleRefreshAllNodes(UBlueprint* InBlueprint)
 	}
 }
 
-void FControlRigBlueprintUtils::HandleRenameVariableReferencesEvent(UBlueprint* InBlueprint, UClass* InVariableClass, const FName& InOldVarName, const FName& InNewVarName)
-{
-	DECLARE_SCOPE_HIERARCHICAL_COUNTER_FUNC()
-
-	if(UControlRigBlueprint* RigBlueprint = Cast<UControlRigBlueprint>(InBlueprint))
-	{
-		if (RigBlueprint->ModelController)
-		{
-			RigBlueprint->ModelController->RenameNode(InOldVarName, InNewVarName);
-		}
-	}
-}
-
 void FControlRigBlueprintUtils::RemoveMemberVariableIfNotUsed(UBlueprint* Blueprint, const FName VarName, UControlRigGraphNode* ToBeDeleted)
 {
 	DECLARE_SCOPE_HIERARCHICAL_COUNTER_FUNC()
 
 	if (Blueprint->IsA<UControlRigBlueprint>())
 	{
-		bool bDeleteVariable = true;
-
-		TArray<UControlRigGraphNode*> AllNodes;
-		FBlueprintEditorUtils::GetAllNodesOfClass(Blueprint, AllNodes);
-
-		for (UControlRigGraphNode* Node : AllNodes)
-		{
-			if (Node != ToBeDeleted && Node->GetPropertyName() == VarName)
-			{
-				bDeleteVariable = false;
-				break;
-			}
-		}
-
-		if (bDeleteVariable)
-		{
-			FBlueprintEditorUtils::RemoveMemberVariable(Blueprint, VarName);
-		}
+		FBlueprintEditorUtils::RemoveMemberVariable(Blueprint, VarName);
 	}
 }
 #undef LOCTEXT_NAMESPACE
\ No newline at end of file
diff --git a/Engine/Plugins/Experimental/ControlRig/Source/ControlRigDeveloper/Private/ControlRigDeveloper.cpp b/Engine/Plugins/Experimental/ControlRig/Source/ControlRigDeveloper/Private/ControlRigDeveloper.cpp
--- a/Engine/Plugins/Experimental/ControlRig/Source/ControlRigDeveloper/Private/ControlRigDeveloper.cpp
+++ b/Engine/Plugins/Experimental/ControlRig/Source/ControlRigDeveloper/Private/ControlRigDeveloper.cpp
@@ -8,6 +8,8 @@
 
 #define LOCTEXT_NAMESPACE "ControlRigDeveloperModule"
 
+DEFINE_LOG_CATEGORY(LogControlRigDeveloper);
+
 class FControlRigDeveloperModule : public IControlRigDeveloperModule
 {
 public:
diff --git a/Engine/Plugins/Experimental/ControlRig/Source/ControlRigDeveloper/Private/ControlRigHierarchyModifier.cpp b/Engine/Plugins/Experimental/ControlRig/Source/ControlRigDeveloper/Private/ControlRigHierarchyModifier.cpp
new file mode 100644
--- /dev/null
+++ b/Engine/Plugins/Experimental/ControlRig/Source/ControlRigDeveloper/Private/ControlRigHierarchyModifier.cpp
@@ -0,0 +1,490 @@
+// Copyright Epic Games, Inc. All Rights Reserved.
+
+#include "ControlRigHierarchyModifier.h"
+
+UControlRigHierarchyModifier::UControlRigHierarchyModifier()
+{
+	Container = nullptr;
+}
+
+#if WITH_EDITOR
+
+TArray<FRigElementKey> UControlRigHierarchyModifier::GetElements() const
+{
+	TArray<FRigElementKey> Elements;
+	if (Container != nullptr)
+	{
+		return Container->GetAllItems(false /* sort */);
+	}
+	return Elements;
+}
+
+FRigElementKey UControlRigHierarchyModifier::AddBone(const FName& InNewName, const FName& InParentName, ERigBoneType InType)
+{
+	if(Container != nullptr)
+	{
+		return Container->BoneHierarchy.Add(InNewName, InParentName, InType).GetElementKey();
+	}
+	return FRigElementKey();
+}
+
+FRigBone UControlRigHierarchyModifier::GetBone(const FRigElementKey& InKey)
+{
+	if (Container != nullptr && InKey.Type == ERigElementType::Bone)
+	{
+		int32 Index = Container->GetIndex(InKey);
+		if (Index != INDEX_NONE)
+		{
+			return Container->BoneHierarchy.Bones[Index];
+		}
+	}
+	return InvalidBone;
+}
+
+void UControlRigHierarchyModifier::SetBone(const FRigBone& InElement)
+{
+	if (Container != nullptr)
+	{
+		int32 Index = Container->GetIndex(InElement.GetElementKey());
+		if (Index != INDEX_NONE)
+		{
+			Container->BoneHierarchy.Bones[Index] = InElement;
+			Container->BoneHierarchy.Initialize();
+		}
+	}
+}
+
+FRigElementKey UControlRigHierarchyModifier::AddControl(
+	const FName& InNewName,
+	ERigControlType InControlType,
+	const FName& InParentName,
+	const FName& InSpaceName,
+	const FName& InGizmoName,
+	const FLinearColor& InGizmoColor
+)
+{
+	if(Container != nullptr)
+	{
+		return Container->ControlHierarchy.Add(InNewName, InControlType, InParentName, InSpaceName, FRigControlValue(), InGizmoName, FTransform::Identity, InGizmoColor).GetElementKey();
+	}
+	return FRigElementKey();
+}
+
+
+FRigControl UControlRigHierarchyModifier::GetControl(const FRigElementKey& InKey)
+{
+	if (Container != nullptr && InKey.Type == ERigElementType::Control)
+	{
+		int32 Index = Container->GetIndex(InKey);
+		if (Index != INDEX_NONE)
+		{
+			return Container->ControlHierarchy.Controls[Index];
+		}
+	}
+	return InvalidControl;
+}
+
+void UControlRigHierarchyModifier::SetControl(const FRigControl& InElement)
+{
+	if (Container != nullptr)
+	{
+		int32 Index = Container->GetIndex(InElement.GetElementKey());
+		if (Index != INDEX_NONE)
+		{
+			Container->ControlHierarchy.Controls[Index] = InElement;
+			Container->ControlHierarchy.Initialize();
+		}
+	}
+}
+
+FRigElementKey UControlRigHierarchyModifier::AddSpace
+(
+	const FName& InNewName,
+	ERigSpaceType InSpaceType,
+	const FName& InParentName
+)
+{
+	if(Container != nullptr)
+	{
+		return Container->SpaceHierarchy.Add(InNewName, InSpaceType, InParentName).GetElementKey();
+	}
+	return FRigElementKey();
+}
+
+FRigSpace UControlRigHierarchyModifier::GetSpace(const FRigElementKey& InKey)
+{
+	if (Container != nullptr && InKey.Type == ERigElementType::Space)
+	{
+		int32 Index = Container->GetIndex(InKey);
+		if (Index != INDEX_NONE)
+		{
+			return Container->SpaceHierarchy.Spaces[Index];
+		}
+	}
+	return InvalidSpace;
+}
+
+void UControlRigHierarchyModifier::SetSpace(const FRigSpace& InElement)
+{
+	if (Container != nullptr)
+	{
+		int32 Index = Container->GetIndex(InElement.GetElementKey());
+		if (Index != INDEX_NONE)
+		{
+			Container->SpaceHierarchy.Spaces[Index] = InElement;
+			Container->SpaceHierarchy.Initialize();
+		}
+	}
+}
+
+FRigElementKey UControlRigHierarchyModifier::AddCurve(const FName& InNewName, float InValue)
+{
+	if(Container != nullptr)
+	{
+		return Container->CurveContainer.Add(InNewName, InValue).GetElementKey();
+	}
+	return FRigElementKey();
+}
+
+FRigCurve UControlRigHierarchyModifier::GetCurve(const FRigElementKey& InKey)
+{
+	if (Container != nullptr && InKey.Type == ERigElementType::Curve)
+	{
+		int32 Index = Container->GetIndex(InKey);
+		if (Index != INDEX_NONE)
+		{
+			return Container->CurveContainer.Curves[Index];
+		}
+	}
+	return InvalidCurve;
+}
+
+void UControlRigHierarchyModifier::SetCurve(const FRigCurve& InElement)
+{
+	if (Container != nullptr)
+	{
+		int32 Index = Container->GetIndex(InElement.GetElementKey());
+		if (Index != INDEX_NONE)
+		{
+			Container->CurveContainer.Curves[Index] = InElement;
+			Container->CurveContainer.Initialize();
+		}
+	}
+}
+
+bool UControlRigHierarchyModifier::RemoveElement(const FRigElementKey& InElement)
+{
+	if(Container == nullptr)
+	{
+		return false;
+	}
+
+	if (Container->GetIndex(InElement) == INDEX_NONE)
+	{
+		return false;
+	}
+
+	switch (InElement.Type)
+	{
+		case ERigElementType::Bone:
+		{
+			Container->BoneHierarchy.Remove(InElement.Name);
+			return true;
+		}
+		case ERigElementType::Control:
+		{
+			Container->ControlHierarchy.Remove(InElement.Name);
+			return true;
+		}
+		case ERigElementType::Space:
+		{
+			Container->SpaceHierarchy.Remove(InElement.Name);
+			return true;
+		}
+		case ERigElementType::Curve:
+		{
+			Container->CurveContainer.Remove(InElement.Name);
+			return true;
+		}
+		default:
+		{
+			ensure(false);
+			break;
+		}
+	}
+	return false;
+}
+
+FRigElementKey UControlRigHierarchyModifier::RenameElement(const FRigElementKey& InElement, const FName& InNewName)
+{
+	if(Container == nullptr)
+	{
+		return FRigElementKey();
+	}
+
+	if (Container->GetIndex(InElement) == INDEX_NONE)
+	{
+		return FRigElementKey();
+	}
+
+	switch (InElement.Type)
+	{
+		case ERigElementType::Bone:
+		{
+			FName NewName = Container->BoneHierarchy.Rename(InElement.Name, InNewName);
+			return FRigElementKey(NewName, ERigElementType::Bone);
+		}
+		case ERigElementType::Control:
+		{
+			FName NewName = Container->ControlHierarchy.Rename(InElement.Name, InNewName);
+			return FRigElementKey(NewName, ERigElementType::Control);
+		}
+		case ERigElementType::Space:
+		{
+			FName NewName = Container->SpaceHierarchy.Rename(InElement.Name, InNewName);
+			return FRigElementKey(NewName, ERigElementType::Space);
+		}
+		case ERigElementType::Curve:
+		{
+			FName NewName = Container->CurveContainer.Rename(InElement.Name, InNewName);
+			return FRigElementKey(NewName, ERigElementType::Curve);
+		}
+		default:
+		{
+			ensure(false);
+			break;
+		}
+	}
+	return FRigElementKey();
+}
+
+bool UControlRigHierarchyModifier::ReparentElement(const FRigElementKey& InElement, const FRigElementKey& InNewParent)
+{
+	if(Container == nullptr)
+	{
+		return false;
+	}
+
+	if (Container->GetIndex(InElement) == INDEX_NONE)
+	{
+		return false;
+	}
+
+	switch (InElement.Type)
+	{
+		case ERigElementType::Bone:
+		{
+			if (InNewParent.Type == ERigElementType::Bone)
+			{
+				return Container->BoneHierarchy.Reparent(InElement.Name, InNewParent.Name);
+			}
+			break;
+		}
+		case ERigElementType::Control:
+		{
+			if (InNewParent.Type == ERigElementType::Control)
+			{
+				if (InNewParent.Name != NAME_None)
+				{
+					Container->ControlHierarchy.SetSpace(InElement.Name, NAME_None);
+				}
+				return Container->ControlHierarchy.Reparent(InElement.Name, InNewParent.Name);
+			}
+			else if (InNewParent.Type == ERigElementType::Space)
+			{
+				Container->ControlHierarchy.SetSpace(InElement.Name, InNewParent.Name);
+				return true;
+			}
+			break;
+		}
+		case ERigElementType::Space:
+		{
+			if (InNewParent.Name == NAME_None)
+			{
+				return Container->SpaceHierarchy.Reparent(InElement.Name, ERigSpaceType::Global, InNewParent.Name);
+			}
+			else if (InNewParent.Type == ERigElementType::Bone)
+			{
+				return Container->SpaceHierarchy.Reparent(InElement.Name, ERigSpaceType::Bone, InNewParent.Name);
+			}
+			else if (InNewParent.Type == ERigElementType::Control)
+			{
+				return Container->SpaceHierarchy.Reparent(InElement.Name, ERigSpaceType::Control, InNewParent.Name);
+			}
+			else if (InNewParent.Type == ERigElementType::Space)
+			{
+				return Container->SpaceHierarchy.Reparent(InElement.Name, ERigSpaceType::Space, InNewParent.Name);
+			}
+			break;
+		}
+		case ERigElementType::Curve:
+		{
+			break;
+		}
+		default:
+		{
+			ensure(false);
+			break;
+		}
+	}
+	return false;
+}
+
+TArray<FRigElementKey> UControlRigHierarchyModifier::GetSelection() const
+{
+	TArray<FRigElementKey> Elements;
+	if (Container != nullptr)
+	{
+		return Container->CurrentSelection();
+	}
+	return Elements;
+}
+
+bool UControlRigHierarchyModifier::Select(const FRigElementKey& InKey, bool bSelect)
+{
+	if(Container == nullptr)
+	{
+		return false;
+	}
+	return Container->Select(InKey, bSelect);
+}
+
+bool UControlRigHierarchyModifier::ClearSelection()
+{
+	if(Container == nullptr)
+	{
+		return false;
+	}
+	return Container->ClearSelection();
+}
+
+bool UControlRigHierarchyModifier::IsSelected(const FRigElementKey& InKey) const
+{
+	if(Container == nullptr)
+	{
+		return false;
+	}
+	return Container->IsSelected(InKey);
+}
+
+void UControlRigHierarchyModifier::Initialize(bool bResetTransforms)
+{
+	if(Container == nullptr)
+	{
+		return;
+	}
+	Container->Initialize(bResetTransforms);
+}
+
+void UControlRigHierarchyModifier::Reset()
+{
+	if(Container == nullptr)
+	{
+		return;
+	}
+	Container->Reset();
+}
+
+void UControlRigHierarchyModifier::ResetTransforms()
+{
+	if(Container == nullptr)
+	{
+		return;
+	}
+	Container->ResetTransforms();
+}
+
+FTransform UControlRigHierarchyModifier::GetInitialTransform(const FRigElementKey& InKey) const
+{
+	if(Container == nullptr)
+	{
+		return FTransform::Identity;
+	}
+	return Container->GetInitialTransform(InKey);
+}
+
+void UControlRigHierarchyModifier::SetInitialTransform(const FRigElementKey& InKey, const FTransform& InTransform)
+{
+	if(Container == nullptr)
+	{
+		return;
+	}
+	return Container->SetInitialTransform(InKey, InTransform);
+}
+
+FTransform UControlRigHierarchyModifier::GetInitialGlobalTransform(const FRigElementKey& InKey) const
+{
+	if(Container == nullptr)
+	{
+		return FTransform::Identity;
+	}
+	return Container->GetInitialGlobalTransform(InKey);
+
+}
+
+void UControlRigHierarchyModifier::SetInitialGlobalTransform(const FRigElementKey& InKey, const FTransform& InTransform)
+{
+	if(Container == nullptr)
+	{
+		return;
+	}
+	return Container->SetInitialGlobalTransform(InKey, InTransform);
+}
+
+FTransform UControlRigHierarchyModifier::GetLocalTransform(const FRigElementKey& InKey) const
+{
+	if(Container == nullptr)
+	{
+		return FTransform::Identity;
+	}
+	return Container->GetLocalTransform(InKey);
+
+}
+
+void UControlRigHierarchyModifier::SetLocalTransform(const FRigElementKey& InKey, const FTransform& InTransform)
+{
+	if(Container == nullptr)
+	{
+		return;
+	}
+	return Container->SetLocalTransform(InKey, InTransform);
+}
+
+FTransform UControlRigHierarchyModifier::GetGlobalTransform(const FRigElementKey& InKey) const
+{
+	if(Container == nullptr)
+	{
+		return FTransform::Identity;
+	}
+	return Container->GetGlobalTransform(InKey);
+
+}
+
+void UControlRigHierarchyModifier::SetGlobalTransform(const FRigElementKey& InKey, const FTransform& InTransform)
+{
+	if(Container == nullptr)
+	{
+		return;
+	}
+	return Container->SetGlobalTransform(InKey, InTransform);
+}
+
+FString UControlRigHierarchyModifier::ExportToText(const TArray<FRigElementKey>& InElementsToExport) const
+{
+	if(Container == nullptr)
+	{
+		return FString();
+	}
+	return Container->ExportToText(InElementsToExport);
+}
+
+TArray<FRigElementKey> UControlRigHierarchyModifier::ImportFromText(const FString& InContent, ERigHierarchyImportMode InImportMode, bool bSelectNewElements)
+{
+	if(Container == nullptr)
+	{
+		return TArray<FRigElementKey>();
+	}
+	return Container->ImportFromText(InContent, InImportMode, bSelectNewElements);
+}
+
+#endif
\ No newline at end of file
