COPY from //Dev-Anim to //Dev-Main

diff --git a/Engine/Plugins/Experimental/ControlRig/Source/ControlRigDeveloper/Private/Graph/ControlRigGraph.cpp b/Engine/Plugins/Experimental/ControlRig/Source/ControlRigDeveloper/Private/Graph/ControlRigGraph.cpp
--- a/Engine/Plugins/Experimental/ControlRig/Source/ControlRigDeveloper/Private/Graph/ControlRigGraph.cpp
+++ b/Engine/Plugins/Experimental/ControlRig/Source/ControlRigDeveloper/Private/Graph/ControlRigGraph.cpp
@@ -5,7 +5,7 @@
 #include "ControlRigBlueprintGeneratedClass.h"
 #include "Graph/ControlRigGraphSchema.h"
 #include "ControlRig.h"
-#include "ControlRigModel.h"
+#include "RigVMModel/RigVMGraph.h"
 #include "ControlRigObjectVersion.h"
 #include "Units/RigUnit.h"
 #include "EdGraphNode_Comment.h"
@@ -25,6 +25,7 @@ UControlRigGraph::UControlRigGraph()
 {
 	bSuspendModelNotifications = false;
 	bIsTemporaryGraphForCopyPaste = false;
+	bIsSelecting = false;
 }
 
 void UControlRigGraph::Initialize(UControlRigBlueprint* InBlueprint)
@@ -32,7 +33,7 @@ void UControlRigGraph::Initialize(UControlRigBlueprint* InBlueprint)
 	DECLARE_SCOPE_HIERARCHICAL_COUNTER_FUNC()
 
 	InBlueprint->OnModified().RemoveAll(this);
-	InBlueprint->OnModified().AddUObject(this, &UControlRigGraph::HandleModelModified);
+	InBlueprint->OnModified().AddUObject(this, &UControlRigGraph::HandleModifiedEvent);
 }
 
 const UControlRigGraphSchema* UControlRigGraph::GetControlRigGraphSchema()
@@ -50,235 +51,17 @@ void UControlRigGraph::Serialize(FArchive& Ar)
 #endif
 
 #if WITH_EDITOR
-void UControlRigGraph::PostLoad()
-{
-	DECLARE_SCOPE_HIERARCHICAL_COUNTER_FUNC()
-
-	FoundHierarchyRefVariableNodes.Reset();
-	FoundHierarchyRefMutableNodes.Reset();
-	FoundHierarchyRefConnections.Reset();
-
-	UBlueprint* Blueprint = Cast<UBlueprint>(GetOuter());
-	if (Blueprint)
-	{
-		UClass* BlueprintClass = Blueprint->GeneratedClass;
-
-		// perform fixes on the graph for backwards compatibility
-		if (GetLinkerCustomVersion(FControlRigObjectVersion::GUID) < FControlRigObjectVersion::RemovalOfHierarchyRefPins)
-		{
-#if WITH_EDITORONLY_DATA
-			for (UEdGraphNode* Node : Nodes)
-			{
-				UControlRigGraphNode* RigNode = Cast<UControlRigGraphNode>(Node);
-				if (RigNode != nullptr)
-				{
-					// store the nodes connected to outputs of hierarchy refs.
-					// this is done for backwards compatibility
-					if (RigNode->HasAnyFlags(RF_NeedPostLoad))
-					{
-						RigNode->CacheHierarchyRefConnectionsOnPostLoad();
-					}
-				}
-			}
-#endif
-
-			for (UEdGraphNode* Node : Nodes)
-			{
-				UControlRigGraphNode* RigNode = Cast<UControlRigGraphNode>(Node);
-				if (RigNode != nullptr)
-				{
-					FStructProperty* Property = CastField<FStructProperty>(BlueprintClass->FindPropertyByName(RigNode->GetPropertyName()));
-					if (Property != nullptr)
-					{
-						// found the hierarchy ref variable
-						if (Property->Struct == FRigHierarchyRef::StaticStruct())
-						{
-							FoundHierarchyRefVariableNodes.Add(RigNode);
-						}
-						// found a former "hierarchy ref" utilizing unit
-						else if (Property->Struct->IsChildOf(FRigUnitMutable::StaticStruct()))
-						{
-							FoundHierarchyRefMutableNodes.Add(RigNode);
-						}
-						else
-						{
-							continue;
-						}
-
-						TArray<UControlRigGraphNode*> LinkedNodes;
-						for(UEdGraphNode* LinkedNode : RigNode->HierarchyRefOutputConnections)
-						{
-							LinkedNodes.Add(CastChecked<UControlRigGraphNode>(LinkedNode));
-						}
-						FoundHierarchyRefConnections.Add(RigNode, LinkedNodes);
-					}
-				}
-			}
-		}
-	}
-
-	Super::PostLoad();
-
-	UControlRigBlueprint* RigBlueprint = Cast<UControlRigBlueprint>(Blueprint);
-	if (RigBlueprint)
-	{
-		if (GetLinkerCustomVersion(FControlRigObjectVersion::GUID) < FControlRigObjectVersion::RemovalOfHierarchyRefPins)
-		{
-			if (BlueprintOnCompiledHandle.IsValid())
-			{
-				Blueprint->OnCompiled().Remove(BlueprintOnCompiledHandle);
-			}
-			BlueprintOnCompiledHandle = Blueprint->OnCompiled().AddUObject(this, &UControlRigGraph::OnBlueprintCompiledPostLoad);
-		}
-
-		RigBlueprint->PopulateModelFromGraph(this);
-	}
-}
 
-void UControlRigGraph::OnBlueprintCompiledPostLoad(UBlueprint* InCompiledBlueprint)
+void UControlRigGraph::CacheNameLists(const FRigHierarchyContainer* HierarchyContainer, const FControlRigDrawContainer* DrawContainer)
 {
-	DECLARE_SCOPE_HIERARCHICAL_COUNTER_FUNC()
-
-	if (GetLinkerCustomVersion(FControlRigObjectVersion::GUID) < FControlRigObjectVersion::RemovalOfHierarchyRefPins)
-	{
-		UControlRigBlueprint* RigBlueprint = Cast<UControlRigBlueprint>(GetOuter());
-		ensure(InCompiledBlueprint == RigBlueprint);
-		RigBlueprint->OnCompiled().Remove(BlueprintOnCompiledHandle);
-		BlueprintOnCompiledHandle.Reset();
-
-		struct FOutStandingLink
-		{
-			FOutStandingLink() {}
-			FOutStandingLink(const FName& InA, const FName& InB, const FName& InC, const FName& InD)
-			{
-				A = InA;
-				B = InB;
-				C = InC;
-				D = InD;
-			}
-
-			FName A;
-			FName B;
-			FName C;
-			FName D;
-		};
-		TArray<FOutStandingLink> OutStandingLinks;
-
-		// create a new "begin execution" unit for each branch
-		for (UControlRigGraphNode* RigNode : FoundHierarchyRefVariableNodes)
-		{
-			TArray<UControlRigGraphNode*>& ConnectedNodes = FoundHierarchyRefConnections.FindChecked(RigNode);
-			for (UControlRigGraphNode* ConnectedNode : ConnectedNodes)
-			{
-				int32 NodePosX = ConnectedNode->NodePosX - 200;
-				int32 NodePosY = ConnectedNode->NodePosY;
-
-				if (RigBlueprint->ModelController->AddNode(FRigUnit_BeginExecution::StaticStruct()->GetFName(), FVector2D((float)NodePosX, (float)NodePosY)))
-				{
-					FName BeginExecNode = RigBlueprint->LastNameFromNotification;
-					if (BeginExecNode != NAME_None)
-					{
-						for (UEdGraphPin* InputPin : ConnectedNode->Pins)
-						{
-							if (InputPin->Direction != EEdGraphPinDirection::EGPD_Input)
-							{
-								continue;
-							}
-							if (InputPin->PinType.PinCategory != UEdGraphSchema_K2::PC_Struct)
-							{
-								continue;
-							}
-							if (InputPin->PinType.PinSubCategoryObject != FRigHierarchyRef::StaticStruct())
-							{
-								continue;
-							}
-
-							OutStandingLinks.Add(FOutStandingLink(BeginExecNode, TEXT("ExecuteContext"), ConnectedNode->PropertyName, TEXT("ExecuteContext")));
-						}
-					}
-				}
-			}
-
-			FBlueprintEditorUtils::RemoveNode(RigBlueprint, RigNode, true);
-		}
-
-		// wire up old hierarchy ref connections to new execution connections
-		for (UControlRigGraphNode* RigNode : FoundHierarchyRefMutableNodes)
-		{
-			for (UEdGraphPin* OutputPin : RigNode->Pins)
-			{
-				if (OutputPin->Direction != EEdGraphPinDirection::EGPD_Output)
-				{
-					continue;
-				}
-				if (OutputPin->PinType.PinCategory != UEdGraphSchema_K2::PC_Struct)
-				{
-					continue;
-				}
-				if (OutputPin->PinType.PinSubCategoryObject != FRigHierarchyRef::StaticStruct())
-				{
-					continue;
-				}
-
-				for (UEdGraphPin* InputPin : OutputPin->LinkedTo)
-				{
-					if (InputPin->Direction != EEdGraphPinDirection::EGPD_Input)
-					{
-						continue;
-					}
-					if (InputPin->PinType.PinCategory != UEdGraphSchema_K2::PC_Struct)
-					{
-						continue;
-					}
-					if (InputPin->PinType.PinSubCategoryObject != FRigHierarchyRef::StaticStruct())
-					{
-						continue;
-					}
-
-					UControlRigGraphNode* InputNode = Cast<UControlRigGraphNode>(InputPin->GetOwningNode());
-					OutStandingLinks.Add(FOutStandingLink(RigNode->PropertyName, TEXT("ExecuteContext"), InputNode->PropertyName, TEXT("ExecuteContext")));
-				}
-			}
-		}
-
-		for (UEdGraphNode* Node : Nodes)
-		{
-			UControlRigGraphNode* RigNode = Cast<UControlRigGraphNode>(Node);
-			if (RigNode != nullptr)
-			{
-				RigNode->ReconstructNode();
-			}
-		}
-
-		for (const FOutStandingLink& Link : OutStandingLinks)
-		{
-			RigBlueprint->ModelController->MakeLink(Link.A, Link.B, Link.C, Link.D);
-		}
-
-		FBlueprintEditorUtils::MarkBlueprintAsStructurallyModified(RigBlueprint);
-
-		FoundHierarchyRefVariableNodes.Reset();
-		FoundHierarchyRefMutableNodes.Reset();
-		FoundHierarchyRefConnections.Reset();
-
-		FNotificationInfo Info(LOCTEXT("ControlRigUpdatedHelpMessage", "The Control Rig has automatically been updated to use execution pins. You will need to compile and re-save."));
-		Info.bFireAndForget = true;
-		Info.FadeOutDuration = 10.0f;
-		Info.ExpireDuration = 0.0f;
-
-		TSharedPtr<SNotificationItem> NotificationPtr = FSlateNotificationManager::Get().AddNotification(Info);
-		NotificationPtr->SetCompletionState(SNotificationItem::CS_Success);
-
-	}
-}
-
-void UControlRigGraph::CacheNameLists(const FRigHierarchyContainer* Container)
-{
-	check(Container);
-	CacheNameList<FRigBoneHierarchy>(Container->BoneHierarchy, BoneNameList);
-	CacheNameList<FRigControlHierarchy>(Container->ControlHierarchy, ControlNameList);
-	CacheNameList<FRigSpaceHierarchy>(Container->SpaceHierarchy, SpaceNameList);
-	CacheNameList<FRigCurveContainer>(Container->CurveContainer, CurveNameList);
+	check(HierarchyContainer);
+	check(DrawContainer);
+
+	CacheNameList<FRigBoneHierarchy>(HierarchyContainer->BoneHierarchy, BoneNameList);
+	CacheNameList<FRigControlHierarchy>(HierarchyContainer->ControlHierarchy, ControlNameList);
+	CacheNameList<FRigSpaceHierarchy>(HierarchyContainer->SpaceHierarchy, SpaceNameList);
+	CacheNameList<FRigCurveContainer>(HierarchyContainer->CurveContainer, CurveNameList);
+	CacheNameList<FControlRigDrawContainer>(*DrawContainer, DrawingNameList);
 }
 
 const TArray<TSharedPtr<FString>>& UControlRigGraph::GetBoneNameList() const
@@ -301,7 +84,12 @@ const TArray<TSharedPtr<FString>>& UControlRigGraph::GetCurveNameList() const
 	return CurveNameList;
 }
 
-void UControlRigGraph::HandleModelModified(const UControlRigModel* InModel, EControlRigModelNotifType InType, const void* InPayload)
+const TArray<TSharedPtr<FString>>& UControlRigGraph::GetDrawingNameList() const
+{
+	return DrawingNameList;
+}
+
+void UControlRigGraph::HandleModifiedEvent(ERigVMGraphNotifType InNotifType, URigVMGraph* InGraph, UObject* InSubject)
 {
 	DECLARE_SCOPE_HIERARCHICAL_COUNTER_FUNC()
 
@@ -310,13 +98,13 @@ void UControlRigGraph::HandleModelModified(const UControlRigModel* InModel, ECon
 		return;
 	}
 
-	switch (InType)
+	switch (InNotifType)
 	{
-		case EControlRigModelNotifType::ModelCleared:
+		case ERigVMGraphNotifType::GraphChanged:
 		{
-			for (const FControlRigModelNode& Node : InModel->Nodes())
+			for (URigVMNode* Node : InGraph->GetNodes())
 			{
-				UEdGraphNode* EdNode = FindNodeFromPropertyName(Node.Name);
+				UEdGraphNode* EdNode = FindNodeForModelNodeName(Node->GetFName());
 				if (EdNode != nullptr)
 				{
 					RemoveNode(EdNode);
@@ -324,220 +112,279 @@ void UControlRigGraph::HandleModelModified(const UControlRigModel* InModel, ECon
 			}
 			break;
 		}
-		case EControlRigModelNotifType::NodeAdded:
+		case ERigVMGraphNotifType::NodeSelectionChanged:
 		{
-			const FControlRigModelNode* Node = (const FControlRigModelNode*)InPayload;
-			if (Node != nullptr)
+			if (bIsSelecting)
+			{
+				return;
+			}
+			TGuardValue<bool> SelectionGuard(bIsSelecting, true);
+
+			TSet<const UEdGraphNode*> NodeSelection;
+			for (FName NodeName : InGraph->GetSelectNodes())
 			{
-				FEdGraphPinType PinType;
-				switch (Node->NodeType)
+				if (UEdGraphNode* EdNode = FindNodeForModelNodeName(NodeName))
 				{
-					case EControlRigModelNodeType::Parameter:
-					{
-						PinType = Node->Pins[0].Type;
-						// no break - fall through
-					}
-					case EControlRigModelNodeType::Function:
+					NodeSelection.Add(EdNode);
+				}
+			}
+			SelectNodeSet(NodeSelection);
+			break;
+		}
+		case ERigVMGraphNotifType::NodeAdded:
+		{
+			if (URigVMNode* ModelNode = Cast<URigVMNode>(InSubject))
+			{
+				if (!ModelNode->IsVisibleInUI())
+				{
+					if (URigVMInjectionInfo* Injection = ModelNode->GetInjectionInfo())
 					{
-						UEdGraphNode* EdNode = FControlRigBlueprintUtils::InstantiateGraphNodeForProperty(this, Node->Name, Node->Position, PinType);
-						if (EdNode != nullptr)
+						if (URigVMPin* ModelPin = Injection->GetPin())
 						{
-							EdNode->CreateNewGuid();
-							if (UControlRigGraphNode* RigNode = Cast<UControlRigGraphNode>(EdNode))
+							URigVMNode* ParentModelNode = ModelPin->GetNode();
+							if (ParentModelNode)
 							{
-								RigNode->ParameterType = (int32)Node->ParameterType;
-								RigNode->SetColorFromModel(Node->Color);
+								UEdGraphNode* EdNode = FindNodeForModelNodeName(ParentModelNode->GetFName());
+								if (EdNode)
+								{
+									if (UControlRigGraphNode* RigNode = Cast<UControlRigGraphNode>(EdNode))
+									{
+										RigNode->ReconstructNode_Internal(true);
+									}
+								}
 							}
 						}
-						break;
 					}
-					case EControlRigModelNodeType::Comment:
+					break;
+				}
+
+				if (URigVMCommentNode* CommentModelNode = Cast<URigVMCommentNode>(ModelNode))
+				{
+					UEdGraphNode_Comment* NewNode = NewObject<UEdGraphNode_Comment>(this, CommentModelNode->GetFName());
+					AddNode(NewNode, false);
+
+					NewNode->CreateNewGuid();
+					NewNode->PostPlacedNewNode();
+					NewNode->AllocateDefaultPins();
+
+					NewNode->NodePosX = ModelNode->GetPosition().X;
+					NewNode->NodePosY = ModelNode->GetPosition().Y;
+					NewNode->NodeWidth = ModelNode->GetSize().X;
+					NewNode->NodeHeight = ModelNode->GetSize().Y;
+					NewNode->CommentColor = ModelNode->GetNodeColor();
+					NewNode->NodeComment = CommentModelNode->GetCommentText();
+					NewNode->SetFlags(RF_Transactional);
+					NewNode->GetNodesUnderComment();
+				}
+				else if (URigVMRerouteNode* RerouteModelNode = Cast<URigVMRerouteNode>(ModelNode))
+				{
+					// todo: collapsed or not?
+					UControlRigGraphNode* NewNode = NewObject<UControlRigGraphNode>(this, ModelNode->GetFName());
+					AddNode(NewNode, false);
+
+					NewNode->ModelNodePath = ModelNode->GetNodePath();
+					NewNode->CreateNewGuid();
+					NewNode->PostPlacedNewNode();
+					NewNode->AllocateDefaultPins();
+
+					NewNode->NodePosX = ModelNode->GetPosition().X;
+					NewNode->NodePosY = ModelNode->GetPosition().Y;
+
+					NewNode->SetFlags(RF_Transactional);
+					NewNode->AllocateDefaultPins();
+
+					if (UEdGraphPin* ValuePin = NewNode->FindPin(ModelNode->FindPin("Value")->GetPinPath()))
 					{
-						UEdGraphNode_Comment* NewNode = NewObject<UEdGraphNode_Comment>(this, Node->Name);
-						AddNode(NewNode, true);
-
-						NewNode->CreateNewGuid();
-						NewNode->PostPlacedNewNode();
-						NewNode->AllocateDefaultPins();
-
-						NewNode->NodePosX = Node->Position.X;
-						NewNode->NodePosY = Node->Position.Y;
-						NewNode->NodeWidth = Node->Size.X;
-						NewNode->NodeHeight= Node->Size.Y;
-						NewNode->CommentColor = Node->Color;
-						NewNode->NodeComment = Node->Text;
-						NewNode->SetFlags(RF_Transactional);
-						NewNode->GetNodesUnderComment();
-
-						break;
+						NewNode->SetColorFromModel(GetSchema()->GetPinTypeColor(ValuePin->PinType));
 					}
-					default:
+				}
+				else // struct, parameter + variable
+				{
+					UControlRigGraphNode* NewNode = NewObject<UControlRigGraphNode>(this, ModelNode->GetFName());
+					AddNode(NewNode, false);
+
+					NewNode->ModelNodePath = ModelNode->GetNodePath();
+					NewNode->CreateNewGuid();
+					NewNode->PostPlacedNewNode();
+					NewNode->AllocateDefaultPins();
+
+					NewNode->NodePosX = ModelNode->GetPosition().X;
+					NewNode->NodePosY = ModelNode->GetPosition().Y;
+					NewNode->SetColorFromModel(ModelNode->GetNodeColor());
+					NewNode->SetFlags(RF_Transactional);
+					NewNode->AllocateDefaultPins();
+				}
+			}
+			break;
+		}
+		case ERigVMGraphNotifType::NodeRemoved:
+		{
+			if (URigVMNode* ModelNode = Cast<URigVMNode>(InSubject))
+			{
+				if (URigVMInjectionInfo* Injection = ModelNode->GetInjectionInfo())
+				{
+					if (URigVMPin* ModelPin = Injection->GetPin())
 					{
-						ensure(false);
-						break;
+						URigVMNode* ParentModelNode = ModelPin->GetNode();
+						if (ParentModelNode)
+						{
+							UEdGraphNode* EdNode = FindNodeForModelNodeName(ParentModelNode->GetFName());
+							if (EdNode)
+							{
+								if (UControlRigGraphNode* RigNode = Cast<UControlRigGraphNode>(EdNode))
+								{
+									RigNode->ReconstructNode_Internal(true);
+								}
+							}
+						}
 					}
+					break;
+				}
+
+				UEdGraphNode* EdNode = FindNodeForModelNodeName(ModelNode->GetFName());
+				if (EdNode)
+				{
+					RemoveNode(EdNode);
 				}
 			}
 			break;
 		}
-		case EControlRigModelNotifType::NodeRemoved:
+		case ERigVMGraphNotifType::NodePositionChanged:
 		{
-			const FControlRigModelNode* Node = (const FControlRigModelNode*)InPayload;
-			if (Node != nullptr)
+			if (URigVMNode* ModelNode = Cast<URigVMNode>(InSubject))
 			{
-				UEdGraphNode* EdNode = FindNodeFromPropertyName(Node->Name);
-				if (EdNode != nullptr)
+				UEdGraphNode* EdNode = FindNodeForModelNodeName(ModelNode->GetFName());
+				if (EdNode)
 				{
-					RemoveNode(EdNode);
+					EdNode->NodePosX = (int32)ModelNode->GetPosition().X;
+					EdNode->NodePosY = (int32)ModelNode->GetPosition().Y;
 				}
 			}
 			break;
 		}
-		case EControlRigModelNotifType::NodeChanged:
+		case ERigVMGraphNotifType::NodeSizeChanged:
 		{
-			if (const FControlRigModelNode* Node = (const FControlRigModelNode*)InPayload)
+			if (URigVMNode* ModelNode = Cast<URigVMNode>(InSubject))
 			{
-				if (UEdGraphNode* EdNode = FindNodeFromPropertyName(Node->Name))
+				UEdGraphNode_Comment* EdNode = Cast<UEdGraphNode_Comment>(FindNodeForModelNodeName(ModelNode->GetFName()));
+				if (EdNode)
 				{
-					EdNode->NodePosX = (int32)Node->Position.X;
-					EdNode->NodePosY = (int32)Node->Position.Y;
-
-					if (UControlRigGraphNode* RigNode = Cast<UControlRigGraphNode>(EdNode))
-					{
-						int32 PreviousParameterType = RigNode->ParameterType;
-						RigNode->ParameterType = (int32)Node->ParameterType;
-						RigNode->SetColorFromModel(Node->Color);
-
-						if (Node->IsParameter() && PreviousParameterType != RigNode->ParameterType)
-						{
-							RigNode->ReconstructNode();
-						}
-					}
-
-					if (UEdGraphNode_Comment* CommentNode = Cast<UEdGraphNode_Comment>(EdNode))
-					{
-						CommentNode->NodeWidth = (int32)Node->Size.X;
-						CommentNode->NodeHeight = (int32)Node->Size.Y;
-						CommentNode->NodeComment = Node->Text;
-						CommentNode->CommentColor = Node->Color;
-					}
+					EdNode->NodeWidth = (int32)ModelNode->GetSize().X;
+					EdNode->NodeHeight = (int32)ModelNode->GetSize().Y;
 				}
 			}
 			break;
 		}
-		case EControlRigModelNotifType::NodeRenamed:
+		case ERigVMGraphNotifType::NodeColorChanged:
 		{
-			const FControlRigModelNodeRenameInfo* Info = (const FControlRigModelNodeRenameInfo*)InPayload;
-			if (Info != nullptr)
+			if (URigVMNode* ModelNode = Cast<URigVMNode>(InSubject))
 			{
-				UControlRigGraphNode* RigNode = Cast<UControlRigGraphNode>(FindNodeFromPropertyName(Info->OldName));
-				if (RigNode != nullptr)
+				UEdGraphNode_Comment* EdNode = Cast<UEdGraphNode_Comment>(FindNodeForModelNodeName(ModelNode->GetFName()));
+				if (EdNode)
 				{
-					RigNode->SetPropertyName(Info->NewName, true);
-					RigNode->InvalidateNodeTitle();
+					EdNode->CommentColor = ModelNode->GetNodeColor();
 				}
 			}
 			break;
 		}
-		case EControlRigModelNotifType::PinAdded:
-		case EControlRigModelNotifType::PinRemoved:
+		case ERigVMGraphNotifType::CommentTextChanged:
 		{
-			const FControlRigModelPin* Pin = (const FControlRigModelPin*)InPayload;
-			if (Pin!= nullptr)
+			if (URigVMCommentNode* ModelNode = Cast<URigVMCommentNode>(InSubject))
 			{
-				const FControlRigModelNode& Node = InModel->Nodes()[Pin->Node];
-				UControlRigGraphNode* RigNode = Cast<UControlRigGraphNode>(FindNodeFromPropertyName(Node.Name));
-				if (RigNode != nullptr)
+				UEdGraphNode_Comment* EdNode = Cast<UEdGraphNode_Comment>(FindNodeForModelNodeName(ModelNode->GetFName()));
+				if (EdNode)
 				{
-					RigNode->ReconstructNode();
+					EdNode->OnUpdateCommentText(ModelNode->GetCommentText());
 				}
 			}
 			break;
 		}
-		case EControlRigModelNotifType::LinkAdded:
-		case EControlRigModelNotifType::LinkRemoved:
+		case ERigVMGraphNotifType::LinkAdded:
+		case ERigVMGraphNotifType::LinkRemoved:
 		{
-			bool AddLink = InType == EControlRigModelNotifType::LinkAdded;
+			bool AddLink = InNotifType == ERigVMGraphNotifType::LinkAdded;
 
-			const FControlRigModelLink* Link = (const FControlRigModelLink*)InPayload;
-			if (Link != nullptr)
+			if (URigVMLink* Link = Cast<URigVMLink>(InSubject))
 			{
-				const FControlRigModelNode& SourceNode = InModel->Nodes()[Link->Source.Node];
-				const FControlRigModelPin& SourcePin = SourceNode.Pins[Link->Source.Pin];
-				const FControlRigModelNode& TargetNode = InModel->Nodes()[Link->Target.Node];
-				const FControlRigModelPin& TargetPin = TargetNode.Pins[Link->Target.Pin];
+				URigVMPin* SourcePin = Link->GetSourcePin();
+				URigVMPin* TargetPin = Link->GetTargetPin();
 
-				UControlRigGraphNode* SourceRigNode = Cast<UControlRigGraphNode>(FindNodeFromPropertyName(SourceNode.Name));
-				UControlRigGraphNode* TargetRigNode = Cast<UControlRigGraphNode>(FindNodeFromPropertyName(TargetNode.Name));
+				SourcePin = SourcePin->GetOriginalPinFromInjectedNode();
+				TargetPin = TargetPin->GetOriginalPinFromInjectedNode();
 
-				if (SourceRigNode != nullptr && TargetRigNode != nullptr)
+				if (SourcePin && TargetPin && SourcePin != TargetPin)
 				{
-					FString SourcePinPath = InModel->GetPinPath(Link->Source, true);
-					FString TargetPinPath = InModel->GetPinPath(Link->Target, true);
+					UControlRigGraphNode* SourceRigNode = Cast<UControlRigGraphNode>(FindNodeForModelNodeName(SourcePin->GetNode()->GetFName()));
+					UControlRigGraphNode* TargetRigNode = Cast<UControlRigGraphNode>(FindNodeForModelNodeName(TargetPin->GetNode()->GetFName()));
 
-					UEdGraphPin* SourceRigPin = SourceRigNode->FindPin(*SourcePinPath, EGPD_Output);
-					UEdGraphPin* TargetRigPin = TargetRigNode->FindPin(*TargetPinPath, EGPD_Input);
-
-					if (SourceRigPin != nullptr && TargetRigPin != nullptr)
+					if (SourceRigNode != nullptr && TargetRigNode != nullptr)
 					{
-						if (AddLink)
-						{
-							SourceRigPin->MakeLinkTo(TargetRigPin);
-						}
-						else
+						FString SourcePinPath = SourcePin->GetPinPath();
+						FString TargetPinPath = TargetPin->GetPinPath();
+						UEdGraphPin* SourceRigPin = SourceRigNode->FindPin(*SourcePinPath, EGPD_Output);
+						UEdGraphPin* TargetRigPin = TargetRigNode->FindPin(*TargetPinPath, EGPD_Input);
+
+						if (SourceRigPin != nullptr && TargetRigPin != nullptr)
 						{
-							SourceRigPin->BreakLinkTo(TargetRigPin);
+							if (AddLink)
+							{
+								SourceRigPin->MakeLinkTo(TargetRigPin);
+							}
+							else
+							{
+								SourceRigPin->BreakLinkTo(TargetRigPin);
+							}
+
 						}
 					}
 				}
-
 			}
 			break;
 		}
-		case EControlRigModelNotifType::PinChanged:
+		case ERigVMGraphNotifType::PinDefaultValueChanged:
 		{
-			const FControlRigModelPin* Pin = (const FControlRigModelPin*)InPayload;
-			if (Pin != nullptr)
+			if (URigVMPin* ModelPin = Cast<URigVMPin>(InSubject))
 			{
-				const FControlRigModelNode& Node = InModel->Nodes()[Pin->Node];
-				UControlRigGraphNode* EdNode = Cast<UControlRigGraphNode>(FindNodeFromPropertyName(Node.Name));
-				if (EdNode != nullptr)
+				if (UControlRigGraphNode* RigNode = Cast<UControlRigGraphNode>(FindNodeForModelNodeName(ModelPin->GetNode()->GetFName())))
 				{
-					FString PinPath = InModel->GetPinPath(Pin->GetPair());
-					UEdGraphPin* EdPin = EdNode->FindPin(*PinPath, Pin->Direction);
-					if (EdPin)
+					UEdGraphPin* RigNodePin = RigNode->FindPin(ModelPin->GetPinPath());
+					check(RigNodePin);
+					RigNode->SetupPinDefaultsFromModel(RigNodePin);
+
+					if (Cast<URigVMVariableNode>(ModelPin->GetNode()))
 					{
-						bool bShouldExpand = EdNode->IsPinExpanded(PinPath) != Pin->bExpanded;
-						if (bShouldExpand)
+						if (ModelPin->GetName() == TEXT("Variable"))
 						{
-							// check if this is an input / output pin,
-							// and in that case: don't expand (since it is handled by the input variant)
-							if (Pin->Direction == EGPD_Output)
-							{
-								if (InModel->FindPinFromPath(PinPath, true /*input*/))
-								{
-									bShouldExpand = false;
-								}
-							}
+							RigNode->InvalidateNodeTitle();
+							RigNode->ReconstructNode_Internal(true);
 						}
-						if (bShouldExpand)
+					}
+					if (Cast<URigVMParameterNode>(ModelPin->GetNode()))
+					{
+						if (ModelPin->GetName() == TEXT("Parameter"))
 						{
-							EdNode->SetPinExpansion(PinPath, Pin->bExpanded);
+							RigNode->InvalidateNodeTitle();
+							RigNode->ReconstructNode_Internal(true);
 						}
-						if (Pin->Direction == EGPD_Input)
+					}
+				}
+				else if (URigVMInjectionInfo* Injection = ModelPin->GetNode()->GetInjectionInfo())
+				{
+					if (Injection->InputPin != ModelPin->GetRootPin())
+					{
+						if (URigVMPin* InjectionPin = Injection->GetPin())
 						{
-							if (!Pin->DefaultValue.IsEmpty())
+							URigVMNode* ParentModelNode = InjectionPin->GetNode();
+							if (ParentModelNode)
 							{
-								if (Pin->Type.PinCategory == UEdGraphSchema_K2::PC_Object)
+								UEdGraphNode* HostEdNode = FindNodeForModelNodeName(ParentModelNode->GetFName());
+								if (HostEdNode)
 								{
-									UClass* Class = Cast<UClass>(Pin->Type.PinSubCategoryObject);
-									if (Class)
+									if (UControlRigGraphNode* HostRigNode = Cast<UControlRigGraphNode>(HostEdNode))
 									{
-										EdPin->DefaultObject = StaticFindObject(Class, ANY_PACKAGE, *Pin->DefaultValue);
+										HostRigNode->ReconstructNode_Internal(true);
 									}
 								}
-
-								EdPin->DefaultValue = Pin->DefaultValue;
 							}
 						}
 					}
@@ -545,15 +392,30 @@ void UControlRigGraph::HandleModelModified(const UControlRigModel* InModel, ECon
 			}
 			break;
 		}
+		case ERigVMGraphNotifType::PinArraySizeChanged:
+		case ERigVMGraphNotifType::PinDirectionChanged:
+		case ERigVMGraphNotifType::PinTypeChanged:
+		{
+			if (URigVMPin* ModelPin = Cast<URigVMPin>(InSubject))
+			{
+				if (UControlRigGraphNode* RigNode = Cast<UControlRigGraphNode>(FindNodeForModelNodeName(ModelPin->GetNode()->GetFName())))
+				{
+					RigNode->ReconstructNode_Internal(true);
+				}
+			}
+			break;
+		}
+		case ERigVMGraphNotifType::PinExpansionChanged:
+		case ERigVMGraphNotifType::NodeSelected:
+		case ERigVMGraphNotifType::NodeDeselected:
 		default:
 		{
-			// todo... other cases
 			break;
 		}
 	}
 }
 
-UEdGraphNode* UControlRigGraph::FindNodeFromPropertyName(const FName& InPropertyName)
+UEdGraphNode* UControlRigGraph::FindNodeForModelNodeName(const FName& InModelNodeName)
 {
 	DECLARE_SCOPE_HIERARCHICAL_COUNTER_FUNC()
 
@@ -561,14 +423,14 @@ UEdGraphNode* UControlRigGraph::FindNodeFromPropertyName(const FName& InProperty
 	{
 		if (UControlRigGraphNode* RigNode = Cast<UControlRigGraphNode>(EdNode))
 		{
-			if (RigNode->PropertyName == InPropertyName)
+			if (RigNode->ModelNodePath == InModelNodeName.ToString())
 			{
 				return EdNode;
 			}
 		}
 		else
 		{
-			if (EdNode->GetFName() == InPropertyName)
+			if (EdNode->GetFName() == InModelNodeName)
 			{
 				return EdNode;
 			}
@@ -577,6 +439,23 @@ UEdGraphNode* UControlRigGraph::FindNodeFromPropertyName(const FName& InProperty
 	return nullptr;
 }
 
+
+URigVMController* UControlRigGraph::GetTemplateController()
+{
+	if (TemplateModel == nullptr)
+	{
+		TemplateModel = NewObject<URigVMGraph>(this, TEXT("TemplateModel"));
+	}
+	if (TemplateController == nullptr)
+	{
+		TemplateController = NewObject<URigVMController>(this, TEXT("TemplateController"));
+		TemplateController->SetGraph(TemplateModel);
+		TemplateController->EnableReporting(false);
+		TemplateController->OnModified().AddUObject(this, &UControlRigGraph::HandleModifiedEvent);
+	}
+	return TemplateController;
+}
+
 #endif
 
 #undef LOCTEXT_NAMESPACE
diff --git a/Engine/Plugins/Experimental/ControlRig/Source/ControlRigDeveloper/Private/Graph/ControlRigGraphNode.cpp b/Engine/Plugins/Experimental/ControlRig/Source/ControlRigDeveloper/Private/Graph/ControlRigGraphNode.cpp
--- a/Engine/Plugins/Experimental/ControlRig/Source/ControlRigDeveloper/Private/Graph/ControlRigGraphNode.cpp
+++ b/Engine/Plugins/Experimental/ControlRig/Source/ControlRigDeveloper/Private/Graph/ControlRigGraphNode.cpp
@@ -20,6 +20,9 @@
 #include "UObject/PropertyPortFlags.h"
 #include "ControlRigBlueprintUtils.h"
 #include "Curves/CurveFloat.h"
+#include "RigVMCore/RigVMExecuteContext.h"
+#include "ControlRigDeveloper.h"
+#include "ControlRigObjectVersion.h"
 
 #if WITH_EDITOR
 #include "IControlRigEditorModule.h"
@@ -29,60 +32,51 @@
 
 UControlRigGraphNode::UControlRigGraphNode()
 : Dimensions(0.0f, 0.0f)
-, NodeTitleFull(FText::GetEmpty())
 , NodeTitle(FText::GetEmpty())
 , CachedTitleColor(FLinearColor(0.f, 0.f, 0.f, 0.f))
 , CachedNodeColor(FLinearColor(0.f, 0.f, 0.f, 0.f))
 {
 	bHasCompilerMessage = false;
 	ErrorType = (int32)EMessageSeverity::Info + 1;
-	ParameterType = (int32)EControlRigModelParameterType::None;
 }
 
 FText UControlRigGraphNode::GetNodeTitle(ENodeTitleType::Type TitleType) const
 {
-	if(NodeTitle.IsEmpty() || NodeTitleFull.IsEmpty())
+	if(NodeTitle.IsEmpty())
 	{
-		UScriptStruct* ScriptStruct = GetUnitScriptStruct();
-		if(ScriptStruct && ScriptStruct->HasMetaData(UControlRig::DisplayNameMetaName))
+		if(URigVMNode* ModelNode = GetModelNode())
 		{
-			if(ScriptStruct->HasMetaData(UControlRig::ShowVariableNameInTitleMetaName))
-			{
-				NodeTitleFull = FText::Format(LOCTEXT("NodeFullTitleFormat", "{0}\n{1}"), FText::FromName(PropertyName), FText::FromString(ScriptStruct->GetMetaData(UControlRig::DisplayNameMetaName)));
-				NodeTitle = FText::FromName(PropertyName);
-			}
-			else
-			{
-				NodeTitle = NodeTitleFull = FText::FromString(ScriptStruct->GetMetaData(UControlRig::DisplayNameMetaName));
-			}
-		}
-		else
-		{
-			NodeTitle = NodeTitleFull = FText::FromName(PropertyName);
+			NodeTitle = FText::FromString(ModelNode->GetNodeTitle());
 		}
 	}
 
-	if(TitleType == ENodeTitleType::FullTitle)
-	{
-		return NodeTitleFull;
-	}
-	else
-	{
-		return NodeTitle;
-	}
+	return NodeTitle;
 }
 
 void UControlRigGraphNode::ReconstructNode()
+{
+	ReconstructNode_Internal();
+}
+
+void UControlRigGraphNode::ReconstructNode_Internal(bool bForce)
 {
 	DECLARE_SCOPE_HIERARCHICAL_COUNTER_FUNC()
 
 	UControlRigGraph* RigGraph = Cast<UControlRigGraph>(GetGraph());
-	if (RigGraph)
+	if (RigGraph && !bForce)
 	{
 		if (RigGraph->bIsTemporaryGraphForCopyPaste)
 		{
 			return;
 		}
+
+		// if this node has been saved prior to our custom version,
+		// don't reset the node
+		int32 LinkerVersion = RigGraph->GetLinkerCustomVersion(FControlRigObjectVersion::GUID);
+		if (LinkerVersion < FControlRigObjectVersion::SwitchedToRigVM)
+		{
+			return;
+		}
 	}
 
 	// Clear previously set messages
@@ -95,6 +89,8 @@ void UControlRigGraphNode::ReconstructNode()
 	Pins.Reset();
 
 	// Recreate the new pins
+	CachedPins.Reset();
+	CachedModelPins.Reset();
 	ReallocatePinsDuringReconstruction(OldPins);
 	RewireOldPinsToNewPins(OldPins, Pins);
 
@@ -102,87 +98,25 @@ void UControlRigGraphNode::ReconstructNode()
 	PostReconstructNode();
 
 	GetGraph()->NotifyGraphChanged();
-
-	UScriptStruct* ScriptStruct = GetUnitScriptStruct();
-	if (ScriptStruct)
-	{
-		StructPath = ScriptStruct->GetPathName();
-	}
 }
 
-void UControlRigGraphNode::CacheHierarchyRefConnectionsOnPostLoad()
+bool UControlRigGraphNode::IsDeprecated() const
 {
-	DECLARE_SCOPE_HIERARCHICAL_COUNTER_FUNC()
-
-	if (HierarchyRefOutputConnections.Num() > 0)
+	if(URigVMNode* ModelNode = GetModelNode())
 	{
-		return;
-	}
-	for (UEdGraphPin* Pin : Pins)
-	{
-		if (Pin->PinType.PinCategory != UEdGraphSchema_K2::PC_Struct)
-		{
-			continue;
-		}
-		if (Pin->PinType.PinSubCategoryObject != FRigHierarchyRef::StaticStruct())
-		{
-			continue;
-		}
-		if (Pin->Direction == EEdGraphPinDirection::EGPD_Output)
-		{
-			for (UEdGraphPin* LinkedPin : Pin->LinkedTo)
-			{
-				HierarchyRefOutputConnections.Add(LinkedPin->GetOwningNode());
-			}
-		}
-		else if (Pin->Direction == EEdGraphPinDirection::EGPD_Input)
+		if(URigVMStructNode* StructModelNode = Cast<URigVMStructNode>(ModelNode))
 		{
-			for (UEdGraphPin* LinkedPin : Pin->LinkedTo)
+			if (UScriptStruct* ScriptStruct = StructModelNode->GetScriptStruct())
 			{
-				UControlRigGraphNode* LinkedNode = Cast<UControlRigGraphNode>(LinkedPin->GetOwningNode());
-				if (LinkedNode)
+				FString DeprecatedMetadata;
+				ScriptStruct->GetStringMetaDataHierarchical(UControlRig::DeprecatedMetaName, &DeprecatedMetadata);
+				if (!DeprecatedMetadata.IsEmpty())
 				{
-					LinkedNode->HierarchyRefOutputConnections.Add(this);
+					return true;
 				}
 			}
 		}
 	}
-}
-
-void UControlRigGraphNode::PrepareForCopying()
-{
-	DECLARE_SCOPE_HIERARCHICAL_COUNTER_FUNC()
-
-	// cache the data we need for paste to work
-	// we fill up struct for rig unit
-	UScriptStruct* ScriptStruct = GetUnitScriptStruct();
-	if (ScriptStruct)
-	{
-		StructPath = ScriptStruct->GetPathName();
-	}
-	// or property
-	FProperty* Property = GetProperty();
-	if (Property)
-	{
-		FString PropertyPath = Property->GetPathName();
-
-		const UEdGraphSchema_K2* Schema = GetDefault<UEdGraphSchema_K2>();	
-		Schema->ConvertPropertyToPinType(Property, PinType);
-	}
-}
-
-bool UControlRigGraphNode::IsDeprecated() const
-{
-	UScriptStruct* ScriptStruct = GetUnitScriptStruct();
-	if (ScriptStruct)
-	{
-		FString DeprecatedMetadata;
-		ScriptStruct->GetStringMetaDataHierarchical(UControlRig::DeprecatedMetaName, &DeprecatedMetadata);
-		if (!DeprecatedMetadata.IsEmpty())
-		{
-			return true;
-		}
-	}
 	return Super::IsDeprecated();
 }
 
@@ -190,16 +124,21 @@ FEdGraphNodeDeprecationResponse UControlRigGraphNode::GetDeprecationResponse(EEd
 {
 	FEdGraphNodeDeprecationResponse Response = Super::GetDeprecationResponse(DeprecationType);
 
-	UScriptStruct* ScriptStruct = GetUnitScriptStruct();
-	if (ScriptStruct)
+	if(URigVMNode* ModelNode = GetModelNode())
 	{
-		FString DeprecatedMetadata;
-		ScriptStruct->GetStringMetaDataHierarchical(UControlRig::DeprecatedMetaName, &DeprecatedMetadata);
-		if (!DeprecatedMetadata.IsEmpty())
+		if(URigVMStructNode* StructModelNode = Cast<URigVMStructNode>(ModelNode))
 		{
-			FFormatNamedArguments Args;
-			Args.Add(TEXT("DeprecatedMetadata"), FText::FromString(DeprecatedMetadata));
-			Response.MessageText = FText::Format(LOCTEXT("ControlRigGraphNodeDeprecationMessage", "Warning: This node is deprecated from: {DeprecatedMetadata}"), Args);
+			if (UScriptStruct* ScriptStruct = StructModelNode->GetScriptStruct())
+			{
+				FString DeprecatedMetadata;
+				ScriptStruct->GetStringMetaDataHierarchical(UControlRig::DeprecatedMetaName, &DeprecatedMetadata);
+				if (!DeprecatedMetadata.IsEmpty())
+				{
+					FFormatNamedArguments Args;
+					Args.Add(TEXT("DeprecatedMetadata"), FText::FromString(DeprecatedMetadata));
+					Response.MessageText = FText::Format(LOCTEXT("ControlRigGraphNodeDeprecationMessage", "Warning: This node is deprecated from: {DeprecatedMetadata}"), Args);
+				}
+			}
 		}
 	}
 
@@ -257,21 +196,14 @@ void UControlRigGraphNode::PostReconstructNode()
 
 	for (UEdGraphPin* Pin : Pins)
 	{
-		SetupPinDefaultsFromCDO(Pin);
+		SetupPinDefaultsFromModel(Pin);
 	}
 
 	bCanRenameNode = false;
 
-	UControlRigBlueprint* Blueprint = Cast<UControlRigBlueprint>(GetOuter()->GetOuter());
-	if (Blueprint)
+	if(URigVMNode* ModelNode = GetModelNode())
 	{
-		if (Blueprint->Model)
-		{
-			if (const FControlRigModelNode* ModelNode = Blueprint->Model->FindNode(PropertyName))
-			{
-				SetColorFromModel(ModelNode->Color);
-			}
-		}
+		SetColorFromModel(ModelNode->GetNodeColor());
 	}
 }
 
@@ -282,294 +214,253 @@ void UControlRigGraphNode::SetColorFromModel(const FLinearColor& InColor)
 	CachedTitleColor = InColor;
 }
 
-#if WITH_EDITORONLY_DATA
-void UControlRigGraphNode::PostLoad()
-{
-	Super::PostLoad();
-	HierarchyRefOutputConnections.Reset();
-}
-#endif
-
-void UControlRigGraphNode::CreateVariablePins(bool bAlwaysCreatePins)
+void UControlRigGraphNode::HandleClearArray(FString InPinPath)
 {
 	DECLARE_SCOPE_HIERARCHICAL_COUNTER_FUNC()
 
-	CacheVariableInfo();
-	CreateExecutionPins(bAlwaysCreatePins);
-	CreateInputPins(bAlwaysCreatePins);
-	CreateInputOutputPins(bAlwaysCreatePins);
-	CreateOutputPins(bAlwaysCreatePins);
+	UControlRigBlueprint* Blueprint = Cast<UControlRigBlueprint>(GetOuter()->GetOuter());
+	if (Blueprint)
+	{
+		Blueprint->Controller->ClearArrayPin(InPinPath);
+	}
 }
 
-void UControlRigGraphNode::HandleClearArray(FString InPropertyPath)
+void UControlRigGraphNode::HandleAddArrayElement(FString InPinPath)
 {
 	DECLARE_SCOPE_HIERARCHICAL_COUNTER_FUNC()
 
 	UControlRigBlueprint* Blueprint = Cast<UControlRigBlueprint>(GetOuter()->GetOuter());
 	if (Blueprint)
 	{
-		if (Blueprint->ModelController && Blueprint->Model)
-		{
-			FString Left, Right;
-			Blueprint->Model->SplitPinPath(InPropertyPath, Left, Right);
-			Blueprint->ModelController->ClearArrayPin(*Left, *Right);
-		}
+		Blueprint->Controller->OpenUndoBracket(TEXT("Add Array Pin"));
+		FString PinPath = Blueprint->Controller->AddArrayPin(InPinPath);
+		Blueprint->Controller->SetPinExpansion(InPinPath, true);
+		Blueprint->Controller->SetPinExpansion(PinPath, true);
+		Blueprint->Controller->CloseUndoBracket();
 	}
 }
 
-void UControlRigGraphNode::HandleAddArrayElement(FString InPropertyPath)
+void UControlRigGraphNode::HandleRemoveArrayElement(FString InPinPath)
 {
 	DECLARE_SCOPE_HIERARCHICAL_COUNTER_FUNC()
 
 	UControlRigBlueprint* Blueprint = Cast<UControlRigBlueprint>(GetOuter()->GetOuter());
 	if (Blueprint)
 	{
-		if (Blueprint->ModelController && Blueprint->Model)
-		{
-			FString Left, Right;
-			Blueprint->Model->SplitPinPath(InPropertyPath, Left, Right);
-			Blueprint->ModelController->AddArrayPin(*Left, *Right, FString());
-		}
+		Blueprint->Controller->RemoveArrayPin(InPinPath);
 	}
 }
 
-void UControlRigGraphNode::HandleRemoveArrayElement(FString InPropertyPath)
+void UControlRigGraphNode::HandleInsertArrayElement(FString InPinPath)
 {
 	DECLARE_SCOPE_HIERARCHICAL_COUNTER_FUNC()
 
 	UControlRigBlueprint* Blueprint = Cast<UControlRigBlueprint>(GetOuter()->GetOuter());
 	if (Blueprint)
 	{
-		if (Blueprint->ModelController && Blueprint->Model)
+		if (URigVMPin* ArrayElementPin = GetModelPinFromPinPath(InPinPath))
 		{
-			const FControlRigModelPin* ChildPin = Blueprint->Model->FindPinFromPath(*InPropertyPath);
-			if (ChildPin)
+			if (URigVMPin* ArrayPin = ArrayElementPin->GetParentPin())
 			{
-				const FControlRigModelPin* ParentPin = Blueprint->Model->GetParentPin(ChildPin->GetPair());
-				if (ParentPin)
-				{
-					FString ParentPinPath = Blueprint->Model->GetPinPath(ParentPin->GetPair());
-					FString Left, Right;
-					Blueprint->Model->SplitPinPath(ParentPinPath, Left, Right);
-					// todo: really should be remove at index
-					Blueprint->ModelController->PopArrayPin(*Left, *Right);
-				}
+				Blueprint->Controller->OpenUndoBracket(TEXT("Add Array Pin"));
+				FString PinPath = Blueprint->Controller->InsertArrayPin(InPinPath, ArrayElementPin->GetPinIndex() + 1, FString());
+				Blueprint->Controller->SetPinExpansion(InPinPath, true);
+				Blueprint->Controller->SetPinExpansion(PinPath, true);
+				Blueprint->Controller->CloseUndoBracket();
 			}
 		}
 	}
 }
 
-void UControlRigGraphNode::HandleInsertArrayElement(FString InPropertyPath)
-{
-	// todo: really should be insert
-	HandleAddArrayElement(InPropertyPath);
-}
-
 void UControlRigGraphNode::AllocateDefaultPins()
 {
-	CreateVariablePins(true);
-}
+	ExecutePins.Reset();
+	InputPins.Reset();
+	InputOutputPins.Reset();
+	OutputPins.Reset();
 
-/** Helper function to check whether this is a struct reference pin */
-static bool IsStructReference(const TSharedPtr<FControlRigField>& InputInfo)
-{
-	if(FStructProperty* StructProperty = CastField<FStructProperty>(InputInfo->GetField()))
+	if (URigVMNode* ModelNode = GetModelNode())
 	{
-		return StructProperty->Struct->IsChildOf(FStructReference::StaticStruct());
+		for (URigVMPin* ModelPin : ModelNode->GetPins())
+		{
+			if (ModelPin->GetDirection() == ERigVMPinDirection::IO)
+			{
+				if (ModelPin->IsStruct())
+				{
+					if (ModelPin->GetScriptStruct()->IsChildOf(FRigVMExecuteContext::StaticStruct()))
+					{
+						ExecutePins.Add(ModelPin);
+						continue;
+					}
+				}
+				InputOutputPins.Add(ModelPin);
+			}
+			else if (ModelPin->GetDirection() == ERigVMPinDirection::Input || 
+				ModelPin->GetDirection() == ERigVMPinDirection::Visible)
+			{
+				InputPins.Add(ModelPin);
+			}
+			else if (ModelPin->GetDirection() == ERigVMPinDirection::Output)
+			{
+				OutputPins.Add(ModelPin);
+			}
+		}
 	}
 
-	return false;
+	CreateExecutionPins();
+	CreateInputPins();
+	CreateInputOutputPins();
+	CreateOutputPins();
 }
 
-void UControlRigGraphNode::CreateExecutionPins(bool bAlwaysCreatePins)
+void UControlRigGraphNode::CreateExecutionPins()
 {
 	DECLARE_SCOPE_HIERARCHICAL_COUNTER_FUNC()
 
-	const TArray<TSharedRef<FControlRigField>>& LocalExecutionInfos = GetExecutionVariableInfo();
-
-	for (const TSharedRef<FControlRigField>& ExecutionInfo : LocalExecutionInfos)
+	const TArray<URigVMPin*> ModelPins = ExecutePins;
+	for (URigVMPin* ModelPin : ModelPins)
 	{
-		if (bAlwaysCreatePins || ExecutionInfo->InputPin == nullptr)
+		PinPair& Pair = CachedPins.FindOrAdd(ModelPin);
+		if (Pair.InputPin == nullptr)
 		{
-			ExecutionInfo->InputPin = CreatePin(EGPD_Input, ExecutionInfo->GetPinType(), FName(*ExecutionInfo->GetPinPath()));
-			ExecutionInfo->InputPin->PinFriendlyName = ExecutionInfo->GetDisplayNameText();
-			ExecutionInfo->InputPin->PinType.bIsReference = IsStructReference(ExecutionInfo);
-			ExecutionInfo->InputPin->DefaultValue = ExecutionInfo->GetPin()->DefaultValue;
+			Pair.InputPin = CreatePin(EGPD_Input, GetPinTypeForModelPin(ModelPin), FName(*ModelPin->GetPinPath()));
+			if (Pair.InputPin != nullptr)
+			{
+				Pair.InputPin->PinFriendlyName = FText::FromName(ModelPin->GetDisplayName());
+				Pair.InputPin->bDisplayAsMutableRef = ModelPin->GetInjectedNodes().Num() > 0;
+			}
 		}
-
-		if (bAlwaysCreatePins || ExecutionInfo->OutputPin == nullptr)
+		if (Pair.OutputPin == nullptr)
 		{
-			ExecutionInfo->OutputPin = CreatePin(EGPD_Output, ExecutionInfo->GetPinType(), FName(*ExecutionInfo->GetPinPath()));
+			Pair.OutputPin = CreatePin(EGPD_Output, GetPinTypeForModelPin(ModelPin), FName(*ModelPin->GetPinPath()));
+			if (Pair.OutputPin != nullptr)
+			{
+				Pair.OutputPin->PinFriendlyName = FText::FromName(ModelPin->GetDisplayName());
+				Pair.OutputPin->bDisplayAsMutableRef = ModelPin->GetInjectedNodes().Num() > 0;
+			}
 		}
-
 		// note: no recursion for execution pins
 	}
 }
 
-void UControlRigGraphNode::CreateInputPins_Recursive(const TSharedPtr<FControlRigField>& InputInfo, bool bAlwaysCreatePins)
-{
-	for (const TSharedPtr<FControlRigField>& ChildInfo : InputInfo->Children)
-	{
-		if (bAlwaysCreatePins || ChildInfo->InputPin == nullptr)
-		{
-			ChildInfo->InputPin = CreatePin(EGPD_Input, ChildInfo->GetPinType(), FName(*ChildInfo->GetPinPath()));
-			ChildInfo->InputPin->PinFriendlyName = ChildInfo->GetDisplayNameText();
-			ChildInfo->InputPin->PinType.bIsReference = IsStructReference(ChildInfo);
-			ChildInfo->InputPin->ParentPin = InputInfo->InputPin;
-			ChildInfo->InputPin->DefaultValue = ChildInfo->GetPin()->DefaultValue;
-			ChildInfo->OutputPin = nullptr;
-			InputInfo->InputPin->SubPins.Add(ChildInfo->InputPin);
-		}
-	}
-
-	for (const TSharedPtr<FControlRigField>& ChildInfo : InputInfo->Children)
-	{
-		CreateInputPins_Recursive(ChildInfo, bAlwaysCreatePins);
-	}
-}
-
-void UControlRigGraphNode::CreateInputPins(bool bAlwaysCreatePins)
+void UControlRigGraphNode::CreateInputPins(URigVMPin* InParentPin)
 {
 	DECLARE_SCOPE_HIERARCHICAL_COUNTER_FUNC()
 
-	const TArray<TSharedRef<FControlRigField>>& LocalInputInfos = GetInputVariableInfo();
-
-	for (const TSharedRef<FControlRigField>& InputInfo : LocalInputInfos)
+	const TArray<URigVMPin*> ModelPins = InParentPin == nullptr ? InputPins : InParentPin->GetSubPins();
+	for (URigVMPin* ModelPin : ModelPins)
 	{
-		if (bAlwaysCreatePins || InputInfo->InputPin == nullptr)
-		{
-			InputInfo->InputPin = CreatePin(EGPD_Input, InputInfo->GetPinType(), FName(*InputInfo->GetPinPath()));
-			InputInfo->InputPin->PinFriendlyName = InputInfo->GetDisplayNameText();
-			InputInfo->InputPin->PinType.bIsReference = IsStructReference(InputInfo);
-			InputInfo->InputPin->DefaultValue = InputInfo->GetPin()->DefaultValue;
-			InputInfo->OutputPin = nullptr;
-		}
-		else
+		PinPair& Pair = CachedPins.FindOrAdd(ModelPin);
+		if (Pair.InputPin == nullptr)
 		{
-			InputInfo->InputPin->DefaultValue = InputInfo->GetPin()->DefaultValue;
-		}
-
-		CreateInputPins_Recursive(InputInfo, bAlwaysCreatePins);
-	}
-}
+			Pair.InputPin = CreatePin(EGPD_Input, GetPinTypeForModelPin(ModelPin), FName(*ModelPin->GetPinPath()));
+			if (Pair.InputPin != nullptr)
+			{
+				Pair.InputPin->PinFriendlyName = FText::FromName(ModelPin->GetDisplayName());
+				Pair.InputPin->bNotConnectable = ModelPin->GetDirection() != ERigVMPinDirection::Input;
+				Pair.InputPin->bDisplayAsMutableRef = ModelPin->GetInjectedNodes().Num() > 0;
 
-void UControlRigGraphNode::CreateInputOutputPins_Recursive(const TSharedPtr<FControlRigField>& InputOutputInfo, bool bAlwaysCreatePins)
-{
-	for (const TSharedPtr<FControlRigField>& ChildInfo : InputOutputInfo->Children)
-	{
-		if (bAlwaysCreatePins || ChildInfo->InputPin == nullptr)
-		{
-			ChildInfo->InputPin = CreatePin(EGPD_Input, ChildInfo->GetPinType(), FName(*ChildInfo->GetPinPath()));
-			ChildInfo->InputPin->PinFriendlyName = ChildInfo->GetDisplayNameText();
-			ChildInfo->InputPin->PinType.bIsReference = IsStructReference(ChildInfo);
-			ChildInfo->InputPin->DefaultValue = ChildInfo->GetPin()->DefaultValue;
-			ChildInfo->InputPin->ParentPin = InputOutputInfo->InputPin;
-			InputOutputInfo->InputPin->SubPins.Add(ChildInfo->InputPin);
-		}
-		else
-		{
-			ChildInfo->InputPin->DefaultValue = ChildInfo->GetPin()->DefaultValue;
-		}
+				SetupPinDefaultsFromModel(Pair.InputPin);
 
-		if (bAlwaysCreatePins || ChildInfo->OutputPin == nullptr)
-		{
-			ChildInfo->OutputPin = CreatePin(EGPD_Output, ChildInfo->GetPinType(), FName(*ChildInfo->GetPinPath()));
-			ChildInfo->OutputPin->PinFriendlyName = ChildInfo->GetDisplayNameText();
-			ChildInfo->OutputPin->ParentPin = InputOutputInfo->OutputPin;
-			ChildInfo->OutputPin->PinType.bIsReference = IsStructReference(ChildInfo);
-			InputOutputInfo->OutputPin->SubPins.Add(ChildInfo->OutputPin);
+				if (InParentPin != nullptr)
+				{
+					PinPair& ParentPair = CachedPins.FindChecked(InParentPin);
+					ParentPair.InputPin->SubPins.Add(Pair.InputPin);
+					Pair.InputPin->ParentPin = ParentPair.InputPin;
+				}
+			}
 		}
-	}
-
-	for (const TSharedPtr<FControlRigField>& ChildInfo : InputOutputInfo->Children)
-	{
-		CreateInputOutputPins_Recursive(ChildInfo, bAlwaysCreatePins);
+		CreateInputPins(ModelPin);
 	}
 }
 
-void UControlRigGraphNode::CreateInputOutputPins(bool bAlwaysCreatePins)
+void UControlRigGraphNode::CreateInputOutputPins(URigVMPin* InParentPin)
 {
 	DECLARE_SCOPE_HIERARCHICAL_COUNTER_FUNC()
 
-	const TArray<TSharedRef<FControlRigField>>& LocalInputOutputInfos = GetInputOutputVariableInfo();
-
-	for (const TSharedRef<FControlRigField>& InputOutputInfo : LocalInputOutputInfos)
+	const TArray<URigVMPin*> ModelPins = InParentPin == nullptr ? InputOutputPins : InParentPin->GetSubPins();
+	for (URigVMPin* ModelPin : ModelPins)
 	{
-		if (bAlwaysCreatePins || InputOutputInfo->InputPin == nullptr)
+		PinPair& Pair = CachedPins.FindOrAdd(ModelPin);
+		if (Pair.InputPin == nullptr)
 		{
-			InputOutputInfo->InputPin = CreatePin(EGPD_Input, InputOutputInfo->GetPinType(), FName(*InputOutputInfo->GetPinPath()));
-			InputOutputInfo->InputPin->PinFriendlyName = InputOutputInfo->GetDisplayNameText();
-			InputOutputInfo->InputPin->PinType.bIsReference = IsStructReference(InputOutputInfo);
-			InputOutputInfo->InputPin->DefaultValue = InputOutputInfo->GetPin()->DefaultValue;
-		}
+			Pair.InputPin = CreatePin(EGPD_Input, GetPinTypeForModelPin(ModelPin), FName(*ModelPin->GetPinPath()));
+			if (Pair.InputPin != nullptr)
+			{
+				Pair.InputPin->PinFriendlyName = FText::FromName(ModelPin->GetDisplayName());
+				Pair.InputPin->bNotConnectable = ModelPin->GetDirection() != ERigVMPinDirection::IO;
+				Pair.InputPin->bDisplayAsMutableRef = ModelPin->GetInjectedNodes().Num() > 0;
 
-		if (bAlwaysCreatePins || InputOutputInfo->OutputPin == nullptr)
-		{
-			InputOutputInfo->OutputPin = CreatePin(EGPD_Output, InputOutputInfo->GetPinType(), FName(*InputOutputInfo->GetPinPath()));
+				SetupPinDefaultsFromModel(Pair.InputPin);
+
+				if (InParentPin != nullptr)
+				{
+					PinPair& ParentPair = CachedPins.FindChecked(InParentPin);
+					ParentPair.InputPin->SubPins.Add(Pair.InputPin);
+					Pair.InputPin->ParentPin = ParentPair.InputPin;
+				}
+			}
 		}
+		if (Pair.OutputPin == nullptr)
+		{
+			Pair.OutputPin = CreatePin(EGPD_Output, GetPinTypeForModelPin(ModelPin), FName(*ModelPin->GetPinPath()));
+			if (Pair.OutputPin != nullptr)
+			{
+				Pair.OutputPin->PinFriendlyName = FText::FromName(ModelPin->GetDisplayName());
+				Pair.OutputPin->bNotConnectable = ModelPin->GetDirection() != ERigVMPinDirection::IO;
+				Pair.OutputPin->bDisplayAsMutableRef = ModelPin->GetInjectedNodes().Num() > 0;
 
-		CreateInputOutputPins_Recursive(InputOutputInfo, bAlwaysCreatePins);
-	}
-}
+				if (InParentPin != nullptr)
+				{
+					PinPair& ParentPair = CachedPins.FindChecked(InParentPin);
+					ParentPair.OutputPin->SubPins.Add(Pair.OutputPin);
+					Pair.OutputPin->ParentPin = ParentPair.OutputPin;
+				}
+			}
+		}
 
-void UControlRigGraphNode::CreateOutputPins_Recursive(const TSharedPtr<FControlRigField>& OutputInfo, bool bAlwaysCreatePins)
-{
-	for (const TSharedPtr<FControlRigField>& ChildInfo : OutputInfo->Children)
-	{
-		if (bAlwaysCreatePins || ChildInfo->OutputPin == nullptr)
+		// don't recurse on knot / compact reroute nodes
+		if(URigVMRerouteNode* RerouteNode = Cast<URigVMRerouteNode>(GetModelNode()))
 		{
-			ChildInfo->OutputPin = CreatePin(EGPD_Output, ChildInfo->GetPinType(), FName(*ChildInfo->GetPinPath()));
-			ChildInfo->OutputPin->PinFriendlyName = ChildInfo->GetDisplayNameText();
-			ChildInfo->OutputPin->PinType.bIsReference = IsStructReference(ChildInfo);
-			ChildInfo->OutputPin->ParentPin = OutputInfo->OutputPin;
-			ChildInfo->InputPin = nullptr;
-			OutputInfo->OutputPin->SubPins.Add(ChildInfo->OutputPin);
+			if (!RerouteNode->GetShowsAsFullNode())
+			{
+				continue;
+			}
 		}
-	}
 
-	for (const TSharedPtr<FControlRigField>& ChildInfo : OutputInfo->Children)
-	{
-		CreateOutputPins_Recursive(ChildInfo, bAlwaysCreatePins);
+		CreateInputOutputPins(ModelPin);
 	}
 }
 
-void UControlRigGraphNode::CreateOutputPins(bool bAlwaysCreatePins)
+void UControlRigGraphNode::CreateOutputPins(URigVMPin* InParentPin)
 {
 	DECLARE_SCOPE_HIERARCHICAL_COUNTER_FUNC()
 
-	const TArray<TSharedRef<FControlRigField>>& LocalOutputInfos = GetOutputVariableInfo();
-
-	for (const TSharedRef<FControlRigField>& OutputInfo : LocalOutputInfos)
+	const TArray<URigVMPin*> ModelPins = InParentPin == nullptr ? OutputPins : InParentPin->GetSubPins();
+	for (URigVMPin* ModelPin : ModelPins)
 	{
-		if (bAlwaysCreatePins || OutputInfo->OutputPin == nullptr)
+		PinPair& Pair = CachedPins.FindOrAdd(ModelPin);
+		if (Pair.OutputPin == nullptr)
 		{
-			OutputInfo->OutputPin = CreatePin(EGPD_Output, OutputInfo->GetPinType(), FName(*OutputInfo->GetPinPath()));
-			OutputInfo->OutputPin->PinFriendlyName = OutputInfo->GetDisplayNameText();
-			OutputInfo->OutputPin->PinType.bIsReference = IsStructReference(OutputInfo);
-			OutputInfo->InputPin = nullptr;
-		}
+			Pair.OutputPin = CreatePin(EGPD_Output, GetPinTypeForModelPin(ModelPin), FName(*ModelPin->GetPinPath()));
+			if (Pair.OutputPin != nullptr)
+			{
+				Pair.OutputPin->PinFriendlyName = FText::FromName(ModelPin->GetDisplayName());
+				Pair.OutputPin->bNotConnectable = ModelPin->GetDirection() != ERigVMPinDirection::Output;
+				Pair.OutputPin->bDisplayAsMutableRef = ModelPin->GetInjectedNodes().Num() > 0;
 
-		CreateOutputPins_Recursive(OutputInfo, bAlwaysCreatePins);
+				if (InParentPin != nullptr)
+				{
+					PinPair& ParentPair = CachedPins.FindChecked(InParentPin);
+					ParentPair.OutputPin->SubPins.Add(Pair.OutputPin);
+					Pair.OutputPin->ParentPin = ParentPair.OutputPin;
+				}
+			}
+		}
+		CreateOutputPins(ModelPin);
 	}
 }
 
-void UControlRigGraphNode::CacheVariableInfo()
-{
-	DECLARE_SCOPE_HIERARCHICAL_COUNTER_FUNC()
-
-	ExecutionInfos.Reset();
-	GetExecutionFields(ExecutionInfos);
-
-	InputInfos.Reset();
-	GetInputFields(InputInfos);
-
-	OutputInfos.Reset();
-	GetOutputFields(OutputInfos);
-
-	InputOutputInfos.Reset();
-	GetInputOutputFields(InputOutputInfos);
-}
-
 UClass* UControlRigGraphNode::GetControlRigGeneratedClass() const
 {
 	UControlRigBlueprint* Blueprint = Cast<UControlRigBlueprint>(GetOuter()->GetOuter());
@@ -599,15 +490,31 @@ UClass* UControlRigGraphNode::GetControlRigSkeletonGeneratedClass() const
 	return nullptr;
 }
 
+FLinearColor UControlRigGraphNode::GetNodeOpacityColor() const
+{
+	if (URigVMNode* ModelNode = GetModelNode())
+	{
+		if (Cast<URigVMParameterNode>(ModelNode) || Cast<URigVMVariableNode>(ModelNode))
+		{
+			return FLinearColor::White;
+		}
+		if (ModelNode->GetInstructionIndex() == INDEX_NONE)
+		{
+			return FLinearColor(0.35f, 0.35f, 0.35f, 0.35f);
+		}
+	}
+	return FLinearColor::White;
+}
+
 FLinearColor UControlRigGraphNode::GetNodeTitleColor() const
 {
 	// return a darkened version of the default node's color
-	return CachedTitleColor;
+	return CachedTitleColor * GetNodeOpacityColor();
 }
 
 FLinearColor UControlRigGraphNode::GetNodeBodyTintColor() const
 {
-	return CachedNodeColor;
+	return CachedNodeColor * GetNodeOpacityColor();
 }
 
 FSlateIcon UControlRigGraphNode::GetIconAndTint(FLinearColor& OutColor) const
@@ -617,364 +524,183 @@ FSlateIcon UControlRigGraphNode::GetIconAndTint(FLinearColor& OutColor) const
 	return Icon;
 }
 
-// bool UControlRigGraphNode::ReferencesVariable(const FName& InVarName, const UStruct* InScope) const
-// {
-// 	return InVarName == PropertyName;
-// }
-
-TSharedPtr<FControlRigField> UControlRigGraphNode::CreateControlRigField(const FControlRigModelPin* InPin, const FString& InPinPath, int32 InArrayIndex) const
+void UControlRigGraphNode::GetNodeContextMenuActions(class UToolMenu* Menu, class UGraphNodeContextMenuContext* Context) const
 {
-	TSharedPtr<FControlRigField> NewField = MakeShareable(new FControlRigPin(InPin, InPinPath, InArrayIndex));
-	NewField->DisplayNameText = InPin->DisplayNameText;
-	NewField->TooltipText = InPin->TooltipText;
-	NewField->InputPin = FindPin(InPinPath, EGPD_Input);
-	NewField->OutputPin = FindPin(InPinPath, EGPD_Output);
-	return NewField;
+#if WITH_EDITOR
+	const UControlRigGraphSchema* Schema = Cast<UControlRigGraphSchema>(GetSchema());
+	IControlRigEditorModule::Get().GetContextMenuActions(Schema, Menu, Context);
+#endif
 }
 
-void UControlRigGraphNode::GetExecutionFields(TArray<TSharedRef<FControlRigField>>& OutFields) const
+bool UControlRigGraphNode::IsPinExpanded(const FString& InPinPath)
 {
-	DECLARE_SCOPE_HIERARCHICAL_COUNTER_FUNC()
-
-	GetFields([](const FControlRigModelPin* InPin, const FControlRigModelNode* InNode)
+	if (URigVMPin* ModelPin = GetModelPinFromPinPath(InPinPath))
 	{
-		FString PinPath = InNode->GetPinPath(InPin->Index, false);
-		return InPin->Direction == EGPD_Output &&
-			InNode->FindPin(*PinPath, true) != nullptr &&
-			InPin->Type.PinSubCategoryObject == FControlRigExecuteContext::StaticStruct();
-	}, OutFields);
+		return ModelPin->IsExpanded();
+	}
+	return false;
 }
 
-void UControlRigGraphNode::GetInputFields(TArray<TSharedRef<FControlRigField>>& OutFields) const
+void UControlRigGraphNode::DestroyNode()
 {
 	DECLARE_SCOPE_HIERARCHICAL_COUNTER_FUNC()
 
-	GetFields([](const FControlRigModelPin* InPin, const FControlRigModelNode* InNode)
+	if(UControlRigGraph* Graph = Cast<UControlRigGraph>(GetOuter()))
 	{
-		FString PinPath = InNode->GetPinPath(InPin->Index, false);
-		if (InPin->Direction != EGPD_Input)
-		{
-			return false;
-		}
-
-		if (InNode->IsParameter() && InNode->ParameterType == EControlRigModelParameterType::Output)
+		UControlRigBlueprint* ControlRigBlueprint = Cast<UControlRigBlueprint>(Graph->GetOuter());
+		if(ControlRigBlueprint)
 		{
-			return true;
+			BreakAllNodeLinks();
+			if(PropertyName_DEPRECATED.IsValid())
+			{
+				FControlRigBlueprintUtils::RemoveMemberVariableIfNotUsed(ControlRigBlueprint, PropertyName_DEPRECATED, this);
+			}
 		}
+	}
 
-		return InNode->FindPin(*PinPath, false) == nullptr;
-	}, OutFields);
+	UEdGraphNode::DestroyNode();
 }
 
-void UControlRigGraphNode::GetOutputFields(TArray<TSharedRef<FControlRigField>>& OutFields) const
+void UControlRigGraphNode::PinDefaultValueChanged(UEdGraphPin* Pin)
 {
-	DECLARE_SCOPE_HIERARCHICAL_COUNTER_FUNC()
-
-	GetFields([](const FControlRigModelPin* InPin, const FControlRigModelNode* InNode)
-	{
-		FString PinPath = InNode->GetPinPath(InPin->Index, false);
-		if (InPin->Direction != EGPD_Output)
-		{
-			return false;
-		}
-		
-		if (InNode->IsParameter() && InNode->ParameterType == EControlRigModelParameterType::Input)
-		{
-			return true;
-		}
-		
-		return InNode->FindPin(*PinPath, true) == nullptr;
-	}, OutFields);
+	CopyPinDefaultsToModel(Pin, true);
 }
 
-void UControlRigGraphNode::GetInputOutputFields(TArray<TSharedRef<FControlRigField>>& OutFields) const
+void UControlRigGraphNode::CopyPinDefaultsToModel(UEdGraphPin* Pin, bool bUndo)
 {
 	DECLARE_SCOPE_HIERARCHICAL_COUNTER_FUNC()
 
-	GetFields([](const FControlRigModelPin* InPin, const FControlRigModelNode* InNode)
-	{
-		if (InNode->IsParameter())
-		{
-			return false;
-		}
-		FString PinPath = InNode->GetPinPath(InPin->Index, false);
-		return InPin->Direction == EGPD_Input && 
-			InNode->FindPin(*PinPath, false) != nullptr &&
-			InPin->Type.PinSubCategoryObject != FControlRigExecuteContext::StaticStruct();
-	}, OutFields);
-}
-
-void UControlRigGraphNode::GetFields(TFunction<bool(const FControlRigModelPin*, const FControlRigModelNode*)> InPinCheckFunction, TArray<TSharedRef<FControlRigField>>& OutFields) const
-{
-	OutFields.Reset();
-
-	FControlRigModelNode Node;
-
-	if (UControlRigBlueprint* RigBlueprint = Cast<UControlRigBlueprint>(GetBlueprint()))
+	if (Pin->Direction != EGPD_Input)
 	{
-		if (RigBlueprint->Model)
-		{
-			const FControlRigModelNode* FoundNode = RigBlueprint->Model->FindNode(GetPropertyName());
-			if (FoundNode != nullptr)
-			{
-				Node = *FoundNode;
-			}
-		}
+		return;
 	}
 
-	if (!Node.IsValid())
+	const UEdGraphSchema_K2* K2Schema = GetDefault<UEdGraphSchema_K2>();
+	if (URigVMPin* ModelPin = GetModelPinFromPinPath(Pin->GetName()))
 	{
-		if (IsVariable())
+		if (ModelPin->GetSubPins().Num() > 0)
 		{
-			FName DataType = PinType.PinCategory;
-			if (UStruct* Struct = Cast<UStruct>(PinType.PinSubCategoryObject))
-			{
-				DataType = Struct->GetFName();
-			}
-			UControlRigController::ConstructPreviewParameter(DataType, EControlRigModelParameterType::Input, Node);
+			return;
 		}
-		else
+
+		FString DefaultValue = Pin->DefaultValue;
+		if (DefaultValue == FName(NAME_None).ToString() && Pin->PinType.PinSubCategory == UEdGraphSchema_K2::PC_Name)
 		{
-			if (UStruct* Struct = GetUnitScriptStruct())
-			{
-				FName FunctionName = Struct->GetFName();
-				UControlRigController::ConstructPreviewNode(FunctionName, Node);
-			}
+			DefaultValue = FString();
 		}
-	}
-
-	if (!Node.IsValid())
-	{
-		return;
-	}
-
-	Node.Name = GetPropertyName();
 
-	TMap<int32, TSharedRef<FControlRigField>> AllFields;
-	for(int32 PinIndex=0;PinIndex < Node.Pins.Num(); PinIndex++)
-	{
-		const FControlRigModelPin& Pin = Node.Pins[PinIndex];
-		if (InPinCheckFunction(&Pin, &Node))
+		if (ModelPin->GetDefaultValue() != DefaultValue)
 		{
-			FString PinPath = Node.GetPinPath(Pin.Index, true);
-			TSharedPtr<FControlRigField> NewField = CreateControlRigField(&Pin, PinPath);
-			if (NewField.IsValid())
-			{
-				TSharedRef<FControlRigField> NewFieldRef = NewField.ToSharedRef();
-				AllFields.Add(Pin.Index, NewFieldRef);
-
-				TSharedRef<FControlRigField>* ParentField = AllFields.Find(Pin.ParentIndex);
-				if (ParentField)
-				{
-					(*ParentField)->Children.Add(NewFieldRef);
-				}
-				else
-				{
-					OutFields.Add(NewFieldRef);
-				}
-			}
+			GetBlueprint()->Controller->SetPinDefaultValue(ModelPin->GetPinPath(), DefaultValue, false, true, false);
 		}
 	}
 }
 
-FStructProperty* UControlRigGraphNode::GetUnitProperty() const
+UControlRigBlueprint* UControlRigGraphNode::GetBlueprint() const
 {
-	FProperty* ClassProperty = GetProperty();
-	if(ClassProperty)
+	if(UControlRigGraph* Graph = Cast<UControlRigGraph>(GetOuter()))
 	{
-		// Check if this is a unit struct and if so extract the pins we want to display...
-		if(FStructProperty* StructProperty = CastField<FStructProperty>(ClassProperty))
-		{
-			if(StructProperty->Struct->IsChildOf(FRigUnit::StaticStruct()))
-			{
-				return StructProperty;
-			}
-		}
+		return Cast<UControlRigBlueprint>(Graph->GetOuter());
 	}
-
 	return nullptr;
 }
 
-UScriptStruct* UControlRigGraphNode::GetUnitScriptStruct() const
+URigVMNode* UControlRigGraphNode::GetModelNode() const
 {
-	if(FStructProperty* StructProperty = GetUnitProperty())
+	if (CachedModelNode)
 	{
-		if(StructProperty->Struct->IsChildOf(FRigUnit::StaticStruct()))
-		{
-			return StructProperty->Struct;
-		}
+		return CachedModelNode;
 	}
-	else 
+
+	UControlRigGraphNode* MutableThis = (UControlRigGraphNode*)this;
+	if (UControlRigGraph* Graph = Cast<UControlRigGraph>(GetOuter()))
 	{
-		// Assume that the property name we have is the name of the struct type
-		UScriptStruct* Struct = FindObject<UScriptStruct>(ANY_PACKAGE, *PropertyName.ToString());
-		if(Struct && Struct->IsChildOf(FRigUnit::StaticStruct()))
+#if WITH_EDITOR
+
+		if (Graph->TemplateModel != nullptr)
 		{
-			return Struct;
+			return MutableThis->CachedModelNode = Graph->TemplateModel->FindNode(ModelNodePath);
 		}
 
-		// if this doesn't work we can still fall back on the struct path
-		FString Prefix, StructName;
-		if (StructPath.Split(TEXT("."), &Prefix, &StructName))
+#endif
+
+		if (UControlRigBlueprint* Blueprint = Cast<UControlRigBlueprint>(Graph->GetOuter()))
 		{
-			Struct = FindObject<UScriptStruct>(ANY_PACKAGE, *StructName);
-			if (Struct && Struct->IsChildOf(FRigUnit::StaticStruct()))
-			{
-				return Struct;
-			}
+			return MutableThis->CachedModelNode = Blueprint->Model->FindNode(ModelNodePath);
 		}
 	}
+
 	return nullptr;
 }
 
-FProperty* UControlRigGraphNode::GetProperty() const
+FName UControlRigGraphNode::GetModelNodeName() const
 {
-	if (UClass* MyControlRigClass = GetControlRigSkeletonGeneratedClass())
+	if (URigVMNode* ModelNode = GetModelNode())
 	{
-		return MyControlRigClass->FindPropertyByName(PropertyName);
+		return ModelNode->GetFName();
 	}
-	return nullptr;
-}
-
-void UControlRigGraphNode::PinConnectionListChanged(UEdGraphPin* Pin) 
-{
-
-}
-
-void UControlRigGraphNode::GetNodeContextMenuActions(UToolMenu* Menu, UGraphNodeContextMenuContext* Context) const
-{
-#if WITH_EDITOR
-	IControlRigEditorModule::Get().GetNodeContextMenuActions(this, Menu, Context);
-#endif
+	return NAME_None;
 }
 
-void UControlRigGraphNode::SetPinExpansion(const FString& InPinPropertyPath, bool bExpanded)
+URigVMPin* UControlRigGraphNode::GetModelPinFromPinPath(const FString& InPinPath) const
 {
-	if(bExpanded)
-	{
-		ExpandedPins.AddUnique(InPinPropertyPath);
-	}
-	else
+	if (URigVMPin*const* CachedModelPin = CachedModelPins.Find(InPinPath))
 	{
-		ExpandedPins.Remove(InPinPropertyPath);
+		return (URigVMPin*)*CachedModelPin;
 	}
-}
 
-bool UControlRigGraphNode::IsPinExpanded(const FString& InPinPropertyPath) const
-{
-	return ExpandedPins.Contains(InPinPropertyPath);
-}
-
-void UControlRigGraphNode::DestroyNode()
-{
-	DECLARE_SCOPE_HIERARCHICAL_COUNTER_FUNC()
-
-	if(UControlRigGraph* Graph = Cast<UControlRigGraph>(GetOuter()))
+	if (URigVMNode* ModelNode = GetModelNode())
 	{
-		UControlRigBlueprint* ControlRigBlueprint = Cast<UControlRigBlueprint>(Graph->GetOuter());
-		if(ControlRigBlueprint)
+		FString PinPath = InPinPath.RightChop(ModelNode->GetNodePath().Len() + 1);
+		URigVMPin* ModelPin = ModelNode->FindPin(PinPath);
+		if (ModelPin)
 		{
-			BreakAllNodeLinks();
-			FControlRigBlueprintUtils::RemoveMemberVariableIfNotUsed(ControlRigBlueprint, PropertyName, this);
+			UControlRigGraphNode* MutableThis = (UControlRigGraphNode*)this;
+			MutableThis->CachedModelPins.Add(InPinPath, ModelPin);
 		}
+		return ModelPin;
 	}
-
-	UEdGraphNode::DestroyNode();
-}
-
-void UControlRigGraphNode::PinDefaultValueChanged(UEdGraphPin* Pin)
-{
-	CopyPinDefaultsToModel(Pin, true);
-}
-
-TSharedPtr<INameValidatorInterface> UControlRigGraphNode::MakeNameValidator() const
-{
-	return MakeShared<FKismetNameValidator>(GetBlueprint(), PropertyName);
+	
+	return nullptr;
 }
 
-void UControlRigGraphNode::CopyPinDefaultsToModel(UEdGraphPin* Pin, bool bUndo)
+void UControlRigGraphNode::SetupPinDefaultsFromModel(UEdGraphPin* Pin)
 {
 	DECLARE_SCOPE_HIERARCHICAL_COUNTER_FUNC()
 
-	const UEdGraphSchema_K2* K2Schema = GetDefault<UEdGraphSchema_K2>();
-
-	if(UControlRigGraph* Graph = Cast<UControlRigGraph>(GetOuter()))
-	{
-		UControlRigBlueprint* ControlRigBlueprint = Cast<UControlRigBlueprint>(Graph->GetOuter());
-		if(ControlRigBlueprint)
-		{
-			if (ControlRigBlueprint->Model)
-			{
-				if (Pin->Direction == EGPD_Input)
-				{
-					FString DefaultValue = Pin->DefaultValue;
-					FString Left, Right;
-					ControlRigBlueprint->Model->SplitPinPath(Pin->GetName(), Left, Right);
-					if (DefaultValue.IsEmpty() && Pin->DefaultObject != nullptr)
-					{
-						DefaultValue = Pin->DefaultObject->GetPathName();
-					}
-					ControlRigBlueprint->ModelController->SetPinDefaultValue(*Left, *Right, DefaultValue, false, bUndo);
-				}
-			}
-		}
-	}
-}
-
-UControlRigBlueprint* UControlRigGraphNode::GetBlueprint() const
-{
-	if(UControlRigGraph* Graph = Cast<UControlRigGraph>(GetOuter()))
+	if (Pin->Direction != EGPD_Input)
 	{
-		return Cast<UControlRigBlueprint>(Graph->GetOuter());
+		return;
 	}
-	return nullptr;
-}
-
-void UControlRigGraphNode::SetupPinDefaultsFromCDO(UEdGraphPin* Pin)
-{
-	DECLARE_SCOPE_HIERARCHICAL_COUNTER_FUNC()
 
 	const UEdGraphSchema_K2* K2Schema = GetDefault<UEdGraphSchema_K2>();
-
-	if(UControlRigGraph* Graph = Cast<UControlRigGraph>(GetOuter()))
+	if (URigVMPin* ModelPin = GetModelPinFromPinPath(Pin->GetName()))
 	{
-		UControlRigBlueprint* ControlRigBlueprint = Cast<UControlRigBlueprint>(Graph->GetOuter());
-		if(ControlRigBlueprint)
+		if (ModelPin->GetSubPins().Num() > 0)
 		{
-			// Note we need the actual generated class here
-			if (UClass* MyControlRigClass = GetControlRigGeneratedClass())
-			{
-				if(UObject* DefaultObject = MyControlRigClass->GetDefaultObject(false))
-				{
-					FString DefaultValueString;
-					FCachedPropertyPath PropertyPath(Pin->PinName.ToString());
-					if(PropertyPathHelpers::GetPropertyValueAsString(DefaultObject, PropertyPath, DefaultValueString))
-					{
-						K2Schema->GetPinDefaultValuesFromString(Pin->PinType, Pin->GetOwningNodeUnchecked(), DefaultValueString, Pin->DefaultValue, Pin->DefaultObject, Pin->DefaultTextValue);
-					}
-				}
-			}
+			return;
+		}
+
+		FString DefaultValueString = ModelPin->GetDefaultValue();
+		if (DefaultValueString.IsEmpty() && ModelPin->GetCPPType() == TEXT("FName"))
+		{
+			DefaultValueString = FName(NAME_None).ToString();
 		}
+		K2Schema->GetPinDefaultValuesFromString(Pin->PinType, Pin->GetOwningNodeUnchecked(), DefaultValueString, Pin->DefaultValue, Pin->DefaultObject, Pin->DefaultTextValue);
 	}
 }
 
 FText UControlRigGraphNode::GetTooltipText() const
 {
-	if(GetUnitScriptStruct())
-	{
-		return GetUnitScriptStruct()->GetToolTipText();
-	}
-	else if(GetUnitProperty())
+	if(URigVMNode* ModelNode = GetModelNode())
 	{
-		return GetUnitProperty()->GetToolTipText();
+		return ModelNode->GetToolTipText();
 	}
-
-	return FText::FromName(PropertyName);
+	return FText::FromString(ModelNodePath);
 }
 
 void UControlRigGraphNode::InvalidateNodeTitle() const
 {
-	NodeTitleFull = FText();
 	NodeTitle = FText();
 }
 
@@ -1001,65 +727,90 @@ void UControlRigGraphNode::AutowireNewNode(UEdGraphPin* FromPin)
 		FPinConnectionResponse ConnectResponse = Schema->CanCreateConnection(FromPin, Pin);
 		if(ConnectResponse.Response != ECanCreateConnectionResponse::CONNECT_RESPONSE_DISALLOW)
 		{
-			if(Schema->TryCreateConnection(FromPin, Pin))
+			if (Schema->TryCreateConnection(FromPin, Pin))
 			{
-				// expand any sub-pins so the connection is visible
-				if(UControlRigGraphNode* OuterNode = Cast<UControlRigGraphNode>(Pin->GetOwningNode()))
-				{
-					UEdGraphPin* ParentPin = Pin->ParentPin;
-					while(ParentPin != nullptr)
-					{
-						OuterNode->SetPinExpansion(ParentPin->PinName.ToString(), true);
-						ParentPin = ParentPin->ParentPin;
-					}
-				}
-				return;
+				break;
 			}
 		}
 	}
 }
 
-void ReplacePropertyName(TArray<TSharedRef<FControlRigField>>& InArray, const FString& OldPropName, const FString& NewPropName)
+bool UControlRigGraphNode::IsSelectedInEditor() const
 {
-	for (int32 Index = 0; Index < InArray.Num(); ++Index)
+	URigVMNode* ModelNode = GetModelNode();
+	if (ModelNode)
 	{
-		InArray[Index]->PinPath = InArray[Index]->PinPath.Replace(*OldPropName, *NewPropName);
-		ReplacePropertyName(InArray[Index]->Children, OldPropName, NewPropName);
+		return ModelNode->IsSelected();
 	}
-};
-
-void UControlRigGraphNode::SetPropertyName(const FName& InPropertyName, bool bReplaceInnerProperties/*=false*/)
-{ 
-	DECLARE_SCOPE_HIERARCHICAL_COUNTER_FUNC()
-
-	const FString OldPropertyName = PropertyName.ToString();
-	const FString NewPropertyName = InPropertyName.ToString();
-	PropertyName = InPropertyName;
+	return false;
+}
 
-	if (bReplaceInnerProperties && InPropertyName != NAME_None)
+bool UControlRigGraphNode::ShouldDrawNodeAsControlPointOnly(int32& OutInputPinIndex, int32& OutOutputPinIndex) const
+{
+	if (URigVMRerouteNode* Reroute = Cast<URigVMRerouteNode>(GetModelNode()))
 	{
-		ReplacePropertyName(InputInfos, OldPropertyName, NewPropertyName);
-		ReplacePropertyName(InputOutputInfos, OldPropertyName, NewPropertyName);
-		ReplacePropertyName(OutputInfos, OldPropertyName, NewPropertyName);
-
-		// now change pins
-		for (int32 Index = 0; Index < Pins.Num(); ++Index)
-		{
-			FString PinString = Pins[Index]->PinName.ToString();
-			Pins[Index]->PinName = FName(*PinString.Replace(*OldPropertyName, *NewPropertyName));
-		}
-
-		for (int32 Index = 0; Index < ExpandedPins.Num(); ++Index)
+		if (!Reroute->GetShowsAsFullNode())
 		{
-			FString& PinString = ExpandedPins[Index];
-			PinString = PinString.Replace(*OldPropertyName, *NewPropertyName);
+			if (Pins.Num() == 2)
+			{
+				OutInputPinIndex = 0;
+				OutOutputPinIndex = 1;
+				return true;
+			}
 		}
 	}
+	return false;
 }
 
-bool UControlRigGraphNode::IsVariable() const
+FEdGraphPinType UControlRigGraphNode::GetPinTypeForModelPin(URigVMPin* InModelPin)
 {
-	return GetUnitScriptStruct() == nullptr;
+	FEdGraphPinType PinType;
+	PinType.ResetToDefaults();
+
+	FName ModelPinCPPType = InModelPin->IsArray() ? *InModelPin->GetArrayElementCppType() : *InModelPin->GetCPPType();
+	if (ModelPinCPPType == TEXT("bool"))
+	{
+		PinType.PinCategory = UEdGraphSchema_K2::PC_Boolean;
+	}
+	else if (ModelPinCPPType == TEXT("int32"))
+	{
+		PinType.PinCategory = UEdGraphSchema_K2::PC_Int;
+	}
+	else if (ModelPinCPPType == TEXT("float"))
+	{
+		PinType.PinCategory = UEdGraphSchema_K2::PC_Float;
+	}
+	else if (ModelPinCPPType == TEXT("FName"))
+	{
+		PinType.PinCategory = UEdGraphSchema_K2::PC_Name;
+	}
+	else if (ModelPinCPPType == TEXT("FString"))
+	{
+		PinType.PinCategory = UEdGraphSchema_K2::PC_String;
+	}
+	else if (InModelPin->GetScriptStruct() != nullptr)
+	{
+		PinType.PinCategory = UEdGraphSchema_K2::PC_Struct;
+		PinType.PinSubCategoryObject = InModelPin->GetScriptStruct();
+	}
+	else if (InModelPin->GetEnum() != nullptr)
+	{
+		PinType.PinCategory = UEdGraphSchema_K2::PC_Byte;
+		PinType.PinSubCategoryObject = InModelPin->GetEnum();
+	}
+
+	if (InModelPin->IsArray())
+	{
+		PinType.ContainerType = EPinContainerType::Array;
+	}
+	else
+	{
+		PinType.ContainerType = EPinContainerType::None;
+	}
+
+	PinType.bIsConst = InModelPin->IsConstant();
+
+	return PinType;
 }
 
 #undef LOCTEXT_NAMESPACE
diff --git a/Engine/Plugins/Experimental/ControlRig/Source/ControlRigDeveloper/Private/Graph/ControlRigGraphSchema.cpp b/Engine/Plugins/Experimental/ControlRig/Source/ControlRigDeveloper/Private/Graph/ControlRigGraphSchema.cpp
--- a/Engine/Plugins/Experimental/ControlRig/Source/ControlRigDeveloper/Private/Graph/ControlRigGraphSchema.cpp
+++ b/Engine/Plugins/Experimental/ControlRig/Source/ControlRigDeveloper/Private/Graph/ControlRigGraphSchema.cpp
@@ -11,12 +11,12 @@
 #include "Framework/MultiBox/MultiBoxBuilder.h"
 #include "GraphEditorActions.h"
 #include "ControlRig.h"
-#include "ControlRigDAG.h"
 #include "ControlRigBlueprint.h"
 #include "ControlRigBlueprintGeneratedClass.h"
 #include "Widgets/Notifications/SNotificationList.h"
 #include "Framework/Notifications/NotificationManager.h"
 #include "EulerTransform.h"
+#include "Curves/CurveFloat.h"
 
 #define LOCTEXT_NAMESPACE "ControlRigGraphSchema"
 
@@ -33,11 +33,14 @@ void UControlRigGraphSchema::GetGraphContextActions(FGraphContextMenuBuilder& Co
 
 void UControlRigGraphSchema::GetContextMenuActions(class UToolMenu* Menu, class UGraphNodeContextMenuContext* Context) const
 {
+	/*
+	// this seems to be taken care of by ControlRigGraphNode
 #if WITH_EDITOR
 	return IControlRigEditorModule::Get().GetContextMenuActions(this, Menu, Context);
 #else
 	check(0);
 #endif
+	*/
 }
 
 bool UControlRigGraphSchema::TryCreateConnection(UEdGraphPin* PinA, UEdGraphPin* PinB) const
@@ -49,6 +52,9 @@ bool UControlRigGraphSchema::TryCreateConnection(UEdGraphPin* PinA, UEdGraphPin*
 	}
 #endif
 
+	UControlRigGraphSchema* MutableThis = (UControlRigGraphSchema*)this;
+	MutableThis->LastPinForCompatibleCheck = nullptr;
+
 	UBlueprint* Blueprint = FBlueprintEditorUtils::FindBlueprintForNodeChecked(PinA->GetOwningNode());
 	UControlRigBlueprint* RigBlueprint = Cast<UControlRigBlueprint>(Blueprint);
 	if (RigBlueprint != nullptr)
@@ -59,11 +65,7 @@ bool UControlRigGraphSchema::TryCreateConnection(UEdGraphPin* PinA, UEdGraphPin*
 			PinA = PinB;
 			PinB = Temp;
 		}
-
-		FString LeftA, LeftB, RightA, RightB;
-		RigBlueprint->Model->SplitPinPath(PinA->GetName(), LeftA, RightA);
-		RigBlueprint->Model->SplitPinPath(PinB->GetName(), LeftB, RightB);
-		return RigBlueprint->ModelController->MakeLink(*LeftA, *RightA, *LeftB, *RightB);
+		return RigBlueprint->Controller->AddLink(PinA->GetName(), PinB->GetName());
 	}
 	return false;
 }
@@ -97,33 +99,39 @@ const FPinConnectionResponse UControlRigGraphSchema::CanCreateConnection(const U
 	UControlRigBlueprint* RigBlueprint = Cast<UControlRigBlueprint>(Blueprint);
 	if (RigBlueprint != nullptr)
 	{
-		FString LeftA, LeftB, RightA, RightB;
-		RigBlueprint->Model->SplitPinPath(A->GetName(), LeftA, RightA);
-		RigBlueprint->Model->SplitPinPath(B->GetName(), LeftB, RightB);
-
-		const FControlRigModelPin* PinA = RigBlueprint->Model->FindPin(*LeftA, *RightA);
-		if (PinA)
-		{
-			RigBlueprint->ModelController->PrepareCycleCheckingForPin(*LeftA, *RightA, A->Direction == EGPD_Input);
-		}
-
-		if (A->Direction == EGPD_Input)
-		{
-			FString Temp = LeftA;
-			LeftA = LeftB;
-			LeftB = Temp;
-			Temp = RightA;
-			RightA = RightB;
-			RightB = Temp;
-		}
+		UControlRigGraphNode* RigNodeA = Cast<UControlRigGraphNode>(A->GetOwningNode());
+		UControlRigGraphNode* RigNodeB = Cast<UControlRigGraphNode>(B->GetOwningNode());
 
-		FString FailureReason;
-		bool bResult = RigBlueprint->ModelController->CanLink(*LeftA, *RightA, *LeftB, *RightB, &FailureReason);
-		if (!bResult)
+		if (RigNodeA && RigNodeB)
 		{
-			return FPinConnectionResponse(CONNECT_RESPONSE_DISALLOW, FText::FromString(FailureReason));
+			URigVMPin* PinA = RigNodeA->GetModelPinFromPinPath(A->GetName());
+			if (PinA)
+			{
+				PinA = PinA->GetPinForLink();
+				RigBlueprint->Model->PrepareCycleChecking(PinA, A->Direction == EGPD_Input);
+			}
+
+			URigVMPin* PinB = RigNodeB->GetModelPinFromPinPath(B->GetName());
+			if (PinB)
+			{
+				PinB = PinB->GetPinForLink();
+			}
+
+			if (A->Direction == EGPD_Input)
+			{
+				URigVMPin* Temp = PinA;
+				PinA = PinB;
+				PinB = Temp;
+			}
+
+			FString FailureReason;
+			bool bResult = RigBlueprint->Model->CanLink(PinA, PinB, &FailureReason);
+			if (!bResult)
+			{
+				return FPinConnectionResponse(CONNECT_RESPONSE_DISALLOW, FText::FromString(FailureReason));
+			}
+			return FPinConnectionResponse(CONNECT_RESPONSE_MAKE, LOCTEXT("ConnectResponse_Allowed", "Connect"));
 		}
-		return FPinConnectionResponse(CONNECT_RESPONSE_MAKE, LOCTEXT("ConnectResponse_Allowed", "Connect"));
 	}
 
 	return FPinConnectionResponse(CONNECT_RESPONSE_DISALLOW, LOCTEXT("ConnectResponse_Disallowed_Unexpected", "Unexpected error"));
@@ -134,9 +142,12 @@ FLinearColor UControlRigGraphSchema::GetPinTypeColor(const FEdGraphPinType& PinT
 	const FName& TypeName = PinType.PinCategory;
 	if (TypeName == UEdGraphSchema_K2::PC_Struct)
 	{
-		if (PinType.PinSubCategoryObject == FControlRigExecuteContext::StaticStruct())
+		if (UStruct* Struct = Cast<UStruct>(PinType.PinSubCategoryObject))
 		{
-			return FLinearColor::White;
+			if (Struct->IsChildOf(FRigVMExecuteContext::StaticStruct()))
+			{
+				return FLinearColor::White;
+			}
 		}
 	}
 	return GetDefault<UEdGraphSchema_K2>()->GetPinTypeColor(PinType);
@@ -151,9 +162,7 @@ void UControlRigGraphSchema::BreakPinLinks(UEdGraphPin& TargetPin, bool bSendsNo
 	const UControlRigBlueprint* RigBlueprint = Cast<UControlRigBlueprint>(Blueprint);
 	if (RigBlueprint != nullptr)
 	{
-		FString Left, Right;
-		RigBlueprint->Model->SplitPinPath(TargetPin.GetName(), Left, Right);
-		RigBlueprint->ModelController->BreakLinks(*Left, *Right, TargetPin.Direction == EGPD_Input);
+		RigBlueprint->Controller->BreakAllLinks(TargetPin.GetName(), TargetPin.Direction == EGPD_Input);
 	}
 }
 
@@ -172,10 +181,7 @@ void UControlRigGraphSchema::BreakSinglePinLink(UEdGraphPin* SourcePin, UEdGraph
 			SourcePin = Temp;
 		}
 		
-		FString SourceLeft, SourceRight, TargetLeft, TargetRight;
-		RigBlueprint->Model->SplitPinPath(SourcePin->GetName(), SourceLeft, SourceRight);
-		RigBlueprint->Model->SplitPinPath(TargetPin->GetName(), TargetLeft, TargetLeft);
-		RigBlueprint->ModelController->BreakLink(*SourceLeft, *SourceRight, *TargetLeft, *TargetRight);
+		RigBlueprint->Controller->BreakLink(SourcePin->GetName(), TargetPin->GetName());
 	}
 }
 
@@ -195,12 +201,67 @@ bool UControlRigGraphSchema::ShouldHidePinDefaultValue(UEdGraphPin* Pin) const
 	return HasParentConnection_Recursive(Pin);
 }
 
+bool UControlRigGraphSchema::IsPinBeingWatched(UEdGraphPin const* Pin) const
+{
+	UBlueprint* Blueprint = FBlueprintEditorUtils::FindBlueprintForNodeChecked(Pin->GetOwningNode());
+	const UControlRigBlueprint* RigBlueprint = Cast<UControlRigBlueprint>(Blueprint);
+	if (RigBlueprint != nullptr)
+	{
+		if (URigVMPin* ModelPin = RigBlueprint->Model->FindPin(Pin->GetName()))
+		{
+			return ModelPin->RequiresWatch();
+		}
+	}
+	return false;
+}
+
+void UControlRigGraphSchema::ClearPinWatch(UEdGraphPin const* Pin) const
+{
+	UBlueprint* Blueprint = FBlueprintEditorUtils::FindBlueprintForNodeChecked(Pin->GetOwningNode());
+	const UControlRigBlueprint* RigBlueprint = Cast<UControlRigBlueprint>(Blueprint);
+	if (RigBlueprint != nullptr)
+	{
+		RigBlueprint->Controller->SetPinIsWatched(Pin->GetName(), false);
+	}
+}
+
+void UControlRigGraphSchema::OnPinConnectionDoubleCicked(UEdGraphPin* PinA, UEdGraphPin* PinB, const FVector2D& GraphPosition) const
+{
+	UBlueprint* Blueprint = FBlueprintEditorUtils::FindBlueprintForNodeChecked(PinA->GetOwningNode());
+	const UControlRigBlueprint* RigBlueprint = Cast<UControlRigBlueprint>(Blueprint);
+	if (RigBlueprint != nullptr)
+	{
+		if (URigVMLink* Link = RigBlueprint->Model->FindLink(FString::Printf(TEXT("%s -> %s"), *PinA->GetName(), *PinB->GetName())))
+		{
+			RigBlueprint->Controller->AddRerouteNodeOnLink(Link, false, GraphPosition);
+		}
+	}
+}
+
+bool UControlRigGraphSchema::MarkBlueprintDirtyFromNewNode(UBlueprint* InBlueprint, UEdGraphNode* InEdGraphNode) const
+{
+	if (InBlueprint == nullptr || InEdGraphNode == nullptr)
+	{
+		return false;
+	}
+	return true;
+}
+
+bool UControlRigGraphSchema::IsStructEditable(UStruct* InStruct) const
+{
+	if (InStruct == FRuntimeFloatCurve::StaticStruct())
+	{
+		return true;
+	}
+	return false;
+}
+
 UControlRigGraphNode* UControlRigGraphSchema::CreateGraphNode(UControlRigGraph* InGraph, const FName& InPropertyName) const
 {
 	const bool bSelectNewNode = true;
 	FGraphNodeCreator<UControlRigGraphNode> GraphNodeCreator(*InGraph);
 	UControlRigGraphNode* ControlRigGraphNode = GraphNodeCreator.CreateNode(bSelectNewNode);
-	ControlRigGraphNode->SetPropertyName(InPropertyName);
+	ControlRigGraphNode->ModelNodePath = InPropertyName.ToString();
 	GraphNodeCreator.Finalize();
 
 	return ControlRigGraphNode;
@@ -246,12 +307,29 @@ bool UControlRigGraphSchema::ArePinsCompatible(const UEdGraphPin* PinA, const UE
 	{
 		return false;
 	}
+
+	// for reroute nodes - always allow it
+	if (PinA->PinType.PinCategory == TEXT("REROUTE"))
+	{
+		UControlRigGraphSchema* MutableThis = (UControlRigGraphSchema*)this;
+		MutableThis->LastPinForCompatibleCheck = PinB;
+		MutableThis->bLastPinWasInput = true;
+		return true;
+	}
+	if (PinB->PinType.PinCategory == TEXT("REROUTE"))
+	{
+		UControlRigGraphSchema* MutableThis = (UControlRigGraphSchema*)this;
+		MutableThis->LastPinForCompatibleCheck = PinA;
+		MutableThis->bLastPinWasInput = false;
+		return true;
+	}
+
 	return GetDefault<UEdGraphSchema_K2>()->ArePinsCompatible(PinA, PinB, CallingContext, bIgnoreArray);
 }
 
 void UControlRigGraphSchema::RenameNode(UControlRigGraphNode* Node, const FName& InNewNodeName) const
 {
-	Node->NodeTitleFull = Node->NodeTitle = FText::FromName(InNewNodeName);
+	Node->NodeTitle = FText::FromName(InNewNodeName);
 	Node->Modify();
 }
 
diff --git a/Engine/Plugins/Experimental/ControlRig/Source/ControlRigDeveloper/Public/ControlRigBlueprint.h b/Engine/Plugins/Experimental/ControlRig/Source/ControlRigDeveloper/Public/ControlRigBlueprint.h
--- a/Engine/Plugins/Experimental/ControlRig/Source/ControlRigDeveloper/Public/ControlRigBlueprint.h
+++ b/Engine/Plugins/Experimental/ControlRig/Source/ControlRigDeveloper/Public/ControlRigBlueprint.h
@@ -11,104 +11,30 @@
 #include "Rigs/RigHierarchyContainer.h"
 #include "Rigs/RigCurveContainer.h"
 #include "Interfaces/Interface_PreviewMeshProvider.h"
-#include "ControlRigController.h"
 #include "ControlRigGizmoLibrary.h"
+#include "RigVMCore/RigVM.h"
+#include "RigVMCore/RigVMStatistics.h"
+#include "RigVMModel/RigVMController.h"
+#include "RigVMCompiler/RigVMCompiler.h"
+#include "ControlRigHierarchyModifier.h"
+#include "Drawing/ControlRigDrawContainer.h"
 #include "ControlRigBlueprint.generated.h"
 
 class UControlRigBlueprintGeneratedClass;
 class USkeletalMesh;
 class UControlRigGraph;
 
-/** 
-  * Source data used by the FControlRigBlueprintCompiler, can't be an editor plugin
-  * because it is needed when running with -game.
-  */
+DECLARE_EVENT_TwoParams(UControlRigBlueprint, FOnVMCompiledEvent, UBlueprint*, URigVM*);
 
-/** A link between two properties. Links become copies between property data at runtime. */
-USTRUCT()
-struct CONTROLRIGDEVELOPER_API FControlRigBlueprintPropertyLink
-{
-	GENERATED_BODY()
-
-	FControlRigBlueprintPropertyLink() 
-		: SourceLinkIndex(0)
-		, DestLinkIndex(0)
-		, SourcePropertyHash(0)
-		, DestPropertyHash(0)
-	{}
-
-	FControlRigBlueprintPropertyLink(const FString& InSourcePropertyPath, const FString& InDestPropertyPath, uint32 InSourceLinkIndex, uint32 InDestLinkIndex)
-		: SourcePropertyPath(InSourcePropertyPath)
-		, DestPropertyPath(InDestPropertyPath)
-		, SourceLinkIndex(InSourceLinkIndex)
-		, DestLinkIndex(InDestLinkIndex)
-		, SourcePropertyHash(FCrc::StrCrc32<TCHAR>(*SourcePropertyPath))
-		, DestPropertyHash(FCrc::StrCrc32<TCHAR>(*DestPropertyPath))
-	{
-	}
-
-	friend bool operator==(const FControlRigBlueprintPropertyLink& A, const FControlRigBlueprintPropertyLink& B)
-	{
-		return A.SourcePropertyHash == B.SourcePropertyHash && A.DestPropertyHash == B.DestPropertyHash;
-	}
-
-	const FString& GetSourcePropertyPath() const { return SourcePropertyPath; }
-	const FString& GetDestPropertyPath() const { return DestPropertyPath; }
-
-	int32 GetSourceLinkIndex() const { return SourceLinkIndex; }
-	int32 GetDestLinkIndex() const { return DestLinkIndex; }
-
-private:
-
-	FString GetSourceUnitName() const { return GetUnitName(SourcePropertyPath); }
-	FString GetDestUnitName() const { return GetUnitName(DestPropertyPath); }
-
-	static FString GetUnitName(FString Input)
-	{
-		int32 ParseIndex = 0;
-		if (Input.FindChar(TCHAR('.'), ParseIndex))
-		{
-			return Input.Left(ParseIndex);
-		}
-		return Input;
-	}
-
-	/** Path to the property we are linking from */
-	UPROPERTY(VisibleAnywhere, Category="Links")
-	FString SourcePropertyPath;
-
-	/** Path to the property we are linking to */
-	UPROPERTY(VisibleAnywhere, Category="Links")
-	FString DestPropertyPath;
-
-	/** Index of the link on the source unit */
-	UPROPERTY()
-	int32 SourceLinkIndex;
-
-	/** Index of the link on the destination unit */
-	UPROPERTY()
-	int32 DestLinkIndex;
-
-	// Hashed strings for faster comparisons
-	UPROPERTY()
-	uint32 SourcePropertyHash;
-
-	// Hashed strings for faster comparisons
-	UPROPERTY()
-	uint32 DestPropertyHash;
-
-	friend class FControlRigBlueprintCompilerContext;
-};
-
-UCLASS(BlueprintType)
+UCLASS(BlueprintType, meta=(IgnoreClassThumbnail))
 class CONTROLRIGDEVELOPER_API UControlRigBlueprint : public UBlueprint, public IInterface_PreviewMeshProvider
 {
-	GENERATED_BODY()
+	GENERATED_UCLASS_BODY()
 
 public:
 	UControlRigBlueprint();
 
-	void InitializeModel();
+	void InitializeModelIfRequired();
 
 	/** Get the (full) generated class for this control rig blueprint */
 	UControlRigBlueprintGeneratedClass* GetControlRigBlueprintGeneratedClass() const;
@@ -124,9 +50,7 @@ public:
 	virtual bool IsValidForBytecodeOnlyRecompile() const override { return false; }
 	virtual void LoadModulesRequiredForCompilation() override;
 	virtual void GetTypeActions(FBlueprintActionDatabaseRegistrar& ActionRegistrar) const override;
-	virtual void GetInstanceActions(FBlueprintActionDatabaseRegistrar& ActionRegistrar) const override;	
 	virtual void SetObjectBeingDebugged(UObject* NewObject) override;
-	virtual bool RequiresMarkAsStructurallyModifiedOnUndo() const override { return false; }
 	virtual void PostLoad() override;
 	virtual void PostTransacted(const FTransactionObjectEvent& TransactionEvent) override;
 
@@ -140,45 +64,85 @@ public:
 
 
 #endif	// #if WITH_EDITOR
-
-	/** Make a property link between the specified properties - used by the compiler */
-	void MakePropertyLink(const FString& InSourcePropertyPath, const FString& InDestPropertyPath, int32 InSourceLinkIndex, int32 InDestLinkIndex);
+	virtual bool ShouldBeMarkedDirtyUponTransaction() const override { return false; }
 
 	/** IInterface_PreviewMeshProvider interface */
 	virtual void SetPreviewMesh(USkeletalMesh* PreviewMesh, bool bMarkAsDirty = true) override;
 	virtual USkeletalMesh* GetPreviewMesh() const override;
 
-	UPROPERTY(transient)
-	UControlRigModel* Model;
+	UFUNCTION()
+	void RecompileVM();
 
-	UPROPERTY(transient)
-	UControlRigController* ModelController;
+	UFUNCTION()
+	void RecompileVMIfRequired();
+	
+	UFUNCTION()
+	void RequestAutoVMRecompilation();
+
+	void IncrementVMRecompileBracket();
+	void DecrementVMRecompileBracket();
+	void RequestControlRigInit();
 
-	// hack until RigVM refactoring
-	TMap<FName, EControlRigModelParameterType> NodeToParameterType;
+	UPROPERTY(EditAnywhere, BlueprintReadWrite, Category = "VM")
+	FRigVMCompileSettings VMCompileSettings;
+
+	UPROPERTY(BlueprintReadOnly, Category = "VM")
+	URigVMGraph* Model;
+
+	UPROPERTY(BlueprintReadOnly, transient, Category = "VM")
+	URigVMController* Controller;
+
+	UPROPERTY(transient)
+	TMap<FString, FRigVMOperand> PinToOperandMap;
 
 	bool bSuspendModelNotificationsForSelf;
 	bool bSuspendModelNotificationsForOthers;
-	FName LastNameFromNotification;
 
-	void PopulateModelFromGraph(const UControlRigGraph* InGraph);
+	void PopulateModelFromGraphForBackwardsCompatibility(UControlRigGraph* InGraph);
 	void RebuildGraphFromModel();
 
-	UControlRigModel::FModifiedEvent& OnModified();
+	FRigVMGraphModifiedEvent& OnModified();
+	FOnVMCompiledEvent& OnVMCompiled();
+
+	UFUNCTION(BlueprintCallable, Category = "VM")
+	static TArray<UControlRigBlueprint*> GetCurrentlyOpenRigBlueprints();
+
+	UFUNCTION(BlueprintCallable, Category = "VM")
+	static TArray<UStruct*> GetAvailableRigUnits();
+
+	UFUNCTION(BlueprintCallable, Category = "Hierarchy")
+	UControlRigHierarchyModifier* GetHierarchyModifier();
 
 #if WITH_EDITORONLY_DATA
 	UPROPERTY(EditAnywhere, config, Category = DefaultGizmo)
 	TAssetPtr<UControlRigGizmoLibrary> GizmoLibrary;
 #endif
 
-private:
-	/** Links between the various properties we have */
-	UPROPERTY()
-	TArray<FControlRigBlueprintPropertyLink> PropertyLinks;
+	UPROPERTY(transient, VisibleAnywhere, Category = "VM", meta = (DisplayName = "VM Statistics", DisplayAfter = "VMCompileSettings"))
+	FRigVMStatistics Statistics;
 
-	/** list of operators. Visible for debug purpose for now */
-	UPROPERTY()
-	TArray<FControlRigOperator> Operators_DEPRECATED;
+	UPROPERTY(EditAnywhere, Category = "Drawing")
+	FControlRigDrawContainer DrawContainer;
+
+#if WITH_EDITOR
+	/** Remove a transient / temporary control used to interact with a pin */
+	FName AddTransientControl(URigVMPin* InPin);
+
+	/** Remove a transient / temporary control used to interact with a pin */
+	FName RemoveTransientControl(URigVMPin* InPin);
+
+	/** Remove a transient / temporary control used to interact with a bone */
+	FName AddTransientControl(const FRigElementKey& InElement);
+
+	/** Remove a transient / temporary control used to interact with a bone */
+	FName RemoveTransientControl(const FRigElementKey& InElement);
+
+	/** Removes all  transient / temporary control used to interact with pins */
+	void ClearTransientControls();
+
+#endif
+
+private:
 
 	// need list of "allow query property" to "source" - whether rig unit or property itself
 	// this will allow it to copy data to target
@@ -198,7 +162,7 @@ private:
 	FRigCurveContainer CurveContainer_DEPRECATED;
 
 	/** The default skeletal mesh to use when previewing this asset */
-	UPROPERTY(DuplicateTransient, AssetRegistrySearchable)
+	UPROPERTY(AssetRegistrySearchable)
 	TSoftObjectPtr<USkeletalMesh> PreviewSkeletalMesh;
 
 	/** The default skeletal mesh to use when previewing this asset */
@@ -208,15 +172,40 @@ private:
 	UPROPERTY(DuplicateTransient, AssetRegistrySearchable)
 	TSoftObjectPtr<UObject> SourceCurveImport;
 
-	UControlRigModel::FModifiedEvent _ModifiedEvent;
-	void HandleModelModified(const UControlRigModel* InModel, EControlRigModelNotifType InType, const void* InPayload);
-	bool UpdateParametersOnControlRig(UControlRig* InRig = nullptr);
-	bool PerformArrayOperation(const FString& InPropertyPath, TFunctionRef<bool(FScriptArrayHelper&, int32)> InOperation, bool bCallModify, bool bPropagateToInstances);
+	UPROPERTY(transient)
+	bool bAutoRecompileVM;
+
+	UPROPERTY(transient)
+	bool bVMRecompilationRequired;
+
+	UPROPERTY(transient)
+	int32 VMRecompilationBracket;
+
+	UPROPERTY(transient)
+	UControlRigHierarchyModifier* HierarchyModifier;
+
+	FRigVMGraphModifiedEvent ModifiedEvent;
+	void Notify(ERigVMGraphNotifType InNotifType, UObject* InSubject);
+	void HandleModifiedEvent(ERigVMGraphNotifType InNotifType, URigVMGraph* InGraph, UObject* InSubject);
+
+	FOnVMCompiledEvent VMCompiledEvent;
+
+	static TArray<UControlRigBlueprint*> sCurrentlyOpenedRigBlueprints;
+
 	void CleanupBoneHierarchyDeprecated();
 
-	void PropagatePoseFromInstanceToBP();
-	void PropagateHierarchyFromBPToInstances(bool bInitialize = true);
+public:
+	void PropagatePoseFromInstanceToBP(UControlRig* InControlRig);
+	void PropagatePoseFromBPToInstances();
+	void PropagateHierarchyFromBPToInstances(bool bInitializeContainer = true, bool bInitializeRigs = true);
+	void PropagateDrawInstructionsFromBPToInstances();
+	void PropagatePropertyFromBPToInstances(FRigElementKey InRigElement, const FProperty* InProperty);
+	void PropagatePropertyFromInstanceToBP(FRigElementKey InRigElement, const FProperty* InProperty, UControlRig* InInstance);
+
+private:
+
 #if WITH_EDITOR
+
 	void HandleOnElementAdded(FRigHierarchyContainer* InContainer, const FRigElementKey& InKey);
 	void HandleOnElementRemoved(FRigHierarchyContainer* InContainer, const FRigElementKey& InKey);
 	void HandleOnElementRenamed(FRigHierarchyContainer* InContainer, ERigElementType InElementType, const FName& InOldName, const FName& InNewName);
@@ -230,4 +219,7 @@ private:
 	friend class FControlRigEditor;
 	friend class UEngineTestControlRig;
 	friend class FControlRigEditMode;
+	friend class FControlRigBlueprintActions;
+	friend class FControlRigDrawContainerDetails;
+	friend class UDefaultControlRigManipulationLayer;
 };
diff --git a/Engine/Plugins/Experimental/ControlRig/Source/ControlRigDeveloper/Public/ControlRigBlueprintCompiler.h b/Engine/Plugins/Experimental/ControlRig/Source/ControlRigDeveloper/Public/ControlRigBlueprintCompiler.h
--- a/Engine/Plugins/Experimental/ControlRig/Source/ControlRigDeveloper/Public/ControlRigBlueprintCompiler.h
+++ b/Engine/Plugins/Experimental/ControlRig/Source/ControlRigDeveloper/Public/ControlRigBlueprintCompiler.h
@@ -29,22 +29,26 @@ public:
 	}
 
 	// FKismetCompilerContext interface
-	virtual void MergeUbergraphPagesIn(UEdGraph* Ubergraph) override;
+	virtual void ValidateLink(const UEdGraphPin* PinA, const UEdGraphPin* PinB) const override {}
+	virtual void ValidatePin(const UEdGraphPin* Pin) const override {}
+	virtual void ValidateNode(const UEdGraphNode* Node) const override {}
+	virtual bool CanIgnoreNode(const UEdGraphNode* Node) const override { return true; }
+	virtual bool ShouldForceKeepNode(const UEdGraphNode* Node) const override { return false; }
+	virtual void PrecompileFunction(FKismetFunctionContext& Context, EInternalCompilerFlags InternalFlags) override {}
 	virtual void PostCompile() override;
 	virtual void CopyTermDefaultsToDefaultObject(UObject* DefaultObject) override;
 	virtual void EnsureProperGeneratedClass(UClass*& TargetUClass) override;
 	virtual void SpawnNewClass(const FString& NewClassName) override;
 	virtual void OnNewClassSet(UBlueprintGeneratedClass* ClassToUse) override;
+	virtual void PruneIsolatedNodes(const TArray<UEdGraphNode*>& RootSet, TArray<UEdGraphNode*>& GraphNodes) override {}
 	virtual void CleanAndSanitizeClass(UBlueprintGeneratedClass* ClassToClean, UObject*& InOldCDO) override;
+	virtual void MergeUbergraphPagesIn(UEdGraph* Ubergraph) override {};
 
 private:
 
 	// used to fail a compilation and mark the blueprint in error
 	void MarkCompilationFailed(const FString& Message);
 
-	// utility function to build property links from the ubergraphs
-	void BuildPropertyLinks();
-
 private:
 	/** the new class we are generating */
 	UControlRigBlueprintGeneratedClass* NewControlRigBlueprintGeneratedClass;
diff --git a/Engine/Plugins/Experimental/ControlRig/Source/ControlRigDeveloper/Public/ControlRigBlueprintUtils.h b/Engine/Plugins/Experimental/ControlRig/Source/ControlRigDeveloper/Public/ControlRigBlueprintUtils.h
--- a/Engine/Plugins/Experimental/ControlRig/Source/ControlRigDeveloper/Public/ControlRigBlueprintUtils.h
+++ b/Engine/Plugins/Experimental/ControlRig/Source/ControlRigDeveloper/Public/ControlRigBlueprintUtils.h
@@ -14,61 +14,8 @@ class UEdGraphPin;
 
 struct CONTROLRIGDEVELOPER_API FControlRigBlueprintUtils
 {
-
-/** 
- * Gets a new name for a unit member struct
- * @param	InBlueprint			The blueprint we want to create a new member in
- * @param	InStructTemplate	The struct template we want to use
- * @return the name of the new member
- */
-static FName GetNewUnitMemberName(UBlueprint* InBlueprint, const UStruct* InStructTemplate);
-
-/** 
- * Adds a new unit member struct. 
- * @param	InBlueprint			The blueprint we want to create a new member in
- * @param	InStructTemplate	The struct template we want to use
- * @return the name of the new member, or NAME_None if the member was not created
- */
-static FName AddUnitMember(UBlueprint* InBlueprint, const UStruct* InStructTemplate, const FName& InName= NAME_None);
-
-/** 
- * Gets a new name for a member
- * @param	InBlueprint			The blueprint we want to create a new member in
- * @param	InVariableDesc		A description of the variable type, used to create a base variable name
- * @return the name of the new member
- */
-static FName GetNewPropertyMemberName(UBlueprint* InBlueprint, const FString& InVariableDesc);
-
-/** 
- * Adds a new property member. 
- * @param	InBlueprint			The blueprint we want to create a new member in
- * @param	InPinType			The type of the property we want to create
- * @param	InVariableDesc		A description of the variable type, used to create a base variable name
- * @return the name of the new member, or NAME_None if the member was not created
- */
-static FName AddPropertyMember(UBlueprint* InBlueprint, const FEdGraphPinType& InPinType, const FString& InVariableDesc);
-
-/**
- * Instantiate a node in the specified graph for the supplied property
- * @param	InGraph				The graph to create the node in
- * @param	InPropertyName		The property the node represents
- * @param	InLocation			Optional location to create the node at
- * @return the new graph node
- */
-static UControlRigGraphNode* InstantiateGraphNodeForProperty(UEdGraph* InGraph, const FName& InPropertyName, const FVector2D& InLocation = FVector2D::ZeroVector, const FEdGraphPinType& InPinType = FEdGraphPinType());
-static UControlRigGraphNode* InstantiateGraphNodeForStructPath(UEdGraph* InGraph, const FName& InPropertyName, const FVector2D& InLocation, const FString& InStructPath);
-
-/**
- * Check whether we can instantiate a node in the specified graph for the specified property
- * We don't allow properties to be instantiated more than once
- * @param	InGraph				The graph to create the node in
- * @param	InPropertyName		The property the node represents
- * @return true if the node can be instantiated
- */
-static bool CanInstantiateGraphNodeForProperty(UEdGraph* InGraph, const FName& InPropertyName);
-
 /** Call a function for each valid rig unit struct */
-static void ForAllRigUnits(TFunction<void(UStruct*)> InFunction);
+static void ForAllRigUnits(TFunction<void(UScriptStruct*)> InFunction);
 
 /** Handle blueprint node reconstruction */
 static void HandleReconstructAllNodes(UBlueprint* InBlueprint);
@@ -76,9 +23,6 @@ static void HandleReconstructAllNodes(UBlueprint* InBlueprint);
 /** Handle blueprint node refresh */
 static void HandleRefreshAllNodes(UBlueprint* InBlueprint);
 
-/** Handle variables getting renamed */
-static void HandleRenameVariableReferencesEvent(UBlueprint* InBlueprint, UClass* InVariableClass, const FName& InOldVarName, const FName& InNewVarName);
-
 /** remove the variable if not used by anybody else but ToBeDeleted*/
 static void RemoveMemberVariableIfNotUsed(UBlueprint* Blueprint, const FName VarName, UControlRigGraphNode* ToBeDeleted);
 
