COPY from //Dev-Anim to //Dev-Main

diff --git a/Engine/Plugins/Experimental/ControlRig/Source/ControlRigDeveloper/Public/ControlRigDeveloper.h b/Engine/Plugins/Experimental/ControlRig/Source/ControlRigDeveloper/Public/ControlRigDeveloper.h
--- a/Engine/Plugins/Experimental/ControlRig/Source/ControlRigDeveloper/Public/ControlRigDeveloper.h
+++ b/Engine/Plugins/Experimental/ControlRig/Source/ControlRigDeveloper/Public/ControlRigDeveloper.h
@@ -5,6 +5,8 @@
 #include "CoreMinimal.h"
 #include "Modules/ModuleInterface.h"
 
+DECLARE_LOG_CATEGORY_EXTERN(LogControlRigDeveloper, Log, All);
+
 class IControlRigDeveloperModule : public IModuleInterface
 {
 public:
diff --git a/Engine/Plugins/Experimental/ControlRig/Source/ControlRigDeveloper/Public/ControlRigHierarchyModifier.h b/Engine/Plugins/Experimental/ControlRig/Source/ControlRigDeveloper/Public/ControlRigHierarchyModifier.h
new file mode 100644
--- /dev/null
+++ b/Engine/Plugins/Experimental/ControlRig/Source/ControlRigDeveloper/Public/ControlRigHierarchyModifier.h
@@ -0,0 +1,180 @@
+// Copyright Epic Games, Inc. All Rights Reserved.
+
+#pragma once
+
+#include "CoreMinimal.h"
+#include "UObject/ObjectMacros.h"
+#include "ControlRigDefines.h"
+#include "Rigs/RigHierarchyContainer.h"
+#include "ControlRigHierarchyModifier.generated.h"
+
+UCLASS(BlueprintType)
+class CONTROLRIGDEVELOPER_API UControlRigHierarchyModifier : public UObject
+{
+	GENERATED_BODY()
+
+public:
+
+	UControlRigHierarchyModifier();
+
+#if WITH_EDITOR
+
+	// Returns the keys of all elements within the hierarchy
+	UFUNCTION(BlueprintCallable, Category = "Hierarchy")
+	TArray<FRigElementKey> GetElements() const;
+
+	// Adds a new single bone
+	UFUNCTION(BlueprintCallable, Category = "Hierarchy")
+	FRigElementKey AddBone(
+		const FName& InNewName,
+		const FName& InParentName = NAME_None,
+		ERigBoneType InType = ERigBoneType::User
+	);
+
+	// Returns a single bone from provided key
+	UFUNCTION(BlueprintCallable, Category = "Hierarchy")
+	FRigBone GetBone(const FRigElementKey& InKey);
+
+	// Updates a single bone
+	UFUNCTION(BlueprintCallable, Category = "Hierarchy")
+	void SetBone(const FRigBone& InElement);
+
+	// Adds a new single control
+	UFUNCTION(BlueprintCallable, Category = "Hierarchy")
+	FRigElementKey AddControl(
+		const FName& InNewName,
+		ERigControlType InControlType = ERigControlType::Transform,
+		const FName& InParentName = NAME_None,
+		const FName& InSpaceName = NAME_None,
+		const FName& InGizmoName = TEXT("Gizmo"),
+		const FLinearColor& InGizmoColor = FLinearColor(1.0, 0.0, 0.0, 1.0)
+	);
+
+	// Returns a single control from provided key
+	UFUNCTION(BlueprintCallable, Category = "Hierarchy")
+	FRigControl GetControl(const FRigElementKey& InKey);
+
+	// Updates a single control
+	UFUNCTION(BlueprintCallable, Category = "Hierarchy")
+	void SetControl(const FRigControl& InElement);
+
+	// Adds a new single space
+	UFUNCTION(BlueprintCallable, Category = "Hierarchy")
+	FRigElementKey AddSpace(
+		const FName& InNewName,
+		ERigSpaceType InSpaceType = ERigSpaceType::Global,
+		const FName& InParentName = NAME_None
+	);
+
+	// Returns a single space from provided key
+	UFUNCTION(BlueprintCallable, Category = "Hierarchy")
+	FRigSpace GetSpace(const FRigElementKey& InKey);
+
+	// Updates a single space
+	UFUNCTION(BlueprintCallable, Category = "Hierarchy")
+	void SetSpace(const FRigSpace& InElement);
+
+	// Adds a new single curve
+	UFUNCTION(BlueprintCallable, Category = "Hierarchy")
+	FRigElementKey AddCurve(const FName& InNewName, float InValue = 0.f);
+
+	// Returns a single curve from provided key
+	UFUNCTION(BlueprintCallable, Category = "Hierarchy")
+	FRigCurve GetCurve(const FRigElementKey& InKey);
+
+	// Updates a single curve
+	UFUNCTION(BlueprintCallable, Category = "Hierarchy")
+	void SetCurve(const FRigCurve& InElement);
+
+	// Removes a single element, returns true if successful
+	UFUNCTION(BlueprintCallable, Category = "Hierarchy")
+	bool RemoveElement(const FRigElementKey& InElement);
+
+	// Renames an existing element and returns the new element key
+	UFUNCTION(BlueprintCallable, Category = "Hierarchy")
+	FRigElementKey RenameElement(const FRigElementKey& InElement, const FName& InNewName);
+
+	// Reparents an element to another element, returns true if successful
+	UFUNCTION(BlueprintCallable, Category = "Hierarchy")
+	bool ReparentElement(const FRigElementKey& InElement, const FRigElementKey& InNewParent);
+
+	// Returns the keys of all selected elements within the hierarchy
+	UFUNCTION(BlueprintCallable, Category = "Hierarchy")
+	TArray<FRigElementKey> GetSelection() const;
+
+	// Selects or deselects a given element
+	UFUNCTION(BlueprintCallable, Category = "Hierarchy")
+	bool Select(const FRigElementKey& InKey, bool bSelect = true);
+	
+	// Clears the selection
+	UFUNCTION(BlueprintCallable, Category = "Hierarchy")
+	bool ClearSelection();
+
+	// Returns true if a given element is currently selected
+	UFUNCTION(BlueprintCallable, Category = "Hierarchy")
+	bool IsSelected(const FRigElementKey& InKey) const;
+
+	// Initializes the rig, but calling reset on all elements
+	UFUNCTION(BlueprintCallable, Category = "Hierarchy")
+	void Initialize(bool bResetTransforms = true);
+	
+	// Removes all elements of the hierarchy
+	UFUNCTION(BlueprintCallable, Category = "Hierarchy")
+	void Reset();
+	
+	// Resets the transforms on all elements of the hierarchy
+	UFUNCTION(BlueprintCallable, Category = "Hierarchy")
+	void ResetTransforms();
+
+	// Returns the initial transform for a given element key
+	UFUNCTION(BlueprintCallable, Category = "Hierarchy")
+	FTransform GetInitialTransform(const FRigElementKey& InKey) const;
+
+	// Sets the initial transform for a given element key
+	UFUNCTION(BlueprintCallable, Category = "Hierarchy")
+	void SetInitialTransform(const FRigElementKey& InKey, const FTransform& InTransform);
+
+	// Returns the initial global transform for a given element key
+	UFUNCTION(BlueprintCallable, Category = "Hierarchy")
+	FTransform GetInitialGlobalTransform(const FRigElementKey& InKey) const;
+
+	// Sets the initial global transform for a given element key
+	UFUNCTION(BlueprintCallable, Category = "Hierarchy")
+	void SetInitialGlobalTransform(const FRigElementKey& InKey, const FTransform& InTransform);
+
+	// Returns the current local transform of a given element key
+	UFUNCTION(BlueprintCallable, Category = "Hierarchy")
+	FTransform GetLocalTransform(const FRigElementKey& InKey) const;
+
+	// Sets the current local transform of a given element key
+	UFUNCTION(BlueprintCallable, Category = "Hierarchy")
+	void SetLocalTransform(const FRigElementKey& InKey, const FTransform& InTransform);
+
+	// Returns the current global transform of a given element key
+	UFUNCTION(BlueprintCallable, Category = "Hierarchy")
+	FTransform GetGlobalTransform(const FRigElementKey& InKey) const;
+
+	// Sets the current global transform of a given element key
+	UFUNCTION(BlueprintCallable, Category = "Hierarchy")
+	void SetGlobalTransform(const FRigElementKey& InKey, const FTransform& InTransform);
+
+	// Exports the elements provided to text (for copy & paste, import / export)
+	UFUNCTION(BlueprintCallable, Category = "Hierarchy")
+	FString ExportToText(const TArray<FRigElementKey>& InElementsToExport) const;
+
+	// Imports the content of the provided text and returns the keys created
+	UFUNCTION(BlueprintCallable, Category = "Hierarchy")
+	TArray<FRigElementKey> ImportFromText(const FString& InContent, ERigHierarchyImportMode InImportMode = ERigHierarchyImportMode::Append, bool bSelectNewElements = true);
+
+#endif // WITH_EDITOR
+
+private:
+
+	FRigHierarchyContainer* Container;
+	FRigBone InvalidBone;
+	FRigControl InvalidControl;
+	FRigSpace InvalidSpace;
+	FRigCurve InvalidCurve;
+
+	friend class UControlRigBlueprint;
+};
diff --git a/Engine/Plugins/Experimental/ControlRig/Source/ControlRigDeveloper/Public/Graph/ControlRigGraph.h b/Engine/Plugins/Experimental/ControlRig/Source/ControlRigDeveloper/Public/Graph/ControlRigGraph.h
--- a/Engine/Plugins/Experimental/ControlRig/Source/ControlRigDeveloper/Public/Graph/ControlRigGraph.h
+++ b/Engine/Plugins/Experimental/ControlRig/Source/ControlRigDeveloper/Public/Graph/ControlRigGraph.h
@@ -5,12 +5,14 @@
 #include "EdGraph/EdGraph.h"
 #include "Graph/ControlRigGraphNode.h"
 #include "Rigs/RigHierarchyContainer.h"
-#include "ControlRigModel.h"
+#include "RigVMModel/RigVMGraph.h"
+#include "Drawing/ControlRigDrawContainer.h"
 #include "ControlRigGraph.generated.h"
 
 class UControlRigBlueprint;
 class UControlRigGraphSchema;
 class UControlRig;
+class URigVMController;
 struct FRigCurveContainer;
 
 UCLASS()
@@ -33,33 +35,31 @@ public:
 #endif
 #if WITH_EDITOR
 
-	virtual void PostLoad() override;
-	void OnBlueprintCompiledPostLoad(UBlueprint*);
-	FDelegateHandle BlueprintOnCompiledHandle;
 
-	void CacheNameLists(const FRigHierarchyContainer* Container);
+	void CacheNameLists(const FRigHierarchyContainer* HierarchyContainer, const FControlRigDrawContainer* DrawContainer);
 
 	const TArray<TSharedPtr<FString>>& GetBoneNameList() const;
 	const TArray<TSharedPtr<FString>>& GetControlNameList() const;
 	const TArray<TSharedPtr<FString>>& GetSpaceNameList() const;
 	const TArray<TSharedPtr<FString>>& GetCurveNameList() const;
+	const TArray<TSharedPtr<FString>>& GetDrawingNameList() const;
 
 	bool bSuspendModelNotifications;
 	bool bIsTemporaryGraphForCopyPaste;
 
-	UEdGraphNode* FindNodeFromPropertyName(const FName& InPropertyName);
+	UEdGraphNode* FindNodeForModelNodeName(const FName& InModelNodeName);
 
 private:
 
-	void HandleModelModified(const UControlRigModel* InModel, EControlRigModelNotifType InType, const void* InPayload);
+	void HandleModifiedEvent(ERigVMGraphNotifType InNotifType, URigVMGraph* InGraph, UObject* InSubject);
 
 	template<class T>
-	void CacheNameList(const T& Hierarchy, TArray<TSharedPtr<FString>>& OutNameList)
+	void CacheNameList(const T& ElementList, TArray<TSharedPtr<FString>>& OutNameList)
 	{
 		DECLARE_SCOPE_HIERARCHICAL_COUNTER_FUNC()
 
 		TArray<FString> Names;
-		for (auto Element : Hierarchy)
+		for (auto Element : ElementList)
 		{
 			Names.Add(Element.Name.ToString());
 		}
@@ -73,14 +73,35 @@ private:
 		}
 	}
 
-	TArray<UControlRigGraphNode*> FoundHierarchyRefVariableNodes;
-	TArray<UControlRigGraphNode*> FoundHierarchyRefMutableNodes;
-	TMap<UControlRigGraphNode*, TArray<UControlRigGraphNode*>> FoundHierarchyRefConnections;
 
 	TArray<TSharedPtr<FString>> BoneNameList;
 	TArray<TSharedPtr<FString>> ControlNameList;
 	TArray<TSharedPtr<FString>> SpaceNameList;
 	TArray<TSharedPtr<FString>> CurveNameList;
+	TArray<TSharedPtr<FString>> DrawingNameList;
+
+	bool bIsSelecting;
+
+#endif
+#if WITH_EDITORONLY_DATA
+
+	UPROPERTY(transient)
+	URigVMGraph* TemplateModel;
+
+	UPROPERTY(transient)
+	URigVMController* TemplateController;
+
+#endif
+#if WITH_EDITOR
+
+	URigVMController* GetTemplateController();
+
+	friend class UControlRigUnitNodeSpawner;
+	friend class UControlRigVariableNodeSpawner;
+	friend class UControlRigParameterNodeSpawner;
+	friend class UControlRigRerouteNodeSpawner;
 #endif
+	friend class UControlRigGraphNode;
+	friend class FControlRigEditor;
 };
 
diff --git a/Engine/Plugins/Experimental/ControlRig/Source/ControlRigDeveloper/Public/Graph/ControlRigGraphNode.h b/Engine/Plugins/Experimental/ControlRig/Source/ControlRigDeveloper/Public/Graph/ControlRigGraphNode.h
--- a/Engine/Plugins/Experimental/ControlRig/Source/ControlRigDeveloper/Public/Graph/ControlRigGraphNode.h
+++ b/Engine/Plugins/Experimental/ControlRig/Source/ControlRigDeveloper/Public/Graph/ControlRigGraphNode.h
@@ -9,7 +9,7 @@
 #include "K2Node.h"
 #include "UObject/SoftObjectPath.h"
 #include "EdGraphSchema_K2.h"
-#include "ControlRigModel.h"
+#include "RigVMModel/RigVMGraph.h"
 #include "ControlRigGraphNode.generated.h"
 
 class FBlueprintActionDatabaseRegistrar;
@@ -17,96 +17,6 @@ class UEdGraph;
 struct FSlateIcon;
 class UControlRigBlueprint;
 
-/** Information about a control rig field */
-class FControlRigField
-{
-public:
-	FControlRigField(const FEdGraphPinType& InPinType, const FString& InPinPath, const FText& InDisplayNameText, int32 InArrayIndex = INDEX_NONE)
-		: InputPin(nullptr)
-		, OutputPin(nullptr)
-		, PinType(InPinType)
-		, PinPath(InPinPath)
-		, DisplayNameText(InDisplayNameText)
-		, ArrayIndex(InArrayIndex)
-	{
-	}
-
-	virtual ~FControlRigField() {}
-
-	/** Get the field we refer to */
-	virtual FField* GetField() const { return nullptr; }
-
-	/** Get the input pin for this item */
-	virtual const FControlRigModelPin* GetPin() const { return nullptr; }
-	
-	/** Get the output pin for this item */
-	virtual UEdGraphPin* GetOutputPin() const { return OutputPin; }
-
-	/** Get the name of this field */
-	virtual FString GetPinPath() const { return PinPath; }
-
-	/** Get the name to display for this field */
-	virtual FText GetDisplayNameText() const { return DisplayNameText; }
-
-	/** Get the tooltip to display for this field */
-	virtual FText GetTooltipText() const { return TooltipText; }
-
-	/** Get the pin type to use for this field */
-	virtual FEdGraphPinType GetPinType() const { return PinType; }
-
-	/** Cached input pin */
-	UEdGraphPin* InputPin;
-
-	/** Cached output pin */
-	UEdGraphPin* OutputPin;
-
-	/** Pin type we use for the field */
-	FEdGraphPinType PinType;
-
-	/** Cached name to display */
-	FString PinPath;
-
-	/** The name to display for this field  */
-	FText DisplayNameText;
-
-	/** The name to display for this field's tooltip  */
-	FText TooltipText;
-
-	/** The array index, or INDEX_NONE if this is not an array-indexed property */
-	int32 ArrayIndex;
-
-	/** Any sub-fields are represented by children of this field */
-	TArray<TSharedRef<FControlRigField>> Children;
-};
-
-/** Information about an input/output pin */
-class FControlRigPin : public FControlRigField
-{
-private:
-	static FEdGraphPinType GetPinTypeFromPin(const FControlRigModelPin* InPin)
-	{
-		return InPin->Type;
-	}
-
-	static FText GetDisplayNameForPin(const FControlRigModelPin* InPin)
-	{
-		return FText::FromString(InPin->Name.ToString());
-	}
-
-public:
-	FControlRigPin(const FControlRigModelPin* InPin, const FString& InPinPath, int32 InArrayIndex = INDEX_NONE)
-		: FControlRigField(GetPinTypeFromPin(InPin), *InPinPath, GetDisplayNameForPin(InPin), InArrayIndex)
-		, Pin(*InPin)
-	{
-	}
-
-	virtual const FControlRigModelPin* GetPin() const override { return &Pin; }
-
-private:
-	/** The field that we use as input/output */
-	FControlRigModelPin Pin;
-};
-
 /** Base class for animation ControlRig-related nodes */
 UCLASS()
 class CONTROLRIGDEVELOPER_API UControlRigGraphNode : public UEdGraphNode
@@ -125,38 +35,41 @@ class CONTROLRIGDEVELOPER_API UControlRigGraphNode : public UEdGraphNode
 	friend class SControlRigGraphPinCurveFloat;
 
 private:
+
+	UPROPERTY()
+	FString ModelNodePath;
+
+	UPROPERTY(transient)
+	URigVMNode* CachedModelNode;
+
+	UPROPERTY(transient)
+	TMap<FString, URigVMPin*> CachedModelPins;
+
 	/** The property we represent. For template nodes this represents the struct/property type name. */
 	UPROPERTY()
-	FName PropertyName;
+	FName PropertyName_DEPRECATED;
 
 	UPROPERTY()
-	FString StructPath;
+	FString StructPath_DEPRECATED;
 
 	/** Pin Type for property */
 	UPROPERTY()
-	FEdGraphPinType PinType;
+	FEdGraphPinType PinType_DEPRECATED;
 
 	/** The type of parameter */
 	UPROPERTY()
-	int32 ParameterType;
+	int32 ParameterType_DEPRECATED;
 
 	/** Expanded pins */
 	UPROPERTY()
-	TArray<FString> ExpandedPins;
+	TArray<FString> ExpandedPins_DEPRECATED;
 
 	/** Cached dimensions of this node (used for auto-layout) */
 	FVector2D Dimensions;
 
 	/** The cached node titles */
-	mutable FText NodeTitleFull;
 	mutable FText NodeTitle;
 
-	/** Cached info about input/output pins */
-	TArray<TSharedRef<FControlRigField>> ExecutionInfos;
-	TArray<TSharedRef<FControlRigField>> InputInfos;
-	TArray<TSharedRef<FControlRigField>> InputOutputInfos;
-	TArray<TSharedRef<FControlRigField>> OutputInfos;
-
 public:
 	UControlRigGraphNode();
 
@@ -167,15 +80,16 @@ public:
 	virtual FSlateIcon GetIconAndTint(FLinearColor& OutColor) const override;
 	virtual void AllocateDefaultPins() override;
 	virtual void ReconstructNode() override;
-	virtual void PinConnectionListChanged(UEdGraphPin* Pin) override;
+	virtual void ReconstructNode_Internal(bool bForce = false);
 	virtual void GetNodeContextMenuActions(class UToolMenu* Menu, class UGraphNodeContextMenuContext* Context) const override;
 	virtual void DestroyNode() override;
 	virtual void PinDefaultValueChanged(UEdGraphPin* Pin) override;
-	virtual TSharedPtr<INameValidatorInterface> MakeNameValidator() const override;
 	virtual FText GetTooltipText() const override;
 	virtual bool CanCreateUnderSpecifiedSchema(const UEdGraphSchema* InSchema) const override;
 	virtual void AutowireNewNode(UEdGraphPin* FromPin) override;	
-	virtual void PrepareForCopying() override;
+	virtual bool SupportsCommentBubble() const override { return false; }
+	virtual bool IsSelectedInEditor() const;
+	virtual bool ShouldDrawNodeAsControlPointOnly(int32& OutInputPinIndex, int32& OutOutputPinIndex) const override;
 
 	virtual bool IsDeprecated() const override;
 	virtual FEdGraphNodeDeprecationResponse GetDeprecationResponse(EEdGraphNodeDeprecationType DeprecationType) const override;
@@ -186,113 +100,47 @@ public:
 	/** Get the cached dimensions of this node */
 	const FVector2D& GetDimensions() const { return Dimensions; }
 
-	/** Set the property name we reference */
-	void SetPropertyName(const FName& InPropertyName, bool bReplaceInnerProperties=false);
-
-	/** Get the property name we reference */
-	FName GetPropertyName() const { return PropertyName; }
-
-	/** Get the execution variable names */
-	const TArray<TSharedRef<FControlRigField>>& GetExecutionVariableInfo() const { return ExecutionInfos; }
-
-	/** Get the input variable names */
-	const TArray<TSharedRef<FControlRigField>>& GetInputVariableInfo() const { return InputInfos; }
-
-	/** Get the input-output variable names */
-	const TArray<TSharedRef<FControlRigField>>& GetInputOutputVariableInfo() const { return InputOutputInfos; }
-
-	/** Get the output variable names */
-	const TArray<TSharedRef<FControlRigField>>& GetOutputVariableInfo() const { return OutputInfos; }
-
-	/** Record a pin's expansion state */
-	void SetPinExpansion(const FString& InPinPropertyPath, bool bExpanded);
-
 	/** Check a pin's expansion state */
-	bool IsPinExpanded(const FString& InPinPropertyPath) const;
+	bool IsPinExpanded(const FString& InPinPath);
 
 	/** Propagate pin defaults to underlying properties if they have changed */
 	void CopyPinDefaultsToModel(UEdGraphPin* Pin, bool bUndo = false);
 
-	/** Check whether we are a property accessor */
-	bool IsPropertyAccessor() const { return GetUnitScriptStruct() == nullptr; }
-
 	/** Get the blueprint that this node is contained within */
 	UControlRigBlueprint* GetBlueprint() const;
 
-	/** Add a new array element to the array referred to by the property path */
-	void HandleAddArrayElement(FString InPropertyPath);
+	/** Get the VM node this is node is wrapping */
+	URigVMNode* GetModelNode() const;
 
-	/** Create Variable Pins on the side */
-	void CreateVariablePins(bool bAlwaysCreatePins = false);
+	/** Get the VM node name this node is wrapping */
+	FName GetModelNodeName() const;
+
+	URigVMPin* GetModelPinFromPinPath(const FString& InPinPath) const;
+
+	/** Add a new array element to the array referred to by the property path */
+	void HandleAddArrayElement(FString InPinPath);
 	
 	/** Clear the array referred to by the property path */
-	void HandleClearArray(FString InPropertyPath);
+	void HandleClearArray(FString InPinPath);
 
 	/** Remove the array element referred to by the property path */
-	void HandleRemoveArrayElement(FString InPropertyPath);
+	void HandleRemoveArrayElement(FString InPinPath);
 
 	/** Insert a new array element after the element referred to by the property path */
-	void HandleInsertArrayElement(FString InPropertyPath);
-
-#if WITH_EDITORONLY_DATA
-	virtual void PostLoad() override;
-	void CacheHierarchyRefConnectionsOnPostLoad();
-#endif
+	void HandleInsertArrayElement(FString InPinPath);
 
 protected:
-	/** Rebuild the cached info about our inputs/outputs */
-	void CacheVariableInfo();
-
-	/** Helper function for AllocateDefaultPins */
-	void CreateExecutionPins(bool bAlwaysCreatePins);
 
-	/** Helper function for AllocateDefaultPins */
-	void CreateInputPins(bool bAlwaysCreatePins);
-	void CreateInputPins_Recursive(const TSharedPtr<FControlRigField>& InputInfo, bool bAlwaysCreatePins);
+	FLinearColor GetNodeOpacityColor() const;
 
 	/** Helper function for AllocateDefaultPins */
-	void CreateInputOutputPins(bool bAlwaysCreatePins);
-	void CreateInputOutputPins_Recursive(const TSharedPtr<FControlRigField>& InputOutputInfo, bool bAlwaysCreatePins);
-
-	/** Helper function for AllocateDefaultPins */
-	void CreateOutputPins(bool bAlwaysCreatePins);
-	void CreateOutputPins_Recursive(const TSharedPtr<FControlRigField>& OutputInfo, bool bAlwaysCreatePins);
-
-	/** Get the generated ControlRig class */
-	UClass* GetControlRigGeneratedClass() const;
-
-	/** Get the skeleton generated ControlRig class */
-	UClass* GetControlRigSkeletonGeneratedClass() const;
-
-	/** Create a ControlRig field from a field on the ControlRig class, if possible */
-	TSharedPtr<FControlRigField> CreateControlRigField(const FControlRigModelPin* InPin, const FString& InPinPath, int32 InArrayIndex = INDEX_NONE) const;
-
-	/** Get all fields that act as execution pins for this node */
-	void GetExecutionFields(TArray<TSharedRef<FControlRigField>>& OutFields) const;
-
-	/** Get all fields that act as inputs for this node */
-	void GetInputFields(TArray<TSharedRef<FControlRigField>>& OutFields) const;
-
-	/** Get all fields that act as outputs for this node */
-	void GetOutputFields(TArray<TSharedRef<FControlRigField>>& OutFields) const;
-
-	/** Get all fields that act as input-outputs for this node */
-	void GetInputOutputFields(TArray<TSharedRef<FControlRigField>>& OutFields) const;
-
-	/** Helper function for GetInputFields/GetOutputFields */
-	void GetFields(TFunction<bool(const FControlRigModelPin*, const FControlRigModelNode*)> InPinCheckFunction, TArray<TSharedRef<FControlRigField>>& OutFields) const;
-
-	/** Get the struct property for the unit we represent, if any (we could just be a property accessor) */
-	FStructProperty* GetUnitProperty() const;
-
-	/** Get the script struct for the unit we represent, if any (we could just be a property accessor) */
-	UScriptStruct* GetUnitScriptStruct() const;
-
-	/** Get the property for the unit we represent */
-	FProperty* GetProperty() const;
+	void CreateExecutionPins();
+	void CreateInputPins(URigVMPin* InParentPin = nullptr);
+	void CreateInputOutputPins(URigVMPin* InParentPin = nullptr);
+	void CreateOutputPins(URigVMPin* InParentPin = nullptr);
 
 	/** Copies default values from underlying properties into pin defaults, for editing */
-	void SetupPinDefaultsFromCDO(UEdGraphPin* Pin);
+	void SetupPinDefaultsFromModel(UEdGraphPin* Pin);
 
 	/** Recreate pins when we reconstruct this node */
 	virtual void ReallocatePinsDuringReconstruction(const TArray<UEdGraphPin*>& OldPins);
@@ -312,14 +160,31 @@ protected:
 	/** Sets the body + title color from a color provided by the model */
 	void SetColorFromModel(const FLinearColor& InColor);
 
-private:
+	UClass* GetControlRigGeneratedClass() const;
+	UClass* GetControlRigSkeletonGeneratedClass() const;
 
-	bool IsVariable() const;
+	static FEdGraphPinType GetPinTypeForModelPin(URigVMPin* InModelPin);
 
-#if WITH_EDITORONLY_DATA
-	TArray<UEdGraphNode*> HierarchyRefOutputConnections;
-#endif
+private:
 
 	FLinearColor CachedTitleColor;
 	FLinearColor CachedNodeColor;
+
+	TArray<URigVMPin*> ExecutePins;
+	TArray<URigVMPin*> InputOutputPins;
+	TArray<URigVMPin*> InputPins;
+	TArray<URigVMPin*> OutputPins;
+	
+	struct PinPair
+	{
+		PinPair()
+			: InputPin(nullptr)
+			, OutputPin(nullptr)
+		{}
+		UEdGraphPin* InputPin;
+		UEdGraphPin* OutputPin;
+	};
+	TMap<URigVMPin*, PinPair> CachedPins;
+
+	friend class SControlRigGraphNode;
 };
diff --git a/Engine/Plugins/Experimental/ControlRig/Source/ControlRigDeveloper/Public/Graph/ControlRigGraphSchema.h b/Engine/Plugins/Experimental/ControlRig/Source/ControlRigDeveloper/Public/Graph/ControlRigGraphSchema.h
--- a/Engine/Plugins/Experimental/ControlRig/Source/ControlRigDeveloper/Public/Graph/ControlRigGraphSchema.h
+++ b/Engine/Plugins/Experimental/ControlRig/Source/ControlRigDeveloper/Public/Graph/ControlRigGraphSchema.h
@@ -67,6 +67,12 @@ public:
 	virtual bool ShouldAlwaysPurgeOnModification() const override { return false; }
 	virtual bool ArePinsCompatible(const UEdGraphPin* PinA, const UEdGraphPin* PinB, const UClass* CallingContext, bool bIgnoreArray /*= false*/) const override;
 	virtual bool DoesSupportPinWatching() const	override { return true; }
+	virtual bool IsPinBeingWatched(UEdGraphPin const* Pin) const override;
+	virtual void ClearPinWatch(UEdGraphPin const* Pin) const override;
+	virtual void OnPinConnectionDoubleCicked(UEdGraphPin* PinA, UEdGraphPin* PinB, const FVector2D& GraphPosition) const override;
+	virtual bool MarkBlueprintDirtyFromNewNode(UBlueprint* InBlueprint, UEdGraphNode* InEdGraphNode) const override;
+
+	virtual bool IsStructEditable(UStruct* InStruct) const;
 
 	/** Create a graph node for a rig */
 	UControlRigGraphNode* CreateGraphNode(UControlRigGraph* InGraph, const FName& InPropertyName) const;
@@ -82,5 +88,12 @@ public:
 
 	/** Returns all of the applicable pin types for variables within a control rig */
 	virtual void GetVariablePinTypes(TArray<FEdGraphPinType>& PinTypes) const;
+
+private:
+
+	const UEdGraphPin* LastPinForCompatibleCheck = nullptr;
+	bool bLastPinWasInput;
+
+	friend class UControlRigRerouteNodeSpawner;
 };
 
diff --git a/Engine/Plugins/Experimental/ControlRig/Source/ControlRigEditor/ControlRigEditor.Build.cs b/Engine/Plugins/Experimental/ControlRig/Source/ControlRigEditor/ControlRigEditor.Build.cs
--- a/Engine/Plugins/Experimental/ControlRig/Source/ControlRigEditor/ControlRigEditor.Build.cs
+++ b/Engine/Plugins/Experimental/ControlRig/Source/ControlRigEditor/ControlRigEditor.Build.cs
@@ -6,11 +6,23 @@ namespace UnrealBuildTool.Rules
     {
         public ControlRigEditor(ReadOnlyTargetRules Target) : base(Target)
         {
-            PrivateIncludePaths.Add("ControlRigEditor/Private");
-            PrivateIncludePaths.Add("ControlRigEditor/Private/Sequencer");
-            PrivateIncludePaths.Add("ControlRigEditor/Private/EditMode");
-            PrivateIncludePaths.Add("ControlRigEditor/Private/Graph");
-            PrivateIncludePaths.Add("ControlRigEditor/Private/Editor");
+            PrivateIncludePaths.AddRange(
+                new string[] {
+                    "ControlRigEditor/Private",
+                    "ControlRigEditor/Private/Sequencer",
+                    "ControlRigEditor/Private/EditMode",
+                    "ControlRigEditor/Private/Graph",
+                    "ControlRigEditor/Private/Editor",
+                    "../../../../Source/Editor/UnrealEd/Private" //compatibility for FBX importer
+				}
+            );
+
+            PublicDependencyModuleNames.AddRange(
+				new string[]
+				{
+					"CurveEditor",
+				}
+	        );
 
             PrivateDependencyModuleNames.AddRange(
                 new string[]
@@ -37,6 +49,7 @@ namespace UnrealBuildTool.Rules
                     "MovieSceneTracks",
                     "MovieSceneTools",
                     "Sequencer",
+					"LevelSequenceEditor",
                     "ClassViewer",
                     "AssetTools",
                     "ContentBrowser",
@@ -53,9 +66,13 @@ namespace UnrealBuildTool.Rules
 					"Json",
 					"DesktopPlatform",
 					"ToolMenus",
-					"ControlRigManipulation"
+					"ControlRigManipulation",
+                    "RigVM",
+                    "RigVMDeveloper"
                 }
             );
+
+            AddEngineThirdPartyPrivateStaticDependencies(Target, "FBX");
         }
     }
 }
diff --git a/Engine/Plugins/Experimental/ControlRig/Source/ControlRigEditor/Private/ControlRigBlueprintActions.cpp b/Engine/Plugins/Experimental/ControlRig/Source/ControlRigEditor/Private/ControlRigBlueprintActions.cpp
--- a/Engine/Plugins/Experimental/ControlRig/Source/ControlRigEditor/Private/ControlRigBlueprintActions.cpp
+++ b/Engine/Plugins/Experimental/ControlRig/Source/ControlRigEditor/Private/ControlRigBlueprintActions.cpp
@@ -3,6 +3,7 @@
 #include "ControlRigBlueprintActions.h"
 #include "ControlRigBlueprintFactory.h"
 #include "ControlRigBlueprint.h"
+#include "ControlRigEditorStyle.h"
 #include "IControlRigEditorModule.h"
 
 #include "Styling/SlateIconFinder.h"
@@ -11,8 +12,28 @@
 #include "EditorStyleSet.h"
 #include "Subsystems/AssetEditorSubsystem.h"
 
+#include "ToolMenus.h"
+#include "ContentBrowserMenuContexts.h"
+#include "AssetTools/Private/AssetTools.h"
+#include "Engine/SkeletalMesh.h"
+#include "Animation/SkeletalMeshActor.h"
+#include "Components/SkeletalMeshComponent.h"
+#include "ScopedTransaction.h"
+#include "LevelSequenceActor.h"
+#include "EngineUtils.h"
+#include "ISequencer.h"
+#include "ILevelSequenceEditorToolkit.h"
+#include "Sequencer/MovieSceneControlRigParameterTrack.h"
+#include "ControlRigSkeletalMeshBinding.h"
+#include "EditMode/ControlRigEditMode.h"
+#include "Editor.h"
+#include "EditorModeManager.h"
+#include "MovieSceneToolsProjectSettings.h"
+
 #define LOCTEXT_NAMESPACE "ControlRigBlueprintActions"
 
+FDelegateHandle FControlRigBlueprintActions::OnSpawnedSkeletalMeshActorChangedHandle;
+
 UFactory* FControlRigBlueprintActions::GetFactoryForBlueprintType(UBlueprint* InBlueprint) const
 {
 	UControlRigBlueprintFactory* ControlRigBlueprintFactory = NewObject<UControlRigBlueprintFactory>();
@@ -59,4 +80,249 @@ TSharedPtr<SWidget> FControlRigBlueprintActions::GetThumbnailOverlay(const FAsse
 		];
 }
 
+void FControlRigBlueprintActions::ExtendSketalMeshToolMenu()
+{
+	TArray<UToolMenu*> MenusToExtend;
+	MenusToExtend.Add(UToolMenus::Get()->ExtendMenu("ContentBrowser.AssetContextMenu.SkeletalMesh"));
+	MenusToExtend.Add(UToolMenus::Get()->ExtendMenu("ContentBrowser.AssetContextMenu.Skeleton"));
+
+	for(UToolMenu* Menu : MenusToExtend)
+	{
+		if (Menu == nullptr)
+		{
+			continue;
+		}
+
+		FToolMenuSection& Section = Menu->FindOrAddSection("GetAssetActions");
+		Section.AddDynamicEntry("GetActions", FNewToolMenuSectionDelegate::CreateLambda([](FToolMenuSection& InSection)
+		{
+			UContentBrowserAssetContextMenuContext* Context = InSection.FindContext<UContentBrowserAssetContextMenuContext>();
+			if (Context)
+			{
+				TArray<UObject*> SelectedObjects = Context->GetSelectedObjects();
+				if (SelectedObjects.Num() > 0)
+				{
+					InSection.AddMenuEntry(
+						"CreateControlRig",
+						LOCTEXT("CreateControlRig", "Create Control Rig"),
+						LOCTEXT("CreateControlRig_ToolTip", "Creates a control rig and preconfigures it for this asset"),
+						FSlateIcon(FControlRigEditorStyle::Get().GetStyleSetName(), "ControlRig", "ControlRig.RigUnit"),
+						FExecuteAction::CreateLambda([SelectedObjects]()
+						{
+							for (UObject* SelectedObject : SelectedObjects)
+							{
+								FControlRigBlueprintActions::ExecuteCreateControlRigFromSkeletalMesh(SelectedObject);
+							}
+						})
+					);
+				}
+			}
+		}));
+	}
+}
+
+void FControlRigBlueprintActions::ExecuteCreateControlRigFromSkeletalMesh(UObject* InSelectedObject)
+{
+	FScopedTransaction Transaction(LOCTEXT("CreateControlRigFromSkeletalMesh", "Create Control Rig"));
+
+	USkeletalMesh* SkeletalMesh = Cast<USkeletalMesh>(InSelectedObject);
+	USkeleton* Skeleton = Cast<USkeleton>(InSelectedObject);
+
+	if(SkeletalMesh)
+	{
+		Skeleton = SkeletalMesh->Skeleton;
+	}
+
+	check(Skeleton);
+
+	UControlRigBlueprintFactory* Factory = NewObject<UControlRigBlueprintFactory>();
+	Factory->ParentClass = UControlRig::StaticClass();
+
+	FString PackagePath = InSelectedObject->GetPathName();
+	FString ControlRigName = FString::Printf(TEXT("%s_CtrlRig"), *InSelectedObject->GetName());
+
+	int32 LastSlashPos = INDEX_NONE;
+	if (PackagePath.FindLastChar('/', LastSlashPos))
+	{
+		PackagePath = PackagePath.Left(LastSlashPos);
+	}
+
+	FAssetToolsModule& AssetToolsModule = FModuleManager::LoadModuleChecked<FAssetToolsModule>("AssetTools");
+
+	FString UniquePackageName;
+	FString UniqueAssetName;
+	AssetToolsModule.Get().CreateUniqueAssetName(PackagePath / ControlRigName, TEXT(""), UniquePackageName, UniqueAssetName);
+
+	UObject* NewAsset = AssetToolsModule.Get().CreateAsset(*UniqueAssetName, *PackagePath, nullptr, Factory);
+
+	UControlRigBlueprint* NewControlRigBlueprint = Cast<UControlRigBlueprint>(NewAsset);
+	check(NewControlRigBlueprint);
+
+	NewControlRigBlueprint->HierarchyContainer.BoneHierarchy.ImportSkeleton(Skeleton->GetReferenceSkeleton(), NAME_None, true, true, false, false);
+	NewControlRigBlueprint->HierarchyContainer.CurveContainer.ImportCurvesFromSkeleton(Skeleton, NAME_None, true, false, false);
+	NewControlRigBlueprint->SourceCurveImport = Skeleton;
+
+	if(SkeletalMesh)
+	{
+		NewControlRigBlueprint->SetPreviewMesh(SkeletalMesh);
+	}
+}
+
+USkeletalMesh* FControlRigBlueprintActions::GetSkeletalMeshFromControlRigBlueprint(UObject* InAsset)
+{
+	if (UControlRigBlueprint* Blueprint = Cast<UControlRigBlueprint>(InAsset))
+	{
+		return Blueprint->GetPreviewMesh();
+	}
+	return nullptr;
+}
+
+void FControlRigBlueprintActions::PostSpawningSkeletalMeshActor(AActor* InSpawnedActor, UObject* InAsset)
+{
+	if (InSpawnedActor->HasAnyFlags(RF_Transient))
+	{
+		return;
+	}
+	
+	OnSpawnedSkeletalMeshActorChangedHandle = FCoreUObjectDelegates::OnObjectPropertyChanged.AddStatic(&FControlRigBlueprintActions::OnSpawnedSkeletalMeshActorChanged, InAsset);
+}
+
+void FControlRigBlueprintActions::OnSpawnedSkeletalMeshActorChanged(UObject* InObject, FPropertyChangedEvent& InEvent, UObject* InAsset)
+{
+	if (!OnSpawnedSkeletalMeshActorChangedHandle.IsValid())
+	{
+		return;
+	}
+
+	// we are waiting for the top level property change event
+	// after the spawn.
+	if (InEvent.Property != nullptr)
+	{
+		return;
+	}
+
+	ASkeletalMeshActor* MeshActor = Cast<ASkeletalMeshActor>(InObject);
+	check(MeshActor);
+
+	FCoreUObjectDelegates::OnObjectPropertyChanged.Remove(OnSpawnedSkeletalMeshActorChangedHandle);
+
+	UControlRigBlueprint* RigBlueprint = Cast<UControlRigBlueprint>(InAsset);
+	if (RigBlueprint == nullptr)
+	{
+		return;
+	}
+	UClass* ControlRigClass = RigBlueprint->GeneratedClass;
+
+	// find a level sequence in the world
+	ALevelSequenceActor* LevelSequenceActor = nullptr;
+	for (ALevelSequenceActor* ExistingLevelSequenceActor : TActorRange<ALevelSequenceActor>(MeshActor->GetWorld()))
+	{
+		LevelSequenceActor = ExistingLevelSequenceActor;
+		break;
+	}
+    //The creation of the FSequencer that happens below is not transactional so we need to make sure we doen't transaction there.
+	GEditor->CancelTransaction(0);
+
+	// todo: show a dialog for picking an existing one or creating a new one
+	if (LevelSequenceActor == nullptr)
+	{
+		FString SequenceName = FString::Printf(TEXT("%s_Take1"), *InAsset->GetName());
+		FString PackagePath = TEXT("/Game");
+
+		FAssetToolsModule& AssetToolsModule = FModuleManager::LoadModuleChecked<FAssetToolsModule>("AssetTools");
+		FString UniquePackageName;
+		FString UniqueAssetName;
+		AssetToolsModule.Get().CreateUniqueAssetName(PackagePath / SequenceName, TEXT(""), UniquePackageName, UniqueAssetName);
+
+		UPackage* Package = CreatePackage(nullptr, *UniquePackageName);
+		ULevelSequence* Sequence = NewObject<ULevelSequence>(Package, *UniqueAssetName, RF_Public | RF_Standalone);
+		Sequence->Initialize(); //creates movie scene
+		Sequence->MarkPackageDirty();
+
+		// Set up some sensible defaults
+		const UMovieSceneToolsProjectSettings* ProjectSettings = GetDefault<UMovieSceneToolsProjectSettings>();
+		FFrameRate TickResolution = Sequence->GetMovieScene()->GetTickResolution();
+		Sequence->GetMovieScene()->SetPlaybackRange((ProjectSettings->DefaultStartTime*TickResolution).FloorToFrame(), (ProjectSettings->DefaultDuration*TickResolution).FloorToFrame().Value);
+
+		UActorFactory* ActorFactory = GEditor->FindActorFactoryForActorClass(ALevelSequenceActor::StaticClass());
+		if (!ensure(ActorFactory))
+		{
+			return;
+		}
+
+		AActor* NewActor = GEditor->UseActorFactory(ActorFactory, FAssetData(Sequence), &FTransform::Identity);
+		if (NewActor == nullptr)
+		{
+			return;
+		}
+
+		LevelSequenceActor = CastChecked<ALevelSequenceActor>(NewActor);
+		LevelSequenceActor->SetSequence(Sequence);
+	}
+
+	ULevelSequence* Sequence = LevelSequenceActor->GetSequence();
+	UMovieScene* MovieScene = Sequence->GetMovieScene();
+
+	GEditor->GetEditorSubsystem<UAssetEditorSubsystem>()->OpenEditorForAsset(Sequence);
+
+	IAssetEditorInstance* AssetEditor = GEditor->GetEditorSubsystem<UAssetEditorSubsystem>()->FindEditorForAsset(Sequence, false);
+	ILevelSequenceEditorToolkit* LevelSequenceEditor = static_cast<ILevelSequenceEditorToolkit*>(AssetEditor);
+	TWeakPtr<ISequencer> WeakSequencer = LevelSequenceEditor ? LevelSequenceEditor->GetSequencer() : nullptr;
+
+	if (WeakSequencer.IsValid())
+	{
+		TArray<TWeakObjectPtr<AActor> > ActorsToAdd;
+		ActorsToAdd.Add(MeshActor);
+		TArray<FGuid> ActorTracks = WeakSequencer.Pin()->AddActors(ActorsToAdd, false);
+		for (FGuid ActorTrackGuid : ActorTracks)
+		{
+			UMovieSceneControlRigParameterTrack* Track = MovieScene->AddTrack<UMovieSceneControlRigParameterTrack>(ActorTrackGuid);
+			if (Track)
+			{
+				USkeletalMesh* SkeletalMesh = MeshActor->GetSkeletalMeshComponent()->SkeletalMesh;
+				USkeleton* Skeleton = SkeletalMesh->Skeleton;
+				
+				FString ObjectName = (ControlRigClass->GetName());
+				ObjectName.RemoveFromEnd(TEXT("_C"));
+
+				UControlRig* ControlRig = NewObject<UControlRig>(Track, ControlRigClass, FName(*ObjectName), RF_Transactional);
+				ControlRig->SetObjectBinding(MakeShared<FControlRigSkeletalMeshBinding>());
+				ControlRig->GetObjectBinding()->BindToObject(SkeletalMesh);
+				ControlRig->GetDataSourceRegistry()->RegisterDataSource(UControlRig::OwnerComponent, ControlRig->GetObjectBinding()->GetBoundObject());
+				ControlRig->Initialize();
+				ControlRig->Execute(EControlRigState::Update);
+				ControlRig->CreateRigControlsForCurveContainer();
+
+				WeakSequencer.Pin()->NotifyMovieSceneDataChanged(EMovieSceneDataChangeType::MovieSceneStructureItemsChanged);
+
+				Track->Modify();
+				UMovieSceneSection* NewSection = Track->CreateControlRigSection(0, ControlRig);
+				//mz todo need to have multiple rigs with same class
+				Track->SetTrackName(FName(*ObjectName));
+				Track->SetDisplayName(FText::FromString(ObjectName));
+
+				WeakSequencer.Pin()->EmptySelection();
+				WeakSequencer.Pin()->SelectSection(NewSection);
+				WeakSequencer.Pin()->ThrobSectionSelection();
+				WeakSequencer.Pin()->ObjectImplicitlyAdded(ControlRig);
+				FText Name = LOCTEXT("SequenceTrackFilter_ControlRigControls", "Control Rig Controls");
+				WeakSequencer.Pin()->SetFilterOn(Name, true);
+				WeakSequencer.Pin()->NotifyMovieSceneDataChanged(EMovieSceneDataChangeType::MovieSceneStructureItemAdded);
+				FControlRigEditMode* ControlRigEditMode = static_cast<FControlRigEditMode*>(GLevelEditorModeTools().GetActiveMode(FControlRigEditMode::ModeName));
+				if (!ControlRigEditMode)
+				{
+					GLevelEditorModeTools().ActivateMode(FControlRigEditMode::ModeName);
+					ControlRigEditMode = static_cast<FControlRigEditMode*>(GLevelEditorModeTools().GetActiveMode(FControlRigEditMode::ModeName));
+
+				}
+				if (ControlRigEditMode)
+				{
+					ControlRigEditMode->SetObjects(ControlRig, FGuid(), nullptr, WeakSequencer.Pin());
+				}
+			}
+		}
+	}
+}
+
 #undef LOCTEXT_NAMESPACE
+
diff --git a/Engine/Plugins/Experimental/ControlRig/Source/ControlRigEditor/Private/ControlRigBlueprintActions.h b/Engine/Plugins/Experimental/ControlRig/Source/ControlRigEditor/Private/ControlRigBlueprintActions.h
--- a/Engine/Plugins/Experimental/ControlRig/Source/ControlRigEditor/Private/ControlRigBlueprintActions.h
+++ b/Engine/Plugins/Experimental/ControlRig/Source/ControlRigEditor/Private/ControlRigBlueprintActions.h
@@ -9,6 +9,8 @@
 
 class FMenuBuilder;
 class UFactory;
+class USkeletalMesh;
+class AActor;
 
 class FControlRigBlueprintActions : public FAssetTypeActions_Blueprint
 {
@@ -23,4 +25,13 @@ public:
 
 	// FAssetTypeActions_Blueprint interface
 	virtual UFactory* GetFactoryForBlueprintType(UBlueprint* InBlueprint) const override;
+
+	static void ExtendSketalMeshToolMenu();
+	static void ExecuteCreateControlRigFromSkeletalMesh(UObject* InSelectedObject);
+
+	static USkeletalMesh* GetSkeletalMeshFromControlRigBlueprint(UObject* InAsset);
+	static void PostSpawningSkeletalMeshActor(AActor* InSpawnedActor, UObject* InAsset);
+	static void OnSpawnedSkeletalMeshActorChanged(UObject* InObject, FPropertyChangedEvent& InEvent, UObject* InAsset);
+
+	static FDelegateHandle OnSpawnedSkeletalMeshActorChangedHandle;
 };
diff --git a/Engine/Plugins/Experimental/ControlRig/Source/ControlRigEditor/Private/ControlRigBlueprintCommands.cpp b/Engine/Plugins/Experimental/ControlRig/Source/ControlRigEditor/Private/ControlRigBlueprintCommands.cpp
--- a/Engine/Plugins/Experimental/ControlRig/Source/ControlRigEditor/Private/ControlRigBlueprintCommands.cpp
+++ b/Engine/Plugins/Experimental/ControlRig/Source/ControlRigEditor/Private/ControlRigBlueprintCommands.cpp
@@ -8,6 +8,7 @@ void FControlRigBlueprintCommands::RegisterCommands()
 {
 	UI_COMMAND(DeleteItem, "Delete", "Deletes the selected items and removes their nodes from the graph.", EUserInterfaceActionType::Button, FInputChord(EKeys::Delete));
 	UI_COMMAND(ExecuteGraph, "Execute", "Execute the rig graph if On.", EUserInterfaceActionType::ToggleButton, FInputChord());
+	UI_COMMAND(AutoCompileGraph, "AutoCompile", "Auto-compile the rig graph if On.", EUserInterfaceActionType::ToggleButton, FInputChord());
 }
 
 #undef LOCTEXT_NAMESPACE
diff --git a/Engine/Plugins/Experimental/ControlRig/Source/ControlRigEditor/Private/ControlRigBlueprintCommands.h b/Engine/Plugins/Experimental/ControlRig/Source/ControlRigEditor/Private/ControlRigBlueprintCommands.h
--- a/Engine/Plugins/Experimental/ControlRig/Source/ControlRigEditor/Private/ControlRigBlueprintCommands.h
+++ b/Engine/Plugins/Experimental/ControlRig/Source/ControlRigEditor/Private/ControlRigBlueprintCommands.h
@@ -23,6 +23,10 @@ public:
 
 	/** Toggle Execute the Graph */
 	TSharedPtr< FUICommandInfo > ExecuteGraph;
+
+	/** Toggle Auto Compilation in the Graph */
+	TSharedPtr< FUICommandInfo > AutoCompileGraph;
+
 	/**
 	 * Initialize commands
 	 */
diff --git a/Engine/Plugins/Experimental/ControlRig/Source/ControlRigEditor/Private/ControlRigBlueprintFactory.cpp b/Engine/Plugins/Experimental/ControlRig/Source/ControlRigEditor/Private/ControlRigBlueprintFactory.cpp
--- a/Engine/Plugins/Experimental/ControlRig/Source/ControlRigEditor/Private/ControlRigBlueprintFactory.cpp
+++ b/Engine/Plugins/Experimental/ControlRig/Source/ControlRigEditor/Private/ControlRigBlueprintFactory.cpp
@@ -134,7 +134,19 @@ private:
 
 				// in the future we might allow it to parent to BP classes, but right now, it won't work well because of multi rig graph
 				// for now we disable it and only allow native class. 
-				return (!InClass->HasAnyClassFlags(CLASS_Deprecated) && InClass->HasAnyClassFlags(CLASS_Native) && InClass->GetOutermost() != GetTransientPackage());
+				if (!InClass->HasAnyClassFlags(CLASS_Deprecated) && InClass->HasAnyClassFlags(CLASS_Native) && InClass->GetOutermost() != GetTransientPackage())
+				{
+					// see if it can be blueprint base
+					const FString BlueprintBaseMetaKey(TEXT("IsBlueprintBase"));
+
+					if (InClass->HasMetaData(*BlueprintBaseMetaKey))
+					{
+						if (InClass->GetMetaData(*BlueprintBaseMetaKey) == TEXT("true"))
+						{
+							return true;
+						}
+					}
+				}
 			}
 
 			return false;
diff --git a/Engine/Plugins/Experimental/ControlRig/Source/ControlRigEditor/Private/ControlRigEditorModule.cpp b/Engine/Plugins/Experimental/ControlRig/Source/ControlRigEditor/Private/ControlRigEditorModule.cpp
--- a/Engine/Plugins/Experimental/ControlRig/Source/ControlRigEditor/Private/ControlRigEditorModule.cpp
+++ b/Engine/Plugins/Experimental/ControlRig/Source/ControlRigEditor/Private/ControlRigEditorModule.cpp
@@ -18,7 +18,6 @@
 #include "Framework/Application/SlateApplication.h"
 #include "LevelEditor.h"
 #include "MovieSceneToolsProjectSettings.h"
-#include "PropertyEditorModule.h"
 #include "Styling/SlateStyle.h"
 #include "WorkflowOrientedApp/WorkflowTabManager.h"
 #include "Modules/ModuleManager.h"
@@ -47,10 +46,10 @@
 #include "ControlRigGizmoLibraryActions.h"
 #include "Graph/ControlRigGraphSchema.h"
 #include "Graph/ControlRigGraph.h"
-#include "Graph/NodeSpawners/ControlRigPropertyNodeSpawner.h"
 #include "Graph/NodeSpawners/ControlRigUnitNodeSpawner.h"
 #include "Graph/NodeSpawners/ControlRigVariableNodeSpawner.h"
-#include "Graph/NodeSpawners/ControlRigCommentNodeSpawner.h"
+#include "Graph/NodeSpawners/ControlRigParameterNodeSpawner.h"
+#include "Graph/NodeSpawners/ControlRigRerouteNodeSpawner.h"
 #include "Kismet2/BlueprintEditorUtils.h"
 #include "Kismet2/KismetDebugUtilities.h"
 #include "Graph/ControlRigGraphNode.h"
@@ -63,22 +62,26 @@
 #include "ControlRigStackCommands.h"
 #include "Animation/AnimSequence.h"
 #include "ControlRigEditorEditMode.h"
-#include "ControlRigDetails.h"
 #include "ControlRigElementDetails.h"
-#include "Units/Deprecated/RigUnitEditor_TwoBoneIKFK.h"
+#include "ControlRigCompilerDetails.h"
+#include "ControlRigDrawingDetails.h"
 #include "Animation/AnimSequence.h"
 #include "Editor/SControlRigProfilingView.h"
 #include "WorkspaceMenuStructure.h"
 #include "WorkspaceMenuStructureModule.h"
 #include "Subsystems/AssetEditorSubsystem.h"
 #include "ControlRigParameterTrackEditor.h"
+#include "ActorFactories/ActorFactorySkeletalMesh.h"
+#include "ControlRigThumbnailRenderer.h"
+#include "RigVMModel/RigVMController.h"
+#include "ControlRigBlueprint.h"
+#include "ControlRig/Private/Units/Simulation/RigUnit_AlphaInterp.h"
+#include "ControlRig/Private/Units/Debug/RigUnit_VisualDebug.h"
 
 #define LOCTEXT_NAMESPACE "ControlRigEditorModule"
 
 DEFINE_LOG_CATEGORY(LogControlRigEditor);
 
-TMap<FName, TSubclassOf<URigUnitEditor_Base>> FControlRigEditorModule::RigUnitEditorClasses;
-
 TSharedRef<SDockTab> SpawnRigProfiler( const FSpawnTabArgs& Args )
 {
 	return SNew(SDockTab)
@@ -126,11 +129,16 @@ void FControlRigEditorModule::StartupModule()
 	PropertyEditorModule.RegisterCustomClassLayout(ClassesToUnregisterOnShutdown.Last(), FOnGetDetailCustomizationInstance::CreateStatic(&FRigSpaceDetails::MakeInstance));
 
 	ClassesToUnregisterOnShutdown.Add(UControlRig::StaticClass()->GetFName());
-	PropertyEditorModule.RegisterCustomClassLayout(ClassesToUnregisterOnShutdown.Last(), FOnGetDetailCustomizationInstance::CreateStatic(&FControlRigDetails::MakeInstance));
 
 	// same as ClassesToUnregisterOnShutdown but for properties, there is none right now
 	PropertiesToUnregisterOnShutdown.Reset();
 
+	PropertiesToUnregisterOnShutdown.Add(FRigVMCompileSettings::StaticStruct()->GetFName());
+	PropertyEditorModule.RegisterCustomPropertyTypeLayout(PropertiesToUnregisterOnShutdown.Last(), FOnGetPropertyTypeCustomizationInstance::CreateStatic(&FRigVMCompileSettingsDetails::MakeInstance));
+
+	PropertiesToUnregisterOnShutdown.Add(FControlRigDrawContainer::StaticStruct()->GetFName());
+	PropertyEditorModule.RegisterCustomPropertyTypeLayout(PropertiesToUnregisterOnShutdown.Last(), FOnGetPropertyTypeCustomizationInstance::CreateStatic(&FControlRigDrawContainerDetails::MakeInstance));
+
 	// Register asset tools
 	auto RegisterAssetTypeAction = [this](const TSharedRef<IAssetTypeActions>& InAssetTypeAction)
 	{
@@ -167,10 +175,6 @@ void FControlRigEditorModule::StartupModule()
 
 	ReconstructAllNodesDelegateHandle = FBlueprintEditorUtils::OnReconstructAllNodesEvent.AddStatic(&FControlRigBlueprintUtils::HandleReconstructAllNodes);
 	RefreshAllNodesDelegateHandle = FBlueprintEditorUtils::OnRefreshAllNodesEvent.AddStatic(&FControlRigBlueprintUtils::HandleRefreshAllNodes);
-	RenameVariableReferencesDelegateHandle = FBlueprintEditorUtils::OnRenameVariableReferencesEvent.AddStatic(&FControlRigBlueprintUtils::HandleRenameVariableReferencesEvent);
-
-	// register rig unit base editor class
-	RegisterRigUnitEditorClass("RigUnit_TwoBoneIKFK", URigUnitEditor_TwoBoneIKFK::StaticClass());
 
 #if WITH_EDITOR
 	if (FSlateApplication::IsInitialized())
@@ -182,6 +186,15 @@ void FControlRigEditorModule::StartupModule()
 			.SetIcon(FSlateIcon(TEXT("ControlRigEditorStyle"), TEXT("ControlRig.RigUnit")));
 	};
 #endif
+
+	FControlRigBlueprintActions::ExtendSketalMeshToolMenu();
+	UActorFactorySkeletalMesh::RegisterDelegatesForAssetClass(
+		UControlRigBlueprint::StaticClass(),
+		FGetSkeletalMeshFromAssetDelegate::CreateStatic(&FControlRigBlueprintActions::GetSkeletalMeshFromControlRigBlueprint),
+		FPostSkeletalMeshActorSpawnedDelegate::CreateStatic(&FControlRigBlueprintActions::PostSpawningSkeletalMeshActor)
+	);
+
+	UThumbnailManager::Get().RegisterCustomRenderer(UControlRigBlueprint::StaticClass(), UControlRigThumbnailRenderer::StaticClass());
 }
 
 void FControlRigEditorModule::ShutdownModule()
@@ -193,9 +206,11 @@ void FControlRigEditorModule::ShutdownModule()
 	}
 #endif
 
+	//UThumbnailManager::Get().UnregisterCustomRenderer(UControlRigBlueprint::StaticClass());
+	//UActorFactorySkeletalMesh::UnregisterDelegatesForAssetClass(UControlRigBlueprint::StaticClass());
+
 	FBlueprintEditorUtils::OnRefreshAllNodesEvent.Remove(RefreshAllNodesDelegateHandle);
 	FBlueprintEditorUtils::OnReconstructAllNodesEvent.Remove(ReconstructAllNodesDelegateHandle);
-	FBlueprintEditorUtils::OnRenameVariableReferencesEvent.Remove(RenameVariableReferencesDelegateHandle);
 
 	FEdGraphUtilities::UnregisterVisualPinFactory(ControlRigGraphPanelPinFactory);
 	FEdGraphUtilities::UnregisterVisualNodeFactory(ControlRigGraphPanelNodeFactory);
@@ -251,6 +266,7 @@ void FControlRigEditorModule::HandleNewBlueprintCreated(UBlueprint* InBlueprint)
 	ControlRigGraph->bAllowDeletion = false;
 	FBlueprintEditorUtils::AddUbergraphPage(InBlueprint, ControlRigGraph);
 	InBlueprint->LastEditedDocuments.AddUnique(ControlRigGraph);
+	InBlueprint->PostLoad();
 }
 
 
@@ -261,17 +277,6 @@ TSharedRef<IControlRigEditor> FControlRigEditorModule::CreateControlRigEditor(co
 	return NewControlRigEditor;
 }
 
-void FControlRigEditorModule::RegisterRigUnitEditorClass(FName RigUnitClassName, TSubclassOf<URigUnitEditor_Base> InClass)
-{
-	TSubclassOf<URigUnitEditor_Base>& Class = RigUnitEditorClasses.FindOrAdd(RigUnitClassName);
-	Class = InClass;
-}
-
-void FControlRigEditorModule::UnregisterRigUnitEditorClass(FName RigUnitClassName)
-{
-	RigUnitEditorClasses.Remove(RigUnitClassName);
-}
-
 void FControlRigEditorModule::GetTypeActions(const UControlRigBlueprint* CRB, FBlueprintActionDatabaseRegistrar& ActionRegistrar)
 {
 	// actions get registered under specific object-keys; the idea is that 
@@ -289,7 +294,7 @@ void FControlRigEditorModule::GetTypeActions(const UControlRigBlueprint* CRB, FB
 	}
 
 	// Add all rig units
-	FControlRigBlueprintUtils::ForAllRigUnits([&](UStruct* InStruct)
+	FControlRigBlueprintUtils::ForAllRigUnits([&](UScriptStruct* InStruct)
 	{
 		FString CategoryMetadata, DisplayNameMetadata, MenuDescSuffixMetadata;
 		InStruct->GetStringMetaDataHierarchical(UControlRig::CategoryMetaName, &CategoryMetadata);
@@ -308,17 +313,13 @@ void FControlRigEditorModule::GetTypeActions(const UControlRigBlueprint* CRB, FB
 		ActionRegistrar.AddBlueprintAction(ActionKey, NodeSpawner);
 	});
 
-	UBlueprintNodeSpawner* CommentNodeSpawner = UControlRigCommentNodeSpawner::Create();
-	check(CommentNodeSpawner != nullptr);
-	ActionRegistrar.AddBlueprintAction(ActionKey, CommentNodeSpawner);
-
 	// Add 'new properties'
 	TArray<FEdGraphPinType> PinTypes;
 	GetDefault<UControlRigGraphSchema>()->GetVariablePinTypes(PinTypes);
 
 	struct Local
 	{
-		static void AddVariableActions_Recursive(UClass* InActionKey, FBlueprintActionDatabaseRegistrar& InActionRegistrar, const FEdGraphPinType& PinType, const FString& InCategory)
+		static void AddVariableActions(UClass* InActionKey, FBlueprintActionDatabaseRegistrar& InActionRegistrar, const FEdGraphPinType& PinType, const FString& InCategory)
 		{
 			static const FString CategoryDelimiter(TEXT("|"));
 
@@ -341,40 +342,67 @@ void FControlRigEditorModule::GetTypeActions(const UControlRigBlueprint* CRB, FB
 			}
 
 
-			UBlueprintNodeSpawner* NodeSpawner = UControlRigVariableNodeSpawner::CreateFromPinType(PinType, MenuDesc, NodeCategory, ToolTip);
-			check(NodeSpawner != nullptr);
-			InActionRegistrar.AddBlueprintAction(InActionKey, NodeSpawner);
+			UBlueprintNodeSpawner* NodeSpawnerGetter = UControlRigVariableNodeSpawner::CreateFromPinType(PinType, true, MenuDesc, NodeCategory, ToolTip);
+			check(NodeSpawnerGetter != nullptr);
+			InActionRegistrar.AddBlueprintAction(InActionKey, NodeSpawnerGetter);
+
+			UBlueprintNodeSpawner* NodeSpawnerSetter = UControlRigVariableNodeSpawner::CreateFromPinType(PinType, false, MenuDesc, NodeCategory, ToolTip);
+			check(NodeSpawnerSetter != nullptr);
+			InActionRegistrar.AddBlueprintAction(InActionKey, NodeSpawnerSetter);
 		}
+
+		/*
+		static void AddParameterActions(UClass* InActionKey, FBlueprintActionDatabaseRegistrar& InActionRegistrar, const FEdGraphPinType& PinType, const FString& InCategory)
+		{
+			static const FString CategoryDelimiter(TEXT("|"));
+
+			FText NodeCategory = FText::FromString(InCategory);
+			FText MenuDesc;
+			FText ToolTip;
+			if (PinType.PinCategory == UEdGraphSchema_K2::PC_Struct)
+			{
+				if (UScriptStruct* Struct = Cast<UScriptStruct>(PinType.PinSubCategoryObject.Get()))
+				{
+					MenuDesc = FText::FromString(Struct->GetName());
+					ToolTip = MenuDesc;
+				}
+
+			}
+			else
+			{
+				MenuDesc = UEdGraphSchema_K2::GetCategoryText(PinType.PinCategory, true);
+				ToolTip = UEdGraphSchema_K2::GetCategoryText(PinType.PinCategory, false);
+			}
+
+
+			UBlueprintNodeSpawner* NodeSpawnerGetter = UControlRigParameterNodeSpawner::CreateFromPinType(PinType, true, MenuDesc, NodeCategory, ToolTip);
+			check(NodeSpawnerGetter != nullptr);
+			InActionRegistrar.AddBlueprintAction(InActionKey, NodeSpawnerGetter);
+
+			UBlueprintNodeSpawner* NodeSpawnerSetter = UControlRigParameterNodeSpawner::CreateFromPinType(PinType, false, MenuDesc, NodeCategory, ToolTip);
+			check(NodeSpawnerSetter != nullptr);
+			InActionRegistrar.AddBlueprintAction(InActionKey, NodeSpawnerSetter);
+		}
+		*/
 	};
 
-	FString CurrentCategory = LOCTEXT("NewVariable", "New Variable").ToString();
+	FString VariableCategory = LOCTEXT("NewVariable", "New Variable").ToString();
+	//FString ParameterCategory = LOCTEXT("NewParameter", "New Parameter").ToString();
 	for (const FEdGraphPinType& PinType: PinTypes)
 	{
-		Local::AddVariableActions_Recursive(ActionKey, ActionRegistrar, PinType, CurrentCategory);
+		Local::AddVariableActions(ActionKey, ActionRegistrar, PinType, VariableCategory);
+		// let's disable parameters in the UI for now.
+		//Local::AddParameterActions(ActionKey, ActionRegistrar, PinType, ParameterCategory);
 	}
-}
 
-void FControlRigEditorModule::GetInstanceActions(const UControlRigBlueprint* CRB, FBlueprintActionDatabaseRegistrar& ActionRegistrar)
-{
-	// actions get registered under specific object-keys; the idea is that 
-	// actions might have to be updated (or deleted) if their object-key is  
-	// mutated (or removed)... here we use the generated class (so if the class 
-	// type disappears, then the action should go with it)
-	UClass* ActionKey = CRB->GeneratedClass;
-	// to keep from needlessly instantiating a UBlueprintNodeSpawner, first   
-	// check to make sure that the registrar is looking for actions of this type
-	// (could be regenerating actions for a specific asset, and therefore the 
-	// registrar would only accept actions corresponding to that asset)
-	if (!ActionRegistrar.IsOpenForRegistration(ActionKey))
-	{
-		return;
-	}
+	// add support for names as parameters
+	//Local::AddParameterActions(ActionKey, ActionRegistrar, FEdGraphPinType(UEdGraphSchema_K2::PC_Name, FName(NAME_None), nullptr, EPinContainerType::None, false, FEdGraphTerminalType()), ParameterCategory);
 
-	for (TFieldIterator<FProperty> PropertyIt(ActionKey, EFieldIteratorFlags::ExcludeSuper); PropertyIt; ++PropertyIt)
-	{
-		UBlueprintNodeSpawner* NodeSpawner = UControlRigPropertyNodeSpawner::CreateFromProperty(UControlRigGraphNode::StaticClass(), *PropertyIt);
-		ActionRegistrar.AddBlueprintAction(ActionKey, NodeSpawner);
-	}
+	UBlueprintNodeSpawner* RerouteNodeSpawner = UControlRigRerouteNodeSpawner::CreateGeneric(
+		LOCTEXT("RerouteSpawnerDesc", "New Reroute Node"),
+		LOCTEXT("RerouteSpawnerCategory", "Organization"), 
+		LOCTEXT("RerouteSpawnerTooltip", "Adds a new reroute node to the graph"));
+	ActionRegistrar.AddBlueprintAction(ActionKey, RerouteNodeSpawner);
 }
 
 FConnectionDrawingPolicy* FControlRigEditorModule::CreateConnectionDrawingPolicy(int32 InBackLayerID, int32 InFrontLayerID, float InZoomFactor, const FSlateRect& InClippingRect, class FSlateWindowElementList& InDrawElements, class UEdGraph* InGraphObj)
@@ -429,7 +457,7 @@ void FControlRigEditorModule::GetContextMenuActions(const UControlRigGraphSchema
 	{
 		Schema->UEdGraphSchema::GetContextMenuActions(Menu, Context);
 
-		if (const UEdGraphPin* InGraphPin = Context->Pin)
+		if (UEdGraphPin* InGraphPin = (UEdGraphPin* )Context->Pin)
 		{
 			{
 				FToolMenuSection& Section = Menu->AddSection("EdGraphSchemaPinActions", LOCTEXT("PinActionsMenuHeader", "Pin Actions"));
@@ -445,8 +473,7 @@ void FControlRigEditorModule::GetContextMenuActions(const UControlRigGraphSchema
 				FToolMenuSection& Section = Menu->AddSection("EdGraphSchemaWatches", LOCTEXT("WatchesHeader", "Watches"));
 				UBlueprint* OwnerBlueprint = FBlueprintEditorUtils::FindBlueprintForGraphChecked(Context->Graph);
 				{
-					const UEdGraphPin* WatchedPin = ((InGraphPin->Direction == EGPD_Input) && (InGraphPin->LinkedTo.Num() > 0)) ? InGraphPin->LinkedTo[0] : InGraphPin;
-					if (FKismetDebugUtilities::IsPinBeingWatched(OwnerBlueprint, WatchedPin))
+					if (FKismetDebugUtilities::IsPinBeingWatched(OwnerBlueprint, InGraphPin))
 					{
 						Section.AddMenuEntry(FGraphEditorCommands::Get().StopWatchingPin);
 					}
@@ -456,21 +483,349 @@ void FControlRigEditorModule::GetContextMenuActions(const UControlRigGraphSchema
 					}
 				}
 			}
+
+			// Add alphainterp menu entries
+			if(UControlRigBlueprint* RigBlueprint = Cast<UControlRigBlueprint>((UBlueprint*)Context->Blueprint))
+			{
+				if (URigVMPin* ModelPin = RigBlueprint->Model->FindPin(InGraphPin->GetName()))
+				{
+
+					if (ModelPin->IsArray())
+					{
+						FToolMenuSection& Section = Menu->AddSection("EdGraphSchemaPinArrays", LOCTEXT("PinArrays", "Arrays"));
+						Section.AddMenuEntry(
+							"ClearPinArray",
+							LOCTEXT("ClearPinArray", "Clear Array"),
+							LOCTEXT("ClearPinArray_Tooltip", "Removes all elements of the array."),
+							FSlateIcon(),
+							FUIAction(FExecuteAction::CreateLambda([RigBlueprint, ModelPin]() {
+								RigBlueprint->Controller->ClearArrayPin(ModelPin->GetPinPath());
+							})
+						));
+					}
+					if(ModelPin->IsArrayElement())
+					{
+						FToolMenuSection& Section = Menu->AddSection("EdGraphSchemaPinArrays", LOCTEXT("PinArrays", "Arrays"));
+						Section.AddMenuEntry(
+							"RemoveArrayPin",
+							LOCTEXT("RemoveArrayPin", "Remove Array Element"),
+							LOCTEXT("RemoveArrayPin_Tooltip", "Removes the selected element from the array"),
+							FSlateIcon(),
+							FUIAction(FExecuteAction::CreateLambda([RigBlueprint, ModelPin]() {
+								RigBlueprint->Controller->RemoveArrayPin(ModelPin->GetPinPath());
+							})
+						));
+						Section.AddMenuEntry(
+							"DuplicateArrayPin",
+							LOCTEXT("DuplicateArrayPin", "Duplicate Array Element"),
+							LOCTEXT("DuplicateArrayPin_Tooltip", "Duplicates the selected element"),
+							FSlateIcon(),
+							FUIAction(FExecuteAction::CreateLambda([RigBlueprint, ModelPin]() {
+								RigBlueprint->Controller->DuplicateArrayPin(ModelPin->GetPinPath());
+							})
+						));
+					}
+
+					if(Cast<URigVMStructNode>(ModelPin->GetNode()))
+					{
+						FToolMenuSection& Section = Menu->AddSection("EdGraphSchemaPinDefaults", LOCTEXT("PinDefaults", "Pin Defaults"));
+						Section.AddMenuEntry(
+							"ResetPinDefaultValue",
+							LOCTEXT("ResetPinDefaultValue", "Reset Pin Value"),
+							LOCTEXT("ResetPinDefaultValue_Tooltip", "Resets the pin's value to its default."),
+							FSlateIcon(),
+							FUIAction(FExecuteAction::CreateLambda([RigBlueprint, ModelPin]() {
+								RigBlueprint->Controller->ResetPinDefaultValue(ModelPin->GetPinPath());
+							})
+						));
+					}
+
+					if ((ModelPin->GetCPPType() == TEXT("FVector") ||
+						 ModelPin->GetCPPType() == TEXT("FQuat") ||
+						 ModelPin->GetCPPType() == TEXT("FTransform")) &&
+						(ModelPin->GetDirection() == ERigVMPinDirection::Input ||
+						 ModelPin->GetDirection() == ERigVMPinDirection::IO) &&
+						 ModelPin->GetPinForLink()->GetRootPin()->GetSourceLinks(true).Num() == 0)
+					{
+						FToolMenuSection& Section = Menu->AddSection("EdGraphSchemaControlPin", LOCTEXT("ControlPin", "Direct Manipulation"));
+						Section.AddMenuEntry(
+							"DirectManipControlPin",
+							LOCTEXT("DirectManipControlPin", "Control Pin Value"),
+							LOCTEXT("DirectManipControlPin_Tooltip", "Configures the pin for direct interaction in the viewport"),
+							FSlateIcon(),
+							FUIAction(FExecuteAction::CreateLambda([RigBlueprint, ModelPin]() {
+								RigBlueprint->AddTransientControl(ModelPin);
+							})
+						));
+					}
+
+					if (ModelPin->GetRootPin() == ModelPin && Cast<URigVMStructNode>(ModelPin->GetNode()) != nullptr)
+					{
+						if (ModelPin->HasInjectedNodes())
+						{
+							FToolMenuSection& Section = Menu->AddSection("EdGraphSchemaNodeEjectionInterp", LOCTEXT("NodeEjectionInterp", "Eject"));
+
+							Section.AddMenuEntry(
+								"EjectLastNode",
+								LOCTEXT("EjectLastNode", "Eject Last Node"),
+								LOCTEXT("EjectLastNode_Tooltip", "Eject the last injected node"),
+								FSlateIcon(),
+								FUIAction(FExecuteAction::CreateLambda([RigBlueprint, ModelPin]() {
+									RigBlueprint->Controller->EjectNodeFromPin(ModelPin->GetPinPath());
+								})
+							));
+						}
+
+						if (ModelPin->GetCPPType() == TEXT("float") ||
+							ModelPin->GetCPPType() == TEXT("FVector"))
+						{
+							FToolMenuSection& Section = Menu->AddSection("EdGraphSchemaNodeInjectionInterp", LOCTEXT("NodeInjectionInterp", "Interpolate"));
+							URigVMNode* InterpNode = nullptr;
+							for (URigVMInjectionInfo* Injection : ModelPin->GetInjectedNodes())
+							{
+								FString PrototypeName;
+								if (Injection->StructNode->GetScriptStruct()->GetStringMetaDataHierarchical(TEXT("PrototypeName"), &PrototypeName))
+								{
+									if (PrototypeName == TEXT("AlphaInterp"))
+									{
+										InterpNode = Injection->StructNode;
+										break;
+									}
+								}
+							}
+
+							if (InterpNode == nullptr)
+							{
+								UScriptStruct* ScriptStruct = nullptr;
+
+								if (ModelPin->GetCPPType() == TEXT("float"))
+								{
+									ScriptStruct = FRigUnit_AlphaInterp::StaticStruct();
+								}
+								else if (ModelPin->GetCPPType() == TEXT("FVector"))
+								{
+									ScriptStruct = FRigUnit_AlphaInterpVector::StaticStruct();
+								}
+								else
+								{
+									checkNoEntry();
+								}
+
+								Section.AddMenuEntry(
+									"AddAlphaInterp",
+									LOCTEXT("AddAlphaInterp", "Add Interpolate"),
+									LOCTEXT("AddAlphaInterp_Tooltip", "Injects an interpolate node"),
+									FSlateIcon(),
+									FUIAction(FExecuteAction::CreateLambda([RigBlueprint, InGraphPin, ModelPin, ScriptStruct]() {
+										URigVMInjectionInfo* Injection = RigBlueprint->Controller->AddInjectedNode(ModelPin->GetPinPath(), ModelPin->GetDirection() != ERigVMPinDirection::Output, ScriptStruct, TEXT("Execute"), TEXT("Value"), TEXT("Result"));
+										if (Injection)
+										{
+											TArray<FName> NodeNames;
+											NodeNames.Add(Injection->StructNode->GetFName());
+											RigBlueprint->Controller->SetNodeSelection(NodeNames);
+										}
+									})
+								));
+							}
+							else
+							{
+								Section.AddMenuEntry(
+									"EditAlphaInterp",
+									LOCTEXT("EditAlphaInterp", "Edit Interpolate"),
+									LOCTEXT("EditAlphaInterp_Tooltip", "Edit the interpolate node"),
+									FSlateIcon(),
+									FUIAction(FExecuteAction::CreateLambda([RigBlueprint, InterpNode]() {
+									TArray<FName> NodeNames;
+									NodeNames.Add(InterpNode->GetFName());
+									RigBlueprint->Controller->SetNodeSelection(NodeNames);
+								})
+									));
+								Section.AddMenuEntry(
+									"RemoveAlphaInterp",
+									LOCTEXT("RemoveAlphaInterp", "Remove Interpolate"),
+									LOCTEXT("RemoveAlphaInterp_Tooltip", "Removes the interpolate node"),
+									FSlateIcon(),
+									FUIAction(FExecuteAction::CreateLambda([RigBlueprint, InGraphPin, ModelPin, InterpNode]() {
+										RigBlueprint->Controller->RemoveNodeByName(InterpNode->GetFName());
+									})
+								));
+							}
+						}
+
+						if (ModelPin->GetCPPType() == TEXT("FVector") ||
+							ModelPin->GetCPPType() == TEXT("FQuat") ||
+							ModelPin->GetCPPType() == TEXT("FTransform"))
+						{
+							FToolMenuSection& Section = Menu->AddSection("EdGraphSchemaNodeInjectionVisualDebug", LOCTEXT("NodeInjectionVisualDebug", "Visual Debug"));
+
+							URigVMNode* VisualDebugNode = nullptr;
+							for (URigVMInjectionInfo* Injection : ModelPin->GetInjectedNodes())
+							{
+								FString PrototypeName;
+								if (Injection->StructNode->GetScriptStruct()->GetStringMetaDataHierarchical(TEXT("PrototypeName"), &PrototypeName))
+								{
+									if (PrototypeName == TEXT("VisualDebug"))
+									{
+										VisualDebugNode = Injection->StructNode;
+										break;
+									}
+								}
+							}
+
+							if (VisualDebugNode == nullptr)
+							{
+								UScriptStruct* ScriptStruct = nullptr;
+
+								if (ModelPin->GetCPPType() == TEXT("FVector"))
+								{
+									ScriptStruct = FRigUnit_VisualDebugVector::StaticStruct();
+								}
+								else if (ModelPin->GetCPPType() == TEXT("FQuat"))
+								{
+									ScriptStruct = FRigUnit_VisualDebugQuat::StaticStruct();
+								}
+								else if (ModelPin->GetCPPType() == TEXT("FTransform"))
+								{
+									ScriptStruct = FRigUnit_VisualDebugTransform::StaticStruct();
+								}
+								else
+								{
+									checkNoEntry();
+								}
+
+								Section.AddMenuEntry(
+									"AddVisualDebug",
+									LOCTEXT("AddVisualDebug", "Add Visual Debug"),
+									LOCTEXT("AddVisualDebug_Tooltip", "Injects a visual debugging node"),
+									FSlateIcon(),
+									FUIAction(FExecuteAction::CreateLambda([RigBlueprint, InGraphPin, ModelPin, ScriptStruct]() {
+										URigVMInjectionInfo* Injection = RigBlueprint->Controller->AddInjectedNode(ModelPin->GetPinPath(), ModelPin->GetDirection() != ERigVMPinDirection::Output, ScriptStruct, TEXT("Execute"), TEXT("Value"), TEXT("Value"));
+										if (Injection)
+										{
+											TArray<FName> NodeNames;
+											NodeNames.Add(Injection->StructNode->GetFName());
+											RigBlueprint->Controller->SetNodeSelection(NodeNames);
+
+											if (URigVMStructNode* StructNode = Cast<URigVMStructNode>(ModelPin->GetNode()))
+											{
+												if (TSharedPtr<FStructOnScope> DefaultStructScope = StructNode->ConstructStructInstance())
+												{
+													FRigVMStruct* DefaultStruct = (FRigVMStruct*)DefaultStructScope->GetStructMemory();
+
+													FString PinPath = ModelPin->GetPinPath();
+													FString Left, Right;
+
+													FName SpaceName = NAME_None;
+													if (URigVMPin::SplitPinPathAtStart(PinPath, Left, Right))
+													{
+														SpaceName = DefaultStruct->DetermineSpaceForPin(Right, &RigBlueprint->HierarchyContainer);
+													}
+
+													if (!SpaceName.IsNone())
+													{
+														if (URigVMPin* BoneSpacePin = Injection->StructNode->FindPin(TEXT("BoneSpace")))
+														{
+															RigBlueprint->Controller->SetPinDefaultValue(BoneSpacePin->GetPinPath(), SpaceName.ToString());
+														}
+													}
+												}
+											}
+										}
+									})
+								));
+							}
+							else
+							{
+								Section.AddMenuEntry(
+									"EditVisualDebug",
+									LOCTEXT("EditVisualDebug", "Edit Visual Debug"),
+									LOCTEXT("EditVisualDebug_Tooltip", "Edit the visual debugging node"),
+									FSlateIcon(),
+									FUIAction(FExecuteAction::CreateLambda([RigBlueprint, VisualDebugNode]() {
+										TArray<FName> NodeNames;
+										NodeNames.Add(VisualDebugNode->GetFName());
+										RigBlueprint->Controller->SetNodeSelection(NodeNames);
+									})
+								));
+								Section.AddMenuEntry(
+									"ToggleVisualDebug",
+									LOCTEXT("ToggleVisualDebug", "Toggle Visual Debug"),
+									LOCTEXT("ToggleVisualDebug_Tooltip", "Toggle the visibility the visual debugging"),
+									FSlateIcon(),
+									FUIAction(FExecuteAction::CreateLambda([RigBlueprint, VisualDebugNode]() {
+										URigVMPin* EnabledPin = VisualDebugNode->FindPin(TEXT("bEnabled"));
+										check(EnabledPin);
+										RigBlueprint->Controller->SetPinDefaultValue(EnabledPin->GetPinPath(), EnabledPin->GetDefaultValue() == TEXT("True") ? TEXT("False") : TEXT("True"), false);
+									})
+								));
+								Section.AddMenuEntry(
+									"RemoveVisualDebug",
+									LOCTEXT("RemoveVisualDebug", "Remove Visual Debug"),
+									LOCTEXT("RemoveVisualDebug_Tooltip", "Removes the visual debugging node"),
+									FSlateIcon(),
+										FUIAction(FExecuteAction::CreateLambda([RigBlueprint, InGraphPin, ModelPin, VisualDebugNode]() {
+										RigBlueprint->Controller->RemoveNodeByName(VisualDebugNode->GetFName());
+									})
+								));
+							}
+						}
+					}
+				}
+			}
 		}
-	}
-}
+		else if(Context->Node) // right clicked on the node
+		{
+			if (UControlRigBlueprint* RigBlueprint = Cast<UControlRigBlueprint>((UBlueprint*)Context->Blueprint))
+			{
+				TArray<FRigElementKey> RigElementsToSelect;
+				if (URigVMNode* ModelNode = RigBlueprint->Model->FindNodeByName(Context->Node->GetFName()))
+				{
+					for (const URigVMPin* Pin : ModelNode->GetPins())
+					{
+						if (Pin->GetCPPType() == TEXT("FName"))
+						{
+							if (Pin->GetCustomWidgetName() == TEXT("BoneName"))
+							{
+								RigElementsToSelect.Add(FRigElementKey(*Pin->GetDefaultValue(), ERigElementType::Bone));
+							}
+							else if (Pin->GetCustomWidgetName() == TEXT("ControlName"))
+							{
+								RigElementsToSelect.Add(FRigElementKey(*Pin->GetDefaultValue(), ERigElementType::Control));
+							}
+							else if (Pin->GetCustomWidgetName() == TEXT("SpaceName"))
+							{
+								RigElementsToSelect.Add(FRigElementKey(*Pin->GetDefaultValue(), ERigElementType::Space));
+							}
+							else if (Pin->GetCustomWidgetName() == TEXT("CurveName"))
+							{
+								RigElementsToSelect.Add(FRigElementKey(*Pin->GetDefaultValue(), ERigElementType::Curve));
+							}
+						}
+					}
+				}
 
-// It's CDO of the class, so we don't want the object to be writable or even if you write, it won't be per instance
-TSubclassOf<URigUnitEditor_Base> FControlRigEditorModule::GetEditorObjectByRigUnit(const FName& RigUnitClassName) 
-{
-	const TSubclassOf<URigUnitEditor_Base>* Class = RigUnitEditorClasses.Find(RigUnitClassName);
-	if (Class)
-	{
-		return *Class;
+				if (RigElementsToSelect.Num() > 0)
+				{
+					FToolMenuSection& Section = Menu->AddSection("EdGraphSchemaHierarchy", LOCTEXT("HierarchyHeader", "Hierarchy"));
+					Section.AddMenuEntry(
+						"SelectRigElements",
+						LOCTEXT("SelectRigElements", "Select Rig Elements"),
+						LOCTEXT("SelectRigElements_Tooltip", "Selects the bone, controls or spaces associated with this node."),
+						FSlateIcon(),
+						FUIAction(FExecuteAction::CreateLambda([RigBlueprint, RigElementsToSelect]() {
+
+							RigBlueprint->HierarchyContainer.ClearSelection();
+							for (const FRigElementKey& RigElementToSelect : RigElementsToSelect)
+							{
+								RigBlueprint->HierarchyContainer.Select(RigElementToSelect, true);
+							}
+
+						})
+					));
+				}
+			}
+		}
 	}
-
-	//if you don't find anything, just send out base one
-	return URigUnitEditor_Base::StaticClass();
 }
 
 IMPLEMENT_MODULE(FControlRigEditorModule, ControlRigEditor)
diff --git a/Engine/Plugins/Experimental/ControlRig/Source/ControlRigEditor/Private/ControlRigEditorModule.h b/Engine/Plugins/Experimental/ControlRig/Source/ControlRigEditor/Private/ControlRigEditorModule.h
--- a/Engine/Plugins/Experimental/ControlRig/Source/ControlRigEditor/Private/ControlRigEditorModule.h
+++ b/Engine/Plugins/Experimental/ControlRig/Source/ControlRigEditor/Private/ControlRigEditorModule.h
@@ -38,17 +38,11 @@ public:
 
 
 	
-	virtual void RegisterRigUnitEditorClass(FName RigUnitClassName, TSubclassOf<URigUnitEditor_Base> Class) override;
-	virtual void UnregisterRigUnitEditorClass(FName RigUnitClassName) override;
-	
 	virtual void GetTypeActions(const UControlRigBlueprint* CRB, FBlueprintActionDatabaseRegistrar& ActionRegistrar) override;
-	virtual void GetInstanceActions(const UControlRigBlueprint* CRB, FBlueprintActionDatabaseRegistrar& ActionRegistrar) override;
 	virtual FConnectionDrawingPolicy* CreateConnectionDrawingPolicy(int32 InBackLayerID, int32 InFrontLayerID, float InZoomFactor, const FSlateRect& InClippingRect, class FSlateWindowElementList& InDrawElements, class UEdGraph* InGraphObj) override;
 	virtual void GetNodeContextMenuActions(const UControlRigGraphNode* Node, class UToolMenu* Menu, class UGraphNodeContextMenuContext* Context) const override;
 	virtual void GetContextMenuActions(const UControlRigGraphSchema* Schema, class UToolMenu* Menu, class UGraphNodeContextMenuContext* Context) const override;
 
-	static TSubclassOf<URigUnitEditor_Base> GetEditorObjectByRigUnit(const FName& RigUnitClassName);
-
 private:
 	/** Handle a new animation controller blueprint being created */
 	void HandleNewBlueprintCreated(UBlueprint* InBlueprint);
@@ -76,8 +70,4 @@ private:
 	/** Delegate handles for blueprint utils */
 	FDelegateHandle RefreshAllNodesDelegateHandle;
 	FDelegateHandle ReconstructAllNodesDelegateHandle;
-	FDelegateHandle RenameVariableReferencesDelegateHandle;
-
-	/** Rig Unit Editor Classes Handler */
-	static TMap<FName, TSubclassOf<URigUnitEditor_Base>> RigUnitEditorClasses;
 };
\ No newline at end of file
diff --git a/Engine/Plugins/Experimental/ControlRig/Source/ControlRigEditor/Private/ControlRigEditorStyle.h b/Engine/Plugins/Experimental/ControlRig/Source/ControlRigEditor/Private/ControlRigEditorStyle.h
--- a/Engine/Plugins/Experimental/ControlRig/Source/ControlRigEditor/Private/ControlRigEditorStyle.h
+++ b/Engine/Plugins/Experimental/ControlRig/Source/ControlRigEditor/Private/ControlRigEditorStyle.h
@@ -63,6 +63,11 @@ public:
 
 			Set("ControlRig.ExecuteGraph", new IMAGE_BRUSH("Slate/ExecuteGraph", Icon40x40));
 			Set("ControlRig.ExecuteGraph.Small", new IMAGE_BRUSH("Slate/ExecuteGraph", Icon20x20));
+
+			Set("ControlRig.AutoCompileGraph", new IMAGE_BRUSH("Slate/AutoCompile", Icon40x40));
+			Set("ControlRig.AutoCompileGraph.Small", new IMAGE_BRUSH("Slate/AutoCompile", Icon20x20));
+
+			Set("ControlRig.RigUnit.VisualDebug", new IMAGE_BRUSH("Slate/ControlRig_Unit_VisualDebug_40x", Icon40x40));
 		}
 
 		// Graph styles
