COPY from //Dev-Anim to //Dev-Main

diff --git a/Engine/Plugins/Experimental/ControlRig/Source/ControlRigEditor/Private/ControlRigElementDetails.cpp b/Engine/Plugins/Experimental/ControlRig/Source/ControlRigEditor/Private/ControlRigElementDetails.cpp
--- a/Engine/Plugins/Experimental/ControlRig/Source/ControlRigEditor/Private/ControlRigElementDetails.cpp
+++ b/Engine/Plugins/Experimental/ControlRig/Source/ControlRigEditor/Private/ControlRigElementDetails.cpp
@@ -9,10 +9,199 @@
 #include "Widgets/Input/SEditableTextBox.h"
 #include "Widgets/Input/SVectorInputBox.h"
 #include "Widgets/Input/SRotatorInputBox.h"
+#include "Widgets/Input/SCheckBox.h"
+#include "Widgets/Input/SNumericEntryBox.h"
 #include "Editor/SControlRigGizmoNameList.h"
+#include "ControlRigBlueprint.h"
+#include "Graph/ControlRigGraph.h"
 
 #define LOCTEXT_NAMESPACE "ControlRigElementDetails"
 
+void FRigUnitDetails::CustomizeDetails(IDetailLayoutBuilder& DetailBuilder)
+{
+	TArray<TSharedPtr<FStructOnScope>> StructsBeingCustomized;
+	DetailBuilder.GetStructsBeingCustomized(StructsBeingCustomized);
+	if (StructsBeingCustomized.Num() == 0)
+	{
+		return;
+	}
+
+	TSharedPtr<FStructOnScope> StructBeingCustomized = StructsBeingCustomized[0];
+
+	BlueprintBeingCustomized = nullptr;
+	if (UPackage* Package = StructBeingCustomized->GetPackage())
+	{
+		TArray<UObject*> SubObjects;
+		Package->GetDefaultSubobjects(SubObjects);
+
+		for (UObject* SubObject : SubObjects)
+		{
+			if (UControlRig* Rig = Cast<UControlRig>(SubObject))
+			{
+				BlueprintBeingCustomized = Cast<UControlRigBlueprint>(Rig->GetClass()->ClassGeneratedBy);
+				if (BlueprintBeingCustomized)
+				{
+					break;
+				}
+			}
+		}
+	}
+
+	if (BlueprintBeingCustomized == nullptr)
+	{
+		return;
+	}
+
+	GraphBeingCustomized = nullptr;
+	for (UEdGraph* Graph : BlueprintBeingCustomized->UbergraphPages)
+	{
+		GraphBeingCustomized = Cast<UControlRigGraph>(Graph);
+		if (GraphBeingCustomized)
+		{
+			break;
+		}
+	}
+
+	if (GraphBeingCustomized == nullptr)
+	{
+		return;
+	}
+
+	UScriptStruct* ScriptStruct = Cast<UScriptStruct>((UStruct*)StructBeingCustomized->GetStruct());
+	check(ScriptStruct);
+
+	IDetailCategoryBuilder& CategoryBuilder = DetailBuilder.EditCategory(*ScriptStruct->GetDisplayNameText().ToString());
+
+	for (TFieldIterator<FProperty> PropertyIt(ScriptStruct); PropertyIt; ++PropertyIt)
+	{
+		FProperty* Property = *PropertyIt;
+		TSharedPtr<IPropertyHandle> PropertyHandle = DetailBuilder.GetProperty(Property->GetFName(), ScriptStruct);
+		if (!PropertyHandle->IsValidHandle())
+		{
+			continue;
+		}
+		DetailBuilder.HideProperty(PropertyHandle);
+
+		if (FNameProperty* NameProperty = CastField<FNameProperty>(Property))
+		{
+			FString CustomWidgetName = NameProperty->GetMetaData(TEXT("CustomWidget"));
+			if (!CustomWidgetName.IsEmpty())
+			{
+				const TArray<TSharedPtr<FString>>* NameList = nullptr;
+				if (CustomWidgetName == TEXT("BoneName"))
+				{
+					NameList = &GraphBeingCustomized->GetBoneNameList();
+				}
+				else if (CustomWidgetName == TEXT("ControlName"))
+				{
+					NameList = &GraphBeingCustomized->GetControlNameList();
+				}
+				else if (CustomWidgetName == TEXT("SpaceName"))
+				{
+					NameList = &GraphBeingCustomized->GetSpaceNameList();
+				}
+				else if (CustomWidgetName == TEXT("CurveName"))
+				{
+					NameList = &GraphBeingCustomized->GetCurveNameList();
+				}
+
+				if (NameList)
+				{
+					TSharedPtr<SControlRigGraphPinNameListValueWidget> NameListWidget;
+
+					CategoryBuilder.AddCustomRow(FText::FromString(Property->GetName()))
+						.NameContent()
+						[
+							PropertyHandle->CreatePropertyNameWidget()
+						]
+						.ValueContent()
+						[
+							SAssignNew(NameListWidget, SControlRigGraphPinNameListValueWidget)
+							.OptionsSource(NameList)
+							.OnGenerateWidget(this, &FRigUnitDetails::MakeNameListItemWidget)
+							.OnSelectionChanged(this, &FRigUnitDetails::OnNameListChanged, StructBeingCustomized, NameProperty, DetailBuilder.GetPropertyUtilities())
+							.OnComboBoxOpening(this, &FRigUnitDetails::OnNameListComboBox, StructBeingCustomized, NameProperty, NameList)
+							.InitiallySelectedItem(GetCurrentlySelectedItem(StructBeingCustomized, NameProperty, NameList))
+							.Content()
+							[
+								SNew(STextBlock)
+								.Text(this, &FRigUnitDetails::GetNameListText, StructBeingCustomized, NameProperty)
+							]
+						];
+
+					NameListWidgets.Add(Property->GetFName(), NameListWidget);
+				}
+				else
+				{
+					CategoryBuilder.AddCustomRow(FText::FromString(Property->GetName()))
+						.NameContent()
+						[
+							PropertyHandle->CreatePropertyNameWidget()
+						];
+				}
+				continue;
+			}
+		}
+		CategoryBuilder.AddProperty(PropertyHandle);
+	}
+}
+
+TSharedRef<SWidget> FRigUnitDetails::MakeNameListItemWidget(TSharedPtr<FString> InItem)
+{
+	return 	SNew(STextBlock).Text(FText::FromString(*InItem));// .Font(FEditorStyle::GetFontStyle(TEXT("PropertyWindow.NormalFont")));
+}
+
+FText FRigUnitDetails::GetNameListText(TSharedPtr<FStructOnScope> InStructOnScope, FNameProperty* InProperty) const
+{
+	if (FName* Value = InProperty->ContainerPtrToValuePtr<FName>(InStructOnScope->GetStructMemory()))
+	{
+		return FText::FromName(*Value);
+	}
+	return FText();
+}
+
+TSharedPtr<FString> FRigUnitDetails::GetCurrentlySelectedItem(TSharedPtr<FStructOnScope> InStructOnScope, FNameProperty* InProperty, const TArray<TSharedPtr<FString>>* InNameList) const
+{
+	FString CurrentItem = GetNameListText(InStructOnScope, InProperty).ToString();
+	for (const TSharedPtr<FString>& Item : *InNameList)
+	{
+		if (Item->Equals(CurrentItem))
+		{
+			return Item;
+		}
+	}
+
+	return TSharedPtr<FString>();
+}
+
+
+void FRigUnitDetails::SetNameListText(const FText& NewTypeInValue, ETextCommit::Type /*CommitInfo*/, TSharedPtr<FStructOnScope> InStructOnScope, FNameProperty* InProperty, TSharedRef<IPropertyUtilities> PropertyUtilities)
+{
+	if (FName* Value = InProperty->ContainerPtrToValuePtr<FName>(InStructOnScope->GetStructMemory()))
+	{
+		*Value = *NewTypeInValue.ToString();
+
+		FPropertyChangedEvent ChangeEvent(InProperty, EPropertyChangeType::ValueSet);
+		PropertyUtilities->NotifyFinishedChangingProperties(ChangeEvent);
+	}
+}
+
+void FRigUnitDetails::OnNameListChanged(TSharedPtr<FString> NewSelection, ESelectInfo::Type SelectInfo, TSharedPtr<FStructOnScope> InStructOnScope, FNameProperty* InProperty, TSharedRef<IPropertyUtilities> PropertyUtilities)
+{
+	if (SelectInfo != ESelectInfo::Direct)
+	{
+		FString NewValue = *NewSelection.Get();
+		SetNameListText(FText::FromString(NewValue), ETextCommit::OnEnter, InStructOnScope, InProperty, PropertyUtilities);
+	}
+}
+
+void FRigUnitDetails::OnNameListComboBox(TSharedPtr<FStructOnScope> InStructOnScope, FNameProperty* InProperty, const TArray<TSharedPtr<FString>>* InNameList)
+{
+	TSharedPtr<SControlRigGraphPinNameListValueWidget> Widget = NameListWidgets.FindChecked(InProperty->GetFName());
+	TSharedPtr<FString> CurrentlySelected = GetCurrentlySelectedItem(InStructOnScope, InProperty, InNameList);
+	Widget->SetSelectedItem(CurrentlySelected);
+}
+
 void FRigElementDetails::CustomizeDetails(IDetailLayoutBuilder& DetailBuilder)
 {
 	DetailBuilder.HideProperty(TEXT("Index"), FRigElement::StaticStruct());
@@ -222,6 +411,8 @@ void FRigBoneDetails::CustomizeDetails( IDetailLayoutBuilder& DetailBuilder )
 	];
 }
 
+TArray<TSharedPtr<FString>> FRigControlDetails::ControlTypeList;
+
 void FRigControlDetails::CustomizeDetails(IDetailLayoutBuilder& DetailBuilder)
 {
 	FRigElementDetails::CustomizeDetails(DetailBuilder);
@@ -231,7 +422,13 @@ void FRigControlDetails::CustomizeDetails(IDetailLayoutBuilder& DetailBuilder)
 	DetailBuilder.HideProperty(TEXT("SpaceName"));
 	DetailBuilder.HideProperty(TEXT("InitialValue"));
 	DetailBuilder.HideProperty(TEXT("Value"));
-	DetailBuilder.HideProperty(TEXT("GizmoName"));
+	DetailBuilder.HideProperty(TEXT("bLimitTranslation"));
+	DetailBuilder.HideProperty(TEXT("bLimitRotation"));
+	DetailBuilder.HideProperty(TEXT("bLimitScale"));
+	DetailBuilder.HideProperty(TEXT("bDrawLimits"));
+	DetailBuilder.HideProperty(TEXT("MinimumValue"));
+	DetailBuilder.HideProperty(TEXT("MaximumValue"));
+	DetailBuilder.HideProperty(TEXT("bIsTransientControl"));
 
 	GizmoNameList.Reset();
 	if (BlueprintBeingCustomized)
@@ -250,10 +447,15 @@ void FRigControlDetails::CustomizeDetails(IDetailLayoutBuilder& DetailBuilder)
 		}
 	}
 
-	IDetailCategoryBuilder& Category = DetailBuilder.EditCategory(TEXT("FRigElement"), LOCTEXT("ControlCategory", "Control"));
-	Category.InitiallyCollapsed(false);
+	IDetailCategoryBuilder& ControlCategory = DetailBuilder.EditCategory(TEXT("Control"), LOCTEXT("ControlCategory", "Control"));
+	IDetailCategoryBuilder& LimitsCategory = DetailBuilder.EditCategory(TEXT("Limits"), LOCTEXT("LimitsCategory", "Limits"));
+	IDetailCategoryBuilder& GizmoCategory = DetailBuilder.EditCategory(TEXT("Gizmo"), LOCTEXT("GizmoCategory", "Gizmo"));
 
-	Category.AddCustomRow(FText::FromString(TEXT("Name")))
+	ControlCategory.InitiallyCollapsed(false);
+	LimitsCategory.InitiallyCollapsed(false);
+	GizmoCategory.InitiallyCollapsed(false);
+
+	ControlCategory.AddCustomRow(FText::FromString(TEXT("Name")))
 	.NameContent()
 	[
 		SNew(STextBlock)
@@ -273,22 +475,180 @@ void FRigControlDetails::CustomizeDetails(IDetailLayoutBuilder& DetailBuilder)
 		return;
 	}
 
+	if (ControlTypeList.Num() == 0)
+	{
+		UEnum* Enum = StaticEnum<ERigControlType>();
+		for (int64 Index = 0; Index < Enum->GetMaxEnumValue(); Index++)
+		{
+			ControlTypeList.Add(MakeShared<FString>(Enum->GetDisplayNameTextByValue(Index).ToString()));
+		}
+	}
+
+	UEnum* ValueTypeEnum = StaticEnum<ERigControlValueType>();
 	FRigControl& Control = ContainerBeingCustomized->ControlHierarchy[ElementKeyBeingCustomized.Name];
+
+	// when control type changes, we have to refresh detail panel
+	TSharedRef<IPropertyHandle> ControlType = DetailBuilder.GetProperty(TEXT("ControlType"));
+	ControlType->SetOnPropertyValueChanged(FSimpleDelegate::CreateLambda(
+		[this, &DetailBuilder]()
+		{
+			DetailBuilder.ForceRefreshDetails();
+
+			if (this->ContainerBeingCustomized && this->ElementKeyBeingCustomized.IsValid())
+			{
+				FRigControl& Control = this->ContainerBeingCustomized->ControlHierarchy[this->ElementKeyBeingCustomized.Name];
+
+				switch (Control.ControlType)
+				{
+					case ERigControlType::Bool:
+					{
+						Control.Value = Control.InitialValue = FRigControlValue::Make<bool>(false);
+						break;
+					}
+					case ERigControlType::Float:
+					{
+						Control.Value = Control.InitialValue = FRigControlValue::Make<float>(0.f);
+						Control.MinimumValue = FRigControlValue::Make<float>(0.f);
+						Control.MaximumValue = FRigControlValue::Make<float>(1.f);
+						break;
+					}
+					case ERigControlType::Vector2D:
+					{
+						Control.Value = Control.InitialValue = FRigControlValue::Make<FVector2D>(FVector2D::ZeroVector);
+						Control.MinimumValue = FRigControlValue::Make<FVector2D>(FVector2D::ZeroVector);
+						Control.MaximumValue = FRigControlValue::Make<FVector2D>(FVector2D(1.f, 1.f));
+						break;
+					}
+					case ERigControlType::Position:
+					{
+						Control.Value = Control.InitialValue = FRigControlValue::Make<FVector>(FVector::ZeroVector);
+						Control.MinimumValue = FRigControlValue::Make<FVector>(-FVector::OneVector);
+						Control.MaximumValue = FRigControlValue::Make<FVector>(FVector::OneVector);
+						break;
+					}
+					case ERigControlType::Scale:
+					{
+						Control.Value = Control.InitialValue = FRigControlValue::Make<FVector>(FVector::OneVector);
+						Control.MinimumValue = FRigControlValue::Make<FVector>(FVector::ZeroVector);
+						Control.MaximumValue = FRigControlValue::Make<FVector>(FVector::OneVector);
+						break;
+					}
+					case ERigControlType::Rotator:
+					{
+						Control.Value = Control.InitialValue = FRigControlValue::Make<FRotator>(FRotator::ZeroRotator);
+						Control.MinimumValue = FRigControlValue::Make<FRotator>(FRotator::ZeroRotator);
+						Control.MaximumValue = FRigControlValue::Make<FRotator>(FRotator(180.f, 180.f, 180.f));
+						break;
+					}
+					case ERigControlType::Transform:
+					{
+						Control.Value = Control.InitialValue = Control.MinimumValue = Control.MaximumValue = FRigControlValue::Make<FTransform>(FTransform::Identity);
+						break;
+					}
+					case ERigControlType::TransformNoScale:
+					{
+						FTransformNoScale Identity = FTransform::Identity;
+						// Helge todo FTransformNoScale::Identity appears to be no good?
+						Control.Value = Control.InitialValue = Control.MinimumValue = Control.MaximumValue = FRigControlValue::Make<FTransformNoScale>(Identity);
+						break;
+					}
+					default:
+					{
+						ensure(false);
+						break;
+					}
+				}
+
+				this->ContainerBeingCustomized->ControlHierarchy.OnControlUISettingsChanged.Broadcast(this->ContainerBeingCustomized, this->ElementKeyBeingCustomized);
+			}
+		}
+	));
+
+	ControlCategory.AddCustomRow(FText::FromString(TEXT("ControlType")))
+	.NameContent()
+	[
+		ControlType->CreatePropertyNameWidget()
+	]
+	.ValueContent()
+	[
+		ControlType->CreatePropertyValueWidget()
+	];
+
 	switch (Control.ControlType)
 	{
-		case ERigControlType::Transform:
+		case ERigControlType::Bool:
+		{
+			for(int64 ValueTypeIndex = 0; ValueTypeIndex < ValueTypeEnum->GetMaxEnumValue(); ValueTypeIndex++)
+			{
+				ERigControlValueType ValueType = (ERigControlValueType)ValueTypeIndex;
+
+				if (ValueType == ERigControlValueType::Minimum ||
+					ValueType == ERigControlValueType::Maximum)
+				{
+					continue;
+				}
+
+				ControlCategory.AddCustomRow(ValueTypeEnum->GetDisplayNameTextByValue(ValueTypeIndex))
+				.NameContent()
+				.VAlign(VAlign_Top)
+				[
+					SNew(STextBlock)
+					.Text(ValueTypeEnum->GetDisplayNameTextByValue(ValueTypeIndex))
+					.Font(IDetailLayoutBuilder::GetDetailFont())
+				]
+				.ValueContent()
+				[
+					SNew(SVerticalBox)
+					+ SVerticalBox::Slot()
+					[
+						SNew(SCheckBox)
+						.IsChecked(this, &FRigControlDetails::GetComponentValueBool, ValueType == ERigControlValueType::Initial)
+						.OnCheckStateChanged(this, &FRigControlDetails::SetComponentValueBool, ValueType == ERigControlValueType::Initial)
+					]
+				];
+			}
+			break;
+		}
+		case ERigControlType::Float:
 		{
-			TArray<FString> Labels = { TEXT("Initial"), TEXT("Current") };
-			for(FString Label : Labels)
+			for(int64 ValueTypeIndex = 0; ValueTypeIndex < ValueTypeEnum->GetMaxEnumValue(); ValueTypeIndex++)
 			{
-				bool bIsInitial = Label == TEXT("Initial");
-				Category.AddCustomRow(FText::FromString(Label))
+				ERigControlValueType ValueType = (ERigControlValueType)ValueTypeIndex;
+
+				if (ValueType == ERigControlValueType::Minimum)
+				{
+					LimitsCategory.AddCustomRow(FText::FromString(TEXT("Limit")))
+					.NameContent()
+					[
+						SNew(STextBlock)
+						.Text(FText::FromString(TEXT("Limit")))
+						.Font(IDetailLayoutBuilder::GetDetailFont())
+					]
+					.ValueContent()
+					[
+						DetailBuilder.GetProperty(TEXT("bLimitTranslation"))->CreatePropertyValueWidget()
+					];
+
+					LimitsCategory.AddCustomRow(FText::FromString(TEXT("bDrawLimits")))
+					.NameContent()
+					[
+						DetailBuilder.GetProperty(TEXT("bDrawLimits"))->CreatePropertyNameWidget()
+					]
+					.ValueContent()
+					[
+						DetailBuilder.GetProperty(TEXT("bDrawLimits"))->CreatePropertyValueWidget()
+					];
+				}
+
+				IDetailCategoryBuilder& Category = (ValueType == ERigControlValueType::Minimum || ValueType == ERigControlValueType::Maximum) ? LimitsCategory : ControlCategory;
+				Category.AddCustomRow(ValueTypeEnum->GetDisplayNameTextByValue(ValueTypeIndex))
 				.NameContent()
 				.VAlign(VAlign_Top)
 				[
 					SNew(STextBlock)
-					.Text(FText::FromString(Label))
+					.Text(ValueTypeEnum->GetDisplayNameTextByValue(ValueTypeIndex))
 					.Font(IDetailLayoutBuilder::GetDetailFont())
+					.IsEnabled(this, &FRigControlDetails::IsEnabled, ValueType)
 				]
 				.ValueContent()
 				.MinDesiredWidth(125.0f * 3.0f) // copied from FComponentTransformDetails
@@ -297,69 +657,538 @@ void FRigControlDetails::CustomizeDetails(IDetailLayoutBuilder& DetailBuilder)
 					SNew(SVerticalBox)
 					+ SVerticalBox::Slot()
 					[
-						SNew(SVectorInputBox)
+						SNew(SNumericEntryBox<float>)
+						.Font(FEditorStyle::GetFontStyle(TEXT("MenuItem.Font")))
 						.AllowSpin(true)
-						.bColorAxisLabels(true)
+						.IsEnabled(this, &FRigControlDetails::IsEnabled, ValueType)
+						.Value(this, &FRigControlDetails::GetComponentValueFloat, ValueType, ERigElementDetailsTransformComponent::TranslationX)
+						.OnValueChanged(this, &FRigControlDetails::SetComponentValueFloat, ValueType, ERigElementDetailsTransformComponent::TranslationX)
+					]
+				];
+			}
+			break;
+		}
+		case ERigControlType::Vector2D:
+		{
+			for(int64 ValueTypeIndex = 0; ValueTypeIndex < ValueTypeEnum->GetMaxEnumValue(); ValueTypeIndex++)
+			{
+				ERigControlValueType ValueType = (ERigControlValueType)ValueTypeIndex;
+
+				if (ValueType == ERigControlValueType::Minimum)
+				{
+					LimitsCategory.AddCustomRow(FText::FromString(TEXT("Limit")))
+					.NameContent()
+					[
+						SNew(STextBlock)
+						.Text(FText::FromString(TEXT("Limit")))
 						.Font(IDetailLayoutBuilder::GetDetailFont())
-						.X(this, &FRigControlDetails::GetComponentValue, bIsInitial, ERigElementDetailsTransformComponent::TranslationX)
-						.Y(this, &FRigControlDetails::GetComponentValue, bIsInitial, ERigElementDetailsTransformComponent::TranslationY)
-						.Z(this, &FRigControlDetails::GetComponentValue, bIsInitial, ERigElementDetailsTransformComponent::TranslationZ)
-						.OnXCommitted(this, &FRigControlDetails::SetComponentValue, bIsInitial, ERigElementDetailsTransformComponent::TranslationX)
-						.OnYCommitted(this, &FRigControlDetails::SetComponentValue, bIsInitial, ERigElementDetailsTransformComponent::TranslationY)
-						.OnZCommitted(this, &FRigControlDetails::SetComponentValue, bIsInitial, ERigElementDetailsTransformComponent::TranslationZ)
 					]
-					+ SVerticalBox::Slot()
+					.ValueContent()
 					[
-						SNew(SRotatorInputBox)
+						DetailBuilder.GetProperty(TEXT("bLimitTranslation"))->CreatePropertyValueWidget()
+					];
+
+					LimitsCategory.AddCustomRow(FText::FromString(TEXT("bDrawLimits")))
+					.NameContent()
+					[
+						DetailBuilder.GetProperty(TEXT("bDrawLimits"))->CreatePropertyNameWidget()
+					]
+					.ValueContent()
+					[
+						DetailBuilder.GetProperty(TEXT("bDrawLimits"))->CreatePropertyValueWidget()
+					];
+				}
+
+				IDetailCategoryBuilder& Category = (ValueType == ERigControlValueType::Minimum || ValueType == ERigControlValueType::Maximum) ? LimitsCategory : ControlCategory;
+				Category.AddCustomRow(ValueTypeEnum->GetDisplayNameTextByValue(ValueTypeIndex))
+				.NameContent()
+				.VAlign(VAlign_Top)
+				[
+					SNew(STextBlock)
+					.Text(ValueTypeEnum->GetDisplayNameTextByValue(ValueTypeIndex))
+					.Font(IDetailLayoutBuilder::GetDetailFont())
+					.IsEnabled(this, &FRigControlDetails::IsEnabled, ValueType)
+				]
+				.ValueContent()
+				.MinDesiredWidth(125.0f * 3.0f) // copied from FComponentTransformDetails
+				.MaxDesiredWidth(125.0f * 3.0f)
+				[
+					SNew(SHorizontalBox)
+					+ SHorizontalBox::Slot()
+					.Padding(0.0f, 1.0f, 2.0f, 1.0f)
+					[
+						SNew(SNumericEntryBox<float>)
+						.Font(FEditorStyle::GetFontStyle(TEXT("MenuItem.Font")))
 						.AllowSpin(true)
-						.bColorAxisLabels(true)
+						.IsEnabled(this, &FRigControlDetails::IsEnabled, ValueType)
+						.Value(this, &FRigControlDetails::GetComponentValueFloat, ValueType, ERigElementDetailsTransformComponent::TranslationX)
+						.OnValueChanged(this, &FRigControlDetails::SetComponentValueFloat, ValueType, ERigElementDetailsTransformComponent::TranslationX)
+					]
+					.Padding(0.0f, 1.0f, 2.0f, 1.0f)
+					+ SHorizontalBox::Slot()
+					[
+						SNew(SNumericEntryBox<float>)
+						.Font(FEditorStyle::GetFontStyle(TEXT("MenuItem.Font")))
+						.AllowSpin(true)
+						.IsEnabled(this, &FRigControlDetails::IsEnabled, ValueType)
+						.Value(this, &FRigControlDetails::GetComponentValueFloat, ValueType, ERigElementDetailsTransformComponent::TranslationY)
+						.OnValueChanged(this, &FRigControlDetails::SetComponentValueFloat, ValueType, ERigElementDetailsTransformComponent::TranslationY)
+					]
+				];
+			}
+			break;
+		}
+		case ERigControlType::Position:
+		{
+			for(int64 ValueTypeIndex = 0; ValueTypeIndex < ValueTypeEnum->GetMaxEnumValue(); ValueTypeIndex++)
+			{
+				ERigControlValueType ValueType = (ERigControlValueType)ValueTypeIndex;
+
+				if (ValueType == ERigControlValueType::Minimum)
+				{
+					LimitsCategory.AddCustomRow(FText::FromString(TEXT("Limit")))
+					.NameContent()
+					[
+						SNew(STextBlock)
+						.Text(FText::FromString(TEXT("Limit")))
 						.Font(IDetailLayoutBuilder::GetDetailFont())
-						.Roll(this, &FRigControlDetails::GetComponentValue, bIsInitial, ERigElementDetailsTransformComponent::RotationRoll)
-						.Pitch(this, &FRigControlDetails::GetComponentValue, bIsInitial, ERigElementDetailsTransformComponent::RotationPitch)
-						.Yaw(this, &FRigControlDetails::GetComponentValue, bIsInitial, ERigElementDetailsTransformComponent::RotationYaw)
-						.OnRollCommitted(this, &FRigControlDetails::SetComponentValue, bIsInitial, ERigElementDetailsTransformComponent::RotationRoll)
-						.OnPitchCommitted(this, &FRigControlDetails::SetComponentValue, bIsInitial, ERigElementDetailsTransformComponent::RotationPitch)
-						.OnYawCommitted(this, &FRigControlDetails::SetComponentValue, bIsInitial, ERigElementDetailsTransformComponent::RotationYaw)
 					]
+					.ValueContent()
+					[
+						DetailBuilder.GetProperty(TEXT("bLimitTranslation"))->CreatePropertyValueWidget()
+					];
+
+					LimitsCategory.AddCustomRow(FText::FromString(TEXT("bDrawLimits")))
+					.NameContent()
+					[
+						DetailBuilder.GetProperty(TEXT("bDrawLimits"))->CreatePropertyNameWidget()
+					]
+					.ValueContent()
+					[
+						DetailBuilder.GetProperty(TEXT("bDrawLimits"))->CreatePropertyValueWidget()
+					];
+				}
+
+				IDetailCategoryBuilder& Category = (ValueType == ERigControlValueType::Minimum || ValueType == ERigControlValueType::Maximum) ? LimitsCategory : ControlCategory;
+				Category.AddCustomRow(ValueTypeEnum->GetDisplayNameTextByValue(ValueTypeIndex))
+				.NameContent()
+				.VAlign(VAlign_Top)
+				[
+					SNew(STextBlock)
+					.Text(ValueTypeEnum->GetDisplayNameTextByValue(ValueTypeIndex))
+					.Font(IDetailLayoutBuilder::GetDetailFont())
+					.IsEnabled(this, &FRigControlDetails::IsEnabled, ValueType)
+				]
+				.ValueContent()
+				.MinDesiredWidth(125.0f * 3.0f) // copied from FComponentTransformDetails
+				.MaxDesiredWidth(125.0f * 3.0f)
+				[
+					SNew(SVerticalBox)
 					+ SVerticalBox::Slot()
 					[
 						SNew(SVectorInputBox)
 						.AllowSpin(true)
+						.IsEnabled(this, &FRigControlDetails::IsEnabled, ValueType)
 						.bColorAxisLabels(true)
 						.Font(IDetailLayoutBuilder::GetDetailFont())
-						.X(this, &FRigControlDetails::GetComponentValue, bIsInitial, ERigElementDetailsTransformComponent::ScaleX)
-						.Y(this, &FRigControlDetails::GetComponentValue, bIsInitial, ERigElementDetailsTransformComponent::ScaleY)
-						.Z(this, &FRigControlDetails::GetComponentValue, bIsInitial, ERigElementDetailsTransformComponent::ScaleZ)
-						.OnXCommitted(this, &FRigControlDetails::SetComponentValue, bIsInitial, ERigElementDetailsTransformComponent::ScaleX)
-						.OnYCommitted(this, &FRigControlDetails::SetComponentValue, bIsInitial, ERigElementDetailsTransformComponent::ScaleY)
-						.OnZCommitted(this, &FRigControlDetails::SetComponentValue, bIsInitial, ERigElementDetailsTransformComponent::ScaleZ)
+						.X(this, &FRigControlDetails::GetComponentValueFloat, ValueType, ERigElementDetailsTransformComponent::TranslationX)
+						.Y(this, &FRigControlDetails::GetComponentValueFloat, ValueType, ERigElementDetailsTransformComponent::TranslationY)
+						.Z(this, &FRigControlDetails::GetComponentValueFloat, ValueType, ERigElementDetailsTransformComponent::TranslationZ)
+						.OnXCommitted(this, &FRigControlDetails::SetComponentValueFloat, ValueType, ERigElementDetailsTransformComponent::TranslationX)
+						.OnYCommitted(this, &FRigControlDetails::SetComponentValueFloat, ValueType, ERigElementDetailsTransformComponent::TranslationY)
+						.OnZCommitted(this, &FRigControlDetails::SetComponentValueFloat, ValueType, ERigElementDetailsTransformComponent::TranslationZ)
 					]
 				];
 			}
 			break;
 		}
-		default:
+		case ERigControlType::Scale:
 		{
-			ensure(false);
-			break;
-		}
-	}
+			for(int64 ValueTypeIndex = 0; ValueTypeIndex < ValueTypeEnum->GetMaxEnumValue(); ValueTypeIndex++)
+			{
+				ERigControlValueType ValueType = (ERigControlValueType)ValueTypeIndex;
 
-	Category.AddCustomRow(FText::FromString(TEXT("GizmoName")))
-	.NameContent()
-	[
-		SNew(STextBlock)
-		.Text(FText::FromString(TEXT("Gizmo Name")))
-		.Font(IDetailLayoutBuilder::GetDetailFont())
-	]
-	.ValueContent()
-	[
-		SNew(SControlRigGizmoNameList, &Control, BlueprintBeingCustomized)
-		.OnGetNameListContent(this, &FRigControlDetails::GetGizmoNameList)
-	];
+				if (ValueType == ERigControlValueType::Minimum)
+				{
+					LimitsCategory.AddCustomRow(FText::FromString(TEXT("Limit")))
+					.NameContent()
+					[
+						SNew(STextBlock)
+						.Text(FText::FromString(TEXT("Limit")))
+						.Font(IDetailLayoutBuilder::GetDetailFont())
+					]
+					.ValueContent()
+					[
+						DetailBuilder.GetProperty(TEXT("bLimitScale"))->CreatePropertyValueWidget()
+					];
+
+					LimitsCategory.AddCustomRow(FText::FromString(TEXT("bDrawLimits")))
+					.NameContent()
+					[
+						DetailBuilder.GetProperty(TEXT("bDrawLimits"))->CreatePropertyNameWidget()
+					]
+					.ValueContent()
+					[
+						DetailBuilder.GetProperty(TEXT("bDrawLimits"))->CreatePropertyValueWidget()
+					];
+				}
+
+				IDetailCategoryBuilder& Category = (ValueType == ERigControlValueType::Minimum || ValueType == ERigControlValueType::Maximum) ? LimitsCategory : ControlCategory;
+				Category.AddCustomRow(ValueTypeEnum->GetDisplayNameTextByValue(ValueTypeIndex))
+				.NameContent()
+				.VAlign(VAlign_Top)
+				[
+					SNew(STextBlock)
+					.Text(ValueTypeEnum->GetDisplayNameTextByValue(ValueTypeIndex))
+					.Font(IDetailLayoutBuilder::GetDetailFont())
+					.IsEnabled(this, &FRigControlDetails::IsEnabled, ValueType)
+				]
+				.ValueContent()
+				.MinDesiredWidth(125.0f * 3.0f) // copied from FComponentTransformDetails
+				.MaxDesiredWidth(125.0f * 3.0f)
+				[
+					SNew(SVerticalBox)
+					+ SVerticalBox::Slot()
+					[
+						SNew(SVectorInputBox)
+						.AllowSpin(true)
+						.IsEnabled(this, &FRigControlDetails::IsEnabled, ValueType)
+						.bColorAxisLabels(true)
+						.Font(IDetailLayoutBuilder::GetDetailFont())
+						.X(this, &FRigControlDetails::GetComponentValueFloat, ValueType, ERigElementDetailsTransformComponent::ScaleX)
+						.Y(this, &FRigControlDetails::GetComponentValueFloat, ValueType, ERigElementDetailsTransformComponent::ScaleY)
+						.Z(this, &FRigControlDetails::GetComponentValueFloat, ValueType, ERigElementDetailsTransformComponent::ScaleZ)
+						.OnXCommitted(this, &FRigControlDetails::SetComponentValueFloat, ValueType, ERigElementDetailsTransformComponent::ScaleX)
+						.OnYCommitted(this, &FRigControlDetails::SetComponentValueFloat, ValueType, ERigElementDetailsTransformComponent::ScaleY)
+						.OnZCommitted(this, &FRigControlDetails::SetComponentValueFloat, ValueType, ERigElementDetailsTransformComponent::ScaleZ)
+					]
+				];
+			}
+			break;
+		}
+		case ERigControlType::Rotator:
+		{
+			for(int64 ValueTypeIndex = 0; ValueTypeIndex < ValueTypeEnum->GetMaxEnumValue(); ValueTypeIndex++)
+			{
+				ERigControlValueType ValueType = (ERigControlValueType)ValueTypeIndex;
+
+				if (ValueType == ERigControlValueType::Minimum)
+				{
+					LimitsCategory.AddCustomRow(FText::FromString(TEXT("Limit")))
+					.NameContent()
+					[
+						SNew(STextBlock)
+						.Text(FText::FromString(TEXT("Limit")))
+						.Font(IDetailLayoutBuilder::GetDetailFont())
+					]
+					.ValueContent()
+					[
+						DetailBuilder.GetProperty(TEXT("bLimitRotation"))->CreatePropertyValueWidget()
+					];
+				}
+
+				IDetailCategoryBuilder& Category = (ValueType == ERigControlValueType::Minimum || ValueType == ERigControlValueType::Maximum) ? LimitsCategory : ControlCategory;
+				Category.AddCustomRow(ValueTypeEnum->GetDisplayNameTextByValue(ValueTypeIndex))
+				.NameContent()
+				.VAlign(VAlign_Top)
+				[
+					SNew(STextBlock)
+					.Text(ValueTypeEnum->GetDisplayNameTextByValue(ValueTypeIndex))
+					.Font(IDetailLayoutBuilder::GetDetailFont())
+					.IsEnabled(this, &FRigControlDetails::IsEnabled, ValueType)
+				]
+				.ValueContent()
+				.MinDesiredWidth(125.0f * 3.0f) // copied from FComponentTransformDetails
+				.MaxDesiredWidth(125.0f * 3.0f)
+				[
+					SNew(SVerticalBox)
+					+ SVerticalBox::Slot()
+					[
+						SNew(SRotatorInputBox)
+						.AllowSpin(true)
+						.IsEnabled(this, &FRigControlDetails::IsEnabled, ValueType)
+						.bColorAxisLabels(true)
+						.Font(IDetailLayoutBuilder::GetDetailFont())
+						.Roll(this, &FRigControlDetails::GetComponentValueFloat, ValueType, ERigElementDetailsTransformComponent::RotationRoll)
+						.Pitch(this, &FRigControlDetails::GetComponentValueFloat, ValueType, ERigElementDetailsTransformComponent::RotationPitch)
+						.Yaw(this, &FRigControlDetails::GetComponentValueFloat, ValueType, ERigElementDetailsTransformComponent::RotationYaw)
+						.OnRollCommitted(this, &FRigControlDetails::SetComponentValueFloat, ValueType, ERigElementDetailsTransformComponent::RotationRoll)
+						.OnPitchCommitted(this, &FRigControlDetails::SetComponentValueFloat, ValueType, ERigElementDetailsTransformComponent::RotationPitch)
+						.OnYawCommitted(this, &FRigControlDetails::SetComponentValueFloat, ValueType, ERigElementDetailsTransformComponent::RotationYaw)
+					]
+				];
+			}
+			break;
+		}
+		case ERigControlType::Transform:
+		{
+			for(int64 ValueTypeIndex = 0; ValueTypeIndex < ValueTypeEnum->GetMaxEnumValue(); ValueTypeIndex++)
+			{
+				ERigControlValueType ValueType = (ERigControlValueType)ValueTypeIndex;
+
+				if (ValueType == ERigControlValueType::Minimum)
+				{
+					LimitsCategory.AddCustomRow(FText::FromString(TEXT("bLimitTranslation")))
+					.NameContent()
+					[
+						DetailBuilder.GetProperty(TEXT("bLimitTranslation"))->CreatePropertyNameWidget()
+					]
+					.ValueContent()
+					[
+						DetailBuilder.GetProperty(TEXT("bLimitTranslation"))->CreatePropertyValueWidget()
+					];
+
+					LimitsCategory.AddCustomRow(FText::FromString(TEXT("bLimitRotation")))
+					.NameContent()
+					[
+						DetailBuilder.GetProperty(TEXT("bLimitRotation"))->CreatePropertyNameWidget()
+					]
+					.ValueContent()
+					[
+						DetailBuilder.GetProperty(TEXT("bLimitRotation"))->CreatePropertyValueWidget()
+					];
+
+					LimitsCategory.AddCustomRow(FText::FromString(TEXT("bLimitScale")))
+					.NameContent()
+					[
+						DetailBuilder.GetProperty(TEXT("bLimitScale"))->CreatePropertyNameWidget()
+					]
+					.ValueContent()
+					[
+						DetailBuilder.GetProperty(TEXT("bLimitScale"))->CreatePropertyValueWidget()
+					];
+
+					LimitsCategory.AddCustomRow(FText::FromString(TEXT("bDrawLimits")))
+					.NameContent()
+					[
+						DetailBuilder.GetProperty(TEXT("bDrawLimits"))->CreatePropertyNameWidget()
+					]
+					.ValueContent()
+					[
+						DetailBuilder.GetProperty(TEXT("bDrawLimits"))->CreatePropertyValueWidget()
+					];
+				}
+
+				IDetailCategoryBuilder& Category = (ValueType == ERigControlValueType::Minimum || ValueType == ERigControlValueType::Maximum) ? LimitsCategory : ControlCategory;
+				Category.AddCustomRow(ValueTypeEnum->GetDisplayNameTextByValue(ValueTypeIndex))
+				.NameContent()
+				.VAlign(VAlign_Top)
+				[
+					SNew(STextBlock)
+					.Text(ValueTypeEnum->GetDisplayNameTextByValue(ValueTypeIndex))
+					.Font(IDetailLayoutBuilder::GetDetailFont())
+					.IsEnabled(this, &FRigControlDetails::IsEnabled, ValueType)
+				]
+				.ValueContent()
+				.MinDesiredWidth(125.0f * 3.0f) // copied from FComponentTransformDetails
+				.MaxDesiredWidth(125.0f * 3.0f)
+				[
+					SNew(SVerticalBox)
+					+ SVerticalBox::Slot()
+					[
+						SNew(SVectorInputBox)
+						.AllowSpin(true)
+						.IsEnabled(this, &FRigControlDetails::IsEnabled, ValueType)
+						.bColorAxisLabels(true)
+						.Font(IDetailLayoutBuilder::GetDetailFont())
+						.X(this, &FRigControlDetails::GetComponentValueFloat, ValueType, ERigElementDetailsTransformComponent::TranslationX)
+						.Y(this, &FRigControlDetails::GetComponentValueFloat, ValueType, ERigElementDetailsTransformComponent::TranslationY)
+						.Z(this, &FRigControlDetails::GetComponentValueFloat, ValueType, ERigElementDetailsTransformComponent::TranslationZ)
+						.OnXCommitted(this, &FRigControlDetails::SetComponentValueFloat, ValueType, ERigElementDetailsTransformComponent::TranslationX)
+						.OnYCommitted(this, &FRigControlDetails::SetComponentValueFloat, ValueType, ERigElementDetailsTransformComponent::TranslationY)
+						.OnZCommitted(this, &FRigControlDetails::SetComponentValueFloat, ValueType, ERigElementDetailsTransformComponent::TranslationZ)
+					]
+					+ SVerticalBox::Slot()
+					[
+						SNew(SRotatorInputBox)
+						.AllowSpin(true)
+						.IsEnabled(this, &FRigControlDetails::IsEnabled, ValueType)
+						.bColorAxisLabels(true)
+						.Font(IDetailLayoutBuilder::GetDetailFont())
+						.Roll(this, &FRigControlDetails::GetComponentValueFloat, ValueType, ERigElementDetailsTransformComponent::RotationRoll)
+						.Pitch(this, &FRigControlDetails::GetComponentValueFloat, ValueType, ERigElementDetailsTransformComponent::RotationPitch)
+						.Yaw(this, &FRigControlDetails::GetComponentValueFloat, ValueType, ERigElementDetailsTransformComponent::RotationYaw)
+						.OnRollCommitted(this, &FRigControlDetails::SetComponentValueFloat, ValueType, ERigElementDetailsTransformComponent::RotationRoll)
+						.OnPitchCommitted(this, &FRigControlDetails::SetComponentValueFloat, ValueType, ERigElementDetailsTransformComponent::RotationPitch)
+						.OnYawCommitted(this, &FRigControlDetails::SetComponentValueFloat, ValueType, ERigElementDetailsTransformComponent::RotationYaw)
+					]
+					+ SVerticalBox::Slot()
+					[
+						SNew(SVectorInputBox)
+						.AllowSpin(true)
+						.IsEnabled(this, &FRigControlDetails::IsEnabled, ValueType)
+						.bColorAxisLabels(true)
+						.Font(IDetailLayoutBuilder::GetDetailFont())
+						.X(this, &FRigControlDetails::GetComponentValueFloat, ValueType, ERigElementDetailsTransformComponent::ScaleX)
+						.Y(this, &FRigControlDetails::GetComponentValueFloat, ValueType, ERigElementDetailsTransformComponent::ScaleY)
+						.Z(this, &FRigControlDetails::GetComponentValueFloat, ValueType, ERigElementDetailsTransformComponent::ScaleZ)
+						.OnXCommitted(this, &FRigControlDetails::SetComponentValueFloat, ValueType, ERigElementDetailsTransformComponent::ScaleX)
+						.OnYCommitted(this, &FRigControlDetails::SetComponentValueFloat, ValueType, ERigElementDetailsTransformComponent::ScaleY)
+						.OnZCommitted(this, &FRigControlDetails::SetComponentValueFloat, ValueType, ERigElementDetailsTransformComponent::ScaleZ)
+					]
+				];
+			}
+			break;
+		}
+		case ERigControlType::TransformNoScale:
+		{
+			for(int64 ValueTypeIndex = 0; ValueTypeIndex < ValueTypeEnum->GetMaxEnumValue(); ValueTypeIndex++)
+			{
+				ERigControlValueType ValueType = (ERigControlValueType)ValueTypeIndex;
+
+				if (ValueType == ERigControlValueType::Minimum)
+				{
+					LimitsCategory.AddCustomRow(FText::FromString(TEXT("bLimitTranslation")))
+					.NameContent()
+					[
+						DetailBuilder.GetProperty(TEXT("bLimitTranslation"))->CreatePropertyNameWidget()
+					]
+					.ValueContent()
+					[
+						DetailBuilder.GetProperty(TEXT("bLimitTranslation"))->CreatePropertyValueWidget()
+					];
+
+					LimitsCategory.AddCustomRow(FText::FromString(TEXT("bLimitRotation")))
+					.NameContent()
+					[
+						DetailBuilder.GetProperty(TEXT("bLimitRotation"))->CreatePropertyNameWidget()
+					]
+					.ValueContent()
+					[
+						DetailBuilder.GetProperty(TEXT("bLimitRotation"))->CreatePropertyValueWidget()
+					];
+
+					LimitsCategory.AddCustomRow(FText::FromString(TEXT("bDrawLimits")))
+					.NameContent()
+					[
+						DetailBuilder.GetProperty(TEXT("bDrawLimits"))->CreatePropertyNameWidget()
+					]
+					.ValueContent()
+					[
+						DetailBuilder.GetProperty(TEXT("bDrawLimits"))->CreatePropertyValueWidget()
+					];
+				}
+
+				IDetailCategoryBuilder& Category = (ValueType == ERigControlValueType::Minimum || ValueType == ERigControlValueType::Maximum) ? LimitsCategory : ControlCategory;
+				Category.AddCustomRow(ValueTypeEnum->GetDisplayNameTextByValue(ValueTypeIndex))
+				.NameContent()
+				.VAlign(VAlign_Top)
+				[
+					SNew(STextBlock)
+					.Text(ValueTypeEnum->GetDisplayNameTextByValue(ValueTypeIndex))
+					.Font(IDetailLayoutBuilder::GetDetailFont())
+					.IsEnabled(this, &FRigControlDetails::IsEnabled, ValueType)
+				]
+				.ValueContent()
+				.MinDesiredWidth(125.0f * 3.0f) // copied from FComponentTransformDetails
+				.MaxDesiredWidth(125.0f * 3.0f)
+				[
+					SNew(SVerticalBox)
+					+ SVerticalBox::Slot()
+					[
+						SNew(SVectorInputBox)
+						.AllowSpin(true)
+						.IsEnabled(this, &FRigControlDetails::IsEnabled, ValueType)
+						.bColorAxisLabels(true)
+						.Font(IDetailLayoutBuilder::GetDetailFont())
+						.X(this, &FRigControlDetails::GetComponentValueFloat, ValueType, ERigElementDetailsTransformComponent::TranslationX)
+						.Y(this, &FRigControlDetails::GetComponentValueFloat, ValueType, ERigElementDetailsTransformComponent::TranslationY)
+						.Z(this, &FRigControlDetails::GetComponentValueFloat, ValueType, ERigElementDetailsTransformComponent::TranslationZ)
+						.OnXCommitted(this, &FRigControlDetails::SetComponentValueFloat, ValueType, ERigElementDetailsTransformComponent::TranslationX)
+						.OnYCommitted(this, &FRigControlDetails::SetComponentValueFloat, ValueType, ERigElementDetailsTransformComponent::TranslationY)
+						.OnZCommitted(this, &FRigControlDetails::SetComponentValueFloat, ValueType, ERigElementDetailsTransformComponent::TranslationZ)
+					]
+					+ SVerticalBox::Slot()
+					[
+						SNew(SRotatorInputBox)
+						.AllowSpin(true)
+						.IsEnabled(this, &FRigControlDetails::IsEnabled, ValueType)
+						.bColorAxisLabels(true)
+						.Font(IDetailLayoutBuilder::GetDetailFont())
+						.Roll(this, &FRigControlDetails::GetComponentValueFloat, ValueType, ERigElementDetailsTransformComponent::RotationRoll)
+						.Pitch(this, &FRigControlDetails::GetComponentValueFloat, ValueType, ERigElementDetailsTransformComponent::RotationPitch)
+						.Yaw(this, &FRigControlDetails::GetComponentValueFloat, ValueType, ERigElementDetailsTransformComponent::RotationYaw)
+						.OnRollCommitted(this, &FRigControlDetails::SetComponentValueFloat, ValueType, ERigElementDetailsTransformComponent::RotationRoll)
+						.OnPitchCommitted(this, &FRigControlDetails::SetComponentValueFloat, ValueType, ERigElementDetailsTransformComponent::RotationPitch)
+						.OnYawCommitted(this, &FRigControlDetails::SetComponentValueFloat, ValueType, ERigElementDetailsTransformComponent::RotationYaw)
+					]
+				];
+			}
+			break;
+		}
+		default:
+		{
+			ensure(false);
+			break;
+		}
+	}
+
+
+	switch (Control.ControlType)
+	{
+		case ERigControlType::Float:
+		case ERigControlType::Vector2D:
+		case ERigControlType::Position:
+		case ERigControlType::Scale:
+		case ERigControlType::Rotator:
+		case ERigControlType::Transform:
+		case ERigControlType::TransformNoScale:
+		{
+			DetailBuilder.HideProperty(TEXT("bGizmoEnabled"));
+			DetailBuilder.HideProperty(TEXT("GizmoName"));
+
+			GizmoCategory.AddCustomRow(FText::FromString(TEXT("bGizmoEnabled")))
+			.NameContent()
+			[
+				DetailBuilder.GetProperty(TEXT("bGizmoEnabled"))->CreatePropertyNameWidget()
+			]
+			.ValueContent()
+			[
+				DetailBuilder.GetProperty(TEXT("bGizmoEnabled"))->CreatePropertyValueWidget()
+			];
+
+			GizmoCategory.AddCustomRow(FText::FromString(TEXT("GizmoName")))
+			.NameContent()
+			[
+				SNew(STextBlock)
+				.Text(FText::FromString(TEXT("Gizmo Name")))
+				.Font(IDetailLayoutBuilder::GetDetailFont())
+				.IsEnabled(this, &FRigControlDetails::IsGizmoEnabled)
+			]
+			.ValueContent()
+			[
+				SNew(SControlRigGizmoNameList, &Control, BlueprintBeingCustomized)
+				.OnGetNameListContent(this, &FRigControlDetails::GetGizmoNameList)
+				.IsEnabled(this, &FRigControlDetails::IsGizmoEnabled)
+			];
+			break;
+		}
+		default:
+		{
+			DetailBuilder.HideProperty(TEXT("bGizmoEnabled"));
+			DetailBuilder.HideProperty(TEXT("GizmoName"));
+			DetailBuilder.HideProperty(TEXT("GizmoTransform"));
+			DetailBuilder.HideProperty(TEXT("GizmoColor"));
+			break;
+		}
+	}
+
+	switch (Control.ControlType)
+	{
+		case ERigControlType::Float:
+		case ERigControlType::Vector2D:
+		{
+			break;
+		}
+		default:
+		{
+			DetailBuilder.HideProperty(TEXT("PrimaryAxis"));
+			break;
+		}
+	}
 }
 
-TOptional<float> FRigControlDetails::GetComponentValue(bool bInitial, ERigElementDetailsTransformComponent Component) const
+ECheckBoxState FRigControlDetails::GetComponentValueBool(bool bInitial) const
 {
 	if (ContainerBeingCustomized != nullptr && ElementKeyBeingCustomized)
 	{
@@ -377,10 +1206,10 @@ TOptional<float> FRigControlDetails::GetComponentValue(bool bInitial, ERigElemen
 
 			switch (Control.ControlType)
 			{
-				case ERigControlType::Transform:
+				case ERigControlType::Bool:
 				{
-					FTransform Transform = bInitial ? Control.InitialValue.Get<FTransform>() : Control.Value.Get<FTransform>();
-					return GetTransformComponent(Transform, Component);
+					bool Value = bInitial ? Control.InitialValue.Get<bool>() : Control.Value.Get<bool>();
+					return Value ? ECheckBoxState::Checked : ECheckBoxState::Unchecked;
 				}
 				default:
 				{
@@ -390,10 +1219,10 @@ TOptional<float> FRigControlDetails::GetComponentValue(bool bInitial, ERigElemen
 			}
 		}
 	}
-	return 0.f;
+	return ECheckBoxState::Unchecked;
 }
 
-void FRigControlDetails::SetComponentValue(float InNewValue, ETextCommit::Type InCommitType, bool bInitial, ERigElementDetailsTransformComponent Component)
+void FRigControlDetails::SetComponentValueBool(ECheckBoxState InNewValue, bool bInitial)
 {
 	if (ContainerBeingCustomized != nullptr && ElementKeyBeingCustomized)
 	{
@@ -405,11 +1234,9 @@ void FRigControlDetails::SetComponentValue(float InNewValue, ETextCommit::Type I
 
 			switch (Control.ControlType)
 			{
-				case ERigControlType::Transform:
+				case ERigControlType::Bool:
 				{
-					FTransform Transform = Value.Get<FTransform>();
-					SetTransformComponent(Transform, Component, InNewValue);
-					Value.Set<FTransform>(Transform);
+					Value.Set<bool>(InNewValue == ECheckBoxState::Checked);
 					break;
 				}
 				default:
@@ -439,11 +1266,354 @@ void FRigControlDetails::SetComponentValue(float InNewValue, ETextCommit::Type I
 	}
 }
 
+TOptional<float> FRigControlDetails::GetComponentValueFloat(ERigControlValueType InValueType, ERigElementDetailsTransformComponent Component) const
+{
+	if (ContainerBeingCustomized != nullptr && ElementKeyBeingCustomized)
+	{
+		int32 Index = ContainerBeingCustomized->ControlHierarchy.GetIndex(ElementKeyBeingCustomized.Name);
+		if (Index != INDEX_NONE)
+		{
+			FRigControl Control = ContainerBeingCustomized->ControlHierarchy[ElementKeyBeingCustomized.Name];
+			if (InValueType == ERigControlValueType::Current)
+			{
+				if (UControlRig* DebuggedRig = Cast<UControlRig>(BlueprintBeingCustomized->GetObjectBeingDebugged()))
+				{
+					Control = DebuggedRig->GetControlHierarchy()[Index];
+				}
+			}
+
+			switch (Control.ControlType)
+			{
+				case ERigControlType::Float:
+				{
+					return Control.GetValue(InValueType).Get<float>();
+				}
+				case ERigControlType::Vector2D:
+				{
+					switch (Component)
+					{
+						case ERigElementDetailsTransformComponent::TranslationX:
+						{
+							return Control.GetValue(InValueType).Get<FVector2D>().X;
+						}
+						case ERigElementDetailsTransformComponent::TranslationY:
+						{
+							return Control.GetValue(InValueType).Get<FVector2D>().Y;
+						}
+						default:
+						{
+							break;
+						}
+					}
+					break;
+				}
+				case ERigControlType::Position:
+				{
+					switch (Component)
+					{
+						case ERigElementDetailsTransformComponent::TranslationX:
+						{
+							return Control.GetValue(InValueType).Get<FVector>().X;
+						}
+						case ERigElementDetailsTransformComponent::TranslationY:
+						{
+							return Control.GetValue(InValueType).Get<FVector>().Y;
+						}
+						case ERigElementDetailsTransformComponent::TranslationZ:
+						{
+							return Control.GetValue(InValueType).Get<FVector>().Z;
+						}
+						default:
+						{
+							break;
+						}
+					}
+					break;
+				}
+				case ERigControlType::Scale:
+				{
+					switch (Component)
+					{
+						case ERigElementDetailsTransformComponent::ScaleX:
+						{
+							return Control.GetValue(InValueType).Get<FVector>().X;
+						}
+						case ERigElementDetailsTransformComponent::ScaleY:
+						{
+							return Control.GetValue(InValueType).Get<FVector>().Y;
+						}
+						case ERigElementDetailsTransformComponent::ScaleZ:
+						{
+							return Control.GetValue(InValueType).Get<FVector>().Z;
+						}
+						default:
+						{
+							break;
+						}
+					}
+					break;
+				}
+				case ERigControlType::Rotator:
+				{
+					switch (Component)
+					{
+						case ERigElementDetailsTransformComponent::RotationPitch:
+						{
+							return Control.GetValue(InValueType).Get<FRotator>().Pitch;
+						}
+						case ERigElementDetailsTransformComponent::RotationYaw:
+						{
+							return Control.GetValue(InValueType).Get<FRotator>().Yaw;
+						}
+						case ERigElementDetailsTransformComponent::RotationRoll:
+						{
+							return Control.GetValue(InValueType).Get<FRotator>().Roll;
+						}
+						default:
+						{
+							break;
+						}
+					}
+					break;
+				}
+				case ERigControlType::Transform:
+				{
+					FTransform Transform = Control.GetValue(InValueType).Get<FTransform>();
+					return GetTransformComponent(Transform, Component);
+				}
+				case ERigControlType::TransformNoScale:
+				{
+					FTransform Transform = Control.GetValue(InValueType).Get<FTransformNoScale>();
+					return GetTransformComponent(Transform, Component);
+				}
+				case ERigControlType::Bool:
+				{
+					break;
+				}
+				default:
+				{
+					ensure(false);
+					break;
+				}
+			}
+		}
+	}
+	return 0.f;
+}
+
+void FRigControlDetails::SetComponentValueFloat(float InNewValue, ETextCommit::Type InCommitType, ERigControlValueType InValueType, ERigElementDetailsTransformComponent Component)
+{
+	SetComponentValueFloat(InNewValue, InValueType, Component);
+}
+
+void FRigControlDetails::SetComponentValueFloat(float InNewValue, ERigControlValueType InValueType, ERigElementDetailsTransformComponent Component)
+{
+	if (ContainerBeingCustomized != nullptr && ElementKeyBeingCustomized)
+	{
+		int32 Index = ContainerBeingCustomized->ControlHierarchy.GetIndex(ElementKeyBeingCustomized.Name);
+		if (Index != INDEX_NONE)
+		{
+			const FRigControl& Control = ContainerBeingCustomized->ControlHierarchy[ElementKeyBeingCustomized.Name];
+			FRigControlValue Value = Control.GetValue(InValueType);
+
+			switch (Control.ControlType)
+			{
+				case ERigControlType::Float:
+				{
+					Value.Set<float>(InNewValue);
+					break;
+				}
+				case ERigControlType::Vector2D:
+				{
+					FVector2D Local = Value.Get<FVector2D>();
+					switch (Component)
+					{
+						case ERigElementDetailsTransformComponent::TranslationX:
+						{
+							Local.X = InNewValue;
+							break;
+						}
+						case ERigElementDetailsTransformComponent::TranslationY:
+						{
+							Local.Y = InNewValue;
+							break;
+						}
+						default:
+						{
+							break;
+						}
+					}
+					Value.Set<FVector2D>(Local);
+					break;
+				}
+				case ERigControlType::Position:
+				{
+					FVector Local = Value.Get<FVector>();
+					switch (Component)
+					{
+						case ERigElementDetailsTransformComponent::TranslationX:
+						{
+							Local.X = InNewValue;
+							break;
+						}
+						case ERigElementDetailsTransformComponent::TranslationY:
+						{
+							Local.Y = InNewValue;
+							break;
+						}
+						case ERigElementDetailsTransformComponent::TranslationZ:
+						{
+							Local.Z = InNewValue;
+							break;
+						}
+						default:
+						{
+							break;
+						}
+					}
+					Value.Set<FVector>(Local);
+					break;
+				}
+				case ERigControlType::Scale:
+				{
+					FVector Local = Value.Get<FVector>();
+					switch (Component)
+					{
+						case ERigElementDetailsTransformComponent::ScaleX:
+						{
+							Local.X = InNewValue;
+							break;
+						}
+						case ERigElementDetailsTransformComponent::ScaleY:
+						{
+							Local.Y = InNewValue;
+							break;
+						}
+						case ERigElementDetailsTransformComponent::ScaleZ:
+						{
+							Local.Z = InNewValue;
+							break;
+						}
+						default:
+						{
+							break;
+						}
+					}
+					Value.Set<FVector>(Local);
+					break;
+				}
+				case ERigControlType::Rotator:
+				{
+					FRotator Local = Value.Get<FRotator>();
+					switch (Component)
+					{
+						case ERigElementDetailsTransformComponent::RotationPitch:
+						{
+							Local.Pitch = InNewValue;
+							break;
+						}
+						case ERigElementDetailsTransformComponent::RotationYaw:
+						{
+							Local.Yaw = InNewValue;
+							break;
+						}
+						case ERigElementDetailsTransformComponent::RotationRoll:
+						{
+							Local.Roll = InNewValue;
+							break;
+						}
+						default:
+						{
+							break;
+						}
+					}
+					Value.Set<FRotator>(Local);
+					break;
+				}
+				case ERigControlType::Transform:
+				{
+					FTransform Transform = Value.Get<FTransform>();
+					SetTransformComponent(Transform, Component, InNewValue);
+					Value.Set<FTransform>(Transform);
+					break;
+				}
+				case ERigControlType::TransformNoScale:
+				{
+					FTransform Transform = Value.Get<FTransformNoScale>();
+					SetTransformComponent(Transform, Component, InNewValue);
+					Value.Set<FTransformNoScale>(Transform);
+					break;
+				}
+				case ERigControlType::Bool:
+				{
+					return;
+				}
+				default:
+				{
+					ensure(false);
+					break;
+				}
+			}
+
+			ContainerBeingCustomized->ControlHierarchy.SetValue(ElementKeyBeingCustomized.Name, Value, InValueType);
+			if (UControlRig* DebuggedRig = Cast<UControlRig>(BlueprintBeingCustomized->GetObjectBeingDebugged()))
+			{
+				((FRigControlHierarchy*)&DebuggedRig->GetControlHierarchy())->SetValue(ElementKeyBeingCustomized.Name, Value, InValueType);
+			}
+		}
+	}
+}
+
+bool FRigControlDetails::IsGizmoEnabled() const
+{
+	if (ContainerBeingCustomized != nullptr && ElementKeyBeingCustomized)
+	{
+		int32 Index = ContainerBeingCustomized->ControlHierarchy.GetIndex(ElementKeyBeingCustomized.Name);
+		if (Index != INDEX_NONE)
+		{
+			const FRigControl& Control = ContainerBeingCustomized->ControlHierarchy[ElementKeyBeingCustomized.Name];
+			return Control.bGizmoEnabled;
+		}
+	}
+	return false;
+}
+
+bool FRigControlDetails::IsEnabled(ERigControlValueType InValueType) const
+{
+	switch (InValueType)
+	{
+		case ERigControlValueType::Minimum:
+		case ERigControlValueType::Maximum:
+		{
+			if (ContainerBeingCustomized != nullptr && ElementKeyBeingCustomized)
+			{
+				int32 Index = ContainerBeingCustomized->ControlHierarchy.GetIndex(ElementKeyBeingCustomized.Name);
+				if (Index != INDEX_NONE)
+				{
+					const FRigControl& Control = ContainerBeingCustomized->ControlHierarchy[ElementKeyBeingCustomized.Name];
+					return Control.bLimitTranslation || Control.bLimitRotation || Control.bLimitScale;
+				}
+			}
+			return false;
+		}
+		default:
+		{
+			break;
+		}
+	}
+	return true;
+}
+
 const TArray<TSharedPtr<FString>>& FRigControlDetails::GetGizmoNameList() const
 {
 	return GizmoNameList;
 }
 
+const TArray<TSharedPtr<FString>>& FRigControlDetails::GetControlTypeList() const
+{
+	return ControlTypeList;
+}
+
 void FRigSpaceDetails::CustomizeDetails(IDetailLayoutBuilder& DetailBuilder)
 {
 	FRigElementDetails::CustomizeDetails(DetailBuilder);
diff --git a/Engine/Plugins/Experimental/ControlRig/Source/ControlRigEditor/Private/ControlRigElementDetails.h b/Engine/Plugins/Experimental/ControlRig/Source/ControlRigEditor/Private/ControlRigElementDetails.h
--- a/Engine/Plugins/Experimental/ControlRig/Source/ControlRigEditor/Private/ControlRigElementDetails.h
+++ b/Engine/Plugins/Experimental/ControlRig/Source/ControlRigEditor/Private/ControlRigElementDetails.h
@@ -8,9 +8,39 @@
 #include "Rigs/RigHierarchyContainer.h"
 #include "ControlRig.h"
 #include "ControlRigBlueprint.h"
+#include "Graph/ControlRigGraph.h"
+#include "Graph/SControlRigGraphPinNameListValueWidget.h"
+#include "Styling/SlateTypes.h"
+#include "IPropertyUtilities.h"
 
 class IPropertyHandle;
 
+class FRigUnitDetails : public IDetailCustomization
+{
+public:
+
+	static TSharedRef<IDetailCustomization> MakeInstance()
+	{
+		return MakeShareable(new FRigUnitDetails);
+	}
+
+	/** IDetailCustomization interface */
+	virtual void CustomizeDetails(IDetailLayoutBuilder& DetailBuilder) override;
+
+protected:
+
+	TSharedRef<SWidget> MakeNameListItemWidget(TSharedPtr<FString> InItem);
+	FText GetNameListText(TSharedPtr<FStructOnScope> InStructOnScope, FNameProperty* InProperty) const;
+	TSharedPtr<FString> GetCurrentlySelectedItem(TSharedPtr<FStructOnScope> InStructOnScope, FNameProperty* InProperty, const TArray<TSharedPtr<FString>>* InNameList) const;
+	void SetNameListText(const FText& NewTypeInValue, ETextCommit::Type /*CommitInfo*/, TSharedPtr<FStructOnScope> InStructOnScope, FNameProperty* InProperty, TSharedRef<IPropertyUtilities> PropertyUtilities);
+	void OnNameListChanged(TSharedPtr<FString> NewSelection, ESelectInfo::Type SelectInfo, TSharedPtr<FStructOnScope> InStructOnScope, FNameProperty* InProperty, TSharedRef<IPropertyUtilities> PropertyUtilities);
+	void OnNameListComboBox(TSharedPtr<FStructOnScope> InStructOnScope, FNameProperty* InProperty, const TArray<TSharedPtr<FString>>* InNameList);
+
+	UControlRigBlueprint* BlueprintBeingCustomized;
+	UControlRigGraph* GraphBeingCustomized;
+	TMap<FName, TSharedPtr<SControlRigGraphPinNameListValueWidget>> NameListWidgets;
+};
+
 UENUM()
 enum class ERigElementDetailsTransformComponent : uint8
 {
@@ -74,13 +104,20 @@ public:
 	/** IDetailCustomization interface */
 	virtual void CustomizeDetails(IDetailLayoutBuilder& DetailBuilder) override;
 
-	TOptional<float> GetComponentValue(bool bInitial, ERigElementDetailsTransformComponent Component) const;
-	void SetComponentValue(float InNewValue, ETextCommit::Type InCommitType, bool bInitial, ERigElementDetailsTransformComponent Component);
+	ECheckBoxState GetComponentValueBool(bool bInitial) const;
+	void SetComponentValueBool(ECheckBoxState InNewValue, bool bInitial);
+	TOptional<float> GetComponentValueFloat(ERigControlValueType InValueType, ERigElementDetailsTransformComponent Component) const;
+	void SetComponentValueFloat(float InNewValue, ETextCommit::Type InCommitType, ERigControlValueType InValueType, ERigElementDetailsTransformComponent Component);
+	void SetComponentValueFloat(float InNewValue, ERigControlValueType InValueType, ERigElementDetailsTransformComponent Component);
+	bool IsGizmoEnabled() const;
+	bool IsEnabled(ERigControlValueType InValueType) const;
 
 	const TArray<TSharedPtr<FString>>& GetGizmoNameList() const;
+	const TArray<TSharedPtr<FString>>& GetControlTypeList() const;
 
 private:
 	TArray<TSharedPtr<FString>> GizmoNameList;
+	static TArray<TSharedPtr<FString>> ControlTypeList;
 };
 
 class FRigSpaceDetails : public FRigElementDetails
diff --git a/Engine/Plugins/Experimental/ControlRig/Source/ControlRigEditor/Private/ControlRigHierarchyCommands.cpp b/Engine/Plugins/Experimental/ControlRig/Source/ControlRigEditor/Private/ControlRigHierarchyCommands.cpp
--- a/Engine/Plugins/Experimental/ControlRig/Source/ControlRigEditor/Private/ControlRigHierarchyCommands.cpp
+++ b/Engine/Plugins/Experimental/ControlRig/Source/ControlRigEditor/Private/ControlRigHierarchyCommands.cpp
@@ -14,14 +14,15 @@ void FControlRigHierarchyCommands::RegisterCommands()
 	UI_COMMAND(RenameItem, "Rename", "Rename the selected item.", EUserInterfaceActionType::Button, FInputChord(EKeys::F2));
 	UI_COMMAND(CopyItems, "Copy", "Copy the selected items.", EUserInterfaceActionType::Button, FInputChord(EKeys::C, EModifierKey::Control));
 	UI_COMMAND(PasteItems, "Paste", "Paste the selected items.", EUserInterfaceActionType::Button, FInputChord(EKeys::V, EModifierKey::Control));
-	UI_COMMAND(PasteLocalTransforms, "Paste Local Xfo", "Paste the local transforms.", EUserInterfaceActionType::Button, FInputChord());
-	UI_COMMAND(PasteGlobalTransforms, "Paste Global Xfo", "Paste the global transforms.", EUserInterfaceActionType::Button, FInputChord(EKeys::V, EModifierKey::Control | EModifierKey::Shift));
+	UI_COMMAND(PasteLocalTransforms, "Paste Local Transform", "Paste the local transforms.", EUserInterfaceActionType::Button, FInputChord());
+	UI_COMMAND(PasteGlobalTransforms, "Paste Global Transform", "Paste the global transforms.", EUserInterfaceActionType::Button, FInputChord(EKeys::V, EModifierKey::Control | EModifierKey::Shift));
 	UI_COMMAND(SetInitialTransformFromClosestBone, "Set Initial Transform from Closest Bone", "Find the Closest Bone to Initial Transform of the Selected Bones", EUserInterfaceActionType::Button, FInputChord());
 	UI_COMMAND(SetInitialTransformFromCurrentTransform, "Set Initial Transform from Current", "Save the Current Transform To Initial Transform of the Selected Bones", EUserInterfaceActionType::Button, FInputChord());
 	UI_COMMAND(ResetTransform, "Reset Transform", "Reset the Transform", EUserInterfaceActionType::Button, FInputChord(EKeys::G, EModifierKey::Control));
-	UI_COMMAND(ResetInitialTransform, "Reset Initial Transform", "Reset the Initial Transform to the Identity", EUserInterfaceActionType::Button, FInputChord());
 	UI_COMMAND(ResetSpace, "Reset Space", "Resets or injects a Space below the Control", EUserInterfaceActionType::Button, FInputChord());
 	UI_COMMAND(FrameSelection, "Frame Selection", "Expands and frames the selection in the tree", EUserInterfaceActionType::Button, FInputChord(EKeys::F));
+	UI_COMMAND(ControlBoneTransform, "Control Bone Transform", "Sets the bone transform using a gizmo", EUserInterfaceActionType::Button, FInputChord());
+	UI_COMMAND(ControlSpaceTransform, "Control Space Transform", "Sets the space transform using a gizmo", EUserInterfaceActionType::Button, FInputChord());
 }
 
 #undef LOCTEXT_NAMESPACE
diff --git a/Engine/Plugins/Experimental/ControlRig/Source/ControlRigEditor/Private/ControlRigHierarchyCommands.h b/Engine/Plugins/Experimental/ControlRig/Source/ControlRigEditor/Private/ControlRigHierarchyCommands.h
--- a/Engine/Plugins/Experimental/ControlRig/Source/ControlRigEditor/Private/ControlRigHierarchyCommands.h
+++ b/Engine/Plugins/Experimental/ControlRig/Source/ControlRigEditor/Private/ControlRigHierarchyCommands.h
@@ -42,18 +42,15 @@ public:
 	/** Paste the selected items. */
 	TSharedPtr< FUICommandInfo > PasteItems;
 
-	/** Paste Local Xfo", "Paste the local transforms. */
+	/** Paste Local Transform", "Paste the local transforms. */
 	TSharedPtr< FUICommandInfo > PasteLocalTransforms;
 
-	/** Paste Global Xfo", "Paste the global transforms. */
+	/** Paste Global Transform", "Paste the global transforms. */
 	TSharedPtr< FUICommandInfo > PasteGlobalTransforms;
 
 	/* Reset transform */
 	TSharedPtr<FUICommandInfo> ResetTransform;
 
-	/* Reset initial transform */
-	TSharedPtr<FUICommandInfo> ResetInitialTransform;
-
 	/* Reset space */
 	TSharedPtr<FUICommandInfo> ResetSpace;
 
@@ -66,6 +63,12 @@ public:
 	/* frames the selection in the tree */
 	TSharedPtr<FUICommandInfo> FrameSelection;
 
+	/* sets the bone transform using a gizmo */
+	TSharedPtr<FUICommandInfo> ControlBoneTransform;
+
+	/* sets the space transform using a gizmo */
+	TSharedPtr<FUICommandInfo> ControlSpaceTransform;
+
 	/**
 	 * Initialize commands
 	 */
diff --git a/Engine/Plugins/Experimental/ControlRig/Source/ControlRigEditor/Private/ControlRigThumbnailRenderer.cpp b/Engine/Plugins/Experimental/ControlRig/Source/ControlRigEditor/Private/ControlRigThumbnailRenderer.cpp
new file mode 100644
--- /dev/null
+++ b/Engine/Plugins/Experimental/ControlRig/Source/ControlRigEditor/Private/ControlRigThumbnailRenderer.cpp
@@ -0,0 +1,125 @@
+// Copyright Epic Games, Inc. All Rights Reserved.
+
+#include "ControlRigThumbnailRenderer.h"
+#include "ThumbnailHelpers.h"
+#include "Engine/StaticMesh.h"
+#include "Components/StaticMeshComponent.h"
+#include "Engine/SkeletalMesh.h"
+#include "Animation/SkeletalMeshActor.h"
+#include "Components/SkeletalMeshComponent.h"
+#include "Materials/MaterialInstanceDynamic.h"
+#include "ControlRig.h"
+
+UControlRigThumbnailRenderer::UControlRigThumbnailRenderer(const FObjectInitializer& ObjectInitializer)
+	: Super(ObjectInitializer)
+{
+	RigBlueprint = nullptr;
+}
+
+bool UControlRigThumbnailRenderer::CanVisualizeAsset(UObject* Object)
+{
+	return Cast<UControlRigBlueprint>(Object) != nullptr;
+}
+
+void UControlRigThumbnailRenderer::Draw(UObject* Object, int32 X, int32 Y, uint32 Width, uint32 Height, FRenderTarget* RenderTarget, FCanvas* Canvas)
+{
+	RigBlueprint = nullptr;
+
+	if (UControlRigBlueprint* InRigBlueprint = Cast<UControlRigBlueprint>(Object))
+	{
+		USkeletalMesh* SkeletalMesh = Cast<USkeletalMesh>(InRigBlueprint->GetPreviewMesh());
+		if (SkeletalMesh != nullptr)
+		{
+			RigBlueprint = InRigBlueprint;
+			Super::Draw(SkeletalMesh, X, Y, Width, Height, RenderTarget, Canvas);
+
+			for (auto Pair : GizmoActors)
+			{
+				if (Pair.Value && Pair.Value->GetOuter())
+				{
+					Pair.Value->Rename(nullptr, GetTransientPackage());
+					Pair.Value->MarkPendingKill();
+				}
+			}
+			GizmoActors.Reset();
+		}
+	}
+}
+
+void UControlRigThumbnailRenderer::AddAdditionalPreviewSceneContent(UObject* Object, UWorld* PreviewWorld)
+{
+	if (ThumbnailScene && ThumbnailScene->GetPreviewActor() && RigBlueprint && RigBlueprint->GizmoLibrary && RigBlueprint->GeneratedClass)
+	{
+		UControlRig* ControlRig = nullptr;
+
+		// reuse the current control rig if possible
+		UControlRig* CDO = Cast<UControlRig>(RigBlueprint->GeneratedClass->GetDefaultObject(true /* create if needed */));
+
+		TArray<UObject*> ArchetypeInstances;
+		CDO->GetArchetypeInstances(ArchetypeInstances);
+		for (UObject* ArchetypeInstance : ArchetypeInstances)
+		{
+			ControlRig = Cast<UControlRig>(ArchetypeInstance);
+			break;
+		}
+
+		bool bRequiresMarkPendingkill = false;
+		if (ControlRig == nullptr)
+		{
+			ControlRig = NewObject<UControlRig>(GetTransientPackage(), RigBlueprint->GeneratedClass);
+			ControlRig->ExecutionType = ERigExecutionType::Editing;
+			ControlRig->Initialize();
+			ControlRig->Execute(EControlRigState::Update);
+			bRequiresMarkPendingkill = true;
+		}
+
+		FTransform ComponentToWorld = ThumbnailScene->GetPreviewActor()->GetSkeletalMeshComponent()->GetComponentToWorld();
+
+		for (const FRigControl& Control : ControlRig->GetControlHierarchy())
+		{
+			switch (Control.ControlType)
+			{
+				case ERigControlType::Float:
+				case ERigControlType::Vector2D:
+				case ERigControlType::Position:
+				case ERigControlType::Scale:
+				case ERigControlType::Rotator:
+				case ERigControlType::Transform:
+				case ERigControlType::TransformNoScale:
+				{
+					if (const FControlRigGizmoDefinition* GizmoDef = RigBlueprint->GizmoLibrary->GetGizmoByName(Control.GizmoName))
+					{
+						FTransform GizmoTransform = Control.GizmoTransform * (GizmoDef->Transform * ControlRig->GetControlGlobalTransform(Control.Name)) * ComponentToWorld;
+						UStaticMesh* StaticMesh = GizmoDef->StaticMesh.LoadSynchronous();
+
+						FActorSpawnParameters SpawnInfo;
+						SpawnInfo.SpawnCollisionHandlingOverride = ESpawnActorCollisionHandlingMethod::AlwaysSpawn;
+						SpawnInfo.bNoFail = true;
+						SpawnInfo.ObjectFlags = RF_Transient;
+						AStaticMeshActor* GizmoActor = PreviewWorld->SpawnActor<AStaticMeshActor>(SpawnInfo);
+						GizmoActor->SetActorEnableCollision(false);
+
+						UMaterialInstanceDynamic* MaterialInstance = UMaterialInstanceDynamic::Create(RigBlueprint->GizmoLibrary->DefaultMaterial.LoadSynchronous(), GizmoActor);
+						MaterialInstance->SetVectorParameterValue(RigBlueprint->GizmoLibrary->MaterialColorParameter, FVector(Control.GizmoColor));
+						GizmoActor->GetStaticMeshComponent()->SetMaterial(0, MaterialInstance);
+
+						GizmoActors.Add(Control.Name, GizmoActor);
+
+						GizmoActor->GetStaticMeshComponent()->SetStaticMesh(StaticMesh);
+						GizmoActor->SetActorTransform(GizmoTransform);
+					}
+					break;
+				}
+				default:
+				{
+					break;
+				}
+			}
+		}
+
+		if (bRequiresMarkPendingkill)
+		{
+			ControlRig->MarkPendingKill();
+		}
+	}
+}
\ No newline at end of file
diff --git a/Engine/Plugins/Experimental/ControlRig/Source/ControlRigEditor/Private/ControlRigVariableDetailsCustomization.cpp b/Engine/Plugins/Experimental/ControlRig/Source/ControlRigEditor/Private/ControlRigVariableDetailsCustomization.cpp
--- a/Engine/Plugins/Experimental/ControlRig/Source/ControlRigEditor/Private/ControlRigVariableDetailsCustomization.cpp
+++ b/Engine/Plugins/Experimental/ControlRig/Source/ControlRigEditor/Private/ControlRigVariableDetailsCustomization.cpp
@@ -11,7 +11,7 @@
 #include "Widgets/Input/SCheckBox.h"
 #include "Engine/BlueprintGeneratedClass.h"
 #include "ControlRigBlueprint.h"
-#include "ControlRigController.h"
+#include "RigVMModel/RigVMController.h"
 
 #define LOCTEXT_NAMESPACE "ControlRigVariableDetailsCustomization"
 
@@ -90,107 +90,25 @@ void FControlRigVariableDetailsCustomization::CustomizeDetails(IDetailLayoutBuil
 
 bool FControlRigVariableDetailsCustomization::IsAnimationFlagEnabled(TWeakFieldPtr<FProperty> PropertyBeingCustomized) const
 {
-	if (PropertyBeingCustomized.IsValid() && BlueprintPtr.IsValid())
-	{
-		UControlRigBlueprint* RigBlueprint = Cast<UControlRigBlueprint>(BlueprintPtr.Get());
-		if (RigBlueprint)
-		{
-			if (RigBlueprint->Model)
-			{
-				FProperty* PropertyPtrBeingCustomized = PropertyBeingCustomized.Get();
-				const FControlRigModelNode* Node = RigBlueprint->Model->FindNode(PropertyPtrBeingCustomized->GetFName());
-				if (Node)
-				{
-					return Node->IsParameter();
-				}
-			}
-		}
-	}
 	return false;
 }
 
 ECheckBoxState FControlRigVariableDetailsCustomization::IsAnimationOutputChecked(TWeakFieldPtr<FProperty> PropertyBeingCustomized) const
 {
-	if (PropertyBeingCustomized.IsValid() && BlueprintPtr.IsValid())
-	{
-		UControlRigBlueprint* RigBlueprint = Cast<UControlRigBlueprint>(BlueprintPtr.Get());
-		if (RigBlueprint)
-		{
-			if (RigBlueprint->Model)
-			{
-				FProperty* PropertyPtrBeingCustomized = PropertyBeingCustomized.Get();
-				const FControlRigModelNode* Node = RigBlueprint->Model->FindNode(PropertyPtrBeingCustomized->GetFName());
-				if (Node)
-				{
-					if (Node->ParameterType == EControlRigModelParameterType::Output)
-					{
-						return ECheckBoxState::Checked;
-					}
-				}
-			}
-		}
-	}
-
 	return ECheckBoxState::Unchecked;
 }
 
 void FControlRigVariableDetailsCustomization::HandleAnimationOutputCheckStateChanged(ECheckBoxState CheckBoxState, TWeakFieldPtr<FProperty> PropertyBeingCustomized)
 {
-	if (PropertyBeingCustomized.IsValid() && BlueprintPtr.IsValid())
-	{
-		UControlRigBlueprint* RigBlueprint = Cast<UControlRigBlueprint>(BlueprintPtr.Get());
-		if (RigBlueprint)
-		{
-			if (RigBlueprint->ModelController)
-			{
-				FProperty* PropertyPtrBeingCustomized = PropertyBeingCustomized.Get();
-				RigBlueprint->ModelController->SetParameterType(PropertyPtrBeingCustomized->GetFName(), CheckBoxState == ECheckBoxState::Checked ? EControlRigModelParameterType::Output : EControlRigModelParameterType::Hidden);
-				FBlueprintEditorUtils::ReconstructAllNodes(BlueprintPtr.Get());
-			}
-		}
-	}
 }
 
 ECheckBoxState FControlRigVariableDetailsCustomization::IsAnimationInputChecked(TWeakFieldPtr<FProperty> PropertyBeingCustomized) const
 {
-	if (PropertyBeingCustomized.IsValid() && BlueprintPtr.IsValid())
-	{
-		UControlRigBlueprint* RigBlueprint = Cast<UControlRigBlueprint>(BlueprintPtr.Get());
-		if (RigBlueprint)
-		{
-			if (RigBlueprint->Model)
-			{
-				FProperty* PropertyPtrBeingCustomized = PropertyBeingCustomized.Get();
-				const FControlRigModelNode* Node = RigBlueprint->Model->FindNode(PropertyPtrBeingCustomized->GetFName());
-				if (Node)
-				{
-					if (Node->ParameterType == EControlRigModelParameterType::Input)
-					{
-						return ECheckBoxState::Checked;
-					}
-				}
-			}
-		}
-	}
-
 	return ECheckBoxState::Unchecked;
 }
 
 void FControlRigVariableDetailsCustomization::HandleAnimationInputCheckStateChanged(ECheckBoxState CheckBoxState, TWeakFieldPtr<FProperty> PropertyBeingCustomized)
 {
-	if (PropertyBeingCustomized.IsValid() && BlueprintPtr.IsValid())
-	{
-		UControlRigBlueprint* RigBlueprint = Cast<UControlRigBlueprint>(BlueprintPtr.Get());
-		if (RigBlueprint)
-		{
-			if (RigBlueprint->ModelController)
-			{
-				FProperty* PropertyPtrBeingCustomized = PropertyBeingCustomized.Get();
-				RigBlueprint->ModelController->SetParameterType(PropertyPtrBeingCustomized->GetFName(), CheckBoxState == ECheckBoxState::Checked ? EControlRigModelParameterType::Input : EControlRigModelParameterType::Hidden);
-				FBlueprintEditorUtils::ReconstructAllNodes(BlueprintPtr.Get());
-			}
-		}
-	}
 }
 
 #undef LOCTEXT_NAMESPACE
diff --git a/Engine/Plugins/Experimental/ControlRig/Source/ControlRigEditor/Private/EditMode/ControlRigControlsProxy.cpp b/Engine/Plugins/Experimental/ControlRig/Source/ControlRigEditor/Private/EditMode/ControlRigControlsProxy.cpp
new file mode 100644
--- /dev/null
+++ b/Engine/Plugins/Experimental/ControlRig/Source/ControlRigEditor/Private/EditMode/ControlRigControlsProxy.cpp
@@ -0,0 +1,470 @@
+// Copyright Epic Games, Inc. All Rights Reserved.
+
+#include "ControlRigControlsProxy.h"
+#include "EditorModeManager.h"
+#include "ControlRigEditMode.h"
+#include "Sequencer/ControlRigSequence.h"
+#include "Rigs/RigControlHierarchy.h"
+
+#include "Components/SkeletalMeshComponent.h"
+#include "Subsystems/AssetEditorSubsystem.h"
+#include "MovieSceneCommonHelpers.h"
+#include "PropertyHandle.h"
+
+void UControlRigControlsProxy::SelectionChanged(bool bInSelected)
+{
+	if (RigControl)
+	{
+		Modify();
+		const FName PropertyName("bSelected");
+		FTrackInstancePropertyBindings Binding(PropertyName, PropertyName.ToString());
+		Binding.CallFunction<bool>(*this, bInSelected);
+	}
+}
+
+void UControlRigControlsProxy::PostEditChangeProperty(struct FPropertyChangedEvent& PropertyChangedEvent)
+{
+	Super::PostEditChangeProperty(PropertyChangedEvent);
+	if (PropertyChangedEvent.Property && PropertyChangedEvent.Property->GetFName() == GET_MEMBER_NAME_CHECKED(UControlRigControlsProxy, bSelected))
+	{
+		if (RigControl && ControlRig.IsValid())
+		{
+			ControlRig->SelectControl(ControlName, bSelected);
+		}
+	}
+}
+
+
+#if WITH_EDITOR
+void UControlRigControlsProxy::PostEditUndo()
+{
+	if (RigControl && ControlRig.IsValid())
+	{
+		ControlRig->SelectControl(ControlName, bSelected);
+	}
+}
+#endif
+
+
+void UControlRigTransformControlProxy::PostEditChangeProperty(struct FPropertyChangedEvent& PropertyChangedEvent)
+{
+	Super::PostEditChangeProperty(PropertyChangedEvent);
+	if (PropertyChangedEvent.Property && PropertyChangedEvent.Property->GetFName() == GET_MEMBER_NAME_CHECKED(UControlRigTransformControlProxy, Transform))
+	{
+		if (RigControl && ControlRig.IsValid())
+		{
+			//MUST set through ControlRig
+			ControlRig->SetControlValue<FTransform>(ControlName, Transform, true, EControlRigSetKey::DoNotCare);
+		}
+	}
+}
+
+void UControlRigTransformControlProxy::ValueChanged()
+{
+	if (RigControl)
+	{
+		Modify();
+		const FName PropertyName("Transform");
+		FTrackInstancePropertyBindings Binding(PropertyName, PropertyName.ToString());
+		FTransform NewTransform = RigControl->Value.Get<FTransform>();
+
+		Binding.CallFunction<FTransform>(*this, NewTransform);
+
+	}
+}
+
+#if WITH_EDITOR
+void UControlRigTransformControlProxy::PostEditUndo()
+{
+	if (RigControl && ControlRig.IsValid())
+	{
+		ControlRig->SelectControl(ControlName, bSelected);
+		ControlRig->SetControlValue<FTransform>(ControlName, Transform, true, EControlRigSetKey::DoNotCare);
+	}
+}
+#endif
+
+void UControlRigTransformControlProxy::SetKey(const IPropertyHandle& KeyedPropertyHandle)
+{
+	if (RigControl)
+	{
+		ControlRig->SetControlValue<FTransform>(ControlName, Transform, true, EControlRigSetKey::Always);
+	}
+}
+
+void UControlRigTransformNoScaleControlProxy::PostEditChangeProperty(struct FPropertyChangedEvent& PropertyChangedEvent)
+{
+	Super::PostEditChangeProperty(PropertyChangedEvent);
+	if (PropertyChangedEvent.Property && PropertyChangedEvent.Property->GetFName() == GET_MEMBER_NAME_CHECKED(UControlRigTransformControlProxy, Transform))
+	{
+		if (RigControl && ControlRig.IsValid())
+		{
+			//MUST set through ControlRig
+			ControlRig->SetControlValue<FTransformNoScale>(ControlName, Transform, true, EControlRigSetKey::DoNotCare);
+		}
+	}
+}
+
+void UControlRigTransformNoScaleControlProxy::ValueChanged()
+{
+	if (RigControl)
+	{
+		Modify();
+		const FName PropertyName("Transform");
+		FTrackInstancePropertyBindings Binding(PropertyName, PropertyName.ToString());
+		FTransformNoScale NewTransform = RigControl->Value.Get<FTransformNoScale>();
+		Binding.CallFunction<FTransformNoScale>(*this, NewTransform);
+	}
+}
+
+#if WITH_EDITOR
+void UControlRigTransformNoScaleControlProxy::PostEditUndo()
+{
+	if (RigControl && ControlRig.IsValid())
+	{
+		ControlRig->SelectControl(ControlName, bSelected);
+		ControlRig->SetControlValue<FTransformNoScale>(ControlName, Transform, true, EControlRigSetKey::DoNotCare);
+	}
+}
+#endif
+
+void UControlRigTransformNoScaleControlProxy::SetKey(const IPropertyHandle& KeyedPropertyHandle)
+{
+	if (RigControl)
+	{
+		ControlRig->SetControlValue<FTransformNoScale>(ControlName, Transform, true, EControlRigSetKey::Always);
+	}
+}
+
+void UControlRigFloatControlProxy::PostEditChangeProperty(struct FPropertyChangedEvent& PropertyChangedEvent)
+{
+	Super::PostEditChangeProperty(PropertyChangedEvent);
+	if (PropertyChangedEvent.Property && PropertyChangedEvent.Property->GetFName() == GET_MEMBER_NAME_CHECKED(UControlRigFloatControlProxy, Float))
+	{
+		if (RigControl && ControlRig.IsValid())
+		{
+			//MUST set through ControlRig
+			ControlRig->SetControlValue<float>(ControlName, Float, true, EControlRigSetKey::DoNotCare);
+		}
+	}
+}
+
+void UControlRigFloatControlProxy::ValueChanged()
+{
+	if (RigControl)
+	{
+		Modify();
+		const FName PropertyName("Float");
+		FTrackInstancePropertyBindings Binding(PropertyName, PropertyName.ToString());
+		float Val = RigControl->Value.Get<float>();
+
+		Binding.CallFunction<float>(*this, Val);
+	}
+}
+
+#if WITH_EDITOR
+void UControlRigFloatControlProxy::PostEditUndo()
+{
+	if (RigControl && ControlRig.IsValid())
+	{
+		ControlRig->SelectControl(ControlName, bSelected);
+		ControlRig->SetControlValue<float>(ControlName, Float, true, EControlRigSetKey::DoNotCare);
+	}
+}
+#endif
+
+
+void UControlRigFloatControlProxy::SetKey(const IPropertyHandle& KeyedPropertyHandle)
+{
+	if (RigControl)
+	{
+		ControlRig->SetControlValue<float>(ControlName, Float, true, EControlRigSetKey::Always);
+	}
+}
+
+
+void UControlRigVectorControlProxy::PostEditChangeProperty(struct FPropertyChangedEvent& PropertyChangedEvent)
+{
+	Super::PostEditChangeProperty(PropertyChangedEvent);
+	if (PropertyChangedEvent.Property && PropertyChangedEvent.Property->GetFName() == GET_MEMBER_NAME_CHECKED(UControlRigVectorControlProxy, Vector))
+	{
+		if (RigControl && ControlRig.IsValid())
+		{
+			//MUST set through ControlRig
+			ControlRig->SetControlValue<FVector>(ControlName, Vector, true, EControlRigSetKey::DoNotCare);
+		}
+	}
+}
+
+void UControlRigVectorControlProxy::ValueChanged()
+{
+	if (RigControl)
+	{
+		Modify();
+		const FName PropertyName("Vector");
+		FTrackInstancePropertyBindings Binding(PropertyName, PropertyName.ToString());
+		FVector Val = RigControl->Value.Get<FVector>();
+
+		Binding.CallFunction<FVector>(*this, Val);
+	}
+}
+
+#if WITH_EDITOR
+void UControlRigVectorControlProxy::PostEditUndo()
+{
+	if (RigControl && ControlRig.IsValid())
+	{
+		ControlRig->SelectControl(ControlName, bSelected);
+		ControlRig->SetControlValue<FVector>(ControlName, Vector, true, EControlRigSetKey::DoNotCare);
+	}
+}
+#endif
+
+
+void UControlRigVectorControlProxy::SetKey(const IPropertyHandle& KeyedPropertyHandle)
+{
+	if (RigControl)
+	{
+		ControlRig->SetControlValue<FVector>(ControlName, Vector, true, EControlRigSetKey::Always);
+	}
+}
+
+void UControlRigVector2DControlProxy::PostEditChangeProperty(struct FPropertyChangedEvent& PropertyChangedEvent)
+{
+	Super::PostEditChangeProperty(PropertyChangedEvent);
+	if (PropertyChangedEvent.Property && PropertyChangedEvent.Property->GetFName() == GET_MEMBER_NAME_CHECKED(UControlRigVector2DControlProxy, Vector2D))
+	{
+		if (RigControl && ControlRig.IsValid())
+		{
+			//MUST set through ControlRig
+			ControlRig->SetControlValue<FVector2D>(ControlName, Vector2D, true, EControlRigSetKey::DoNotCare);
+		}
+	}
+}
+
+void UControlRigVector2DControlProxy::ValueChanged()
+{
+	if (RigControl)
+	{
+		Modify();
+		const FName PropertyName("Vector2D");
+		FTrackInstancePropertyBindings Binding(PropertyName, PropertyName.ToString());
+		FVector2D Val = RigControl->Value.Get<FVector2D>();
+
+		Binding.CallFunction<FVector2D>(*this, Val);
+	}
+}
+
+#if WITH_EDITOR
+void UControlRigVector2DControlProxy::PostEditUndo()
+{
+	if (RigControl && ControlRig.IsValid())
+	{
+		ControlRig->SelectControl(ControlName, bSelected);
+		ControlRig->SetControlValue<FVector2D>(ControlName, Vector2D, true, EControlRigSetKey::DoNotCare);
+	}
+}
+#endif
+
+
+void UControlRigVector2DControlProxy::SetKey(const IPropertyHandle& KeyedPropertyHandle)
+{
+	if (RigControl)
+	{
+		ControlRig->SetControlValue<FVector2D>(ControlName, Vector2D, true, EControlRigSetKey::Always);
+	}
+}
+
+
+void UControlRigBoolControlProxy::PostEditChangeProperty(struct FPropertyChangedEvent& PropertyChangedEvent)
+{
+	Super::PostEditChangeProperty(PropertyChangedEvent);
+	if (PropertyChangedEvent.Property && PropertyChangedEvent.Property->GetFName() == GET_MEMBER_NAME_CHECKED(UControlRigBoolControlProxy, Bool))
+	{
+		if (RigControl && ControlRig.IsValid())
+		{
+			//MUST set through ControlRig
+			ControlRig->SetControlValue<bool>(ControlName, Bool, true, EControlRigSetKey::DoNotCare);
+		}
+	}
+}
+
+void UControlRigBoolControlProxy::ValueChanged()
+{
+	if (RigControl)
+	{
+		Modify();
+		const FName PropertyName("Bool");
+		FTrackInstancePropertyBindings Binding(PropertyName, PropertyName.ToString());
+		bool Val = RigControl->Value.Get<bool>();
+
+		Binding.CallFunction<bool>(*this, Val);
+	}
+}
+
+#if WITH_EDITOR
+void UControlRigBoolControlProxy::PostEditUndo()
+{
+	if (RigControl && ControlRig.IsValid())
+	{
+		ControlRig->SelectControl(ControlName, bSelected);
+		ControlRig->SetControlValue<bool>(ControlName, Bool, true, EControlRigSetKey::DoNotCare);
+	}
+}
+#endif
+
+
+void UControlRigBoolControlProxy::SetKey(const IPropertyHandle& KeyedPropertyHandle)
+{
+	if (RigControl)
+	{
+		ControlRig->SetControlValue<bool>(ControlName, Bool, true, EControlRigSetKey::Always);
+	}
+}
+
+
+//////UControlDetailPanelControlProxies////////
+
+UControlRigControlsProxy* UControlRigDetailPanelControlProxies::FindProxy(const FName& Name) const
+{
+	UControlRigControlsProxy* const* Proxy = AllProxies.Find(Name);
+	if (Proxy &&  Proxy[0])
+	{
+		return Proxy[0];
+	}
+	return nullptr;
+}
+
+void UControlRigDetailPanelControlProxies::AddProxy(const FName& Name, UControlRig* ControlRig, FRigControl* RigControl)
+{
+	UControlRigControlsProxy* Proxy = FindProxy(Name);
+	if (!Proxy && RigControl != nullptr)
+	{
+		switch(RigControl->ControlType)
+		{
+			case ERigControlType::Transform:
+			{
+				Proxy = NewObject<UControlRigTransformControlProxy>(GetTransientPackage(), NAME_None);
+				break;
+
+			}
+			case ERigControlType::TransformNoScale:
+			{
+				Proxy = NewObject<UControlRigTransformNoScaleControlProxy>(GetTransientPackage(), NAME_None);
+				break;
+
+			}
+			case ERigControlType::Float:
+			{
+				Proxy = NewObject<UControlRigFloatControlProxy>(GetTransientPackage(), NAME_None);
+				break;
+
+			}
+			case ERigControlType::Position:
+			case ERigControlType::Rotator:
+			case ERigControlType::Scale:
+			{
+				Proxy = NewObject<UControlRigVectorControlProxy>(GetTransientPackage(), NAME_None);
+				break;
+
+			}
+			case ERigControlType::Vector2D:
+			{
+				Proxy = NewObject<UControlRigVector2DControlProxy>(GetTransientPackage(), NAME_None);
+				break;
+
+			}
+			case ERigControlType::Bool:
+			{
+				Proxy = NewObject<UControlRigBoolControlProxy>(GetTransientPackage(), NAME_None);
+				break;
+
+			}
+			default:
+				break;
+		}
+		if (Proxy)
+		{
+			Proxy->SetFlags(RF_Transactional);
+			Proxy->SetName(Name);
+			Proxy->RigControl = RigControl;
+			Proxy->ControlRig = ControlRig;
+			Proxy->ValueChanged();
+			AllProxies.Add(Name, Proxy);
+		}
+
+	}
+}
+
+void UControlRigDetailPanelControlProxies::RemoveProxy(const FName& Name)
+{
+	UControlRigControlsProxy* ExistingProxy = FindProxy(Name);
+	if (ExistingProxy)
+	{
+		ExistingProxy->Rename(nullptr, GetTransientPackage(), REN_ForceNoResetLoaders);
+		ExistingProxy->MarkPendingKill();
+	}
+	AllProxies.Remove(Name);
+}
+
+void UControlRigDetailPanelControlProxies::RemoveAllProxies()
+{
+	for (TPair<FName, UControlRigControlsProxy * >Pair: AllProxies)
+	{
+		UControlRigControlsProxy* ExistingProxy = Pair.Value;
+		ExistingProxy->Rename(nullptr, GetTransientPackage(), REN_ForceNoResetLoaders);
+		ExistingProxy->MarkPendingKill();
+	}
+	AllProxies.Empty();
+	SelectedProxies.SetNum(0);
+}
+
+void UControlRigDetailPanelControlProxies::RecreateAllProxies(UControlRig* InControlRig)
+{
+	RemoveAllProxies();
+	const TArray<FRigControl>& Controls = InControlRig->AvailableControls();
+	for (const FRigControl& RigControl : Controls)
+	{
+		FRigControl* PRigControl = const_cast<FRigControl*>(&RigControl);
+		if(PRigControl->bGizmoEnabled)
+		{
+			AddProxy(RigControl.Name, InControlRig, PRigControl);
+		}
+	}
+}
+
+void UControlRigDetailPanelControlProxies::ProxyChanged(const FName& Name)
+{
+	UControlRigControlsProxy* Proxy = FindProxy(Name);
+	if (Proxy)
+	{
+		Modify();
+		Proxy->ValueChanged();
+	}
+}
+
+void UControlRigDetailPanelControlProxies::SelectProxy(const FName& Name, bool bSelected)
+{
+	UControlRigControlsProxy* Proxy = FindProxy(Name);
+	if (Proxy)
+	{
+		Modify();
+		if (bSelected)
+		{
+			if (!SelectedProxies.Contains(Proxy))
+			{
+				//don't show more than 5 for performance
+				if (SelectedProxies.Num() < 5)
+				{
+					SelectedProxies.Add(Proxy);
+				}
+			}
+		}
+		else
+		{
+			SelectedProxies.Remove(Proxy);
+		}
+		Proxy->SelectionChanged(bSelected);
+	}
+}
+
diff --git a/Engine/Plugins/Experimental/ControlRig/Source/ControlRigEditor/Private/EditMode/ControlRigControlsProxy.h b/Engine/Plugins/Experimental/ControlRig/Source/ControlRigEditor/Private/EditMode/ControlRigControlsProxy.h
new file mode 100644
--- /dev/null
+++ b/Engine/Plugins/Experimental/ControlRig/Source/ControlRigEditor/Private/EditMode/ControlRigControlsProxy.h
@@ -0,0 +1,204 @@
+// Copyright Epic Games, Inc. All Rights Reserved.
+#pragma once
+
+#include "CoreMinimal.h"
+#include "UObject/Object.h"
+#include "UObject/LazyObjectPtr.h"
+#include "TransformNoScale.h"
+#include "ControlRigControlsProxy.generated.h"
+
+struct FRigControl;
+class UControlRig;
+class IPropertyHandle;
+
+UCLASS(Abstract)
+class UControlRigControlsProxy : public UObject
+{
+	GENERATED_BODY()
+
+public:
+	UControlRigControlsProxy() : bSelected(false) {}
+	virtual void SetName(const FName& InName)  { ControlName = InName; }
+	virtual FName GetName() const { return ControlName; }
+	virtual void ValueChanged() {}
+	virtual void SelectionChanged(bool bInSelected);
+	virtual void SetKey(const IPropertyHandle& KeyedPropertyHandle) {};
+
+	// UObject interface
+	virtual void PostEditChangeProperty(struct FPropertyChangedEvent& PropertyChangedEvent) override;
+#if WITH_EDITOR
+	virtual void PostEditUndo() override;
+#endif
+
+public:
+
+	FRigControl* RigControl;
+	TWeakObjectPtr<UControlRig> ControlRig;
+
+	UPROPERTY()
+	bool bSelected;
+	UPROPERTY(VisibleAnywhere, AdvancedDisplay, Category = "Control")
+	FName ControlName;
+};
+
+UCLASS()
+class UControlRigTransformControlProxy : public UControlRigControlsProxy
+{
+	GENERATED_BODY()
+	UControlRigTransformControlProxy() {}
+
+	// UObject interface
+	virtual void PostEditChangeProperty(struct FPropertyChangedEvent& PropertyChangedEvent) override;
+#if WITH_EDITOR
+	virtual void PostEditUndo() override;
+#endif
+
+	//UControlRigControlsProxy
+	virtual void ValueChanged() override;
+	virtual void SetKey(const IPropertyHandle& KeyedPropertyHandle) override;
+
+public:
+	
+	UPROPERTY(EditAnywhere, Interp, AdvancedDisplay, Category = "Control")
+	FTransform Transform;
+};
+
+UCLASS()
+class UControlRigTransformNoScaleControlProxy : public UControlRigControlsProxy
+{
+	GENERATED_BODY()
+	UControlRigTransformNoScaleControlProxy() {}
+
+	// UObject interface
+	virtual void PostEditChangeProperty(struct FPropertyChangedEvent& PropertyChangedEvent) override;
+#if WITH_EDITOR
+	virtual void PostEditUndo() override;
+#endif
+
+	//UControlRigControlsProxy
+	virtual void ValueChanged() override;
+	virtual void SetKey(const IPropertyHandle& KeyedPropertyHandle) override;
+
+public:
+
+	UPROPERTY(EditAnywhere, Interp, AdvancedDisplay, Category = "Control")
+	FTransformNoScale Transform;
+};
+
+UCLASS()
+class UControlRigFloatControlProxy : public UControlRigControlsProxy
+{
+	GENERATED_BODY()
+	UControlRigFloatControlProxy()  {}
+
+	// UObject interface
+	virtual void PostEditChangeProperty(struct FPropertyChangedEvent& PropertyChangedEvent) override;
+#if WITH_EDITOR
+	virtual void PostEditUndo() override;
+#endif
+
+	//UControlRigControlsProxy
+	virtual void ValueChanged() override;
+	virtual void SetKey(const IPropertyHandle& KeyedPropertyHandle) override;
+
+public:
+
+	UPROPERTY(EditAnywhere, Interp, AdvancedDisplay, Category = "Control")
+	float Float;
+};
+
+
+UCLASS()
+class UControlRigVectorControlProxy : public UControlRigControlsProxy
+{
+	GENERATED_BODY()
+	UControlRigVectorControlProxy() {}
+
+	// UObject interface
+	virtual void PostEditChangeProperty(struct FPropertyChangedEvent& PropertyChangedEvent) override;
+#if WITH_EDITOR
+	virtual void PostEditUndo() override;
+#endif
+
+	//UControlRigControlsProxy
+	virtual void ValueChanged() override;
+	virtual void SetKey(const IPropertyHandle& KeyedPropertyHandle) override;
+
+public:
+
+	UPROPERTY(EditAnywhere, Interp, AdvancedDisplay, Category = "Control")
+	FVector Vector;
+};
+
+UCLASS()
+class UControlRigVector2DControlProxy : public UControlRigControlsProxy
+{
+	GENERATED_BODY()
+	UControlRigVector2DControlProxy() {}
+
+	// UObject interface
+	virtual void PostEditChangeProperty(struct FPropertyChangedEvent& PropertyChangedEvent) override;
+#if WITH_EDITOR
+	virtual void PostEditUndo() override;
+#endif
+
+	//UControlRigControlsProxy
+	virtual void ValueChanged() override;
+	virtual void SetKey(const IPropertyHandle& KeyedPropertyHandle) override;
+
+public:
+
+	UPROPERTY(EditAnywhere, Interp, AdvancedDisplay, Category = "Control")
+	FVector2D Vector2D;
+};
+
+UCLASS()
+class UControlRigBoolControlProxy : public UControlRigControlsProxy
+{
+	GENERATED_BODY()
+	UControlRigBoolControlProxy() {}
+
+	// UObject interface
+	virtual void PostEditChangeProperty(struct FPropertyChangedEvent& PropertyChangedEvent) override;
+#if WITH_EDITOR
+	virtual void PostEditUndo() override;
+#endif
+
+	//UControlRigControlsProxy
+	virtual void ValueChanged() override;
+	virtual void SetKey(const IPropertyHandle& KeyedPropertyHandle) override;
+
+public:
+
+	UPROPERTY(EditAnywhere, Interp, AdvancedDisplay, Category = "Control")
+	bool Bool;
+};
+
+
+/** Proxy in Details Panel */
+UCLASS()
+class UControlRigDetailPanelControlProxies :public UObject
+{
+	GENERATED_BODY()
+
+	UControlRigDetailPanelControlProxies() {}
+protected:
+
+	UPROPERTY()
+	TMap<FName, UControlRigControlsProxy*> AllProxies;
+
+	UPROPERTY()
+	TArray< UControlRigControlsProxy*> SelectedProxies;
+
+
+public:
+	UControlRigControlsProxy* FindProxy(const FName& Name) const;
+	void AddProxy(const FName& Name, UControlRig* InControlRig, FRigControl* InControl);
+	void RemoveProxy(const FName& Name );
+	void ProxyChanged(const FName& Name);
+	void RemoveAllProxies();
+	void RecreateAllProxies(UControlRig* InControlRig);
+	void SelectProxy(const FName& Name, bool bSelected);
+	const TArray<UControlRigControlsProxy*>& GetSelectedProxies() const { return SelectedProxies;}
+
+};
