COPY from //Dev-Anim to //Dev-Main

diff --git a/Engine/Plugins/Experimental/ControlRig/Source/ControlRigEditor/Private/EditMode/ControlRigEditMode.cpp b/Engine/Plugins/Experimental/ControlRig/Source/ControlRigEditor/Private/EditMode/ControlRigEditMode.cpp
--- a/Engine/Plugins/Experimental/ControlRig/Source/ControlRigEditor/Private/EditMode/ControlRigEditMode.cpp
+++ b/Engine/Plugins/Experimental/ControlRig/Source/ControlRigEditor/Private/EditMode/ControlRigEditMode.cpp
@@ -31,14 +31,15 @@
 #include "Kismet2/BlueprintEditorUtils.h"
 #include "Drawing/ControlRigDrawInterface.h"
 #include "ControlRigBlueprint.h"
-#include "ControlRigController.h"
 #include "ControlRigGizmoActor.h"
 #include "DefaultControlRigManipulationLayer.h"
 #include "Framework/MultiBox/MultiBoxBuilder.h"
 #include "SEditorViewport.h"
-
+#include "ControlRigControlsProxy.h"
+#include "ScopedTransaction.h"
 #include "Framework/Notifications/NotificationManager.h"
 #include "Widgets/Notifications/SNotificationList.h"
+#include "RigVMModel/RigVMController.h"
 
 FName FControlRigEditMode::ModeName("EditMode.ControlRig");
 
@@ -64,12 +65,11 @@ FControlRigEditMode::FControlRigEditMode()
 	, bRecreateManipulationLayerRequired(false)
 	, ManipulationLayer(nullptr)
 	, CurrentViewportClient(nullptr)
+	, bIsChangingCoordSystem(false)
 {
 	Settings = NewObject<UControlRigEditModeSettings>(GetTransientPackage(), TEXT("Settings"));
-	
-	// @todo: thinking of removing this, but for now leaving it. 
-	// it indicates, execution is stopped, and you can modify default bone transform or controls if set to be true
-	bEnableRigElementDefaultPoseEditing = false;
+	ControlProxy = NewObject<UControlRigDetailPanelControlProxies>(GetTransientPackage(), TEXT("Controls"));
+	ControlProxy->SetFlags(RF_Transactional);
 
 	CommandBindings = MakeShareable(new FUICommandList);
 	BindCommands();
@@ -83,62 +83,82 @@ FControlRigEditMode::~FControlRigEditMode()
 {	
 	CommandBindings = nullptr;
 
+	GLevelEditorModeTools().OnWidgetModeChanged().RemoveAll(this);
+	GLevelEditorModeTools().OnCoordSystemChanged().RemoveAll(this);
+
 #if WITH_EDITOR
 	GEditor->OnObjectsReplaced().RemoveAll(this);
 #endif
 }
 
-void FControlRigEditMode::SetObjects(const TWeakObjectPtr<>& InSelectedObject, const FGuid& InObjectBinding, UObject* BindingObject)
+void FControlRigEditMode::SetObjects(const TWeakObjectPtr<>& InSelectedObject, const FGuid& InObjectBinding, UObject* BindingObject, TWeakPtr<ISequencer> InSequencer)
 {
 	WeakControlRigEditing = Cast<UControlRig>(InSelectedObject.Get());
 	ControlRigGuid = InObjectBinding;
-
+	WeakSequencer = InSequencer;
 	// if we get binding object, set it to control rig binding object
 	if (BindingObject)
 	{
 		if (UControlRig* ControlRig = WeakControlRigEditing.Get())
 		{
 			if (TSharedPtr<IControlRigObjectBinding> ObjectBinding = ControlRig->GetObjectBinding())
-		{
-				if (ObjectBinding->GetBoundObject() == nullptr)
 			{
+				if (ObjectBinding->GetBoundObject() == nullptr)
+				{
 					ObjectBinding->BindToObject(BindingObject);
+				}
 			}
 		}
-	}
 
 	}
-
+	if (UsesToolkits() && Toolkit.IsValid())
+	{
+		StaticCastSharedPtr<SControlRigEditModeTools>(Toolkit->GetInlineContent())->SetControlRig(WeakControlRigEditing.Get());
+	}
 	SetObjects_Internal();
 }
 
-void FControlRigEditMode::SetObjects_Internal()
+void FControlRigEditMode::SetUpDetailPanel()
 {
-	TArray<TWeakObjectPtr<>> SelectedObjects;
-	if(IsInLevelEditor())
+	if (IsInLevelEditor())
 	{
-		SelectedObjects.Add(Settings);
+		TArray<TWeakObjectPtr<>> SelectedObjects;
+		if (WeakControlRigEditing.IsValid())
+		{
+			const TArray<UControlRigControlsProxy*>& Proxies = ControlProxy->GetSelectedProxies();
+			for (UControlRigControlsProxy* Proxy : Proxies)
+			{
+				SelectedObjects.Add(Proxy);
+			}
+			SelectedObjects.Add(Settings);
+		}
+		StaticCastSharedPtr<SControlRigEditModeTools>(Toolkit->GetInlineContent())->SetSequencer(WeakSequencer.Pin());
+		StaticCastSharedPtr<SControlRigEditModeTools>(Toolkit->GetInlineContent())->SetDetailsObjects(SelectedObjects);
 	}
+}
+
+void FControlRigEditMode::SetObjects_Internal()
+{
 
 	if (WeakControlRigEditing.IsValid())
 	{
 		//Don't add the WeakControlRig Editing...SelectedObjects.Add(WeakControlRigEditing);
 		WeakControlRigEditing.Get()->DrawInterface = &DrawInterface;
 		if (IsInLevelEditor())
-	{
+		{
+			WeakControlRigEditing->Hierarchy.OnElementSelected.RemoveAll(this);
+			WeakControlRigEditing->ControlModified().RemoveAll(this);
+
 			WeakControlRigEditing->Hierarchy.OnElementSelected.AddSP(this, &FControlRigEditMode::OnRigElementSelected);
+			WeakControlRigEditing->ControlModified().AddSP(this, &FControlRigEditMode::OnControlModified);
 		}
 
-
-		if (UsesToolkits())
-			{
-			StaticCastSharedPtr<SControlRigEditModeTools>(Toolkit->GetInlineContent())->SetDetailsObjects(SelectedObjects);
-					}
-
 		// create default manipulation layer
 		RecreateManipulationLayer();
 		HandleSelectionChanged();
 	}
+	SetUpDetailPanel();
+
 }
 
 bool FControlRigEditMode::UsesToolkits() const
@@ -159,6 +179,21 @@ void FControlRigEditMode::Enter()
 		}
 
 		Toolkit->Init(Owner->GetToolkitHost());
+
+		bIsChangingCoordSystem = false;
+		if (CoordSystemPerWidgetMode.Num() < (FWidget::WM_Max))
+		{
+			CoordSystemPerWidgetMode.SetNum(FWidget::WM_Max);
+			ECoordSystem CoordSystem = GLevelEditorModeTools().GetCoordSystem();
+			for (int32 i = 0; i < FWidget::WM_Max; ++i)
+			{
+				CoordSystemPerWidgetMode[i] = CoordSystem;
+			}
+		}
+	
+		GLevelEditorModeTools().OnWidgetModeChanged().AddSP(this, &FControlRigEditMode::OnWidgetModeChanged);
+		GLevelEditorModeTools().OnCoordSystemChanged().AddSP(this, &FControlRigEditMode::OnCoordSystemChanged);
+
 	}
 
 	SetObjects_Internal();
@@ -175,7 +210,7 @@ void FControlRigEditMode::Exit()
 
 		if (GEditor)
 		{
-		GEditor->EndTransaction();
+			GEditor->EndTransaction();
 		}
 
 		bIsTransacting = false;
@@ -195,6 +230,14 @@ void FControlRigEditMode::Exit()
 	}
 	//clear actors
 	GizmoActors.SetNum(0);
+
+	//clear delegates
+	GLevelEditorModeTools().OnWidgetModeChanged().RemoveAll(this);
+	GLevelEditorModeTools().OnCoordSystemChanged().RemoveAll(this);
+
+	//clear proxies
+	ControlProxy->RemoveAllProxies();
+
 	// Call parent implementation
 	FEdMode::Exit();
 }
@@ -203,15 +246,15 @@ void FControlRigEditMode::Tick(FEditorViewportClient* ViewportClient, float Delt
 {
 	FEdMode::Tick(ViewportClient, DeltaTime);
 
-		ViewportClient->Invalidate();
+	ViewportClient->Invalidate();
 
 	if (ManipulationLayer)
-		{
+	{
 		RecalcPivotTransform();
-		}
+	}
 
 	if (bRecreateManipulationLayerRequired)
-		{
+	{
 		RecreateManipulationLayer();
 
 		for (const FRigElementKey& SelectedKey : SelectedRigElements)
@@ -223,8 +266,21 @@ void FControlRigEditMode::Tick(FEditorViewportClient* ViewportClient, float Delt
 				{
 					GizmoActor->SetSelected(true);
 				}
+
+				if (WeakControlRigEditing.IsValid())
+				{
+					FRigControl* Control = WeakControlRigEditing->FindControl(SelectedKey.Name);
+					if (Control)
+					{
+						if (!WeakControlRigEditing->IsCurveControl(Control))
+						{
+							ControlProxy->AddProxy(SelectedKey.Name, WeakControlRigEditing.Get(), Control);
+						}
+					}
+				}
 			}
 		}
+		SetUpDetailPanel();
 		HandleSelectionChanged();
 		bRecreateManipulationLayerRequired = false;
 	}
@@ -232,8 +288,45 @@ void FControlRigEditMode::Tick(FEditorViewportClient* ViewportClient, float Delt
 	// We need to tick here since changing a bone for example
 	// might have changed the transform of the Control
 	if (ManipulationLayer)
+	{
+		if (WeakControlRigEditing.IsValid())
 		{
-		ManipulationLayer->TickManipulatableObjects(DeltaTime);
+			const FWidget::EWidgetMode CurrentWidgetMode = ViewportClient->GetWidgetMode();
+			for (FRigElementKey SelectedRigElement : SelectedRigElements)
+			{
+				if (AControlRigGizmoActor* GizmoActor = GetGizmoFromControlName(SelectedRigElement.Name))
+				{
+					if (!ManipulationLayer->ModeSupportedByGizmoActor(GizmoActor, CurrentWidgetMode))
+					{
+						if (FRigControl* Control = WeakControlRigEditing->FindControl(SelectedRigElement.Name))
+						{
+							switch (Control->ControlType)
+							{
+								case ERigControlType::Float:
+								case ERigControlType::Vector2D:
+								case ERigControlType::Position:
+								case ERigControlType::Transform:
+								case ERigControlType::TransformNoScale:
+								{
+									ViewportClient->SetWidgetMode(FWidget::WM_Translate);
+									break;
+								}
+								case ERigControlType::Rotator:
+								{
+									ViewportClient->SetWidgetMode(FWidget::WM_Rotate);
+									break;
+								}
+								case ERigControlType::Scale:
+								{
+									ViewportClient->SetWidgetMode(FWidget::WM_Scale);
+									break;
+								}
+							}
+						}
+					}
+				}
+			}
+		}
 	}
 }
 
@@ -241,49 +334,50 @@ void FControlRigEditMode::Render(const FSceneView* View, FViewport* Viewport, FP
 {	
 	if (!WeakControlRigEditing.IsValid() || ManipulationLayer == nullptr)
 	{
-		DrawInterface.DrawInstructions.Reset();
+		DrawInterface.Instructions.Reset();
 		return;
-		}
+	}
 
 	bool bRender = !Settings->bHideManipulators;
 
 	FTransform ComponentTransform = (ManipulationLayer)? ManipulationLayer->GetSkeletalMeshComponentTransform() : FTransform::Identity;
-		if (bRender)
-		{
+	if (bRender)
+	{
 		for (AControlRigGizmoActor* Actor : GizmoActors)
-			{
+		{
 			//Actor->SetActorHiddenInGame(bIsHidden);
 			if (GIsEditor && Actor->GetWorld() != nullptr && !Actor->GetWorld()->IsPlayInEditor())
-				{
+			{
 				Actor->SetIsTemporarilyHiddenInEditor(false);
 			}
-				}
+		}
 		if (Settings->bDisplayHierarchy)
 		{
 			UControlRig* ControlRig = WeakControlRigEditing.Get();
-				// each base hierarchy Bone
+			// each base hierarchy Bone
 			const FRigBoneHierarchy& BaseHierarchy = ControlRig->GetBoneHierarchy();
 			for (int32 BoneIndex = 0; BoneIndex < BaseHierarchy.Num(); ++BoneIndex)
-				{
+			{
 				const FRigBone& CurrentBone = BaseHierarchy[BoneIndex];
-					const FTransform Transform = BaseHierarchy.GetGlobalTransform(BoneIndex);
+				const FTransform Transform = BaseHierarchy.GetGlobalTransform(BoneIndex);
 
-					if (CurrentBone.ParentIndex != INDEX_NONE)
-					{
-						const FTransform ParentTransform = BaseHierarchy.GetGlobalTransform(CurrentBone.ParentIndex);
+				if (CurrentBone.ParentIndex != INDEX_NONE)
+				{
+					const FTransform ParentTransform = BaseHierarchy.GetGlobalTransform(CurrentBone.ParentIndex);
 
 					PDI->DrawLine(ComponentTransform.TransformPosition(Transform.GetLocation()),ComponentTransform.TransformPosition(ParentTransform.GetLocation()), FLinearColor::White, SDPG_Foreground);
-					}
+				}
 
 				PDI->DrawPoint(ComponentTransform.TransformPosition(Transform.GetLocation()), FLinearColor::White, 5.0f, SDPG_Foreground);
-				}
 			}
+		}
 
 		if (Settings->bDisplayAxesOnSelection && Settings->AxisScale > SMALL_NUMBER)
 		{
 			UControlRig* ControlRig = WeakControlRigEditing.Get();
 			const FRigHierarchyContainer* Hierarchy = ControlRig->GetHierarchy();
 			const float Scale = Settings->AxisScale;
+			PDI->AddReserveLines(SDPG_Foreground, SelectedRigElements.Num() * 3);
 
 			for (const FRigElementKey& SelectedElement : SelectedRigElements)
 			{
@@ -293,44 +387,48 @@ void FControlRigEditMode::Render(const FSceneView* View, FViewport* Viewport, FP
 				PDI->DrawLine(ElementTransform.GetTranslation(), ElementTransform.TransformPosition(FVector(Scale, 0.f, 0.f)), FLinearColor::Red, SDPG_Foreground);
 				PDI->DrawLine(ElementTransform.GetTranslation(), ElementTransform.TransformPosition(FVector(0.f, Scale, 0.f)), FLinearColor::Green, SDPG_Foreground);
 				PDI->DrawLine(ElementTransform.GetTranslation(), ElementTransform.TransformPosition(FVector(0.f, 0.f, Scale)), FLinearColor::Blue, SDPG_Foreground);
+			}
 		}
-		}
-		for (const FControlRigDrawInterface::FDrawIntruction& Instruction : DrawInterface.DrawInstructions)
+		for (const FControlRigDrawInstruction& Instruction : DrawInterface)
+		{
+			if (Instruction.Positions.Num() == 0)
 			{
-				if (Instruction.Positions.Num() == 0)
-				{
-					continue;
-				}
-				switch (Instruction.DrawType)
+				continue;
+			}
+
+			FTransform InstructionTransform = Instruction.Transform * ComponentTransform;
+			switch (Instruction.PrimitiveType)
+			{
+				case EControlRigDrawSettings::Points:
 				{
-					case FControlRigDrawInterface::EDrawType_Point:
+					for (const FVector& Point : Instruction.Positions)
 					{
-						for (const FVector& Point : Instruction.Positions)
-						{
-					PDI->DrawPoint(ComponentTransform.TransformPosition(Point), Instruction.Color, Instruction.Thickness, SDPG_Foreground);
-						}
-						break;
+						PDI->DrawPoint(InstructionTransform.TransformPosition(Point), Instruction.Color, Instruction.Thickness, SDPG_Foreground);
 					}
-					case FControlRigDrawInterface::EDrawType_Lines:
+					break;
+				}
+				case EControlRigDrawSettings::Lines:
+				{
+					const TArray<FVector>& Points = Instruction.Positions;
+					PDI->AddReserveLines(SDPG_Foreground, Points.Num() / 2, false, Instruction.Thickness > SMALL_NUMBER);
+					for (int32 PointIndex = 0; PointIndex < Points.Num() - 1; PointIndex += 2)
 					{
-						const TArray<FVector>& Points = Instruction.Positions;
-						for (int32 PointIndex = 0; PointIndex < Points.Num() - 1; PointIndex += 2)
-						{
-					PDI->DrawLine(ComponentTransform.TransformPosition(Points[PointIndex]), ComponentTransform.TransformPosition(Points[PointIndex + 1]), Instruction.Color, SDPG_Foreground, Instruction.Thickness);
-						}
-						break;
+						PDI->DrawLine(InstructionTransform.TransformPosition(Points[PointIndex]), InstructionTransform.TransformPosition(Points[PointIndex + 1]), Instruction.Color, SDPG_Foreground, Instruction.Thickness);
 					}
-					case FControlRigDrawInterface::EDrawType_LineStrip:
+					break;
+				}
+				case EControlRigDrawSettings::LineStrip:
+				{
+					const TArray<FVector>& Points = Instruction.Positions;
+					PDI->AddReserveLines(SDPG_Foreground, Points.Num() - 1, false, Instruction.Thickness > SMALL_NUMBER);
+					for (int32 PointIndex = 0; PointIndex < Points.Num() - 1; PointIndex++)
 					{
-						const TArray<FVector>& Points = Instruction.Positions;
-						for (int32 PointIndex = 0; PointIndex < Points.Num() - 1; PointIndex++)
-						{
-					PDI->DrawLine(ComponentTransform.TransformPosition(Points[PointIndex]), ComponentTransform.TransformPosition(Points[PointIndex + 1]), Instruction.Color, SDPG_Foreground, Instruction.Thickness);
-						}
-						break;
+						PDI->DrawLine(InstructionTransform.TransformPosition(Points[PointIndex]), InstructionTransform.TransformPosition(Points[PointIndex + 1]), Instruction.Color, SDPG_Foreground, Instruction.Thickness);
 					}
+					break;
 				}
 			}
+		}
 	}
 	else
 	{
@@ -344,9 +442,7 @@ void FControlRigEditMode::Render(const FSceneView* View, FViewport* Viewport, FP
 		}
 	}
 
-	//draw debug info
-	
-	DrawInterface.DrawInstructions.Reset();
+	DrawInterface.Reset();
 }
 
 bool FControlRigEditMode::InputKey(FEditorViewportClient* InViewportClient, FViewport* InViewport, FKey InKey, EInputEvent InEvent)
@@ -373,18 +469,25 @@ bool FControlRigEditMode::EndTracking(FEditorViewportClient* InViewportClient, F
 		{
 			// One final notify of our manipulators to make sure the property is keyed
 			for (AControlRigGizmoActor* GizmoActor : GizmoActors)
-				{
+			{
 				if (GizmoActor->IsManipulating())
-					{
+				{
 					GizmoActor->SetManipulating(false);
-					}
 				}
 			}
 
-		ManipulationLayer->EndTransaction();
-		GEditor->EndTransaction();
+			bManipulatorMadeChange = false;
+			GEditor->EndTransaction();
+
+		}
+
+		if (ManipulationLayer)
+		{
+			ManipulationLayer->EndTransaction();
+		}
+
+
 		bIsTransacting = false;
-		bManipulatorMadeChange = false;
 		return true;
 	}
 
@@ -397,14 +500,13 @@ bool FControlRigEditMode::StartTracking(FEditorViewportClient* InViewportClient,
 {
 	if (!bIsTransacting)
 	{
-		GEditor->BeginTransaction(LOCTEXT("MoveControlTransaction", "Move Control"));
 
 		if (ManipulationLayer != nullptr)
 		{
 			ManipulationLayer->BeginTransaction();
 		}
 		for (AControlRigGizmoActor* GizmoActor : GizmoActors)
-			{
+		{
 			GizmoActor->SetManipulating(true);
 		}
 
@@ -422,10 +524,10 @@ bool FControlRigEditMode::UsesTransformWidget() const
 	for (const AControlRigGizmoActor* GizmoActor : GizmoActors)
 	{
 		if (GizmoActor->IsSelected())
-			{
-				return true;
-			}
+		{
+			return true;
 		}
+	}
 
 	if (AreRigElementSelectedAndMovable())
 	{
@@ -442,13 +544,13 @@ bool FControlRigEditMode::UsesTransformWidget(FWidget::EWidgetMode CheckMode) co
 		if (GizmoActor->IsSelected() && ManipulationLayer != nullptr)
 		{
 			return ManipulationLayer->ModeSupportedByGizmoActor(GizmoActor, CheckMode);
-			}
 		}
+	}
 
 	if (AreRigElementSelectedAndMovable())
-		{
-			return true;
-		}
+	{
+		return true;
+	}
 
 	return FEdMode::UsesTransformWidget(CheckMode);
 }
@@ -456,10 +558,10 @@ bool FControlRigEditMode::UsesTransformWidget(FWidget::EWidgetMode CheckMode) co
 FVector FControlRigEditMode::GetWidgetLocation() const
 {
 	if (AreRigElementSelectedAndMovable() && ManipulationLayer != nullptr)
-			{
+	{
 		FTransform ComponentTransform = ManipulationLayer->GetSkeletalMeshComponentTransform();
-				return ComponentTransform.TransformPosition(PivotTransform.GetLocation());
-			}
+		return ComponentTransform.TransformPosition(PivotTransform.GetLocation());
+	}
 
 	return FEdMode::GetWidgetLocation();
 }
@@ -467,10 +569,10 @@ FVector FControlRigEditMode::GetWidgetLocation() const
 bool FControlRigEditMode::GetCustomDrawingCoordinateSystem(FMatrix& OutMatrix, void* InData)
 {
 	if (AreRigElementSelectedAndMovable())
-			{
-				OutMatrix = PivotTransform.ToMatrixNoScale().RemoveTranslation();
-				return true;
-			}
+	{
+		OutMatrix = PivotTransform.ToMatrixNoScale().RemoveTranslation();
+		return true;
+	}
 
 	return false;
 }
@@ -492,14 +594,16 @@ bool FControlRigEditMode::HandleClick(FEditorViewportClient* InViewportClient, H
 				const FControlData* ControlData = ManipulationLayer != nullptr ? ManipulationLayer->GetControlDataFromGizmo(GizmoActor) :  nullptr;
 
 				if (ControlData)
-		{
+				{
+					FScopedTransaction ScopedTransaction(LOCTEXT("SelectControlTransaction", "Select Control"), !GIsTransacting);
+
 					const FName& ControlName = ControlData->ControlName;
-			if (Click.IsShiftDown() || Click.IsControlDown())
-			{
+					if (Click.IsShiftDown() || Click.IsControlDown())
+					{
 						SetRigElementSelection(ERigElementType::Control, ControlName, true);
-			}
-			else
-			{
+					}
+					else
+					{
 						ClearRigElementSelection(FRigElementTypeHelper::ToMask(ERigElementType::Control));
 						SetRigElementSelection(ERigElementType::Control, ControlName, true);
 					}
@@ -510,12 +614,12 @@ bool FControlRigEditMode::HandleClick(FEditorViewportClient* InViewportClient, H
 				if (Click.GetKey() == EKeys::RightMouseButton)
 				{
 					OpenContextMenu(InViewportClient);
-			}
+				}
 
-			return true;
+				return true;
+			}
 		}
 	}
-	}
 
 	// for now we show this menu all the time if body is selected
 	// if we want some global menu, we'll have to move this
@@ -525,6 +629,12 @@ bool FControlRigEditMode::HandleClick(FEditorViewportClient* InViewportClient, H
 		return true;
 	}
 
+	if (Settings  &&  Settings->bOnlySelectRigControls)
+	{
+		return true;
+	}
+	FScopedTransaction ScopedTransaction(LOCTEXT("SelectControlTransaction", "Select Control"), !GIsTransacting);
+
 	// clear selected controls
 	ClearRigElementSelection(FRigElementTypeHelper::ToMask(ERigElementType::All));
 
@@ -639,97 +749,105 @@ void FControlRigEditMode::SelectNone()
 
 bool FControlRigEditMode::InputDelta(FEditorViewportClient* InViewportClient, FViewport* InViewport, FVector& InDrag, FRotator& InRot, FVector& InScale)
 {
-		FVector Drag = InDrag;
-		FRotator Rot = InRot;
-		FVector Scale = InScale;
+	FVector Drag = InDrag;
+	FRotator Rot = InRot;
+	FVector Scale = InScale;
 
-		const bool bCtrlDown = InViewport->KeyState(EKeys::LeftControl) || InViewport->KeyState(EKeys::RightControl);
-		const bool bShiftDown = InViewport->KeyState(EKeys::LeftShift) || InViewport->KeyState(EKeys::RightShift);
-		const bool bAltDown = InViewport->KeyState(EKeys::LeftAlt) || InViewport->KeyState(EKeys::RightAlt);
-		const bool bMouseButtonDown = InViewport->KeyState(EKeys::LeftMouseButton);
+	const bool bCtrlDown = InViewport->KeyState(EKeys::LeftControl) || InViewport->KeyState(EKeys::RightControl);
+	const bool bShiftDown = InViewport->KeyState(EKeys::LeftShift) || InViewport->KeyState(EKeys::RightShift);
+	const bool bAltDown = InViewport->KeyState(EKeys::LeftAlt) || InViewport->KeyState(EKeys::RightAlt);
+	const bool bMouseButtonDown = InViewport->KeyState(EKeys::LeftMouseButton);
 
-		const FWidget::EWidgetMode WidgetMode = InViewportClient->GetWidgetMode();
-		const EAxisList::Type CurrentAxis = InViewportClient->GetCurrentWidgetAxis();
-		const ECoordSystem CoordSystem = InViewportClient->GetWidgetCoordSystemSpace();
+	const FWidget::EWidgetMode WidgetMode = InViewportClient->GetWidgetMode();
+	const EAxisList::Type CurrentAxis = InViewportClient->GetCurrentWidgetAxis();
+	const ECoordSystem CoordSystem = InViewportClient->GetWidgetCoordSystemSpace();
 
-		if (bIsTransacting && bMouseButtonDown && !bCtrlDown && !bShiftDown && !bAltDown && CurrentAxis != EAxisList::None)
-		{
-			const bool bDoRotation = !Rot.IsZero() && (WidgetMode == FWidget::WM_Rotate || WidgetMode == FWidget::WM_TranslateRotateZ);
-			const bool bDoTranslation = !Drag.IsZero() && (WidgetMode == FWidget::WM_Translate || WidgetMode == FWidget::WM_TranslateRotateZ);
-			const bool bDoScale = !Scale.IsZero() && WidgetMode == FWidget::WM_Scale;
+	if (bIsTransacting && bMouseButtonDown && !bCtrlDown && !bShiftDown && !bAltDown && CurrentAxis != EAxisList::None)
+	{
+		const bool bDoRotation = !Rot.IsZero() && (WidgetMode == FWidget::WM_Rotate || WidgetMode == FWidget::WM_TranslateRotateZ);
+		const bool bDoTranslation = !Drag.IsZero() && (WidgetMode == FWidget::WM_Translate || WidgetMode == FWidget::WM_TranslateRotateZ);
+		const bool bDoScale = !Scale.IsZero() && WidgetMode == FWidget::WM_Scale;
 
 		if (ManipulationLayer != nullptr && AreRigElementsSelected(FRigElementTypeHelper::ToMask(ERigElementType::Control)))
-			{
+		{
 			FTransform ComponentTransform = ManipulationLayer->GetSkeletalMeshComponentTransform();
 			for (AControlRigGizmoActor* GizmoActor : GizmoActors)
-				{
+			{
 				if (GizmoActor->IsSelected())
-					{
+				{
 					// test local vs global
+					if (bManipulatorMadeChange == false)
+					{
+						GEditor->BeginTransaction(LOCTEXT("MoveControlTransaction", "Move Control"));
+					}
 					ManipulationLayer->MoveGizmo(GizmoActor, bDoTranslation, InDrag, bDoRotation, InRot, bDoScale, InScale, ComponentTransform);
 					bManipulatorMadeChange = true;
 				}
-							}
+			}
 
 			RecalcPivotTransform();
 
 			if (bManipulatorMadeChange)
-							{
+			{
 				ManipulationLayer->TickManipulatableObjects(0.f);
-							}
+			}
 			return true;
-								}
+		}
 		else if (ManipulationLayer != nullptr && AreRigElementSelectedAndMovable())
-									{
+		{
 			FTransform ComponentTransform = ManipulationLayer->GetSkeletalMeshComponentTransform();
 
 			// set Bone transform
 			// that will set initial Bone transform
 			for (int32 Index = 0; Index < SelectedRigElements.Num(); ++Index)
-										{
+			{
 				const ERigElementType SelectedRigElementType = SelectedRigElements[Index].Type;
 
-				if (SelectedRigElementType == ERigElementType::Bone || SelectedRigElementType == ERigElementType::Control)
-			{
-					FTransform NewWorldTransform = OnGetRigElementTransformDelegate.Execute(SelectedRigElements[Index], false) * ComponentTransform;
-				bool bTransformChanged = false;
-				if (bDoRotation)
+				if (SelectedRigElementType == ERigElementType::Control)
 				{
-					FQuat CurrentRotation = NewWorldTransform.GetRotation();
-					CurrentRotation = (Rot.Quaternion() * CurrentRotation);
-					NewWorldTransform.SetRotation(CurrentRotation);
-					bTransformChanged = true;
-				}
+					FTransform NewWorldTransform = OnGetRigElementTransformDelegate.Execute(SelectedRigElements[Index], false) * ComponentTransform;
+					bool bTransformChanged = false;
+					if (bDoRotation)
+					{
+						FQuat CurrentRotation = NewWorldTransform.GetRotation();
+						CurrentRotation = (Rot.Quaternion() * CurrentRotation);
+						NewWorldTransform.SetRotation(CurrentRotation);
+						bTransformChanged = true;
+					}
 
-				if (bDoTranslation)
-				{
-					FVector CurrentLocation = NewWorldTransform.GetLocation();
-					CurrentLocation = CurrentLocation + Drag;
-					NewWorldTransform.SetLocation(CurrentLocation);
-					bTransformChanged = true;
-				}
+					if (bDoTranslation)
+					{
+						FVector CurrentLocation = NewWorldTransform.GetLocation();
+						CurrentLocation = CurrentLocation + Drag;
+						NewWorldTransform.SetLocation(CurrentLocation);
+						bTransformChanged = true;
+					}
 
-				if (bDoScale)
-				{
-					FVector CurrentScale = NewWorldTransform.GetScale3D();
-					CurrentScale = CurrentScale + Scale;
-					NewWorldTransform.SetScale3D(CurrentScale);
-					bTransformChanged = true;
-				}
+					if (bDoScale)
+					{
+						FVector CurrentScale = NewWorldTransform.GetScale3D();
+						CurrentScale = CurrentScale + Scale;
+						NewWorldTransform.SetScale3D(CurrentScale);
+						bTransformChanged = true;
+					}
 
-				if (bTransformChanged)
-				{
-					FTransform NewComponentTransform = NewWorldTransform.GetRelativeTransform(ComponentTransform);
+					if (bTransformChanged)
+					{
+						if (bManipulatorMadeChange == false)
+						{
+							GEditor->BeginTransaction(LOCTEXT("MoveControlTransaction", "Move Control"));
+						}
+						FTransform NewComponentTransform = NewWorldTransform.GetRelativeTransform(ComponentTransform);
 						OnSetRigElementTransformDelegate.Execute(SelectedRigElements[Index], NewComponentTransform, false);
 						bManipulatorMadeChange = true;
 					}
 				}
-				}
+			}
 
 			// not sure this makes sense @rethink
 			return bManipulatorMadeChange;
-			}
 		}
+	}
 	return false;
 }
 
@@ -756,22 +874,25 @@ void FControlRigEditMode::AddReferencedObjects( FReferenceCollector& Collector )
 {
 	if (Settings)
 	{
-	Collector.AddReferencedObject(Settings);
+		Collector.AddReferencedObject(Settings);
 	}
 	if (ManipulationLayer)
 	{
 		Collector.AddReferencedObject(ManipulationLayer);
 	}
 	if (GizmoActors.Num() > 0)
-					{
+	{
 		for (AControlRigGizmoActor* GizmoActor : GizmoActors)
-						{
+		{
 			Collector.AddReferencedObject(GizmoActor);
-			}
 		}
+	}
+	if (ControlProxy)
+	{
+		Collector.AddReferencedObject(ControlProxy);
+	}
 }
 
-
 void FControlRigEditMode::ClearRigElementSelection(uint32 InTypes)
 {
 	if (!WeakControlRigEditing.IsValid())
@@ -781,11 +902,11 @@ void FControlRigEditMode::ClearRigElementSelection(uint32 InTypes)
 
 	UControlRigBlueprint* Blueprint = Cast<UControlRigBlueprint>(WeakControlRigEditing->GetClass()->ClassGeneratedBy);
 	if (Blueprint)
-		{
+	{
 		Blueprint->HierarchyContainer.ClearSelection();
 	}
-	if(IsInLevelEditor())
-				{
+	if (IsInLevelEditor())
+	{
 		WeakControlRigEditing->Hierarchy.ClearSelection();
 	}
 }
@@ -800,11 +921,11 @@ void FControlRigEditMode::SetRigElementSelectionInternal(ERigElementType Type, c
 
 	UControlRigBlueprint* Blueprint = Cast<UControlRigBlueprint>(WeakControlRigEditing->GetClass()->ClassGeneratedBy);
 	if (Blueprint)
-		{
+	{
 		Blueprint->HierarchyContainer.Select(FRigElementKey(InRigElementName, Type), bSelected);
 	}
-	if(IsInLevelEditor())
-							{
+	if (IsInLevelEditor())
+	{
 		WeakControlRigEditing->Hierarchy.Select(FRigElementKey(InRigElementName, Type), bSelected);
 	}
 }
@@ -828,9 +949,9 @@ void FControlRigEditMode::SetRigElementSelection(ERigElementType Type, const TAr
 		TGuardValue<bool> ReentrantGuard(bSelecting, true);
 
 		for (const FName& ElementName : InRigElementNames)
-	{
+		{
 			SetRigElementSelectionInternal(Type, ElementName, bSelected);
-	}
+		}
 
 		HandleSelectionChanged();
 	}
@@ -852,19 +973,19 @@ bool FControlRigEditMode::AreRigElementsSelected(uint32 InTypes) const
 int32 FControlRigEditMode::GetNumSelectedRigElements(uint32 InTypes) const
 {
 	if (FRigElementTypeHelper::DoesHave(InTypes, ERigElementType::All))
-		{
+	{
 		return SelectedRigElements.Num();
-		}
+	}
 	else
 	{
 		int32 NumSelected = 0;
 		for (const FRigElementKey& Ele : SelectedRigElements)
-			{
+		{
 			if (FRigElementTypeHelper::DoesHave(InTypes, Ele.Type))
-				{
+			{
 				++NumSelected;
-				}
 			}
+		}
 
 		return NumSelected;
 	}
@@ -876,58 +997,71 @@ int32 FControlRigEditMode::GetNumSelectedRigElements(uint32 InTypes) const
 void FControlRigEditMode::RefreshObjects()
 {
 	WeakControlRigEditing = nullptr;
-				ControlRigGuid.Invalidate();
+	ControlRigGuid.Invalidate();
 
-			SetObjects_Internal();
+	SetObjects_Internal();
 }
 
-void FControlRigEditMode::EnableRigElementEditing(bool bEnabled)
+bool FControlRigEditMode::CanRemoveFromPreviewScene(const USceneComponent* InComponent)
 {
-	bEnableRigElementDefaultPoseEditing = bEnabled;
-	RecalcPivotTransform();
+	for (AControlRigGizmoActor* GizmoActor : GizmoActors)
+	{
+		TInlineComponentArray<USceneComponent*> SceneComponents;
+		GizmoActor->GetComponents(SceneComponents, true);
+		if (SceneComponents.Contains(InComponent))
+		{
+			return false;
+		}
+	}
+
+	// we don't need it 
+	return true;
 }
 
 void FControlRigEditMode::RecalcPivotTransform()
 {
 	PivotTransform = FTransform::Identity;
-
+	if (!ManipulationLayer)
+	{
+		return;
+	}
 	// @todo: support bones also
 	if(AreRigElementsSelected(FRigElementTypeHelper::ToMask(ERigElementType::Control)))
-		{
-			FTransform LastTransform = FTransform::Identity;
+	{
+		FTransform LastTransform = FTransform::Identity;
 
 		// recalc coord system too
 		FTransform ComponentTransform = ManipulationLayer->GetSkeletalMeshComponentTransform();
 
-			// Use average location as pivot location
-			FVector PivotLocation = FVector::ZeroVector;
+		// Use average location as pivot location
+		FVector PivotLocation = FVector::ZeroVector;
 
 		int32 NumSelectedControls = 0;
 		for (const AControlRigGizmoActor* GizmoActor : GizmoActors)
-			{
+		{
 			if (GizmoActor->IsSelected())
-				{
+			{
 				LastTransform = GizmoActor->GetActorTransform().GetRelativeTransform(ComponentTransform);
 				PivotLocation += LastTransform.GetLocation();
 				++NumSelectedControls;
-				}
 			}
+		}
 
-			PivotLocation /= (float)NumSelectedControls;
-			PivotTransform.SetLocation(PivotLocation);
+		PivotLocation /= (float)FMath::Max(1, NumSelectedControls);
+		PivotTransform.SetLocation(PivotLocation);
 		
-			if (NumSelectedControls == 1)
-			{
-				// A single Bone just uses its own transform
-				FTransform WorldTransform = LastTransform * ComponentTransform;
-				PivotTransform.SetRotation(WorldTransform.GetRotation());
-			}
-			else if (NumSelectedControls > 1)
-			{
-				// If we have more than one Bone selected, use the coordinate space of the component
-				PivotTransform.SetRotation(ComponentTransform.GetRotation());
-			}
+		if (NumSelectedControls == 1)
+		{
+			// A single Bone just uses its own transform
+			FTransform WorldTransform = LastTransform * ComponentTransform;
+			PivotTransform.SetRotation(WorldTransform.GetRotation());
 		}
+		else if (NumSelectedControls > 1)
+		{
+			// If we have more than one Bone selected, use the coordinate space of the component
+			PivotTransform.SetRotation(ComponentTransform.GetRotation());
+		}
+	}
 	else if (AreRigElementSelectedAndMovable())
 	{
 		// recalc coord system too
@@ -939,19 +1073,19 @@ void FControlRigEditMode::RecalcPivotTransform()
 		FTransform LastTransform = FTransform::Identity;
 		for (int32 Index = 0; Index < SelectedRigElements.Num(); ++Index)
 		{
-			if (SelectedRigElements[Index].Type == ERigElementType::Control || SelectedRigElements[Index].Type == ERigElementType::Bone)
+			if (SelectedRigElements[Index].Type == ERigElementType::Control)
 			{
 				LastTransform = OnGetRigElementTransformDelegate.Execute(SelectedRigElements[Index], false);
 				PivotLocation += LastTransform.GetLocation();
 				++NumSelection;
 			}
-	}
+		}
 
-		PivotLocation /= (float)NumSelection;
+		PivotLocation /= (float)FMath::Max(1, NumSelection);
 		PivotTransform.SetLocation(PivotLocation);
 
 		if (NumSelection == 1)
-	{
+		{
 			// A single Bone just uses its own transform
 			FTransform WorldTransform = LastTransform * ComponentTransform;
 			PivotTransform.SetRotation(WorldTransform.GetRotation());
@@ -971,11 +1105,11 @@ void FControlRigEditMode::HandleSelectionChanged()
 	{
 		TInlineComponentArray<UPrimitiveComponent*> PrimitiveComponents;
 		GizmoActor->GetComponents(PrimitiveComponents, true);
-			for (UPrimitiveComponent* PrimitiveComponent : PrimitiveComponents)
-			{
-				PrimitiveComponent->PushSelectionToProxy();
-			}
+		for (UPrimitiveComponent* PrimitiveComponent : PrimitiveComponents)
+		{
+			PrimitiveComponent->PushSelectionToProxy();
 		}
+	}
 
 	
 	// update the pivot transform of our selected objects (they could be animating)
@@ -997,13 +1131,6 @@ bool FControlRigEditMode::IsControlSelected() const
 	return (AreRigElementsSelected(TypeFlag));
 }
 
-bool FControlRigEditMode::IsControlOrSpaceOrBoneSelected() const
-{
-	static uint32 TypeFlag = (uint32)ERigElementType::Bone | (uint32)ERigElementType::Control | (uint32)ERigElementType::Space;
-	return (AreRigElementsSelected(TypeFlag));
-}
-
-
 bool FControlRigEditMode::GetRigElementGlobalTransform(const FRigElementKey& InElement, FTransform& OutGlobalTransform) const
 {
 	// if control, go through manipulation layer
@@ -1075,15 +1202,25 @@ void FControlRigEditMode::RecreateManipulationLayer()
 {
 	if (ManipulationLayer)
 	{
+		if (AnimInitDelegateHandle.IsValid())
+		{
+			ManipulationLayer->OnAnimSystemInitialized.Remove(AnimInitDelegateHandle);
+			AnimInitDelegateHandle.Reset();
+		}
+
 		ManipulationLayer->DestroyLayer();
 	}
-
+	if (ControlProxy)
+	{
+		ControlProxy->RemoveAllProxies();
+	}
 	if (UControlRig* ControlRig = WeakControlRigEditing.Get())
 	{
 		ManipulationLayer = NewObject<UDefaultControlRigManipulationLayer>();
-		
+
 		// create layer
 		ManipulationLayer->CreateLayer();
+		AnimInitDelegateHandle = ManipulationLayer->OnAnimSystemInitialized.Add(OnAnimSystemInitializedDelegate);
 
 		// default manipulation layer can support any control rig
 		ManipulationLayer->AddManipulatableObject(ControlRig);
@@ -1111,15 +1248,13 @@ void FControlRigEditMode::RecreateManipulationLayer()
 				}
 			}
 		}
+		if (ControlProxy)
+		{
+			ControlProxy->RecreateAllProxies(ControlRig);
+		}
 	}
 }
 
-void RequestToRecreateManipulationLayer()
-{
-
-}
-
-
 FControlRigEditMode* FControlRigEditMode::GetEditModeFromWorldContext(UWorld* InWorldContext)
 {
 	return nullptr;
@@ -1127,11 +1262,15 @@ FControlRigEditMode* FControlRigEditMode::GetEditModeFromWorldContext(UWorld* In
 
 bool FControlRigEditMode::GizmoSelectionOverride(const UPrimitiveComponent* InComponent) const
 {
-	AControlRigGizmoActor* OwnerActor = Cast<AControlRigGizmoActor>(InComponent->GetOwner());
-	if (OwnerActor)
+    //Think we only want to do this in regular editor, in the level editor we are driving selection
+	if (!IsInLevelEditor())
 	{
-		// See if the actor is in a selected unit proxy
-		return OwnerActor->IsSelected();
+		AControlRigGizmoActor* OwnerActor = Cast<AControlRigGizmoActor>(InComponent->GetOwner());
+		if (OwnerActor)
+		{
+			// See if the actor is in a selected unit proxy
+			return OwnerActor->IsSelected();
+		}
 	}
 
 	return false;
@@ -1164,12 +1303,6 @@ bool FControlRigEditMode::AreRigElementSelectedAndMovable() const
 		return true;
 	}
 
-	// if bone, we want to make sure execution is stopped
-	if (bEnableRigElementDefaultPoseEditing && AreRigElementsSelected(FRigElementTypeHelper::ToMask(ERigElementType::Bone)))
-	{
-		return true;
-	}
-
 	if (IsTransformDelegateAvailable())
 	{
 		return true;
@@ -1215,26 +1348,42 @@ void FControlRigEditMode::OnRigElementSelected(FRigHierarchyContainer* Container
 			else
 			{
 				SelectedRigElements.Remove(InKey);
-		}
+			}
 
 			// if it's control
 			if (InKey.Type == ERigElementType::Control)
-		{
+			{
+				FScopedTransaction ScopedTransaction(LOCTEXT("SelectControlTransaction", "Select Control"), !GIsTransacting );
+				ControlProxy->Modify();
 				// users may select gizmo and control rig units, so we have to let them go through both of them if they do
 				// first go through gizmo actor
 				AControlRigGizmoActor* GizmoActor = GetGizmoFromControlName(InKey.Name);
 				if (GizmoActor)
-			{
+				{
 					GizmoActor->SetSelected(bSelected);
+				}
+				if (bSelected)
+				{
+					if (WeakControlRigEditing.IsValid())
+					{
+						FRigControl* Control = WeakControlRigEditing->FindControl(InKey.Name);
+						if (Control)
+						{
+							ControlProxy->SelectProxy(InKey.Name, true);
+						}
+					}
+				}
+				else
+				{
+					ControlProxy->SelectProxy(InKey.Name, false);
+				}
 			}
-			}
-
+			SetUpDetailPanel();
 			HandleSelectionChanged();
 			break;
 		}
 		default:
 		{
-			// todo
 			ensureMsgf(false, TEXT("Unsupported Type of RigElement: %d"), InKey.Type);
 			break;
 		}
@@ -1249,6 +1398,17 @@ void FControlRigEditMode::OnRigElementChanged(FRigHierarchyContainer* Container,
 	}
 }
 
+void FControlRigEditMode::OnControlModified(IControlRigManipulatable* Subject, const FRigControl& Control, EControlRigSetKey InSetKey)
+{
+	FScopedTransaction ScopedTransaction(LOCTEXT("MoveControlTransaction", "Modify Control"),!GIsTransacting && InSetKey != EControlRigSetKey::Never);
+	ControlProxy->Modify();
+	ControlProxy->ProxyChanged(Control.Name);
+	RecalcPivotTransform();
+}
+
+void FControlRigEditMode::HandleModifiedEvent(ERigVMGraphNotifType InNotifType, URigVMGraph* InGraph, UObject* InSubject)
+{
+}
 
 void FControlRigEditMode::OnControlUISettingChanged(FRigHierarchyContainer* Container, const FRigElementKey& InKey)
 {
@@ -1260,4 +1420,25 @@ AControlRigGizmoActor* FControlRigEditMode::GetGizmoFromControlName(const FName&
 	return ManipulationLayer ? ManipulationLayer->GetGizmoFromControlName(InControlName) : nullptr;
 }
 
+void FControlRigEditMode::OnWidgetModeChanged(FWidget::EWidgetMode InWidgetMode)
+{
+	if (Settings && Settings->bCoordSystemPerWidgetMode)
+	{
+		TGuardValue<bool> ReentrantGuardSelf(bIsChangingCoordSystem, true);
+
+		int32 WidgetMode = (int32)GLevelEditorModeTools().GetWidgetMode();
+		GLevelEditorModeTools().SetCoordSystem(CoordSystemPerWidgetMode[WidgetMode]);
+	}
+}
+
+void FControlRigEditMode::OnCoordSystemChanged(ECoordSystem InCoordSystem)
+{
+	TGuardValue<bool> ReentrantGuardSelf(bIsChangingCoordSystem, true);
+
+	int32 WidgetMode = (int32)GLevelEditorModeTools().GetWidgetMode();
+	ECoordSystem CoordSystem = GLevelEditorModeTools().GetCoordSystem();
+	CoordSystemPerWidgetMode[WidgetMode] = CoordSystem;
+}
+
+
 #undef LOCTEXT_NAMESPACE
diff --git a/Engine/Plugins/Experimental/ControlRig/Source/ControlRigEditor/Private/EditMode/ControlRigEditMode.h b/Engine/Plugins/Experimental/ControlRig/Source/ControlRigEditor/Private/EditMode/ControlRigEditMode.h
--- a/Engine/Plugins/Experimental/ControlRig/Source/ControlRigEditor/Private/EditMode/ControlRigEditMode.h
+++ b/Engine/Plugins/Experimental/ControlRig/Source/ControlRigEditor/Private/EditMode/ControlRigEditMode.h
@@ -5,7 +5,7 @@
 #include "CoreMinimal.h"
 #include "InputCoreTypes.h"
 #include "IPersonaEditMode.h"
-#include "ControlRigModel.h"
+#include "RigVMModel/RigVMGraph.h"
 #include "Rigs/RigHierarchyContainer.h"
 #include "Drawing/ControlRigDrawInterface.h"
 #include "Units/RigUnitContext.h"
@@ -26,6 +26,12 @@ class FExtender;
 class IMovieScenePlayer;
 class AControlRigGizmoActor;
 class UDefaultControlRigManipulationLayer;
+class UControlRigDetailPanelControlProxies;
+class UControlRigControlsProxy;
+struct FRigControl;
+class IControlRigManipulatable;
+class ISequencer;
+enum class EControlRigSetKey : uint8;
 
 DECLARE_DELEGATE_RetVal_TwoParams(FTransform, FOnGetRigElementTransform, const FRigElementKey& /*RigElementKey*/, bool /*bLocal*/);
 DECLARE_DELEGATE_ThreeParams(FOnSetRigElementTransform, const FRigElementKey& /*RigElementKey*/, const FTransform& /*Transform*/, bool /*bLocal*/);
@@ -40,7 +46,7 @@ public:
 	~FControlRigEditMode();
 
 	/** Set the objects to be displayed in the details panel */
-	void SetObjects(const TWeakObjectPtr<>& InSelectedObject, const FGuid& InObjectBinding, UObject* BindingObject);
+	void SetObjects(const TWeakObjectPtr<>& InSelectedObject, const FGuid& InObjectBinding, UObject* BindingObject, TWeakPtr<ISequencer> InSequencer);
 
 	/** This edit mode is re-used between the level editor and the control rig editor. Calling this indicates which context we are in */
 	virtual bool IsInLevelEditor() const { return true; }
@@ -93,6 +99,10 @@ public:
 	/** Context Menu Delegates */
 	FNewMenuDelegate& OnContextMenu() { return OnContextMenuDelegate; }
 	FNewMenuCommandsDelegate& OnContextMenuCommands() { return OnContextMenuCommandsDelegate; }
+	FSimpleDelegate&  OnAnimSystemInitialized() { return OnAnimSystemInitializedDelegate; }
+
+	FRigVMGraphModifiedEvent ModifiedEvent;
+	void HandleModifiedEvent(ERigVMGraphNotifType InNotifType, URigVMGraph* InGraph, UObject* InSubject);
 
 	// callback that gets called when rig element is selected in other view
 	void OnRigElementAdded(FRigHierarchyContainer* Container, const FRigElementKey& InKey);
@@ -102,9 +112,11 @@ public:
 	void OnRigElementSelected(FRigHierarchyContainer* Container, const FRigElementKey& InKey, bool bSelected);
 	void OnRigElementChanged(FRigHierarchyContainer* Container, const FRigElementKey& InKey);
 	void OnControlUISettingChanged(FRigHierarchyContainer* Container, const FRigElementKey& InKey);
+	void OnControlModified(IControlRigManipulatable* Subject, const FRigControl& Control, EControlRigSetKey InSetKey);
 
-	/** Enable RigElement Editing */
-	void EnableRigElementEditing(bool bEnabled);
+	/** return true if it can be removed from preview scene 
+	- this is to ensure preview scene doesn't remove Gizmo actors */
+	bool CanRemoveFromPreviewScene(const USceneComponent* InComponent);
 
 	/** Get Control Rig, could be more than one later,  we are animating. Currently used by Sequencer for cross selection.*/
 	UControlRig*  GetControlRig() { return WeakControlRigEditing.IsValid() ? WeakControlRigEditing.Get() : nullptr; }
@@ -114,10 +126,14 @@ protected:
 	/** Get the node name from the property path */
 	AControlRigGizmoActor* GetGizmoFromControlName(const FName& InControlName) const;
 
+
 protected:
 	/** Helper function: set ControlRigs array to the details panel */
 	void SetObjects_Internal();
 
+	/** Set up Details Panel based upon Selected Objects*/
+	void SetUpDetailPanel();
+
 	/** Updates cached pivot transform */
 	void RecalcPivotTransform();
 
@@ -143,6 +159,9 @@ protected:
 	bool GizmoSelectionOverride(const UPrimitiveComponent* InComponent) const;
 
 protected:
+
+	TWeakPtr<ISequencer> WeakSequencer;
+
 	/** Settings object used to insert controls into the details panel */
 	UControlRigEditModeSettings* Settings;
 
@@ -179,22 +198,21 @@ protected:
 	FOnSetRigElementTransform OnSetRigElementTransformDelegate;
 	FNewMenuDelegate OnContextMenuDelegate;
 	FNewMenuCommandsDelegate OnContextMenuCommandsDelegate;
+	FSimpleDelegate OnAnimSystemInitializedDelegate;
+	FDelegateHandle AnimInitDelegateHandle;
 	
 	TArray<FRigElementKey> SelectedRigElements;
 
-	/* Flag to enable element init pose editing @Todo: */
-	bool bEnableRigElementDefaultPoseEditing;
-
 	/* Flag to recreate manipulation layer during tick */
 	bool bRecreateManipulationLayerRequired;
 
 	/** Default Manipulation Layer */
 	UDefaultControlRigManipulationLayer* ManipulationLayer;
 	TArray<AControlRigGizmoActor*> GizmoActors;
+	UControlRigDetailPanelControlProxies* ControlProxy;
 
 	/** Utility functions for UI/Some other viewport manipulation*/
 	bool IsControlSelected() const;
-	bool IsControlOrSpaceOrBoneSelected() const;
 	bool AreRigElementSelectedAndMovable() const;
 	
 	/** Set initial transform handlers */
@@ -210,6 +228,8 @@ public:
 	/** Set a RigElement's selection state */
 	void SetRigElementSelection(ERigElementType Type, const FName& InRigElementName, bool bSelected);
 
+	void SetSelectedRigElement(const FName& InElementName, ERigElementType InRigElementType);
+
 	/** Set multiple RigElement's selection states */
 	void SetRigElementSelection(ERigElementType Type, const TArray<FName>& InRigElementNames, bool bSelected);
 
@@ -222,9 +242,18 @@ public:
 private:
 	/** Set a RigElement's selection state */
 	void SetRigElementSelectionInternal(ERigElementType Type, const FName& InRigElementName, bool bSelected);
-
+	
 	FEditorViewportClient* CurrentViewportClient;
 
+/* store coordinate system per widget mode*/
+private:
+	void OnWidgetModeChanged(FWidget::EWidgetMode InWidgetMode);
+	void OnCoordSystemChanged(ECoordSystem InCoordSystem);
+	TArray<ECoordSystem> CoordSystemPerWidgetMode;
+	bool bIsChangingCoordSystem;
+
+private:
+
 	friend class FControlRigEditorModule;
 	friend class UControlRigPickerWidget;
 	friend class FControlRigEditor;
diff --git a/Engine/Plugins/Experimental/ControlRig/Source/ControlRigEditor/Private/EditMode/ControlRigEditModeCommands.cpp b/Engine/Plugins/Experimental/ControlRig/Source/ControlRigEditor/Private/EditMode/ControlRigEditModeCommands.cpp
--- a/Engine/Plugins/Experimental/ControlRig/Source/ControlRigEditor/Private/EditMode/ControlRigEditModeCommands.cpp
+++ b/Engine/Plugins/Experimental/ControlRig/Source/ControlRigEditor/Private/EditMode/ControlRigEditModeCommands.cpp
@@ -7,7 +7,7 @@
 void FControlRigEditModeCommands::RegisterCommands()
 {
 	UI_COMMAND(SetKey, "Set Key", "Sets a key on the selected nodes", EUserInterfaceActionType::Button, FInputChord(EKeys::Enter));
-	UI_COMMAND(ToggleManipulators, "Toggle Manipulators", "Toggles visibility of manipulators in the viewport", EUserInterfaceActionType::Button, FInputChord(EKeys::R));
+	UI_COMMAND(ToggleManipulators, "Toggle Manipulators", "Toggles visibility of manipulators in the viewport", EUserInterfaceActionType::Button, FInputChord(EKeys::T));
 	UI_COMMAND(ExportAnimSequence, "Export To Animation", "Export this sequence to an optimized skeleton-specific animation sequence", EUserInterfaceActionType::Button, FInputChord());
 	UI_COMMAND(ReExportAnimSequence, "Re-export To Animation", "Re-export this sequence to an animation sequence using the previous export settings", EUserInterfaceActionType::Button, FInputChord());
 	UI_COMMAND(ImportFromRigSequence, "Import From Rig Sequence", "Import this animation sequence from a source rig sequence", EUserInterfaceActionType::Button, FInputChord());
diff --git a/Engine/Plugins/Experimental/ControlRig/Source/ControlRigEditor/Private/EditMode/ControlRigEditModeSettings.cpp b/Engine/Plugins/Experimental/ControlRig/Source/ControlRigEditor/Private/EditMode/ControlRigEditModeSettings.cpp
--- a/Engine/Plugins/Experimental/ControlRig/Source/ControlRigEditor/Private/EditMode/ControlRigEditModeSettings.cpp
+++ b/Engine/Plugins/Experimental/ControlRig/Source/ControlRigEditor/Private/EditMode/ControlRigEditModeSettings.cpp
@@ -4,16 +4,22 @@
 #include "EditorModeManager.h"
 #include "ControlRigEditMode.h"
 #include "Sequencer/ControlRigSequence.h"
+#include "Rigs/RigControlHierarchy.h"
 
 #include "Components/SkeletalMeshComponent.h"
 #include "Subsystems/AssetEditorSubsystem.h"
+#include "MovieSceneCommonHelpers.h"
+
 
 void UControlRigEditModeSettings::PreEditChange(FProperty* PropertyAboutToChange)
 {
-	
+	Super::PreEditChange(PropertyAboutToChange);
+
 }
 
 void UControlRigEditModeSettings::PostEditChangeProperty(struct FPropertyChangedEvent& PropertyChangedEvent)
 {
-	
+	Super::PostEditChangeProperty(PropertyChangedEvent);
+
 }
+
diff --git a/Engine/Plugins/Experimental/ControlRig/Source/ControlRigEditor/Private/EditMode/ControlRigEditModeSettings.h b/Engine/Plugins/Experimental/ControlRig/Source/ControlRigEditor/Private/EditMode/ControlRigEditModeSettings.h
--- a/Engine/Plugins/Experimental/ControlRig/Source/ControlRigEditor/Private/EditMode/ControlRigEditModeSettings.h
+++ b/Engine/Plugins/Experimental/ControlRig/Source/ControlRigEditor/Private/EditMode/ControlRigEditModeSettings.h
@@ -7,10 +7,6 @@
 #include "UObject/LazyObjectPtr.h"
 #include "ControlRigEditModeSettings.generated.h"
 
-class AActor;
-class UControlRig;
-class UControlRigSequence;
-
 /** Settings object used to show useful information in the details panel */
 UCLASS()
 class UControlRigEditModeSettings : public UObject
@@ -22,6 +18,8 @@ class UControlRigEditModeSettings : public UObject
 		, bHideManipulators(false)
 		, bDisplayAxesOnSelection(true)
 		, AxisScale(10.f)
+		, bCoordSystemPerWidgetMode(true)
+		, bOnlySelectRigControls(false)
 	{}
 
 	// UObject interface
@@ -46,4 +44,12 @@ public:
 	UPROPERTY(EditAnywhere, AdvancedDisplay, Category = "Animation")
 	float AxisScale;
 
+	/** If true we restore the coordinate space when changing Widget Modes in the Viewport*/
+	UPROPERTY(EditAnywhere, AdvancedDisplay, Category = "Animation")
+	bool bCoordSystemPerWidgetMode;
+
+	/** If true we can only select Rig Controls in the scene not other Actors. */
+	UPROPERTY(EditAnywhere, AdvancedDisplay, Category = "Animation")
+	bool bOnlySelectRigControls;
+	
 };
\ No newline at end of file
diff --git a/Engine/Plugins/Experimental/ControlRig/Source/ControlRigEditor/Private/EditMode/ControlRigEditModeToolkit.h b/Engine/Plugins/Experimental/ControlRig/Source/ControlRigEditor/Private/EditMode/ControlRigEditModeToolkit.h
--- a/Engine/Plugins/Experimental/ControlRig/Source/ControlRigEditor/Private/EditMode/ControlRigEditModeToolkit.h
+++ b/Engine/Plugins/Experimental/ControlRig/Source/ControlRigEditor/Private/EditMode/ControlRigEditModeToolkit.h
@@ -18,7 +18,7 @@ public:
 	FControlRigEditModeToolkit(FControlRigEditMode& InEditMode)
 		: EditMode(InEditMode)
 	{
-		SAssignNew(ModeTools, SControlRigEditModeTools, EditMode.GetWorld());
+		SAssignNew(ModeTools, SControlRigEditModeTools, EditMode, EditMode.GetWorld());
 	}
 
 	/** IToolkit interface */
diff --git a/Engine/Plugins/Experimental/ControlRig/Source/ControlRigEditor/Private/EditMode/SControlHierarchy.cpp b/Engine/Plugins/Experimental/ControlRig/Source/ControlRigEditor/Private/EditMode/SControlHierarchy.cpp
new file mode 100644
--- /dev/null
+++ b/Engine/Plugins/Experimental/ControlRig/Source/ControlRigEditor/Private/EditMode/SControlHierarchy.cpp
@@ -0,0 +1,458 @@
+// Copyright Epic Games, Inc. All Rights Reserved.
+
+#include "SControlHierarchy.h"
+#include "Widgets/Input/SComboButton.h"
+#include "EditorStyleSet.h"
+#include "Widgets/Input/SSearchBox.h"
+#include "ControlRigEditMode.h"
+#include "Widgets/Layout/SSpacer.h"
+#include "Widgets/Images/SImage.h"
+#include "Widgets/Text/STextBlock.h"
+#include "Widgets/Text/SInlineEditableTextBlock.h"
+#include "ControlRig.h"
+#include "ControlRigEditorStyle.h"
+#include "ScopedTransaction.h"
+
+#define LOCTEXT_NAMESPACE "SControlHierarchy"
+
+//////////////////////////////////////////////////////////////
+/// FControlTreeElement
+///////////////////////////////////////////////////////////
+FControlTreeElement::FControlTreeElement(const FRigElementKey& InKey, TWeakPtr<SControlHierarchy> InHierarchyHandler)
+{
+	Key = InKey;
+}
+
+TSharedRef<ITableRow> FControlTreeElement::MakeTreeRowWidget( const TSharedRef<STableViewBase>& InOwnerTable, TSharedRef<FControlTreeElement> InRigTreeElement, TSharedPtr<SControlHierarchy> InHierarchy)
+{
+	return SNew(SControlHierarchyItem, InOwnerTable, InRigTreeElement, InHierarchy);
+		
+}
+
+
+//////////////////////////////////////////////////////////////
+/// SControlHierarchyItem
+///////////////////////////////////////////////////////////
+void SControlHierarchyItem::Construct(const FArguments& InArgs, const TSharedRef<STableViewBase>& OwnerTable, TSharedRef<FControlTreeElement> InRigTreeElement, TSharedPtr<SControlHierarchy> InHierarchy)
+{
+	WeakRigTreeElement = InRigTreeElement;
+
+	TSharedPtr< SInlineEditableTextBlock > InlineWidget;
+
+	const FSlateBrush* Brush = nullptr;
+	switch (InRigTreeElement->Key.Type)
+		{
+		case ERigElementType::Control:
+		{
+			Brush = FControlRigEditorStyle::Get().GetBrush("ControlRig.Tree.Control");
+			break;
+		}
+	
+		default:
+		{
+			Brush = FControlRigEditorStyle::Get().GetBrush("ControlRig.Tree.BoneUser");
+			break;
+		}
+	}
+
+	STableRow<TSharedPtr<FControlTreeElement>>::Construct(
+		STableRow<TSharedPtr<FControlTreeElement>>::FArguments()
+		.Content()
+		[
+			SNew(SHorizontalBox)
+			+ SHorizontalBox::Slot()
+		.MaxWidth(18)
+		.FillWidth(1.0)
+		.HAlign(HAlign_Left)
+		.VAlign(VAlign_Center)
+		[
+			SNew(SImage)
+			.Image(Brush)
+		]
+	+ SHorizontalBox::Slot()
+		.AutoWidth()
+		.VAlign(VAlign_Center)
+		[
+			SAssignNew(InlineWidget, SInlineEditableTextBlock)
+			.Text(this, &SControlHierarchyItem::GetName)
+		.MultiLine(false)
+		]
+		], OwnerTable);
+
+}
+
+FText SControlHierarchyItem::GetName() const
+{
+	return (FText::FromName(WeakRigTreeElement.Pin()->Key.Name));
+}
+
+///////////////////////////////////////////////////////////
+
+SControlHierarchy::~SControlHierarchy()
+{
+	if (ControlRig.IsValid())
+	{
+		ControlRig->ControlSelected().RemoveAll(this);
+	}	
+}
+
+void SControlHierarchy::Construct(const FArguments& InArgs, UControlRig* InControlRig)
+{
+	
+	ChildSlot
+	[
+		SNew(SVerticalBox)
+		+SVerticalBox::Slot()
+		.AutoHeight()
+		.VAlign(VAlign_Top)
+		.Padding(0.0f)
+		[
+			SNew(SBorder)
+			.Padding(0.0f)
+			.BorderImage(FEditorStyle::GetBrush("DetailsView.CategoryTop"))
+			.BorderBackgroundColor(FLinearColor(0.6f, 0.6f, 0.6f, 1.0f))
+			[
+				SNew(SVerticalBox)
+				+SVerticalBox::Slot()
+				.AutoHeight()
+				.VAlign(VAlign_Top)
+				[
+					SNew(SHorizontalBox)
+					+SHorizontalBox::Slot()
+					.VAlign(VAlign_Center)
+					.Padding(3.0f, 1.0f)
+					[
+						SAssignNew(FilterBox, SSearchBox)
+						.OnTextChanged(this, &SControlHierarchy::OnFilterTextChanged)
+					]
+				]
+			]
+		]
+		+SVerticalBox::Slot()
+		.Padding(0.0f, 0.0f)
+		[
+			SNew(SBorder)
+			.Padding(2.0f)
+			.BorderImage(FEditorStyle::GetBrush("SCSEditor.TreePanel"))
+			[
+				SAssignNew(TreeView, STreeView<TSharedPtr<FControlTreeElement>>)
+				.TreeItemsSource(&RootElements)
+				.SelectionMode(ESelectionMode::Multi)
+				.OnGenerateRow(this, &SControlHierarchy::MakeTableRowWidget)
+				.OnGetChildren(this, &SControlHierarchy::HandleGetChildrenForTree)
+				.OnSelectionChanged(this, &SControlHierarchy::OnSelectionChanged)
+				.HighlightParentNodesForSelection(true)
+				.ItemHeight(24)
+			]
+		]
+	];
+
+	SetControlRig(InControlRig);
+}
+
+void SControlHierarchy::SetControlRig(UControlRig* InControlRig)
+{
+	bSelecting = false;
+	if (ControlRig.IsValid())
+	{
+		ControlRig->ControlSelected().RemoveAll(this);
+	}
+	ControlRig = InControlRig;
+	if (ControlRig.IsValid())
+	{
+		ControlRig->ControlSelected().AddRaw(this, &SControlHierarchy::OnRigElementSelected);
+	}
+	RefreshTreeView();
+
+}
+void SControlHierarchy::OnFilterTextChanged(const FText& SearchText)
+{
+	FilterText = SearchText;
+
+	RefreshTreeView();
+}
+
+void SControlHierarchy::RefreshTreeView()
+{
+	
+	TMap<FRigElementKey, bool> ExpansionState;
+	for (TPair<FRigElementKey, TSharedPtr<FControlTreeElement>> Pair : ElementMap)
+	{
+		ExpansionState.FindOrAdd(Pair.Key) = TreeView->IsItemExpanded(Pair.Value);
+	}
+
+	RootElements.Reset();
+	ElementMap.Reset();
+	ParentMap.Reset();
+
+	if (ControlRig.IsValid())
+	{
+		FRigControlHierarchy& ControlHierarchy = ControlRig->GetControlHierarchy();
+		for (const FRigControl& Element : ControlHierarchy)
+		{
+			if (!ControlRig->IsCurveControl(&Element))
+			{
+				AddControlElement(Element);
+			}
+		}
+		if (ExpansionState.Num() == 0)
+		{
+			for (TSharedPtr<FControlTreeElement> RootElement : RootElements)
+			{
+				SetExpansionRecursive(RootElement, false);
+			}
+		}
+		else
+		{
+			for (TPair<FRigElementKey, bool> Pair : ExpansionState)
+			{
+				if (!Pair.Value)
+				{
+					continue;
+				}
+				if (TSharedPtr<FControlTreeElement>* ItemPtr = ElementMap.Find(Pair.Key))
+				{
+					TreeView->SetItemExpansion(*ItemPtr, true);
+				}
+			}
+		}
+
+		TreeView->RequestTreeRefresh();
+
+		TArray<FName> Selection = ControlHierarchy.CurrentSelection();
+		for (const FName& Name : Selection)
+		{
+			for (const FRigControl& Control : ControlHierarchy.GetControls())
+			{
+				if (Name == Control.Name)
+				{
+					OnRigElementSelected(ControlRig.Get(),Control, true);
+					break;
+				}
+			}
+		}
+	}
+}
+
+void SControlHierarchy::SetExpansionRecursive(TSharedPtr<FControlTreeElement> InElement, bool bTowardsParent)
+{
+	TreeView->SetItemExpansion(InElement, true);
+
+	if (bTowardsParent)
+	{
+		if (const FRigElementKey* ParentKey = ParentMap.Find(InElement->Key))
+		{
+			if (TSharedPtr<FControlTreeElement>* ParentItem = ElementMap.Find(*ParentKey))
+			{
+				SetExpansionRecursive(*ParentItem, bTowardsParent);
+			}
+		}
+	}
+	else
+	{
+		for (int32 ChildIndex = 0; ChildIndex < InElement->Children.Num(); ++ChildIndex)
+		{
+			SetExpansionRecursive(InElement->Children[ChildIndex], bTowardsParent);
+		}
+	}
+}
+TSharedRef<ITableRow> SControlHierarchy::MakeTableRowWidget(TSharedPtr<FControlTreeElement> InItem, const TSharedRef<STableViewBase>& OwnerTable)
+{
+	return InItem->MakeTreeRowWidget( OwnerTable, InItem.ToSharedRef(), SharedThis(this));
+}
+
+void SControlHierarchy::HandleGetChildrenForTree(TSharedPtr<FControlTreeElement> InItem, TArray<TSharedPtr<FControlTreeElement>>& OutChildren)
+{
+	OutChildren = InItem.Get()->Children;
+}
+
+void SControlHierarchy::OnSelectionChanged(TSharedPtr<FControlTreeElement> Selection, ESelectInfo::Type SelectInfo)
+{
+	if (!bSelecting)
+	{
+		TGuardValue<bool>(bSelecting, true);
+		FRigHierarchyContainer* Hierarchy = GetHierarchyContainer();
+
+		if (Hierarchy)
+		{
+			FScopedTransaction ScopedTransaction(LOCTEXT("SelectControlTransaction", "Select Control"), !GIsTransacting);
+
+			TArray<FRigElementKey> OldSelection = Hierarchy->CurrentSelection();
+			TArray<FRigElementKey> NewSelection;
+
+			TArray<TSharedPtr<FControlTreeElement>> SelectedItems = TreeView->GetSelectedItems();
+			for (const TSharedPtr<FControlTreeElement>& SelectedItem : SelectedItems)
+			{
+				NewSelection.Add(SelectedItem->Key);
+			}
+
+			for (const FRigElementKey& PreviouslySelected : OldSelection)
+			{
+				if (NewSelection.Contains(PreviouslySelected))
+				{
+					continue;
+				}
+				Hierarchy->Select(PreviouslySelected, false);
+			}
+
+			for (const FRigElementKey& NewlySelected : NewSelection)
+			{
+				Hierarchy->Select(NewlySelected, true);
+			}
+		}
+	}
+}
+
+
+TSharedPtr<FControlTreeElement> SControlHierarchy::FindElement(const FRigElementKey& InElementKey, TSharedPtr<FControlTreeElement> CurrentItem)
+{
+	if (CurrentItem->Key == InElementKey)
+	{
+		return CurrentItem;
+	}
+
+	for (int32 ChildIndex = 0; ChildIndex < CurrentItem->Children.Num(); ++ChildIndex)
+	{
+		TSharedPtr<FControlTreeElement> Found = FindElement(InElementKey, CurrentItem->Children[ChildIndex]);
+		if (Found.IsValid())
+		{
+			return Found;
+		}
+	}
+
+	return TSharedPtr<FControlTreeElement>();
+}
+
+void SControlHierarchy::OnRigElementSelected(IControlRigManipulatable* ControlRigManp, const FRigControl& Control, bool bSelected)
+{
+	if (!bSelecting)
+	{
+		TGuardValue<bool>(bSelecting, true);
+		FRigElementKey Key;
+		Key.Name = Control.Name;
+		Key.Type = ERigElementType::Control;
+		for (int32 RootIndex = 0; RootIndex < RootElements.Num(); ++RootIndex)
+		{
+			TSharedPtr<FControlTreeElement> Found = FindElement(Key, RootElements[RootIndex]);
+			if (Found.IsValid())
+			{
+				TreeView->SetItemSelection(Found, bSelected, ESelectInfo::OnNavigation);
+				TArray<TSharedPtr<FControlTreeElement>> SelectedItems = TreeView->GetSelectedItems();
+				for (TSharedPtr<FControlTreeElement> SelectedItem : SelectedItems)
+				{
+					SetExpansionRecursive(SelectedItem, true);
+				}
+
+				if (SelectedItems.Num() > 0)
+				{
+					TreeView->RequestScrollIntoView(SelectedItems.Last());
+				}
+			}
+		}
+	}
+}
+
+void SControlHierarchy::AddControlElement(FRigControl InControl)
+{
+	const FRigHierarchyContainer* Container = GetHierarchyContainer();
+	const FRigControlHierarchy& ControlHierarchy = Container->ControlHierarchy;
+	const FRigSpaceHierarchy& SpaceHierarchy = Container->SpaceHierarchy;
+
+	FRigElementKey ParentKey;
+	if (InControl.SpaceIndex != INDEX_NONE)
+	{
+		AddSpaceElement(SpaceHierarchy[InControl.SpaceIndex]);
+		ParentKey = SpaceHierarchy[InControl.SpaceIndex].GetElementKey();
+	}
+	else if (InControl.ParentIndex != INDEX_NONE)
+	{
+		AddControlElement(ControlHierarchy[InControl.ParentIndex]);
+		ParentKey = ControlHierarchy[InControl.ParentIndex].GetElementKey();
+	}
+	AddElement(InControl.GetElementKey(), ParentKey);
+}
+
+
+void SControlHierarchy::AddSpaceElement(FRigSpace InSpace)
+{
+	const FRigHierarchyContainer* Container = GetHierarchyContainer();
+	const FRigBoneHierarchy& BoneHierarchy = Container->BoneHierarchy;
+	const FRigControlHierarchy& ControlHierarchy = Container->ControlHierarchy;
+	const FRigSpaceHierarchy& SpaceHierarchy = Container->SpaceHierarchy;
+
+	FRigElementKey ParentKey;
+	if (InSpace.ParentIndex != INDEX_NONE)
+	{
+		switch (InSpace.SpaceType)
+		{
+		case ERigSpaceType::Control:
+		{
+			AddControlElement(ControlHierarchy[InSpace.ParentIndex]);
+			ParentKey = ControlHierarchy[InSpace.ParentIndex].GetElementKey();
+			break;
+		}
+		case ERigSpaceType::Space:
+		{
+			AddSpaceElement(SpaceHierarchy[InSpace.ParentIndex]);
+			ParentKey = SpaceHierarchy[InSpace.ParentIndex].GetElementKey();
+			break;
+		}
+		default:
+		{
+			break;
+		}
+		}
+	}
+	AddElement(InSpace.GetElementKey(), ParentKey);
+}
+
+void SControlHierarchy::AddElement(FRigElementKey InKey, FRigElementKey InParentKey)
+{
+	if (ElementMap.Contains(InKey))
+	{
+		return;
+	}
+
+	FString FilteredString = FilterText.ToString();
+	if (FilteredString.IsEmpty())
+	{
+		TSharedPtr<FControlTreeElement> NewItem = MakeShared<FControlTreeElement>(InKey, SharedThis(this));
+		ElementMap.Add(InKey, NewItem);
+		if (InParentKey)
+		{
+			ParentMap.Add(InKey, InParentKey);
+		}
+
+		if (InParentKey)
+		{
+			TSharedPtr<FControlTreeElement>* FoundItem = ElementMap.Find(InParentKey);
+			check(FoundItem);
+			FoundItem->Get()->Children.Add(NewItem);
+		}
+		else
+		{
+			RootElements.Add(NewItem);
+		}
+	}
+	else
+	{
+		FString FilteredStringUnderScores = FilteredString.Replace(TEXT(" "), TEXT("_"));
+		if (InKey.Name.ToString().Contains(FilteredString) || InKey.Name.ToString().Contains(FilteredStringUnderScores))
+		{
+			TSharedPtr<FControlTreeElement> NewItem = MakeShared<FControlTreeElement>(InKey, SharedThis(this));
+			RootElements.Add(NewItem);
+		}
+	}
+}
+
+
+FRigHierarchyContainer* SControlHierarchy::GetHierarchyContainer() const
+{
+	if (ControlRig.IsValid())
+	{
+		return ControlRig->GetHierarchy();
+	}
+
+	return nullptr;
+}
+#undef LOCTEXT_NAMESPACE
\ No newline at end of file
diff --git a/Engine/Plugins/Experimental/ControlRig/Source/ControlRigEditor/Private/EditMode/SControlHierarchy.h b/Engine/Plugins/Experimental/ControlRig/Source/ControlRigEditor/Private/EditMode/SControlHierarchy.h
new file mode 100644
--- /dev/null
+++ b/Engine/Plugins/Experimental/ControlRig/Source/ControlRigEditor/Private/EditMode/SControlHierarchy.h
@@ -0,0 +1,96 @@
+// Copyright Epic Games, Inc. All Rights Reserved.
+#pragma once
+#include "Widgets/Views/STreeView.h"
+#include "Rigs/RigHierarchyContainer.h"
+
+class FControlRigEditMode;
+struct FRigElementKey;
+class SSearchBox;
+class UControlRig;
+class IControlRigManipulatable;
+struct FRigControl;
+class SControlHierarchy;
+
+class FControlTreeElement : public TSharedFromThis<FControlTreeElement>
+{
+public:
+	FControlTreeElement(const FRigElementKey& InKey, TWeakPtr<SControlHierarchy> InHierarchyHandler);
+public:
+	FRigElementKey Key;
+	TArray<TSharedPtr<FControlTreeElement>> Children;
+	TSharedRef<ITableRow> MakeTreeRowWidget(const TSharedRef<STableViewBase>& InOwnerTable, TSharedRef<FControlTreeElement> InRigTreeElement, TSharedPtr<SControlHierarchy> InHierarchy);
+};
+
+class SControlHierarchyItem : public STableRow<TSharedPtr<FControlTreeElement>>
+{
+	SLATE_BEGIN_ARGS(SControlHierarchyItem) {}
+	SLATE_END_ARGS()
+
+	void Construct(const FArguments& InArgs, const TSharedRef<STableViewBase>& OwnerTable, TSharedRef<FControlTreeElement> InRigTreeElement, TSharedPtr<SControlHierarchy> InHierarchy);
+
+private:
+	TWeakPtr<FControlTreeElement> WeakRigTreeElement;
+	FText GetName() const;
+};
+
+class SControlHierarchy : public SCompoundWidget
+{
+public:
+	SLATE_BEGIN_ARGS(SControlHierarchy) {}
+	SLATE_END_ARGS()
+
+	~SControlHierarchy();
+
+	void Construct(const FArguments& InArgs, UControlRig* InControlRig);
+	void SetControlRig(UControlRig* InControlRig);
+private:
+	/** Rebuild the tree view */
+	void RefreshTreeView();
+
+	/** Make a row widget for the table */
+	TSharedRef<ITableRow> MakeTableRowWidget(TSharedPtr<FControlTreeElement> InItem, const TSharedRef<STableViewBase>& OwnerTable);
+
+	/** Get children for the tree */
+	void HandleGetChildrenForTree(TSharedPtr<FControlTreeElement> InItem, TArray<TSharedPtr<FControlTreeElement>>& OutChildren);
+
+	/** Set Selection Changed */
+	void OnSelectionChanged(TSharedPtr<FControlTreeElement> Selection, ESelectInfo::Type SelectInfo);
+
+private:
+
+	/** If selecting*/
+	bool bSelecting;
+
+	/** ControlRig */
+	TWeakObjectPtr<UControlRig> ControlRig;
+
+	/** Search box widget */
+	TSharedPtr<SSearchBox> FilterBox;
+	FText FilterText;
+
+	void OnFilterTextChanged(const FText& SearchText);
+
+	/** Tree view widget */
+	TSharedPtr<STreeView<TSharedPtr<FControlTreeElement>>> TreeView;
+
+	/** Backing array for tree view */
+	TArray<TSharedPtr<FControlTreeElement>> RootElements;
+
+	/** A map for looking up items based on their key */
+	TMap<FRigElementKey, TSharedPtr<FControlTreeElement>> ElementMap;
+
+	/** A map for looking up a parent based on their key */
+	TMap<FRigElementKey, FRigElementKey> ParentMap;
+private:
+	static TSharedPtr<FControlTreeElement> FindElement(const FRigElementKey& InElementKey, TSharedPtr<FControlTreeElement> CurrentItem);
+
+	void SetExpansionRecursive(TSharedPtr<FControlTreeElement> InElements, bool bTowardsParent);
+	void OnRigElementSelected(IControlRigManipulatable* ControlRigManp, const FRigControl& Control, bool bSelected);
+	FRigHierarchyContainer* GetHierarchyContainer() const;
+	void AddElement(FRigElementKey InKey, FRigElementKey InParentKey = FRigElementKey());
+	void AddControlElement(FRigControl InControl);
+	void AddSpaceElement(FRigSpace InSpace);
+public:
+
+	friend class SControlHierarchyItem;
+};
diff --git a/Engine/Plugins/Experimental/ControlRig/Source/ControlRigEditor/Private/EditMode/SControlRigEditModeTools.cpp b/Engine/Plugins/Experimental/ControlRig/Source/ControlRigEditor/Private/EditMode/SControlRigEditModeTools.cpp
--- a/Engine/Plugins/Experimental/ControlRig/Source/ControlRigEditor/Private/EditMode/SControlRigEditModeTools.cpp
+++ b/Engine/Plugins/Experimental/ControlRig/Source/ControlRigEditor/Private/EditMode/SControlRigEditModeTools.cpp
@@ -1,6 +1,7 @@
 // Copyright Epic Games, Inc. All Rights Reserved.
 
 #include "SControlRigEditModeTools.h"
+#include "ControlRigControlsProxy.h"
 #include "PropertyEditorModule.h"
 #include "IDetailsView.h"
 #include "ISequencer.h"
@@ -15,12 +16,11 @@
 #include "Widgets/Layout/SScrollBox.h"
 #include "MovieSceneSequence.h"
 #include "MovieScene.h"
-
+#include "SControlHierarchy.h"
 #include "SControlPicker.h"
 
 #define LOCTEXT_NAMESPACE "ControlRigRootCustomization"
 
-
 class FControlRigRootCustomization : public IDetailRootObjectCustomization
 {
 public:
@@ -35,13 +35,18 @@ public:
 		return true;
 	}
 
-	virtual bool ShouldDisplayHeader(const UObject* InRootObject) const override
+	virtual bool ShouldDisplayHeader(const UObject* InRootObAject) const override
 	{
 		return false;
 	}
 };
 
-void SControlRigEditModeTools::Construct(const FArguments& InArgs, UWorld* InWorld)
+void SControlRigEditModeTools::SetControlRig(UControlRig* ControlRig)
+{
+	ControlHierarchy->SetControlRig(ControlRig);
+}
+
+void SControlRigEditModeTools::Construct(const FArguments& InArgs, FControlRigEditMode& InEditMode,UWorld* InWorld)
 {
 	// initialize settings view
 	FDetailsViewArgs DetailsViewArgs;
@@ -73,7 +78,7 @@ void SControlRigEditModeTools::Construct(const FArguments& InArgs, UWorld* InWor
 		+ SScrollBox::Slot()
 		[
 			SNew(SVerticalBox)
-			/*
+			/* We don't do the picker nor the float controls but leeaving this here inc case we 
 			+ SVerticalBox::Slot()
 			.AutoHeight()
 			[
@@ -87,8 +92,36 @@ void SControlRigEditModeTools::Construct(const FArguments& InArgs, UWorld* InWor
 					SAssignNew(ControlPicker, SControlPicker, InWorld)
 				]
 			]
+			
+
+			+SVerticalBox::Slot()
+			.AutoHeight()
+			[
+				SAssignNew(PickerExpander, SExpandableArea)
+				.InitiallyCollapsed(true)
+				.AreaTitle(LOCTEXT("CurveControl_Header", "Curve Controls"))
+				.AreaTitleFont(FEditorStyle::GetFontStyle("DetailsView.CategoryFontStyle"))
+				.BorderBackgroundColor(FLinearColor(.6f, .6f, .6f))
+				.BodyContent()
+				[
+					SAssignNew(CurveControlContainer, SCurveControlContainer, InEditMode.GetControlRig())
+				]
+			]
 			*/
 
+			+SVerticalBox::Slot()
+			.AutoHeight()
+			[
+				SAssignNew(PickerExpander, SExpandableArea)
+				.InitiallyCollapsed(true)
+				.AreaTitle(LOCTEXT("Picker_Header", "Controls"))
+				.AreaTitleFont(FEditorStyle::GetFontStyle("DetailsView.CategoryFontStyle"))
+				.BorderBackgroundColor(FLinearColor(.6f, .6f, .6f))
+				.BodyContent()
+				[
+					SAssignNew(ControlHierarchy, SControlHierarchy, InEditMode.GetControlRig())
+				]
+			]
 			+SVerticalBox::Slot()
 			.AutoHeight()
 			[
@@ -96,6 +129,13 @@ void SControlRigEditModeTools::Construct(const FArguments& InArgs, UWorld* InWor
 			]
 		]
 	];
+
+	// Bind notification when edit mode selection changes, so we can update picker
+	FControlRigEditMode* ControlRigEditMode = static_cast<FControlRigEditMode*>(GLevelEditorModeTools().GetActiveMode(FControlRigEditMode::ModeName));
+	if (ControlRigEditMode)
+	{
+		ControlRigEditMode->ModifiedEvent.AddSP(this, &SControlRigEditModeTools::HandleModifiedEvent);
+	}	
 }
 
 void SControlRigEditModeTools::SetDetailsObjects(const TArray<TWeakObjectPtr<>>& InObjects)
@@ -117,9 +157,9 @@ void SControlRigEditModeTools::SetDetailsObjects(const TArray<TWeakObjectPtr<>>&
 
 }
 
-void SControlRigEditModeTools::SetSequencer(TSharedPtr<ISequencer> InSequencer)
+void SControlRigEditModeTools::SetSequencer(TWeakPtr<ISequencer> InSequencer)
 {
-	WeakSequencer = InSequencer;
+	WeakSequencer = InSequencer.Pin();
 }
 
 bool SControlRigEditModeTools::IsPropertyKeyable(UClass* InObjectClass, const IPropertyHandle& InPropertyHandle) const
@@ -168,14 +208,20 @@ bool SControlRigEditModeTools::IsPropertyAnimated(const IPropertyHandle& Propert
 
 void SControlRigEditModeTools::OnKeyPropertyClicked(const IPropertyHandle& KeyedPropertyHandle)
 {
+	if (WeakSequencer.IsValid() && !WeakSequencer.Pin()->IsAllowedToChange())
+	{
+		return;
+	}
+
 	TArray<UObject*> Objects;
 	KeyedPropertyHandle.GetOuterObjects(Objects);
-	FKeyPropertyParams KeyPropertyParams(Objects, KeyedPropertyHandle, ESequencerKeyMode::ManualKeyForced);
-
-	TSharedPtr<ISequencer> Sequencer = WeakSequencer.Pin();
-	if (Sequencer.IsValid())
+	for (UObject *Object : Objects)
 	{
-		Sequencer->KeyProperty(KeyPropertyParams);
+		UControlRigControlsProxy* Proxy = Cast< UControlRigControlsProxy>(Object);
+		if (Proxy)
+		{
+			Proxy->SetKey(KeyedPropertyHandle);
+		}
 	}
 }
 
@@ -191,7 +237,14 @@ bool SControlRigEditModeTools::ShouldShowPropertyOnDetailCustomization(const FPr
 */
 
 		// Always show settings properties
-		bShow |= InProperty.GetOwner<UClass>() == UControlRigEditModeSettings::StaticClass();
+		const UClass* OwnerClass = InProperty.GetOwner<UClass>();
+		bShow |= OwnerClass == UControlRigEditModeSettings::StaticClass();
+		bShow |= OwnerClass == UControlRigTransformControlProxy::StaticClass();		
+		bShow |= OwnerClass == UControlRigTransformNoScaleControlProxy::StaticClass();
+		bShow |= OwnerClass == UControlRigFloatControlProxy::StaticClass();
+		bShow |= OwnerClass == UControlRigVectorControlProxy::StaticClass();
+		bShow |= OwnerClass == UControlRigVector2DControlProxy::StaticClass();
+		bShow |= OwnerClass == UControlRigBoolControlProxy::StaticClass();
 
 		return bShow;
 	};
@@ -220,7 +273,14 @@ bool SControlRigEditModeTools::IsReadOnlyPropertyOnDetailCustomization(const FPr
 		bool bShow = InProperty.HasAnyPropertyFlags(CPF_Interp) || InProperty.HasMetaData(UControlRig::InputMetaName);
 
 		// Always show settings properties
-		bShow |= Cast<UClass>(InProperty.GetOwner<UObject>()) == UControlRigEditModeSettings::StaticClass();
+		const UClass* OwnerClass = InProperty.GetOwner<UClass>();
+		bShow |= OwnerClass == UControlRigEditModeSettings::StaticClass();
+		bShow |= OwnerClass == UControlRigTransformControlProxy::StaticClass();
+		bShow |= OwnerClass == UControlRigTransformNoScaleControlProxy::StaticClass();
+		bShow |= OwnerClass == UControlRigFloatControlProxy::StaticClass();
+		bShow |= OwnerClass == UControlRigVectorControlProxy::StaticClass();
+		bShow |= OwnerClass == UControlRigVector2DControlProxy::StaticClass();
+		bShow |= OwnerClass == UControlRigBoolControlProxy::StaticClass();
 
 		return bShow;
 	};
@@ -258,4 +318,32 @@ void SControlRigEditModeTools::OnManipulatorsPicked(const TArray<FName>& Manipul
 	}
 }
 
+void SControlRigEditModeTools::HandleModifiedEvent(ERigVMGraphNotifType InNotifType, URigVMGraph* InGraph, UObject* InSubject)
+{
+	if (bPickerChangingSelection)
+	{
+		return;
+	}
+
+	TGuardValue<bool> SelectGuard(bPickerChangingSelection, true);
+	switch (InNotifType)
+	{
+		case ERigVMGraphNotifType::NodeSelected:
+		case ERigVMGraphNotifType::NodeDeselected:
+		{
+			URigVMNode* Node = Cast<URigVMNode>(InSubject);
+			if (Node)
+			{
+				// those are not yet implemented yet
+				// ControlPicker->SelectManipulator(Node->Name, InType == EControlRigModelNotifType::NodeSelected);
+			}
+			break;
+		}
+		default:
+		{
+			break;
+		}
+	}
+}
+
 #undef LOCTEXT_NAMESPACE
diff --git a/Engine/Plugins/Experimental/ControlRig/Source/ControlRigEditor/Private/EditMode/SControlRigEditModeTools.h b/Engine/Plugins/Experimental/ControlRig/Source/ControlRigEditor/Private/EditMode/SControlRigEditModeTools.h
--- a/Engine/Plugins/Experimental/ControlRig/Source/ControlRigEditor/Private/EditMode/SControlRigEditModeTools.h
+++ b/Engine/Plugins/Experimental/ControlRig/Source/ControlRigEditor/Private/EditMode/SControlRigEditModeTools.h
@@ -5,13 +5,16 @@
 #include "CoreMinimal.h"
 #include "Widgets/SCompoundWidget.h"
 #include "IDetailKeyframeHandler.h"
-#include "ControlRigModel.h"
+#include "RigVMModel/RigVMGraph.h"
 #include "Widgets/DeclarativeSyntaxSupport.h"
 
+class FControlRigEditMode;
 class IDetailsView;
 class ISequencer;
 class SControlPicker;
 class SExpandableArea;
+class SControlHierarchy;
+class UControlRig;
 
 class SControlRigEditModeTools : public SCompoundWidget, public IDetailKeyframeHandler
 {
@@ -19,13 +22,16 @@ public:
 	SLATE_BEGIN_ARGS(SControlRigEditModeTools) {}
 	SLATE_END_ARGS();
 
-	void Construct(const FArguments& InArgs, UWorld* InWorld);
+	void Construct(const FArguments& InArgs, FControlRigEditMode& InEditMode, UWorld* InWorld);
 
 	/** Set the objects to be displayed in the details panel */
 	void SetDetailsObjects(const TArray<TWeakObjectPtr<>>& InObjects);
 
 	/** Set the sequencer we are bound to */
-	void SetSequencer(TSharedPtr<ISequencer> InSequencer);
+	void SetSequencer(TWeakPtr<ISequencer> InSequencer);
+
+	/** Set The Control Rig we are using*/
+	void SetControlRig(UControlRig* ControlRig);
 
 	// IDetailKeyframeHandler interface
 	virtual bool IsPropertyKeyable(UClass* InObjectClass, const class IPropertyHandle& PropertyHandle) const override;
@@ -39,9 +45,11 @@ private:
 	/** The details view we do most of our work within */
 	TSharedPtr<IDetailsView> DetailsView;
 
-	/** Special picker for controls */
-	TSharedPtr<SControlPicker> ControlPicker;
+	/** Hierarchy picker for controls*/
+	TSharedPtr<SControlHierarchy> ControlHierarchy;
 
+	/** Special picker for controls, no longer used */
+	TSharedPtr<SControlPicker> ControlPicker;
 	TSharedPtr<SExpandableArea> PickerExpander;
 
 	/** Display or edit set up for property */
@@ -50,4 +58,6 @@ private:
 
 	/** Called when a manipulator is selected in the picker */
 	void OnManipulatorsPicked(const TArray<FName>& Manipulators);
+
+	void HandleModifiedEvent(ERigVMGraphNotifType InNotifType, URigVMGraph* InGraph, UObject* InSubject);
 };
