COPY from //Dev-Anim to //Dev-Main

diff --git a/Engine/Plugins/Experimental/ControlRig/Source/ControlRigEditor/Private/Editor/ControlRigEditor.cpp b/Engine/Plugins/Experimental/ControlRig/Source/ControlRigEditor/Private/Editor/ControlRigEditor.cpp
--- a/Engine/Plugins/Experimental/ControlRig/Source/ControlRigEditor/Private/Editor/ControlRigEditor.cpp
+++ b/Engine/Plugins/Experimental/ControlRig/Source/ControlRigEditor/Private/Editor/ControlRigEditor.cpp
@@ -13,7 +13,6 @@
 #include "Graph/ControlRigGraph.h"
 #include "BlueprintActionDatabase.h"
 #include "ControlRigBlueprintCommands.h"
-#include "ScopedTransaction.h"
 #include "Kismet2/BlueprintEditorUtils.h"
 #include "IPersonaToolkit.h"
 #include "PersonaModule.h"
@@ -21,9 +20,12 @@
 #include "ControlRigEditModeSettings.h"
 #include "AssetEditorModeManager.h"
 #include "ControlRigBlueprintGeneratedClass.h"
+#include "AnimCustomInstanceHelper.h"
+#include "Sequencer/ControlRigLayerInstance.h"
 #include "Sequencer/ControlRigSequencerAnimInstance.h"
 #include "Animation/DebugSkelMeshComponent.h"
 #include "IPersonaPreviewScene.h"
+#include "Persona/Private/AnimationEditorPreviewScene.h"
 #include "Animation/AnimData/BoneMaskFilter.h"
 #include "ControlRig.h"
 #include "ControlRigSkeletalMeshComponent.h"
@@ -52,10 +54,11 @@
 #include "ControlRig/Private/Units/Hierarchy/RigUnit_GetRelativeBoneTransform.h"
 #include "ControlRig/Private/Units/Hierarchy/RigUnit_SetRelativeBoneTransform.h"
 #include "ControlRig/Private/Units/Hierarchy/RigUnit_GetInitialBoneTransform.h"
-#include "ControlRig/Private/Units/Hierarchy/RigUnit_AddBoneTransform.h"
 #include "ControlRig/Private/Units/Execution/RigUnit_BeginExecution.h"
 #include "ControlRig/Private/Units/Hierarchy/RigUnit_GetControlTransform.h"
+#include "ControlRig/Private/Units/Hierarchy/RigUnit_SetControlTransform.h"
 #include "ControlRig/Private/Units/Hierarchy/RigUnit_GetSpaceTransform.h"
+#include "Units/Hierarchy/RigUnit_AddBoneTransform.h"
 #include "Graph/NodeSpawners/ControlRigUnitNodeSpawner.h"
 #include "Graph/ControlRigGraphSchema.h"
 #include "ControlRigObjectVersion.h"
@@ -65,6 +68,10 @@
 #include "SNodePanel.h"
 #include "Kismet/Private/SMyBlueprint.h"
 #include "Kismet/Private/SBlueprintEditorSelectedDebugObjectWidget.h"
+#include "Exporters/Exporter.h"
+#include "UnrealExporter.h"
+#include "ControlRigElementDetails.h"
+#include "PropertyEditorModule.h"
 
 #define LOCTEXT_NAMESPACE "ControlRigEditor"
 
@@ -82,8 +89,9 @@ namespace ControlRigEditorTabs
 FControlRigEditor::FControlRigEditor()
 	: ControlRig(nullptr)
 	, bControlRigEditorInitialized(false)
-	, bIsSelecting(false)
 	, bIsSettingObjectBeingDebugged(false)
+	, NodeDetailStruct(nullptr)
+	, NodeDetailName(NAME_None)
 	, bExecutionControlRig(true)
 {
 }
@@ -93,15 +101,30 @@ FControlRigEditor::~FControlRigEditor()
 	UControlRigBlueprint* RigBlueprint = GetControlRigBlueprint();
 	if (RigBlueprint)
 	{
-		RigBlueprint->HierarchyContainer.OnElementChanged.Clear();
-		RigBlueprint->HierarchyContainer.OnElementAdded.Clear();
-		RigBlueprint->HierarchyContainer.OnElementRemoved.Clear();
-		RigBlueprint->HierarchyContainer.OnElementRenamed.Clear();
-		RigBlueprint->HierarchyContainer.OnElementReparented.Clear();
-		RigBlueprint->HierarchyContainer.OnElementSelected.Clear();
-		RigBlueprint->HierarchyContainer.OnElementChanged.Clear();
-		RigBlueprint->HierarchyContainer.ControlHierarchy.OnControlUISettingsChanged.Clear();
- 	}
+		UControlRigBlueprint::sCurrentlyOpenedRigBlueprints.Remove(RigBlueprint);
+
+		RigBlueprint->HierarchyContainer.OnElementChanged.RemoveAll(this);
+		RigBlueprint->HierarchyContainer.OnElementAdded.RemoveAll(this);
+		RigBlueprint->HierarchyContainer.OnElementRemoved.RemoveAll(this);
+		RigBlueprint->HierarchyContainer.OnElementRenamed.RemoveAll(this);
+		RigBlueprint->HierarchyContainer.OnElementReparented.RemoveAll(this);
+		RigBlueprint->HierarchyContainer.OnElementSelected.RemoveAll(this);
+		if (FControlRigEditMode* EditMode = GetEditMode())
+		{
+			RigBlueprint->HierarchyContainer.OnElementSelected.RemoveAll(EditMode);
+		}
+	}
+
+	if (NodeDetailBuffer.Num() > 0 && NodeDetailStruct != nullptr)
+	{
+		NodeDetailStruct->DestroyStruct(NodeDetailBuffer.GetData(), 1);
+	}
+
+	if (UWorld* PreviewWorld = GetPersonaToolkit()->GetPreviewScene()->GetWorld())
+	{
+		PreviewWorld->MarkObjectsPendingKill();
+		PreviewWorld->MarkPendingKill();
+	}
 }
 
 UControlRigBlueprint* FControlRigEditor::GetControlRigBlueprint() const
@@ -162,10 +185,12 @@ void FControlRigEditor::InitControlRigEditor(const EToolkitMode::Type Mode, cons
 	const bool bCreateDefaultToolbar = true;
 	InitAssetEditor(Mode, InitToolkitHost, ControlRigEditorAppName, DummyLayout, bCreateDefaultStandaloneMenu, bCreateDefaultToolbar, ObjectsBeingEdited);
 
+	UControlRigBlueprint::sCurrentlyOpenedRigBlueprints.AddUnique(InControlRigBlueprint);
+
 	TArray<UBlueprint*> ControlRigBlueprints;
 	ControlRigBlueprints.Add(InControlRigBlueprint);
 
-	InControlRigBlueprint->InitializeModel();
+	InControlRigBlueprint->InitializeModelIfRequired();
 
 	CommonInitialization(ControlRigBlueprints);
 
@@ -180,18 +205,12 @@ void FControlRigEditor::InitControlRigEditor(const EToolkitMode::Type Mode, cons
 			}
 
 			RigGraph->Initialize(InControlRigBlueprint);
-
-			if (RigGraph->GetLinkerCustomVersion(FControlRigObjectVersion::GUID) < FControlRigObjectVersion::RemovalOfHierarchyRefPins)
-			{
-				// recompile in case this control rig requires a rebuild
-				// since we've removed the Hierarchy Ref pins of the first version.
-				Compile();
-			}
 		}
 
 	}
 
-	InControlRigBlueprint->OnModified().AddSP(this, &FControlRigEditor::HandleModelModified);
+	InControlRigBlueprint->OnModified().AddSP(this, &FControlRigEditor::HandleModifiedEvent);
+	InControlRigBlueprint->OnVMCompiled().AddSP(this, &FControlRigEditor::HandleVMCompiledEvent);
 
 	BindCommands();
 
@@ -220,6 +239,11 @@ void FControlRigEditor::InitControlRigEditor(const EToolkitMode::Type Mode, cons
 		EditMode->OnSetRigElementTransform() = FOnSetRigElementTransform::CreateSP(this, &FControlRigEditor::SetRigElementTransform);
 		EditMode->OnContextMenu() = FNewMenuDelegate::CreateSP(this, &FControlRigEditor::HandleOnViewportContextMenuDelegate);
 		EditMode->OnContextMenuCommands() = FNewMenuCommandsDelegate::CreateSP(this, &FControlRigEditor::HandleOnViewportContextMenuCommandsDelegate);
+		EditMode->OnAnimSystemInitialized() = FSimpleDelegate::CreateRaw(this, &FControlRigEditor::OnAnimInitialized);
+		
+		PersonaToolkit->GetPreviewScene()->SetRemoveAttachedComponentFilter(FOnRemoveAttachedComponentFilter::CreateSP(EditMode, &FControlRigEditMode::CanRemoveFromPreviewScene));
+
+		InControlRigBlueprint->OnModified().AddSP(EditMode, &FControlRigEditMode::HandleModifiedEvent);
 	}
 
 	UpdateControlRig();
@@ -242,13 +266,16 @@ void FControlRigEditor::InitControlRigEditor(const EToolkitMode::Type Mode, cons
 
 	if (InControlRigBlueprint)
 	{
-		if (UControlRigModel* Model = InControlRigBlueprint->Model)
+		if (URigVMGraph* Model = InControlRigBlueprint->Model)
 		{
-			if (Model->Nodes().Num() == 0)
+			if (Model->GetNodes().Num() == 0)
 			{
-				if (UControlRigController* Controller = InControlRigBlueprint->ModelController)
+				URigVMNode* Node = InControlRigBlueprint->Controller->AddStructNode(FRigUnit_BeginExecution::StaticStruct(), TEXT("Execute"), FVector2D::ZeroVector, FString(), false);
+				if (Node)
 				{
-					Controller->AddNode(FRigUnit_BeginExecution::StaticStruct()->GetFName());
+					TArray<FName> NodeNames;
+					NodeNames.Add(Node->GetFName());
+					InControlRigBlueprint->Controller->SetNodeSelection(NodeNames, false);
 				}
 			}
 			else
@@ -259,7 +286,9 @@ void FControlRigEditor::InitControlRigEditor(const EToolkitMode::Type Mode, cons
 				{
 					bWasDirty = Package->IsDirty();
 				}
+
 				InControlRigBlueprint->RebuildGraphFromModel();
+
 				if (Package)
 				{
 					if (Package->IsDirty() != bWasDirty)
@@ -268,7 +297,7 @@ void FControlRigEditor::InitControlRigEditor(const EToolkitMode::Type Mode, cons
 					}
 				}
 			}
-			}
+		}
 
 		InControlRigBlueprint->HierarchyContainer.OnElementAdded.AddSP(this, &FControlRigEditor::OnRigElementAdded);
 		InControlRigBlueprint->HierarchyContainer.OnElementRemoved.AddSP(this, &FControlRigEditor::OnRigElementRemoved);
@@ -289,6 +318,7 @@ void FControlRigEditor::InitControlRigEditor(const EToolkitMode::Type Mode, cons
 		}
 	}
 
+	UpdateStaleWatchedPins();
 	bControlRigEditorInitialized = true;
 }
 
@@ -296,11 +326,19 @@ void FControlRigEditor::BindCommands()
 {
 	DECLARE_SCOPE_HIERARCHICAL_COUNTER_FUNC()
 
+	/*
 	GetToolkitCommands()->MapAction(
 		FControlRigBlueprintCommands::Get().ExecuteGraph,
 		FExecuteAction::CreateSP(this, &FControlRigEditor::ToggleExecuteGraph), 
 		FCanExecuteAction(), 
 		FIsActionChecked::CreateSP(this, &FControlRigEditor::IsExecuteGraphOn));
+	*/
+
+	GetToolkitCommands()->MapAction(
+		FControlRigBlueprintCommands::Get().AutoCompileGraph,
+		FExecuteAction::CreateSP(this, &FControlRigEditor::ToggleAutoCompileGraph), 
+		FCanExecuteAction(), 
+		FIsActionChecked::CreateSP(this, &FControlRigEditor::IsAutoCompileGraphOn));
 }
 
 void FControlRigEditor::ToggleExecuteGraph()
@@ -309,12 +347,6 @@ void FControlRigEditor::ToggleExecuteGraph()
 	{
 		bExecutionControlRig = !bExecutionControlRig;
 
-		// enable rig element init pose change, if not executing, we allow rig element change
-		if (FControlRigEditMode* EditMode = GetEditMode())
-		{
-			EditMode->EnableRigElementEditing(!bExecutionControlRig);
-		}
-
 		// This is required now since we update execution/input flag on update controlrig
 		// @fixme: change this to just change flag instead of updating whole control rig
 		// I'll do this before first check-in
@@ -327,6 +359,27 @@ bool FControlRigEditor::IsExecuteGraphOn() const
 	return bExecutionControlRig;
 }
 
+void FControlRigEditor::ToggleAutoCompileGraph()
+{
+	if (UControlRigBlueprint* RigBlueprint = Cast<UControlRigBlueprint>(GetBlueprintObj()))
+	{
+		RigBlueprint->bAutoRecompileVM = !RigBlueprint->bAutoRecompileVM;
+		if (RigBlueprint->bAutoRecompileVM)
+		{
+			RigBlueprint->RequestAutoVMRecompilation();
+		}
+	}
+}
+
+bool FControlRigEditor::IsAutoCompileGraphOn() const
+{
+	if (UControlRigBlueprint* RigBlueprint = Cast<UControlRigBlueprint>(GetBlueprintObj()))
+	{
+		return RigBlueprint->bAutoRecompileVM;
+	}
+	return false;
+}
+
 void FControlRigEditor::ExtendToolbar()
 {
 	DECLARE_SCOPE_HIERARCHICAL_COUNTER_FUNC()
@@ -367,8 +420,12 @@ void FControlRigEditor::FillToolbar(FToolBarBuilder& ToolbarBuilder)
 {
 	ToolbarBuilder.BeginSection("Toolbar");
 	{
+		/*
 		ToolbarBuilder.AddToolBarButton(FControlRigBlueprintCommands::Get().ExecuteGraph,
 			NAME_None, TAttribute<FText>(), TAttribute<FText>(), FSlateIcon(FControlRigEditorStyle::Get().GetStyleSetName(), "ControlRig.ExecuteGraph"));
+		*/
+		ToolbarBuilder.AddToolBarButton(FControlRigBlueprintCommands::Get().AutoCompileGraph,
+			NAME_None, TAttribute<FText>(), TAttribute<FText>(), FSlateIcon(FControlRigEditorStyle::Get().GetStyleSetName(), "ControlRig.AutoCompileGraph"));
 
 		ToolbarBuilder.AddWidget(SNew(SBlueprintEditorSelectedDebugObjectWidget, SharedThis(this)));
 	}
@@ -496,9 +553,18 @@ void FControlRigEditor::HandleSetObjectBeingDebugged(UObject* InObject)
 		}
 	}
 
-	if (GetBlueprintObj())
+	if (UControlRigBlueprint* RigBlueprint = Cast<UControlRigBlueprint>(GetBlueprintObj()))
 	{
-		FBlueprintEditorUtils::UpdateStalePinWatches(GetBlueprintObj());
+		UpdateStaleWatchedPins();
+		if (UControlRigBlueprintGeneratedClass* GeneratedClass = RigBlueprint->GetControlRigBlueprintGeneratedClass())
+		{
+			UControlRig* CDO = Cast<UControlRig>(GeneratedClass->GetDefaultObject(true /* create if needed */));
+			if (CDO->VM->GetInstructions().Num() <= 1 /* only exit */)
+			{
+				RigBlueprint->RecompileVM();
+				RigBlueprint->RequestControlRigInit();
+			}
+		}
 	}
 
 	if (DebuggedControlRig)
@@ -510,28 +576,25 @@ void FControlRigEditor::HandleSetObjectBeingDebugged(UObject* InObject)
 		UControlRigSkeletalMeshComponent* EditorSkelComp = Cast<UControlRigSkeletalMeshComponent>(GetPersonaToolkit()->GetPreviewScene()->GetPreviewMeshComponent());
 		if (EditorSkelComp)
 		{
-			UControlRigSequencerAnimInstance* AnimInstance = Cast<UControlRigSequencerAnimInstance>(EditorSkelComp->GetAnimInstance());
+			UControlRigLayerInstance* AnimInstance = Cast<UControlRigLayerInstance>(EditorSkelComp->GetAnimInstance());
 			if (AnimInstance)
 			{
+				FControlRigIOSettings IOSettings = FControlRigIOSettings::MakeEnabled();
+				IOSettings.bUpdatePose = bShouldExecute;
+				IOSettings.bUpdateCurves = bShouldExecute;
+
 				// we might want to move this into another method
 				FInputBlendPose Filter;
-				AnimInstance->UpdateControlRig(
-					DebuggedControlRig,
-					0,
-					false,
-					false,
-					Filter,
-					1.0f,
-					bShouldExecute, // update input
-					bShouldExecute // execute rig
-				);
+				AnimInstance->ResetControlRigTracks();
+				AnimInstance->AddControlRigTrack(0, DebuggedControlRig);
+				AnimInstance->UpdateControlRigTrack(0, 1.0f, IOSettings, bShouldExecute);
 				AnimInstance->RecalcRequiredBones();
 
 				// since rig has changed, rebuild draw skeleton
 				EditorSkelComp->RebuildDebugDrawSkeleton();
 				if (FControlRigEditMode* EditMode = GetEditMode())
 				{
-					EditMode->SetObjects(DebuggedControlRig, FGuid(), EditorSkelComp);
+					EditMode->SetObjects(DebuggedControlRig, FGuid(), EditorSkelComp,nullptr);
 				}
 			}
 		}
@@ -540,7 +603,7 @@ void FControlRigEditor::HandleSetObjectBeingDebugged(UObject* InObject)
 	{
 		if (FControlRigEditMode* EditMode = GetEditMode())
 		{
-			EditMode->SetObjects(nullptr, FGuid(), nullptr);
+			EditMode->SetObjects(nullptr, FGuid(), nullptr,nullptr);
 		}
 	}
 }
@@ -580,6 +643,8 @@ UBlueprint* FControlRigEditor::GetBlueprintObj() const
 
 void FControlRigEditor::SetDetailObjects(const TArray<UObject*>& InObjects)
 {
+	ClearDetailObject();
+
 	RigElementInDetailPanel = FRigElementKey();
 	SKismetInspector::FShowDetailsOptions Options;
 	Options.bForceRefresh = true;
@@ -588,7 +653,88 @@ void FControlRigEditor::SetDetailObjects(const TArray<UObject*>& InObjects)
 
 void FControlRigEditor::SetDetailObject(UObject* Obj)
 {
-	RigElementInDetailPanel = FRigElementKey();
+	if (URigVMStructNode* StructNode = Cast<URigVMStructNode>(Obj))
+	{
+		ClearDetailObject();
+		NodeDetailStruct = StructNode->GetScriptStruct();
+		NodeDetailName = StructNode->GetFName();
+
+		if (NodeDetailStruct)
+		{
+			NodeDetailBuffer.AddUninitialized(NodeDetailStruct->GetStructureSize());
+			NodeDetailStruct->InitializeDefaultValue(NodeDetailBuffer.GetData());
+
+			FString StructDefaultValue = StructNode->GetStructDefaultValue();
+			NodeDetailStruct->ImportText(*StructDefaultValue, NodeDetailBuffer.GetData(), nullptr, PPF_None, nullptr, NodeDetailStruct->GetName());
+
+			TSharedPtr<FStructOnScope> StructToDisplay = MakeShareable(new FStructOnScope(NodeDetailStruct, (uint8*)NodeDetailBuffer.GetData()));
+			if (StructToDisplay.IsValid())
+			{
+				StructToDisplay->SetPackage(GetControlRigBlueprint()->GetOutermost());
+			}
+
+			// mark all input properties with editanywhere
+			for (TFieldIterator<FProperty> PropertyIt(NodeDetailStruct); PropertyIt; ++PropertyIt)
+			{
+				FProperty* Property = *PropertyIt;
+				if (!Property->HasMetaData(TEXT("Input")))
+				{
+					continue;
+				}
+
+				// filter out execute pins
+				if (FStructProperty* StructProperty = CastField<FStructProperty>(Property))
+				{
+					if (StructProperty->Struct->IsChildOf(FRigVMExecuteContext::StaticStruct()))
+					{
+						continue;
+					}
+				}
+
+				bool bEditable = true;
+				if (URigVMPin* Pin = StructNode->FindPin(Property->GetName()))
+				{
+					Pin = Pin->GetPinForLink();
+					if (Pin->GetDirection() == ERigVMPinDirection::Output)
+					{
+						if (Pin->GetTargetLinks().Num() > 0)
+						{
+							bEditable = false;
+						}
+					}
+					else if (Pin->GetDirection() == ERigVMPinDirection::Input ||
+						Pin->GetDirection() == ERigVMPinDirection::IO)
+					{
+						if (Pin->GetSourceLinks().Num() > 0)
+						{
+							bEditable = false;
+						}
+					}
+				}
+
+				Property->SetPropertyFlags(Property->GetPropertyFlags() | CPF_Edit);
+
+				if (bEditable)
+				{
+					Property->ClearPropertyFlags(CPF_EditConst);
+				}
+				else
+				{
+					Property->SetPropertyFlags(Property->GetPropertyFlags() | CPF_EditConst);
+				}
+			}
+
+			FPropertyEditorModule& PropertyEditorModule = FModuleManager::LoadModuleChecked<FPropertyEditorModule>("PropertyEditor");
+			if (!PropertyEditorModule.GetClassNameToDetailLayoutNameMap().Contains(NodeDetailStruct->GetFName()))
+			{
+				PropertyEditorModule.RegisterCustomClassLayout(NodeDetailStruct->GetFName(), FOnGetDetailCustomizationInstance::CreateStatic(&FRigUnitDetails::MakeInstance));
+			}
+
+			Inspector->ShowSingleStruct(StructToDisplay);
+		}
+		return;
+	}
+
 	TArray<UObject*> Objects;
 	if (Obj)
 	{
@@ -603,6 +749,9 @@ void FControlRigEditor::SetDetailStruct(const FRigElementKey& InElement, TShared
 	{
 		return;
 	}
+
+	ClearDetailObject();
+
 	RigElementInDetailPanel = InElement;
 	if (StructToDisplay.IsValid())
 	{
@@ -614,6 +763,15 @@ void FControlRigEditor::SetDetailStruct(const FRigElementKey& InElement, TShared
 void FControlRigEditor::ClearDetailObject()
 {
 	RigElementInDetailPanel = FRigElementKey();
+
+	if (NodeDetailBuffer.Num() > 0 && NodeDetailStruct != nullptr)
+	{
+		NodeDetailStruct->DestroyStruct(NodeDetailBuffer.GetData(), 1);
+		NodeDetailBuffer.Reset();
+		NodeDetailStruct = nullptr;
+	}
+	NodeDetailName = NAME_None;
+
 	Inspector->ShowDetailsForObjects(TArray<UObject*>());
 	Inspector->ShowSingleStruct(TSharedPtr<FStructOnScope>());
 }
@@ -646,12 +804,14 @@ void FControlRigEditor::Compile()
 
 		FString LastDebuggedObjectName = GetCustomDebugObjectLabel(GetBlueprintObj()->GetObjectBeingDebugged());
 		GetBlueprintObj()->SetObjectBeingDebugged(nullptr);
+
+		FRigElementKey SelectedKey = RigElementInDetailPanel;
 		ClearDetailObject();
 
 		if (ControlRig)
 		{
-			ControlRig->OnInitialized().Clear();
-			ControlRig->OnExecuted().Clear();
+			ControlRig->OnInitialized_AnyThread().Clear();
+			ControlRig->OnExecuted_AnyThread().Clear();
 		}
 
 		{
@@ -665,7 +825,10 @@ void FControlRigEditor::Compile()
 			UControlRigBlueprintGeneratedClass* GeneratedClass = Cast<UControlRigBlueprintGeneratedClass>(ControlRig->GetClass());
 			if (GeneratedClass)
 			{
-				if (GeneratedClass->Operators.Num() == 1) // just the "done" operator
+				UControlRig* CDO = Cast<UControlRig>(GeneratedClass->GetDefaultObject(true /* create if needed */));
+				FRigVMInstructionArray Instructions = CDO->VM->GetInstructions();
+
+				if (Instructions.Num() <= 1) // just the "done" operator
 				{
 					FNotificationInfo Info(LOCTEXT("ControlRigBlueprintCompilerEmptyRigMessage", "The Control Rig you compiled doesn't do anything. Did you forget to add a Begin_Execution node?"));
 					Info.bFireAndForget = true;
@@ -687,6 +850,47 @@ void FControlRigEditor::Compile()
 				GetBlueprintObj()->SetObjectBeingDebugged(DebugObject.Object);
 			}
 		}
+
+		if(SelectedKey.IsValid())
+		{
+			UControlRigBlueprint* RigBlueprint = Cast<UControlRigBlueprint>(GetBlueprintObj());
+			switch (SelectedKey.Type)
+			{
+				case ERigElementType::Bone:
+				{
+					FRigBoneHierarchy& BoneHierarchy = RigBlueprint->HierarchyContainer.BoneHierarchy;
+					SetDetailStruct(SelectedKey, MakeShareable(new FStructOnScope(FRigBone::StaticStruct(), (uint8*)&(BoneHierarchy[SelectedKey.Name]))));
+					break;
+				}
+				case ERigElementType::Control:
+				{
+					FRigControlHierarchy& ControlHierarchy = RigBlueprint->HierarchyContainer.ControlHierarchy;
+					SetDetailStruct(SelectedKey, MakeShareable(new FStructOnScope(FRigControl::StaticStruct(), (uint8*)&(ControlHierarchy[SelectedKey.Name]))));
+					break;
+				}
+				case ERigElementType::Space:
+				{
+					FRigSpaceHierarchy& SpaceHierarchy = RigBlueprint->HierarchyContainer.SpaceHierarchy;
+					SetDetailStruct(SelectedKey, MakeShareable(new FStructOnScope(FRigSpace::StaticStruct(), (uint8*)&(SpaceHierarchy[SelectedKey.Name]))));
+					break;
+				}
+				case ERigElementType::Curve:
+				{
+					FRigCurveContainer& CurveContainer = RigBlueprint->HierarchyContainer.CurveContainer;
+					SetDetailStruct(SelectedKey, MakeShareable(new FStructOnScope(FRigBone::StaticStruct(), (uint8*)&(CurveContainer[SelectedKey.Name]))));
+					break;
+				}
+				default:
+				{
+					break;
+				}
+			}
+		}
+
+		if (PreviewInstance)
+		{
+			PreviewInstance->ResetModifiedBone();
+		}
 	}
 
 	// enable this for creating a new unit test
@@ -722,6 +926,21 @@ FLinearColor FControlRigEditor::GetWorldCentricTabColorScale() const
 	return FLinearColor( 0.5f, 0.25f, 0.35f, 0.5f );
 }
 
+bool FControlRigEditor::TransactionObjectAffectsBlueprint(UObject* InTransactedObject)
+{
+	UControlRigBlueprint* RigBlueprint = Cast<UControlRigBlueprint>(GetBlueprintObj());
+	if (RigBlueprint == nullptr)
+	{
+		return false;
+	}
+
+	if (InTransactedObject->GetOuter() == RigBlueprint->Controller)
+	{
+		return false;
+	}
+	return FBlueprintEditor::TransactionObjectAffectsBlueprint(InTransactedObject);
+}
+
 void FControlRigEditor::DeleteSelectedNodes()
 {
 	DECLARE_SCOPE_HIERARCHICAL_COUNTER_FUNC()
@@ -735,6 +954,9 @@ void FControlRigEditor::DeleteSelectedNodes()
 	const FGraphPanelSelectionSet SelectedNodes = GetSelectedNodes();
 	SetUISelectionState(NAME_None);
 
+	bool DeletedAnything = false;
+	RigBlueprint->Controller->OpenUndoBracket(TEXT("Delete selected nodes"));
+
 	for (FGraphPanelSelectionSet::TConstIterator NodeIt(SelectedNodes); NodeIt; ++NodeIt)
 	{
 		if (UEdGraphNode* Node = Cast<UEdGraphNode>(*NodeIt))
@@ -744,179 +966,184 @@ void FControlRigEditor::DeleteSelectedNodes()
 				AnalyticsTrackNodeEvent(GetBlueprintObj(), Node, true);
 				if (UControlRigGraphNode* RigNode = Cast<UControlRigGraphNode>(Node))
 				{
-					RigBlueprint->ModelController->RemoveNode(RigNode->PropertyName);
+					if(RigBlueprint->Controller->RemoveNodeByName(*RigNode->ModelNodePath))
+					{
+						DeletedAnything = true;
+					}
 				}
 				else if (UEdGraphNode_Comment* CommentNode = Cast<UEdGraphNode_Comment>(Node))
 				{
-					RigBlueprint->ModelController->RemoveNode(CommentNode->GetFName());
+					if(RigBlueprint->Controller->RemoveNodeByName(CommentNode->GetFName()))
+					{
+						DeletedAnything = true;
+					}
 				}
 				else
 				{
-					ensure(false);
+					Node->GetGraph()->RemoveNode(Node);
 				}
 			}
 		}
 	}
+
+	if(DeletedAnything)
+	{
+		RigBlueprint->Controller->CloseUndoBracket();
+	}
+	else
+	{
+		RigBlueprint->Controller->CancelUndoBracket();
+	}
 }
 
-void FControlRigEditor::PasteNodesHere(class UEdGraph* DestinationGraph, const FVector2D& GraphLocation)
+bool FControlRigEditor::CanDeleteNodes() const
 {
-	DECLARE_SCOPE_HIERARCHICAL_COUNTER_FUNC()
+	return true;
+}
 
+void FControlRigEditor::CopySelectedNodes()
+{
 	UControlRigBlueprint* RigBlueprint = Cast<UControlRigBlueprint>(GetBlueprintObj());
 	if (RigBlueprint == nullptr)
 	{
 		return;
 	}
 
-	UControlRigGraph* TempGraph = NewObject<UControlRigGraph>(GetTransientPackage());
-	TempGraph->bIsTemporaryGraphForCopyPaste = true;
-	TempGraph->Schema = DestinationGraph->Schema;
+	FString ExportedText = RigBlueprint->Controller->ExportSelectedNodesToText();
+	FPlatformApplicationMisc::ClipboardCopy(*ExportedText);
+}
+
+bool FControlRigEditor::CanCopyNodes() const
+{
+	UControlRigBlueprint* RigBlueprint = Cast<UControlRigBlueprint>(GetBlueprintObj());
+	if (RigBlueprint == nullptr)
+	{
+		return false;
+	}
+	return RigBlueprint->Model->GetSelectNodes().Num() > 0;
+}
+
+bool FControlRigEditor::CanPasteNodes() const
+{
+	UControlRigBlueprint* RigBlueprint = Cast<UControlRigBlueprint>(GetBlueprintObj());
+	if (RigBlueprint == nullptr)
+	{
+		return false;
+	}
 
-	// Grab the text to paste from the clipboard.
 	FString TextToImport;
 	FPlatformApplicationMisc::ClipboardPaste(TextToImport);
+	return RigBlueprint->Controller->CanImportNodesFromText(TextToImport);
+}
+
+void FControlRigEditor::PasteNodes()
+{
+	DECLARE_SCOPE_HIERARCHICAL_COUNTER_FUNC()
 
-	// Import the nodes
-	TSet<UEdGraphNode*> PastedNodes;
-	FEdGraphUtilities::ImportNodesFromText(TempGraph, TextToImport, PastedNodes);
-	if (PastedNodes.Num() == 0)
+	UControlRigBlueprint* RigBlueprint = Cast<UControlRigBlueprint>(GetBlueprintObj());
+	if (RigBlueprint == nullptr)
 	{
 		return;
 	}
 
-	FVector2D AvgNodePosition(0.0f, 0.0f);
-	for (TSet<UEdGraphNode*>::TIterator It(PastedNodes); It; ++It)
+	RigBlueprint->Controller->OpenUndoBracket(TEXT("Pasted Nodes."));
+
+	FString TextToImport;
+	FPlatformApplicationMisc::ClipboardPaste(TextToImport);
+	TArray<FName> NodeNames = RigBlueprint->Controller->ImportNodesFromText(TextToImport);
+
+	if (NodeNames.Num() > 0)
+	{
+		for (const FName& NodeName : NodeNames)
+		{
+			FVector2D Position = RigBlueprint->Model->FindNodeByName(NodeName)->GetPosition();
+			RigBlueprint->Controller->SetNodePositionByName(NodeName, Position + FVector2D(50.f, 50.f));
+		}
+		RigBlueprint->Controller->SetNodeSelection(NodeNames);
+		RigBlueprint->Controller->CloseUndoBracket();
+	}
+	else
 	{
-		UEdGraphNode* Node = *It;
-		AvgNodePosition.X += Node->NodePosX;
-		AvgNodePosition.Y += Node->NodePosY;
+		RigBlueprint->Controller->CancelUndoBracket();
 	}
+}
 
-	float InvNumNodes = 1.0f / float(PastedNodes.Num());
-	AvgNodePosition.X *= InvNumNodes;
-	AvgNodePosition.Y *= InvNumNodes;
+void FControlRigEditor::PostUndo(bool bSuccess)
+{
+	IControlRigEditor::PostUndo(bSuccess);
+	EnsureValidRigElementInDetailPanel();
+}
 
-	RigBlueprint->ModelController->OpenUndoBracket(TEXT("Pasted Nodes."));
+void FControlRigEditor::PostRedo(bool bSuccess)
+{
+	IControlRigEditor::PostRedo(bSuccess);
+	EnsureValidRigElementInDetailPanel();
+}
 
-	TMap<FString, FString> NameMap;
-	for (TSet<UEdGraphNode*>::TIterator It(PastedNodes); It; ++It)
+void FControlRigEditor::EnsureValidRigElementInDetailPanel()
+{
+	if (RigElementInDetailPanel.IsValid())
 	{
-		UEdGraphNode* Node = *It;
-		Node->NodePosX = (Node->NodePosX - AvgNodePosition.X) + GraphLocation.X;
-		Node->NodePosY = (Node->NodePosY - AvgNodePosition.Y) + GraphLocation.Y;
-		Node->SnapToGrid(SNodePanel::GetSnapGridSize());
-
-		if (UControlRigGraphNode* RigNode = Cast<UControlRigGraphNode>(Node))
+		if (UControlRigBlueprint * RigBlueprint = Cast<UControlRigBlueprint>(GetBlueprintObj()))
 		{
-			FVector2D NodePosition = FVector2D(Node->NodePosX, Node->NodePosY);
-			UScriptStruct* ScriptStruct = RigNode->GetUnitScriptStruct();
-			bool bAddedNode = false;
-			if (ScriptStruct)
+			if (RigBlueprint->HierarchyContainer.GetIndex(RigElementInDetailPanel) == INDEX_NONE)
 			{
-				bAddedNode = RigBlueprint->ModelController->AddNode(ScriptStruct->GetFName(), NodePosition, RigNode->GetPropertyName());
+				ClearDetailObject();
 			}
-			else
-			{
-				FName DataType = RigNode->PinType.PinCategory;
-				if (UStruct* Struct = Cast<UStruct>(RigNode->PinType.PinSubCategoryObject))
-				{
-					DataType = Struct->GetFName();
-				}
+		}
+	}
+}
 
-				bAddedNode = RigBlueprint->ModelController->AddParameter(RigNode->GetPropertyName(), DataType, EControlRigModelParameterType::Hidden, NodePosition);
-			}
+void FControlRigEditor::OnStartWatchingPin()
+{
+	if (UControlRigBlueprint* ControlRigBlueprint = CastChecked<UControlRigBlueprint>(GetBlueprintObj()))
+	{
+		if (UEdGraphPin* Pin = GetCurrentlySelectedPin())
+		{
+			ControlRigBlueprint->Controller->SetPinIsWatched(Pin->GetName(), true);
+		}
+	}
+}
 
-			if(bAddedNode && RigBlueprint->LastNameFromNotification != NAME_None)
+bool FControlRigEditor::CanStartWatchingPin() const
+{
+	if (UControlRigBlueprint* ControlRigBlueprint = CastChecked<UControlRigBlueprint>(GetBlueprintObj()))
+	{
+		if (UEdGraphPin* Pin = GetCurrentlySelectedPin())
+		{
+			if (URigVMPin* ModelPin = ControlRigBlueprint->Model->FindPin(Pin->GetName()))
 			{
-				FName AddedNodeName = RigBlueprint->LastNameFromNotification;
-				NameMap.Add(Node->GetName(), AddedNodeName.ToString());
-
-				for (UEdGraphPin* Pin : Node->Pins)
-				{
-					FString Left, Right;
-					RigBlueprint->Model->SplitPinPath(Pin->GetName(), Left, Right);
-					Left = *AddedNodeName.ToString();
-
-					if (Pin->PinType.ContainerType == EPinContainerType::Array)
-					{
-						RigBlueprint->ModelController->SetArrayPinSize(*Left, *Right, Pin->SubPins.Num());
-					}
-
-					if (RigNode->ExpandedPins.Contains(Pin->GetName()))
-					{
-						RigBlueprint->ModelController->ExpandPin(*Left, *Right, Pin->Direction == EGPD_Input, true);
-					}
-
-					if (Pin->Direction == EGPD_Input)
-					{
-						FString DefaultValue = Pin->DefaultValue;
-						if (DefaultValue.IsEmpty() && Pin->DefaultObject != nullptr)
-						{
-							DefaultValue = Pin->DefaultObject->GetPathName();
-						}
-						if (!DefaultValue.IsEmpty())
-						{
-							RigBlueprint->ModelController->SetPinDefaultValue(*Left, *Right, DefaultValue, false);
-						}
-					}
-				}
+				return ModelPin->GetParentPin() == nullptr && !ModelPin->RequiresWatch();
 			}
 		}
-		else if (const UEdGraphNode_Comment* CommentNode = Cast<UEdGraphNode_Comment>(Node))
+	}
+	return false;
+}
+
+void FControlRigEditor::OnStopWatchingPin()
+{
+	if (UControlRigBlueprint* ControlRigBlueprint = CastChecked<UControlRigBlueprint>(GetBlueprintObj()))
+	{
+		if (UEdGraphPin* Pin = GetCurrentlySelectedPin())
 		{
-			FVector2D NodePosition = FVector2D((float)CommentNode->NodePosX, (float)CommentNode->NodePosY);
-			FVector2D NodeSize = FVector2D((float)CommentNode->NodeWidth, (float)CommentNode->NodeHeight);
-			RigBlueprint->ModelController->AddComment(CommentNode->GetFName(), CommentNode->NodeComment, NodePosition, NodeSize, CommentNode->CommentColor, false);
+			ControlRigBlueprint->Controller->SetPinIsWatched(Pin->GetName(), false);
 		}
 	}
+}
 
-	for (TSet<UEdGraphNode*>::TIterator It(PastedNodes); It; ++It)
+bool FControlRigEditor::CanStopWatchingPin() const
+{
+	if (UControlRigBlueprint* ControlRigBlueprint = CastChecked<UControlRigBlueprint>(GetBlueprintObj()))
 	{
-		UEdGraphNode* Node = *It;
-		UControlRigGraphNode* RigNode = Cast<UControlRigGraphNode>(Node);
-		if (RigNode != nullptr)
+		if (UEdGraphPin* Pin = GetCurrentlySelectedPin())
 		{
-			for (UEdGraphPin* Pin : RigNode->Pins)
+			if (URigVMPin* ModelPin = ControlRigBlueprint->Model->FindPin(Pin->GetName()))
 			{
-				if (Pin->Direction == EGPD_Input)
-				{
-					continue;
-				}
-
-				for (UEdGraphPin* OtherPin : Pin->LinkedTo)
-				{
-					UControlRigGraphNode* OtherRigNode = Cast<UControlRigGraphNode>(OtherPin->GetOwningNode());
-					if (OtherRigNode)
-					{
-						FString* RemappedNodeName = NameMap.Find(RigNode->GetName());
-						FString* OtherRemappedNodeName = NameMap.Find(OtherRigNode->GetName());
-
-						if (RemappedNodeName != nullptr && OtherRemappedNodeName != nullptr)
-						{
-							FString PinPath = Pin->GetName();
-							if (PinPath.StartsWith(RigNode->GetPropertyName().ToString()))
-							{
-								PinPath = *RemappedNodeName + PinPath.RightChop(RigNode->GetPropertyName().ToString().Len());
-							}
-							FString OtherPinPath = OtherPin->GetName();
-							if (OtherPinPath.StartsWith(OtherRigNode->GetPropertyName().ToString()))
-							{
-								OtherPinPath = *OtherRemappedNodeName + OtherPinPath.RightChop(OtherRigNode->GetPropertyName().ToString().Len());
-							}
-
-							FString SourceLeft, SourceRight, TargetLeft, TargetRight;
-							RigBlueprint->Model->SplitPinPath(PinPath, SourceLeft, SourceRight);
-							RigBlueprint->Model->SplitPinPath(OtherPinPath, TargetLeft, TargetRight);
-							RigBlueprint->ModelController->MakeLink(*SourceLeft, *SourceRight, *TargetLeft, *TargetRight);
-						}
-					}
-				}
+				return ModelPin->RequiresWatch();
 			}
 		}
 	}
-
-	RigBlueprint->ModelController->CloseUndoBracket();
+	return false;
 }
 
 void FControlRigEditor::OnToolkitHostingStarted(const TSharedRef<class IToolkit>& Toolkit)
@@ -946,19 +1173,18 @@ void FControlRigEditor::OnActiveTabChanged( TSharedPtr<SDockTab> PreviouslyActiv
 	}
 }
 
-void FControlRigEditor::PostUndo(bool bSuccess)
+void FControlRigEditor::OnAnimInitialized()
 {
-	//DocumentManager->CleanInvalidTabs();
-	//DocumentManager->RefreshAllTabs();
-
-	FBlueprintEditor::PostUndo(bSuccess);
-}
-
-void FControlRigEditor::PostRedo(bool bSuccess)
-{
-	//DocumentManager->RefreshAllTabs();
-
-	FBlueprintEditor::PostRedo(bSuccess);
+	UControlRigSkeletalMeshComponent* EditorSkelComp = Cast<UControlRigSkeletalMeshComponent>(GetPersonaToolkit()->GetPreviewScene()->GetPreviewMeshComponent());
+	UControlRigLayerInstance* AnimInstance = Cast<UControlRigLayerInstance>(EditorSkelComp->GetAnimInstance());
+	if (AnimInstance && ControlRig)
+	{
+		// update control rig data to anim instance since animation system has been reinitialized
+		FInputBlendPose Filter;
+		AnimInstance->ResetControlRigTracks();
+		AnimInstance->AddControlRigTrack(0, ControlRig);
+		AnimInstance->UpdateControlRigTrack(0, 1.0f, FControlRigIOSettings::MakeEnabled(), bExecutionControlRig);
+	}
 }
 
 void FControlRigEditor::UndoAction()
@@ -1004,128 +1230,88 @@ FGraphAppearanceInfo FControlRigEditor::GetGraphAppearance(UEdGraph* InGraph) co
 	return AppearanceInfo;
 }
 
-void FControlRigEditor::NotifyPostChange(const FPropertyChangedEvent& PropertyChangedEvent, FProperty* PropertyThatChanged)
-{
-	FBlueprintEditor::NotifyPostChange(PropertyChangedEvent, PropertyThatChanged);
-}
-
-void FControlRigEditor::HandleModelModified(const UControlRigModel* InModel, EControlRigModelNotifType InType, const void* InPayload)
+void FControlRigEditor::HandleModifiedEvent(ERigVMGraphNotifType InNotifType, URigVMGraph* InGraph, UObject* InSubject)
 {
 	DECLARE_SCOPE_HIERARCHICAL_COUNTER_FUNC()
 
-	UControlRigBlueprint* RigBlueprint = Cast<UControlRigBlueprint>(GetBlueprintObj());
-	if (RigBlueprint == nullptr)
+	UControlRigBlueprint* ControlRigBlueprint = CastChecked<UControlRigBlueprint>(GetBlueprintObj());
+	if (ControlRigBlueprint == nullptr)
 	{
 		return;
 	}
 
-	switch (InType)
+	switch (InNotifType)
 	{
-		case EControlRigModelNotifType::NodeSelected:
-		case EControlRigModelNotifType::NodeDeselected:
+		case ERigVMGraphNotifType::NodeSelected:
+		case ERigVMGraphNotifType::NodeDeselected:
 		{
-			const FControlRigModelNode* Node = (const FControlRigModelNode*)InPayload;
-			if (Node != nullptr)
+			URigVMNode* Node = Cast<URigVMNode>(InSubject);
+
+			if (FocusedGraphEdPtr.IsValid())
 			{
-					if (FocusedGraphEdPtr.IsValid())
+				TSharedPtr<SGraphEditor> FocusedGraphEd = FocusedGraphEdPtr.Pin();
+				if (UControlRigGraph* RigGraph = Cast<UControlRigGraph>(FocusedGraphEd->GetCurrentGraph()))
+				{
+					if (InNotifType == ERigVMGraphNotifType::NodeSelected)
 					{
-						TSharedPtr<SGraphEditor> FocusedGraphEd = FocusedGraphEdPtr.Pin();
-						if (UControlRigGraph* RigGraph = Cast<UControlRigGraph>(FocusedGraphEd->GetCurrentGraph()))
-						{
-						TArray<FControlRigModelNode> Selection;
-						bool bSelect = InType == EControlRigModelNotifType::NodeSelected;
-						if (bSelect)
-						{
-							Selection = InModel->SelectedNodes();
-						}
-						else
-						{
-							Selection.Add(*Node);
-						}
-
-						TArray<FRigElementKey> RigElementsToSelect;
-						for (const FControlRigModelNode& SelectedNode : Selection)
-						{
-							if (UControlRigGraphNode* RigNode = Cast<UControlRigGraphNode>(RigGraph->FindNodeFromPropertyName(SelectedNode.Name)))
-							{
-								for (const FControlRigModelPin& Pin : Node->Pins)
-								{
-									if (Pin.Type.PinCategory == UEdGraphSchema_K2::PC_Name)
-									{
-										if (Pin.CustomWidgetName == TEXT("BoneName"))
-										{
-											RigElementsToSelect.Add(FRigElementKey(*Pin.DefaultValue, ERigElementType::Bone));
-										}
-										else if (Pin.CustomWidgetName == TEXT("ControlName"))
-										{
-											RigElementsToSelect.Add(FRigElementKey(*Pin.DefaultValue, ERigElementType::Control));
-										}
-										else if (Pin.CustomWidgetName == TEXT("SpaceName"))
-										{
-											RigElementsToSelect.Add(FRigElementKey(*Pin.DefaultValue, ERigElementType::Space));
-										}
-										else if (Pin.CustomWidgetName == TEXT("CurveName"))
-										{
-											RigElementsToSelect.Add(FRigElementKey(*Pin.DefaultValue, ERigElementType::Curve));
-										}
-									}
-								}
-							}
-						}
-
-						if (RigElementsToSelect.Num() > 0)
-						{
-							if (bSelect)
-							{
-								RigBlueprint->HierarchyContainer.ClearSelection();
-							}
-							for (const FRigElementKey& RigElementToSelect : RigElementsToSelect)
-							{
-								RigBlueprint->HierarchyContainer.Select(RigElementToSelect, bSelect);
-							}
-						}
+						SetDetailObject(Node);
+					}
 
-						if (!bIsSelecting)
+					if (!RigGraph->bIsSelecting)
+					{
+						TGuardValue<bool> SelectingGuard(RigGraph->bIsSelecting, true);
+						if (URigVMNode* ModelNode = Cast<URigVMNode>(InSubject))
 						{
-							TGuardValue<bool> SelectingGuard(bIsSelecting, true);
-							if (UEdGraphNode* EdNode = RigGraph->FindNodeFromPropertyName(Node->Name))
+							if (UEdGraphNode* EdNode = RigGraph->FindNodeForModelNodeName(ModelNode->GetFName()))
 							{
-								FocusedGraphEd->SetNodeSelection(EdNode, InType == EControlRigModelNotifType::NodeSelected);
+								FocusedGraphEd->SetNodeSelection(EdNode, InNotifType == ERigVMGraphNotifType::NodeSelected);
 							}
-							break;
 						}
 					}
 				}
+			}
+			break;
+		}
+		case ERigVMGraphNotifType::PinDefaultValueChanged:
+		{
+			URigVMPin* Pin = Cast<URigVMPin>(InSubject);
 
-				if (InType == EControlRigModelNotifType::NodeSelected)
+			if (ControlRigBlueprint && NodeDetailBuffer.Num() > 0 && NodeDetailStruct != nullptr && !NodeDetailName.IsNone())
+			{
+				if (Pin->GetNode()->GetFName() == NodeDetailName)
 				{
-					UClass* Class = RigBlueprint->GeneratedClass.Get();
-					if (Class)
+					URigVMPin* RootPin = Pin->GetRootPin();
+					if (FProperty* Property = NodeDetailStruct->FindPropertyByName(RootPin->GetFName()))
 					{
-						if (FProperty* Property = Class->FindPropertyByName(Node->Name))
-						{
-							TSet<class UObject*> SelectedObjects;
-							SelectedObjects.Add(Property->GetUPropertyWrapper());
-							FBlueprintEditor::OnSelectedNodesChangedImpl(SelectedObjects);
-						}
-						else
+						FString DefaultValue = RootPin->GetDefaultValue();
+						if (!DefaultValue.IsEmpty())
 						{
-							TSharedPtr<SGraphEditor> FocusedGraphEd = FocusedGraphEdPtr.Pin();
-							if (UControlRigGraph* RigGraph = Cast<UControlRigGraph>(FocusedGraphEd->GetCurrentGraph()))
-							{
-								if(UEdGraphNode* EdNode = RigGraph->FindNodeFromPropertyName(Node->Name))
-								{
-									TSet<class UObject*> SelectedObjects;
-									SelectedObjects.Add(EdNode);
-									FBlueprintEditor::OnSelectedNodesChangedImpl(SelectedObjects);
-								}
-							}
+							uint8* PropertyValuePtr = (uint8*)NodeDetailBuffer.GetData();
+							PropertyValuePtr += Property->GetOffset_ReplaceWith_ContainerPtrToValuePtr();
+							Property->ImportText(*DefaultValue, PropertyValuePtr, PPF_None, nullptr);
 						}
 					}
 				}
 			}
+
+			break;
+		}
+		case ERigVMGraphNotifType::PinArraySizeChanged:
+		{
+			URigVMPin* Pin = Cast<URigVMPin>(InSubject);
+
+			if (ControlRigBlueprint && NodeDetailBuffer.Num() > 0 && NodeDetailStruct != nullptr && !NodeDetailName.IsNone())
+			{
+				if (Pin->GetNode()->GetFName() == NodeDetailName)
+				{
+					// refresh the details panel
+					SetDetailObject(Pin->GetNode());
+				}
+			}
+
 			break;
 		}
+		case ERigVMGraphNotifType::NodeSelectionChanged:
 		default:
 		{
 			break;
@@ -1133,9 +1319,30 @@ void FControlRigEditor::HandleModelModified(const UControlRigModel* InModel, ECo
 	}
 }
 
+void FControlRigEditor::HandleVMCompiledEvent(UBlueprint* InBlueprint, URigVM* InVM)
+{
+}
+
+void FControlRigEditor::HandleControlRigInitializedEvent(UControlRig* InControlRig, const EControlRigState InState)
+{
+	UpdateGraphCompilerErrors();
+}
+
 void FControlRigEditor::Tick(float DeltaTime)
 {
 	FBlueprintEditor::Tick(DeltaTime);
+
+	// tick the control rig in case we don't have skeletal mesh
+	if (UControlRigBlueprint* Blueprint = GetControlRigBlueprint())
+	{
+		if (Blueprint->GetPreviewMesh() == nullptr && 
+			ControlRig != nullptr && 
+			bExecutionControlRig)
+		{
+			ControlRig->SetDeltaTime(DeltaTime);
+			ControlRig->Execute(EControlRigState::Update);
+		}
+	}
 }
 
 bool FControlRigEditor::IsEditable(UEdGraph* InGraph) const
@@ -1162,32 +1369,46 @@ void FControlRigEditor::OnSelectedNodesChangedImpl(const TSet<class UObject*>& N
 {
 	DECLARE_SCOPE_HIERARCHICAL_COUNTER_FUNC()
 
-	if (bIsSelecting)
+	UControlRigGraph* RigGraph = Cast<UControlRigGraph>(GetFocusedGraph());
+	if (RigGraph == nullptr)
 	{
 		return;
 	}
 
-	TGuardValue<bool> SelectingGuard(bIsSelecting, true);
+	if (RigGraph->bIsSelecting || GIsTransacting)
+	{
+		return;
+	}
+	TGuardValue<bool> SelectGuard(RigGraph->bIsSelecting, true);
 
 	UControlRigBlueprint* ControlRigBlueprint = CastChecked<UControlRigBlueprint>(GetBlueprintObj());
 	if (ControlRigBlueprint)
 	{
-		if (ControlRigBlueprint->ModelController)
+		TArray<FName> NodeNamesToSelect;
+		for (UObject* Object : NewSelection)
 		{
-			TArray<FName> NodeNamesToSelect;
-			for (UObject* Object : NewSelection)
+			if (UControlRigGraphNode* ControlRigGraphNode = Cast<UControlRigGraphNode>(Object))
 			{
-				if (UControlRigGraphNode* ControlRigGraphNode = Cast<UControlRigGraphNode>(Object))
-				{
-					NodeNamesToSelect.Add(ControlRigGraphNode->GetPropertyName());
-				}
-				else if(UEdGraphNode* Node = Cast<UEdGraphNode>(Object))
+				NodeNamesToSelect.Add(ControlRigGraphNode->GetModelNodeName());
+			}
+			else if(UEdGraphNode* Node = Cast<UEdGraphNode>(Object))
+			{
+				if (UEdGraphNode_Comment* CommentNode = Cast<UEdGraphNode_Comment>(Node))
 				{
-					NodeNamesToSelect.Add(Node->GetFName());
+					URigVMNode* ModelNode = ControlRigBlueprint->Model->FindNodeByName(Node->GetFName());
+					if (ModelNode == nullptr)
+					{
+						TGuardValue<bool> BlueprintNotifGuard(ControlRigBlueprint->bSuspendModelNotificationsForOthers, true);
+						FVector2D NodePos(CommentNode->NodePosX, CommentNode->NodePosY);
+						FVector2D NodeSize(CommentNode->NodeWidth, CommentNode->NodeHeight);
+						FLinearColor NodeColor = CommentNode->CommentColor;
+						ControlRigBlueprint->Controller->AddCommentNode(CommentNode->NodeComment, NodePos, NodeSize, NodeColor, CommentNode->GetName(), true);
+					}
 				}
+				NodeNamesToSelect.Add(Node->GetFName());
 			}
-			ControlRigBlueprint->ModelController->SetSelection(NodeNamesToSelect);
 		}
+		ControlRigBlueprint->Controller->SetNodeSelection(NodeNamesToSelect, true);
 	}
 }
 
@@ -1236,34 +1457,6 @@ void FControlRigEditor::OnBlueprintChangedImpl(UBlueprint* InBlueprint, bool bIs
 		{
 			UpdateControlRig();
 		}
-
-		OnSelectedNodesChangedImpl(GetSelectedNodes());
-
-		UClass* Class = GetBlueprintObj()->GeneratedClass.Get();
-		if (Class)
-		{
-			TSet<class UObject*> SelectedObjects;
-			FGraphPanelSelectionSet SelectedNodes = GetSelectedNodes();
-			for (UObject* SelectedNode : SelectedNodes)
-			{
-				if (UControlRigGraphNode* RigNode = Cast<UControlRigGraphNode>(SelectedNode))
-				{
-					FProperty* Property = Class->FindPropertyByName(RigNode->GetPropertyName());
-					if (Property)
-					{
-						SelectedObjects.Add(Property->GetUPropertyWrapper());
-					}
-				}
-				else
-				{
-					SelectedObjects.Add(SelectedNode);
-				}
-			}
-			if (SelectedObjects.Num() > 0)
-			{
-				FBlueprintEditor::OnSelectedNodesChangedImpl(SelectedObjects);
-			}
-		}
 	}
 }
 
@@ -1437,6 +1630,43 @@ void FControlRigEditor::HandleViewportCreated(const TSharedRef<class IPersonaVie
 					], 
 					LOCTEXT("ControlRigAxesScale", "Axes Scale")
 				);
+
+				if (UControlRigBlueprint* ControlRigBlueprint = CastChecked<UControlRigBlueprint>(GetBlueprintObj()))
+				{
+					for (UEdGraph* Graph : ControlRigBlueprint->UbergraphPages)
+					{
+						if (UControlRigGraph* RigGraph = Cast<UControlRigGraph>(Graph))
+						{
+							const TArray<TSharedPtr<FString>>* BoneNameList = &RigGraph->GetBoneNameList();
+
+							InMenuBuilder.AddWidget(
+								SNew(SBox)
+								.HAlign(HAlign_Right)
+								[
+									SNew(SBox)
+									.Padding(FMargin(4.0f, 0.0f, 0.0f, 0.0f))
+									.WidthOverride(100.0f)
+									.IsEnabled(this, &FControlRigEditor::IsPinControlNameListEnabled)
+									[
+										SAssignNew(PinControlNameList, SControlRigGraphPinNameListValueWidget)
+										.OptionsSource(BoneNameList)
+										.OnGenerateWidget(this, &FControlRigEditor::MakePinControlNameListItemWidget)
+										.OnSelectionChanged(this, &FControlRigEditor::OnPinControlNameListChanged)
+										.OnComboBoxOpening(this, &FControlRigEditor::OnPinControlNameListComboBox, BoneNameList)
+										.InitiallySelectedItem(GetPinControlCurrentlySelectedItem(BoneNameList))
+										.Content()
+										[
+											SNew(STextBlock)
+											.Text(this, &FControlRigEditor::GetPinControlNameListText)
+										]
+									]
+								],
+								LOCTEXT("ControlRigAuthoringSpace", "Pin Control Space")
+							);
+							break;
+						}
+					}
+				}
 			}
 			InMenuBuilder.EndSection();
 		}
@@ -1468,6 +1698,100 @@ void FControlRigEditor::OnToolbarAxesScaleChanged(float InValue)
 	}
 }
 
+bool FControlRigEditor::IsPinControlNameListEnabled() const
+{
+	if (ControlRig)
+	{
+		if (ControlRig->TransientControls.Num() > 0)
+		{
+			return true;
+		}
+	}
+	return false;
+}
+
+TSharedRef<SWidget> FControlRigEditor::MakePinControlNameListItemWidget(TSharedPtr<FString> InItem)
+{
+	return 	SNew(STextBlock).Text(FText::FromString(*InItem));
+}
+
+FText FControlRigEditor::GetPinControlNameListText() const
+{
+	if (ControlRig)
+	{
+		if (ControlRig->TransientControls.Num() > 0)
+		{
+			return FText::FromName(ControlRig->TransientControls[0].ParentName);
+		}
+	}
+	return FText::FromName(NAME_None);
+}
+
+TSharedPtr<FString> FControlRigEditor::GetPinControlCurrentlySelectedItem(const TArray<TSharedPtr<FString>>* InNameList) const
+{
+	FString CurrentItem = GetPinControlNameListText().ToString();
+	for (const TSharedPtr<FString>& Item : *InNameList)
+	{
+		if (Item->Equals(CurrentItem))
+		{
+			return Item;
+		}
+	}
+	return TSharedPtr<FString>();
+}
+
+void FControlRigEditor::SetPinControlNameListText(const FText& NewTypeInValue, ETextCommit::Type /*CommitInfo*/)
+{
+	if (ControlRig)
+	{
+		if (ControlRig->TransientControls.Num() > 0)
+		{
+			FRigControl& Control = ControlRig->TransientControls[0];
+
+			Control.ParentIndex = ControlRig->Hierarchy.BoneHierarchy.GetIndex(*NewTypeInValue.ToString());
+			if (Control.ParentIndex == INDEX_NONE)
+			{
+				Control.ParentName = NAME_None;
+			}
+			else
+			{
+				Control.ParentName = ControlRig->Hierarchy.BoneHierarchy[Control.ParentIndex].Name;
+			}
+
+			// find out if the controlled pin is part of a visual debug node
+			if (UControlRigBlueprint* ControlRigBlueprint = CastChecked<UControlRigBlueprint>(GetBlueprintObj()))
+			{
+				if (URigVMPin* ControlledPin = ControlRigBlueprint->Model->FindPin(Control.Name.ToString()))
+				{
+					URigVMNode* ControlledNode = ControlledPin->GetPinForLink()->GetNode();
+					if (URigVMPin* BoneSpacePin = ControlledNode->FindPin(TEXT("BoneSpace")))
+					{
+						if (BoneSpacePin->GetCPPType() == TEXT("FName") && BoneSpacePin->GetCustomWidgetName() == TEXT("BoneName"))
+						{
+							ControlRigBlueprint->Controller->SetPinDefaultValue(BoneSpacePin->GetPinPath(), Control.ParentName.ToString(), false, false, false);
+						}
+					}
+				}
+			}
+		}
+	}
+}
+
+void FControlRigEditor::OnPinControlNameListChanged(TSharedPtr<FString> NewSelection, ESelectInfo::Type SelectInfo)
+{
+	if (SelectInfo != ESelectInfo::Direct)
+	{
+		FString NewValue = *NewSelection.Get();
+		SetPinControlNameListText(FText::FromString(NewValue), ETextCommit::OnEnter);
+	}
+}
+
+void FControlRigEditor::OnPinControlNameListComboBox(const TArray<TSharedPtr<FString>>* InNameList)
+{
+	TSharedPtr<FString> CurrentlySelected = GetPinControlCurrentlySelectedItem(InNameList);
+	PinControlNameList->SetSelectedItem(CurrentlySelected);
+}
+
 void FControlRigEditor::HandlePreviewSceneCreated(const TSharedRef<IPersonaPreviewScene>& InPersonaPreviewScene)
 {
 	DECLARE_SCOPE_HIERARCHICAL_COUNTER_FUNC()
@@ -1480,7 +1804,7 @@ void FControlRigEditor::HandlePreviewSceneCreated(const TSharedRef<IPersonaPrevi
 	EditorSkelComp->SetSkeletalMesh(InPersonaPreviewScene->GetPersonaToolkit()->GetPreviewMesh());
 	InPersonaPreviewScene->SetPreviewMeshComponent(EditorSkelComp);
 	bool bWasCreated = false;
-	UAnimCustomInstance::BindToSkeletalMeshComponent<UControlRigSequencerAnimInstance>(EditorSkelComp, bWasCreated);
+	FAnimCustomInstanceHelper::BindToSkeletalMeshComponent<UControlRigLayerInstance>(EditorSkelComp, bWasCreated);
 	InPersonaPreviewScene->AddComponent(EditorSkelComp, FTransform::Identity);
 
 	// set root component, so we can attach to it. 
@@ -1488,16 +1812,37 @@ void FControlRigEditor::HandlePreviewSceneCreated(const TSharedRef<IPersonaPrevi
 
 	// set to use custom default mode defined in mesh component
 	InPersonaPreviewScene->SetDefaultAnimationMode(EPreviewSceneDefaultAnimationMode::Custom);
+
+	PreviewInstance = nullptr;
+	if (UControlRigLayerInstance* ControlRigLayerInstance = Cast<UControlRigLayerInstance>(EditorSkelComp->GetAnimInstance()))
+	{
+		PreviewInstance = Cast<UAnimPreviewInstance>(ControlRigLayerInstance->GetSourceAnimInstance());
+	}
+	else
+	{
+		PreviewInstance = Cast<UAnimPreviewInstance>(EditorSkelComp->GetAnimInstance());
+	}
+
+	if (GEditor)
+	{
+		// remove the preview scene undo handling - it has unwanted side effects
+		FAnimationEditorPreviewScene* AnimationEditorPreviewScene = static_cast<FAnimationEditorPreviewScene*>(&InPersonaPreviewScene.Get());
+		if (AnimationEditorPreviewScene)
+		{
+			GEditor->UnregisterForUndo(AnimationEditorPreviewScene);
+		}
+	}
 }
 
 void FControlRigEditor::UpdateControlRig()
 {
 	DECLARE_SCOPE_HIERARCHICAL_COUNTER_FUNC()
 
-	if(UClass* Class = GetBlueprintObj()->GeneratedClass)
+	UControlRigBlueprint* Blueprint = Cast<UControlRigBlueprint>(GetBlueprintObj());
+	if(UClass* Class = Blueprint->GeneratedClass)
 	{
 		UControlRigSkeletalMeshComponent* EditorSkelComp = Cast<UControlRigSkeletalMeshComponent>(GetPersonaToolkit()->GetPreviewScene()->GetPreviewMeshComponent());
-		UControlRigSequencerAnimInstance* AnimInstance = Cast<UControlRigSequencerAnimInstance>(EditorSkelComp->GetAnimInstance());
+		UControlRigLayerInstance* AnimInstance = Cast<UControlRigLayerInstance>(EditorSkelComp->GetAnimInstance());
 
 		if (AnimInstance)
 		{
@@ -1509,13 +1854,14 @@ void FControlRigEditor::UpdateControlRig()
 				ControlRig->ControlRigLog = &ControlRigLog;
 
 				ControlRig->InitializeFromCDO();
-				ControlRig->InstantiateOperatorsFromGeneratedClass();
  			}
 
+			ControlRig->PreviewInstance = PreviewInstance;
+
 			if (UControlRig* CDO = Cast<UControlRig>(Class->GetDefaultObject()))
 			{
 				CDO->GizmoLibrary = GetControlRigBlueprint()->GizmoLibrary;
- 			}
+			}
 
 			CacheNameLists();
 
@@ -1530,18 +1876,23 @@ void FControlRigEditor::UpdateControlRig()
 
 			// initialize is moved post reinstance
 			FInputBlendPose Filter;
-			AnimInstance->UpdateControlRig(ControlRig, 0, false, false, Filter, 1.0f, true, bExecutionControlRig);
+			AnimInstance->ResetControlRigTracks();
+			AnimInstance->AddControlRigTrack(0, ControlRig);
+			AnimInstance->UpdateControlRigTrack(0, 1.0f, FControlRigIOSettings::MakeEnabled(), bExecutionControlRig);
 			AnimInstance->RecalcRequiredBones();
 
 			// since rig has changed, rebuild draw skeleton
 			EditorSkelComp->RebuildDebugDrawSkeleton();
 			if (FControlRigEditMode* EditMode = GetEditMode())
 			{
-				EditMode->SetObjects(ControlRig, FGuid(), EditorSkelComp);
+				EditMode->SetObjects(ControlRig, FGuid(), EditorSkelComp,nullptr);
 			}
 
-			// update the graph with the compiler errors
-			UpdateGraphCompilerErrors();
+			Blueprint->RecompileVM();
+			
+			ControlRig->OnInitialized_AnyThread().AddSP(this, &FControlRigEditor::HandleControlRigInitializedEvent);
+			ControlRig->RequestInit();
+			ControlRig->ControlModified().AddSP(this, &FControlRigEditor::HandleOnControlModified);
 		}
 	}
 }
@@ -1561,7 +1912,7 @@ void FControlRigEditor::CacheNameLists()
 				continue;
 			}
 
-			RigGraph->CacheNameLists(&ControlRigBP->HierarchyContainer);
+			RigGraph->CacheNameLists(&ControlRigBP->HierarchyContainer, &ControlRigBP->DrawContainer);
 		}
 	}
 }
@@ -1576,6 +1927,14 @@ void FControlRigEditor::AddReferencedObjects( FReferenceCollector& Collector )
 void FControlRigEditor::HandlePreviewMeshChanged(USkeletalMesh* InOldSkeletalMesh, USkeletalMesh* InNewSkeletalMesh)
 {
 	RebindToSkeletalMeshComponent();
+
+	if (GetObjectsCurrentlyBeingEdited()->Num() > 0)
+	{
+		if (UControlRigBlueprint* ControlRigBP = GetControlRigBlueprint())
+		{
+			ControlRigBP->SetPreviewMesh(InNewSkeletalMesh);
+		}
+	}
 }
 
 void FControlRigEditor::RebindToSkeletalMeshComponent()
@@ -1586,7 +1945,69 @@ void FControlRigEditor::RebindToSkeletalMeshComponent()
 	if (MeshComponent)
 	{
 		bool bWasCreated = false;
-		UAnimCustomInstance::BindToSkeletalMeshComponent<UControlRigSequencerAnimInstance>(MeshComponent , bWasCreated);
+		FAnimCustomInstanceHelper::BindToSkeletalMeshComponent<UControlRigLayerInstance>(MeshComponent , bWasCreated);
+	}
+}
+
+void FControlRigEditor::UpdateStaleWatchedPins()
+{
+	UControlRigBlueprint* ControlRigBP = GetControlRigBlueprint();
+	if (ControlRigBP == nullptr)
+	{
+		return;
+	}
+
+	TSet<UEdGraphPin*> AllPins;
+
+	// Find all unique pins being watched
+	for (const FEdGraphPinReference& PinRef : ControlRigBP->WatchedPins)
+	{
+		UEdGraphPin* Pin = PinRef.Get();
+		if (Pin == nullptr)
+		{
+			continue;
+		}
+
+		UEdGraphNode* OwningNode = Pin->GetOwningNode();
+		// during node reconstruction, dead pins get moved to the transient 
+		// package (so just in case this blueprint got saved with dead pin watches)
+		if (OwningNode == nullptr)
+		{
+			continue;
+		}
+
+		if (!OwningNode->Pins.Contains(Pin))
+		{
+			continue;
+		}
+
+		AllPins.Add(Pin);
+	}
+
+	// Refresh watched pins with unique pins (throw away null or duplicate watches)
+	if (ControlRigBP->WatchedPins.Num() != AllPins.Num())
+	{
+		ControlRigBP->Status = BS_Dirty;
+	}
+
+	ControlRigBP->WatchedPins.Empty();
+
+	for (URigVMNode* ModelNode : ControlRigBP->Model->GetNodes())
+	{
+		TArray<URigVMPin*> ModelPins = ModelNode->GetAllPinsRecursively();
+		for (URigVMPin* ModelPin : ModelPins)
+		{
+			if (ModelPin->RequiresWatch())
+			{
+				ControlRigBP->Controller->SetPinIsWatched(ModelPin->GetPinPath(), false, false);
+			}
+		}
+	}
+
+	for (UEdGraphPin* Pin : AllPins)
+	{
+		ControlRigBP->WatchedPins.Add(Pin);
+		ControlRigBP->Controller->SetPinIsWatched(Pin->GetName(), true, false);
 	}
 }
 
@@ -1613,10 +2034,7 @@ void FControlRigEditor::OnNodeTitleCommitted(const FText& NewText, ETextCommit::
 	{
 		if (UControlRigBlueprint* ControlRigBP = GetControlRigBlueprint())
 		{
-			if (ControlRigBP->ModelController)
-			{
-				ControlRigBP->ModelController->SetCommentText(CommentBeingChanged->GetFName(), NewText.ToString(), true);
-			}
+			ControlRigBP->Controller->SetCommentTextByName(CommentBeingChanged->GetFName(), NewText.ToString(), true);
 		}
 	}
 }
@@ -1715,11 +2133,73 @@ void FControlRigEditor::SetRigElementTransform(const FRigElementKey& InElement,
 	}
 }
 
+void FControlRigEditor::NotifyPreChange(FProperty* PropertyAboutToChange)
+{
+	FBlueprintEditor::NotifyPreChange(PropertyAboutToChange);
+
+	if (UControlRigBlueprint* ControlRigBP = GetControlRigBlueprint())
+	{
+		if (RigElementInDetailPanel.IsValid())
+		{
+			ControlRigBP->Modify();
+		}
+	}
+}
+
+void FControlRigEditor::NotifyPostChange(const FPropertyChangedEvent& PropertyChangedEvent, FProperty* PropertyThatChanged)
+{
+	FBlueprintEditor::NotifyPostChange(PropertyChangedEvent, PropertyThatChanged);
+}
+
 void FControlRigEditor::OnFinishedChangingProperties(const FPropertyChangedEvent& PropertyChangedEvent)
 {
 	DECLARE_SCOPE_HIERARCHICAL_COUNTER_FUNC()
 
 	UControlRigBlueprint* ControlRigBP = GetControlRigBlueprint();
+
+	if (ControlRigBP)
+	{
+		if (PropertyChangedEvent.MemberProperty->GetNameCPP() == TEXT("VMCompileSettings"))
+		{
+			ControlRigBP->RecompileVM();
+			return;
+		}
+
+		if (PropertyChangedEvent.MemberProperty->GetNameCPP() == TEXT("DrawContainer"))
+		{
+			ControlRigBP->PropagateDrawInstructionsFromBPToInstances();
+			return;
+		}
+	}
+
+	if (ControlRig && ControlRigBP && NodeDetailBuffer.Num() > 0 && NodeDetailStruct != nullptr && !NodeDetailName.IsNone())
+	{
+		const uint8* PropertyValuePtr = NodeDetailBuffer.GetData();
+		PropertyValuePtr += PropertyChangedEvent.MemberProperty->GetOffset_ReplaceWith_ContainerPtrToValuePtr();
+
+		FString DefaultValue;
+		PropertyChangedEvent.MemberProperty->ExportTextItem(DefaultValue, PropertyValuePtr, PropertyValuePtr, nullptr, PPF_None);
+
+		if (DefaultValue.IsEmpty())
+		{
+			if (CastField<FArrayProperty>(PropertyChangedEvent.MemberProperty))
+			{
+				DefaultValue = TEXT("()");
+			}
+			else if (CastField<FNameProperty>(PropertyChangedEvent.MemberProperty) != nullptr ||
+				CastField<FStrProperty>(PropertyChangedEvent.MemberProperty) != nullptr)
+			{
+				DefaultValue = TEXT("\"\"");
+			}
+		}
+
+		if (!DefaultValue.IsEmpty())
+		{
+			FString PinPath = FString::Printf(TEXT("%s.%s"), *NodeDetailName.ToString(), *PropertyChangedEvent.MemberProperty->GetName());
+			ControlRigBP->Controller->SetPinDefaultValue(PinPath, DefaultValue, true, true);
+		}
+	}
+
 	if (ControlRig && ControlRigBP && RigElementInDetailPanel)
 	{
 		UControlRig* DebuggedControlRig = Cast<UControlRig>(GetBlueprintObj()->GetObjectBeingDebugged());
@@ -1729,53 +2209,97 @@ void FControlRigEditor::OnFinishedChangingProperties(const FPropertyChangedEvent
 			return;
 		}
 
-		UScriptStruct* ScriptStruct = PropertyChangedEvent.Property->GetOwner<UScriptStruct>();
+		UScriptStruct* ScriptStruct = PropertyChangedEvent.MemberProperty->GetOwner<UScriptStruct>();
 		if (ScriptStruct)
 		{
 			if (ScriptStruct == FRigBone::StaticStruct() && RigElementInDetailPanel.Type == ERigElementType::Bone)
 			{
-				ControlRig->Hierarchy.BoneHierarchy[RigElementInDetailPanel.Name] = ControlRigBP->HierarchyContainer.BoneHierarchy[RigElementInDetailPanel.Name];
-				if (DebuggedControlRig && DebuggedControlRig != ControlRig)
+				if (PropertyChangedEvent.MemberProperty->GetFName() == TEXT("LocalTransform"))
 				{
-					DebuggedControlRig->Hierarchy.BoneHierarchy[RigElementInDetailPanel.Name] = ControlRigBP->HierarchyContainer.BoneHierarchy[RigElementInDetailPanel.Name];
+					FRigBone& Bone = ControlRigBP->HierarchyContainer.BoneHierarchy[RigElementInDetailPanel.Name];
+					FTransform ParentTransform = FTransform::Identity;
+					if (Bone.ParentIndex != INDEX_NONE)
+					{
+						ParentTransform = ControlRigBP->HierarchyContainer.BoneHierarchy.GetGlobalTransform(Bone.ParentIndex);
+					}
+					Bone.GlobalTransform = Bone.LocalTransform * ParentTransform;
+
+					ControlRigBP->PropagatePropertyFromBPToInstances(RigElementInDetailPanel, PropertyChangedEvent.MemberProperty);
+					ControlRigBP->PropagatePropertyFromBPToInstances(RigElementInDetailPanel, ScriptStruct->FindPropertyByName(TEXT("GlobalTransform")));
 				}
-			}
-			else if (ScriptStruct == FRigSpace::StaticStruct() && RigElementInDetailPanel.Type == ERigElementType::Space)
-	{
-				ControlRig->Hierarchy.SpaceHierarchy[RigElementInDetailPanel.Name] = ControlRigBP->HierarchyContainer.SpaceHierarchy[RigElementInDetailPanel.Name];
+				else if (PropertyChangedEvent.MemberProperty->GetFName() == TEXT("GlobalTransform"))
+				{
+					FRigBone& Bone = ControlRigBP->HierarchyContainer.BoneHierarchy[RigElementInDetailPanel.Name];
+					FTransform ParentTransform = FTransform::Identity;
+					if (Bone.ParentIndex != INDEX_NONE)
+					{
+						ParentTransform = ControlRigBP->HierarchyContainer.BoneHierarchy.GetGlobalTransform(Bone.ParentIndex);
+					}
+					Bone.LocalTransform = Bone.GlobalTransform.GetRelativeTransform(ParentTransform);
+
+					ControlRigBP->PropagatePropertyFromBPToInstances(RigElementInDetailPanel, PropertyChangedEvent.MemberProperty);
+					ControlRigBP->PropagatePropertyFromBPToInstances(RigElementInDetailPanel, ScriptStruct->FindPropertyByName(TEXT("LocalTransform")));
+				}
+
+				ControlRig->SetTransientControlValue(RigElementInDetailPanel);
 				if (DebuggedControlRig && DebuggedControlRig != ControlRig)
 				{
-					DebuggedControlRig->Hierarchy.SpaceHierarchy[RigElementInDetailPanel.Name] = ControlRigBP->HierarchyContainer.SpaceHierarchy[RigElementInDetailPanel.Name];
+					DebuggedControlRig->SetTransientControlValue(RigElementInDetailPanel);
+				}
+
+				if (PreviewInstance)
+				{
+					if (FAnimNode_ModifyBone* Modify = PreviewInstance->FindModifiedBone(RigElementInDetailPanel.Name))
+					{
+						FTransform LocalTransform = ControlRigBP->HierarchyContainer.BoneHierarchy[RigElementInDetailPanel.Name].LocalTransform;
+						Modify->Translation = LocalTransform.GetTranslation();
+						Modify->Rotation = LocalTransform.GetRotation().Rotator();
+						Modify->TranslationSpace = EBoneControlSpace::BCS_ParentBoneSpace;
+						Modify->RotationSpace = EBoneControlSpace::BCS_ParentBoneSpace;
+					}
 				}
 			}
-			else if (ScriptStruct == FRigControl::StaticStruct() && RigElementInDetailPanel.Type == ERigElementType::Control)
+			else if (ScriptStruct == FRigSpace::StaticStruct() && RigElementInDetailPanel.Type == ERigElementType::Space)
 			{
-				ControlRig->Hierarchy.ControlHierarchy[RigElementInDetailPanel.Name] = ControlRigBP->HierarchyContainer.ControlHierarchy[RigElementInDetailPanel.Name];
+				ControlRigBP->PropagatePropertyFromBPToInstances(RigElementInDetailPanel, PropertyChangedEvent.MemberProperty);
+
+				if (PropertyChangedEvent.MemberProperty->GetName() == TEXT("InitialTransform"))
+				{
+					FRigSpace& Space = ControlRigBP->HierarchyContainer.SpaceHierarchy[RigElementInDetailPanel.Name];
+					Space.LocalTransform = Space.InitialTransform;
+					ControlRigBP->PropagatePropertyFromBPToInstances(RigElementInDetailPanel, ScriptStruct->FindPropertyByName(TEXT("LocalTransform")));
+				}
+
+				ControlRig->SetTransientControlValue(RigElementInDetailPanel);
 				if (DebuggedControlRig && DebuggedControlRig != ControlRig)
-		{
-					DebuggedControlRig->Hierarchy.ControlHierarchy[RigElementInDetailPanel.Name] = ControlRigBP->HierarchyContainer.ControlHierarchy[RigElementInDetailPanel.Name];
+				{
+					DebuggedControlRig->SetTransientControlValue(RigElementInDetailPanel);
 				}
 			}
-			else if (ScriptStruct == FRigCurve::StaticStruct() && RigElementInDetailPanel.Type == ERigElementType::Curve)
+			else if (ScriptStruct == FRigControl::StaticStruct() && RigElementInDetailPanel.Type == ERigElementType::Control)
 			{
-				ControlRig->Hierarchy.CurveContainer[RigElementInDetailPanel.Name] = ControlRigBP->HierarchyContainer.CurveContainer[RigElementInDetailPanel.Name];
-				if (DebuggedControlRig && DebuggedControlRig != ControlRig)
+				if (PropertyChangedEvent.MemberProperty->GetName() == TEXT("GizmoColor"))
 				{
-					DebuggedControlRig->Hierarchy.CurveContainer[RigElementInDetailPanel.Name] = ControlRigBP->HierarchyContainer.CurveContainer[RigElementInDetailPanel.Name];
+					FRigControl& Control = ControlRigBP->HierarchyContainer.ControlHierarchy[RigElementInDetailPanel.Name];
+					Control.GizmoColor.R = FMath::Clamp<float>(Control.GizmoColor.R, 0.f, 1.f);
+					Control.GizmoColor.G = FMath::Clamp<float>(Control.GizmoColor.G, 0.f, 1.f);
+					Control.GizmoColor.B = FMath::Clamp<float>(Control.GizmoColor.B, 0.f, 1.f);
+					Control.GizmoColor.A = FMath::Clamp<float>(Control.GizmoColor.A, 0.f, 1.f);
 				}
-			}
 
-			ControlRigBP->MarkPackageDirty();
-		}
+				ControlRigBP->PropagatePropertyFromBPToInstances(RigElementInDetailPanel, PropertyChangedEvent.MemberProperty);
 
-		if (PropertyChangedEvent.Property->GetName().StartsWith(TEXT("Gizmo")) && RigElementInDetailPanel.Type == ERigElementType::Control)
-		{
-			ControlRig->Hierarchy.ControlHierarchy[RigElementInDetailPanel.Name] = ControlRigBP->HierarchyContainer.ControlHierarchy[RigElementInDetailPanel.Name];
-			if (DebuggedControlRig && DebuggedControlRig != ControlRig)
+				if (PropertyChangedEvent.MemberProperty->GetName().Contains(TEXT("Gizmo")))
+				{
+					ControlRigBP->HierarchyContainer.ControlHierarchy.OnControlUISettingsChanged.Broadcast(&ControlRigBP->HierarchyContainer, FRigElementKey(RigElementInDetailPanel.Name, ERigElementType::Control));
+				}
+			}
+			else if (ScriptStruct == FRigCurve::StaticStruct() && RigElementInDetailPanel.Type == ERigElementType::Curve)
 			{
-				DebuggedControlRig->Hierarchy.ControlHierarchy[RigElementInDetailPanel.Name] = ControlRigBP->HierarchyContainer.ControlHierarchy[RigElementInDetailPanel.Name];
+				ControlRigBP->PropagatePropertyFromBPToInstances(RigElementInDetailPanel, PropertyChangedEvent.MemberProperty);
 			}
-			ControlRigBP->HierarchyContainer.ControlHierarchy.OnControlUISettingsChanged.Broadcast(&ControlRigBP->HierarchyContainer, FRigElementKey(RigElementInDetailPanel.Name, ERigElementType::Control));
+
+			ControlRigBP->Modify();
 			ControlRigBP->MarkPackageDirty();
 		}
 	}
@@ -1790,13 +2314,26 @@ void FControlRigEditor::OnHierarchyChanged()
 	{
 		ControlRigBP->PropagateHierarchyFromBPToInstances();
 
+		FBlueprintEditorUtils::MarkBlueprintAsModified(GetControlRigBlueprint());
 		TArray<FRigElementKey> SelectedElements = ControlRigBP->HierarchyContainer.CurrentSelection();
 		for(const FRigElementKey& SelectedElement : SelectedElements)
 		{
 			ControlRigBP->HierarchyContainer.OnElementSelected.Broadcast(&ControlRigBP->HierarchyContainer, SelectedElement, true);
 		}
+		GetControlRigBlueprint()->RecompileVM();
 
-		if (ControlRigBP->HierarchyContainer.GetIndex(RigElementInDetailPanel) == INDEX_NONE)
+		if (NodeDetailStruct != nullptr && NodeDetailBuffer.Num() > 0 && NodeDetailName != NAME_None)
+		{
+			if (URigVMNode* Node = ControlRigBP->Model->FindNode(NodeDetailName.ToString()))
+			{
+				SetDetailObject(Node);
+			}
+			else
+			{
+				ClearDetailObject();
+			}
+		}
+		else if (ControlRigBP->HierarchyContainer.GetIndex(RigElementInDetailPanel) == INDEX_NONE)
 		{
 			ClearDetailObject();
 		}
@@ -1830,71 +2367,22 @@ void FControlRigEditor::OnRigElementRenamed(FRigHierarchyContainer* Container, E
 			continue;
 		}
 
-		CacheNameLists();
-
 		for (UEdGraphNode* Node : RigGraph->Nodes)
 		{
-			UControlRigGraphNode* RigNode = Cast<UControlRigGraphNode>(Node);
-			if (RigNode == nullptr)
-			{
-				continue;
-			}
-
-			FStructProperty* UnitProperty = RigNode->GetUnitProperty();
-			UStruct* UnitStruct = RigNode->GetUnitScriptStruct();
-			if (UnitProperty && UnitStruct)
+			if (UControlRigGraphNode* RigNode = Cast<UControlRigGraphNode>(Node))
 			{
-				for (TFieldIterator<FNameProperty> It(UnitStruct); It; ++It)
+				if (URigVMNode* ModelNode = RigNode->GetModelNode())
 				{
-					if (It->HasMetaData(UControlRig::BoneNameMetaName) && ElementType == ERigElementType::Bone)
+					for (URigVMPin * ModelPin : ModelNode->GetPins())
 					{
-						FString PinName = FString::Printf(TEXT("%s.%s"), *UnitProperty->GetName(), *It->GetName());
-						UEdGraphPin* Pin = Node->FindPin(PinName, EEdGraphPinDirection::EGPD_Input);
-						if (Pin)
+						if ((ModelPin->GetCPPType() == TEXT("FName") && ModelPin->GetCustomWidgetName() == TEXT("BoneName") && ElementType == ERigElementType::Bone) ||
+							(ModelPin->GetCPPType() == TEXT("FName") && ModelPin->GetCustomWidgetName() == TEXT("ControlName") && ElementType == ERigElementType::Control) ||
+							(ModelPin->GetCPPType() == TEXT("FName") && ModelPin->GetCustomWidgetName() == TEXT("SpaceName") && ElementType == ERigElementType::Space) ||
+							(ModelPin->GetCPPType() == TEXT("FName") && ModelPin->GetCustomWidgetName() == TEXT("CurveName") && ElementType == ERigElementType::Curve))
 						{
-							FName CurrentBone = FName(*Pin->GetDefaultAsString());
-							if (CurrentBone == InOldName)
+							if (ModelPin->GetDefaultValue() == InOldName.ToString())
 							{
-								Pin->GetSchema()->TrySetDefaultValue(*Pin, InNewName.ToString());
-							}
-						}
-					}
-					if (It->HasMetaData(UControlRig::ControlNameMetaName) && ElementType == ERigElementType::Control)
-					{
-						FString PinName = FString::Printf(TEXT("%s.%s"), *UnitProperty->GetName(), *It->GetName());
-						UEdGraphPin* Pin = Node->FindPin(PinName, EEdGraphPinDirection::EGPD_Input);
-						if (Pin)
-						{
-							FName CurrentControl = FName(*Pin->GetDefaultAsString());
-							if (CurrentControl == InOldName)
-							{
-								Pin->GetSchema()->TrySetDefaultValue(*Pin, InNewName.ToString());
-							}
-						}
-					}
-					if (It->HasMetaData(UControlRig::SpaceNameMetaName) && ElementType == ERigElementType::Space)
-					{
-						FString PinName = FString::Printf(TEXT("%s.%s"), *UnitProperty->GetName(), *It->GetName());
-						UEdGraphPin* Pin = Node->FindPin(PinName, EEdGraphPinDirection::EGPD_Input);
-						if (Pin)
-						{
-							FName CurrentSpace = FName(*Pin->GetDefaultAsString());
-							if (CurrentSpace == InOldName)
-							{
-								Pin->GetSchema()->TrySetDefaultValue(*Pin, InNewName.ToString());
-							}
-						}
-					}
-					if (It->HasMetaData(UControlRig::CurveNameMetaName) && ElementType == ERigElementType::Curve)
-					{
-						FString PinName = FString::Printf(TEXT("%s.%s"), *UnitProperty->GetName(), *It->GetName());
-						UEdGraphPin* Pin = Node->FindPin(PinName, EEdGraphPinDirection::EGPD_Input);
-						if (Pin)
-						{
-							FName CurrentCurve = FName(*Pin->GetDefaultAsString());
-							if (CurrentCurve == InOldName)
-							{
-								Pin->GetSchema()->TrySetDefaultValue(*Pin, InNewName.ToString());
+								Blueprint->Controller->SetPinDefaultValue(ModelPin->GetPinPath(), InNewName.ToString(), false);
 							}
 						}
 					}
@@ -1904,6 +2392,7 @@ void FControlRigEditor::OnRigElementRenamed(FRigHierarchyContainer* Container, E
 
 		CacheNameLists();
 	}
+
 }
 
 void FControlRigEditor::OnRigElementReparented(FRigHierarchyContainer* Container, const FRigElementKey& InKey, const FName& InOldParentName, const FName& InNewParentName)
@@ -1919,84 +2408,76 @@ void FControlRigEditor::OnRigElementSelected(FRigHierarchyContainer* Container,
 		return;
 	}
 
-	switch (InKey.Type)
+	if (InKey.Type == ERigElementType::Bone)
 	{
-		case ERigElementType::Bone:
+		UControlRigSkeletalMeshComponent* EditorSkelComp = Cast<UControlRigSkeletalMeshComponent>(GetPersonaToolkit()->GetPreviewScene()->GetPreviewMeshComponent());
+		if (EditorSkelComp)
 		{
-			FRigBoneHierarchy& BoneHierarchy = RigBlueprint->HierarchyContainer.BoneHierarchy;
-			if (bSelected)
-							{
-				SetDetailStruct(InKey, MakeShareable(new FStructOnScope(FRigBone::StaticStruct(), (uint8*)&(BoneHierarchy[InKey.Name]))));
-							}
-			else
+			int32 Index = ControlRig->Hierarchy.BoneHierarchy.GetIndex(InKey.Name);
+			if (Index != INDEX_NONE)
 			{
-				ClearDetailObject();
-						}
-
-			UControlRigSkeletalMeshComponent* EditorSkelComp = Cast<UControlRigSkeletalMeshComponent>(GetPersonaToolkit()->GetPreviewScene()->GetPreviewMeshComponent());
-			if (EditorSkelComp)
-			{
-				int32 Index = ControlRig->Hierarchy.BoneHierarchy.GetIndex(InKey.Name);
-				if (Index != INDEX_NONE)
+				if (bSelected)
 				{
-					if (bSelected)
-					{
-						EditorSkelComp->BonesOfInterest.Add(Index);
-					}
-					else
-					{
-						EditorSkelComp->BonesOfInterest.Remove(Index);
+					EditorSkelComp->BonesOfInterest.AddUnique(Index);
+				}
+				else
+				{
+					EditorSkelComp->BonesOfInterest.Remove(Index);
 				}
 			}
 		}
+	}
 
-			break;
-		}
-		case ERigElementType::Space:
+	if (bSelected)
+	{
+		switch (InKey.Type)
 		{
-			if (bSelected)
+			case ERigElementType::Bone:
 			{
-				FRigSpaceHierarchy& SpaceHierarchy = RigBlueprint->HierarchyContainer.SpaceHierarchy;
-				SetDetailStruct(InKey, MakeShareable(new FStructOnScope(FRigSpace::StaticStruct(), (uint8*)&(SpaceHierarchy[InKey.Name]))));
+				FRigBoneHierarchy& BoneHierarchy = RigBlueprint->HierarchyContainer.BoneHierarchy;
+				SetDetailStruct(InKey, MakeShareable(new FStructOnScope(FRigBone::StaticStruct(), (uint8*)&(BoneHierarchy[InKey.Name]))));
+				break;
 			}
-			else
+			case ERigElementType::Space:
 			{
-				ClearDetailObject();
+				FRigSpaceHierarchy& SpaceHierarchy = RigBlueprint->HierarchyContainer.SpaceHierarchy;
+				SetDetailStruct(InKey, MakeShareable(new FStructOnScope(FRigSpace::StaticStruct(), (uint8*)&(SpaceHierarchy[InKey.Name]))));
+				break;
 			}
-			break;
-		}
-		case ERigElementType::Control:
-		{
-			if (bSelected)
+			case ERigElementType::Control:
 			{
 				FRigControlHierarchy& ControlHierarchy = RigBlueprint->HierarchyContainer.ControlHierarchy;
-				SetDetailStruct(InKey, MakeShareable(new FStructOnScope(FRigControl::StaticStruct(), (uint8*)&(ControlHierarchy[InKey.Name]))));
-			}
-			else
-			{
-				ClearDetailObject();
+				int32 ControlIndex = ControlHierarchy.GetIndex(InKey.Name);
+				if (ControlIndex != INDEX_NONE)
+				{
+					SetDetailStruct(InKey, MakeShareable(new FStructOnScope(FRigControl::StaticStruct(), (uint8*)&(ControlHierarchy[ControlIndex]))));
+				}
+				break;
 			}
-			break;
-		}
-		case ERigElementType::Curve:
-		{
-			if (bSelected)
+			case ERigElementType::Curve:
 			{
 				FRigCurveContainer& CurveContainer = RigBlueprint->HierarchyContainer.CurveContainer;
 				SetDetailStruct(InKey, MakeShareable(new FStructOnScope(FRigCurve::StaticStruct(), (uint8*)&(CurveContainer[InKey.Name]))));
+				break;
 			}
-			else
-			{
-				ClearDetailObject();
-			}
-			break;
+		}
+	}
+	else
+	{
+		TArray<FRigElementKey> CurrentSelection = RigBlueprint->HierarchyContainer.CurrentSelection();
+		if (CurrentSelection.Num() > 0)
+		{
+			OnRigElementSelected(Container, CurrentSelection.Last(), true);
+		}
+		else
+		{
+			ClearDetailObject();
 		}
 	}
 }
 
 void FControlRigEditor::OnControlUISettingChanged(FRigHierarchyContainer* Container, const FRigElementKey& InKey)
 {
-	OnHierarchyChanged();
 }
 
 FControlRigEditorEditMode* FControlRigEditor::GetEditMode() const
@@ -2079,82 +2560,91 @@ void FControlRigEditor::OnGraphNodeDropToPerform(TSharedPtr<FGraphNodeDragDropOp
 					)
 				);
 
-			MenuBuilder.AddMenuSeparator();
+				MenuBuilder.AddMenuSeparator();
 
-			MenuBuilder.AddMenuEntry(
-				LOCTEXT("CreateSetBoneRotation", "Set Bone Rotation"),
-				LOCTEXT("CreateSetBoneRotationTooltip", "Setter for bone Rotation\n"),
-				FSlateIcon(),
-				FUIAction(
-					FExecuteAction::CreateSP(this, &FControlRigEditor::HandleMakeElementGetterSetter, ERigElementGetterSetterType_Rotation, false, DraggedKeys, Graph, NodePosition),
-					FCanExecuteAction()
-				)
-			);
+				MenuBuilder.AddMenuEntry(
+					LOCTEXT("CreateSetBoneRotation", "Set Bone Rotation"),
+					LOCTEXT("CreateSetBoneRotationTooltip", "Setter for bone Rotation\n"),
+					FSlateIcon(),
+					FUIAction(
+						FExecuteAction::CreateSP(this, &FControlRigEditor::HandleMakeElementGetterSetter, ERigElementGetterSetterType_Rotation, false, DraggedKeys, Graph, NodePosition),
+						FCanExecuteAction()
+					)
+				);
 
-			MenuBuilder.AddMenuEntry(
-				LOCTEXT("CreateSetBoneTranslation", "Set Bone Translation"),
-				LOCTEXT("CreateSetBoneTranslationTooltip", "Setter for bone translation\n"),
-				FSlateIcon(),
-				FUIAction(
-					FExecuteAction::CreateSP(this, &FControlRigEditor::HandleMakeElementGetterSetter, ERigElementGetterSetterType_Translation, false, DraggedKeys, Graph, NodePosition),
-					FCanExecuteAction()
-				)
-			);
+				MenuBuilder.AddMenuEntry(
+					LOCTEXT("CreateSetBoneTranslation", "Set Bone Translation"),
+					LOCTEXT("CreateSetBoneTranslationTooltip", "Setter for bone translation\n"),
+					FSlateIcon(),
+					FUIAction(
+						FExecuteAction::CreateSP(this, &FControlRigEditor::HandleMakeElementGetterSetter, ERigElementGetterSetterType_Translation, false, DraggedKeys, Graph, NodePosition),
+						FCanExecuteAction()
+					)
+				);
 
-			MenuBuilder.AddMenuEntry(
-				LOCTEXT("CreateSetBoneOffset", "Set Bone Offset"),
-				LOCTEXT("CreateSetBoneOffsetTooltip", "Setter for bone offset\n"),
-				FSlateIcon(),
-				FUIAction(
-					FExecuteAction::CreateSP(this, &FControlRigEditor::HandleMakeElementGetterSetter, ERigElementGetterSetterType_Offset, false, DraggedKeys, Graph, NodePosition),
-					FCanExecuteAction()
-				)
-			);
+				MenuBuilder.AddMenuEntry(
+					LOCTEXT("CreateSetBoneOffset", "Set Bone Offset"),
+					LOCTEXT("CreateSetBoneOffsetTooltip", "Setter for bone offset\n"),
+					FSlateIcon(),
+					FUIAction(
+						FExecuteAction::CreateSP(this, &FControlRigEditor::HandleMakeElementGetterSetter, ERigElementGetterSetterType_Offset, false, DraggedKeys, Graph, NodePosition),
+						FCanExecuteAction()
+					)
+				);
 
-			MenuBuilder.AddMenuEntry(
-				LOCTEXT("CreateGetInitialBoneTransform", "Get Initial Bone Transform"),
-				LOCTEXT("CreateGetInitialBoneTransformTooltip", "Getter for initial bone transform\n"),
-				FSlateIcon(),
-				FUIAction(
-					FExecuteAction::CreateSP(this, &FControlRigEditor::HandleMakeElementGetterSetter, ERigElementGetterSetterType_Initial, true, DraggedKeys, Graph, NodePosition),
-					FCanExecuteAction()
-				)
-			);
+				MenuBuilder.AddMenuEntry(
+					LOCTEXT("CreateGetInitialBoneTransform", "Get Initial Bone Transform"),
+					LOCTEXT("CreateGetInitialBoneTransformTooltip", "Getter for initial bone transform\n"),
+					FSlateIcon(),
+					FUIAction(
+						FExecuteAction::CreateSP(this, &FControlRigEditor::HandleMakeElementGetterSetter, ERigElementGetterSetterType_Initial, true, DraggedKeys, Graph, NodePosition),
+						FCanExecuteAction()
+					)
+				);
 
-			MenuBuilder.AddMenuSeparator();
+				MenuBuilder.AddMenuSeparator();
 
-			MenuBuilder.AddMenuEntry(
-				LOCTEXT("CreateGetBoneRelativeTransform", "Get Bone Relative Transform"),
-				LOCTEXT("CreateGetBoneRelativeTransformTooltip", "Getter for bone relative transform\n"),
-				FSlateIcon(),
-				FUIAction(
-					FExecuteAction::CreateSP(this, &FControlRigEditor::HandleMakeElementGetterSetter, ERigElementGetterSetterType_Relative, true, DraggedKeys, Graph, NodePosition),
-					FCanExecuteAction()
-				)
-			);
-			MenuBuilder.AddMenuEntry(
-				LOCTEXT("CreateSetBoneRelativeTransform", "Set Bone Relative Transform"),
-				LOCTEXT("CreateSetBoneRelativeTransformTooltip", "Setter for bone relative transform\n"),
-				FSlateIcon(),
-				FUIAction(
-					FExecuteAction::CreateSP(this, &FControlRigEditor::HandleMakeElementGetterSetter, ERigElementGetterSetterType_Relative, false, DraggedKeys, Graph, NodePosition),
-					FCanExecuteAction()
-				)
-			);
+				MenuBuilder.AddMenuEntry(
+					LOCTEXT("CreateGetBoneRelativeTransform", "Get Bone Relative Transform"),
+					LOCTEXT("CreateGetBoneRelativeTransformTooltip", "Getter for bone relative transform\n"),
+					FSlateIcon(),
+					FUIAction(
+						FExecuteAction::CreateSP(this, &FControlRigEditor::HandleMakeElementGetterSetter, ERigElementGetterSetterType_Relative, true, DraggedKeys, Graph, NodePosition),
+						FCanExecuteAction()
+					)
+				);
+				MenuBuilder.AddMenuEntry(
+					LOCTEXT("CreateSetBoneRelativeTransform", "Set Bone Relative Transform"),
+					LOCTEXT("CreateSetBoneRelativeTransformTooltip", "Setter for bone relative transform\n"),
+					FSlateIcon(),
+					FUIAction(
+						FExecuteAction::CreateSP(this, &FControlRigEditor::HandleMakeElementGetterSetter, ERigElementGetterSetterType_Relative, false, DraggedKeys, Graph, NodePosition),
+						FCanExecuteAction()
+					)
+				);
 
 			}
 
 			if ((DraggedTypes & (uint8)ERigElementType::Control) != 0)
 			{
 				MenuBuilder.AddMenuEntry(
-					LOCTEXT("CreateGetControlTransform", "Get Control Transform"),
-					LOCTEXT("CreateGetControlTransformTooltip", "Getter for control transform\n"),
+					LOCTEXT("CreateGetControl", "Get Control"),
+					LOCTEXT("CreateGetControlTooltip", "Getter for control\n"),
 					FSlateIcon(),
 					FUIAction(
 						FExecuteAction::CreateSP(this, &FControlRigEditor::HandleMakeElementGetterSetter, ERigElementGetterSetterType_Transform, true, DraggedKeys, Graph, NodePosition),
 						FCanExecuteAction()
 					)
 				);
+				MenuBuilder.AddMenuEntry(
+					LOCTEXT("CreateSetControl", "Set Control"),
+					LOCTEXT("CreateSetControlTooltip", "Setter for control\n"),
+					FSlateIcon(),
+					FUIAction(
+						FExecuteAction::CreateSP(this, &FControlRigEditor::HandleMakeElementGetterSetter, ERigElementGetterSetterType_Transform, false, DraggedKeys, Graph, NodePosition),
+						FCanExecuteAction()
+					)
+				);
 			}
 
 			if ((DraggedTypes & (uint8)ERigElementType::Space) != 0)
@@ -2200,18 +2690,29 @@ void FControlRigEditor::HandleMakeElementGetterSetter(ERigElementGetterSetterTyp
 	{
 		return;
 	}
-	if (Blueprint->ModelController == nullptr)
+	if (Blueprint->Controller == nullptr)
 	{
 		return;
 	}
 
-	Blueprint->ModelController->OpenUndoBracket(TEXT("Adding Nodes from Hierarchy"));
+	Blueprint->Controller->OpenUndoBracket(TEXT("Adding Nodes from Hierarchy"));
 
-	TArray<FName> NewNodeNames;
+	struct FNewNodeData
+	{
+		FName Name;
+		FName ValuePinName;
+		ERigControlType ValueType;
+		FRigControlValue Value;
+	};
+	TArray<FNewNodeData> NewNodes;
 
 	for (const FRigElementKey& Key : Keys)
 	{
-	UStruct* StructTemplate = nullptr;
+		UScriptStruct* StructTemplate = nullptr;
+
+		FNewNodeData NewNode;
+		NewNode.Name = NAME_None;
+		NewNode.ValuePinName = NAME_None;
 
 		FName PinName;
 		switch (Key.Type)
@@ -2220,100 +2721,207 @@ void FControlRigEditor::HandleMakeElementGetterSetter(ERigElementGetterSetterTyp
 			{
 				PinName = TEXT("Bone");
 
-	if (bIsGetter)
-	{
-		switch (Type)
-		{
+				if (bIsGetter)
+				{
+					switch (Type)
+					{
 						case ERigElementGetterSetterType_Transform:
-			{
-				StructTemplate = FRigUnit_GetBoneTransform::StaticStruct();
-				break;
-			}
+						{
+							StructTemplate = FRigUnit_GetBoneTransform::StaticStruct();
+							break;
+						}
 						case ERigElementGetterSetterType_Initial:
-			{
-				StructTemplate = FRigUnit_GetInitialBoneTransform::StaticStruct();
-				break;
-			}
+						{
+							StructTemplate = FRigUnit_GetInitialBoneTransform::StaticStruct();
+							break;
+						}
 						case ERigElementGetterSetterType_Relative:
-			{
-				StructTemplate = FRigUnit_GetRelativeBoneTransform::StaticStruct();
-				break;
-			}
+						{
+							StructTemplate = FRigUnit_GetRelativeBoneTransform::StaticStruct();
+							break;
+						}
 						case ERigElementGetterSetterType_Name:
-			{
-				StructTemplate = FRigUnit_BoneName::StaticStruct();
-				break;
-			}
-			default:
-			{
-				break;
-			}
-		}
-	}
-	else
-	{
-		switch (Type)
-		{
+						{
+							StructTemplate = FRigUnit_BoneName::StaticStruct();
+							break;
+						}
+						default:
+						{
+							break;
+						}
+					}
+				}
+				else
+				{
+					switch (Type)
+					{
 						case ERigElementGetterSetterType_Transform:
-			{
-				StructTemplate = FRigUnit_SetBoneTransform::StaticStruct();
-				break;
-			}
+						{
+							StructTemplate = FRigUnit_SetBoneTransform::StaticStruct();
+							NewNode.ValuePinName = TEXT("Transform");
+							NewNode.ValueType = ERigControlType::Transform;
+							NewNode.Value = FRigControlValue::Make<FTransform>(Blueprint->HierarchyContainer.GetLocalTransform(Key));
+							break;
+						}
 						case ERigElementGetterSetterType_Relative:
-			{
-				StructTemplate = FRigUnit_SetRelativeBoneTransform::StaticStruct();
-				break;
-			}
+						{
+							StructTemplate = FRigUnit_SetRelativeBoneTransform::StaticStruct();
+							break;
+						}
 						case ERigElementGetterSetterType_Rotation:
-			{
-				StructTemplate = FRigUnit_SetBoneRotation::StaticStruct();
-				break;
-			}
+						{
+							StructTemplate = FRigUnit_SetBoneRotation::StaticStruct();
+							NewNode.ValuePinName = TEXT("Rotation");
+							NewNode.ValueType = ERigControlType::Rotator;
+							NewNode.Value = FRigControlValue::Make<FRotator>(Blueprint->HierarchyContainer.GetLocalTransform(Key).Rotator());
+							break;
+						}
 						case ERigElementGetterSetterType_Translation:
-			{
-				StructTemplate = FRigUnit_SetBoneTranslation::StaticStruct();
-				break;
-			}
+						{
+							StructTemplate = FRigUnit_SetBoneTranslation::StaticStruct();
+							NewNode.ValuePinName = TEXT("Translation");
+							NewNode.ValueType = ERigControlType::Position;
+							NewNode.Value = FRigControlValue::Make<FVector>(Blueprint->HierarchyContainer.GetLocalTransform(Key).GetLocation());
+							break;
+						}
 						case ERigElementGetterSetterType_Offset:
-			{
-				StructTemplate = FRigUnit_AddBoneTransform::StaticStruct();
-				break;
-			}
-			default:
-			{
-				break;
-			}
-		}
-	}
+						{
+							StructTemplate = FRigUnit_AddBoneTransform::StaticStruct();
+							break;
+						}
+						default:
+						{
+							break;
+						}
+					}
+				}
 				break;
 			}
 			case ERigElementType::Control:
 			{
 				PinName = TEXT("Control");
+
+				const FRigControl& Control = Blueprint->HierarchyContainer.ControlHierarchy[Key.Name];
 				if (bIsGetter)
-	{
-					StructTemplate = FRigUnit_GetControlTransform::StaticStruct();
+				{
+					switch (Control.ControlType)
+					{
+						case ERigControlType::Bool:
+						{
+							StructTemplate = FRigUnit_GetControlBool::StaticStruct();
+							break;
+						}
+						case ERigControlType::Float:
+						{
+							StructTemplate = FRigUnit_GetControlFloat::StaticStruct();
+							break;
+						}
+						case ERigControlType::Vector2D:
+						{
+							StructTemplate = FRigUnit_GetControlVector2D::StaticStruct();
+							break;
+						}
+						case ERigControlType::Position:
+						case ERigControlType::Scale:
+						{
+							StructTemplate = FRigUnit_GetControlVector::StaticStruct();
+							break;
+						}
+						case ERigControlType::Rotator:
+						{
+							StructTemplate = FRigUnit_GetControlRotator::StaticStruct();
+							break;
+						}
+						case ERigControlType::Transform:
+						case ERigControlType::TransformNoScale:
+						{
+							StructTemplate = FRigUnit_GetControlTransform::StaticStruct();
+							break;
+						}
+						default:
+						{
+							break;
+						}
+					}
+				}
+				else
+				{
+					switch (Control.ControlType)
+					{
+						case ERigControlType::Bool:
+						{
+							StructTemplate = FRigUnit_SetControlBool::StaticStruct();
+							break;
+						}
+						case ERigControlType::Float:
+						{
+							StructTemplate = FRigUnit_SetControlFloat::StaticStruct();
+							break;
+						}
+						case ERigControlType::Vector2D:
+						{
+							StructTemplate = FRigUnit_SetControlVector2D::StaticStruct();
+							break;
+						}
+						case ERigControlType::Position:
+						{
+							StructTemplate = FRigUnit_SetControlVector::StaticStruct();
+							NewNode.ValuePinName = TEXT("Vector");
+							NewNode.ValueType = ERigControlType::Position;
+							NewNode.Value = FRigControlValue::Make<FVector>(Blueprint->HierarchyContainer.GetGlobalTransform(Key).GetLocation());
+							break;
+						}
+						case ERigControlType::Scale:
+						{
+							StructTemplate = FRigUnit_SetControlVector::StaticStruct();
+							NewNode.ValuePinName = TEXT("Vector");
+							NewNode.ValueType = ERigControlType::Scale;
+							NewNode.Value = FRigControlValue::Make<FVector>(Blueprint->HierarchyContainer.GetGlobalTransform(Key).GetScale3D());
+							break;
+						}
+						case ERigControlType::Rotator:
+						{
+							StructTemplate = FRigUnit_SetControlRotator::StaticStruct();
+							NewNode.ValuePinName = TEXT("Rotator");
+							NewNode.ValueType = ERigControlType::Rotator;
+							NewNode.Value = FRigControlValue::Make<FRotator>(Blueprint->HierarchyContainer.GetGlobalTransform(Key).Rotator());
+							break;
+						}
+						case ERigControlType::Transform:
+						case ERigControlType::TransformNoScale:
+						{
+							StructTemplate = FRigUnit_SetControlTransform::StaticStruct();
+							NewNode.ValuePinName = TEXT("Transform");
+							NewNode.ValueType = ERigControlType::Transform;
+							NewNode.Value = FRigControlValue::Make<FTransform>(Blueprint->HierarchyContainer.GetGlobalTransform(Key));
+							break;
+						}
+						default:
+						{
+							break;
+						}
+					}
 				}
 				break;
-	}
+			}
 			case ERigElementType::Space:
 			{
 				PinName = TEXT("Space");
 				if (bIsGetter)
-	{
+				{
 					StructTemplate = FRigUnit_GetSpaceTransform::StaticStruct();
 				}
 				break;
-	}
+			}
 			default:
-	{
+			{
 				break;
 			}
-	}
+		}
 
 		if (StructTemplate == nullptr)
 		{
-			continue;
+			return;
 		}
 
 		FVector2D NodePositionIncrement(0.f, 120.f);
@@ -2323,27 +2931,163 @@ void FControlRigEditor::HandleMakeElementGetterSetter(ERigElementGetterSetterTyp
 		}
 
 		FName Name = FControlRigBlueprintUtils::ValidateName(Blueprint, StructTemplate->GetName());
-		if (Blueprint->ModelController->AddNode(StructTemplate->GetFName(), NodePosition, Name))
+		if (URigVMStructNode* ModelNode = Blueprint->Controller->AddStructNode(StructTemplate, TEXT("Execute"), NodePosition))
 		{
-			const FControlRigModelNode* Node = Blueprint->Model->FindNode(Blueprint->LastNameFromNotification);
-			if (Node)
+			NewNode.Name = ModelNode->GetFName();
+			NewNodes.Add(NewNode);
+			Blueprint->Controller->SetPinDefaultValue(FString::Printf(TEXT("%s.%s"), *ModelNode->GetName(), *PinName.ToString()), Key.Name.ToString());
+
+			if (!NewNode.ValuePinName.IsNone())
 			{
-				NewNodeNames.Add(Node->Name);
-				Blueprint->ModelController->SetPinDefaultValueName(Node->Name, PinName, Key.Name, true);
+				FString DefaultValue;
+
+				switch (NewNode.ValueType)
+				{
+					case ERigControlType::Position:
+					case ERigControlType::Scale:
+					{
+						DefaultValue = NewNode.Value.ToString<FVector>();
+						break;
+					}
+					case ERigControlType::Rotator:
+					{
+						DefaultValue = NewNode.Value.ToString<FRotator>();
+						break;
+					}
+					case ERigControlType::Transform:
+					{
+						DefaultValue = NewNode.Value.ToString<FTransform>();
+						break;
+					}
+					default:
+					{
+						break;
+					}
+				}
+				if (!DefaultValue.IsEmpty())
+				{
+					Blueprint->Controller->SetPinDefaultValue(FString::Printf(TEXT("%s.%s"), *ModelNode->GetName(), *NewNode.ValuePinName.ToString()), DefaultValue);
+				}
 			}
+
+			UControlRigUnitNodeSpawner::HookupMutableNode(ModelNode, Blueprint);
 		}
 
 		NodePosition += NodePositionIncrement;
 	}
 
-	if (NewNodeNames.Num() > 0)
+	if (NewNodes.Num() > 0)
 	{
-		Blueprint->ModelController->CloseUndoBracket();
-		Blueprint->ModelController->SetSelection(NewNodeNames);
+		TArray<FName> NewNodeNames;
+		for (const FNewNodeData& NewNode : NewNodes)
+		{
+			NewNodeNames.Add(NewNode.Name);
+		}
+		Blueprint->Controller->SetNodeSelection(NewNodeNames);
+		Blueprint->Controller->CloseUndoBracket();
 	}
 	else
 	{
-		Blueprint->ModelController->CancelUndoBracket();
+		Blueprint->Controller->CancelUndoBracket();
+	}
+}
+
+void FControlRigEditor::HandleOnControlModified(IControlRigManipulatable* Subject, const FRigControl& Control, EControlRigSetKey InSetKey)
+{
+	if (Subject != ControlRig)
+	{
+		return;
+	}
+
+	if (!Control.bIsTransientControl)
+	{
+		return;
+	}
+
+	UControlRigBlueprint* Blueprint = Cast<UControlRigBlueprint>(GetBlueprintObj());
+	if (Blueprint)
+	{
+		if (URigVMPin* Pin = Blueprint->Model->FindPin(Control.Name.ToString()))
+		{
+			FString NewDefaultValue;
+			switch (Control.ControlType)
+			{
+				case ERigControlType::Position:
+				case ERigControlType::Scale:
+				{
+					NewDefaultValue = Control.Value.ToString<FVector>();
+					break;
+				}
+				case ERigControlType::Rotator:
+				{
+					FRotator Rotator = Control.Value.Get<FRotator>();
+					FRigControlValue QuatValue = FRigControlValue::Make<FQuat>(FQuat(Rotator));
+					NewDefaultValue = QuatValue.ToString<FQuat>();
+					break;
+				}
+				case ERigControlType::Transform:
+				{
+					NewDefaultValue = Control.Value.ToString<FTransform>();
+					break;
+				}
+				case ERigControlType::TransformNoScale:
+				{
+					NewDefaultValue = Control.Value.ToString<FTransformNoScale>();
+					break;
+				}
+				default:
+				{
+					break;
+				}
+			}
+
+			if (!NewDefaultValue.IsEmpty())
+			{
+				Blueprint->Controller->SetPinDefaultValue(Pin->GetPinPath(), NewDefaultValue, true, true, true);
+			}
+		}
+		else
+		{
+			static FString ControlRigForElementBoneName;
+			static FString ControlRigForElementSpaceName;
+
+			if (ControlRigForElementBoneName.IsEmpty())
+			{
+				ControlRigForElementBoneName = FString::Printf(TEXT("ControlForRigElement_%s_"),
+					*StaticEnum<ERigElementType>()->GetNameByValue((int64)ERigElementType::Bone).ToString());
+				ControlRigForElementSpaceName = FString::Printf(TEXT("ControlForRigElement_%s_"),
+					*StaticEnum<ERigElementType>()->GetNameByValue((int64)ERigElementType::Space).ToString());
+			}
+
+			if (Control.Name.ToString().StartsWith(ControlRigForElementBoneName))
+			{
+				FName BoneName = *Control.Name.ToString().RightChop(ControlRigForElementBoneName.Len());
+
+				FTransform Transform = Control.Value.Get<FTransform>();
+				Blueprint->HierarchyContainer.BoneHierarchy.SetLocalTransform(BoneName, Transform);
+				Blueprint->PropagateHierarchyFromBPToInstances(false, false);
+
+				if (PreviewInstance)
+				{
+					if (FAnimNode_ModifyBone* Modify = PreviewInstance->FindModifiedBone(BoneName))
+					{
+						Modify->Translation = Transform.GetTranslation();
+						Modify->Rotation = Transform.GetRotation().Rotator();
+						Modify->TranslationSpace = EBoneControlSpace::BCS_ParentBoneSpace;
+						Modify->RotationSpace = EBoneControlSpace::BCS_ParentBoneSpace;
+					}
+				}
+			}
+			else if (Control.Name.ToString().StartsWith(ControlRigForElementSpaceName))
+			{
+				FName SpaceName = *Control.Name.ToString().RightChop(ControlRigForElementSpaceName.Len());
+
+				FTransform GlobalTransform = ControlRig->GetControlGlobalTransform(Control.Name);
+				Blueprint->HierarchyContainer.SpaceHierarchy.SetGlobalTransform(SpaceName, GlobalTransform);
+				Blueprint->HierarchyContainer.SpaceHierarchy.SetInitialGlobalTransform(SpaceName, GlobalTransform);
+				Blueprint->PropagateHierarchyFromBPToInstances(false, false);
+			}
+		}
 	}
 }
 
@@ -2370,64 +3114,58 @@ void FControlRigEditor::UpdateGraphCompilerErrors()
 			// reset all nodes and store them in the map
 			bool bFoundWarning = false;
 			bool bFoundError = false;
-			TMap<FName, UControlRigGraphNode*> UnitNameToNode;
+			TMap<int32, UEdGraphNode*> InstructionIndexToNode;
 			for (UEdGraphNode* GraphNode : Graph->Nodes)
 			{
 				if (UControlRigGraphNode* ControlRigGraphNode = Cast<UControlRigGraphNode>(GraphNode))
 				{
 					bFoundError = bFoundError || ControlRigGraphNode->ErrorType <= (int32)EMessageSeverity::Error;
 					bFoundWarning = bFoundWarning || ControlRigGraphNode->ErrorType <= (int32)EMessageSeverity::Warning;
-
-					if (ControlRigGraphNode->GetUnitProperty())
+					if (URigVMNode* ModelNode = ControlRigGraphNode->GetModelNode())
 					{
-						UnitNameToNode.Add(ControlRigGraphNode->GetUnitProperty()->GetFName(), ControlRigGraphNode);
+						if (ModelNode->GetInstructionIndex() != INDEX_NONE)
+						{
+							InstructionIndexToNode.Add(ModelNode->GetInstructionIndex(), GraphNode);
+						}
 					}
 				}
+				GraphNode->ErrorType = int32(EMessageSeverity::Info) + 1;
 			}
 
 			// update the nodes' error messages
 			bool bFoundErrorOrWarningInLog = false;
 			for (const FControlRigLog::FLogEntry& Entry : ControlRigLog.Entries)
 			{
-				UControlRigGraphNode** RigNodePtr = UnitNameToNode.Find(Entry.Unit);
-				if (RigNodePtr == nullptr)
+				UEdGraphNode** GraphNodePtr = InstructionIndexToNode.Find(Entry.InstructionIndex);
+				if (GraphNodePtr == nullptr)
 				{
 					continue;
 				}
-				UControlRigGraphNode* RigNode = *RigNodePtr;
+				UEdGraphNode* GraphNode = *GraphNodePtr;
 
 				bFoundError = bFoundError || Entry.Severity <= EMessageSeverity::Error;
 				bFoundWarning = bFoundWarning || Entry.Severity <= EMessageSeverity::Warning;
 				bFoundErrorOrWarningInLog = bFoundErrorOrWarningInLog || Entry.Severity <= EMessageSeverity::Warning;
 
 				int32 ErrorType = (int32)Entry.Severity;
-				if (RigNode->ErrorType < ErrorType)
+				if (GraphNode->ErrorType < ErrorType)
 				{
 					continue;
 				}
-				else if (RigNode->ErrorType == ErrorType)
+				else if (GraphNode->ErrorType == ErrorType)
 				{
-					RigNode->ErrorMsg = FString::Printf(TEXT("%s\n%s"), *RigNode->ErrorMsg, *Entry.Message);
+					GraphNode->ErrorMsg = FString::Printf(TEXT("%s\n%s"), *GraphNode->ErrorMsg, *Entry.Message);
 				}
 				else
 				{
-					RigNode->ErrorMsg = Entry.Message;
-					RigNode->ErrorType = ErrorType;
+					GraphNode->ErrorMsg = Entry.Message;
+					GraphNode->ErrorType = ErrorType;
 				}
 			}
 
 			for (UEdGraphNode* GraphNode : Graph->Nodes)
 			{
-				if (UControlRigGraphNode* ControlRigGraphNode = Cast<UControlRigGraphNode>(GraphNode))
-				{
-					bool PreviouslyHadError = ControlRigGraphNode->bHasCompilerMessage;
-					bool CurrentlyHasError = ControlRigGraphNode->ErrorType <= int32(EMessageSeverity::Info);
-					if (CurrentlyHasError != PreviouslyHadError)
-					{
-						ControlRigGraphNode->bHasCompilerMessage = CurrentlyHasError;
-						ControlRigGraphNode->Modify();
-					}
-				}
+				GraphNode->bHasCompilerMessage = GraphNode->ErrorType <= int32(EMessageSeverity::Info);
 			}
 
 			if (bFoundError)
@@ -2441,7 +3179,7 @@ void FControlRigEditor::UpdateGraphCompilerErrors()
 					Info.bUseSuccessFailIcons = true;
 					Info.Image = FEditorStyle::GetBrush(TEXT("MessageLog.Error"));
 					Info.bFireAndForget = true;
-					Info.FadeOutDuration = 10.0f;
+					Info.FadeOutDuration = 5.0f;
 					Info.ExpireDuration = 0.0f;
 					TSharedPtr<SNotificationItem> NotificationPtr = FSlateNotificationManager::Get().AddNotification(Info);
 					NotificationPtr->SetCompletionState(SNotificationItem::CS_Success);
@@ -2455,7 +3193,7 @@ void FControlRigEditor::UpdateGraphCompilerErrors()
 					Info.bUseSuccessFailIcons = true;
 					Info.Image = FEditorStyle::GetBrush(TEXT("MessageLog.Warning"));
 					Info.bFireAndForget = true;
-					Info.FadeOutDuration = 10.0f;
+					Info.FadeOutDuration = 5.0f;
 					Info.ExpireDuration = 0.0f;
 					TSharedPtr<SNotificationItem> NotificationPtr = FSlateNotificationManager::Get().AddNotification(Info);
 					NotificationPtr->SetCompletionState(SNotificationItem::CS_Success);
