COPY from //Dev-Anim to //Dev-Main

diff --git a/Engine/Plugins/Experimental/ControlRig/Source/ControlRigEditor/Private/Editor/ControlRigEditor.h b/Engine/Plugins/Experimental/ControlRig/Source/ControlRigEditor/Private/Editor/ControlRigEditor.h
--- a/Engine/Plugins/Experimental/ControlRig/Source/ControlRigEditor/Private/Editor/ControlRigEditor.h
+++ b/Engine/Plugins/Experimental/ControlRig/Source/ControlRigEditor/Private/Editor/ControlRigEditor.h
@@ -7,10 +7,16 @@
 #include "AssetEditorModeManager.h"
 #include "DragAndDrop/GraphNodeDragDropOp.h"
 #include "ControlRigDefines.h"
+#include "Units/RigUnitContext.h"
 #include "ControlRigLog.h"
-#include "ControlRigModel.h"
 #include "IPersonaViewport.h"
 #include "Framework/MultiBox/MultiBoxBuilder.h"
+#include "RigVMModel/RigVMGraph.h"
+#include "RigVMCore/RigVM.h"
+#include "Graph/SControlRigGraphPinNameListValueWidget.h"
+#include "Styling/SlateTypes.h"
+#include "AnimPreviewInstance.h"
+#include "ScopedTransaction.h"
 
 class UControlRigBlueprint;
 class IPersonaToolkit;
@@ -73,9 +79,28 @@ public:
 	}
 
 	// BlueprintEditor interface
+	virtual bool TransactionObjectAffectsBlueprint(UObject* InTransactedObject) override;
 	virtual bool CanAddNewLocalVariable() const override { return false; }
-	virtual void DeleteSelectedNodes();
-	virtual void PasteNodesHere(class UEdGraph* DestinationGraph, const FVector2D& GraphLocation) override;
+	virtual void DeleteSelectedNodes() override;
+	virtual bool CanDeleteNodes() const override;
+
+	virtual void CopySelectedNodes() override;
+	virtual bool CanCopyNodes() const override;
+
+	virtual void PasteNodes() override;
+	virtual bool CanPasteNodes() const override;
+
+	// FEditorUndoClient Interface
+	virtual void PostUndo(bool bSuccess) override;
+	virtual void PostRedo(bool bSuccess) override;
+
+	void EnsureValidRigElementInDetailPanel();
+
+	virtual void OnStartWatchingPin();
+	virtual bool CanStartWatchingPin() const;
+
+	virtual void OnStopWatchingPin();
+	virtual bool CanStopWatchingPin() const;
 
 	// IToolkitHost Interface
 	virtual void OnToolkitHostingStarted(const TSharedRef<class IToolkit>& Toolkit) override;
@@ -123,6 +148,9 @@ public:
 	FNewMenuDelegate& OnViewportContextMenu() { return OnViewportContextMenuDelegate; }
 	FNewMenuCommandsDelegate& OnViewportContextMenuCommands() { return OnViewportContextMenuCommandsDelegate; }
 
+	DECLARE_EVENT_OneParam(FControlRigEditor, FPreviewControlRigUpdated, FControlRigEditor*);
+	FPreviewControlRigUpdated& OnPreviewControlRigUpdated() { return PreviewControlRigUpdated;  }
+
 protected:
 
 	void OnHierarchyChanged();
@@ -144,14 +172,9 @@ protected:
 	virtual void OnBlueprintChangedImpl(UBlueprint* InBlueprint, bool bIsJustBeingCompiled) override;
 	virtual void SetupGraphEditorEvents(UEdGraph* InGraph, SGraphEditor::FGraphEditorEvents& InEvents) override;
 
-	// FEditorUndoClient Interface
-	virtual void PostUndo(bool bSuccess) override;
-	virtual void PostRedo(bool bSuccess) override;
-
-	// FNotifyHook Interface
-	virtual void NotifyPostChange(const FPropertyChangedEvent& PropertyChangedEvent, FProperty* PropertyThatChanged) override;
-
-	void HandleModelModified(const UControlRigModel* InModel, EControlRigModelNotifType InType, const void* InPayload);
+	void HandleModifiedEvent(ERigVMGraphNotifType InNotifType, URigVMGraph* InGraph, UObject* InSubject);
+	void HandleVMCompiledEvent(UBlueprint* InBlueprint, URigVM* InVM);
+	void HandleControlRigInitializedEvent(UControlRig* InControlRig, const EControlRigState InState);
 
 	// FGCObject Interface
 	virtual void AddReferencedObjects( FReferenceCollector& Collector ) override;
@@ -190,7 +213,7 @@ private:
 	TOptional<float> GetToolbarAxesScale() const;
 	void OnToolbarAxesScaleChanged(float InValue);
 
-		/** Handle switching skeletal meshes */
+	/** Handle switching skeletal meshes */
 	void HandlePreviewMeshChanged(USkeletalMesh* InOldSkeletalMesh, USkeletalMesh* InNewSkeletalMesh);
 
 	/** Push a newly compiled/opened control rig to the edit mode */
@@ -202,12 +225,17 @@ private:
 	/** Rebind our anim instance to the preview's skeletal mesh component */
 	void RebindToSkeletalMeshComponent();
 
+	/** Update stale watch pins */
+	void UpdateStaleWatchedPins();
+
 	/** Wraps the normal blueprint editor's action menu creation callback */
 	FActionMenuContent HandleCreateGraphActionMenu(UEdGraph* InGraph, const FVector2D& InNodePosition, const TArray<UEdGraphPin*>& InDraggedPins, bool bAutoExpand, SGraphEditor::FActionMenuClosed InOnMenuClosed);
 	void OnNodeTitleCommitted(const FText& NewText, ETextCommit::Type CommitInfo, UEdGraphNode* NodeBeingChanged);
 
 	void ToggleExecuteGraph();
 	bool IsExecuteGraphOn() const;
+	void ToggleAutoCompileGraph();
+	bool IsAutoCompileGraphOn() const;
 
 	enum ERigElementGetterSetterType
 	{
@@ -222,6 +250,8 @@ private:
 
 	 void HandleMakeElementGetterSetter(ERigElementGetterSetterType Type, bool bIsGetter, TArray<FRigElementKey> Keys, UEdGraph* Graph, FVector2D NodePosition);
 
+	 void HandleOnControlModified(IControlRigManipulatable* Subject, const FRigControl& Control, EControlRigSetKey InSetKey);
+
 protected:
 
 	/** Toolbox hosting widget */
@@ -245,6 +275,9 @@ protected:
 	/** preview scene */
 	TSharedPtr<IPersonaPreviewScene> PreviewScene;
 
+	/** preview animation instance */
+	UAnimPreviewInstance* PreviewInstance;
+
 	/** Delegate to deal with key down evens in the viewport / editor */
 	FPersonaViewportKeyDownDelegate OnKeyDownDelegate;
 
@@ -258,14 +291,20 @@ protected:
 	FTransform GetRigElementTransform(const FRigElementKey& InElement, bool bLocal) const;
 	void SetRigElementTransform(const FRigElementKey& InElement, const FTransform& InTransform, bool bLocal);
 	
+	// FNotifyHook Interface
+	virtual void NotifyPreChange(FProperty* PropertyAboutToChange) override;
+	virtual void NotifyPostChange(const FPropertyChangedEvent& PropertyChangedEvent, FProperty* PropertyThatChanged) override;
 	/** delegate for changing property */
 	virtual void OnFinishedChangingProperties(const FPropertyChangedEvent& PropertyChangedEvent) override;
 
 	bool bControlRigEditorInitialized;
-	bool bIsSelecting;
 	bool bIsSettingObjectBeingDebugged;
 	FRigElementKey RigElementInDetailPanel;
 
+	TArray<uint8, TAlignedHeapAllocator<16>> NodeDetailBuffer;
+	UScriptStruct* NodeDetailStruct;
+	FName NodeDetailName;
+
 	/** Currently executing ControlRig or not - later maybe this will change to enum for whatever different mode*/
 	bool bExecutionControlRig;
 
@@ -277,6 +316,19 @@ protected:
 	/** This can be used to enable dumping of a unit test */
 	void DumpUnitTestCode();
 
+	void OnAnimInitialized();
+
+	FPreviewControlRigUpdated PreviewControlRigUpdated;
+
+	TSharedPtr<SControlRigGraphPinNameListValueWidget> PinControlNameList;
+	bool IsPinControlNameListEnabled() const;
+	TSharedRef<SWidget> MakePinControlNameListItemWidget(TSharedPtr<FString> InItem);
+	FText GetPinControlNameListText() const;
+	TSharedPtr<FString> GetPinControlCurrentlySelectedItem(const TArray<TSharedPtr<FString>>* InNameList) const;
+	void SetPinControlNameListText(const FText& NewTypeInValue, ETextCommit::Type /*CommitInfo*/);
+	void OnPinControlNameListChanged(TSharedPtr<FString> NewSelection, ESelectInfo::Type SelectInfo);
+	void OnPinControlNameListComboBox(const TArray<TSharedPtr<FString>>* InNameList);
+
 	friend class FControlRigEditorMode;
 	friend class SControlRigStackView;
 	friend class SRigHierarchy;
diff --git a/Engine/Plugins/Experimental/ControlRig/Source/ControlRigEditor/Private/Editor/ControlRigSkeletalMeshComponent.cpp b/Engine/Plugins/Experimental/ControlRig/Source/ControlRigEditor/Private/Editor/ControlRigSkeletalMeshComponent.cpp
--- a/Engine/Plugins/Experimental/ControlRig/Source/ControlRigEditor/Private/Editor/ControlRigSkeletalMeshComponent.cpp
+++ b/Engine/Plugins/Experimental/ControlRig/Source/ControlRigEditor/Private/Editor/ControlRigSkeletalMeshComponent.cpp
@@ -1,9 +1,10 @@
 // Copyright Epic Games, Inc. All Rights Reserved.
 
 #include "ControlRigSkeletalMeshComponent.h"
-#include "Sequencer/ControlRigSequencerAnimInstance.h"
+#include "Sequencer/ControlRigLayerInstance.h" 
 #include "SkeletalDebugRendering.h"
 #include "ControlRig.h"
+#include "AnimPreviewInstance.h"
 
 UControlRigSkeletalMeshComponent::UControlRigSkeletalMeshComponent(const FObjectInitializer& ObjectInitializer)
 	: Super(ObjectInitializer)
@@ -15,11 +16,21 @@ UControlRigSkeletalMeshComponent::UControlRigSkeletalMeshComponent(const FObject
 void UControlRigSkeletalMeshComponent::InitAnim(bool bForceReinit)
 {
 	// skip preview init entirely, just init the super class
-	USkeletalMeshComponent::InitAnim(bForceReinit);
+	Super::InitAnim(bForceReinit);
+
+	UControlRigLayerInstance* ControlRigInstance = Cast<UControlRigLayerInstance>(GetAnimInstance());
+	if (ControlRigInstance)
+	{
+		ControlRigInstance->SetSourceAnimInstance(PreviewInstance);
+	}
 
 	RebuildDebugDrawSkeleton();
 }
 
+bool UControlRigSkeletalMeshComponent::IsPreviewOn() const
+{
+	return (PreviewInstance != nullptr);
+}
 
 void UControlRigSkeletalMeshComponent::SetCustomDefaultPose()
 {
@@ -28,54 +39,59 @@ void UControlRigSkeletalMeshComponent::SetCustomDefaultPose()
 
 void UControlRigSkeletalMeshComponent::RebuildDebugDrawSkeleton()
 {
-	UControlRigSequencerAnimInstance* ControlRigInstance = Cast<UControlRigSequencerAnimInstance>(GetAnimInstance());
+	UControlRigLayerInstance* ControlRigInstance = Cast<UControlRigLayerInstance>(GetAnimInstance());
 
-	if (ControlRigInstance && ControlRigInstance->CachedControlRig.IsValid())
+	if (ControlRigInstance)
 	{
-		UControlRig* ControlRig = ControlRigInstance->CachedControlRig.Get();
-		// just copy it because this is not thread safe
-		const FRigBoneHierarchy& BaseHiearchy = ControlRig->GetBoneHierarchy();
-
-		DebugDrawSkeleton.Empty();
-		DebugDrawBones.Reset();
-
-		// create ref modifier
- 		FReferenceSkeletonModifier RefSkelModifier(DebugDrawSkeleton, nullptr);
- 
- 		for (int32 Index = 0; Index < BaseHiearchy.Num(); Index++)
- 		{
-			FMeshBoneInfo NewMeshBoneInfo;
-			NewMeshBoneInfo.Name = BaseHiearchy.GetName(Index);
-			NewMeshBoneInfo.ParentIndex = BaseHiearchy[Index].ParentIndex;
-			// give ref pose here
-			RefSkelModifier.Add(NewMeshBoneInfo, BaseHiearchy.GetInitialTransform(Index));
-
-			DebugDrawBones.Add(Index);
+		UControlRig* ControlRig = ControlRigInstance->GetFirstAvailableControlRig();
+		if (ControlRig)
+		{
+			// just copy it because this is not thread safe
+			const FRigBoneHierarchy& BaseHiearchy = ControlRig->GetBoneHierarchy();
+
+			DebugDrawSkeleton.Empty();
+			DebugDrawBones.Reset();
+
+			// create ref modifier
+			FReferenceSkeletonModifier RefSkelModifier(DebugDrawSkeleton, nullptr);
+
+			for (int32 Index = 0; Index < BaseHiearchy.Num(); Index++)
+			{
+				FMeshBoneInfo NewMeshBoneInfo;
+				NewMeshBoneInfo.Name = BaseHiearchy.GetName(Index);
+				NewMeshBoneInfo.ParentIndex = BaseHiearchy[Index].ParentIndex;
+				// give ref pose here
+				RefSkelModifier.Add(NewMeshBoneInfo, BaseHiearchy.GetInitialTransform(Index));
+
+				DebugDrawBones.Add(Index);
+			}
 		}
 	}
 }
 
 FTransform UControlRigSkeletalMeshComponent::GetDrawTransform(int32 BoneIndex) const
 {
-	UControlRigSequencerAnimInstance* ControlRigInstance = Cast<UControlRigSequencerAnimInstance>(GetAnimInstance());
+	UControlRigLayerInstance* ControlRigInstance = Cast<UControlRigLayerInstance>(GetAnimInstance());
 
-	if (ControlRigInstance && ControlRigInstance->CachedControlRig.IsValid())
+	if (ControlRigInstance)
 	{
-		UControlRig* ControlRig = ControlRigInstance->CachedControlRig.Get();
-		// just copy it because this is not thread safe
-		const FRigBoneHierarchy& BaseHiearchy = ControlRig->GetBoneHierarchy();
-		return BaseHiearchy.GetGlobalTransform(BoneIndex);
+		UControlRig* ControlRig = ControlRigInstance->GetFirstAvailableControlRig();
+		if (ControlRig)
+		{
+			// just copy it because this is not thread safe
+			const FRigBoneHierarchy& BaseHiearchy = ControlRig->GetBoneHierarchy();
+			return BaseHiearchy.GetGlobalTransform(BoneIndex);
+		}
 	}
 
 	return FTransform::Identity;
 }
 
+
 void UControlRigSkeletalMeshComponent::EnablePreview(bool bEnable, UAnimationAsset* PreviewAsset)
 {
-	UControlRigSequencerAnimInstance* ControlRigInstance = Cast<UControlRigSequencerAnimInstance>(GetAnimInstance());
-
-	if (ControlRigInstance && ControlRigInstance->CachedControlRig.IsValid())
+	if (PreviewInstance)
 	{
-		ControlRigInstance->SetAnimationAsset(PreviewAsset);
+		PreviewInstance->SetAnimationAsset(PreviewAsset);
 	}
 }
diff --git a/Engine/Plugins/Experimental/ControlRig/Source/ControlRigEditor/Private/Editor/ControlRigSkeletalMeshComponent.h b/Engine/Plugins/Experimental/ControlRig/Source/ControlRigEditor/Private/Editor/ControlRigSkeletalMeshComponent.h
--- a/Engine/Plugins/Experimental/ControlRig/Source/ControlRigEditor/Private/Editor/ControlRigSkeletalMeshComponent.h
+++ b/Engine/Plugins/Experimental/ControlRig/Source/ControlRigEditor/Private/Editor/ControlRigSkeletalMeshComponent.h
@@ -34,6 +34,8 @@ class UControlRigSkeletalMeshComponent : public UDebugSkelMeshComponent
 
 	virtual void EnablePreview(bool bEnable, class UAnimationAsset * PreviewAsset) override;
 
+	// return true if preview animation is active 
+	virtual bool IsPreviewOn() const override;
 	// END UDebugSkeletalMeshComponent interface
 
 public:
diff --git a/Engine/Plugins/Experimental/ControlRig/Source/ControlRigEditor/Private/Editor/SControlRigGizmoNameList.cpp b/Engine/Plugins/Experimental/ControlRig/Source/ControlRigEditor/Private/Editor/SControlRigGizmoNameList.cpp
--- a/Engine/Plugins/Experimental/ControlRig/Source/ControlRigEditor/Private/Editor/SControlRigGizmoNameList.cpp
+++ b/Engine/Plugins/Experimental/ControlRig/Source/ControlRigEditor/Private/Editor/SControlRigGizmoNameList.cpp
@@ -77,6 +77,7 @@ void SControlRigGizmoNameList::SetNameListText(const FText& NewTypeInValue, ETex
 			const FScopedTransaction Transaction(NSLOCTEXT("ControlRigEditor", "ChangeGizmoName", "Change Gizmo Name"));
 			Blueprint->Modify();
 			Control.GizmoName = NewName;
+			Blueprint->PropagatePropertyFromBPToInstances(ControlKey, FRigControl::StaticStruct()->FindPropertyByName(TEXT("GizmoName")));
 			Blueprint->HierarchyContainer.ControlHierarchy.OnControlUISettingsChanged.Broadcast(&Blueprint->HierarchyContainer, ControlKey);
 		}
 	}
diff --git a/Engine/Plugins/Experimental/ControlRig/Source/ControlRigEditor/Private/Editor/SControlRigProfilingView.cpp b/Engine/Plugins/Experimental/ControlRig/Source/ControlRigEditor/Private/Editor/SControlRigProfilingView.cpp
--- a/Engine/Plugins/Experimental/ControlRig/Source/ControlRigEditor/Private/Editor/SControlRigProfilingView.cpp
+++ b/Engine/Plugins/Experimental/ControlRig/Source/ControlRigEditor/Private/Editor/SControlRigProfilingView.cpp
@@ -220,7 +220,7 @@ TSharedRef<SWidget> SControlRigProfilingItem::GenerateWidgetForColumn(const FNam
 		[
 			SNew(STextBlock)
 			.Text(this, &SControlRigProfilingItem::GetLabelText)
-			.ToolTipText(this, &SControlRigProfilingItem::GetTooltipText)
+			.ToolTipText(this, &SControlRigProfilingItem::GetToolTipText)
 			.ColorAndOpacity(this, &SControlRigProfilingItem::GetTextColor)
 			.OnDoubleClicked(this, &SControlRigProfilingItem::OnDoubleClicked)
 		];
@@ -305,7 +305,7 @@ FText SControlRigProfilingItem::GetLabelText() const
 	return (FText::FromName(Name));
 }
 
-FText SControlRigProfilingItem::GetTooltipText() const
+FText SControlRigProfilingItem::GetToolTipText() const
 {
 	FString Path = WeakTreeElement.Pin()->GetPath();
 	
diff --git a/Engine/Plugins/Experimental/ControlRig/Source/ControlRigEditor/Private/Editor/SControlRigProfilingView.h b/Engine/Plugins/Experimental/ControlRig/Source/ControlRigEditor/Private/Editor/SControlRigProfilingView.h
--- a/Engine/Plugins/Experimental/ControlRig/Source/ControlRigEditor/Private/Editor/SControlRigProfilingView.h
+++ b/Engine/Plugins/Experimental/ControlRig/Source/ControlRigEditor/Private/Editor/SControlRigProfilingView.h
@@ -34,7 +34,7 @@ private:
 	TWeakPtr<FStatsTreeElement> WeakTreeElement;
 
 	FText GetLabelText() const;
-	FText GetTooltipText() const;
+	FText GetToolTipText() const;
 	FSlateColor GetTextColor() const;
 	FText GetTotalTimeText(bool bInclusive) const;
 	FText GetAverageTimeText(bool bInclusive) const;
diff --git a/Engine/Plugins/Experimental/ControlRig/Source/ControlRigEditor/Private/Editor/SControlRigStackView.cpp b/Engine/Plugins/Experimental/ControlRig/Source/ControlRigEditor/Private/Editor/SControlRigStackView.cpp
--- a/Engine/Plugins/Experimental/ControlRig/Source/ControlRigEditor/Private/Editor/SControlRigStackView.cpp
+++ b/Engine/Plugins/Experimental/ControlRig/Source/ControlRigEditor/Private/Editor/SControlRigStackView.cpp
@@ -17,12 +17,11 @@
 //////////////////////////////////////////////////////////////
 /// FRigStackEntry
 ///////////////////////////////////////////////////////////
-FRigStackEntry::FRigStackEntry(int32 InEntryIndex, ERigStackEntry::Type InEntryType, int32 InOpIndex, EControlRigOpCode InOpCode, const FName& InName, const FString& InLabel)
+FRigStackEntry::FRigStackEntry(int32 InEntryIndex, ERigStackEntry::Type InEntryType, int32 InInstructionIndex, ERigVMOpCode InOpCode, const FString& InLabel)
 	: EntryIndex(InEntryIndex)
 	, EntryType(InEntryType)
-	, OpIndex(InOpIndex)
+	, InstructionIndex(InInstructionIndex)
 	, OpCode(InOpCode)
-	, Name(InName)
 	, Label(InLabel)
 {
 
@@ -128,14 +127,15 @@ SControlRigStackView::~SControlRigStackView()
 		{
 			ControlRigEditor.Pin()->GetControlRigBlueprint()->OnModified().Remove(OnModelModified);
 		}
-		if (OnControlRigInitializedHandle.IsValid() && ControlRigEditor.Pin()->ControlRig)
+		if (OnPreviewControlRigUpdatedHandle.IsValid())
 		{
-			ControlRigEditor.Pin()->ControlRig->OnInitialized().Remove(OnControlRigInitializedHandle);
+			ControlRigEditor.Pin()->ControlRig->OnInitialized_AnyThread().Remove(OnControlRigInitializedHandle);
+			ControlRigEditor.Pin()->OnPreviewControlRigUpdated().Remove(OnPreviewControlRigUpdatedHandle);
 		}
 	}
-	if (ControlRigBlueprint.IsValid() && OnBlueprintCompiledHandle.IsValid())
+	if (ControlRigBlueprint.IsValid() && OnVMCompiledHandle.IsValid())
 	{
-		ControlRigBlueprint->OnCompiled().Remove(OnBlueprintCompiledHandle);
+		ControlRigBlueprint->OnVMCompiled().Remove(OnVMCompiledHandle);
 	}
 }
 
@@ -146,6 +146,7 @@ void SControlRigStackView::Construct( const FArguments& InArgs, TSharedRef<FCont
 	Graph = Cast<UControlRigGraph>(ControlRigBlueprint->GetLastEditedUberGraph());
 	CommandList = MakeShared<FUICommandList>();
 	bSuspendModelNotifications = false;
+	bSuspendControllerSelection = false;
 
 	BindCommands();
 
@@ -171,22 +172,39 @@ void SControlRigStackView::Construct( const FArguments& InArgs, TSharedRef<FCont
 			]
 		];
 
-	RefreshTreeView();
+	RefreshTreeView(nullptr);
 
 	if (ControlRigBlueprint.IsValid())
 	{
-		if(OnBlueprintCompiledHandle.IsValid())
+		if (OnVMCompiledHandle.IsValid())
 		{
-			ControlRigBlueprint->OnCompiled().Remove(OnBlueprintCompiledHandle);
+			ControlRigBlueprint->OnVMCompiled().Remove(OnVMCompiledHandle);
+		}
+		if (OnModelModified.IsValid())
+		{
+			ControlRigBlueprint->OnModified().Remove(OnModelModified);
+		}
+		OnVMCompiledHandle = ControlRigBlueprint->OnVMCompiled().AddSP(this, &SControlRigStackView::OnVMCompiled);
+		OnModelModified = ControlRigBlueprint->OnModified().AddSP(this, &SControlRigStackView::HandleModifiedEvent);
+
+		if (ControlRigEditor.IsValid())
+		{
+			if (UControlRig* ControlRig = ControlRigEditor.Pin()->ControlRig)
+			{
+				OnVMCompiled(ControlRigBlueprint.Get(), ControlRig->VM);
+			}
 		}
-		OnBlueprintCompiledHandle = ControlRigBlueprint->OnCompiled().AddSP(this, &SControlRigStackView::OnBlueprintCompiled);
-		OnModelModified = ControlRigBlueprint->OnModified().AddSP(this, &SControlRigStackView::HandleModelModified);
+	}
+
+	if (ControlRigEditor.IsValid())
+	{
+		OnPreviewControlRigUpdatedHandle = ControlRigEditor.Pin()->OnPreviewControlRigUpdated().AddSP(this, &SControlRigStackView::HandlePreviewControlRigUpdated);
 	}
 }
 
 void SControlRigStackView::OnSelectionChanged(TSharedPtr<FRigStackEntry> Selection, ESelectInfo::Type SelectInfo)
 {
-	if (bSuspendModelNotifications)
+	if (bSuspendModelNotifications || bSuspendControllerSelection)
 	{
 		return;
 	}
@@ -206,23 +224,26 @@ void SControlRigStackView::OnSelectionChanged(TSharedPtr<FRigStackEntry> Selecti
 			return;
 		}
 
-		TArray<FName> SelectedNodes;
-		for (TSharedPtr<FRigStackEntry>& Entry : SelectedItems)
+		TMap<int32, FName> InstructionIndexToNodeName;
+		for (URigVMNode* ModelNode : ControlRigBlueprint->Model->GetNodes())
 		{
-			if (Entry->OpIndex >= GeneratedClass->Operators.Num())
+			if (ModelNode->GetInstructionIndex() != INDEX_NONE)
 			{
-				return;
+				InstructionIndexToNodeName.Add(ModelNode->GetInstructionIndex(), ModelNode->GetFName());
 			}
+		}
 
-			FControlRigOperator& Operator = GeneratedClass->Operators[Entry->OpIndex];
-			if (Operator.OpCode == EControlRigOpCode::Exec)
+		TArray<FName> SelectedNodes;
+		for (TSharedPtr<FRigStackEntry>& Entry : SelectedItems)
+		{
+			const FName* NodeName = InstructionIndexToNodeName.Find(Entry->InstructionIndex);
+			if (NodeName)
 			{
-				FString PropertyPath = Operator.CachedPropertyPath1.ToString();
-				SelectedNodes.Add(*PropertyPath);
+				SelectedNodes.Add(*NodeName);
 			}
 		}
 
-		ControlRigBlueprint->ModelController->SetSelection(SelectedNodes);
+		ControlRigBlueprint->Controller->SetNodeSelection(SelectedNodes);
 	}
 }
 
@@ -243,108 +264,138 @@ void SControlRigStackView::HandleGetChildrenForTree(TSharedPtr<FRigStackEntry> I
 	OutChildren = InItem->Children;
 }
 
-void SControlRigStackView::RefreshTreeView(UControlRig* ControlRig)
+void SControlRigStackView::RefreshTreeView(URigVM* InVM)
 {
 	Operators.Reset();
 
-	if (ControlRigBlueprint.IsValid())
+	if (InVM)
 	{
-		UControlRigBlueprintGeneratedClass* GeneratedClass = ControlRigBlueprint->GetControlRigBlueprintGeneratedClass();
-		if (GeneratedClass)
-		{
-			TMap<FName, int32> UnitToOperatorIndex;
+		FRigVMInstructionArray Instructions = InVM->GetInstructions();
 
-			for (int32 OperatorIndex=0;OperatorIndex< GeneratedClass->Operators.Num();OperatorIndex++)
+		for (int32 InstructionIndex=0; InstructionIndex < Instructions.Num(); InstructionIndex++)
+		{
+			FRigVMInstruction Instruction = Instructions[InstructionIndex];
+			switch (Instruction.OpCode)
 			{
-				FControlRigOperator Operator = GeneratedClass->Operators[OperatorIndex];
-				switch (Operator.OpCode)
+				case ERigVMOpCode::Execute_0_Operands:
+				case ERigVMOpCode::Execute_1_Operands:
+				case ERigVMOpCode::Execute_2_Operands:
+				case ERigVMOpCode::Execute_3_Operands:
+				case ERigVMOpCode::Execute_4_Operands:
+				case ERigVMOpCode::Execute_5_Operands:
+				case ERigVMOpCode::Execute_6_Operands:
+				case ERigVMOpCode::Execute_7_Operands:
+				case ERigVMOpCode::Execute_8_Operands:
+				case ERigVMOpCode::Execute_9_Operands:
+				case ERigVMOpCode::Execute_10_Operands:
+				case ERigVMOpCode::Execute_11_Operands:
+				case ERigVMOpCode::Execute_12_Operands:
+				case ERigVMOpCode::Execute_13_Operands:
+				case ERigVMOpCode::Execute_14_Operands:
+				case ERigVMOpCode::Execute_15_Operands:
+				case ERigVMOpCode::Execute_16_Operands:
+				case ERigVMOpCode::Execute_17_Operands:
+				case ERigVMOpCode::Execute_18_Operands:
+				case ERigVMOpCode::Execute_19_Operands:
+				case ERigVMOpCode::Execute_20_Operands:
+				case ERigVMOpCode::Execute_21_Operands:
+				case ERigVMOpCode::Execute_22_Operands:
+				case ERigVMOpCode::Execute_23_Operands:
+				case ERigVMOpCode::Execute_24_Operands:
+				case ERigVMOpCode::Execute_25_Operands:
+				case ERigVMOpCode::Execute_26_Operands:
+				case ERigVMOpCode::Execute_27_Operands:
+				case ERigVMOpCode::Execute_28_Operands:
+				case ERigVMOpCode::Execute_29_Operands:
+				case ERigVMOpCode::Execute_30_Operands:
+				case ERigVMOpCode::Execute_31_Operands:
+				case ERigVMOpCode::Execute_32_Operands:
+				case ERigVMOpCode::Execute_33_Operands:
+				case ERigVMOpCode::Execute_34_Operands:
+				case ERigVMOpCode::Execute_35_Operands:
+				case ERigVMOpCode::Execute_36_Operands:
+				case ERigVMOpCode::Execute_37_Operands:
+				case ERigVMOpCode::Execute_38_Operands:
+				case ERigVMOpCode::Execute_39_Operands:
+				case ERigVMOpCode::Execute_40_Operands:
+				case ERigVMOpCode::Execute_41_Operands:
+				case ERigVMOpCode::Execute_42_Operands:
+				case ERigVMOpCode::Execute_43_Operands:
+				case ERigVMOpCode::Execute_44_Operands:
+				case ERigVMOpCode::Execute_45_Operands:
+				case ERigVMOpCode::Execute_46_Operands:
+				case ERigVMOpCode::Execute_47_Operands:
+				case ERigVMOpCode::Execute_48_Operands:
+				case ERigVMOpCode::Execute_49_Operands:
+				case ERigVMOpCode::Execute_50_Operands:
+				case ERigVMOpCode::Execute_51_Operands:
+				case ERigVMOpCode::Execute_52_Operands:
+				case ERigVMOpCode::Execute_53_Operands:
+				case ERigVMOpCode::Execute_54_Operands:
+				case ERigVMOpCode::Execute_55_Operands:
+				case ERigVMOpCode::Execute_56_Operands:
+				case ERigVMOpCode::Execute_57_Operands:
+				case ERigVMOpCode::Execute_58_Operands:
+				case ERigVMOpCode::Execute_59_Operands:
+				case ERigVMOpCode::Execute_60_Operands:
+				case ERigVMOpCode::Execute_61_Operands:
+				case ERigVMOpCode::Execute_62_Operands:
+				case ERigVMOpCode::Execute_63_Operands:
+				case ERigVMOpCode::Execute_64_Operands:
 				{
-					case EControlRigOpCode::Exec:
-					{
-						FName UnitPath = *Operator.CachedPropertyPath1.ToString();
-						FString OperatorLabel;
-
-						FStructProperty * StructProperty = CastField<FStructProperty>(GeneratedClass->FindPropertyByName(UnitPath));
-						if (StructProperty)
-						{
-							if (StructProperty->Struct->IsChildOf(FRigUnit::StaticStruct()))
-							{
-								StructProperty->Struct->GetStringMetaDataHierarchical(UControlRig::DisplayNameMetaName, &OperatorLabel);
-								if (OperatorLabel.IsEmpty())
-								{
-									OperatorLabel = FName::NameToDisplayString(StructProperty->Struct->GetFName().ToString(), false);
-								}
-
-								if (ControlRig)
-								{
-									const FRigUnit* UnitPtr = StructProperty->ContainerPtrToValuePtr<FRigUnit>(ControlRig);
-									if (UnitPtr)
-									{
-										FString UnitLabel = UnitPtr->GetUnitLabel();
-										if (!UnitLabel.IsEmpty())
-										{
-											OperatorLabel = UnitLabel;
-										}
-									}
-								}
-							}
-						}
-
-						if (OperatorLabel.IsEmpty())
-						{
-							OperatorLabel = UnitPath.ToString();
-						}
-
-						TSharedPtr<FRigStackEntry> NewEntry = MakeShared<FRigStackEntry>(Operators.Num(), ERigStackEntry::Operator, OperatorIndex, Operator.OpCode, UnitPath, OperatorLabel);
-						UnitToOperatorIndex.Add(UnitPath, Operators.Num());
-						Operators.Add(NewEntry);
-						break;
-					}
-					default:
-					{
-						break;
-					}
+					FRigVMExecuteOp Op = InVM->ByteCode.GetOpAt<FRigVMExecuteOp>(Instruction);
+					FString OperatorLabel = InVM->GetRigVMFunctionName(Op.FunctionIndex);
+					TSharedPtr<FRigStackEntry> NewEntry = MakeShared<FRigStackEntry>(Operators.Num(), ERigStackEntry::Operator, InstructionIndex, Instruction.OpCode, OperatorLabel);
+					Operators.Add(NewEntry);
+					break;
+				}
+				default:
+				{
+					FString OperatorLabel = StaticEnum<ERigVMOpCode>()->GetNameStringByValue((int64)Instruction.OpCode);
+					TSharedPtr<FRigStackEntry> NewEntry = MakeShared<FRigStackEntry>(Operators.Num(), ERigStackEntry::Operator, InstructionIndex, Instruction.OpCode, OperatorLabel);
+					Operators.Add(NewEntry);
+					break;
 				}
 			}
+		}
 
-			// fill the children from the log
-			if (ControlRig)
+		// fill the children from the log
+		if (ControlRigEditor.IsValid())
+		{
+			UControlRig* ControlRig = ControlRigEditor.Pin()->ControlRig;
+			if(ControlRig && ControlRig->ControlRigLog)
 			{
-				if(ControlRig->ControlRigLog)
+				const TArray<FControlRigLog::FLogEntry>& LogEntries = ControlRig->ControlRigLog->Entries;
+				for (const FControlRigLog::FLogEntry& LogEntry : LogEntries)
 				{
-					const TArray<FControlRigLog::FLogEntry>& LogEntries = ControlRig->ControlRigLog->Entries;
-					for (const FControlRigLog::FLogEntry& LogEntry : LogEntries)
+					if (Operators.Num() <= LogEntry.InstructionIndex)
+					{
+						continue;
+					}
+					int32 ChildIndex = Operators[LogEntry.InstructionIndex]->Children.Num();
+					switch (LogEntry.Severity)
 					{
-						if (!UnitToOperatorIndex.Contains(LogEntry.Unit))
-						{
-							continue;
-						}
-						int32 OperatorIndex = UnitToOperatorIndex.FindChecked(LogEntry.Unit);
-						int32 ChildIndex = Operators[OperatorIndex]->Children.Num();
-						switch (LogEntry.Severity)
-						{
 						case EMessageSeverity::Info:
 						{
-							Operators[OperatorIndex]->Children.Add(MakeShared<FRigStackEntry>(ChildIndex, ERigStackEntry::Info, OperatorIndex, EControlRigOpCode::Invalid, LogEntry.Unit, LogEntry.Message));
+							Operators[LogEntry.InstructionIndex]->Children.Add(MakeShared<FRigStackEntry>(ChildIndex, ERigStackEntry::Info, LogEntry.InstructionIndex, ERigVMOpCode::Invalid, LogEntry.Message));
 							break;
 						}
 						case EMessageSeverity::Warning:
 						case EMessageSeverity::PerformanceWarning:
 						{
-							Operators[OperatorIndex]->Children.Add(MakeShared<FRigStackEntry>(ChildIndex, ERigStackEntry::Warning, OperatorIndex, EControlRigOpCode::Invalid, LogEntry.Unit, LogEntry.Message));
+							Operators[LogEntry.InstructionIndex]->Children.Add(MakeShared<FRigStackEntry>(ChildIndex, ERigStackEntry::Warning, LogEntry.InstructionIndex, ERigVMOpCode::Invalid, LogEntry.Message));
 							break;
 						}
 						case EMessageSeverity::Error:
 						case EMessageSeverity::CriticalError:
 						{
-							Operators[OperatorIndex]->Children.Add(MakeShared<FRigStackEntry>(ChildIndex, ERigStackEntry::Error, OperatorIndex, EControlRigOpCode::Invalid, LogEntry.Unit, LogEntry.Message));
+							Operators[LogEntry.InstructionIndex]->Children.Add(MakeShared<FRigStackEntry>(ChildIndex, ERigStackEntry::Error, LogEntry.InstructionIndex, ERigVMOpCode::Invalid, LogEntry.Message));
 							break;
 						}
 						default:
 						{
 							break;
 						}
-						}
 					}
 				}
 			}
@@ -356,18 +407,19 @@ void SControlRigStackView::RefreshTreeView(UControlRig* ControlRig)
 
 TSharedPtr< SWidget > SControlRigStackView::CreateContextMenu()
 {
-	const FControlRigStackCommands& Actions = FControlRigStackCommands::Get();
-
 	TArray<TSharedPtr<FRigStackEntry>> SelectedItems = TreeView->GetSelectedItems();
+	if(SelectedItems.Num() == 0)
+	{
+		return nullptr;
+	}
+
+	const FControlRigStackCommands& Actions = FControlRigStackCommands::Get();
 
 	const bool CloseAfterSelection = true;
 	FMenuBuilder MenuBuilder(CloseAfterSelection, CommandList);
 	{
 		MenuBuilder.BeginSection("RigStackToolsAction", LOCTEXT("ToolsAction", "Tools"));
-		if (SelectedItems.Num() > 0)
-		{
-			MenuBuilder.AddMenuEntry(Actions.FocusOnSelection);
-		}
+		MenuBuilder.AddMenuEntry(Actions.FocusOnSelection);
 		MenuBuilder.EndSection();
 	}
 
@@ -380,26 +432,55 @@ void SControlRigStackView::HandleFocusOnSelectedGraphNode()
 	ControlRigEditor.Pin()->ZoomToSelection_Clicked();
 }
 
-void SControlRigStackView::OnBlueprintCompiled(UBlueprint* InCompiledBlueprint)
+void SControlRigStackView::OnVMCompiled(UBlueprint* InCompiledBlueprint, URigVM* InCompiledVM)
 {
-	UControlRig* ControlRig = ControlRigEditor.Pin()->ControlRig;
-	if (ControlRig == nullptr)
+	RefreshTreeView(InCompiledVM);
+
+	if (ControlRigEditor.IsValid() && !OnControlRigInitializedHandle.IsValid())
 	{
-		return;
+		UControlRig* ControlRig = ControlRigEditor.Pin()->ControlRig;
+		OnControlRigInitializedHandle = ControlRig->OnInitialized_AnyThread().AddSP(this, &SControlRigStackView::HandleControlRigInitializedEvent);
 	}
+}
 
-	RefreshTreeView(ControlRig);
-	OnSelectionChanged(TSharedPtr<FRigStackEntry>(), ESelectInfo::Direct);
+void SControlRigStackView::HandleModifiedEvent(ERigVMGraphNotifType InNotifType, URigVMGraph* InGraph, UObject* InSubject)
+{
+	if (bSuspendModelNotifications)
+	{
+		return;
+	}
 
-	if (ControlRigEditor.IsValid())
+	switch (InNotifType)
 	{
-		OnControlRigInitializedHandle = ControlRig->OnInitialized().AddSP(this, &SControlRigStackView::OnControlRigInitialized);
+		case ERigVMGraphNotifType::NodeSelected:
+		case ERigVMGraphNotifType::NodeDeselected:
+		{
+			URigVMNode* Node = Cast<URigVMNode>(InSubject);
+			if (Node)
+			{
+				if (Node->GetInstructionIndex() != INDEX_NONE)
+				{
+					if (Operators.Num() > Node->GetInstructionIndex())
+					{
+						TGuardValue<bool> SuspendNotifs(bSuspendModelNotifications, true);
+						TreeView->SetItemSelection(Operators[Node->GetInstructionIndex()], InNotifType == ERigVMGraphNotifType::NodeSelected, ESelectInfo::Direct);
+					}
+				}
+			}
+			break;
+		}
+		default:
+		{
+			break;
+		}
 	}
 }
 
-void SControlRigStackView::OnControlRigInitialized(UControlRig* ControlRig, EControlRigState State)
+void SControlRigStackView::HandleControlRigInitializedEvent(UControlRig* InControlRig, const EControlRigState InState)
 {
-	RefreshTreeView(ControlRig);
+	TGuardValue<bool> SuspendControllerSelection(bSuspendControllerSelection, true);
+
+	RefreshTreeView(InControlRig->VM);
 	OnSelectionChanged(TSharedPtr<FRigStackEntry>(), ESelectInfo::Direct);
 
 	for (TSharedPtr<FRigStackEntry>& Operator : Operators)
@@ -413,51 +494,27 @@ void SControlRigStackView::OnControlRigInitialized(UControlRig* ControlRig, ECon
 			}
 		}
 	}
-}
-
-void SControlRigStackView::HandleModelModified(const UControlRigModel* InModel, EControlRigModelNotifType InType, const void* InPayload)
-{
-	if (bSuspendModelNotifications)
+	
+	if (ControlRigEditor.IsValid())
 	{
-		return;
-	}
+		InControlRig->OnInitialized_AnyThread().Remove(OnControlRigInitializedHandle);
+		OnControlRigInitializedHandle.Reset();
 
-	switch (InType)
-	{
-		case EControlRigModelNotifType::NodeSelected:
-		case EControlRigModelNotifType::NodeDeselected:
+		if (UControlRigBlueprint* RigBlueprint = ControlRigEditor.Pin()->GetControlRigBlueprint())
 		{
-			const FControlRigModelNode* Node = (const FControlRigModelNode*)InPayload;
-			if (Node != nullptr)
+			for (URigVMNode* ModelNode : RigBlueprint->Model->GetNodes())
 			{
-				if (!ControlRigBlueprint.IsValid())
-				{
-					return;
-				}
-
-				UControlRigBlueprintGeneratedClass* GeneratedClass = ControlRigBlueprint->GetControlRigBlueprintGeneratedClass();
-				if (GeneratedClass == nullptr)
-				{
-					return;
-				}
-
-				TGuardValue<bool> SuspendNotifs(bSuspendModelNotifications, true);
-				for (const TSharedPtr<FRigStackEntry>& Entry : Operators)
+				if (ModelNode->IsSelected())
 				{
-					if (Node->Name == Entry->Name)
-					{
-						TreeView->SetItemSelection(Entry, InType == EControlRigModelNotifType::NodeSelected, ESelectInfo::Direct);
-					}
+					HandleModifiedEvent(ERigVMGraphNotifType::NodeSelected, RigBlueprint->Model, ModelNode);
 				}
 			}
-			break;
-		}
-		default:
-		{
-			// todo... other cases
-			break;
 		}
 	}
 }
 
+void SControlRigStackView::HandlePreviewControlRigUpdated(FControlRigEditor* InEditor)
+{
+}
+
 #undef LOCTEXT_NAMESPACE
diff --git a/Engine/Plugins/Experimental/ControlRig/Source/ControlRigEditor/Private/Editor/SControlRigStackView.h b/Engine/Plugins/Experimental/ControlRig/Source/ControlRigEditor/Private/Editor/SControlRigStackView.h
--- a/Engine/Plugins/Experimental/ControlRig/Source/ControlRigEditor/Private/Editor/SControlRigStackView.h
+++ b/Engine/Plugins/Experimental/ControlRig/Source/ControlRigEditor/Private/Editor/SControlRigStackView.h
@@ -7,6 +7,7 @@
 #include "ControlRig.h"
 #include "ControlRigBlueprint.h"
 #include "Widgets/Views/STreeView.h"
+#include "RigVMCore/RigVM.h"
 
 class FControlRigEditor;
 class SControlRigStackView;
@@ -27,15 +28,14 @@ namespace ERigStackEntry
 class FRigStackEntry : public TSharedFromThis<FRigStackEntry>
 {
 public:
-	FRigStackEntry(int32 InEntryIndex, ERigStackEntry::Type InEntryType, int32 InOpIndex, EControlRigOpCode InOpCode, const FName& InName, const FString& InLabel);
+	FRigStackEntry(int32 InEntryIndex, ERigStackEntry::Type InEntryType, int32 InInstructionIndex, ERigVMOpCode InOpCode, const FString& InLabel);
 
 	TSharedRef<ITableRow> MakeTreeRowWidget(const TSharedRef<STableViewBase>& InOwnerTable, TSharedRef<FRigStackEntry> InEntry, TSharedRef<FUICommandList> InCommandList, TSharedPtr<SControlRigStackView> InStackView);
 
 	int32 EntryIndex;
 	ERigStackEntry::Type EntryType;
-	int32 OpIndex;
-	EControlRigOpCode OpCode;
-	FName Name;
+	int32 InstructionIndex;
+	ERigVMOpCode OpCode;
 	FString Label;
 	TArray<TSharedPtr<FRigStackEntry>> Children;
 };
@@ -73,7 +73,7 @@ public:
 protected:
 
 	/** Rebuild the tree view */
-	void RefreshTreeView(UControlRig* ControlRig = nullptr);
+	void RefreshTreeView(URigVM* InVM);
 
 private:
 
@@ -89,12 +89,13 @@ private:
 	/** Focus on the selected operator in the graph*/
 	void HandleFocusOnSelectedGraphNode();
 
-	void OnBlueprintCompiled(UBlueprint* InCompiledBlueprint);
-
-	void OnControlRigInitialized(UControlRig* ControlRig, EControlRigState State);
+	void OnVMCompiled(UBlueprint* InCompiledBlueprint, URigVM* InCompiledVM);
 
 	bool bSuspendModelNotifications;
-	void HandleModelModified(const UControlRigModel* InModel, EControlRigModelNotifType InType, const void* InPayload);
+	bool bSuspendControllerSelection;
+	void HandleModifiedEvent(ERigVMGraphNotifType InNotifType, URigVMGraph* InGraph, UObject* InSubject);
+	void HandleControlRigInitializedEvent(UControlRig* InControlRig, const EControlRigState InState);
+	void HandlePreviewControlRigUpdated(FControlRigEditor* InEditor);
 
 	TSharedPtr<STreeView<TSharedPtr<FRigStackEntry>>> TreeView;
 
@@ -108,6 +109,7 @@ private:
 	TArray<TSharedPtr<FRigStackEntry>> Operators;
 
 	FDelegateHandle OnModelModified;
-	FDelegateHandle OnBlueprintCompiledHandle;
 	FDelegateHandle OnControlRigInitializedHandle;
+	FDelegateHandle OnVMCompiledHandle;
+	FDelegateHandle OnPreviewControlRigUpdatedHandle;
 };
diff --git a/Engine/Plugins/Experimental/ControlRig/Source/ControlRigEditor/Private/Editor/SRigCurveContainer.cpp b/Engine/Plugins/Experimental/ControlRig/Source/ControlRigEditor/Private/Editor/SRigCurveContainer.cpp
--- a/Engine/Plugins/Experimental/ControlRig/Source/ControlRigEditor/Private/Editor/SRigCurveContainer.cpp
+++ b/Engine/Plugins/Experimental/ControlRig/Source/ControlRigEditor/Private/Editor/SRigCurveContainer.cpp
@@ -17,6 +17,7 @@
 #include "ControlRigBlueprint.h"
 #include "ScopedTransaction.h"
 #include "Kismet2/BlueprintEditorUtils.h"
+#include "ScopedTransaction.h"
 
 #define LOCTEXT_NAMESPACE "SRigCurveContainer"
 
@@ -484,6 +485,9 @@ void SRigCurveContainer::OnSelectionChanged(FDisplayedRigCurveInfoPtr Selection,
 
 	if (Container)
 	{
+
+		FScopedTransaction ScopedTransaction(LOCTEXT("SelectCurveTransaction", "Select Curve"), !GIsTransacting);
+
 		TGuardValue<bool> GuardRigHierarchyChanges(bIsChangingRigHierarchy, true);
 
 		TArray<FName> OldSelection = Container->CurrentSelection();
@@ -574,6 +578,7 @@ void SRigCurveContainer::CreateImportMenu(FMenuBuilder& MenuBuilder)
 		.Padding(3)
 		[
 			SNew(SObjectPropertyEntryBox)
+			//.ObjectPath_UObject(this, &SRigCurveContainer::GetCurrentHierarchy)
 			.OnShouldFilterAsset(this, &SRigCurveContainer::ShouldFilterOnImport)
 			.OnObjectChanged(this, &SRigCurveContainer::ImportCurve)
 		]
@@ -607,25 +612,12 @@ void SRigCurveContainer::ImportCurve(const FAssetData& InAssetData)
 
 		if (Skeleton)
 		{
+			TGuardValue<bool> GuardReentry(bIsChangingRigHierarchy, true);
+
 			FScopedTransaction Transaction(LOCTEXT("CurveImport", "Import Curve"));
 			ControlRigBlueprint->Modify();
-				
-			const FSmartNameMapping* SmartNameMapping = Skeleton->GetSmartNameContainer(USkeleton::AnimCurveMappingName);
-
 			Container->ClearSelection();
-
-			TArray<FName> NameArray;
-			SmartNameMapping->FillNameArray(NameArray);
-			for (int32 Index = 0; Index < NameArray.Num() ; ++Index)
-			{
-				TGuardValue<bool> GuardReentry(bIsChangingRigHierarchy, true);
-				Container->Add(NameArray[Index]);
-			}
-
-			for (int32 Index = 0; Index < NameArray.Num(); ++Index)
-			{
-				Container->Select(NameArray[Index]);
-			}
+			Container->ImportCurvesFromSkeleton(Skeleton, NAME_None, false, false, false /* notify */);
 
 			FSlateApplication::Get().DismissAllMenus();
 			RefreshCurveList();
diff --git a/Engine/Plugins/Experimental/ControlRig/Source/ControlRigEditor/Private/Editor/SRigHierarchy.cpp b/Engine/Plugins/Experimental/ControlRig/Source/ControlRigEditor/Private/Editor/SRigHierarchy.cpp
--- a/Engine/Plugins/Experimental/ControlRig/Source/ControlRigEditor/Private/Editor/SRigHierarchy.cpp
+++ b/Engine/Plugins/Experimental/ControlRig/Source/ControlRigEditor/Private/Editor/SRigHierarchy.cpp
@@ -3,6 +3,7 @@
 #include "SRigHierarchy.h"
 #include "Widgets/Input/SComboButton.h"
 #include "EditorStyleSet.h"
+#include "Widgets/Input/SButton.h"
 #include "Widgets/Input/SSearchBox.h"
 #include "Kismet2/BlueprintEditorUtils.h"
 #include "ScopedTransaction.h"
@@ -14,7 +15,6 @@
 #include "Widgets/Text/STextBlock.h"
 #include "K2Node_VariableGet.h"
 #include "ControlRigBlueprintUtils.h"
-#include "NodeSpawners/ControlRigPropertyNodeSpawner.h"
 #include "ControlRigHierarchyCommands.h"
 #include "ControlRigBlueprint.h"
 #include "Graph/ControlRigGraph.h"
@@ -32,6 +32,8 @@
 #include "ControlRigEditorStyle.h"
 #include "HAL/PlatformApplicationMisc.h"
 #include "Dialogs/Dialogs.h"
+#include "IPersonaToolkit.h"
+#include "SKismetInspector.h"
 
 #define LOCTEXT_NAMESPACE "SRigHierarchy"
 
@@ -226,16 +228,15 @@ SRigHierarchy::~SRigHierarchy()
 		ControlRigEditor.Pin()->GetKeyDownDelegate().Unbind();
 		ControlRigEditor.Pin()->OnViewportContextMenu().Unbind();
 		ControlRigEditor.Pin()->OnViewportContextMenuCommands().Unbind();
+	}
 
-		ControlRigBlueprint = ControlRigEditor.Pin()->GetControlRigBlueprint();
-		if (ControlRigBlueprint.IsValid())
-		{
-			ControlRigBlueprint->HierarchyContainer.OnElementAdded.RemoveAll(this);
-			ControlRigBlueprint->HierarchyContainer.OnElementRemoved.RemoveAll(this);
-			ControlRigBlueprint->HierarchyContainer.OnElementRenamed.RemoveAll(this);
-			ControlRigBlueprint->HierarchyContainer.OnElementReparented.RemoveAll(this);
-			ControlRigBlueprint->HierarchyContainer.OnElementSelected.RemoveAll(this);
-		}
+	if (ControlRigBlueprint.IsValid())
+	{
+		ControlRigBlueprint->HierarchyContainer.OnElementAdded.RemoveAll(this);
+		ControlRigBlueprint->HierarchyContainer.OnElementRemoved.RemoveAll(this);
+		ControlRigBlueprint->HierarchyContainer.OnElementRenamed.RemoveAll(this);
+		ControlRigBlueprint->HierarchyContainer.OnElementReparented.RemoveAll(this);
+		ControlRigBlueprint->HierarchyContainer.OnElementSelected.RemoveAll(this);
 	}
 }
 
@@ -277,11 +278,34 @@ void SRigHierarchy::Construct(const FArguments& InArgs, TSharedRef<FControlRigEd
 			.BorderBackgroundColor(FLinearColor(0.6f, 0.6f, 0.6f, 1.0f))
 			[
 				SNew(SVerticalBox)
+
+				+ SVerticalBox::Slot()
+				.AutoHeight()
+				.VAlign(VAlign_Top)
+				[
+					SNew(SHorizontalBox)
+					.Visibility(this, &SRigHierarchy::IsToolbarVisible)
+
+					+ SHorizontalBox::Slot()
+					.VAlign(VAlign_Center)
+					.MaxWidth(180.0f)
+					.Padding(3.0f, 1.0f)
+					[
+						SNew(SButton)
+						.ButtonStyle(FEditorStyle::Get(), "FlatButton.Success")
+						.ForegroundColor(FLinearColor::White)
+						.OnClicked(FOnClicked::CreateSP(this, &SRigHierarchy::OnImportSkeletonClicked))
+						.Text(FText::FromString(TEXT("Import Hierarchy")))
+					]
+				]
+
 				+SVerticalBox::Slot()
 				.AutoHeight()
 				.VAlign(VAlign_Top)
 				[
 					SNew(SHorizontalBox)
+					.Visibility(this, &SRigHierarchy::IsSearchbarVisible)
+
 					+SHorizontalBox::Slot()
 					.VAlign(VAlign_Center)
 					.Padding(3.0f, 1.0f)
@@ -306,6 +330,7 @@ void SRigHierarchy::Construct(const FArguments& InArgs, TSharedRef<FControlRigEd
 				.OnGetChildren(this, &SRigHierarchy::HandleGetChildrenForTree)
 				.OnSelectionChanged(this, &SRigHierarchy::OnSelectionChanged)
 				.OnContextMenuOpening(this, &SRigHierarchy::CreateContextMenu)
+				.OnMouseButtonDoubleClick(this, &SRigHierarchy::OnItemDoubleClicked)
 				.HighlightParentNodesForSelection(true)
 				.ItemHeight(24)
 			]
@@ -372,16 +397,6 @@ void SRigHierarchy::BindCommands()
 		FExecuteAction::CreateSP(this, &SRigHierarchy::HandleResetTransform),
 		FCanExecuteAction::CreateSP(this, &SRigHierarchy::IsMultiSelected));
 
-	CommandList->MapAction(
-		Commands.ResetInitialTransform,
-		FExecuteAction::CreateSP(this, &SRigHierarchy::HandleResetInitialTransform),
-		FCanExecuteAction::CreateSP(this, &SRigHierarchy::IsMultiSelected));
-
-	CommandList->MapAction(
-		Commands.ResetSpace,
-		FExecuteAction::CreateSP(this, &SRigHierarchy::HandleResetSpace),
-		FCanExecuteAction::CreateSP(this, &SRigHierarchy::IsControlSelected));
-
 	CommandList->MapAction(
 		Commands.SetInitialTransformFromClosestBone,
 		FExecuteAction::CreateSP(this, &SRigHierarchy::HandleSetInitialTransformFromClosestBone),
@@ -390,12 +405,22 @@ void SRigHierarchy::BindCommands()
 	CommandList->MapAction(
 		Commands.SetInitialTransformFromCurrentTransform,
 		FExecuteAction::CreateSP(this, &SRigHierarchy::HandleSetInitialTransformFromCurrentTransform),
-		FCanExecuteAction::CreateSP(this, &SRigHierarchy::IsControlOrSpaceSelected));
+		FCanExecuteAction::CreateSP(this, &SRigHierarchy::IsMultiSelected));
 
 	CommandList->MapAction(
 		Commands.FrameSelection,
 		FExecuteAction::CreateSP(this, &SRigHierarchy::HandleFrameSelection),
 		FCanExecuteAction::CreateSP(this, &SRigHierarchy::IsMultiSelected));
+
+	CommandList->MapAction(
+		Commands.ControlBoneTransform,
+		FExecuteAction::CreateSP(this, &SRigHierarchy::HandleControlBoneOrSpaceTransform),
+		FCanExecuteAction::CreateSP(this, &SRigHierarchy::IsSingleBoneSelected));
+
+	CommandList->MapAction(
+		Commands.ControlSpaceTransform,
+		FExecuteAction::CreateSP(this, &SRigHierarchy::HandleControlBoneOrSpaceTransform),
+		FCanExecuteAction::CreateSP(this, &SRigHierarchy::IsSingleSpaceSelected));
 }
 
 FReply SRigHierarchy::OnKeyDown(const FGeometry& MyGeometry, const FKeyEvent& InKeyEvent)
@@ -407,6 +432,48 @@ FReply SRigHierarchy::OnKeyDown(const FGeometry& MyGeometry, const FKeyEvent& In
 	return FReply::Unhandled();
 }
 
+EVisibility SRigHierarchy::IsToolbarVisible() const
+{
+	if (FRigHierarchyContainer* Container = GetHierarchyContainer())
+	{
+		if (Container->BoneHierarchy.Num() > 0)
+		{
+			return EVisibility::Collapsed;
+		}
+	}
+	return EVisibility::Visible;
+}
+
+EVisibility SRigHierarchy::IsSearchbarVisible() const
+{
+	if (FRigHierarchyContainer* Container = GetHierarchyContainer())
+	{
+		if ((Container->BoneHierarchy.Num() + Container->SpaceHierarchy.Num() + Container->ControlHierarchy.Num()) > 0)
+		{
+			return EVisibility::Visible;
+		}
+	}
+	return EVisibility::Collapsed;
+}
+
+FReply SRigHierarchy::OnImportSkeletonClicked()
+{
+	FRigHierarchyImportSettings Settings;
+	TSharedPtr<FStructOnScope> StructToDisplay = MakeShareable(new FStructOnScope(FRigHierarchyImportSettings::StaticStruct(), (uint8*)&Settings));
+
+	TSharedRef<SKismetInspector> KismetInspector = SNew(SKismetInspector);
+	KismetInspector->ShowSingleStruct(StructToDisplay);
+
+	SGenericDialogWidget::OpenDialog(LOCTEXT("ControlRigHierarchyImport", "Import Hierarchy"), KismetInspector, SGenericDialogWidget::FArguments(), true);
+
+	if (Settings.Mesh != nullptr)
+	{
+		ImportHierarchy(FAssetData(Settings.Mesh));
+	}
+
+	return FReply::Handled();
+}
+
 void SRigHierarchy::OnFilterTextChanged(const FText& SearchText)
 {
 	FilterText = SearchText;
@@ -602,6 +669,7 @@ void SRigHierarchy::AddElement(FRigElementKey InKey, FRigElementKey InParentKey)
 		if (InKey.Name.ToString().Contains(FilteredString) || InKey.Name.ToString().Contains(FilteredStringUnderScores))	
 		{
 			TSharedPtr<FRigTreeElement> NewItem = MakeShared<FRigTreeElement>(InKey, SharedThis(this));
+			ElementMap.Add(InKey, NewItem);
 			RootElements.Add(NewItem);
 		}
 	}
@@ -688,6 +756,15 @@ void SRigHierarchy::OnRigElementAdded(FRigHierarchyContainer* Container, const F
 	{
 		return;
 	}
+
+	if (InKey.Type == ERigElementType::Control)
+	{
+		if (Container->GetIndex(InKey) == INDEX_NONE)
+		{
+			return;
+		}
+	}
+
 	RefreshTreeView();
 }
 
@@ -697,6 +774,15 @@ void SRigHierarchy::OnRigElementRemoved(FRigHierarchyContainer* Container, const
 	{
 		return;
 	}
+
+	if (InKey.Type == ERigElementType::Control)
+	{
+		if (Container->GetIndex(InKey) == INDEX_NONE)
+		{
+			return;
+		}
+	}
+
 	RefreshTreeView();
 }
 
@@ -725,6 +811,14 @@ void SRigHierarchy::OnRigElementSelected(FRigHierarchyContainer* Container, cons
 		return;
 	}
 
+	if (InKey.Type == ERigElementType::Control)
+	{
+		if (Container->GetIndex(InKey) == INDEX_NONE)
+		{
+			return;
+		}
+	}
+
 	for (int32 RootIndex = 0; RootIndex < RootElements.Num(); ++RootIndex)
 	{
 		TSharedPtr<FRigTreeElement> Found = FindElement(InKey, RootElements[RootIndex]);
@@ -749,6 +843,19 @@ TSharedPtr< SWidget > SRigHierarchy::CreateContextMenu()
 	return MenuBuilder.MakeWidget();
 }
 
+void SRigHierarchy::OnItemDoubleClicked(TSharedPtr<FRigTreeElement> InItem)
+{
+	if (TreeView->IsItemExpanded(InItem))
+	{
+		TreeView->SetItemExpansion(InItem, false);
+	}
+	else
+	{
+		SetExpansionRecursive(InItem, false);
+	}
+}
+
+
 void SRigHierarchy::FillContextMenu(class FMenuBuilder& MenuBuilder)
 {
 	const FControlRigHierarchyCommands& Actions = FControlRigHierarchyCommands::Get();
@@ -770,7 +877,6 @@ void SRigHierarchy::FillContextMenu(class FMenuBuilder& MenuBuilder)
 				{
 					InSubMenuBuilder.AddMenuEntry(Actions.AddBoneItem);
 				}
-
 				InSubMenuBuilder.AddMenuEntry(Actions.AddControlItem);
 				InSubMenuBuilder.AddMenuEntry(Actions.AddSpaceItem);
 			}
@@ -785,6 +891,20 @@ void SRigHierarchy::FillContextMenu(class FMenuBuilder& MenuBuilder)
 		MenuBuilder.AddMenuEntry(Actions.RenameItem);
 		MenuBuilder.EndSection();
 
+		if (IsSingleBoneSelected())
+		{
+			MenuBuilder.BeginSection("Interaction", LOCTEXT("InteractionHeader", "Interaction"));
+			MenuBuilder.AddMenuEntry(Actions.ControlBoneTransform);
+			MenuBuilder.EndSection();
+		}
+
+		if (IsSingleSpaceSelected())
+		{
+			MenuBuilder.BeginSection("Interaction", LOCTEXT("InteractionHeader", "Interaction"));
+			MenuBuilder.AddMenuEntry(Actions.ControlSpaceTransform);
+			MenuBuilder.EndSection();
+		}
+
 		MenuBuilder.BeginSection("Copy&Paste", LOCTEXT("Copy&PasteHeader", "Copy & Paste"));
 		MenuBuilder.AddMenuEntry(Actions.CopyItems);
 		MenuBuilder.AddMenuEntry(Actions.PasteItems);
@@ -794,10 +914,8 @@ void SRigHierarchy::FillContextMenu(class FMenuBuilder& MenuBuilder)
 
 		MenuBuilder.BeginSection("Transforms", LOCTEXT("TransformsHeader", "Transforms"));
 		MenuBuilder.AddMenuEntry(Actions.ResetTransform);
-		MenuBuilder.AddMenuEntry(Actions.ResetInitialTransform);
 		MenuBuilder.AddMenuEntry(Actions.SetInitialTransformFromCurrentTransform);
 		MenuBuilder.AddMenuEntry(Actions.SetInitialTransformFromClosestBone);
-		MenuBuilder.AddMenuEntry(Actions.ResetSpace);
 		MenuBuilder.EndSection();
 
 		MenuBuilder.BeginSection("Assets", LOCTEXT("AssetsHeader", "Assets"));
@@ -859,9 +977,13 @@ void SRigHierarchy::RefreshHierarchy(const FAssetData& InAssetData)
 		ControlRigBlueprint->Modify();
 
 		const FReferenceSkeleton& RefSkeleton = Mesh->RefSkeleton;
-		Hierarchy->BoneHierarchy.ImportSkeleton(RefSkeleton, NAME_None, true, true, true);
+		Hierarchy->BoneHierarchy.ImportSkeleton(RefSkeleton, NAME_None, true, true, true, false /* notify */);
 	}
+
+	FSlateApplication::Get().DismissAllMenus();
+	RefreshTreeView();
 }
+
 void SRigHierarchy::CreateImportMenu(FMenuBuilder& MenuBuilder)
 {
 	MenuBuilder.AddWidget(
@@ -899,10 +1021,22 @@ void SRigHierarchy::ImportHierarchy(const FAssetData& InAssetData)
 		ControlRigBlueprint->Modify();
 
 		const FReferenceSkeleton& RefSkeleton = Mesh->RefSkeleton;
-		Hierarchy->BoneHierarchy.ImportSkeleton(RefSkeleton, NAME_None, false, false, true);
+		Hierarchy->BoneHierarchy.ImportSkeleton(RefSkeleton, NAME_None, false, false, true, false /* notify */);
+	}
 
-		FSlateApplication::Get().DismissAllMenus();
-		RefreshTreeView();
+	FSlateApplication::Get().DismissAllMenus();
+	RefreshTreeView();
+
+	if (ControlRigBlueprint->GetPreviewMesh() == nullptr &&
+		ControlRigEditor.IsValid() && 
+		Mesh != nullptr)
+	{
+		ControlRigEditor.Pin()->GetPersonaToolkit()->SetPreviewMesh(Mesh, true);
+	}
+
+	if (ControlRigEditor.IsValid())
+	{
+		ControlRigEditor.Pin()->Compile();
 	}
 }
 
@@ -916,6 +1050,24 @@ bool SRigHierarchy::IsSingleSelected() const
 	return TreeView->GetNumItemsSelected() == 1;
 }
 
+bool SRigHierarchy::IsSingleBoneSelected() const
+{
+	if(!IsSingleSelected())
+	{
+		return false;
+	}
+	return TreeView->GetSelectedItems()[0]->Key.Type == ERigElementType::Bone;
+}
+
+bool SRigHierarchy::IsSingleSpaceSelected() const
+{
+	if(!IsSingleSelected())
+	{
+		return false;
+	}
+	return TreeView->GetSelectedItems()[0]->Key.Type == ERigElementType::Space;
+}
+
 bool SRigHierarchy::IsControlSelected() const
 {
 	TArray<TSharedPtr<FRigTreeElement>> SelectedItems = TreeView->GetSelectedItems();
@@ -1022,10 +1174,10 @@ void SRigHierarchy::HandleDeleteItem()
 				}
 			}
 		}
-
-		RefreshTreeView();
-		FSlateApplication::Get().DismissAllMenus();
  	}
+
+	RefreshTreeView();
+	FSlateApplication::Get().DismissAllMenus();
 }
 
 bool SRigHierarchy::CanDeleteItem() const
@@ -1131,12 +1283,13 @@ void SRigHierarchy::HandleNewItem(ERigElementType InElementType)
 				}
 			}
 		}
+
 		Hierarchy->ClearSelection();
 		Hierarchy->Select(FRigElementKey(NewElementName, InElementType));
-
-		FSlateApplication::Get().DismissAllMenus();
-		RefreshTreeView();
 	}
+
+	FSlateApplication::Get().DismissAllMenus();
+	RefreshTreeView();
 }
 
 /** Check whether we can deleting the selected item(s) */
@@ -1204,10 +1357,10 @@ void SRigHierarchy::HandleDuplicateItem()
 		{
 			Hierarchy->Select(NewKeys[Index]);
 		}
-
-		FSlateApplication::Get().DismissAllMenus();
-		RefreshTreeView();
 	}
+
+	FSlateApplication::Get().DismissAllMenus();
+	RefreshTreeView();
 }
 
 /** Check whether we can deleting the selected item(s) */
@@ -1270,7 +1423,7 @@ bool SRigHierarchy::CanCopyOrPasteItems() const
 
 void SRigHierarchy::HandleCopyItems()
 {
-	if (FRigHierarchyContainer* Hierarchy = GetHierarchyContainer())
+	if (FRigHierarchyContainer* Hierarchy = GetDebuggedHierarchyContainer())
 	{
 		TArray<FRigElementKey> Selection = GetHierarchyContainer()->CurrentSelection();
 		FString Content = Hierarchy->ExportToText(Selection);
@@ -1289,8 +1442,9 @@ void SRigHierarchy::HandlePasteItems()
 		ControlRigBlueprint->Modify();
 
 		Hierarchy->ImportFromText(Content, ERigHierarchyImportMode::Append, true);
-		ControlRigBlueprint->PropagateHierarchyFromBPToInstances(false);
 	}
+
+	ControlRigBlueprint->PropagateHierarchyFromBPToInstances(false);
 }
 
 void SRigHierarchy::HandlePasteLocalTransforms()
@@ -1304,8 +1458,9 @@ void SRigHierarchy::HandlePasteLocalTransforms()
 		ControlRigBlueprint->Modify();
 
 		Hierarchy->ImportFromText(Content, ERigHierarchyImportMode::ReplaceLocalTransform, true);
-		ControlRigBlueprint->PropagateHierarchyFromBPToInstances(false);
 	}
+
+	ControlRigBlueprint->PropagateHierarchyFromBPToInstances(false);
 }
 
 void SRigHierarchy::HandlePasteGlobalTransforms()
@@ -1319,8 +1474,9 @@ void SRigHierarchy::HandlePasteGlobalTransforms()
 		ControlRigBlueprint->Modify();
 
 		Hierarchy->ImportFromText(Content, ERigHierarchyImportMode::ReplaceGlobalTransform, true);
-		ControlRigBlueprint->PropagateHierarchyFromBPToInstances(false);
 	}
+
+	ControlRigBlueprint->PropagateHierarchyFromBPToInstances(false);
 }
 
 FRigHierarchyContainer* SRigHierarchy::GetHierarchyContainer() const
@@ -1379,7 +1535,6 @@ void SRigHierarchy::PostRedo(bool bSuccess)
 {
 	if (bSuccess)
 	{
-		ControlRigBlueprint->PropagateHierarchyFromBPToInstances(false);
 		RefreshTreeView();
 	}
 }
@@ -1388,7 +1543,6 @@ void SRigHierarchy::PostUndo(bool bSuccess)
 {
 	if (bSuccess)
 	{
-		ControlRigBlueprint->PropagateHierarchyFromBPToInstances(false);
 		RefreshTreeView();
 	}
 }
@@ -1732,46 +1886,37 @@ void SRigHierarchy::HandleResetTransform()
 		{
 			if (FRigHierarchyContainer* DebuggedContainer = GetDebuggedHierarchyContainer())
 			{
-				FScopedTransaction Transaction(LOCTEXT("HierarchyResetInitialTransforms", "Reset Initial Transforms"));
+				FScopedTransaction Transaction(LOCTEXT("HierarchyResetTransforms", "Reset Transforms"));
 				Blueprint->Modify();
 
 				TArray<TSharedPtr<FRigTreeElement>> SelectedItems = TreeView->GetSelectedItems();
 				for (TSharedPtr<FRigTreeElement> SelectedItem : SelectedItems)
 				{
 					FRigElementKey Key = SelectedItem->Key;
-					if (Key.Type == ERigElementType::Control || Key.Type == ERigElementType::Space)
+					if (Key.Type == ERigElementType::Control)
+					{
+						GetHierarchyContainer()->SetLocalTransform(Key, FTransform::Identity);
+						DebuggedContainer->SetLocalTransform(Key, FTransform::Identity);
+					}
+
+					else if (Key.Type == ERigElementType::Space)
 					{
 						FTransform InitialTransform = GetHierarchyContainer()->GetInitialTransform(Key);
 						GetHierarchyContainer()->SetLocalTransform(Key, InitialTransform);
 						DebuggedContainer->SetLocalTransform(Key, InitialTransform);
 					}
-				}
-			}
-		}
-	}
-}
-
-void SRigHierarchy::HandleResetInitialTransform()
-{
-	if (IsMultiSelected() && ControlRigEditor.IsValid())
-	{
-		if (UControlRigBlueprint* Blueprint = ControlRigEditor.Pin()->GetControlRigBlueprint())
-		{
-			if (FRigHierarchyContainer* DebuggedContainer = GetDebuggedHierarchyContainer())
-			{
-				FScopedTransaction Transaction(LOCTEXT("HierarchyResetInitialTransforms", "Reset Initial Transforms"));
-				Blueprint->Modify();
 
-				TArray<TSharedPtr<FRigTreeElement>> SelectedItems = TreeView->GetSelectedItems();
-				for (TSharedPtr<FRigTreeElement> SelectedItem : SelectedItems)
-				{
-					FRigElementKey Key = SelectedItem->Key;
-					if (Key.Type == ERigElementType::Control || Key.Type == ERigElementType::Space)
+					else if (Key.Type == ERigElementType::Bone)
 					{
-						GetHierarchyContainer()->SetInitialTransform(Key, FTransform::Identity);
-						GetHierarchyContainer()->SetLocalTransform(Key, FTransform::Identity);
-						DebuggedContainer->SetInitialTransform(Key, FTransform::Identity);
-						DebuggedContainer->SetLocalTransform(Key, FTransform::Identity);
+						FTransform InitialTransform = GetHierarchyContainer()->GetInitialGlobalTransform(Key);
+						GetHierarchyContainer()->SetGlobalTransform(Key, InitialTransform);
+						DebuggedContainer->SetGlobalTransform(Key, InitialTransform);
+
+						Blueprint->RemoveTransientControl(Key);
+						if (ControlRigEditor.Pin()->PreviewInstance)
+						{
+							ControlRigEditor.Pin()->PreviewInstance->RemoveBoneModification(Key.Name);
+						}
 					}
 				}
 			}
@@ -1917,14 +2062,40 @@ void SRigHierarchy::HandleSetInitialTransformFromCurrentTransform()
 				for (TSharedPtr<FRigTreeElement> SelectedItem : SelectedItems)
 				{
 					FRigElementKey Key = SelectedItem->Key;
-					if (Key.Type == ERigElementType::Control || Key.Type == ERigElementType::Space || Key.Type == ERigElementType::Bone)
+
+					FTransform GlobalTransform = GlobalTransforms[Key];
+					FTransform LocalTransform = GlobalTransform.GetRelativeTransform(ParentGlobalTransforms[Key]);
+
+					if (Key.Type == ERigElementType::Control )
+					{
+						GetHierarchyContainer()->SetInitialTransform(Key, LocalTransform);
+						GetHierarchyContainer()->SetLocalTransform(Key, FTransform::Identity);
+						DebuggedContainer->SetInitialTransform(Key, LocalTransform);
+						DebuggedContainer->SetLocalTransform(Key, FTransform::Identity);
+					}
+					else if (Key.Type == ERigElementType::Space)
 					{
-						FTransform LocalTransform = GlobalTransforms[Key].GetRelativeTransform(ParentGlobalTransforms[Key]);
 						GetHierarchyContainer()->SetInitialTransform(Key, LocalTransform);
 						GetHierarchyContainer()->SetLocalTransform(Key, LocalTransform);
 						DebuggedContainer->SetInitialTransform(Key, LocalTransform);
 						DebuggedContainer->SetLocalTransform(Key, LocalTransform);
 					}
+					else if (Key.Type == ERigElementType::Bone)
+					{
+						FTransform InitialTransform = GlobalTransform;
+						if (ControlRigEditor.Pin()->PreviewInstance)
+						{
+							if (FAnimNode_ModifyBone* ModifyBone = ControlRigEditor.Pin()->PreviewInstance->FindModifiedBone(Key.Name))
+							{
+								InitialTransform.SetTranslation(ModifyBone->Translation);
+								InitialTransform.SetRotation(FQuat(ModifyBone->Rotation));
+								InitialTransform.SetScale3D(ModifyBone->Scale);
+								InitialTransform = InitialTransform * ParentGlobalTransforms[Key];
+							}
+						}
+						GetHierarchyContainer()->SetInitialGlobalTransform(Key, InitialTransform);
+						Blueprint->PropagateHierarchyFromBPToInstances(false, false);
+					}
 				}
 			}
 		}
@@ -1945,6 +2116,28 @@ void SRigHierarchy::HandleFrameSelection()
 	}
 }
 
+void SRigHierarchy::HandleControlBoneOrSpaceTransform()
+{
+	UControlRigBlueprint* Blueprint = ControlRigEditor.Pin()->GetControlRigBlueprint();
+	if (Blueprint == nullptr)
+	{
+		return;
+	}
+
+	TArray<TSharedPtr<FRigTreeElement>> SelectedItems = TreeView->GetSelectedItems();
+	for (TSharedPtr<FRigTreeElement> SelectedItem : SelectedItems)
+	{
+		if (SelectedItem->Key.Type == ERigElementType::Bone ||
+			SelectedItem->Key.Type == ERigElementType::Space)
+		{
+			Blueprint->AddTransientControl(SelectedItem->Key);
+			Blueprint->HierarchyContainer.ClearSelection();
+			Blueprint->HierarchyContainer.Select(SelectedItem->Key, true);
+			return;
+		}
+	}
+}
+
 bool SRigHierarchy::FindClosestBone(const FVector& Point, FName& OutRigElementName, FTransform& OutGlobalTransform) const
 {
 	if (FRigHierarchyContainer* DebuggedContainer = GetDebuggedHierarchyContainer())
@@ -2058,9 +2251,9 @@ void SRigHierarchy::HandleSetInitialTransformFromClosestBone()
 
 						FTransform LocalTransform = ClosestTransforms[Key].GetRelativeTransform(ParentGlobalTransforms[Key]);
 						GetHierarchyContainer()->SetInitialTransform(Key, LocalTransform);
-						GetHierarchyContainer()->SetLocalTransform(Key, LocalTransform);
+						GetHierarchyContainer()->SetLocalTransform(Key, Key.Type == ERigElementType::Control ? FTransform::Identity : LocalTransform);
 						DebuggedContainer->SetInitialTransform(Key, LocalTransform);
-						DebuggedContainer->SetLocalTransform(Key, LocalTransform);
+						DebuggedContainer->SetLocalTransform(Key, Key.Type == ERigElementType::Control ? FTransform::Identity : LocalTransform);
 					}
 				}
 			}
diff --git a/Engine/Plugins/Experimental/ControlRig/Source/ControlRigEditor/Private/Editor/SRigHierarchy.h b/Engine/Plugins/Experimental/ControlRig/Source/ControlRigEditor/Private/Editor/SRigHierarchy.h
--- a/Engine/Plugins/Experimental/ControlRig/Source/ControlRigEditor/Private/Editor/SRigHierarchy.h
+++ b/Engine/Plugins/Experimental/ControlRig/Source/ControlRigEditor/Private/Editor/SRigHierarchy.h
@@ -7,6 +7,8 @@
 #include "Rigs/RigHierarchyContainer.h"
 #include "EditorUndoClient.h"
 #include "DragAndDrop/GraphNodeDragDropOp.h"
+#include "Engine/SkeletalMesh.h"
+#include "SRigHierarchy.generated.h"
 
 class SRigHierarchy;
 class FControlRigEditor;
@@ -92,6 +94,19 @@ private:
 	FText GetName() const;
 };
 
+USTRUCT()
+struct FRigHierarchyImportSettings
+{
+	GENERATED_BODY()
+
+	FRigHierarchyImportSettings()
+	: Mesh(nullptr)
+	{}
+
+	UPROPERTY(EditAnywhere, Category = "Hierachy Import")
+	USkeletalMesh* Mesh;
+};
+
 /** Widget allowing editing of a control rig's structure */
 class SRigHierarchy : public SCompoundWidget, public FEditorUndoClient
 {
@@ -151,6 +166,7 @@ private:
 	void OnSelectionChanged(TSharedPtr<FRigTreeElement> Selection, ESelectInfo::Type SelectInfo);
 
 	TSharedPtr< SWidget > CreateContextMenu();
+	void OnItemDoubleClicked(TSharedPtr<FRigTreeElement> InItem);
 
 	// FEditorUndoClient
 	virtual void PostUndo(bool bSuccess) override;
@@ -173,6 +189,9 @@ private:
 	TSharedPtr<SSearchBox> FilterBox;
 	FText FilterText;
 
+	EVisibility IsToolbarVisible() const;
+	EVisibility IsSearchbarVisible() const;
+	FReply OnImportSkeletonClicked();
 	void OnFilterTextChanged(const FText& SearchText);
 
 	/** Tree view widget */
@@ -194,6 +213,8 @@ private:
 
 	bool IsMultiSelected() const;
 	bool IsSingleSelected() const;
+	bool IsSingleBoneSelected() const;
+	bool IsSingleSpaceSelected() const;
 	bool IsControlSelected() const;
 	bool IsControlOrSpaceSelected() const;
 
@@ -212,6 +233,7 @@ private:
 	void HandleSetInitialTransformFromCurrentTransform();
 	void HandleSetInitialTransformFromClosestBone();
 	void HandleFrameSelection();
+	void HandleControlBoneOrSpaceTransform();
 	bool FindClosestBone(const FVector& Point, FName& OutRigElementName, FTransform& OutGlobalTransform) const;
 
 	FName CreateUniqueName(const FName& InBaseName, ERigElementType InElementType) const;
diff --git a/Engine/Plugins/Experimental/ControlRig/Source/ControlRigEditor/Private/Graph/ControlRigConnectionDrawingPolicy.cpp b/Engine/Plugins/Experimental/ControlRig/Source/ControlRigEditor/Private/Graph/ControlRigConnectionDrawingPolicy.cpp
--- a/Engine/Plugins/Experimental/ControlRig/Source/ControlRigEditor/Private/Graph/ControlRigConnectionDrawingPolicy.cpp
+++ b/Engine/Plugins/Experimental/ControlRig/Source/ControlRigEditor/Private/Graph/ControlRigConnectionDrawingPolicy.cpp
@@ -4,7 +4,7 @@
 #include "Graph/ControlRigGraph.h"
 #include "Graph/ControlRigGraphNode.h"
 #include "ControlRigBlueprint.h"
-#include "ControlRigController.h"
+#include "RigVMModel/RigVMController.h"
 #include "Kismet2/BlueprintEditorUtils.h"
 
 void FControlRigConnectionDrawingPolicy::SetIncompatiblePinDrawState(const TSharedPtr<SGraphPin>& StartPin, const TSet< TSharedRef<SWidget> >& VisiblePins)
@@ -12,13 +12,13 @@ void FControlRigConnectionDrawingPolicy::SetIncompatiblePinDrawState(const TShar
 	UEdGraphPin* Pin = StartPin->GetPinObj();
 	if (Pin != nullptr)
 	{
-		UBlueprint* Blueprint = FBlueprintEditorUtils::FindBlueprintForNodeChecked(Pin->GetOwningNode());
-		UControlRigBlueprint* RigBlueprint = Cast<UControlRigBlueprint>(Blueprint);
-		if (RigBlueprint != nullptr)
+		UControlRigGraphNode* RigNode = Cast<UControlRigGraphNode>(Pin->GetOwningNode());
+		if(RigNode)
 		{
-			FString Left, Right;
-			RigBlueprint->Model->SplitPinPath(Pin->GetName(), Left, Right);
-			RigBlueprint->ModelController->PrepareCycleCheckingForPin(*Left, *Right, Pin->Direction == EGPD_Input);
+			if(URigVMPin* ModelPin = RigNode->GetModelPinFromPinPath(Pin->GetName()))
+			{
+				ModelPin->GetGraph()->PrepareCycleChecking(ModelPin->GetPinForLink(), Pin->Direction == EGPD_Input);
+			}
 		}
 	}
 	FKismetConnectionDrawingPolicy::SetIncompatiblePinDrawState(StartPin, VisiblePins);
@@ -37,7 +37,7 @@ void FControlRigConnectionDrawingPolicy::ResetIncompatiblePinDrawState(const TSe
 			UControlRigBlueprint* RigBlueprint = Cast<UControlRigBlueprint>(Blueprint);
 			if (RigBlueprint != nullptr)
 			{
-				RigBlueprint->ModelController->ResetCycleCheck();
+				RigBlueprint->Model->PrepareCycleChecking(nullptr, true);
 			}
 		}
 	}
@@ -141,4 +141,52 @@ void FControlRigConnectionDrawingPolicy::DetermineLinkGeometry(
 	{
 		EndWidgetGeometry = PinGeometries->Find((*pInputWidget).ToSharedRef());
 	}
-}
\ No newline at end of file
+}
+
+void FControlRigConnectionDrawingPolicy::DetermineWiringStyle(UEdGraphPin* OutputPin, UEdGraphPin* InputPin, /*inout*/ FConnectionParams& Params)
+{
+	FKismetConnectionDrawingPolicy::DetermineWiringStyle(OutputPin, InputPin, Params);
+	if (OutputPin == nullptr || InputPin == nullptr)
+	{
+		return;
+	}
+
+	UControlRigGraphNode* OutputNode = Cast<UControlRigGraphNode>(OutputPin->GetOwningNode());
+	UControlRigGraphNode* InputNode = Cast<UControlRigGraphNode>(InputPin->GetOwningNode());
+	if (OutputNode && InputNode)
+	{
+		bool bInjectionIsSelected = false;
+		if (URigVMPin* OutputModelPin = OutputNode->GetModelPinFromPinPath(OutputPin->GetName()))
+		{
+			OutputModelPin = OutputModelPin->GetPinForLink();
+			if (URigVMInjectionInfo* OutputInjection = OutputModelPin->GetNode()->GetInjectionInfo())
+			{
+				if (OutputModelPin->GetNode()->IsSelected())
+				{
+					bInjectionIsSelected = true;
+				}
+			}
+		}
+
+		if (!bInjectionIsSelected)
+		{
+			if (URigVMPin* InputModelPin = InputNode->GetModelPinFromPinPath(InputPin->GetName()))
+			{
+				InputModelPin = InputModelPin->GetPinForLink();
+				if (URigVMInjectionInfo* InputInjection = InputModelPin->GetNode()->GetInjectionInfo())
+				{
+					if (InputModelPin->GetNode()->IsSelected())
+					{
+						bInjectionIsSelected = true;
+					}
+				}
+			}
+		}
+
+		if (bInjectionIsSelected)
+		{
+			Params.WireThickness = Settings->TraceAttackWireThickness;
+			Params.WireColor = Settings->TraceAttackColor;
+		}
+	}
+}
diff --git a/Engine/Plugins/Experimental/ControlRig/Source/ControlRigEditor/Private/Graph/ControlRigConnectionDrawingPolicy.h b/Engine/Plugins/Experimental/ControlRig/Source/ControlRigEditor/Private/Graph/ControlRigConnectionDrawingPolicy.h
--- a/Engine/Plugins/Experimental/ControlRig/Source/ControlRigEditor/Private/Graph/ControlRigConnectionDrawingPolicy.h
+++ b/Engine/Plugins/Experimental/ControlRig/Source/ControlRigEditor/Private/Graph/ControlRigConnectionDrawingPolicy.h
@@ -23,4 +23,5 @@ public:
 	UEdGraphPin* InputPin,
 	/*out*/ FArrangedWidget*& StartWidgetGeometry,
 	/*out*/ FArrangedWidget*& EndWidgetGeometry) override;
+	virtual void DetermineWiringStyle(UEdGraphPin* OutputPin, UEdGraphPin* InputPin, /*inout*/ FConnectionParams& Params) override;
 };
diff --git a/Engine/Plugins/Experimental/ControlRig/Source/ControlRigEditor/Private/Graph/ControlRigGraphPanelNodeFactory.cpp b/Engine/Plugins/Experimental/ControlRig/Source/ControlRigEditor/Private/Graph/ControlRigGraphPanelNodeFactory.cpp
--- a/Engine/Plugins/Experimental/ControlRig/Source/ControlRigEditor/Private/Graph/ControlRigGraphPanelNodeFactory.cpp
+++ b/Engine/Plugins/Experimental/ControlRig/Source/ControlRigEditor/Private/Graph/ControlRigGraphPanelNodeFactory.cpp
@@ -7,11 +7,21 @@
 #include "SControlRigGraphNode.h"
 #include "SControlRigGraphNodeComment.h"
 #include "Graph/ControlRigGraphSchema.h"
+#include "SGraphNodeKnot.h"
+#include "RigVMModel/RigVMNode.h"
+#include "RigVMModel/Nodes/RigVMRerouteNode.h"
 
 TSharedPtr<SGraphNode> FControlRigGraphPanelNodeFactory::CreateNode(UEdGraphNode* Node) const
 {
 	if (UControlRigGraphNode* ControlRigGraphNode = Cast<UControlRigGraphNode>(Node))
 	{
+		int32 InputPin = 0, OutputPin = 0;
+		if (Node->ShouldDrawNodeAsControlPointOnly(InputPin, OutputPin))
+		{
+			TSharedRef<SGraphNode> GraphNode = SNew(SGraphNodeKnot, Node);
+			return GraphNode;
+		}
+
 		TSharedRef<SGraphNode> GraphNode = 
 			SNew(SControlRigGraphNode)
 			.GraphNodeObj(ControlRigGraphNode);
diff --git a/Engine/Plugins/Experimental/ControlRig/Source/ControlRigEditor/Private/Graph/ControlRigGraphPanelPinFactory.cpp b/Engine/Plugins/Experimental/ControlRig/Source/ControlRigEditor/Private/Graph/ControlRigGraphPanelPinFactory.cpp
--- a/Engine/Plugins/Experimental/ControlRig/Source/ControlRigEditor/Private/Graph/ControlRigGraphPanelPinFactory.cpp
+++ b/Engine/Plugins/Experimental/ControlRig/Source/ControlRigEditor/Private/Graph/ControlRigGraphPanelPinFactory.cpp
@@ -6,69 +6,65 @@
 #include "Graph/ControlRigGraph.h"
 #include "Graph/SControlRigGraphPinNameList.h"
 #include "Graph/SControlRigGraphPinCurveFloat.h"
+#include "Graph/SControlRigGraphPinVariableName.h"
+#include "Graph/SControlRigGraphPinParameterName.h"
 #include "KismetPins/SGraphPinExec.h"
 #include "ControlRig.h"
 #include "NodeFactory.h"
 #include "EdGraphSchema_K2.h"
 #include "Curves/CurveFloat.h"
+#include "RigVMModel/Nodes/RigVMStructNode.h"
+#include "RigVMCore/RigVMExecuteContext.h"
 
 TSharedPtr<SGraphPin> FControlRigGraphPanelPinFactory::CreatePin(UEdGraphPin* InPin) const
 {
 	if (InPin)
 	{
-		UControlRigGraphNode* RigNode = Cast<UControlRigGraphNode>(InPin->GetOwningNode());
-		if (RigNode)
+		if (UControlRigGraphNode* RigNode = Cast<UControlRigGraphNode>(InPin->GetOwningNode()))
 		{
 			UControlRigGraph* RigGraph = Cast<UControlRigGraph>(RigNode->GetGraph());
 
-			// use a bone name widget in case we are looking at a name with appropriate metadata
-			if (InPin->PinType.PinCategory == UEdGraphSchema_K2::PC_Name)
+			if (URigVMPin* ModelPin = RigNode->GetModelPinFromPinPath(InPin->GetName()))
 			{
-				UStruct* Struct = nullptr;
-				if (InPin->ParentPin != nullptr)
+				FName CustomWidgetName = ModelPin->GetCustomWidgetName();
+				if (CustomWidgetName == TEXT("BoneName"))
 				{
-					Struct = Cast<UStruct>(InPin->ParentPin->PinType.PinSubCategoryObject);
+					return SNew(SControlRigGraphPinNameList, InPin)
+						.OnGetNameListContent_UObject(RigGraph, &UControlRigGraph::GetBoneNameList);
 				}
-				if (Struct == nullptr)
+				else if (CustomWidgetName == TEXT("ControlName"))
 				{
-					Struct = RigNode->GetUnitScriptStruct();
+					return SNew(SControlRigGraphPinNameList, InPin)
+						.OnGetNameListContent_UObject(RigGraph, &UControlRigGraph::GetControlNameList);
 				}
-				if (Struct)
+				else if (CustomWidgetName == TEXT("SpaceName"))
 				{
-					FString NodeName, PropertyName;
-					if (InPin->GetName().Split(TEXT("."), &NodeName, &PropertyName, ESearchCase::IgnoreCase, ESearchDir::FromEnd))
-					{
-						FProperty* Property = Struct->FindPropertyByName(*PropertyName);
-						if (Property)
-						{
-							if (Property->HasMetaData(UControlRig::BoneNameMetaName))
-							{
-								return SNew(SControlRigGraphPinNameList, InPin)
-									.OnGetNameListContent_UObject(RigGraph, &UControlRigGraph::GetBoneNameList);
-							}
-							if (Property->HasMetaData(UControlRig::ControlNameMetaName))
-							{
-								return SNew(SControlRigGraphPinNameList, InPin)
-									.OnGetNameListContent_UObject(RigGraph, &UControlRigGraph::GetControlNameList);
-							}
-							if (Property->HasMetaData(UControlRig::SpaceNameMetaName))
-							{
-								return SNew(SControlRigGraphPinNameList, InPin)
-									.OnGetNameListContent_UObject(RigGraph, &UControlRigGraph::GetSpaceNameList);
-							}
-							else if (Property->HasMetaData(UControlRig::CurveNameMetaName))
-							{
-								return SNew(SControlRigGraphPinNameList, InPin)
-									.OnGetNameListContent_UObject(RigGraph, &UControlRigGraph::GetCurveNameList);
-							}
-						}
-					}
+					return SNew(SControlRigGraphPinNameList, InPin)
+						.OnGetNameListContent_UObject(RigGraph, &UControlRigGraph::GetSpaceNameList);
+				}
+				else if (CustomWidgetName == TEXT("CurveName"))
+				{
+					return SNew(SControlRigGraphPinNameList, InPin)
+						.OnGetNameListContent_UObject(RigGraph, &UControlRigGraph::GetCurveNameList);
+				}
+				else if (CustomWidgetName == TEXT("DrawingName"))
+				{
+					return SNew(SControlRigGraphPinNameList, InPin)
+						.OnGetNameListContent_UObject(RigGraph, &UControlRigGraph::GetDrawingNameList);
+				}
+				else if (CustomWidgetName == TEXT("VariableName"))
+				{
+					return SNew(SControlRigGraphPinVariableName, InPin);
+				}
+				else if (CustomWidgetName == TEXT("ParameterName"))
+				{
+					return SNew(SControlRigGraphPinParameterName, InPin);
 				}
 			}
 
-			else if (InPin->PinType.PinCategory == UEdGraphSchema_K2::PC_Struct)
+			if (InPin->PinType.PinCategory == UEdGraphSchema_K2::PC_Struct)
 			{
-				if (InPin->PinType.PinSubCategoryObject == FControlRigExecuteContext::StaticStruct())
+				if (Cast<UStruct>(InPin->PinType.PinSubCategoryObject)->IsChildOf(FRigVMExecuteContext::StaticStruct()))
 				{
 					return SNew(SGraphPinExec, InPin);
 				}
diff --git a/Engine/Plugins/Experimental/ControlRig/Source/ControlRigEditor/Private/Graph/NodeSpawners/ControlRigParameterNodeSpawner.cpp b/Engine/Plugins/Experimental/ControlRig/Source/ControlRigEditor/Private/Graph/NodeSpawners/ControlRigParameterNodeSpawner.cpp
new file mode 100644
--- /dev/null
+++ b/Engine/Plugins/Experimental/ControlRig/Source/ControlRigEditor/Private/Graph/NodeSpawners/ControlRigParameterNodeSpawner.cpp
@@ -0,0 +1,167 @@
+// Copyright Epic Games, Inc. All Rights Reserved.
+
+#include "ControlRigParameterNodeSpawner.h"
+#include "Graph/ControlRigGraph.h"
+#include "Graph/ControlRigGraphNode.h"
+#include "EdGraphSchema_K2.h"
+#include "Kismet2/BlueprintEditorUtils.h"
+#include "Classes/EditorStyleSettings.h"
+#include "Editor/EditorEngine.h"
+#include "ObjectEditorUtils.h"
+#include "EditorCategoryUtils.h"
+#include "K2Node_Variable.h"
+#include "BlueprintNodeTemplateCache.h"
+#include "RigVMModel/RigVMController.h"
+#include "ControlRigBlueprint.h"
+#include "ControlRigBlueprintUtils.h"
+#include "RigVMModel/Nodes/RigVMParameterNode.h"
+
+#if WITH_EDITOR
+#include "Editor.h"
+#endif
+
+#define LOCTEXT_NAMESPACE "ControlRigParameterNodeSpawner"
+
+UControlRigParameterNodeSpawner* UControlRigParameterNodeSpawner::CreateFromPinType(const FEdGraphPinType& InPinType, bool bInIsGetter, const FText& InMenuDesc, const FText& InCategory, const FText& InTooltip)
+{
+	UControlRigParameterNodeSpawner* NodeSpawner = NewObject<UControlRigParameterNodeSpawner>(GetTransientPackage());
+	NodeSpawner->EdGraphPinType = InPinType;
+	NodeSpawner->bIsGetter = bInIsGetter;
+	NodeSpawner->NodeClass = UControlRigGraphNode::StaticClass();
+
+	FBlueprintActionUiSpec& MenuSignature = NodeSpawner->DefaultMenuSignature;
+	
+	MenuSignature.MenuName = FText::FromString(FString::Printf(TEXT("%s %s"), bInIsGetter ? TEXT("Get") : TEXT("Set"), *InMenuDesc.ToString()));
+	MenuSignature.Tooltip  = InTooltip;
+	MenuSignature.Category = InCategory;
+	MenuSignature.Keywords = FText::FromString(TEXT("Parameter"));
+	MenuSignature.Icon = UK2Node_Variable::GetVarIconFromPinType(InPinType, MenuSignature.IconTint);
+
+	return NodeSpawner;
+}
+
+void UControlRigParameterNodeSpawner::Prime()
+{
+	// we expect that you don't need a node template to construct menu entries
+	// from this, so we choose not to pre-cache one here
+}
+
+FBlueprintNodeSignature UControlRigParameterNodeSpawner::GetSpawnerSignature() const
+{
+	return FBlueprintNodeSignature(NodeClass);
+}
+
+FBlueprintActionUiSpec UControlRigParameterNodeSpawner::GetUiSpec(FBlueprintActionContext const& Context, FBindingSet const& Bindings) const
+{
+	UEdGraph* TargetGraph = (Context.Graphs.Num() > 0) ? Context.Graphs[0] : nullptr;
+	FBlueprintActionUiSpec MenuSignature = PrimeDefaultUiSpec(TargetGraph);
+
+	DynamicUiSignatureGetter.ExecuteIfBound(Context, Bindings, &MenuSignature);
+	return MenuSignature;
+}
+
+UEdGraphNode* UControlRigParameterNodeSpawner::Invoke(UEdGraph* ParentGraph, FBindingSet const& Bindings, FVector2D const Location) const
+{
+	UControlRigGraphNode* NewNode = nullptr;
+
+	bool const bIsTemplateNode = FBlueprintNodeTemplateCache::IsTemplateOuter(ParentGraph);
+
+	// First create a backing member for our node
+	UControlRigGraph* RigGraph = Cast<UControlRigGraph>(ParentGraph);
+	check(RigGraph);
+	UControlRigBlueprint* RigBlueprint = Cast<UControlRigBlueprint>(ParentGraph->GetOuter());
+	check(RigBlueprint);
+
+	FName MemberName = NAME_None;
+
+#if WITH_EDITOR
+	if (GEditor && !bIsTemplateNode)
+	{
+		GEditor->CancelTransaction(0);
+	}
+#endif
+
+	URigVMController* Controller = bIsTemplateNode ? RigGraph->GetTemplateController() : RigBlueprint->Controller;
+
+	FName DataType = EdGraphPinType.PinCategory;
+	if (DataType == UEdGraphSchema_K2::PC_Int)
+	{
+		DataType = TEXT("int32");
+	}
+	else if (DataType == UEdGraphSchema_K2::PC_Name)
+	{
+		DataType = TEXT("FName");
+	}
+	else if (DataType == UEdGraphSchema_K2::PC_String)
+	{
+		DataType = TEXT("FString");
+	}
+	else if (UStruct* Struct = Cast<UStruct>(EdGraphPinType.PinSubCategoryObject))
+	{
+		DataType = *FString::Printf(TEXT("F%s"), *Struct->GetFName().ToString());
+	}
+	else if (UEnum* Enum = Cast<UEnum>(EdGraphPinType.PinSubCategoryObject))
+	{
+		DataType = *FString::Printf(TEXT("E%s"), *Enum->GetName());
+	}
+
+	FString DataTypeForParameterName = DataType.ToString();
+	if (DataTypeForParameterName.StartsWith(TEXT("F"), ESearchCase::CaseSensitive) || DataTypeForParameterName.StartsWith(TEXT("E"), ESearchCase::CaseSensitive))
+	{
+		DataTypeForParameterName = DataTypeForParameterName.RightChop(1);
+	}
+	DataTypeForParameterName = DataTypeForParameterName.Left(1).ToUpper() + DataTypeForParameterName.RightChop(1);
+	FString ParameterNamePrefix = FString::Printf(TEXT("%sParam"), *DataTypeForParameterName);
+
+	TMap<FName, int32> NameToIndex;
+	TArray<FRigVMGraphParameterDescription> ExistingParameters = Controller->GetGraph()->GetParameterDescriptions();
+	for (int32 ParameterIndex = 0; ParameterIndex < ExistingParameters.Num(); ParameterIndex++)
+	{
+		NameToIndex.Add(ExistingParameters[ParameterIndex].Name, ParameterIndex);
+	}
+
+	FName ParameterName = URigVMController::GetUniqueName(*ParameterNamePrefix, [NameToIndex](const FName& InName) {
+		return !NameToIndex.Contains(InName);
+	});
+
+	if (!bIsTemplateNode)
+	{
+		Controller->OpenUndoBracket(FString::Printf(TEXT("Add '%s' Parameter"), *DataType.ToString()));
+	}
+
+	if (URigVMNode* ModelNode = Controller->AddParameterNodeFromObjectPath(ParameterName, DataType.ToString(), FString(), bIsGetter, FString(), Location, FString(), !bIsTemplateNode))
+	{
+		for (UEdGraphNode* Node : ParentGraph->Nodes)
+		{
+			if (UControlRigGraphNode* RigNode = Cast<UControlRigGraphNode>(Node))
+			{
+				if (RigNode->GetModelNodeName() == ModelNode->GetFName())
+				{
+					NewNode = RigNode;
+					break;
+				}
+			}
+		}
+
+		if (!bIsTemplateNode)
+		{
+			if (NewNode)
+			{
+				Controller->ClearNodeSelection(true);
+				Controller->SelectNode(ModelNode, true, true);
+			}
+			Controller->CloseUndoBracket();
+		}
+	}
+	else
+	{
+		if (!bIsTemplateNode)
+		{
+			Controller->CancelUndoBracket();
+		}
+	}
+
+	return NewNode;
+}
+
+#undef LOCTEXT_NAMESPACE
diff --git a/Engine/Plugins/Experimental/ControlRig/Source/ControlRigEditor/Private/Graph/NodeSpawners/ControlRigParameterNodeSpawner.h b/Engine/Plugins/Experimental/ControlRig/Source/ControlRigEditor/Private/Graph/NodeSpawners/ControlRigParameterNodeSpawner.h
new file mode 100644
--- /dev/null
+++ b/Engine/Plugins/Experimental/ControlRig/Source/ControlRigEditor/Private/Graph/NodeSpawners/ControlRigParameterNodeSpawner.h
@@ -0,0 +1,46 @@
+// Copyright Epic Games, Inc. All Rights Reserved.
+
+#pragma once
+
+#include "CoreMinimal.h"
+#include "UObject/ObjectMacros.h"
+#include "UObject/Object.h"
+#include "Templates/SubclassOf.h"
+#include "UObject/UnrealType.h"
+#include "EdGraph/EdGraphPin.h"
+#include "Engine/Blueprint.h"
+#include "BlueprintActionFilter.h"
+#include "BlueprintNodeSignature.h"
+#include "BlueprintFieldNodeSpawner.h"
+#include "ControlRigParameterNodeSpawner.generated.h"
+
+class UControlRigGraphNode;
+
+UCLASS(Transient)
+class CONTROLRIGEDITOR_API UControlRigParameterNodeSpawner : public UBlueprintNodeSpawner
+{
+	GENERATED_BODY()
+
+public:
+	/**
+	 * Creates a new UControlRigParameterNodeSpawner, charged with spawning 
+	 * a new member-parameter node
+	 * 
+	 * @return A newly allocated instance of this class.
+	 */
+	static UControlRigParameterNodeSpawner* CreateFromPinType(const FEdGraphPinType& InPinType, bool bInIsGetter, const FText& InMenuDesc, const FText& InCategory, const FText& InTooltip);
+
+	// UBlueprintNodeSpawner interface
+	virtual void Prime() override;
+	virtual FBlueprintNodeSignature GetSpawnerSignature() const override;
+	virtual FBlueprintActionUiSpec GetUiSpec(FBlueprintActionContext const& Context, FBindingSet const& Bindings) const override;
+	virtual UEdGraphNode* Invoke(UEdGraph* ParentGraph, FBindingSet const& Bindings, FVector2D const Location) const override;
+	// End UBlueprintNodeSpawner interface
+
+private:
+	/** The pin type we will spawn */
+	FEdGraphPinType EdGraphPinType;
+	bool bIsGetter;
+
+	friend class UEngineTestControlRig;
+};
diff --git a/Engine/Plugins/Experimental/ControlRig/Source/ControlRigEditor/Private/Graph/NodeSpawners/ControlRigRerouteNodeSpawner.cpp b/Engine/Plugins/Experimental/ControlRig/Source/ControlRigEditor/Private/Graph/NodeSpawners/ControlRigRerouteNodeSpawner.cpp
new file mode 100644
--- /dev/null
+++ b/Engine/Plugins/Experimental/ControlRig/Source/ControlRigEditor/Private/Graph/NodeSpawners/ControlRigRerouteNodeSpawner.cpp
@@ -0,0 +1,170 @@
+// Copyright Epic Games, Inc. All Rights Reserved.
+
+#include "ControlRigRerouteNodeSpawner.h"
+#include "Graph/ControlRigGraph.h"
+#include "Graph/ControlRigGraphNode.h"
+#include "Graph/ControlRigGraphSchema.h"
+#include "EdGraphSchema_K2.h"
+#include "Kismet2/BlueprintEditorUtils.h"
+#include "Classes/EditorStyleSettings.h"
+#include "Editor/EditorEngine.h"
+#include "ObjectEditorUtils.h"
+#include "EditorCategoryUtils.h"
+#include "K2Node_Variable.h"
+#include "BlueprintNodeTemplateCache.h"
+#include "RigVMModel/RigVMController.h"
+#include "ControlRigBlueprint.h"
+#include "ControlRigBlueprintUtils.h"
+#include "RigVMModel/Nodes/RigVMRerouteNode.h"
+
+#if WITH_EDITOR
+#include "Editor.h"
+#endif
+
+#define LOCTEXT_NAMESPACE "ControlRigRerouteNodeSpawner"
+
+UControlRigRerouteNodeSpawner* UControlRigRerouteNodeSpawner::CreateGeneric(const FText& InMenuDesc, const FText& InCategory, const FText& InTooltip)
+{
+	UControlRigRerouteNodeSpawner* NodeSpawner = NewObject<UControlRigRerouteNodeSpawner>(GetTransientPackage());
+	NodeSpawner->NodeClass = UControlRigGraphNode::StaticClass();
+
+	FBlueprintActionUiSpec& MenuSignature = NodeSpawner->DefaultMenuSignature;
+	
+	MenuSignature.MenuName = InMenuDesc;
+	MenuSignature.Tooltip  = InTooltip;
+	MenuSignature.Category = InCategory;
+	MenuSignature.Keywords = FText::FromString(TEXT("Reroute,Elbow,Wire,Literal,Make Literal,Constant"));
+
+	return NodeSpawner;
+}
+
+void UControlRigRerouteNodeSpawner::Prime()
+{
+	// we expect that you don't need a node template to construct menu entries
+	// from this, so we choose not to pre-cache one here
+}
+
+FBlueprintNodeSignature UControlRigRerouteNodeSpawner::GetSpawnerSignature() const
+{
+	return FBlueprintNodeSignature(NodeClass);
+}
+
+FBlueprintActionUiSpec UControlRigRerouteNodeSpawner::GetUiSpec(FBlueprintActionContext const& Context, FBindingSet const& Bindings) const
+{
+	UEdGraph* TargetGraph = (Context.Graphs.Num() > 0) ? Context.Graphs[0] : nullptr;
+	FBlueprintActionUiSpec MenuSignature = PrimeDefaultUiSpec(TargetGraph);
+
+	DynamicUiSignatureGetter.ExecuteIfBound(Context, Bindings, &MenuSignature);
+	return MenuSignature;
+}
+
+UEdGraphNode* UControlRigRerouteNodeSpawner::Invoke(UEdGraph* ParentGraph, FBindingSet const& Bindings, FVector2D const Location) const
+{
+	UControlRigGraphNode* NewNode = nullptr;
+
+	bool const bIsTemplateNode = FBlueprintNodeTemplateCache::IsTemplateOuter(ParentGraph);
+
+	// First create a backing member for our node
+	UControlRigGraph* RigGraph = Cast<UControlRigGraph>(ParentGraph);
+	check(RigGraph);
+	UControlRigBlueprint* RigBlueprint = Cast<UControlRigBlueprint>(ParentGraph->GetOuter());
+	check(RigBlueprint);
+
+	FName MemberName = NAME_None;
+
+#if WITH_EDITOR
+	if (GEditor && !bIsTemplateNode)
+	{
+		GEditor->CancelTransaction(0);
+	}
+#endif
+
+	if (bIsTemplateNode)
+	{
+		NewNode = NewObject<UControlRigGraphNode>(ParentGraph, TEXT("RerouteNode"));
+		ParentGraph->AddNode(NewNode, false);
+		
+		NewNode->CreateNewGuid();
+		NewNode->PostPlacedNewNode();
+
+		UEdGraphPin* InputValuePin = UEdGraphPin::CreatePin(NewNode);
+		UEdGraphPin* OutputValuePin = UEdGraphPin::CreatePin(NewNode);
+		NewNode->Pins.Add(InputValuePin);
+		NewNode->Pins.Add(OutputValuePin);
+
+		InputValuePin->PinType.PinCategory = TEXT("REROUTE");
+		OutputValuePin->PinType.PinCategory = TEXT("REROUTE");
+		InputValuePin->Direction = EGPD_Input;
+		OutputValuePin->Direction = EGPD_Output;
+		NewNode->SetFlags(RF_Transactional);
+
+		return NewNode;
+	}
+	else
+	{
+		URigVMController* Controller = RigBlueprint->Controller;
+		Controller->OpenUndoBracket(TEXT("Added Reroute Node."));
+
+		FString PinPath;
+		bool bIsInput = false;
+
+		if (UControlRigGraphSchema* RigSchema = Cast<UControlRigGraphSchema>((UEdGraphSchema*)ParentGraph->GetSchema()))
+		{
+			bIsInput = RigSchema->bLastPinWasInput;
+
+			if (const UEdGraphPin* LastPin = RigSchema->LastPinForCompatibleCheck)
+			{
+				if (URigVMPin* ModelPin = RigBlueprint->Model->FindPin(LastPin->GetName()))
+				{
+					PinPath = ModelPin->GetPinPath();
+				}
+			}
+		}
+
+		if (URigVMNode* ModelNode = Controller->AddRerouteNodeOnPin(PinPath, bIsInput, true, Location, FString(), true))
+		{
+			for (UEdGraphNode* Node : ParentGraph->Nodes)
+			{
+				if (UControlRigGraphNode* RigNode = Cast<UControlRigGraphNode>(Node))
+				{
+					if (RigNode->GetModelNodeName() == ModelNode->GetFName())
+					{
+						NewNode = RigNode;
+						break;
+					}
+				}
+			}
+
+			if (NewNode)
+			{
+				Controller->ClearNodeSelection(true);
+				Controller->SelectNode(ModelNode, true, true);
+			}
+			Controller->CloseUndoBracket();
+		}
+		else
+		{
+			Controller->CancelUndoBracket();
+		}
+	}
+
+	return NewNode;
+}
+
+bool UControlRigRerouteNodeSpawner::IsTemplateNodeFilteredOut(FBlueprintActionFilter const& Filter) const
+{
+	for (UBlueprint* Blueprint : Filter.Context.Blueprints)
+	{
+		if (UControlRigBlueprint* RigBlueprint = Cast<UControlRigBlueprint>(Blueprint))
+		{
+			if (Filter.Context.Pins.Num() == 0)
+			{
+				return true;
+			}
+			return false;
+		}
+	}
+	return true;
+}
+
+#undef LOCTEXT_NAMESPACE
diff --git a/Engine/Plugins/Experimental/ControlRig/Source/ControlRigEditor/Private/Graph/NodeSpawners/ControlRigRerouteNodeSpawner.h b/Engine/Plugins/Experimental/ControlRig/Source/ControlRigEditor/Private/Graph/NodeSpawners/ControlRigRerouteNodeSpawner.h
new file mode 100644
--- /dev/null
+++ b/Engine/Plugins/Experimental/ControlRig/Source/ControlRigEditor/Private/Graph/NodeSpawners/ControlRigRerouteNodeSpawner.h
@@ -0,0 +1,45 @@
+// Copyright Epic Games, Inc. All Rights Reserved.
+
+#pragma once
+
+#include "CoreMinimal.h"
+#include "UObject/ObjectMacros.h"
+#include "UObject/Object.h"
+#include "Templates/SubclassOf.h"
+#include "UObject/UnrealType.h"
+#include "EdGraph/EdGraphPin.h"
+#include "Engine/Blueprint.h"
+#include "BlueprintActionFilter.h"
+#include "BlueprintNodeSignature.h"
+#include "BlueprintFieldNodeSpawner.h"
+#include "ControlRigRerouteNodeSpawner.generated.h"
+
+class UControlRigGraphNode;
+
+UCLASS(Transient)
+class CONTROLRIGEDITOR_API UControlRigRerouteNodeSpawner : public UBlueprintNodeSpawner
+{
+	GENERATED_BODY()
+
+public:
+	/**
+	 * Creates a new UControlRigRerouteNodeSpawner, charged with spawning 
+	 * a new member-variable node
+	 * 
+	 * @return A newly allocated instance of this class.
+	 */
+	static UControlRigRerouteNodeSpawner* CreateGeneric(const FText& InMenuDesc, const FText& InCategory, const FText& InTooltip);
+
+	// UBlueprintNodeSpawner interface
+	virtual void Prime() override;
+	virtual FBlueprintNodeSignature GetSpawnerSignature() const override;
+	virtual FBlueprintActionUiSpec GetUiSpec(FBlueprintActionContext const& Context, FBindingSet const& Bindings) const override;
+	virtual UEdGraphNode* Invoke(UEdGraph* ParentGraph, FBindingSet const& Bindings, FVector2D const Location) const override;
+	virtual bool IsTemplateNodeFilteredOut(FBlueprintActionFilter const& Filter) const override;
+
+	// End UBlueprintNodeSpawner interface
+
+private:
+
+	friend class UEngineTestControlRig;
+};
diff --git a/Engine/Plugins/Experimental/ControlRig/Source/ControlRigEditor/Private/Graph/NodeSpawners/ControlRigUnitNodeSpawner.cpp b/Engine/Plugins/Experimental/ControlRig/Source/ControlRigEditor/Private/Graph/NodeSpawners/ControlRigUnitNodeSpawner.cpp
--- a/Engine/Plugins/Experimental/ControlRig/Source/ControlRigEditor/Private/Graph/NodeSpawners/ControlRigUnitNodeSpawner.cpp
+++ b/Engine/Plugins/Experimental/ControlRig/Source/ControlRigEditor/Private/Graph/NodeSpawners/ControlRigUnitNodeSpawner.cpp
@@ -16,7 +16,6 @@
 #include "ScopedTransaction.h"
 #include "ControlRig/Private/Units/Execution/RigUnit_BeginExecution.h"
 #include "ControlRig.h"
-#include "ControlRigBlueprint.h"
 
 #if WITH_EDITOR
 #include "Editor.h"
@@ -24,7 +23,7 @@
 
 #define LOCTEXT_NAMESPACE "ControlRigUnitNodeSpawner"
 
-UControlRigUnitNodeSpawner* UControlRigUnitNodeSpawner::CreateFromStruct(UStruct* InStruct, const FText& InMenuDesc, const FText& InCategory, const FText& InTooltip)
+UControlRigUnitNodeSpawner* UControlRigUnitNodeSpawner::CreateFromStruct(UScriptStruct* InStruct, const FText& InMenuDesc, const FText& InCategory, const FText& InTooltip)
 {
 	UControlRigUnitNodeSpawner* NodeSpawner = NewObject<UControlRigUnitNodeSpawner>(GetTransientPackage());
 	NodeSpawner->StructTemplate = InStruct;
@@ -121,42 +120,116 @@ bool UControlRigUnitNodeSpawner::IsTemplateNodeFilteredOut(FBlueprintActionFilte
 	return Super::IsTemplateNodeFilteredOut(Filter);
 }
 
-UControlRigGraphNode* UControlRigUnitNodeSpawner::SpawnNode(UEdGraph* ParentGraph, UBlueprint* Blueprint, UStruct* StructTemplate, FVector2D const Location)
+UControlRigGraphNode* UControlRigUnitNodeSpawner::SpawnNode(UEdGraph* ParentGraph, UBlueprint* Blueprint, UScriptStruct* StructTemplate, FVector2D const Location)
 {
 	UControlRigGraphNode* NewNode = nullptr;
 	UControlRigBlueprint* RigBlueprint = Cast<UControlRigBlueprint>(Blueprint);
-	if (RigBlueprint != nullptr)
+	UControlRigGraph* RigGraph = Cast<UControlRigGraph>(ParentGraph);
+
+	if (RigBlueprint != nullptr && RigGraph != nullptr)
 	{
 		bool const bIsTemplateNode = FBlueprintNodeTemplateCache::IsTemplateOuter(ParentGraph);
+		bool const bUndo = !bIsTemplateNode;
 
-		// First create a backing member for our node
-		FName MemberName = NAME_None;
+		FName Name = bIsTemplateNode ? *StructTemplate->GetDisplayNameText().ToString() : FControlRigBlueprintUtils::ValidateName(RigBlueprint, StructTemplate->GetFName().ToString());
+		URigVMController* Controller = bIsTemplateNode ? RigGraph->GetTemplateController() : RigBlueprint->Controller;
 
 		if (!bIsTemplateNode)
 		{
-			FName Name = FControlRigBlueprintUtils::ValidateName(RigBlueprint, StructTemplate->GetFName().ToString());
-			if (RigBlueprint->ModelController->AddNode(StructTemplate->GetFName(), Location, Name))
+			Controller->OpenUndoBracket(FString::Printf(TEXT("Add '%s' Node"), *Name.ToString()));
+		}
+
+		if (URigVMStructNode* ModelNode = Controller->AddStructNode(StructTemplate, TEXT("Execute"), Location, Name.ToString(), bUndo))
+		{
+			NewNode = Cast<UControlRigGraphNode>(RigGraph->FindNodeForModelNodeName(ModelNode->GetFName()));
+
+			if (NewNode && bUndo)
+			{
+				Controller->ClearNodeSelection(true);
+				Controller->SelectNode(ModelNode, true, true);
+
+				HookupMutableNode(ModelNode, RigBlueprint);
+			}
+
+			if (bUndo)
+			{
+				Controller->CloseUndoBracket();
+			}
+		}
+		else
+		{
+			if (bUndo)
+			{
+				Controller->CancelUndoBracket();
+			}
+		}
+	}
+	return NewNode;
+}
+
+void UControlRigUnitNodeSpawner::HookupMutableNode(URigVMStructNode* InModelNode, UControlRigBlueprint* InRigBlueprint)
+{
+	URigVMController* Controller = InRigBlueprint->Controller;
+
+	Controller->ClearNodeSelection(true);
+	Controller->SelectNode(InModelNode, true, true);
+
+	// see if the node has an execute pin
+	URigVMPin* ModelNodeExecutePin = nullptr;
+	for (URigVMPin* Pin : InModelNode->GetPins())
+	{
+		if (Pin->GetScriptStruct())
+		{
+			if (Pin->GetScriptStruct()->IsChildOf(FRigVMExecuteContext::StaticStruct()))
 			{
-				MemberName = RigBlueprint->LastNameFromNotification;
-				for (UEdGraphNode* Node : ParentGraph->Nodes)
+				if (Pin->GetDirection() == ERigVMPinDirection::IO)
 				{
-					if (UControlRigGraphNode* RigNode = Cast<UControlRigGraphNode>(Node))
+					ModelNodeExecutePin = Pin;
+					break;
+				}
+			}
+		}
+	}
+
+	// we have an execute pin - so we have to hook it up
+	if (ModelNodeExecutePin)
+	{
+		URigVMPin* ClosestOtherModelNodeExecutePin = nullptr;
+		float ClosestDistance = FLT_MAX;
+
+		for (URigVMNode* OtherModelNode : Controller->GetGraph()->GetNodes())
+		{
+			if (OtherModelNode == InModelNode)
+			{
+				continue;
+			}
+
+			for (URigVMPin* Pin : OtherModelNode->GetPins())
+			{
+				if (Pin->GetScriptStruct())
+				{
+					if (Pin->GetScriptStruct()->IsChildOf(FRigVMExecuteContext::StaticStruct()))
 					{
-						if (RigNode->GetPropertyName() == MemberName)
+						if (Pin->GetDirection() == ERigVMPinDirection::IO || Pin->GetDirection() == ERigVMPinDirection::Output)
 						{
-							NewNode = RigNode;
+							float Distance = (OtherModelNode->GetPosition() - InModelNode->GetPosition()).Size();
+							if (Distance < ClosestDistance)
+							{
+								ClosestOtherModelNodeExecutePin = Pin;
+								ClosestDistance = Distance;
+							}
 							break;
 						}
 					}
 				}
 			}
 		}
-		else
+
+		if (ClosestOtherModelNodeExecutePin)
 		{
-			NewNode = FControlRigBlueprintUtils::InstantiateGraphNodeForStructPath(ParentGraph, *StructTemplate->GetDisplayNameText().ToString(), Location, StructTemplate->GetPathName());
+			Controller->AddLink(ClosestOtherModelNodeExecutePin->GetPinPath(), ModelNodeExecutePin->GetPinPath(), true);
 		}
 	}
-	return NewNode;
 }
 
 #undef LOCTEXT_NAMESPACE
diff --git a/Engine/Plugins/Experimental/ControlRig/Source/ControlRigEditor/Private/Graph/NodeSpawners/ControlRigUnitNodeSpawner.h b/Engine/Plugins/Experimental/ControlRig/Source/ControlRigEditor/Private/Graph/NodeSpawners/ControlRigUnitNodeSpawner.h
--- a/Engine/Plugins/Experimental/ControlRig/Source/ControlRigEditor/Private/Graph/NodeSpawners/ControlRigUnitNodeSpawner.h
+++ b/Engine/Plugins/Experimental/ControlRig/Source/ControlRigEditor/Private/Graph/NodeSpawners/ControlRigUnitNodeSpawner.h
@@ -12,6 +12,10 @@
 #include "BlueprintActionFilter.h"
 #include "BlueprintNodeSignature.h"
 #include "BlueprintFieldNodeSpawner.h"
+#include "ControlRigBlueprint.h"
+#include "Graph/ControlRigGraph.h"
+#include "Graph/ControlRigGraphNode.h"
+#include "RigVMModel/Nodes/RigVMStructNode.h"
 #include "ControlRigUnitNodeSpawner.generated.h"
 
 class UControlRigGraphNode;
@@ -28,7 +32,7 @@ public:
 	 * 
 	 * @return A newly allocated instance of this class.
 	 */
-	static UControlRigUnitNodeSpawner* CreateFromStruct(UStruct* InStruct, const FText& InMenuDesc, const FText& InCategory, const FText& InTooltip);
+	static UControlRigUnitNodeSpawner* CreateFromStruct(UScriptStruct* InStruct, const FText& InMenuDesc, const FText& InCategory, const FText& InTooltip);
 
 	// UBlueprintNodeSpawner interface
 	virtual void Prime() override;
@@ -36,14 +40,15 @@ public:
 	virtual FBlueprintActionUiSpec GetUiSpec(FBlueprintActionContext const& Context, FBindingSet const& Bindings) const override;
 	virtual UEdGraphNode* Invoke(UEdGraph* ParentGraph, FBindingSet const& Bindings, FVector2D const Location) const override;
 	virtual bool IsTemplateNodeFilteredOut(FBlueprintActionFilter const& Filter) const override;
+	static void HookupMutableNode(URigVMStructNode* InModelNode, UControlRigBlueprint* InRigBlueprint);
 	// End UBlueprintNodeSpawner interface
 
 private:
 	/** The unit type we will spawn */
 	UPROPERTY(Transient)
-	UStruct* StructTemplate;
+	UScriptStruct* StructTemplate;
 
-	static UControlRigGraphNode* SpawnNode(UEdGraph* ParentGraph, UBlueprint* Blueprint, UStruct* StructTemplate, FVector2D const Location);
+	static UControlRigGraphNode* SpawnNode(UEdGraph* ParentGraph, UBlueprint* Blueprint, UScriptStruct* StructTemplate, FVector2D const Location);
 
 	friend class UEngineTestControlRig;
 	friend class FControlRigEditor;
