COPY from //Dev-Anim to //Dev-Main

diff --git a/Engine/Plugins/Experimental/ControlRig/Source/ControlRigEditor/Private/Graph/NodeSpawners/ControlRigVariableNodeSpawner.cpp b/Engine/Plugins/Experimental/ControlRig/Source/ControlRigEditor/Private/Graph/NodeSpawners/ControlRigVariableNodeSpawner.cpp
--- a/Engine/Plugins/Experimental/ControlRig/Source/ControlRigEditor/Private/Graph/NodeSpawners/ControlRigVariableNodeSpawner.cpp
+++ b/Engine/Plugins/Experimental/ControlRig/Source/ControlRigEditor/Private/Graph/NodeSpawners/ControlRigVariableNodeSpawner.cpp
@@ -1,6 +1,7 @@
 // Copyright Epic Games, Inc. All Rights Reserved.
 
 #include "ControlRigVariableNodeSpawner.h"
+#include "Graph/ControlRigGraph.h"
 #include "Graph/ControlRigGraphNode.h"
 #include "EdGraphSchema_K2.h"
 #include "Kismet2/BlueprintEditorUtils.h"
@@ -10,10 +11,10 @@
 #include "EditorCategoryUtils.h"
 #include "K2Node_Variable.h"
 #include "BlueprintNodeTemplateCache.h"
-#include "ControlRigController.h"
+#include "RigVMModel/RigVMController.h"
 #include "ControlRigBlueprint.h"
 #include "ControlRigBlueprintUtils.h"
-#include "Units/RigUnit.h"
+#include "RigVMModel/Nodes/RigVMVariableNode.h"
 
 #if WITH_EDITOR
 #include "Editor.h"
@@ -21,40 +22,20 @@
 
 #define LOCTEXT_NAMESPACE "ControlRigVariableNodeSpawner"
 
-const TArray<FString> GControlRigVariableNodeSpawnerAllowedStructTypes = {
-	TEXT("FBox"),
-	TEXT("FBox2D"),
-	TEXT("FColor"),
-	TEXT("FLinearColor"),
-	TEXT("FVector"),
-	TEXT("FVector2D"),
-	TEXT("FVector4"),
-	TEXT("FRotator"),
-	TEXT("FQuat"),
-	TEXT("FPlane"),
-	TEXT("FMatrix"),
-	TEXT("FRotationMatrix"),
-	TEXT("FScaleMatrix"),
-	TEXT("FTransform"),
-	TEXT("FEulerTransform")
-};
-
-const TArray<FString> GControlRigVariableNodeSpawnerAllowedEnumTypes = {
-};
-
-UControlRigVariableNodeSpawner* UControlRigVariableNodeSpawner::CreateFromPinType(const FEdGraphPinType& InPinType, const FText& InMenuDesc, const FText& InCategory, const FText& InTooltip)
+UControlRigVariableNodeSpawner* UControlRigVariableNodeSpawner::CreateFromPinType(const FEdGraphPinType& InPinType, bool bInIsGetter, const FText& InMenuDesc, const FText& InCategory, const FText& InTooltip)
 {
 	UControlRigVariableNodeSpawner* NodeSpawner = NewObject<UControlRigVariableNodeSpawner>(GetTransientPackage());
 	NodeSpawner->EdGraphPinType = InPinType;
+	NodeSpawner->bIsGetter = bInIsGetter;
 	NodeSpawner->NodeClass = UControlRigGraphNode::StaticClass();
 
 	FBlueprintActionUiSpec& MenuSignature = NodeSpawner->DefaultMenuSignature;
 	
-	MenuSignature.MenuName = InMenuDesc;
+	MenuSignature.MenuName = FText::FromString(FString::Printf(TEXT("%s %s"), bInIsGetter ? TEXT("Get") : TEXT("Set"), *InMenuDesc.ToString()));
 	MenuSignature.Tooltip  = InTooltip;
 	MenuSignature.Category = InCategory;
 	MenuSignature.Keywords = FText::FromString(TEXT("Variable"));
-	MenuSignature.Icon = UK2Node_Variable::GetVarIconFromPinType(NodeSpawner->GetVarType(), MenuSignature.IconTint);
+	MenuSignature.Icon = UK2Node_Variable::GetVarIconFromPinType(InPinType, MenuSignature.IconTint);
 
 	return NodeSpawner;
 }
@@ -86,107 +67,102 @@ UEdGraphNode* UControlRigVariableNodeSpawner::Invoke(UEdGraph* ParentGraph, FBin
 	bool const bIsTemplateNode = FBlueprintNodeTemplateCache::IsTemplateOuter(ParentGraph);
 
 	// First create a backing member for our node
-	UBlueprint* Blueprint = CastChecked<UBlueprint>(ParentGraph->GetOuter());
+	UControlRigGraph* RigGraph = Cast<UControlRigGraph>(ParentGraph);
+	check(RigGraph);
+	UControlRigBlueprint* RigBlueprint = Cast<UControlRigBlueprint>(ParentGraph->GetOuter());
+	check(RigBlueprint);
+
 	FName MemberName = NAME_None;
-	if(!bIsTemplateNode)
-	{
+
 #if WITH_EDITOR
-		if (GEditor)
-		{
-			GEditor->CancelTransaction(0);
-		}
+	if (GEditor && !bIsTemplateNode)
+	{
+		GEditor->CancelTransaction(0);
+	}
 #endif
 
-		if (UControlRigBlueprint* RigBlueprint = Cast<UControlRigBlueprint>(Blueprint))
-		{
-			FName DataType = EdGraphPinType.PinCategory;
-			if (UStruct* Struct = Cast<UStruct>(EdGraphPinType.PinSubCategoryObject))
-			{
-				DataType = Struct->GetFName();
-			}
+	URigVMController* Controller = bIsTemplateNode ? RigGraph->GetTemplateController() : RigBlueprint->Controller;
 
-			FName Name = FControlRigBlueprintUtils::ValidateName(RigBlueprint, DefaultMenuSignature.MenuName.ToString());
-			if (RigBlueprint->ModelController->AddParameter(*Name.ToString(), DataType, EControlRigModelParameterType::Hidden, Location))
-			{
-				MemberName = RigBlueprint->LastNameFromNotification;
-				for (UEdGraphNode* Node : ParentGraph->Nodes)
-				{
-					if (UControlRigGraphNode* RigNode = Cast<UControlRigGraphNode>(Node))
-					{
-						if (RigNode->GetPropertyName() == MemberName)
-						{
-							NewNode = RigNode;
-							break;
-						}
-					}
-				}
-			}
-		}
+	FName DataType = EdGraphPinType.PinCategory;
+	if (DataType == UEdGraphSchema_K2::PC_Int)
+	{
+		DataType = TEXT("int32");
 	}
-	else
+	else if (DataType == UEdGraphSchema_K2::PC_Name)
 	{
-		NewNode = FControlRigBlueprintUtils::InstantiateGraphNodeForProperty(ParentGraph, *DefaultMenuSignature.MenuName.ToString(), Location, EdGraphPinType);
+		DataType = TEXT("FName");
+	}
+	else if (DataType == UEdGraphSchema_K2::PC_String)
+	{
+		DataType = TEXT("FString");
+	}
+	else if (UStruct* Struct = Cast<UStruct>(EdGraphPinType.PinSubCategoryObject))
+	{
+		DataType = *FString::Printf(TEXT("F%s"), *Struct->GetFName().ToString());
+	}
+	else if (UEnum* Enum = Cast<UEnum>(EdGraphPinType.PinSubCategoryObject))
+	{
+		DataType = *FString::Printf(TEXT("E%s"), *Enum->GetName());
 	}
 
-	return NewNode;
-}
+	FString DataTypeForVariableName = DataType.ToString();
+	if (DataTypeForVariableName.StartsWith(TEXT("F"), ESearchCase::CaseSensitive) || DataTypeForVariableName.StartsWith(TEXT("E"), ESearchCase::CaseSensitive))
+	{
+		DataTypeForVariableName = DataTypeForVariableName.RightChop(1);
+	}
+	DataTypeForVariableName = DataTypeForVariableName.Left(1).ToUpper() + DataTypeForVariableName.RightChop(1);
+	FString VariableNamePrefix = FString::Printf(TEXT("%sVar"), *DataTypeForVariableName);
 
-bool UControlRigVariableNodeSpawner::IsTemplateNodeFilteredOut(FBlueprintActionFilter const& Filter) const
-{
-	if (EdGraphPinType.PinCategory == UEdGraphSchema_K2::PC_Struct)
+	TMap<FName, int32> NameToIndex;
+	TArray<FRigVMGraphVariableDescription> ExistingVariables = Controller->GetGraph()->GetVariableDescriptions();
+	for (int32 VariableIndex = 0; VariableIndex < ExistingVariables.Num(); VariableIndex++)
 	{
-		UStruct* Struct = Cast<UStruct>(EdGraphPinType.PinSubCategoryObject);
-		if (Struct == nullptr)
-		{
-			return true;
-		}
-		if (Struct->IsChildOf(FRigUnit::StaticStruct()))
-		{
-			return true;
-		}
+		NameToIndex.Add(ExistingVariables[VariableIndex].Name, VariableIndex);
+	}
 
-		UScriptStruct* ScriptStruct = Cast<UScriptStruct>(Struct);
-		if (ScriptStruct == nullptr)
-		{
-			// for now filter out anything which is not a script struct
-			return true;
-		}
+	FName VariableName = URigVMController::GetUniqueName(*VariableNamePrefix, [NameToIndex](const FName& InName) {
+		return !NameToIndex.Contains(InName);
+	});
 
-		// check if it is any of the math types
-		FString StructName = ScriptStruct->GetStructCPPName();
-		if (!GControlRigVariableNodeSpawnerAllowedStructTypes.Contains(StructName))
-		{
-			return true;
-		}
+	if (!bIsTemplateNode)
+	{
+		Controller->OpenUndoBracket(FString::Printf(TEXT("Add '%s' Variable"), *DataType.ToString()));
 	}
-	else if (EdGraphPinType.PinCategory == UEdGraphSchema_K2::PC_Enum || 
-			EdGraphPinType.PinCategory == UEdGraphSchema_K2::PC_Byte)
+
+	if (URigVMNode* ModelNode = Controller->AddVariableNodeFromObjectPath(VariableName, DataType.ToString(), FString(), bIsGetter, FString(), Location, FString(), !bIsTemplateNode))
 	{
-		UEnum* Enum = Cast<UEnum>(EdGraphPinType.PinSubCategoryObject);
-		if (Enum == nullptr)
+		for (UEdGraphNode* Node : ParentGraph->Nodes)
 		{
-			return true;
+			if (UControlRigGraphNode* RigNode = Cast<UControlRigGraphNode>(Node))
+			{
+				if (RigNode->GetModelNodeName() == ModelNode->GetFName())
+				{
+					NewNode = RigNode;
+					break;
+				}
+			}
 		}
 
-		if (!GControlRigVariableNodeSpawnerAllowedEnumTypes.Contains(Enum->CppType))
+		if (!bIsTemplateNode)
 		{
-			return true;
+			if (NewNode)
+			{
+				Controller->ClearNodeSelection(true);
+				Controller->SelectNode(ModelNode, true, true);
+			}
+			Controller->CloseUndoBracket();
 		}
 	}
-	else if (EdGraphPinType.PinCategory == UEdGraphSchema_K2::AllObjectTypes ||
-			EdGraphPinType.PinCategory == UEdGraphSchema_K2::PC_Object ||
-			EdGraphPinType.PinCategory == UEdGraphSchema_K2::PC_Delegate || 
-			EdGraphPinType.PinCategory == UEdGraphSchema_K2::PC_Interface)
+	else
 	{
-		// we don't allow objects, delegate or interfaces
-		return true;
+		if (!bIsTemplateNode)
+		{
+			Controller->CancelUndoBracket();
+		}
 	}
-	return Super::IsTemplateNodeFilteredOut(Filter);
-}
 
-FEdGraphPinType UControlRigVariableNodeSpawner::GetVarType() const
-{
-	return EdGraphPinType;
+
+	return NewNode;
 }
 
 #undef LOCTEXT_NAMESPACE
diff --git a/Engine/Plugins/Experimental/ControlRig/Source/ControlRigEditor/Private/Graph/NodeSpawners/ControlRigVariableNodeSpawner.h b/Engine/Plugins/Experimental/ControlRig/Source/ControlRigEditor/Private/Graph/NodeSpawners/ControlRigVariableNodeSpawner.h
--- a/Engine/Plugins/Experimental/ControlRig/Source/ControlRigEditor/Private/Graph/NodeSpawners/ControlRigVariableNodeSpawner.h
+++ b/Engine/Plugins/Experimental/ControlRig/Source/ControlRigEditor/Private/Graph/NodeSpawners/ControlRigVariableNodeSpawner.h
@@ -28,28 +28,19 @@ public:
 	 * 
 	 * @return A newly allocated instance of this class.
 	 */
-	static UControlRigVariableNodeSpawner* CreateFromPinType(const FEdGraphPinType& InPinType, const FText& InMenuDesc, const FText& InCategory, const FText& InTooltip);
+	static UControlRigVariableNodeSpawner* CreateFromPinType(const FEdGraphPinType& InPinType, bool bInIsGetter, const FText& InMenuDesc, const FText& InCategory, const FText& InTooltip);
 
 	// UBlueprintNodeSpawner interface
 	virtual void Prime() override;
 	virtual FBlueprintNodeSignature GetSpawnerSignature() const override;
 	virtual FBlueprintActionUiSpec GetUiSpec(FBlueprintActionContext const& Context, FBindingSet const& Bindings) const override;
 	virtual UEdGraphNode* Invoke(UEdGraph* ParentGraph, FBindingSet const& Bindings, FVector2D const Location) const override;
-	virtual bool IsTemplateNodeFilteredOut(FBlueprintActionFilter const& Filter) const override;
 	// End UBlueprintNodeSpawner interface
 
-	/**
-	 * Utility function for easily accessing the variable's type (needs to pull
-	 * the information differently if it is a local variable as opposed to a
-	 * member variable with a FProperty).
-	 * 
-	 * @return A struct detailing the wrapped variable's type.
-	 */
-	FEdGraphPinType GetVarType() const;
-
 private:
 	/** The pin type we will spawn */
 	FEdGraphPinType EdGraphPinType;
+	bool bIsGetter;
 
 	friend class UEngineTestControlRig;
 };
diff --git a/Engine/Plugins/Experimental/ControlRig/Source/ControlRigEditor/Private/Graph/SControlRigGraphNode.cpp b/Engine/Plugins/Experimental/ControlRig/Source/ControlRigEditor/Private/Graph/SControlRigGraphNode.cpp
--- a/Engine/Plugins/Experimental/ControlRig/Source/ControlRigEditor/Private/Graph/SControlRigGraphNode.cpp
+++ b/Engine/Plugins/Experimental/ControlRig/Source/ControlRigEditor/Private/Graph/SControlRigGraphNode.cpp
@@ -1,14 +1,15 @@
 // Copyright Epic Games, Inc. All Rights Reserved.
 
 #include "SControlRigGraphNode.h"
+#include "ControlRig.h"
 #include "Graph/ControlRigGraphNode.h"
+#include "Graph/ControlRigGraph.h"
 #include "SGraphPin.h"
 #include "Graph/ControlRigGraphSchema.h"
 #include "Widgets/Layout/SBorder.h"
 #include "Widgets/SBoxPanel.h"
 #include "Widgets/Layout/SSpacer.h"
 #include "Widgets/Input/SButton.h"
-#include "Widgets/Images/SImage.h"
 #include "Widgets/Layout/SScrollBar.h"
 #include "GraphEditorSettings.h"
 #include "ControlRigEditorStyle.h"
@@ -19,7 +20,12 @@
 #include "PropertyPathHelpers.h"
 #include "UObject/PropertyPortFlags.h"
 #include "ControlRigBlueprint.h"
-#include "ControlRigController.h"
+#include "RigVMCore/RigVM.h"
+#include "RigVMModel/RigVMController.h"
+#include "RigVMModel/RigVMNode.h"
+#include "RigVMModel/RigVMPin.h"
+#include "RigVMCompiler/RigVMCompiler.h"
+#include "IDocumentation.h"
 
 #if WITH_EDITOR
 #include "Editor.h"
@@ -38,14 +44,15 @@ void SControlRigGraphNode::Construct( const FArguments& InArgs )
 
 	// Re-cache variable info here (unit structure could have changed since last reconstruction, e.g. array add/remove)
 	// and also create missing pins if it hasn't created yet
-	ControlRigGraphNode->CreateVariablePins(false);
+	ControlRigGraphNode->AllocateDefaultPins();
 	
+	NodeErrorType = int32(EMessageSeverity::Info) + 1;
 	InputTree = nullptr;
 	OutputTree = nullptr;
 	InputOutputTree = nullptr;
 	this->UpdateGraphNode();
 
-	SetIsEditable(ControlRigGraphNode->IsPropertyAccessor());
+	SetIsEditable(false);
 
 	ScrollBar = SNew(SScrollBar);
 
@@ -53,9 +60,9 @@ void SControlRigGraphNode::Construct( const FArguments& InArgs )
 	LeftNodeBox->AddSlot()
 		.AutoHeight()
 		[
-			SAssignNew(ExecutionTree, STreeView<TSharedRef<FControlRigField>>)
+			SAssignNew(ExecutionTree, STreeView<URigVMPin*>)
 			.Visibility(this, &SControlRigGraphNode::GetExecutionTreeVisibility)
-			.TreeItemsSource(&ControlRigGraphNode->GetExecutionVariableInfo())
+			.TreeItemsSource(&ControlRigGraphNode->ExecutePins)
 			.SelectionMode(ESelectionMode::None)
 			.OnGenerateRow(this, &SControlRigGraphNode::MakeTableRowWidget)
 			.OnGetChildren(this, &SControlRigGraphNode::HandleGetChildrenForTree)
@@ -67,9 +74,9 @@ void SControlRigGraphNode::Construct( const FArguments& InArgs )
 	LeftNodeBox->AddSlot()
 		.AutoHeight()
 		[
-			SAssignNew(InputTree, STreeView<TSharedRef<FControlRigField>>)
+			SAssignNew(InputTree, STreeView<URigVMPin*>)
 			.Visibility(this, &SControlRigGraphNode::GetInputTreeVisibility)
-			.TreeItemsSource(&ControlRigGraphNode->GetInputVariableInfo())
+			.TreeItemsSource(&ControlRigGraphNode->InputPins)
 			.SelectionMode(ESelectionMode::None)
 			.OnGenerateRow(this, &SControlRigGraphNode::MakeTableRowWidget)
 			.OnGetChildren(this, &SControlRigGraphNode::HandleGetChildrenForTree)
@@ -81,9 +88,9 @@ void SControlRigGraphNode::Construct( const FArguments& InArgs )
 	LeftNodeBox->AddSlot()
 		.AutoHeight()
 		[
-			SAssignNew(InputOutputTree, STreeView<TSharedRef<FControlRigField>>)
+			SAssignNew(InputOutputTree, STreeView<URigVMPin*>)
 			.Visibility(this, &SControlRigGraphNode::GetInputOutputTreeVisibility)
-			.TreeItemsSource(&ControlRigGraphNode->GetInputOutputVariableInfo())
+			.TreeItemsSource(&ControlRigGraphNode->InputOutputPins)
 			.SelectionMode(ESelectionMode::None)
 			.OnGenerateRow(this, &SControlRigGraphNode::MakeTableRowWidget)
 			.OnGetChildren(this, &SControlRigGraphNode::HandleGetChildrenForTree)
@@ -95,9 +102,9 @@ void SControlRigGraphNode::Construct( const FArguments& InArgs )
 	LeftNodeBox->AddSlot()
 		.AutoHeight()
 		[
-			SAssignNew(OutputTree, STreeView<TSharedRef<FControlRigField>>)
+			SAssignNew(OutputTree, STreeView<URigVMPin*>)
 			.Visibility(this, &SControlRigGraphNode::GetOutputTreeVisibility)
-			.TreeItemsSource(&ControlRigGraphNode->GetOutputVariableInfo())
+			.TreeItemsSource(&ControlRigGraphNode->OutputPins)
 			.SelectionMode(ESelectionMode::None)
 			.OnGenerateRow(this, &SControlRigGraphNode::MakeTableRowWidget)
 			.OnGetChildren(this, &SControlRigGraphNode::HandleGetChildrenForTree)
@@ -108,24 +115,23 @@ void SControlRigGraphNode::Construct( const FArguments& InArgs )
 
 	struct Local
 	{
-		static void SetItemExpansion_Recursive(UControlRigGraphNode* InControlRigGraphNode, TSharedPtr<STreeView<TSharedRef<FControlRigField>>>& TreeWidget, const TArray<TSharedRef<FControlRigField>>& InItems)
+		static void SetItemExpansion_Recursive(UControlRigGraphNode* InControlRigGraphNode, TSharedPtr<STreeView<URigVMPin*>>& TreeWidget, const TArray<URigVMPin*>& InItems)
 		{
-			for(const TSharedRef<FControlRigField>& Field : InItems)
+			for(URigVMPin* Pin : InItems)
 			{
-				if(InControlRigGraphNode->IsPinExpanded(Field->GetPinPath()))
+				if(InControlRigGraphNode->IsPinExpanded(Pin->GetPinPath()))
 				{
-					TreeWidget->SetItemExpansion(Field, true);
-
-					SetItemExpansion_Recursive(InControlRigGraphNode, TreeWidget, Field->Children);
+					TreeWidget->SetItemExpansion(Pin, true);
+					SetItemExpansion_Recursive(InControlRigGraphNode, TreeWidget, Pin->GetSubPins());
 				}
 			}
 		}
 	};
 
-	Local::SetItemExpansion_Recursive(ControlRigGraphNode, ExecutionTree, ControlRigGraphNode->GetExecutionVariableInfo());
-	Local::SetItemExpansion_Recursive(ControlRigGraphNode, InputTree, ControlRigGraphNode->GetInputVariableInfo());
-	Local::SetItemExpansion_Recursive(ControlRigGraphNode, InputOutputTree, ControlRigGraphNode->GetInputOutputVariableInfo());
-	Local::SetItemExpansion_Recursive(ControlRigGraphNode, OutputTree, ControlRigGraphNode->GetOutputVariableInfo());
+	Local::SetItemExpansion_Recursive(ControlRigGraphNode, ExecutionTree, ControlRigGraphNode->ExecutePins);
+	Local::SetItemExpansion_Recursive(ControlRigGraphNode, InputTree, ControlRigGraphNode->InputPins);
+	Local::SetItemExpansion_Recursive(ControlRigGraphNode, InputOutputTree, ControlRigGraphNode->InputOutputPins);
+	Local::SetItemExpansion_Recursive(ControlRigGraphNode, OutputTree, ControlRigGraphNode->OutputPins);
 
 
 	// force the regeneration of all pins.
@@ -140,6 +146,13 @@ void SControlRigGraphNode::Construct( const FArguments& InArgs )
 	InputTree->Tick(DummyGeometry, 0.f, 0.f);
 	InputOutputTree->Tick(DummyGeometry, 0.f, 0.f);
 	OutputTree->Tick(DummyGeometry, 0.f, 0.f);
+
+	const FSlateBrush* ImageBrush = FControlRigEditorStyle::Get().GetBrush(TEXT("ControlRig.RigUnit.VisualDebug"));
+
+	VisualDebugIndicatorWidget =
+		SNew(SImage)
+		.Image(ImageBrush)
+		.Visibility(EVisibility::Visible);
 }
 
 TSharedRef<SWidget> SControlRigGraphNode::CreateNodeContentArea()
@@ -199,8 +212,32 @@ void SControlRigGraphNode::EndUserInteraction() const
 	if (GraphNode)
 	{
 		UControlRigGraphNode* ControlRigGraphNode = CastChecked<UControlRigGraphNode>(GraphNode);
-		FVector2D Position(ControlRigGraphNode->NodePosX, ControlRigGraphNode->NodePosY);
-		ControlRigGraphNode->GetBlueprint()->ModelController->SetNodePosition(ControlRigGraphNode->GetPropertyName(), Position, true);
+		URigVMController* Controller = ControlRigGraphNode->GetBlueprint()->Controller;
+		if (URigVMGraph* Model = Controller->GetGraph())
+		{
+			UControlRigGraph* Graph = Cast<UControlRigGraph>(ControlRigGraphNode->GetGraph());
+			Controller->OpenUndoBracket(TEXT("Moved Nodes."));
+			bool bMovedSomething = false;
+			for (const FName& SelectedNodeName : Model->GetSelectNodes())
+			{
+				if (UEdGraphNode* SelectedNode = Graph->FindNodeForModelNodeName(SelectedNodeName))
+				{
+					FVector2D Position(SelectedNode->NodePosX, SelectedNode->NodePosY);
+					if (Controller->SetNodePositionByName(SelectedNodeName, Position, true, true))
+					{
+						bMovedSomething = true;
+					}
+				}
+			}
+			if (bMovedSomething)
+			{
+				Controller->CloseUndoBracket();
+			}
+			else
+			{
+				Controller->CancelUndoBracket();
+			}
+		}
 	}
 }
 
@@ -209,25 +246,29 @@ void SControlRigGraphNode::AddPin(const TSharedRef<SGraphPin>& PinToAdd)
 	// We show our own label
 	PinToAdd->SetShowLabel(false);
 
-	const UEdGraphPin* PinObj = PinToAdd->GetPinObj();
-
-	// Remove value widget from combined pin content
-	TSharedPtr<SWrapBox> LabelAndValueWidget = PinToAdd->GetLabelAndValue();
-	TSharedPtr<SHorizontalBox> FullPinHorizontalRowWidget = PinToAdd->GetFullPinHorizontalRowWidget().Pin();
-	if(LabelAndValueWidget.IsValid() && FullPinHorizontalRowWidget.IsValid())
+	UControlRigGraphNode* ControlRigGraphNode = CastChecked<UControlRigGraphNode>(GraphNode);
+	if(URigVMNode* ModelNode = ControlRigGraphNode->GetModelNode())
 	{
-		FullPinHorizontalRowWidget->RemoveSlot(LabelAndValueWidget.ToSharedRef());
-	}
+		const UEdGraphPin* EdPinObj = PinToAdd->GetPinObj();
 
-	PinToAdd->SetOwner(SharedThis(this));
-	PinWidgetMap.Add(PinObj, PinToAdd);
-	if(PinObj->Direction == EGPD_Input)
-	{
-		InputPins.Add(PinToAdd);
-	}
-	else
-	{
-		OutputPins.Add(PinToAdd);
+		// Remove value widget from combined pin content
+		TSharedPtr<SWrapBox> LabelAndValueWidget = PinToAdd->GetLabelAndValue();
+		TSharedPtr<SHorizontalBox> FullPinHorizontalRowWidget = PinToAdd->GetFullPinHorizontalRowWidget().Pin();
+		if(LabelAndValueWidget.IsValid() && FullPinHorizontalRowWidget.IsValid())
+		{
+			FullPinHorizontalRowWidget->RemoveSlot(LabelAndValueWidget.ToSharedRef());
+		}
+
+		PinToAdd->SetOwner(SharedThis(this));
+		PinWidgetMap.Add(EdPinObj, PinToAdd);
+		if(EdPinObj->Direction == EGPD_Input)
+		{
+			InputPins.Add(PinToAdd);
+		}
+		else
+		{
+			OutputPins.Add(PinToAdd);
+		}
 	}
 }
 
@@ -249,25 +290,25 @@ EVisibility SControlRigGraphNode::GetTitleVisibility() const
 EVisibility SControlRigGraphNode::GetExecutionTreeVisibility() const
 {
 	UControlRigGraphNode* ControlRigGraphNode = CastChecked<UControlRigGraphNode>(GraphNode);
-	return ControlRigGraphNode->GetExecutionVariableInfo().Num() > 0 ? EVisibility::Visible : EVisibility::Collapsed;
+	return ControlRigGraphNode->ExecutePins.Num() > 0 ? EVisibility::Visible : EVisibility::Collapsed;
 }
 
 EVisibility SControlRigGraphNode::GetInputTreeVisibility() const
 {
 	UControlRigGraphNode* ControlRigGraphNode = CastChecked<UControlRigGraphNode>(GraphNode);
-	return ControlRigGraphNode->GetInputVariableInfo().Num() > 0 ? EVisibility::Visible : EVisibility::Collapsed;
+	return ControlRigGraphNode->InputPins.Num() > 0 ? EVisibility::Visible : EVisibility::Collapsed;
 }
 
 EVisibility SControlRigGraphNode::GetInputOutputTreeVisibility() const
 {
 	UControlRigGraphNode* ControlRigGraphNode = CastChecked<UControlRigGraphNode>(GraphNode);
-	return ControlRigGraphNode->GetInputOutputVariableInfo().Num() > 0 ? EVisibility::Visible : EVisibility::Collapsed;
+	return ControlRigGraphNode->InputOutputPins.Num() > 0 ? EVisibility::Visible : EVisibility::Collapsed;
 }
 
 EVisibility SControlRigGraphNode::GetOutputTreeVisibility() const
 {
 	UControlRigGraphNode* ControlRigGraphNode = CastChecked<UControlRigGraphNode>(GraphNode);
-	return ControlRigGraphNode->GetOutputVariableInfo().Num() > 0 ? EVisibility::Visible : EVisibility::Collapsed;
+	return ControlRigGraphNode->OutputPins.Num() > 0 ? EVisibility::Visible : EVisibility::Collapsed;
 }
 
 TSharedRef<SWidget> SControlRigGraphNode::CreateTitleWidget(TSharedPtr<SNodeTitle> NodeTitle)
@@ -370,7 +411,7 @@ class SControlRigExpanderArrow : public SExpanderArrow
 	bool bLeftAligned;
 };
 
-class SControlRigPinTreeRow : public STableRow<TSharedRef<FControlRigField>>
+class SControlRigPinTreeRow : public STableRow<URigVMPin*>
 {
 public:
 	SLATE_BEGIN_ARGS(SControlRigPinTreeRow) {}
@@ -383,7 +424,7 @@ public:
 	{
 		bLeftAligned = InArgs._LeftAligned;
 
-		STableRow<TSharedRef<FControlRigField>>::Construct(STableRow<TSharedRef<FControlRigField>>::FArguments(), InOwnerTableView);
+		STableRow<URigVMPin*>::Construct(STableRow<URigVMPin*>::FArguments(), InOwnerTableView);
 	}
 
 	virtual void ConstructChildren( ETableViewMode::Type InOwnerTableMode, const TAttribute<FMargin>& InPadding, const TSharedRef<SWidget>& InContent ) override
@@ -478,14 +519,12 @@ public:
 	bool bLeftAligned;
 };
 
-TSharedRef<SWidget> SControlRigGraphNode::AddContainerPinContent(TSharedRef<FControlRigField> InItem, FText InTooltipText)
+TSharedRef<SWidget> SControlRigGraphNode::AddContainerPinContent(URigVMPin* InItem, FText InTooltipText)
 {
-	TWeakPtr<FControlRigField> WeakItem = InItem;
-
 	return SNew(SButton)
 	.ContentPadding(0.0f)
 	.ButtonStyle(FEditorStyle::Get(), "NoBorder")
-	.OnClicked(this, &SControlRigGraphNode::HandleAddArrayElement, WeakItem)
+	.OnClicked(this, &SControlRigGraphNode::HandleAddArrayElement, InItem)
 	.IsEnabled(this, &SGraphNode::IsNodeEditable)
 	.ToolTipText(InTooltipText)
 	.Cursor(EMouseCursor::Default)
@@ -502,24 +541,39 @@ TSharedRef<SWidget> SControlRigGraphNode::AddContainerPinContent(TSharedRef<FCon
 	];
 }
 
-TSharedRef<ITableRow> SControlRigGraphNode::MakeTableRowWidget(TSharedRef<FControlRigField> InItem, const TSharedRef<STableViewBase>& OwnerTable)
+TSharedRef<ITableRow> SControlRigGraphNode::MakeTableRowWidget(URigVMPin* InItem, const TSharedRef<STableViewBase>& OwnerTable)
 {
-	const bool bLeaf = InItem->Children.Num() == 0;
-	const bool bIsContainer = InItem->PinType.IsContainer();
+	const bool bLeaf = InItem->GetSubPins().Num() == 0;
+	const bool bIsContainer = InItem->IsArray();
 
 	TSharedPtr<SGraphPin> InputPinWidget;
 	TSharedPtr<SGraphPin> OutputPinWidget;
 	TSharedPtr<SWidget> InputPinValueWidget;
 
-	if(InItem->InputPin)
+	UControlRigGraphNode* RigNode = Cast<UControlRigGraphNode>(GraphNode);
+	const UControlRigGraphNode::PinPair& Pair = RigNode->CachedPins.FindChecked(InItem);
+
+	if (Pair.InputPin)
 	{
-		TSharedPtr<SGraphPin>* InputGraphPinPtr = PinWidgetMap.Find(InItem->InputPin);
+		TSharedPtr<SGraphPin>* InputGraphPinPtr = PinWidgetMap.Find(Pair.InputPin);
 		if(InputGraphPinPtr != nullptr)
 		{
 			InputPinWidget = *InputGraphPinPtr;
 
+			bool bIsPlainOrEditableStruct = !InItem->IsStruct();
+			if (!bIsPlainOrEditableStruct)
+			{
+				if (InItem->GetSubPins().Num() == 0)
+				{
+					if (const UControlRigGraphSchema* RigSchema = Cast<UControlRigGraphSchema>(RigNode->GetSchema()))
+					{
+						bIsPlainOrEditableStruct = RigSchema->IsStructEditable(InItem->GetScriptStruct());
+					}
+				}
+			}
+
 			// Only leaf pins have value widgets, but not containers
-			if(bLeaf && !bIsContainer)
+			if(bLeaf && !bIsContainer && bIsPlainOrEditableStruct)
 			{
 				InputPinValueWidget = (*InputGraphPinPtr)->GetValueWidget();
 			}
@@ -531,9 +585,9 @@ TSharedRef<ITableRow> SControlRigGraphNode::MakeTableRowWidget(TSharedRef<FContr
 		}
 	}
 
-	if(InItem->OutputPin)
+	if (Pair.OutputPin)
 	{
-		TSharedPtr<SGraphPin>* OutputGraphPinPtr = PinWidgetMap.Find(InItem->OutputPin);
+		TSharedPtr<SGraphPin>* OutputGraphPinPtr = PinWidgetMap.Find(Pair.OutputPin);
 		if(OutputGraphPinPtr != nullptr)
 		{
 			OutputPinWidget = *OutputGraphPinPtr;
@@ -542,7 +596,7 @@ TSharedRef<ITableRow> SControlRigGraphNode::MakeTableRowWidget(TSharedRef<FContr
 
 	TSharedRef<SControlRigPinTreeRow> ControlRigPinTreeRow = SNew(SControlRigPinTreeRow, OwnerTable)
 		.LeftAligned(!(OutputPinWidget.IsValid() && !InputPinWidget.IsValid()))
-		.ToolTipText(InItem->GetTooltipText());
+		.ToolTipText(InItem->GetToolTipText());
 
 	if(InputPinWidget.IsValid() || OutputPinWidget.IsValid())
 	{
@@ -652,25 +706,19 @@ TSharedRef<ITableRow> SControlRigGraphNode::MakeTableRowWidget(TSharedRef<FContr
 	return ControlRigPinTreeRow;
 }
 	
-void SControlRigGraphNode::HandleGetChildrenForTree(TSharedRef<FControlRigField> InItem, TArray<TSharedRef<FControlRigField>>& OutChildren)
+void SControlRigGraphNode::HandleGetChildrenForTree(URigVMPin* InItem, TArray<URigVMPin*>& OutChildren)
 {
-	OutChildren.Append(InItem->Children);
+	OutChildren.Append(InItem->GetSubPins());
 }
 
-void SControlRigGraphNode::HandleExpansionChanged(TSharedRef<FControlRigField> InItem, bool bExpanded)
+void SControlRigGraphNode::HandleExpansionChanged(URigVMPin* InItem, bool bExpanded)
 {
 	if (GraphNode)
 	{
 		UControlRigBlueprint* ControlRigBlueprint = Cast<UControlRigBlueprint>(GraphNode->GetGraph()->GetOuter());
 		if (ControlRigBlueprint)
 		{
-			if (ControlRigBlueprint->ModelController)
-			{
-				FString PinPath = InItem->GetPinPath();
-				FString Left, Right;
-				ControlRigBlueprint->Model->SplitPinPath(PinPath, Left, Right);
-				ControlRigBlueprint->ModelController->ExpandPin(*Left, *Right, InItem->GetPin()->Direction == EGPD_Input, bExpanded);
-			}
+			ControlRigBlueprint->Controller->SetPinExpansion(InItem->GetPinPath(), bExpanded, true);
 		}
 	}
 }
@@ -705,15 +753,13 @@ FSlateColor SControlRigGraphNode::GetPinTextColor(TWeakPtr<SGraphPin> GraphPin)
 	return FLinearColor::White;
 }
 
-FReply SControlRigGraphNode::HandleAddArrayElement(TWeakPtr<FControlRigField> InWeakItem)
+FReply SControlRigGraphNode::HandleAddArrayElement(URigVMPin* InItem)
 {
-	TSharedPtr<FControlRigField> Item = InWeakItem.Pin();
-	if(Item.IsValid())
+	if(InItem)
 	{
 		if (UControlRigGraphNode* ControlRigGraphNode = Cast<UControlRigGraphNode>(GraphNode))
 		{
-			// todo ?
-			ControlRigGraphNode->HandleAddArrayElement(Item->GetPinPath());
+			ControlRigGraphNode->HandleAddArrayElement(InItem->GetPinPath());
 		}
 	}
 
@@ -727,13 +773,15 @@ void SControlRigGraphNode::GetNodeInfoPopups(FNodeInfoContext* Context, TArray<F
 	const FLinearColor LatentBubbleColor(1.f, 0.5f, 0.25f);
 	const FLinearColor PinnedWatchColor(0.35f, 0.25f, 0.25f);
 
+	UControlRig* ActiveObject = Cast<UControlRig>(K2Context->ActiveObjectBeingDebugged);
+	UControlRigBlueprint* RigBlueprint = Cast<UControlRigBlueprint>(K2Context->SourceBlueprint);
+
 	// Display any pending latent actions
-	if (UObject* ActiveObject = K2Context->ActiveObjectBeingDebugged)
+	if (ActiveObject && RigBlueprint)
 	{
 		// Display pinned watches
 		if (K2Context->WatchedNodeSet.Contains(GraphNode))
 		{
-			UBlueprint* Blueprint = K2Context->SourceBlueprint;
 			const UEdGraphSchema* Schema = GraphNode->GetSchema();
 
 			FString PinnedWatchText;
@@ -743,26 +791,47 @@ void SControlRigGraphNode::GetNodeInfoPopups(FNodeInfoContext* Context, TArray<F
 				UEdGraphPin* WatchPin = GraphNode->Pins[PinIndex];
 				if (K2Context->WatchedPinSet.Contains(WatchPin))
 				{
-					if (ValidWatchCount > 0)
-					{
-						PinnedWatchText += TEXT("\n");
-					}
-
-					FString PinName = UEdGraphSchema_K2::TypeToText(WatchPin->PinType).ToString();
-					PinName += TEXT(" ");
-					PinName += Schema->GetPinDisplayName(WatchPin).ToString();
-
-					FString WatchText;
-					if (PropertyPathHelpers::GetPropertyValueAsString(ActiveObject, WatchPin->PinName.ToString(), WatchText))
+					if (URigVMPin* ModelPin = RigBlueprint->Model->FindPin(WatchPin->GetName()))
 					{
-						PinnedWatchText += FText::Format(LOCTEXT("WatchingAndValidFmt", "Watching {0}\n\t{1}"), FText::FromString(PinName), FText::FromString(WatchText)).ToString();//@TODO: Print out object being debugged name?
+						if (ValidWatchCount > 0)
+						{
+							PinnedWatchText += TEXT("\n");
+						}
+
+						FString PinName = Schema->GetPinDisplayName(WatchPin).ToString();
+						PinName += TEXT(" (");
+						PinName += UEdGraphSchema_K2::TypeToText(WatchPin->PinType).ToString();
+						PinName += TEXT(")");
+
+						FString WatchText;
+						FString PinHash = URigVMCompiler::GetPinHash(ModelPin, nullptr, true);
+						if (const FRigVMOperand* WatchOperand = RigBlueprint->PinToOperandMap.Find(PinHash))
+						{
+							FRigVMMemoryContainer& Memory = 
+								WatchOperand->GetMemoryType() == ERigVMMemoryType::Literal ?
+								ActiveObject->GetVM()->LiteralMemory : ActiveObject->GetVM()->WorkMemory;
+
+							TArray<FString> DefaultValues = Memory.GetRegisterValueAsString(*WatchOperand, ModelPin->GetCPPType(), ModelPin->GetCPPTypeObject());
+							if (DefaultValues.Num() == 1)
+							{
+								WatchText = DefaultValues[0];
+							}
+							else if (DefaultValues.Num() > 1)
+							{
+								WatchText = FString::Printf(TEXT("[%s]"), *FString::Join(DefaultValues, TEXT(", ")));
+							}
+							if (!WatchText.IsEmpty())
+							{
+								PinnedWatchText += FText::Format(LOCTEXT("WatchingAndValidFmt", "{0}\n\t{1}"), FText::FromString(PinName), FText::FromString(WatchText)).ToString();//@TODO: Print out object being debugged name?
+							}
+							else
+							{
+								PinnedWatchText += FText::Format(LOCTEXT("InvalidPropertyFmt", "No watch found for {0}"), Schema->GetPinDisplayName(WatchPin)).ToString();//@TODO: Print out object being debugged name?
+							}
+
+							ValidWatchCount++;
+						}
 					}
-					else
-					{
-						PinnedWatchText += FText::Format(LOCTEXT("InvalidPropertyFmt", "Invalid Property {0}"), FText::FromString(PinName)).ToString();//@TODO: Print out object being debugged name?
-					}
-
-					ValidWatchCount++;
 				}
 			}
 
@@ -774,6 +843,84 @@ void SControlRigGraphNode::GetNodeInfoPopups(FNodeInfoContext* Context, TArray<F
 	}
 }
 
+TArray<FOverlayWidgetInfo> SControlRigGraphNode::GetOverlayWidgets(bool bSelected, const FVector2D& WidgetSize) const
+{
+	TArray<FOverlayWidgetInfo> Widgets = SGraphNode::GetOverlayWidgets(bSelected, WidgetSize);
+
+	if (UControlRigGraphNode* RigNode = CastChecked<UControlRigGraphNode>(GraphNode, ECastCheckedType::NullAllowed))
+	{
+		if (URigVMNode* ModelNode = RigNode->GetModelNode())
+		{
+			bool bSetColor = false;
+			FLinearColor Color = FLinearColor::Black;
+			int32 PreviousNumWidgets = Widgets.Num();
+			VisualDebugIndicatorWidget->SetColorAndOpacity(Color);
+
+			for (URigVMPin* ModelPin : ModelNode->GetPins())
+			{
+				if (ModelPin->HasInjectedNodes())
+				{
+					for (URigVMInjectionInfo* Injection : ModelPin->GetInjectedNodes())
+					{
+						URigVMStructNode* VisualDebugNode = Injection->StructNode;
+
+						FString PrototypeName;
+						if (VisualDebugNode->GetScriptStruct()->GetStringMetaDataHierarchical(TEXT("PrototypeName"), &PrototypeName))
+						{
+							if (PrototypeName == TEXT("VisualDebug"))
+							{
+								if (!bSetColor)
+								{
+									if (VisualDebugNode->FindPin(TEXT("bEnabled"))->GetDefaultValue() == TEXT("True"))
+									{
+										if (URigVMPin* ColorPin = VisualDebugNode->FindPin(TEXT("Color")))
+										{
+											TBaseStructure<FLinearColor>::Get()->ImportText(*ColorPin->GetDefaultValue(), &Color, nullptr, PPF_None, nullptr, TBaseStructure<FLinearColor>::Get()->GetName());
+										}
+										else
+										{
+											Color = FLinearColor::White;
+										}
+
+										VisualDebugIndicatorWidget->SetColorAndOpacity(Color);
+										bSetColor = true;
+									}
+								}
+
+								if (Widgets.Num() == PreviousNumWidgets)
+								{
+									const FSlateBrush* ImageBrush = FEditorStyle::Get().GetBrush(TEXT("ControlRig.RigUnit.VisualDebug"));
+									FVector2D ImageSize = VisualDebugIndicatorWidget->GetDesiredSize();
+
+									FOverlayWidgetInfo Info;
+									Info.OverlayOffset = FVector2D(WidgetSize.X - (ImageSize.X * 0.8f), -(ImageSize.Y * 0.5f));
+									Info.Widget = VisualDebugIndicatorWidget;
+
+									Widgets.Add(Info);
+								}
+							}
+						}
+					}
+				}
+			}
+		}
+	}
+
+	return Widgets;
+}
+
+void SControlRigGraphNode::RefreshErrorInfo()
+{
+	if (GraphNode)
+	{
+		if (NodeErrorType != GraphNode->ErrorType)
+		{
+			SGraphNode::RefreshErrorInfo();
+			NodeErrorType = GraphNode->ErrorType;
+		}
+	}
+}
+
 void SControlRigGraphNode::Tick(const FGeometry& AllottedGeometry, const double InCurrentTime, const float InDeltaTime)
 {
 	SGraphNode::Tick(AllottedGeometry, InCurrentTime, InDeltaTime);
@@ -782,6 +929,7 @@ void SControlRigGraphNode::Tick(const FGeometry& AllottedGeometry, const double
 	{
 		GraphNode->NodeWidth = (int32)AllottedGeometry.Size.X;
 		GraphNode->NodeHeight = (int32)AllottedGeometry.Size.Y;
+		RefreshErrorInfo();
 	}
 }
 
diff --git a/Engine/Plugins/Experimental/ControlRig/Source/ControlRigEditor/Private/Graph/SControlRigGraphNode.h b/Engine/Plugins/Experimental/ControlRig/Source/ControlRigEditor/Private/Graph/SControlRigGraphNode.h
--- a/Engine/Plugins/Experimental/ControlRig/Source/ControlRigEditor/Private/Graph/SControlRigGraphNode.h
+++ b/Engine/Plugins/Experimental/ControlRig/Source/ControlRigEditor/Private/Graph/SControlRigGraphNode.h
@@ -6,10 +6,11 @@
 #include "Widgets/DeclarativeSyntaxSupport.h"
 #include "SGraphNode.h"
 #include "Widgets/Views/STreeView.h"
+#include "Widgets/Images/SImage.h"
+#include "RigVMModel/RigVMPin.h"
 
 class UControlRigGraphNode;
 class STableViewBase;
-class FControlRigField;
 class SOverlay;
 class SGraphPin;
 class UEdGraphPin;
@@ -42,7 +43,9 @@ public:
 	virtual TSharedRef<SWidget> CreateNodeContentArea() override;
 	virtual TSharedPtr<SGraphPin> GetHoveredPin( const FGeometry& MyGeometry, const FPointerEvent& MouseEvent ) const override;
 	virtual void GetNodeInfoPopups(FNodeInfoContext* Context, TArray<FGraphInformationPopupInfo>& Popups) const override;
+	virtual TArray<FOverlayWidgetInfo> GetOverlayWidgets(bool bSelected, const FVector2D& WidgetSize) const override;
 
+	virtual void RefreshErrorInfo() override;
 	virtual void Tick(const FGeometry& AllottedGeometry, const double InCurrentTime, const float InDeltaTime) override;
 
 private:
@@ -61,35 +64,35 @@ private:
 
 	EVisibility GetOutputTreeVisibility() const;
 
-	TSharedRef<ITableRow> MakeTableRowWidget(TSharedRef<FControlRigField> InItem, const TSharedRef<STableViewBase>& OwnerTable);
+	TSharedRef<ITableRow> MakeTableRowWidget(URigVMPin* InItem, const TSharedRef<STableViewBase>& OwnerTable);
 
-	void HandleGetChildrenForTree(TSharedRef<FControlRigField> InItem, TArray<TSharedRef<FControlRigField>>& OutChildren);
+	void HandleGetChildrenForTree(URigVMPin* InItem, TArray<URigVMPin*>& OutChildren);
 
-	void HandleExpansionChanged(TSharedRef<FControlRigField> InItem, bool bExpanded);
+	void HandleExpansionChanged(URigVMPin* InItem, bool bExpanded);
 
 	FText GetPinLabel(TWeakPtr<SGraphPin> GraphPin) const;
 
 	FSlateColor GetPinTextColor(TWeakPtr<SGraphPin> GraphPin) const;
 
-	TSharedRef<SWidget> AddContainerPinContent(TSharedRef<FControlRigField> InItem, FText InTooltipText);
+	TSharedRef<SWidget> AddContainerPinContent(URigVMPin* InItem, FText InTooltipText);
 
-	FReply HandleAddArrayElement(TWeakPtr<FControlRigField> InWeakItem);
+	FReply HandleAddArrayElement(URigVMPin* InItem);
 
 private:
 	/** Cached widget title area */
 	TSharedPtr<SOverlay> TitleAreaWidget;
 
 	/** Widget representing collapsible execution pins */
-	TSharedPtr<STreeView<TSharedRef<FControlRigField>>> ExecutionTree;
+	TSharedPtr<STreeView<URigVMPin*>> ExecutionTree;
 
 	/** Widget representing collapsible input pins */
-	TSharedPtr<STreeView<TSharedRef<FControlRigField>>> InputTree;
+	TSharedPtr<STreeView<URigVMPin*>> InputTree;
 
 	/** Widget representing collapsible input-output pins */
-	TSharedPtr<STreeView<TSharedRef<FControlRigField>>> InputOutputTree;
+	TSharedPtr<STreeView<URigVMPin*>> InputOutputTree;
 
 	/** Widget representing collapsible output pins */
-	TSharedPtr<STreeView<TSharedRef<FControlRigField>>> OutputTree;
+	TSharedPtr<STreeView<URigVMPin*>> OutputTree;
 
 	/** Dummy scrollbar, as we cant create a tree view without one! */
 	TSharedPtr<SScrollBar> ScrollBar;
@@ -99,4 +102,8 @@ private:
 
 	/** Map of pin widgets to extra pin widgets */
 	TMap<TSharedRef<SWidget>, TSharedRef<SGraphPin>> ExtraWidgetToPinMap;
+
+	int32 NodeErrorType;
+
+	TSharedPtr<SImage> VisualDebugIndicatorWidget;
 };
diff --git a/Engine/Plugins/Experimental/ControlRig/Source/ControlRigEditor/Private/Graph/SControlRigGraphNodeComment.cpp b/Engine/Plugins/Experimental/ControlRig/Source/ControlRigEditor/Private/Graph/SControlRigGraphNodeComment.cpp
--- a/Engine/Plugins/Experimental/ControlRig/Source/ControlRigEditor/Private/Graph/SControlRigGraphNodeComment.cpp
+++ b/Engine/Plugins/Experimental/ControlRig/Source/ControlRigEditor/Private/Graph/SControlRigGraphNodeComment.cpp
@@ -7,7 +7,7 @@
 #include "UObject/PropertyPortFlags.h"
 #include "ControlRigBlueprint.h"
 #include "Graph/ControlRigGraph.h"
-#include "ControlRigController.h"
+#include "RigVMModel/RigVMController.h"
 #if WITH_EDITOR
 #include "Editor.h"
 #endif
@@ -34,10 +34,10 @@ FReply SControlRigGraphNodeComment::OnMouseButtonUp(const FGeometry& MyGeometry,
 				{
 					FVector2D Position(CommentNode->NodePosX, CommentNode->NodePosY);
 					FVector2D Size(CommentNode->NodeWidth, CommentNode->NodeHeight);
-					Blueprint->ModelController->OpenUndoBracket(TEXT("Resize Comment Box"));
-					Blueprint->ModelController->SetNodePosition(CommentNode->GetFName(), Position, true);
-					Blueprint->ModelController->SetNodeSize(CommentNode->GetFName(), Size, true);
-					Blueprint->ModelController->CloseUndoBracket();
+					Blueprint->Controller->OpenUndoBracket(TEXT("Resize Comment Box"));
+					Blueprint->Controller->SetNodePositionByName(CommentNode->GetFName(), Position, true);
+					Blueprint->Controller->SetNodeSizeByName(CommentNode->GetFName(), Size, true);
+					Blueprint->Controller->CloseUndoBracket();
 				}
 			}
 		}
@@ -63,36 +63,31 @@ void SControlRigGraphNodeComment::EndUserInteraction() const
 			{
 				if (CommentNode->MoveMode == ECommentBoxMode::GroupMovement)
 				{
-					Blueprint->ModelController->OpenUndoBracket(TEXT("Move Comment Box"));
+					Blueprint->Controller->OpenUndoBracket(TEXT("Move Comment Box"));
 
 					for (FCommentNodeSet::TConstIterator NodeIt(CommentNode->GetNodesUnderComment()); NodeIt; ++NodeIt)
 					{
 						if (UEdGraphNode* EdNode = Cast<UEdGraphNode>(*NodeIt))
 						{
 							FName NodeName = EdNode->GetFName();
-							if (UControlRigGraphNode* RigNode = Cast<UControlRigGraphNode>(EdNode))
+							if (URigVMNode* ModelNode = Blueprint->Model->FindNodeByName(NodeName))
 							{
-								NodeName = RigNode->GetPropertyName();
-							}
-
-							if (const FControlRigModelNode* ModelNode = Blueprint->Model->FindNode(NodeName))
-							{
-								if (!Blueprint->Model->IsNodeSelected(ModelNode->Name))
+								if (!ModelNode->IsSelected())
 								{
 									FVector2D Position(EdNode->NodePosX, EdNode->NodePosY);
-									Blueprint->ModelController->SetNodePosition(ModelNode->Name, Position, true);
+									Blueprint->Controller->SetNodePositionByName(NodeName, Position, true);
 								}
 							}
 						}
 					}
 					FVector2D Position(CommentNode->NodePosX, CommentNode->NodePosY);
-					Blueprint->ModelController->SetNodePosition(CommentNode->GetFName(), Position, true);
-					Blueprint->ModelController->CloseUndoBracket();
+					Blueprint->Controller->SetNodePositionByName(CommentNode->GetFName(), Position, true);
+					Blueprint->Controller->CloseUndoBracket();
 				}
 				else
 				{
 					FVector2D Position(CommentNode->NodePosX, CommentNode->NodePosY);
-					Blueprint->ModelController->SetNodePosition(CommentNode->GetFName(), Position, true);
+					Blueprint->Controller->SetNodePositionByName(CommentNode->GetFName(), Position, true);
 				}
 			}
 		}
@@ -115,7 +110,7 @@ void SControlRigGraphNodeComment::Tick(const FGeometry& AllottedGeometry, const
 			{
 				if (UControlRigBlueprint* Blueprint = Cast<UControlRigBlueprint>(Graph->GetOuter()))
 				{
-					Blueprint->ModelController->SetCommentText(CommentNode->GetFName(), CurrentCommentTitle, true);
+					Blueprint->Controller->SetCommentTextByName(CommentNode->GetFName(), CurrentCommentTitle, true);
 				}
 			}
 		}
@@ -134,7 +129,7 @@ void SControlRigGraphNodeComment::Tick(const FGeometry& AllottedGeometry, const
 					if (UControlRigBlueprint* Blueprint = Cast<UControlRigBlueprint>(Graph->GetOuter()))
 					{
 						// for now we won't use our undo for this kind of change
-						Blueprint->ModelController->SetNodeColor(GraphNode->GetFName(), CurrentNodeCommentColor, false);
+						Blueprint->Controller->SetNodeColorByName(GraphNode->GetFName(), CurrentNodeCommentColor, false, true);
 						CachedNodeCommentColor = CurrentNodeCommentColor;
 					}
 				}
diff --git a/Engine/Plugins/Experimental/ControlRig/Source/ControlRigEditor/Private/Graph/SControlRigGraphPinCurveFloat.cpp b/Engine/Plugins/Experimental/ControlRig/Source/ControlRigEditor/Private/Graph/SControlRigGraphPinCurveFloat.cpp
--- a/Engine/Plugins/Experimental/ControlRig/Source/ControlRigEditor/Private/Graph/SControlRigGraphPinCurveFloat.cpp
+++ b/Engine/Plugins/Experimental/ControlRig/Source/ControlRigEditor/Private/Graph/SControlRigGraphPinCurveFloat.cpp
@@ -3,12 +3,16 @@
 
 #include "Graph/SControlRigGraphPinCurveFloat.h"
 #include "Graph/ControlRigGraph.h"
+#include "Graph/ControlRigGraphNode.h"
+#include "ControlRigBlueprint.h"
 #include "Widgets/Layout/SBox.h"
 #include "Widgets/Input/SEditableTextBox.h"
 #include "EdGraphSchema_K2.h"
 #include "ScopedTransaction.h"
 #include "DetailLayoutBuilder.h"
 #include "PropertyPathHelpers.h"
+#include "RigVMModel/Nodes/RigVMStructNode.h"
+#include "RigVMModel/RigVMController.h"
 
 #include "IControlRigEditorModule.h"
 
@@ -61,7 +65,14 @@ FRuntimeFloatCurve& SControlRigGraphPinCurveFloat::UpdateAndGetCurve()
 {
 	if (UEdGraphPin* Pin = GetPinObj())
 	{
-		FRuntimeFloatCurve::StaticStruct()->ImportText(*Pin->DefaultValue, &Curve, nullptr, EPropertyPortFlags::PPF_None, nullptr, FRuntimeFloatCurve::StaticStruct()->GetName(), true);
+		if (UControlRigGraphNode* RigNode = Cast<UControlRigGraphNode>(Pin->GetOwningNode()))
+		{
+			if (URigVMPin* ModelPin = RigNode->GetModelPinFromPinPath(Pin->GetName()))
+			{
+				FString TextToImport = ModelPin->GetDefaultValue();
+				FRuntimeFloatCurve::StaticStruct()->ImportText(*TextToImport, &Curve, nullptr, EPropertyPortFlags::PPF_None, nullptr, FRuntimeFloatCurve::StaticStruct()->GetName(), true);
+			}
+		}
 	}
 	return Curve;
 }
@@ -98,12 +109,12 @@ bool SControlRigGraphPinCurveFloat::IsValidCurve(FRichCurveEditInfo CurveInfo)
 	{
 		if (UControlRigGraphNode* Node = Cast<UControlRigGraphNode>(Pin->GetOwningNode()))
 		{
-			if (const FStructProperty* StructProperty = Node->GetUnitProperty())
+			if (URigVMStructNode* StructModelNode = Cast<URigVMStructNode>(Node->GetModelNode()))
 			{
 				FString NodeName, PropertyName;
 				Pin->PinName.ToString().Split(TEXT("."), &NodeName, &PropertyName);
 
-				if (FProperty* CurveProperty = StructProperty->Struct->FindPropertyByName(*PropertyName))
+				if (FProperty* CurveProperty = StructModelNode->GetScriptStruct()->FindPropertyByName(*PropertyName))
 				{
 					return true;
 				}
@@ -119,9 +130,12 @@ void SControlRigGraphPinCurveFloat::OnCurveChanged(const TArray<FRichCurveEditIn
 	{
 		if (UControlRigGraphNode* Node = Cast<UControlRigGraphNode>(Pin->GetOwningNode()))
 		{
-			Pin->DefaultValue.Empty();
-			FRuntimeFloatCurve::StaticStruct()->ExportText(Pin->DefaultValue, &Curve, nullptr, nullptr, EPropertyPortFlags::PPF_None, nullptr, true);
-			Node->CopyPinDefaultsToModel(Pin);
+			if (UControlRigBlueprint* RigBlueprint = Cast<UControlRigBlueprint>(Node->GetGraph()->GetOuter()))
+			{
+				FString ExportedText;
+				FRuntimeFloatCurve::StaticStruct()->ExportText(ExportedText, &Curve, nullptr, nullptr, EPropertyPortFlags::PPF_None, nullptr, true);
+				RigBlueprint->Controller->SetPinDefaultValue(Pin->GetName(), ExportedText, true, true, true);
+			}
 		}
 	}
 	ModifyOwner();
diff --git a/Engine/Plugins/Experimental/ControlRig/Source/ControlRigEditor/Private/Graph/SControlRigGraphPinEditableNameValueWidget.cpp b/Engine/Plugins/Experimental/ControlRig/Source/ControlRigEditor/Private/Graph/SControlRigGraphPinEditableNameValueWidget.cpp
new file mode 100644
--- /dev/null
+++ b/Engine/Plugins/Experimental/ControlRig/Source/ControlRigEditor/Private/Graph/SControlRigGraphPinEditableNameValueWidget.cpp
@@ -0,0 +1,260 @@
+// Copyright Epic Games, Inc. All Rights Reserved.
+
+#include "SControlRigGraphPinEditableNameValueWidget.h"
+#include "DetailLayoutBuilder.h"
+
+#define LOCTEXT_NAMESPACE "GraphPinEditableNameValueWidget"
+
+void SControlRigGraphPinEditableNameValueWidget::Construct(const FArguments& InArgs)
+{
+	this->OnComboBoxOpening = InArgs._OnComboBoxOpening;
+	this->OnSelectionChanged = InArgs._OnSelectionChanged;
+	this->OnGenerateWidget = InArgs._OnGenerateWidget;
+
+	OptionsSource = InArgs._OptionsSource;
+	CustomScrollbar = InArgs._CustomScrollbar;
+
+	TSharedRef<SWidget> ComboBoxMenuContent =
+		SNew(SBox)
+		.MaxDesiredHeight(InArgs._MaxListHeight)
+		[
+			SNew(SVerticalBox)
+
+			+ SVerticalBox::Slot()
+			.AutoHeight()
+			[
+				SAssignNew(this->SearchField, SEditableTextBox)
+				.HintText(LOCTEXT("SearchOrRename", "Search / Rename"))
+				.OnTextChanged(this, &SControlRigGraphPinEditableNameValueWidget::OnSearchTextChanged)
+				.OnTextCommitted(this, &SControlRigGraphPinEditableNameValueWidget::OnSearchTextCommitted)
+			]
+
+			+ SVerticalBox::Slot()
+			[
+				SAssignNew(this->ComboListView, SComboListType)
+				.ListItemsSource(OptionsSource)
+				.OnGenerateRow(this, &SControlRigGraphPinEditableNameValueWidget::GenerateMenuItemRow)
+				.OnSelectionChanged(this, &SControlRigGraphPinEditableNameValueWidget::OnSelectionChanged_Internal)
+				.SelectionMode(ESelectionMode::Single)
+				.ExternalScrollbar(InArgs._CustomScrollbar)
+			]
+		];
+
+	// Set up content
+	TSharedPtr<SWidget> ButtonContent = InArgs._Content.Widget;
+	if (InArgs._Content.Widget == SNullWidget::NullWidget)
+	{
+		SAssignNew(ButtonContent, STextBlock)
+			.Text(NSLOCTEXT("SControlRigGraphPinEditableNameValueWidget", "ContentWarning", "No Content Provided"))
+			.ColorAndOpacity(FLinearColor::Red);
+	}
+
+
+	SComboButton::Construct(SComboButton::FArguments()
+		.Method(InArgs._Method)
+		.ButtonContent()
+		[
+			ButtonContent.ToSharedRef()
+		]
+		.MenuContent()
+		[
+			ComboBoxMenuContent
+		]
+		.HasDownArrow(InArgs._HasDownArrow)
+		.ContentPadding(InArgs._ContentPadding)
+		.OnMenuOpenChanged(this, &SControlRigGraphPinEditableNameValueWidget::OnMenuOpenChanged)
+		.IsFocusable(true)
+		);
+	SetMenuContentWidgetToFocus(ComboListView);
+
+	// Need to establish the selected item at point of construction so its available for querying
+	// NB: If you need a selection to fire use SetItemSelection rather than setting an IntiallySelectedItem
+	SelectedItem = InArgs._InitiallySelectedItem;
+	if (TListTypeTraits<TSharedPtr<FString>>::IsPtrValid(SelectedItem))
+	{
+		ComboListView->Private_SetItemSelection(SelectedItem, true);
+	}
+
+}
+
+void SControlRigGraphPinEditableNameValueWidget::ClearSelection()
+{
+	ComboListView->ClearSelection();
+}
+
+void SControlRigGraphPinEditableNameValueWidget::SetSelectedItem(TSharedPtr<FString> InSelectedItem)
+{
+	if (TListTypeTraits<TSharedPtr<FString>>::IsPtrValid(InSelectedItem))
+	{
+		ComboListView->SetSelection(InSelectedItem);
+	}
+	else
+	{
+		ComboListView->ClearSelection();
+	}
+}
+
+TSharedPtr<FString> SControlRigGraphPinEditableNameValueWidget::GetSelectedItem()
+{
+	return SelectedItem;
+}
+
+void SControlRigGraphPinEditableNameValueWidget::RefreshOptions()
+{
+	if (!ComboListView->IsPendingRefresh())
+	{
+		ComboListView->RequestListRefresh();
+	}
+}
+
+TSharedRef<ITableRow> SControlRigGraphPinEditableNameValueWidget::GenerateMenuItemRow(TSharedPtr<FString> InItem, const TSharedRef<STableViewBase>& OwnerTable)
+{
+	if (OnGenerateWidget.IsBound())
+	{
+		FString SearchToken = SearchField->GetText().ToString().ToLower();
+		EVisibility WidgetVisibility = EVisibility::Visible;
+		if (!SearchToken.IsEmpty())
+		{
+			if (InItem->ToLower().Find(SearchToken) < 0)
+			{
+				WidgetVisibility = EVisibility::Collapsed;
+			}
+		}
+		return SNew(SComboRow<TSharedPtr<FString>>, OwnerTable)
+			.Visibility(WidgetVisibility)
+			[
+				OnGenerateWidget.Execute(InItem)
+			];
+	}
+	else
+	{
+		return SNew(SComboRow<TSharedPtr<FString>>, OwnerTable)
+			[
+				SNew(STextBlock).Text(NSLOCTEXT("SlateCore", "ComboBoxMissingOnGenerateWidgetMethod", "Please provide a .OnGenerateWidget() handler."))
+			];
+
+	}
+}
+
+void SControlRigGraphPinEditableNameValueWidget::OnMenuOpenChanged(bool bOpen)
+{
+	if (bOpen == false)
+	{
+		if (TListTypeTraits<TSharedPtr<FString>>::IsPtrValid(SelectedItem))
+		{
+			// Ensure the ListView selection is set back to the last committed selection
+			ComboListView->SetSelection(SelectedItem, ESelectInfo::OnNavigation);
+			ComboListView->RequestScrollIntoView(SelectedItem, 0);
+		}
+
+		// Set focus back to ComboBox for users focusing the ListView that just closed
+		FSlateApplication::Get().ForEachUser([&](FSlateUser& User) {
+			if (FSlateApplication::Get().HasUserFocusedDescendants(AsShared(), User.GetUserIndex()))
+			{
+				FSlateApplication::Get().SetUserFocus(User.GetUserIndex(), AsShared(), EFocusCause::SetDirectly);
+			}
+		});
+
+	}
+	else
+	{
+		RegisterActiveTimer(0.f, FWidgetActiveTimerDelegate::CreateSP(this, &SControlRigGraphPinEditableNameValueWidget::SetFocusPostConstruct));
+	}
+}
+
+EActiveTimerReturnType SControlRigGraphPinEditableNameValueWidget::SetFocusPostConstruct(double InCurrentTime, float InDeltaTime)
+{
+	if (SearchField.IsValid())
+	{
+		bool bSucceeded = false;
+		FSlateApplication::Get().ForEachUser([&](FSlateUser& User) {
+			if (FSlateApplication::Get().SetUserFocus(User.GetUserIndex(), SearchField->AsShared(), EFocusCause::SetDirectly))
+			{
+				bSucceeded = true;
+			}
+		});
+
+		if (bSucceeded)
+		{
+			return EActiveTimerReturnType::Stop;
+		}
+	}
+
+	return EActiveTimerReturnType::Continue;
+}
+
+void SControlRigGraphPinEditableNameValueWidget::OnSelectionChanged_Internal(TSharedPtr<FString> ProposedSelection, ESelectInfo::Type SelectInfo)
+{
+	// Ensure that the proposed selection is different
+	if (SelectInfo != ESelectInfo::OnNavigation)
+	{
+		// Ensure that the proposed selection is different from selected
+		if (ProposedSelection != SelectedItem)
+		{
+			SelectedItem = ProposedSelection;
+			OnSelectionChanged.ExecuteIfBound(ProposedSelection, SelectInfo);
+		}
+		// close combo even if user reselected item
+		this->SetIsOpen(false);
+	}
+}
+
+void SControlRigGraphPinEditableNameValueWidget::OnSearchTextChanged(const FText& ChangedText)
+{
+	FString SearchToken = ChangedText.ToString().ToLower();
+	for (int32 i = 0; i < OptionsSource->Num(); i++)
+	{
+		TSharedPtr<ITableRow> Row = ComboListView->WidgetFromItem((*OptionsSource)[i]);
+		if (Row)
+		{
+			if (SearchToken.IsEmpty())
+			{
+				Row->AsWidget()->SetVisibility(EVisibility::Visible);
+			}
+			else if ((*OptionsSource)[i]->ToLower().Find(SearchToken) >= 0)
+			{
+				Row->AsWidget()->SetVisibility(EVisibility::Visible);
+			}
+			else
+			{
+				Row->AsWidget()->SetVisibility(EVisibility::Collapsed);
+			}
+		}
+	}
+
+	ComboListView->RequestListRefresh();
+
+	SelectedItem = TSharedPtr< FString >();
+}
+
+void SControlRigGraphPinEditableNameValueWidget::OnSearchTextCommitted(const FText& ChangedText, ETextCommit::Type CommitType)
+{
+	if (CommitType == ETextCommit::OnEnter)
+	{
+		TSharedPtr<FString> ProposedSelection = MakeShared<FString>(ChangedText.ToString());
+		OnSelectionChanged_Internal(ProposedSelection, ESelectInfo::OnKeyPress);
+	}
+}
+
+FReply SControlRigGraphPinEditableNameValueWidget::OnButtonClicked()
+{
+	// if user clicked to close the combo menu
+	if (this->IsOpen())
+	{
+		// Re-select first selected item, just in case it was selected by navigation previously
+		TArray<TSharedPtr<FString>> SelectedItems = ComboListView->GetSelectedItems();
+		if (SelectedItems.Num() > 0)
+		{
+			OnSelectionChanged_Internal(SelectedItems[0], ESelectInfo::Direct);
+		}
+	}
+	else
+	{
+		SearchField->SetText(FText::GetEmpty());
+		OnComboBoxOpening.ExecuteIfBound();
+	}
+
+	return SComboButton::OnButtonClicked();
+}
+
+#undef LOCTEXT_NAMESPACE
diff --git a/Engine/Plugins/Experimental/ControlRig/Source/ControlRigEditor/Private/Graph/SControlRigGraphPinEditableNameValueWidget.h b/Engine/Plugins/Experimental/ControlRig/Source/ControlRigEditor/Private/Graph/SControlRigGraphPinEditableNameValueWidget.h
new file mode 100644
--- /dev/null
+++ b/Engine/Plugins/Experimental/ControlRig/Source/ControlRigEditor/Private/Graph/SControlRigGraphPinEditableNameValueWidget.h
@@ -0,0 +1,134 @@
+// Copyright Epic Games, Inc. All Rights Reserved.
+
+#pragma once
+
+#include "CoreMinimal.h"
+#include "Input/Reply.h"
+#include "Widgets/SWidget.h"
+#include "Widgets/Input/SComboButton.h"
+#include "Widgets/Input/SEditableTextBox.h"
+#include "Widgets/Views/SListView.h"
+#include "Widgets/Input/SComboBox.h"
+
+/**
+ * A searchable text combo box
+ */
+class SControlRigGraphPinEditableNameValueWidget : public SComboButton
+{
+public:
+	/** Type of list used for showing menu options. */
+	typedef SListView< TSharedPtr<FString> > SComboListType;
+	/** Delegate type used to generate widgets that represent Options */
+	typedef typename TSlateDelegates< TSharedPtr<FString> >::FOnGenerateWidget FOnGenerateWidget;
+	typedef typename TSlateDelegates< TSharedPtr<FString> >::FOnSelectionChanged FOnSelectionChanged;
+
+	SLATE_BEGIN_ARGS(SControlRigGraphPinEditableNameValueWidget)
+		: _Content()
+		, _ContentPadding(FMargin(3.0, 3.0))
+		, _OptionsSource()
+		, _OnSelectionChanged()
+		, _OnGenerateWidget()
+		, _InitiallySelectedItem(nullptr)
+		, _Method()
+		, _MaxListHeight(450.0f)
+		, _HasDownArrow(true)
+	{}
+
+	/** Slot for this button's content (optional) */
+	SLATE_DEFAULT_SLOT(FArguments, Content)
+
+		SLATE_ATTRIBUTE(FMargin, ContentPadding)
+
+		SLATE_ARGUMENT(const TArray< TSharedPtr<FString> >*, OptionsSource)
+		SLATE_EVENT(FOnSelectionChanged, OnSelectionChanged)
+		SLATE_EVENT(FOnGenerateWidget, OnGenerateWidget)
+
+		/** Called when combo box is opened, before list is actually created */
+		SLATE_EVENT(FOnComboBoxOpening, OnComboBoxOpening)
+
+		/** The custom scrollbar to use in the ListView */
+		SLATE_ARGUMENT(TSharedPtr<SScrollBar>, CustomScrollbar)
+
+		/** The option that should be selected when the combo box is first created */
+		SLATE_ARGUMENT(TSharedPtr<FString>, InitiallySelectedItem)
+
+		SLATE_ARGUMENT(TOptional<EPopupMethod>, Method)
+
+		/** The max height of the combo box menu */
+		SLATE_ARGUMENT(float, MaxListHeight)
+
+		/**
+		 * When false, the down arrow is not generated and it is up to the API consumer
+		 * to make their own visual hint that this is a drop down.
+		 */
+		SLATE_ARGUMENT(bool, HasDownArrow)
+
+	SLATE_END_ARGS()
+
+	/**
+	 * Construct the widget from a declaration
+	 *
+	 * @param InArgs   Declaration from which to construct the combo box
+	 */
+	void Construct(const FArguments& InArgs);
+
+	void ClearSelection();
+
+	void SetSelectedItem(TSharedPtr<FString> InSelectedItem);
+
+	/** @return the item currently selected by the combo box. */
+	TSharedPtr<FString> GetSelectedItem();
+
+	/**
+	 * Requests a list refresh after updating options
+	 * Call SetSelectedItem to update the selected item if required
+	 * @see SetSelectedItem
+	 */
+	void RefreshOptions();
+
+private:
+
+	/** Generate a row for the InItem in the combo box's list (passed in as OwnerTable). Do this by calling the user-specified OnGenerateWidget */
+	TSharedRef<ITableRow> GenerateMenuItemRow(TSharedPtr<FString> InItem, const TSharedRef<STableViewBase>& OwnerTable);
+
+	//** Called if the menu is closed
+	void OnMenuOpenChanged(bool bOpen);
+
+	/** Invoked when the selection in the list changes */
+	void OnSelectionChanged_Internal(TSharedPtr<FString> ProposedSelection, ESelectInfo::Type SelectInfo);
+
+	/** Invoked when the search text changes */
+	void OnSearchTextChanged(const FText& ChangedText);
+
+	/** Invoked when the search is committed*/
+	void OnSearchTextCommitted(const FText& ChangedText, ETextCommit::Type CommitType);
+
+	/** Handle clicking on the content menu */
+	virtual FReply OnButtonClicked() override;
+
+	/** The item style to use. */
+	const FTableRowStyle* ItemStyle;
+
+private:
+	/** Delegate that is invoked when the selected item in the combo box changes */
+	FOnSelectionChanged OnSelectionChanged;
+	/** The item currently selected in the combo box */
+	TSharedPtr<FString> SelectedItem;
+	/** The search field used for the combox box's contents */
+	TSharedPtr< SEditableTextBox > SearchField;
+	/** The ListView that we pop up; visualized the available options. */
+	TSharedPtr< SComboListType > ComboListView;
+	/** The Scrollbar used in the ListView. */
+	TSharedPtr< SScrollBar > CustomScrollbar;
+	/** Delegate to invoke before the combo box is opening. */
+	FOnComboBoxOpening OnComboBoxOpening;
+	/** Delegate to invoke when we need to visualize an option as a widget. */
+	FOnGenerateWidget OnGenerateWidget;
+
+	const TArray< TSharedPtr<FString> >* OptionsSource;
+
+	/** Used to focus the name box immediately following construction */
+	EActiveTimerReturnType SetFocusPostConstruct(double InCurrentTime, float InDeltaTime);
+};
+
+
diff --git a/Engine/Plugins/Experimental/ControlRig/Source/ControlRigEditor/Private/Graph/SControlRigGraphPinNameList.cpp b/Engine/Plugins/Experimental/ControlRig/Source/ControlRigEditor/Private/Graph/SControlRigGraphPinNameList.cpp
--- a/Engine/Plugins/Experimental/ControlRig/Source/ControlRigEditor/Private/Graph/SControlRigGraphPinNameList.cpp
+++ b/Engine/Plugins/Experimental/ControlRig/Source/ControlRigEditor/Private/Graph/SControlRigGraphPinNameList.cpp
@@ -29,7 +29,6 @@ TSharedRef<SWidget>	SControlRigGraphPinNameList::GetDefaultValueWidget()
 
 	return SNew(SBox)
 		.MinDesiredWidth(150)
-		.MaxDesiredWidth(400)
 		[
 			SAssignNew(NameListComboBox, SControlRigGraphPinNameListValueWidget)
 				.Visibility(this, &SGraphPin::GetDefaultValueVisibility)
diff --git a/Engine/Plugins/Experimental/ControlRig/Source/ControlRigEditor/Private/Graph/SControlRigGraphPinNameListValueWidget.cpp b/Engine/Plugins/Experimental/ControlRig/Source/ControlRigEditor/Private/Graph/SControlRigGraphPinNameListValueWidget.cpp
--- a/Engine/Plugins/Experimental/ControlRig/Source/ControlRigEditor/Private/Graph/SControlRigGraphPinNameListValueWidget.cpp
+++ b/Engine/Plugins/Experimental/ControlRig/Source/ControlRigEditor/Private/Graph/SControlRigGraphPinNameListValueWidget.cpp
@@ -17,6 +17,7 @@ void SControlRigGraphPinNameListValueWidget::Construct(const FArguments& InArgs)
 
 	TSharedRef<SWidget> ComboBoxMenuContent =
 		SNew(SBox)
+		.MinDesiredWidth(150.0f)
 		.MaxDesiredHeight(InArgs._MaxListHeight)
 		[
 			SNew(SVerticalBox)
diff --git a/Engine/Plugins/Experimental/ControlRig/Source/ControlRigEditor/Private/Graph/SControlRigGraphPinParameterName.cpp b/Engine/Plugins/Experimental/ControlRig/Source/ControlRigEditor/Private/Graph/SControlRigGraphPinParameterName.cpp
new file mode 100644
--- /dev/null
+++ b/Engine/Plugins/Experimental/ControlRig/Source/ControlRigEditor/Private/Graph/SControlRigGraphPinParameterName.cpp
@@ -0,0 +1,120 @@
+// Copyright Epic Games, Inc. All Rights Reserved.
+
+
+#include "Graph/SControlRigGraphPinParameterName.h"
+#include "Graph/ControlRigGraph.h"
+#include "Widgets/Layout/SBox.h"
+#include "Widgets/Input/SEditableTextBox.h"
+#include "EdGraphSchema_K2.h"
+#include "ScopedTransaction.h"
+#include "DetailLayoutBuilder.h"
+
+void SControlRigGraphPinParameterName::Construct(const FArguments& InArgs, UEdGraphPin* InGraphPinObj)
+{
+	SGraphPin::Construct(SGraphPin::FArguments(), InGraphPinObj);
+}
+
+TSharedRef<SWidget>	SControlRigGraphPinParameterName::GetDefaultValueWidget()
+{
+	UControlRigGraphNode* RigNode = Cast<UControlRigGraphNode>(GraphPinObj->GetOwningNode());
+	URigVMGraph* Model = RigNode->GetModelNode()->GetGraph();
+
+	TSharedPtr<FString> InitialSelected;
+
+	TArray<TSharedPtr<FString>>& LocalParameterNames = GetParameterNames();
+	for (TSharedPtr<FString> Item : LocalParameterNames)
+	{
+		if (Item->Equals(GetParameterNameText().ToString()))
+		{
+			InitialSelected = Item;
+		}
+	}
+
+	return SNew(SBox)
+		.MinDesiredWidth(150)
+		.MaxDesiredWidth(400)
+		[
+			SAssignNew(NameComboBox, SControlRigGraphPinEditableNameValueWidget)
+				.Visibility(this, &SGraphPin::GetDefaultValueVisibility)
+				.OptionsSource(&ParameterNames)
+				.OnGenerateWidget(this, &SControlRigGraphPinParameterName::MakeParameterNameItemWidget)
+				.OnSelectionChanged(this, &SControlRigGraphPinParameterName::OnParameterNameChanged)
+				.OnComboBoxOpening(this, &SControlRigGraphPinParameterName::OnParameterNameComboBox)
+				.InitiallySelectedItem(InitialSelected)
+				.Content()
+				[
+					SNew(STextBlock)
+					.Text(this, &SControlRigGraphPinParameterName::GetParameterNameText)
+				]
+		];
+}
+
+FText SControlRigGraphPinParameterName::GetParameterNameText() const
+{
+	return FText::FromString( GraphPinObj->GetDefaultAsString() );
+}
+
+void SControlRigGraphPinParameterName::SetParameterNameText(const FText& NewTypeInValue, ETextCommit::Type /*CommitInfo*/)
+{
+	if(!GraphPinObj->GetDefaultAsString().Equals(NewTypeInValue.ToString()))
+	{
+		const FScopedTransaction Transaction( NSLOCTEXT("GraphEditor", "ChangeParameterNamePinValue", "Change Bone Name Pin Value" ) );
+		GraphPinObj->Modify();
+		GraphPinObj->GetSchema()->TrySetDefaultValue(*GraphPinObj, NewTypeInValue.ToString());
+	}
+}
+
+TSharedRef<SWidget> SControlRigGraphPinParameterName::MakeParameterNameItemWidget(TSharedPtr<FString> InItem)
+{
+	return 	SNew(STextBlock).Text(FText::FromString(*InItem));// .Font(FEditorStyle::GetFontStyle(TEXT("PropertyWindow.NormalFont")));
+}
+
+void SControlRigGraphPinParameterName::OnParameterNameChanged(TSharedPtr<FString> NewSelection, ESelectInfo::Type SelectInfo)
+{
+	if (SelectInfo != ESelectInfo::Direct)
+	{
+		FString NewValue = *NewSelection.Get();
+		SetParameterNameText(FText::FromString(NewValue), ETextCommit::OnEnter);
+	}
+}
+
+void SControlRigGraphPinParameterName::OnParameterNameComboBox()
+{
+	TSharedPtr<FString> CurrentlySelected;
+	TArray<TSharedPtr<FString>>& LocalParameterNames = GetParameterNames();
+	for (TSharedPtr<FString> Item : LocalParameterNames)
+	{
+		if (Item->Equals(GetParameterNameText().ToString()))
+		{
+			CurrentlySelected = Item;
+		}
+	}
+
+	NameComboBox->SetSelectedItem(CurrentlySelected);
+}
+
+TArray<TSharedPtr<FString>>& SControlRigGraphPinParameterName::GetParameterNames()
+{
+	if(UControlRigGraphNode* RigNode = Cast<UControlRigGraphNode>(GraphPinObj->GetOwningNode()))
+	{
+		if(URigVMParameterNode* ModelNode = Cast<URigVMParameterNode>(RigNode->GetModelNode()))
+		{
+			ParameterNames.Reset();
+
+			FRigVMGraphParameterDescription MyDescription = ModelNode->GetParameterDescription();
+
+			TArray<FRigVMGraphParameterDescription> ParameterDescriptions = ModelNode->GetGraph()->GetParameterDescriptions();
+			Algo::SortBy(ParameterDescriptions, &FRigVMGraphParameterDescription::Name, FNameLexicalLess());
+
+			for (FRigVMGraphParameterDescription& ParameterDescription : ParameterDescriptions)
+			{
+				if (ParameterDescription.CPPType == MyDescription.CPPType && ParameterDescription.bIsInput == MyDescription.bIsInput)
+				{
+					ParameterNames.Add(MakeShared<FString>(ParameterDescription.Name.ToString()));
+				}
+			}
+		}
+	}
+
+	return ParameterNames;
+}
diff --git a/Engine/Plugins/Experimental/ControlRig/Source/ControlRigEditor/Private/Graph/SControlRigGraphPinParameterName.h b/Engine/Plugins/Experimental/ControlRig/Source/ControlRigEditor/Private/Graph/SControlRigGraphPinParameterName.h
new file mode 100644
--- /dev/null
+++ b/Engine/Plugins/Experimental/ControlRig/Source/ControlRigEditor/Private/Graph/SControlRigGraphPinParameterName.h
@@ -0,0 +1,35 @@
+// Copyright Epic Games, Inc. All Rights Reserved.
+
+#pragma once
+
+#include "CoreMinimal.h"
+#include "Widgets/DeclarativeSyntaxSupport.h"
+#include "Widgets/SWidget.h"
+#include "SGraphPin.h"
+#include "SControlRigGraphPinEditableNameValueWidget.h"
+
+class SControlRigGraphPinParameterName : public SGraphPin
+{
+public:
+	SLATE_BEGIN_ARGS(SControlRigGraphPinParameterName) {}
+	SLATE_END_ARGS()
+
+		void Construct(const FArguments& InArgs, UEdGraphPin* InGraphPinObj);
+
+protected:
+	//~ Begin SGraphPin Interface
+	virtual TSharedRef<SWidget>	GetDefaultValueWidget() override;
+	//~ End SGraphPin Interface
+
+	FText GetParameterNameText() const;
+	virtual void SetParameterNameText(const FText& NewTypeInValue, ETextCommit::Type CommitInfo);
+
+	TSharedRef<SWidget> MakeParameterNameItemWidget(TSharedPtr<FString> InItem);
+	void OnParameterNameChanged(TSharedPtr<FString> NewSelection, ESelectInfo::Type SelectInfo);
+	void OnParameterNameComboBox();
+	TArray<TSharedPtr<FString>>& GetParameterNames();
+
+	TArray<TSharedPtr<FString>> ParameterNames;
+	TSharedPtr<SControlRigGraphPinEditableNameValueWidget> NameComboBox;
+
+};
diff --git a/Engine/Plugins/Experimental/ControlRig/Source/ControlRigEditor/Private/Graph/SControlRigGraphPinVariableName.cpp b/Engine/Plugins/Experimental/ControlRig/Source/ControlRigEditor/Private/Graph/SControlRigGraphPinVariableName.cpp
new file mode 100644
--- /dev/null
+++ b/Engine/Plugins/Experimental/ControlRig/Source/ControlRigEditor/Private/Graph/SControlRigGraphPinVariableName.cpp
@@ -0,0 +1,120 @@
+// Copyright Epic Games, Inc. All Rights Reserved.
+
+
+#include "Graph/SControlRigGraphPinVariableName.h"
+#include "Graph/ControlRigGraph.h"
+#include "Widgets/Layout/SBox.h"
+#include "Widgets/Input/SEditableTextBox.h"
+#include "EdGraphSchema_K2.h"
+#include "ScopedTransaction.h"
+#include "DetailLayoutBuilder.h"
+
+void SControlRigGraphPinVariableName::Construct(const FArguments& InArgs, UEdGraphPin* InGraphPinObj)
+{
+	SGraphPin::Construct(SGraphPin::FArguments(), InGraphPinObj);
+}
+
+TSharedRef<SWidget>	SControlRigGraphPinVariableName::GetDefaultValueWidget()
+{
+	UControlRigGraphNode* RigNode = Cast<UControlRigGraphNode>(GraphPinObj->GetOwningNode());
+	URigVMGraph* Model = RigNode->GetModelNode()->GetGraph();
+
+	TSharedPtr<FString> InitialSelected;
+
+	TArray<TSharedPtr<FString>>& LocalVariableNames = GetVariableNames();
+	for (TSharedPtr<FString> Item : LocalVariableNames)
+	{
+		if (Item->Equals(GetVariableNameText().ToString()))
+		{
+			InitialSelected = Item;
+		}
+	}
+
+	return SNew(SBox)
+		.MinDesiredWidth(150)
+		.MaxDesiredWidth(400)
+		[
+			SAssignNew(NameComboBox, SControlRigGraphPinEditableNameValueWidget)
+				.Visibility(this, &SGraphPin::GetDefaultValueVisibility)
+				.OptionsSource(&VariableNames)
+				.OnGenerateWidget(this, &SControlRigGraphPinVariableName::MakeVariableNameItemWidget)
+				.OnSelectionChanged(this, &SControlRigGraphPinVariableName::OnVariableNameChanged)
+				.OnComboBoxOpening(this, &SControlRigGraphPinVariableName::OnVariableNameComboBox)
+				.InitiallySelectedItem(InitialSelected)
+				.Content()
+				[
+					SNew(STextBlock)
+					.Text(this, &SControlRigGraphPinVariableName::GetVariableNameText)
+				]
+		];
+}
+
+FText SControlRigGraphPinVariableName::GetVariableNameText() const
+{
+	return FText::FromString( GraphPinObj->GetDefaultAsString() );
+}
+
+void SControlRigGraphPinVariableName::SetVariableNameText(const FText& NewTypeInValue, ETextCommit::Type /*CommitInfo*/)
+{
+	if(!GraphPinObj->GetDefaultAsString().Equals(NewTypeInValue.ToString()))
+	{
+		const FScopedTransaction Transaction( NSLOCTEXT("GraphEditor", "ChangeVariableNamePinValue", "Change Bone Name Pin Value" ) );
+		GraphPinObj->Modify();
+		GraphPinObj->GetSchema()->TrySetDefaultValue(*GraphPinObj, NewTypeInValue.ToString());
+	}
+}
+
+TSharedRef<SWidget> SControlRigGraphPinVariableName::MakeVariableNameItemWidget(TSharedPtr<FString> InItem)
+{
+	return 	SNew(STextBlock).Text(FText::FromString(*InItem));// .Font(FEditorStyle::GetFontStyle(TEXT("PropertyWindow.NormalFont")));
+}
+
+void SControlRigGraphPinVariableName::OnVariableNameChanged(TSharedPtr<FString> NewSelection, ESelectInfo::Type SelectInfo)
+{
+	if (SelectInfo != ESelectInfo::Direct)
+	{
+		FString NewValue = *NewSelection.Get();
+		SetVariableNameText(FText::FromString(NewValue), ETextCommit::OnEnter);
+	}
+}
+
+void SControlRigGraphPinVariableName::OnVariableNameComboBox()
+{
+	TSharedPtr<FString> CurrentlySelected;
+	TArray<TSharedPtr<FString>>& LocalVariableNames = GetVariableNames();
+	for (TSharedPtr<FString> Item : LocalVariableNames)
+	{
+		if (Item->Equals(GetVariableNameText().ToString()))
+		{
+			CurrentlySelected = Item;
+		}
+	}
+
+	NameComboBox->SetSelectedItem(CurrentlySelected);
+}
+
+TArray<TSharedPtr<FString>>& SControlRigGraphPinVariableName::GetVariableNames()
+{
+	if(UControlRigGraphNode* RigNode = Cast<UControlRigGraphNode>(GraphPinObj->GetOwningNode()))
+	{
+		if(URigVMVariableNode* ModelNode = Cast<URigVMVariableNode>(RigNode->GetModelNode()))
+		{
+			VariableNames.Reset();
+
+			FRigVMGraphVariableDescription MyDescription = ModelNode->GetVariableDescription();
+
+			TArray<FRigVMGraphVariableDescription> VariableDescriptions = ModelNode->GetGraph()->GetVariableDescriptions();
+			Algo::SortBy(VariableDescriptions, &FRigVMGraphVariableDescription::Name, FNameLexicalLess());
+
+			for (FRigVMGraphVariableDescription& VariableDescription : VariableDescriptions)
+			{
+				if (VariableDescription.CPPType == MyDescription.CPPType)
+				{
+					VariableNames.Add(MakeShared<FString>(VariableDescription.Name.ToString()));
+				}
+			}
+		}
+	}
+
+	return VariableNames;
+}
diff --git a/Engine/Plugins/Experimental/ControlRig/Source/ControlRigEditor/Private/Graph/SControlRigGraphPinVariableName.h b/Engine/Plugins/Experimental/ControlRig/Source/ControlRigEditor/Private/Graph/SControlRigGraphPinVariableName.h
new file mode 100644
--- /dev/null
+++ b/Engine/Plugins/Experimental/ControlRig/Source/ControlRigEditor/Private/Graph/SControlRigGraphPinVariableName.h
@@ -0,0 +1,35 @@
+// Copyright Epic Games, Inc. All Rights Reserved.
+
+#pragma once
+
+#include "CoreMinimal.h"
+#include "Widgets/DeclarativeSyntaxSupport.h"
+#include "Widgets/SWidget.h"
+#include "SGraphPin.h"
+#include "SControlRigGraphPinEditableNameValueWidget.h"
+
+class SControlRigGraphPinVariableName : public SGraphPin
+{
+public:
+	SLATE_BEGIN_ARGS(SControlRigGraphPinVariableName) {}
+	SLATE_END_ARGS()
+
+		void Construct(const FArguments& InArgs, UEdGraphPin* InGraphPinObj);
+
+protected:
+	//~ Begin SGraphPin Interface
+	virtual TSharedRef<SWidget>	GetDefaultValueWidget() override;
+	//~ End SGraphPin Interface
+
+	FText GetVariableNameText() const;
+	virtual void SetVariableNameText(const FText& NewTypeInValue, ETextCommit::Type CommitInfo);
+
+	TSharedRef<SWidget> MakeVariableNameItemWidget(TSharedPtr<FString> InItem);
+	void OnVariableNameChanged(TSharedPtr<FString> NewSelection, ESelectInfo::Type SelectInfo);
+	void OnVariableNameComboBox();
+	TArray<TSharedPtr<FString>>& GetVariableNames();
+
+	TArray<TSharedPtr<FString>> VariableNames;
+	TSharedPtr<SControlRigGraphPinEditableNameValueWidget> NameComboBox;
+
+};
