COPY from //Dev-Anim to //Dev-Main

diff --git a/Engine/Plugins/Experimental/ControlRig/Source/ControlRigEditor/Private/Sequencer/ControlRigParameterTrackEditor.cpp b/Engine/Plugins/Experimental/ControlRig/Source/ControlRigEditor/Private/Sequencer/ControlRigParameterTrackEditor.cpp
--- a/Engine/Plugins/Experimental/ControlRig/Source/ControlRigEditor/Private/Sequencer/ControlRigParameterTrackEditor.cpp
+++ b/Engine/Plugins/Experimental/ControlRig/Source/ControlRigEditor/Private/Sequencer/ControlRigParameterTrackEditor.cpp
@@ -47,54 +47,17 @@
 #include "ControlRigSkeletalMeshBinding.h"
 #include "LevelEditorViewport.h"
 #include "IKeyArea.h"
+#include "ISequencer.h"
+#include "CurveModel.h"
+#include "ControlRigEditorModule.h"
+#include "SequencerSettings.h"
+#include "Channels/FloatChannelCurveModel.h"
+#include "Sequencer/ControlRigSortedControls.h"
+#include "TransformNoScale.h"
 
-#define LOCTEXT_NAMESPACE "FControlRigParameterTrackEditor"
-
-
-
-class FControlRigTrackCommands
-	: public TCommands<FControlRigTrackCommands>
-{
-public:
-
-	FControlRigTrackCommands()
-		: TCommands<FControlRigTrackCommands>
-		(
-			"ControlRigTrack",
-			NSLOCTEXT("Contexts", "ControlRigTrack", "ControlRigTrack"),
-			NAME_None,
-			FEditorStyle::GetStyleSetName() // Icon Style Set
-			)
-		, BindingCount(0)
-	{ }
-
-	/** Sets a transform key at the current time for the selected actor */
-	TSharedPtr< FUICommandInfo > AddTransformKey;
-
-	/** Sets a translation key at the current time for the selected actor */
-	TSharedPtr< FUICommandInfo > AddTranslationKey;
 
-	/** Sets a rotation key at the current time for the selected actor */
-	TSharedPtr< FUICommandInfo > AddRotationKey;
 
-	/** Sets a scale key at the current time for the selected actor */
-	TSharedPtr< FUICommandInfo > AddScaleKey;
-
-	/**
-	 * Initialize commands
-	 */
-	virtual void RegisterCommands() override;
-
-	mutable uint32 BindingCount;
-};
-
-void FControlRigTrackCommands::RegisterCommands()
-{
-	UI_COMMAND(AddTransformKey, "Add Transform Key", "Add a transform key at the current time for the selected control rig control.", EUserInterfaceActionType::Button, FInputChord(EKeys::S));
-	UI_COMMAND(AddTranslationKey, "Add Translation Key", "Add a translation key at the current time for the selected control rig control.", EUserInterfaceActionType::Button, FInputChord(EModifierKey::Shift, EKeys::W));
-	UI_COMMAND(AddRotationKey, "Add Rotation Key", "Add a rotation key at the current time for the selected control rig control.", EUserInterfaceActionType::Button, FInputChord(EModifierKey::Shift, EKeys::E));
-	UI_COMMAND(AddScaleKey, "Add Scale Key", "Add a scale key at the current time for the selected control rig control.", EUserInterfaceActionType::Button, FInputChord(EModifierKey::Shift, EKeys::R));
-}
+#define LOCTEXT_NAMESPACE "FControlRigParameterTrackEditor"
 
 static USkeleton* GetSkeletonFromComponent(UActorComponent* InComponent)
 {
@@ -163,12 +126,12 @@ static USkeleton* AcquireSkeletonFromObjectGuid(const FGuid& Guid, UObject** Obj
 }
 
 FControlRigParameterTrackEditor::FControlRigParameterTrackEditor(TSharedRef<ISequencer> InSequencer)
-	: FKeyframeTrackEditor<UMovieSceneControlRigParameterTrack>(InSequencer), bIsDoingASelection(false)
+	: FKeyframeTrackEditor<UMovieSceneControlRigParameterTrack>(InSequencer), bIsDoingSelection(false)
 {
 	SelectionChangedHandle = InSequencer->GetSelectionChangedTracks().AddRaw(this, &FControlRigParameterTrackEditor::OnSelectionChanged);
 	SequencerChangedHandle = InSequencer->OnMovieSceneDataChanged().AddRaw(this, &FControlRigParameterTrackEditor::OnSequencerDataChanged);
+	SelectionChangedHandle = InSequencer->GetCurveDisplayChanged().AddRaw(this, &FControlRigParameterTrackEditor::OnCurveDisplayChanged);
 
-	FControlRigTrackCommands::Register();
 	//register all modified/selections for control rigs
 	UMovieScene* MovieScene = InSequencer->GetFocusedMovieSceneSequence()->GetMovieScene();
 	const TArray<FMovieSceneBinding>& Bindings = MovieScene->GetBindings();
@@ -187,6 +150,17 @@ FControlRigParameterTrackEditor::~FControlRigParameterTrackEditor()
 {
 }
 
+void FControlRigParameterTrackEditor::ObjectImplicitlyAdded(UObject* InObject) 
+{
+	UControlRig* ControlRig = Cast<UControlRig>(InObject);
+	if (ControlRig)
+	{
+		ControlRig->ControlModified().AddRaw(this, &FControlRigParameterTrackEditor::HandleControlModified);
+		ControlRig->ControlSelected().AddRaw(this, &FControlRigParameterTrackEditor::HandleControlSelected);
+	}
+}
+
+
 void FControlRigParameterTrackEditor::OnRelease()
 {
 	if (GetSequencer().IsValid())
@@ -222,15 +196,7 @@ void FControlRigParameterTrackEditor::OnRelease()
 			GLevelEditorModeTools().DeactivateMode(FControlRigEditMode::ModeName);
 		}
 
-		ControlRigEditMode->SetObjects(nullptr, FGuid(), nullptr);
-	}
-
-	const FControlRigTrackCommands& Commands = FControlRigTrackCommands::Get();
-	Commands.BindingCount--;
-
-	if (Commands.BindingCount < 1)
-	{
-		FControlRigTrackCommands::Unregister();
+		ControlRigEditMode->SetObjects(nullptr, FGuid(), nullptr,GetSequencer());
 	}
 
 }
@@ -436,27 +402,33 @@ void FControlRigParameterTrackEditor::AddControlRig(UClass* InClass, UObject* Bo
 
 	if (InClass && InClass->IsChildOf(UControlRig::StaticClass()) && SequencerParent.IsValid())
 	{
+		UMovieSceneSequence* OwnerSequence = GetSequencer()->GetFocusedMovieSceneSequence();
+		UMovieScene* OwnerMovieScene = OwnerSequence->GetMovieScene();
 		FScopedTransaction AddControlRigTrackTransaction(LOCTEXT("AddControlRigTrack_Transaction", "Add Control Rig Track"));
 
-		UMovieSceneControlRigParameterTrack* Track = Cast<UMovieSceneControlRigParameterTrack>(AddTrack(GetSequencer()->GetFocusedMovieSceneSequence()->GetMovieScene(), SkelMeshBinding, UMovieSceneControlRigParameterTrack::StaticClass(), NAME_None));
+		UMovieSceneControlRigParameterTrack* Track = Cast<UMovieSceneControlRigParameterTrack>(AddTrack(OwnerMovieScene, SkelMeshBinding, UMovieSceneControlRigParameterTrack::StaticClass(), NAME_None));
 		if (Track)
 		{
+			OwnerSequence->Modify();
+			OwnerMovieScene->Modify();
 
 			FString ObjectName = (InClass->GetName());
 			ObjectName.RemoveFromEnd(TEXT("_C"));
 
 			UControlRig* ControlRig = NewObject<UControlRig>(Track, InClass, FName(*ObjectName), RF_Transactional);
-
+			ControlRig->Modify();
 			ControlRig->SetObjectBinding(MakeShared<FControlRigSkeletalMeshBinding>());
 			ControlRig->GetObjectBinding()->BindToObject(BoundSkelMesh);
+			ControlRig->GetDataSourceRegistry()->RegisterDataSource(UControlRig::OwnerComponent, ControlRig->GetObjectBinding()->GetBoundObject());
 			ControlRig->Initialize();
 			ControlRig->Execute(EControlRigState::Update);
+			ControlRig->CreateRigControlsForCurveContainer();
 
 			SequencerParent->NotifyMovieSceneDataChanged(EMovieSceneDataChangeType::MovieSceneStructureItemsChanged);
 
 			Track->Modify();
 			UMovieSceneSection* NewSection = Track->CreateControlRigSection(0, ControlRig);
-
+			NewSection->Modify();
 			//mz todo need to have multiple rigs with same class
 			Track->SetTrackName(FName(*ObjectName));
 			Track->SetDisplayName(FText::FromString(ObjectName));
@@ -475,7 +447,7 @@ void FControlRigParameterTrackEditor::AddControlRig(UClass* InClass, UObject* Bo
 			}
 			if (ControlRigEditMode)
 			{
-				ControlRigEditMode->SetObjects(ControlRig, FGuid(), nullptr);
+				ControlRigEditMode->SetObjects(ControlRig, FGuid(), nullptr, GetSequencer());
 			}
 			ControlRig->ControlModified().AddRaw(this, &FControlRigParameterTrackEditor::HandleControlModified);
 			ControlRig->ControlSelected().AddRaw(this, &FControlRigParameterTrackEditor::HandleControlSelected);
@@ -483,42 +455,12 @@ void FControlRigParameterTrackEditor::AddControlRig(UClass* InClass, UObject* Bo
 	}
 }
 
-
-void FControlRigParameterTrackEditor::BindCommands(TSharedRef<FUICommandList> SequencerCommandBindings)
+bool FControlRigParameterTrackEditor::HasTransformKeyOverridePriority() const
 {
-	const FControlRigTrackCommands& Commands = FControlRigTrackCommands::Get();
-
-	CommandBindings = MakeShared<FUICommandList>();
-	CommandBindings->MapAction(
-		Commands.AddTransformKey,
-		FExecuteAction::CreateSP(this, &FControlRigParameterTrackEditor::OnAddTransformKeysForSelectedObjects, EMovieSceneTransformChannel::All),
-		FCanExecuteAction::CreateSP(this, &FControlRigParameterTrackEditor::CanAddTransformKeysForSelectedObjects));
-
-	CommandBindings->MapAction(
-		Commands.AddTranslationKey,
-		FExecuteAction::CreateSP(this, &FControlRigParameterTrackEditor::OnAddTransformKeysForSelectedObjects, EMovieSceneTransformChannel::Translation),
-		FCanExecuteAction::CreateSP(this, &FControlRigParameterTrackEditor::CanAddTransformKeysForSelectedObjects));
+	return CanAddTransformKeysForSelectedObjects();
 
-	CommandBindings->MapAction(
-		Commands.AddRotationKey,
-		FExecuteAction::CreateSP(this, &FControlRigParameterTrackEditor::OnAddTransformKeysForSelectedObjects, EMovieSceneTransformChannel::Rotation),
-		FCanExecuteAction::CreateSP(this, &FControlRigParameterTrackEditor::CanAddTransformKeysForSelectedObjects));
-
-	CommandBindings->MapAction(
-		Commands.AddScaleKey,
-		FExecuteAction::CreateSP(this, &FControlRigParameterTrackEditor::OnAddTransformKeysForSelectedObjects, EMovieSceneTransformChannel::Scale),
-		FCanExecuteAction::CreateSP(this, &FControlRigParameterTrackEditor::CanAddTransformKeysForSelectedObjects));
-
-	Commands.BindingCount++;
-
-	// Add these bindings to Sequencer
-	SequencerCommandBindings->Append(CommandBindings.ToSharedRef());
-
-	// Also add them to the Curve Editor 
-	GetSequencer()->GetCommandBindings(ESequencerCommandBindings::CurveEditor)->Append(CommandBindings.ToSharedRef());
 }
-
-bool FControlRigParameterTrackEditor::CanAddTransformKeysForSelectedObjects()
+bool FControlRigParameterTrackEditor::CanAddTransformKeysForSelectedObjects() const
 {
 	if (!GetSequencer()->IsAllowedToChange())
 	{
@@ -561,8 +503,7 @@ void FControlRigParameterTrackEditor::OnAddTransformKeysForSelectedObjects(EMovi
 				USceneComponent* Component = Cast<USkeletalMeshComponent>(ObjectBinding->GetBoundObject());
 				if (Component)
 				{
-					UObject *ActorObject = Component->GetOwner();
-					AddControlKeys(ActorObject, ControlRig, Name, ControlName, Channel, ESequencerKeyMode::ManualKey);
+					AddControlKeys(Component, ControlRig, Name, ControlName, Channel, ESequencerKeyMode::ManualKeyForced);
 				}
 			}
 		}
@@ -593,30 +534,93 @@ void FControlRigParameterTrackEditor::OnSequencerDataChanged(EMovieSceneDataChan
 		{
 			GLevelEditorModeTools().DeactivateMode(FControlRigEditMode::ModeName);
 		}
-		ControlRigEditMode->SetObjects(nullptr, FGuid(), nullptr);
+		ControlRigEditMode->SetObjects(nullptr, FGuid(), nullptr, GetSequencer());
 	}
 }
 
-void FControlRigParameterTrackEditor::OnSelectionChanged(TArray<UMovieSceneTrack*> InTracks)
+void FControlRigParameterTrackEditor::OnCurveDisplayChanged(FCurveModel* CurveModel, bool bDisplayed)
 {
-	if (bIsDoingASelection)
+	if (bIsDoingSelection)
 	{
 		return;
 	}
+	TGuardValue<bool> Guard(bIsDoingSelection, true);
+	FScopedTransaction ScopedTransaction(LOCTEXT("SelectControlTransaction", "Select Control"), !GIsTransacting);
 
-	TGuardValue<bool> Guard(bIsDoingASelection, true);
 	TArray<FString> StringArray;
 	FControlRigEditMode* ControlRigEditMode = static_cast<FControlRigEditMode*>(GLevelEditorModeTools().GetActiveMode(FControlRigEditMode::ModeName));
 	UControlRig* ControlRig = nullptr;
-	//Always clear the control rig(s) in the edit mode.
-	if (ControlRigEditMode)
+	TArray<IKeyArea> KeyAreas;
+	if (CurveModel)
 	{
-		ControlRig = ControlRigEditMode->GetControlRig();
-		if (ControlRig)
+		UMovieSceneControlRigParameterSection* MovieSection = Cast<UMovieSceneControlRigParameterSection>(CurveModel->GetOwningObject());
+		if (MovieSection)
 		{
-			ControlRig->ClearControlSelection();
+			ControlRig = MovieSection->ControlRig;
+			//Only create the edit mode if we have a  curve selected and it's not set and we have some boundobjects.
+			if (!ControlRigEditMode)
+			{
+				GLevelEditorModeTools().ActivateMode(FControlRigEditMode::ModeName);
+				ControlRigEditMode = static_cast<FControlRigEditMode*>(GLevelEditorModeTools().GetActiveMode(FControlRigEditMode::ModeName));
+				if (TSharedPtr<IControlRigObjectBinding> ObjectBinding = ControlRig->GetObjectBinding())
+				{
+					if (ControlRigEditMode)
+					{
+						ControlRigEditMode->SetObjects(ControlRig, FGuid(), nullptr, GetSequencer());
+					}
+				}
+			}
+			else
+			{
+				if (ControlRigEditMode->GetControlRig() != ControlRig)
+				{
+					ControlRigEditMode->SetObjects(ControlRig, FGuid(), nullptr, GetSequencer());
+				}
+			}
+			//Not 100% safe but for now it is since that's all we show in the curve editor
+			//We need the Float Curve Model so we can get the ChannelHandle so we can also select the keyarea in the sequencer window if needed.
+			FFloatChannelCurveModel* FCurveModel = static_cast<FFloatChannelCurveModel*>(CurveModel);
+			FString String = CurveModel->GetLongDisplayName().ToString();
+			StringArray.SetNum(0);
+			String.ParseIntoArray(StringArray, TEXT("."));
+			if (StringArray.Num()  > 2)
+			{
+				//Not great but it should always be the third name
+				FName ControlName(*StringArray[2]);
+				ControlRig->SelectControl(ControlName, bDisplayed);
+				const TMovieSceneChannelHandle<FMovieSceneFloatChannel>& ChannelHandle = FCurveModel->GetChannelHandle();
+				IKeyArea KeyArea(MovieSection, ChannelHandle);
+				KeyAreas.Add(KeyArea);
+			}
+			else
+			{
+				UE_LOG(LogControlRigEditor, Display, TEXT("Could not find Rig Control From FCurveModel::LongName"));
+			}	
+		}
+		if (KeyAreas.Num() > 0)
+		{
+			bool bSync = GetSequencer()->GetSequencerSettings()->ShouldSyncCurveEditorSelection();
+			GetSequencer()->SuspendSelectionBroadcast();
+			GetSequencer()->GetSequencerSettings()->SyncCurveEditorSelection(false);
+			GetSequencer()->SelectByKeyAreas(MovieSection, KeyAreas, true, bDisplayed);
+			GetSequencer()->GetSequencerSettings()->SyncCurveEditorSelection(bSync);
+			GetSequencer()->ResumeSelectionBroadcast();
 		}
 	}
+}
+
+void FControlRigParameterTrackEditor::OnSelectionChanged(TArray<UMovieSceneTrack*> InTracks)
+{
+	if (bIsDoingSelection)
+	{
+		return;
+	}
+	TGuardValue<bool> Guard(bIsDoingSelection, true);
+
+	TArray<FString> StringArray;
+	FControlRigEditMode* ControlRigEditMode = static_cast<FControlRigEditMode*>(GLevelEditorModeTools().GetActiveMode(FControlRigEditMode::ModeName));
+	UControlRig* ControlRig = nullptr;
+	
 	TArray<const IKeyArea*> KeyAreas;
 	GetSequencer()->GetSelectedKeyAreas(KeyAreas);
 	if (KeyAreas.Num() <= 0)
@@ -624,6 +628,8 @@ void FControlRigParameterTrackEditor::OnSelectionChanged(TArray<UMovieSceneTrack
 		return;
 	}
 
+	FScopedTransaction ScopedTransaction(LOCTEXT("SelectControlTransaction", "Select Control"), !GIsTransacting);
+
 	TMap<UControlRig *, TSet<FName>> RigsAndControls;
 	for (const IKeyArea* KeyArea : KeyAreas)
 	{
@@ -640,7 +646,7 @@ void FControlRigParameterTrackEditor::OnSelectionChanged(TArray<UMovieSceneTrack
 				{
 					if (ControlRigEditMode)
 					{
-						ControlRigEditMode->SetObjects(ControlRig, FGuid(), nullptr);
+						ControlRigEditMode->SetObjects(ControlRig, FGuid(), nullptr, GetSequencer());
 					}
 				}
 			}
@@ -648,7 +654,7 @@ void FControlRigParameterTrackEditor::OnSelectionChanged(TArray<UMovieSceneTrack
 			{
 				if (ControlRigEditMode->GetControlRig() != ControlRig)
 				{
-					ControlRigEditMode->SetObjects(ControlRig, FGuid(), nullptr);
+					ControlRigEditMode->SetObjects(ControlRig, FGuid(), nullptr, GetSequencer());
 				}
 			}
 			
@@ -666,129 +672,214 @@ void FControlRigParameterTrackEditor::OnSelectionChanged(TArray<UMovieSceneTrack
 			}
 		}
 	}
+
+	ControlRig = nullptr;
+	//Always clear the control rig(s) in the edit mode.
+	if (ControlRigEditMode)
+	{
+		ControlRig = ControlRigEditMode->GetControlRig();
+		if (ControlRig)
+		{
+			ControlRig->ClearControlSelection();
+		}
+	}
 	for (TPair<UControlRig *, TSet<FName>> Pair : RigsAndControls)
 	{
-		Pair.Key->ClearControlSelection();
+		if (Pair.Key != ControlRig)
+		{
+			Pair.Key->ClearControlSelection();
+		}
 		for (const FName Name : Pair.Value)
 		{
 			Pair.Key->SelectControl(Name, true);
 		}
 	}
 }
+FMovieSceneTrackEditor::FFindOrCreateHandleResult FControlRigParameterTrackEditor::FindOrCreateHandleToSceneCompOrOwner(USceneComponent* InComp)
+{
+
+	bool bCreateHandleIfMissing = false;
+	FName CreatedFolderName = NAME_None;
+
+	FFindOrCreateHandleResult Result;
+	bool bHandleWasValid = GetSequencer()->GetHandleToObject(InComp, false).IsValid();
+
+
+	Result.Handle = GetSequencer()->GetHandleToObject(InComp, bCreateHandleIfMissing, CreatedFolderName);
+	Result.bWasCreated = bHandleWasValid == false && Result.Handle.IsValid();
+
+	if (Result.Handle.IsValid() == false)
+	{
+		UObject* OwnerObject = InComp->GetOwner();
+		bHandleWasValid = GetSequencer()->GetHandleToObject(OwnerObject, false).IsValid();
+
+		Result.Handle = GetSequencer()->GetHandleToObject(OwnerObject, bCreateHandleIfMissing, CreatedFolderName);
+		Result.bWasCreated = bHandleWasValid == false && Result.Handle.IsValid();
+
+	}
+	return Result;
+}
 
 
 void FControlRigParameterTrackEditor::HandleControlSelected(IControlRigManipulatable* ControlRigManp, const FRigControl& Control, bool bSelected)
 {
-	if (!bIsDoingASelection)
+
+	if (bIsDoingSelection)
 	{
-		TGuardValue<bool> Guard(bIsDoingASelection, true);
-		FString OurName = ControlRigManp->GetName();
-		FName ControlRigName(*OurName);
-		if (TSharedPtr<IControlRigObjectBinding> ObjectBinding = ControlRigManp->GetObjectBinding())
+		return;
+	}
+	TGuardValue<bool> Guard(bIsDoingSelection, true);
+
+	FString OurName = ControlRigManp->GetName();
+	FName ControlRigName(*OurName);
+	if (TSharedPtr<IControlRigObjectBinding> ObjectBinding = ControlRigManp->GetObjectBinding())
+	{
+		UObject *ActorObject = nullptr;
+		USceneComponent* Component = Cast<USkeletalMeshComponent>(ObjectBinding->GetBoundObject());
+		if (!Component)
 		{
-			UObject *ActorObject = nullptr;
-			USceneComponent* Component = Cast<USkeletalMeshComponent>(ObjectBinding->GetBoundObject());
-			if (!Component)
-			{
-				return;
-			}
-			ActorObject = Component->GetOwner();
-			bool bCreateTrack = false;
-			bool bCreateHandle = false;
-			FFindOrCreateHandleResult HandleResult = FindOrCreateHandleToObject(ActorObject, bCreateHandle, NAME_None);
-			FGuid ObjectHandle = HandleResult.Handle;
-			if (!ObjectHandle.IsValid())
+			return;
+		}
+		ActorObject = Component->GetOwner();
+		if (bSelected)
+		{
+			//make sure the actor is selected, otherwise we really can't select the proxies.
+			if (AActor* SkelMeshActor = Cast<AActor>(ActorObject))
 			{
-				return;
+				if (SkelMeshActor->IsSelected() == false)
+				{
+					GEditor->SelectActor(SkelMeshActor, true, true, false);
+				}
 			}
-			FFindOrCreateTrackResult TrackResult = FindOrCreateTrackForObject(ObjectHandle, UMovieSceneControlRigParameterTrack::StaticClass(), ControlRigName, bCreateTrack);
-			UMovieSceneControlRigParameterTrack* Track = CastChecked<UMovieSceneControlRigParameterTrack>(TrackResult.Track, ECastCheckedType::NullAllowed);
-			if (Track)
-			{
-				float Weight = 1.0f;
-				TArray<IKeyArea> KeyAreas;
-				TArray<FString> StringArray;
+		}
+		bool bCreateTrack = false;
+		FFindOrCreateHandleResult HandleResult = FindOrCreateHandleToSceneCompOrOwner(Component);
+		FGuid ObjectHandle = HandleResult.Handle;
+		if (!ObjectHandle.IsValid())
+		{
+			return;
+		}
+
+		FFindOrCreateTrackResult TrackResult = FindOrCreateTrackForObject(ObjectHandle, UMovieSceneControlRigParameterTrack::StaticClass(), ControlRigName, bCreateTrack);
+		UMovieSceneControlRigParameterTrack* Track = CastChecked<UMovieSceneControlRigParameterTrack>(TrackResult.Track, ECastCheckedType::NullAllowed);
+		if (Track)
+		{
+			float Weight = 1.0f;
+			TArray<IKeyArea> KeyAreas;
+			TArray<FString> StringArray;
+			GetSequencer()->SuspendSelectionBroadcast();
 
-				for (UMovieSceneSection* Section : Track->GetAllSections())
+			for (UMovieSceneSection* Section : Track->GetAllSections())
+			{
+				UMovieSceneControlRigParameterSection* ParamSection = Cast<UMovieSceneControlRigParameterSection>(Section);
+				if (ParamSection)
 				{
-					UMovieSceneControlRigParameterSection* ParamSection = Cast<UMovieSceneControlRigParameterSection>(Section);
-					if (ParamSection)
+					FChannelMapInfo* pChannelIndex = ParamSection->ControlChannelMap.Find(Control.Name);
+					if (pChannelIndex != nullptr)
 					{
-						FMovieSceneChannelProxy& ChannelProxy = Section->GetChannelProxy();
-						for (const FMovieSceneChannelEntry& Entry : Section->GetChannelProxy().GetAllEntries())
-						{
-							const FName ChannelTypeName = Entry.GetChannelTypeName();
-
-							// One editor data ptr per channel
-							TArrayView<FMovieSceneChannel* const>        Channels = Entry.GetChannels();
-							TArrayView<const FMovieSceneChannelMetaData> AllMetaData = Entry.GetMetaData();
-
-							for (int32 Index = 0; Index < Channels.Num(); ++Index)
-							{
-								FMovieSceneChannelHandle ChannelHandle = ChannelProxy.MakeHandle(ChannelTypeName, Index);
-								const FMovieSceneChannelMetaData& MetaData = AllMetaData[Index];
-								FString String = MetaData.Name.ToString();
-								String.ParseIntoArray(StringArray, TEXT("."));
-								if (StringArray.Num() > 0)
-								{
-									FName ControlName(*StringArray[0]);
-									if (ControlName == Control.Name)
-									{
-										IKeyArea KeyArea(Section, ChannelHandle);
-										KeyAreas.Add(KeyArea);
-									}
-								}
-							}
-						}
+						GetSequencer()->SelectByNthCategoryNode(Section, pChannelIndex->ControlIndex, bSelected);
 					}
 				}
-				if (KeyAreas.Num())
-				{
-					GetSequencer()->SelectByKeyAreas(KeyAreas, true, bSelected);
-				}
 			}
+			GetSequencer()->ResumeSelectionBroadcast();
+
+
+			//Force refresh now, not later
+			GetSequencer()->RefreshTree();
 		}
 	}
 }
 
-void FControlRigParameterTrackEditor::HandleControlModified(IControlRigManipulatable* ControlRigManp, const FRigControl& Control)
+void FControlRigParameterTrackEditor::HandleControlModified(IControlRigManipulatable* ControlRigManp, const FRigControl& Control, EControlRigSetKey InSetKey)
 {
 	if (GetSequencer().IsValid() && !GetSequencer()->IsAllowedToChange())
 	{
 		return;
 	}
-	FString OurName = ControlRigManp->GetName();
-	FName Name(*OurName);
-	if (TSharedPtr<IControlRigObjectBinding> ObjectBinding = ControlRigManp->GetObjectBinding())
+	UMovieScene* MovieScene = GetSequencer()->GetFocusedMovieSceneSequence()->GetMovieScene();
+	FControlRigEditMode* ControlRigEditMode = static_cast<FControlRigEditMode*>(GLevelEditorModeTools().GetActiveMode(FControlRigEditMode::ModeName));
+
+	if (ControlRigEditMode && ControlRigEditMode->GetControlRig() != nullptr && MovieScene)
 	{
-		UObject *ActorObject = nullptr;
-		USceneComponent* Component = Cast<USkeletalMeshComponent>(ObjectBinding->GetBoundObject());
-		if (Component)
+		bool bTrackIsValid = false;
+		const TArray<FMovieSceneBinding>& Bindings = MovieScene->GetBindings();
+		for (const FMovieSceneBinding& Binding : Bindings)
 		{
-			ActorObject = Component->GetOwner();
-			AddControlKeys(ActorObject, ControlRigManp, Name, Control.Name, EMovieSceneTransformChannel::All, ESequencerKeyMode::AutoKey);
+			UMovieSceneControlRigParameterTrack* Track = Cast<UMovieSceneControlRigParameterTrack>(MovieScene->FindTrack(UMovieSceneControlRigParameterTrack::StaticClass(), Binding.GetObjectGuid(), NAME_None));
+			if (Track && Track->GetControlRig() == ControlRigManp)
+			{
+				bTrackIsValid = true;
+			}
+		}
+		if (bTrackIsValid)
+		{
+			FString OurName = ControlRigManp->GetName();
+			FName Name(*OurName);
+			if (TSharedPtr<IControlRigObjectBinding> ObjectBinding = ControlRigManp->GetObjectBinding())
+			{
+				USceneComponent* Component = Cast<USkeletalMeshComponent>(ObjectBinding->GetBoundObject());
+				if (Component)
+				{
+					ESequencerKeyMode KeyMode = ESequencerKeyMode::AutoKey;
+					if (InSetKey == EControlRigSetKey::Always)
+					{
+						KeyMode = ESequencerKeyMode::ManualKeyForced;
+					}
+					else if (InSetKey == EControlRigSetKey::Never)
+					{
+						KeyMode = ESequencerKeyMode::ManualKey; //best we have here
+					}
+					AddControlKeys(Component, ControlRigManp, Name, Control.Name, EMovieSceneTransformChannel::All, KeyMode);
+				}
+			}
+		}
+		else
+		{
+			//okay no good track so deactive it and delete it's Control Rig and bingings.
+			if (GLevelEditorModeTools().HasToolkitHost())
+			{
+				GLevelEditorModeTools().DeactivateMode(FControlRigEditMode::ModeName);
+			}
 		}
 	}
 }
 
 void FControlRigParameterTrackEditor::GetControlRigKeys(IControlRigManipulatable* Manip, FName ParameterName, EMovieSceneTransformChannel ChannelsToKey, FGeneratedTrackKeys& OutGeneratedKeys)
 {
-	const TArray<FRigControl>& Controls = Manip->AvailableControls();
+	TArray<FRigControl> Controls;
+	FControlRigSortedControls::GetControlsInOrder(Manip, Controls);
 	// If key all is enabled, for a key on all the channels
 	if (GetSequencer()->GetKeyGroupMode() == EKeyGroupMode::KeyAll)
 	{
 		ChannelsToKey = EMovieSceneTransformChannel::All;
 	}
 
+	//Need seperate index fo bools and floats since there are seperate entries for each later when they are accessed by the set key stuff.
 	int32 ChannelIndex = 0;
+	int32 BoolChannelIndex = 0;
 	for (const FRigControl& RigControl : Controls)
 	{
+		bool bSetKey = RigControl.Name == ParameterName;
+
 		switch (RigControl.ControlType)
 		{
+		case ERigControlType::Bool:
+		{
+			bool Val = RigControl.Value.Get<bool>();
+			OutGeneratedKeys.Add(FMovieSceneChannelValueSetter::Create<FMovieSceneBoolChannel>(BoolChannelIndex++, Val, bSetKey));
+			break;
+		}
 		case ERigControlType::Float:
 		{
 			float Val = RigControl.Value.Get<float>();
-			OutGeneratedKeys.Add(FMovieSceneChannelValueSetter::Create<FMovieSceneFloatChannel>(ChannelIndex++, Val, true));
+			OutGeneratedKeys.Add(FMovieSceneChannelValueSetter::Create<FMovieSceneFloatChannel>(ChannelIndex++, Val, bSetKey));
+			break;
+		}
+		case ERigControlType::Vector2D:
+		{
+			FVector2D Val = RigControl.Value.Get<FVector2D>();
+			OutGeneratedKeys.Add(FMovieSceneChannelValueSetter::Create<FMovieSceneFloatChannel>(ChannelIndex++, Val.X, bSetKey));
+			OutGeneratedKeys.Add(FMovieSceneChannelValueSetter::Create<FMovieSceneFloatChannel>(ChannelIndex++, Val.Y, bSetKey));
 			break;
 		}
 		case ERigControlType::Position:
@@ -796,19 +887,29 @@ void FControlRigParameterTrackEditor::GetControlRigKeys(IControlRigManipulatable
 		case ERigControlType::Rotator:
 		{
 			FVector Val = RigControl.Value.Get<FVector>();
-			OutGeneratedKeys.Add(FMovieSceneChannelValueSetter::Create<FMovieSceneFloatChannel>(ChannelIndex++, Val.X, true));
-			OutGeneratedKeys.Add(FMovieSceneChannelValueSetter::Create<FMovieSceneFloatChannel>(ChannelIndex++, Val.Y, true));
-			OutGeneratedKeys.Add(FMovieSceneChannelValueSetter::Create<FMovieSceneFloatChannel>(ChannelIndex++, Val.Z, true));
+			OutGeneratedKeys.Add(FMovieSceneChannelValueSetter::Create<FMovieSceneFloatChannel>(ChannelIndex++, Val.X, bSetKey));
+			OutGeneratedKeys.Add(FMovieSceneChannelValueSetter::Create<FMovieSceneFloatChannel>(ChannelIndex++, Val.Y, bSetKey));
+			OutGeneratedKeys.Add(FMovieSceneChannelValueSetter::Create<FMovieSceneFloatChannel>(ChannelIndex++, Val.Z, bSetKey));
 			break;
 		}
 
 		case ERigControlType::Transform:
+		case ERigControlType::TransformNoScale:
 		{
-			FTransform Val = RigControl.Value.Get<FTransform>();
+			FTransform Val = FTransform::Identity;
+			if (RigControl.ControlType == ERigControlType::TransformNoScale)
+			{
+				FTransformNoScale NoScale = RigControl.Value.Get<FTransformNoScale>();
+				Val = NoScale;
+			}
+			else
+			{
+				Val = RigControl.Value.Get<FTransform>();
+			}
 			FVector CurrentVector = Val.GetTranslation();
-			bool bKeyX = EnumHasAnyFlags(ChannelsToKey, EMovieSceneTransformChannel::TranslationX);
-			bool bKeyY = EnumHasAnyFlags(ChannelsToKey, EMovieSceneTransformChannel::TranslationY);
-			bool bKeyZ = EnumHasAnyFlags(ChannelsToKey, EMovieSceneTransformChannel::TranslationZ);
+			bool bKeyX = bSetKey && EnumHasAnyFlags(ChannelsToKey, EMovieSceneTransformChannel::TranslationX);
+			bool bKeyY = bSetKey && EnumHasAnyFlags(ChannelsToKey, EMovieSceneTransformChannel::TranslationY);
+			bool bKeyZ = bSetKey && EnumHasAnyFlags(ChannelsToKey, EMovieSceneTransformChannel::TranslationZ);
 			if (GetSequencer()->GetKeyGroupMode() == EKeyGroupMode::KeyGroup && (bKeyX || bKeyY || bKeyZ))
 			{
 				bKeyX = bKeyY = bKeyZ = true;
@@ -818,9 +919,9 @@ void FControlRigParameterTrackEditor::GetControlRigKeys(IControlRigManipulatable
 			OutGeneratedKeys.Add(FMovieSceneChannelValueSetter::Create<FMovieSceneFloatChannel>(ChannelIndex++, CurrentVector.Z, bKeyZ));
 
 			CurrentVector = Val.GetRotation().Euler();
-			bKeyX = EnumHasAnyFlags(ChannelsToKey, EMovieSceneTransformChannel::RotationX);
-			bKeyY = EnumHasAnyFlags(ChannelsToKey, EMovieSceneTransformChannel::RotationY);
-			bKeyZ = EnumHasAnyFlags(ChannelsToKey, EMovieSceneTransformChannel::RotationZ);
+			bKeyX = bSetKey && EnumHasAnyFlags(ChannelsToKey, EMovieSceneTransformChannel::RotationX);
+			bKeyY = bSetKey && EnumHasAnyFlags(ChannelsToKey, EMovieSceneTransformChannel::RotationY);
+			bKeyZ = bSetKey && EnumHasAnyFlags(ChannelsToKey, EMovieSceneTransformChannel::RotationZ);
 			if (GetSequencer()->GetKeyGroupMode() == EKeyGroupMode::KeyGroup && (bKeyX || bKeyY || bKeyZ))
 			{
 				bKeyX = bKeyY = bKeyZ = true;
@@ -829,26 +930,27 @@ void FControlRigParameterTrackEditor::GetControlRigKeys(IControlRigManipulatable
 			OutGeneratedKeys.Add(FMovieSceneChannelValueSetter::Create<FMovieSceneFloatChannel>(ChannelIndex++, CurrentVector.Y, bKeyY));
 			OutGeneratedKeys.Add(FMovieSceneChannelValueSetter::Create<FMovieSceneFloatChannel>(ChannelIndex++, CurrentVector.Z, bKeyZ));
 
-			CurrentVector = Val.GetScale3D();
-			bKeyX = EnumHasAnyFlags(ChannelsToKey, EMovieSceneTransformChannel::ScaleX);
-			bKeyY = EnumHasAnyFlags(ChannelsToKey, EMovieSceneTransformChannel::ScaleY);
-			bKeyZ = EnumHasAnyFlags(ChannelsToKey, EMovieSceneTransformChannel::ScaleZ);
-			if (GetSequencer()->GetKeyGroupMode() == EKeyGroupMode::KeyGroup && (bKeyX || bKeyY || bKeyZ))
+			if (RigControl.ControlType == ERigControlType::Transform)
 			{
-				bKeyX = bKeyY = bKeyZ = true;
+				CurrentVector = Val.GetScale3D();
+				bKeyX = bSetKey && EnumHasAnyFlags(ChannelsToKey, EMovieSceneTransformChannel::ScaleX);
+				bKeyY = bSetKey && EnumHasAnyFlags(ChannelsToKey, EMovieSceneTransformChannel::ScaleY);
+				bKeyZ = bSetKey && EnumHasAnyFlags(ChannelsToKey, EMovieSceneTransformChannel::ScaleZ);
+				if (GetSequencer()->GetKeyGroupMode() == EKeyGroupMode::KeyGroup && (bKeyX || bKeyY || bKeyZ))
+				{
+					bKeyX = bKeyY = bKeyZ = true;
+				}
+				OutGeneratedKeys.Add(FMovieSceneChannelValueSetter::Create<FMovieSceneFloatChannel>(ChannelIndex++, CurrentVector.X, bKeyX));
+				OutGeneratedKeys.Add(FMovieSceneChannelValueSetter::Create<FMovieSceneFloatChannel>(ChannelIndex++, CurrentVector.Y, bKeyY));
+				OutGeneratedKeys.Add(FMovieSceneChannelValueSetter::Create<FMovieSceneFloatChannel>(ChannelIndex++, CurrentVector.Z, bKeyZ));
 			}
-			OutGeneratedKeys.Add(FMovieSceneChannelValueSetter::Create<FMovieSceneFloatChannel>(ChannelIndex++, CurrentVector.X, bKeyX));
-			OutGeneratedKeys.Add(FMovieSceneChannelValueSetter::Create<FMovieSceneFloatChannel>(ChannelIndex++, CurrentVector.Y, bKeyY));
-			OutGeneratedKeys.Add(FMovieSceneChannelValueSetter::Create<FMovieSceneFloatChannel>(ChannelIndex++, CurrentVector.Z, bKeyZ));
-
 			break;
 		}
-		//mz todo the other types
 		}
 	}
 }
 
-FKeyPropertyResult FControlRigParameterTrackEditor::AddKeysToControlRigHandle(UObject *Object, IControlRigManipulatable* Manip,
+FKeyPropertyResult FControlRigParameterTrackEditor::AddKeysToControlRigHandle(USceneComponent *InSceneComp, IControlRigManipulatable* Manip,
 	FGuid ObjectHandle, FFrameNumber KeyTime, FGeneratedTrackKeys& GeneratedKeys,
 	ESequencerKeyMode KeyMode, TSubclassOf<UMovieSceneTrack> TrackClass, FName ControlRigName)
 {
@@ -897,7 +999,7 @@ FKeyPropertyResult FControlRigParameterTrackEditor::AddKeysToControlRigHandle(UO
 		{
 			if (!bTrackCreated)
 			{
-				ModifyOurGeneratedKeysByCurrentAndWeight(Object, Manip, Track, SectionToKey, KeyTime, GeneratedKeys, Weight);
+				ModifyOurGeneratedKeysByCurrentAndWeight(InSceneComp, Manip, Track, SectionToKey, KeyTime, GeneratedKeys, Weight);
 			}
 			UMovieSceneControlRigParameterSection* ParamSection = Cast<UMovieSceneControlRigParameterSection>(SectionToKey);
 			if (!ParamSection->GetDoNotKey())
@@ -912,7 +1014,7 @@ FKeyPropertyResult FControlRigParameterTrackEditor::AddKeysToControlRigHandle(UO
 }
 
 FKeyPropertyResult FControlRigParameterTrackEditor::AddKeysToControlRig(
-	UObject* Object, IControlRigManipulatable* Manip, FFrameNumber KeyTime, FGeneratedTrackKeys& GeneratedKeys,
+	USceneComponent *InSceneComp, IControlRigManipulatable* Manip, FFrameNumber KeyTime, FGeneratedTrackKeys& GeneratedKeys,
 	ESequencerKeyMode KeyMode, TSubclassOf<UMovieSceneTrack> TrackClass, FName PropertyName)
 {
 	FKeyPropertyResult KeyPropertyResult;
@@ -924,26 +1026,26 @@ FKeyPropertyResult FControlRigParameterTrackEditor::AddKeysToControlRig(
 		KeyMode == ESequencerKeyMode::ManualKeyForced ||
 		AllowEditsMode == EAllowEditsMode::AllowSequencerEditsOnly;
 
-	FFindOrCreateHandleResult HandleResult = FindOrCreateHandleToObject(Object, bCreateHandle, NAME_None);
+	FFindOrCreateHandleResult HandleResult = FindOrCreateHandleToSceneCompOrOwner(InSceneComp);
 	FGuid ObjectHandle = HandleResult.Handle;
 	KeyPropertyResult.bHandleCreated = HandleResult.bWasCreated;
 	if (ObjectHandle.IsValid())
 	{
-		KeyPropertyResult |= AddKeysToControlRigHandle(Object, Manip, ObjectHandle, KeyTime, GeneratedKeys, KeyMode, TrackClass, PropertyName);
+		KeyPropertyResult |= AddKeysToControlRigHandle(InSceneComp, Manip, ObjectHandle, KeyTime, GeneratedKeys, KeyMode, TrackClass, PropertyName);
 	}
 
 	return KeyPropertyResult;
 }
 
-void FControlRigParameterTrackEditor::AddControlKeys(UObject* InObject, IControlRigManipulatable* Manip, FName ControlRigName, FName RigControlName, EMovieSceneTransformChannel ChannelsToKey, ESequencerKeyMode KeyMode)
+void FControlRigParameterTrackEditor::AddControlKeys(USceneComponent *InSceneComp, IControlRigManipulatable* Manip, FName ControlRigName, FName RigControlName, EMovieSceneTransformChannel ChannelsToKey, ESequencerKeyMode KeyMode)
 {
-	if (GetSequencer().IsValid() && !GetSequencer()->IsAllowedToChange())
+	if (KeyMode == ESequencerKeyMode::ManualKey || (GetSequencer().IsValid() && !GetSequencer()->IsAllowedToChange()))
 	{
 		return;
 	}
 	bool bCreateTrack = false;
 	bool bCreateHandle = false;
-	FFindOrCreateHandleResult HandleResult = FindOrCreateHandleToObject(InObject, bCreateHandle, NAME_None);
+	FFindOrCreateHandleResult HandleResult = FindOrCreateHandleToSceneCompOrOwner(InSceneComp);
 	FGuid ObjectHandle = HandleResult.Handle;
 	if (!ObjectHandle.IsValid())
 	{
@@ -966,15 +1068,18 @@ void FControlRigParameterTrackEditor::AddControlKeys(UObject* InObject, IControl
 	TSharedRef<FGeneratedTrackKeys> GeneratedKeys = MakeShared<FGeneratedTrackKeys>();
 
 	GetControlRigKeys(Manip, RigControlName, ChannelsToKey, *GeneratedKeys);
+	TGuardValue<bool> Guard(bIsDoingSelection, true);
 
 	auto OnKeyProperty = [=](FFrameNumber Time) -> FKeyPropertyResult
 	{
-		return this->AddKeysToControlRig(InObject, Manip, Time, *GeneratedKeys, KeyMode, UMovieSceneControlRigParameterTrack::StaticClass(), ControlRigName);
+		return this->AddKeysToControlRig(InSceneComp, Manip, Time, *GeneratedKeys, KeyMode, UMovieSceneControlRigParameterTrack::StaticClass(), ControlRigName);
 	};
 
 	AnimatablePropertyChanged(FOnKeyProperty::CreateLambda(OnKeyProperty));
+
 }
 
+
 bool FControlRigParameterTrackEditor::ModifyOurGeneratedKeysByCurrentAndWeight(UObject *Object, IControlRigManipulatable* Manip, UMovieSceneTrack *Track, UMovieSceneSection* SectionToKey, FFrameNumber KeyTime, FGeneratedTrackKeys& GeneratedTotalKeys, float Weight) const
 {
 	FFrameRate TickResolution = GetSequencer()->GetFocusedTickResolution();
@@ -986,73 +1091,124 @@ bool FControlRigParameterTrackEditor::ModifyOurGeneratedKeysByCurrentAndWeight(U
 	FMovieSceneContext Context(FMovieSceneEvaluationRange(KeyTime, GetSequencer()->GetFocusedTickResolution()));
 	EvalTrack.Interrogate(Context, InterrogationData, Object);
 	const TArray<FRigControl>& Controls = Manip->AvailableControls();
-
+	UMovieSceneControlRigParameterSection* Section = Cast<UMovieSceneControlRigParameterSection>(SectionToKey);
 	FMovieSceneChannelProxy& Proxy = SectionToKey->GetChannelProxy();
 	int32 ChannelIndex = 0;
+	FChannelMapInfo* pChannelIndex = nullptr;
 	for (const FRigControl& RigControl : Controls)
 	{
 		switch (RigControl.ControlType)
 		{
-		case ERigControlType::Float:
-		{
-			for (const FFloatInterrogationData& Val : InterrogationData.Iterate<FFloatInterrogationData>(UMovieSceneControlRigParameterSection::GetFloatInterrogationKey()))
+			case ERigControlType::Float:
 			{
-				if (Val.ParameterName == RigControl.Name)
+				for (const FFloatInterrogationData& Val : InterrogationData.Iterate<FFloatInterrogationData>(UMovieSceneControlRigParameterSection::GetFloatInterrogationKey()))
 				{
-					GeneratedTotalKeys[ChannelIndex]->ModifyByCurrentAndWeight(Proxy, KeyTime, (void *)&Val.Val, Weight);
-					break;
+					if (Val.ParameterName == RigControl.Name)
+					{
+						pChannelIndex = Section->ControlChannelMap.Find(RigControl.Name);
+						if (pChannelIndex)
+						{
+							ChannelIndex = pChannelIndex->ChannelIndex;
+						}
+						
+						GeneratedTotalKeys[ChannelIndex]->ModifyByCurrentAndWeight(Proxy, KeyTime, (void *)&Val.Val, Weight);
+						break;
+					}
 				}
+				++ChannelIndex;
+				break;
 			}
-			++ChannelIndex;
-			break;
-		}
-		case ERigControlType::Position:
-		case ERigControlType::Scale:
-		case ERigControlType::Rotator:
-		{
-			for (const FVectorInterrogationData& Val : InterrogationData.Iterate<FVectorInterrogationData>(UMovieSceneControlRigParameterSection::GetVectorInterrogationKey()))
+			case ERigControlType::Bool:
 			{
-				if (Val.ParameterName == RigControl.Name)
+				//no blending of bools
+				++ChannelIndex;
+				break;
+			}
+			case ERigControlType::Vector2D:
+			{
+				for (const FVector2DInterrogationData& Val : InterrogationData.Iterate<FVector2DInterrogationData>(UMovieSceneControlRigParameterSection::GetVector2DInterrogationKey()))
 				{
-					GeneratedTotalKeys[ChannelIndex]->ModifyByCurrentAndWeight(Proxy, KeyTime, (void *)&Val.Val, Weight);
-					GeneratedTotalKeys[ChannelIndex + 1]->ModifyByCurrentAndWeight(Proxy, KeyTime, (void *)&Val.Val, Weight);
-					GeneratedTotalKeys[ChannelIndex + 2]->ModifyByCurrentAndWeight(Proxy, KeyTime, (void *)&Val.Val, Weight);
+					if (Val.ParameterName == RigControl.Name)
+					{
+						pChannelIndex = Section->ControlChannelMap.Find(RigControl.Name);
+						if (pChannelIndex)
+						{
+							ChannelIndex = pChannelIndex->ChannelIndex;
+						}
+						GeneratedTotalKeys[ChannelIndex]->ModifyByCurrentAndWeight(Proxy, KeyTime, (void *)&Val.Val.X, Weight);
+						GeneratedTotalKeys[ChannelIndex + 1]->ModifyByCurrentAndWeight(Proxy, KeyTime, (void *)&Val.Val.Y, Weight);
 
-					break;
+						break;
+					}
+					ChannelIndex += 2;
 				}
+				break;
 			}
-			ChannelIndex += 3;
+			case ERigControlType::Position:
+			case ERigControlType::Scale:
+			case ERigControlType::Rotator:
+			{
+				for (const FVectorInterrogationData& Val : InterrogationData.Iterate<FVectorInterrogationData>(UMovieSceneControlRigParameterSection::GetVectorInterrogationKey()))
+				{
+					if (Val.ParameterName == RigControl.Name)
+					{
+						pChannelIndex = Section->ControlChannelMap.Find(RigControl.Name);
+						if (pChannelIndex)
+						{
+							ChannelIndex = pChannelIndex->ChannelIndex;
+						}
+						GeneratedTotalKeys[ChannelIndex]->ModifyByCurrentAndWeight(Proxy, KeyTime, (void *)&Val.Val.X, Weight);
+						GeneratedTotalKeys[ChannelIndex + 1]->ModifyByCurrentAndWeight(Proxy, KeyTime, (void *)&Val.Val.Y, Weight);
+						GeneratedTotalKeys[ChannelIndex + 2]->ModifyByCurrentAndWeight(Proxy, KeyTime, (void *)&Val.Val.Z, Weight);
 
-			break;
-		}
+						break;
+					}
+					ChannelIndex += 3;
+				}
 
-		case ERigControlType::Transform:
-		{
-			for (const FTransformInterrogationData& Val : InterrogationData.Iterate<FTransformInterrogationData>(UMovieSceneControlRigParameterSection::GetTransformInterrogationKey()))
+				break;
+			}
+
+			case ERigControlType::Transform:
+			case ERigControlType::TransformNoScale:
 			{
-				if (Val.ParameterName == RigControl.Name)
+				for (const FTransformInterrogationData& Val : InterrogationData.Iterate<FTransformInterrogationData>(UMovieSceneControlRigParameterSection::GetTransformInterrogationKey()))
 				{
-					FVector CurrentPos = Val.Val.GetTranslation();
-					FRotator CurrentRot = Val.Val.GetRotation().Rotator();
-					FVector CurrentScale = Val.Val.GetScale3D();
-					GeneratedTotalKeys[ChannelIndex]->ModifyByCurrentAndWeight(Proxy, KeyTime, (void *)&CurrentPos.X, Weight);
-					GeneratedTotalKeys[ChannelIndex + 1]->ModifyByCurrentAndWeight(Proxy, KeyTime, (void *)&CurrentPos.Y, Weight);
-					GeneratedTotalKeys[ChannelIndex + 2]->ModifyByCurrentAndWeight(Proxy, KeyTime, (void *)&CurrentPos.Z, Weight);
-					GeneratedTotalKeys[ChannelIndex + 3]->ModifyByCurrentAndWeight(Proxy, KeyTime, (void *)&CurrentRot.Roll, Weight);
-					GeneratedTotalKeys[ChannelIndex + 4]->ModifyByCurrentAndWeight(Proxy, KeyTime, (void *)&CurrentRot.Pitch, Weight);
-					GeneratedTotalKeys[ChannelIndex + 5]->ModifyByCurrentAndWeight(Proxy, KeyTime, (void *)&CurrentRot.Yaw, Weight);
-					GeneratedTotalKeys[ChannelIndex + 6]->ModifyByCurrentAndWeight(Proxy, KeyTime, (void *)&CurrentScale.X, Weight);
-					GeneratedTotalKeys[ChannelIndex + 7]->ModifyByCurrentAndWeight(Proxy, KeyTime, (void *)&CurrentScale.Y, Weight);
-					GeneratedTotalKeys[ChannelIndex + 8]->ModifyByCurrentAndWeight(Proxy, KeyTime, (void *)&CurrentScale.Z, Weight);
-
-					break;
+					if (Val.ParameterName == RigControl.Name)
+					{
+						pChannelIndex = Section->ControlChannelMap.Find(RigControl.Name);
+						if (pChannelIndex)
+						{
+							ChannelIndex = pChannelIndex->ChannelIndex;
+						}
+						FVector CurrentPos = Val.Val.GetTranslation();
+						FRotator CurrentRot = Val.Val.GetRotation().Rotator();
+						GeneratedTotalKeys[ChannelIndex]->ModifyByCurrentAndWeight(Proxy, KeyTime, (void *)&CurrentPos.X, Weight);
+						GeneratedTotalKeys[ChannelIndex + 1]->ModifyByCurrentAndWeight(Proxy, KeyTime, (void *)&CurrentPos.Y, Weight);
+						GeneratedTotalKeys[ChannelIndex + 2]->ModifyByCurrentAndWeight(Proxy, KeyTime, (void *)&CurrentPos.Z, Weight);
+						GeneratedTotalKeys[ChannelIndex + 3]->ModifyByCurrentAndWeight(Proxy, KeyTime, (void *)&CurrentRot.Roll, Weight);
+						GeneratedTotalKeys[ChannelIndex + 4]->ModifyByCurrentAndWeight(Proxy, KeyTime, (void *)&CurrentRot.Pitch, Weight);
+						GeneratedTotalKeys[ChannelIndex + 5]->ModifyByCurrentAndWeight(Proxy, KeyTime, (void *)&CurrentRot.Yaw, Weight);
+						if (RigControl.ControlType == ERigControlType::Transform)
+						{
+							FVector CurrentScale = Val.Val.GetScale3D();
+							GeneratedTotalKeys[ChannelIndex + 6]->ModifyByCurrentAndWeight(Proxy, KeyTime, (void *)&CurrentScale.X, Weight);
+							GeneratedTotalKeys[ChannelIndex + 7]->ModifyByCurrentAndWeight(Proxy, KeyTime, (void *)&CurrentScale.Y, Weight);
+							GeneratedTotalKeys[ChannelIndex + 8]->ModifyByCurrentAndWeight(Proxy, KeyTime, (void *)&CurrentScale.Z, Weight);
+						}
+						break;
+					}
+				}
+				if (RigControl.ControlType == ERigControlType::Transform)
+				{
+					ChannelIndex += 9;
+				}
+				else
+				{
+					ChannelIndex += 6;
 				}
+				break;
 			}
-			ChannelIndex += 9;
-
-			break;
-		}
-		//mz todo the other types
 		}
 	}
 	return true;
diff --git a/Engine/Plugins/Experimental/ControlRig/Source/ControlRigEditor/Private/Sequencer/ControlRigParameterTrackEditor.h b/Engine/Plugins/Experimental/ControlRig/Source/ControlRigEditor/Private/Sequencer/ControlRigParameterTrackEditor.h
--- a/Engine/Plugins/Experimental/ControlRig/Source/ControlRigEditor/Private/Sequencer/ControlRigParameterTrackEditor.h
+++ b/Engine/Plugins/Experimental/ControlRig/Source/ControlRigEditor/Private/Sequencer/ControlRigParameterTrackEditor.h
@@ -46,10 +46,14 @@ public:
 
 	// ISequencerTrackEditor interface
 	virtual void OnRelease() override;
-	virtual void BindCommands(TSharedRef<FUICommandList> SequencerCommandBindings) override;
 	virtual void BuildObjectBindingTrackMenu(FMenuBuilder& MenuBuilder, const TArray<FGuid>& ObjectBindings, const UClass* ObjectClass) override;
 	virtual TSharedRef<ISequencerSection> MakeSectionInterface(UMovieSceneSection& SectionObject, UMovieSceneTrack& Track, FGuid ObjectBinding) override;
 	virtual bool SupportsType(TSubclassOf<UMovieSceneTrack> Type) const override;
+	virtual bool HasTransformKeyBindings() const override { return true; }
+	virtual bool CanAddTransformKeysForSelectedObjects() const override;
+	virtual void OnAddTransformKeysForSelectedObjects(EMovieSceneTransformChannel Channel);
+	virtual bool HasTransformKeyOverridePriority() const override;
+	virtual void ObjectImplicitlyAdded(UObject* InObject)  override;
 
 private:
 
@@ -64,23 +68,28 @@ private:
 	/** Delegate for MovieScene Changing so we can see if our track got deleted*/
 	void OnSequencerDataChanged(EMovieSceneDataChangeType DataChangeType);
 
+	/** Delegate for Curve Selection Changed Event */
+	void OnCurveDisplayChanged(FCurveModel* InCurveModel, bool bDisplayed);
+
 	/** Control Rig Delegates*/
-	void HandleControlModified(IControlRigManipulatable* Subject, const FRigControl& Control);
+	void HandleControlModified(IControlRigManipulatable* Subject, const FRigControl& Control, EControlRigSetKey InSetKey);
 	void HandleControlSelected(IControlRigManipulatable* Subject, const FRigControl& Control, bool bSelected);
 
+	/** Handle Creattion for SkelMeshComp or Actor Owner, either may have a binding*/
+	FMovieSceneTrackEditor::FFindOrCreateHandleResult FindOrCreateHandleToSceneCompOrOwner(USceneComponent* InComp);
+
 private:
-	//MZ TODO Store map of control rig to parameter values to see if they changed?
 
 	/** Command Bindings added by the Transform Track Editor to Sequencer and curve editor. */
 	TSharedPtr<FUICommandList> CommandBindings;
 public:
 
-	void AddControlKeys(UObject* InObject, IControlRigManipulatable* Manip, FName PropertyName, FName ParameterName, EMovieSceneTransformChannel ChannelsToKey, ESequencerKeyMode KeyMode);
+	void AddControlKeys(USceneComponent *InSceneComp, IControlRigManipulatable* Manip, FName PropertyName, FName ParameterName, EMovieSceneTransformChannel ChannelsToKey, ESequencerKeyMode KeyMode);
 	void GetControlRigKeys(IControlRigManipulatable* Manip, FName ParameterName, EMovieSceneTransformChannel ChannelsToKey, FGeneratedTrackKeys& OutGeneratedKeys);
 	FKeyPropertyResult AddKeysToControlRig(
-		UObject* Object, IControlRigManipulatable* Manip, FFrameNumber KeyTime, FGeneratedTrackKeys& GeneratedKeys,
+		USceneComponent *InSceneComp, IControlRigManipulatable* Manip, FFrameNumber KeyTime, FGeneratedTrackKeys& GeneratedKeys,
 		ESequencerKeyMode KeyMode, TSubclassOf<UMovieSceneTrack> TrackClass, FName PropertyName);
-	FKeyPropertyResult AddKeysToControlRigHandle(UObject *Object, IControlRigManipulatable* Manip,
+	FKeyPropertyResult AddKeysToControlRigHandle(USceneComponent *InSceneComp, IControlRigManipulatable* Manip,
 		FGuid ObjectHandle, FFrameNumber KeyTime, FGeneratedTrackKeys& GeneratedKeys,
 		ESequencerKeyMode KeyMode, TSubclassOf<UMovieSceneTrack> TrackClass, FName PropertyName);
 	/**
@@ -103,15 +112,13 @@ public:
 	void OnControlRigAssetEnterPressed(const TArray<FAssetData>& AssetData, TArray<FGuid> ObjectBindings,  USkeleton* Skeleton);
 	*/
 private:
-	bool CanAddTransformKeysForSelectedObjects();
-	void OnAddTransformKeysForSelectedObjects(EMovieSceneTransformChannel Channel);
 	FDelegateHandle SelectionChangedHandle;
 	FDelegateHandle SequencerChangedHandle;
 
 private:
 
 	/** Guard to stop infinite loops when handling control selections*/
-	bool bIsDoingASelection;
+	bool bIsDoingSelection;
 
 };
 
diff --git a/Engine/Plugins/Experimental/ControlRig/Source/ControlRigEditor/Private/Sequencer/ControlRigSequencerFilter.cpp b/Engine/Plugins/Experimental/ControlRig/Source/ControlRigEditor/Private/Sequencer/ControlRigSequencerFilter.cpp
new file mode 100644
--- /dev/null
+++ b/Engine/Plugins/Experimental/ControlRig/Source/ControlRigEditor/Private/Sequencer/ControlRigSequencerFilter.cpp
@@ -0,0 +1,93 @@
+// Copyright Epic Games, Inc. All Rights Reserved.
+
+#include "ControlRigSequencerFilter.h"
+#include "ControlRig.h"
+#include "ControlRigSkeletalMeshComponent.h"
+#include "Sequencer/MovieSceneControlRigParameterTrack.h"
+#include "Styling/SlateIconFinder.h"
+
+#define LOCTEXT_NAMESPACE "ControlRigSequencerTrackFilters"
+
+//////////////////////////////////////////////////////////////////////////
+//
+
+class FSequencerTrackFilter_ControlRigControls : public  FSequencerTrackFilter
+{
+	virtual FString GetName() const override { return TEXT("ControlRigControlsFilter"); }
+	virtual FText GetDisplayName() const override { return LOCTEXT("SequenceTrackFilter_ControlRigControls", "Control Rig Controls"); }
+	virtual FText GetToolTipText() const override { return LOCTEXT("SequencerTrackFilter_ControlRigControlsToolTip", "Show Only Control Rig Controls."); }
+	virtual FSlateIcon GetIcon() const { return FSlateIconFinder::FindIconForClass(UControlRig::StaticClass()); }
+
+
+	virtual bool PassesFilterWithDisplayName(FTrackFilterType InItem, const FText& InText) const
+	{
+
+		const UMovieSceneControlRigParameterTrack* Track = Cast<UMovieSceneControlRigParameterTrack>(InItem);
+		return (Track != nullptr);
+	}
+
+	virtual bool PassesFilter(FTrackFilterType InItem) const override
+	{
+		return true;
+	}
+};
+
+class FSequencerTrackFilter_ControlRigSelectedControls : public FSequencerTrackFilter
+{
+	virtual FString GetName() const override { return TEXT("ControlRigControlsSelectedFilter"); }
+	virtual FText GetDisplayName() const override { return LOCTEXT("SequenceTrackFilter_ControlRigSelectedControls", "Selected Control Rig Controls"); }
+	virtual FText GetToolTipText() const override { return LOCTEXT("SequencerTrackFilter_ControlRigSelectedControls", "Show Only Selected Control Rig Controls."); }
+	virtual FSlateIcon GetIcon() const { return FSlateIconFinder::FindIconForClass(UControlRig::StaticClass()); }
+
+	virtual bool PassesFilter(FTrackFilterType InItem) const override
+	{
+		return true;
+	}
+
+	virtual bool PassesFilterWithDisplayName(FTrackFilterType InItem, const FText& InText) const
+	{
+		const UMovieSceneControlRigParameterTrack* Track = Cast<UMovieSceneControlRigParameterTrack>(InItem);
+		if (Track)
+		{
+			UControlRig *ControlRig = Track->GetControlRig();
+			if (ControlRig)
+			{
+				FName Name(*InText.ToString());
+				TArray<FName> ControlNames = ControlRig->CurrentControlSelection();
+				for (const FName& ControlName : ControlNames)
+				{
+					if (Name == ControlName)
+					{
+						return true;
+					}
+				}
+			}
+
+		}
+		return false;
+	}
+
+};
+
+/*
+
+	Bool,
+	Float,
+	Vector2D,
+	Position,
+	Scale,
+	Rotator,
+	Transform
+*/
+
+//////////////////////////////////////////////////////////////////////////
+//
+
+void UControlRigTrackFilter::AddTrackFilterExtensions(TArray< TSharedRef<class FSequencerTrackFilter> >& InOutFilterList) const
+{
+	InOutFilterList.Add(MakeShared<FSequencerTrackFilter_ControlRigControls>());
+	InOutFilterList.Add(MakeShared<FSequencerTrackFilter_ControlRigSelectedControls>());
+
+}
+
+#undef LOCTEXT_NAMESPACE
diff --git a/Engine/Plugins/Experimental/ControlRig/Source/ControlRigEditor/Private/Sequencer/ControlRigSequencerFilter.h b/Engine/Plugins/Experimental/ControlRig/Source/ControlRigEditor/Private/Sequencer/ControlRigSequencerFilter.h
new file mode 100644
--- /dev/null
+++ b/Engine/Plugins/Experimental/ControlRig/Source/ControlRigEditor/Private/Sequencer/ControlRigSequencerFilter.h
@@ -0,0 +1,20 @@
+// Copyright Epic Games, Inc. All Rights Reserved.
+
+#pragma once
+
+#include "CoreMinimal.h"
+#include "UObject/ObjectMacros.h"
+#include "SequencerTrackFilterBase.h"
+#include "SequencerTrackFilterExtension.h"
+#include "ControlRigSequencerFilter.generated.h"
+
+UCLASS()
+class UControlRigTrackFilter : public USequencerTrackFilterExtension
+{
+public:
+	GENERATED_BODY()
+
+	// USequencerTrackFilterExtension interface
+	virtual void AddTrackFilterExtensions(TArray< TSharedRef<class FSequencerTrackFilter> >& InOutFilterList) const override;
+	// End of USequencerTrackFilterExtension interface
+};
diff --git a/Engine/Plugins/Experimental/ControlRig/Source/ControlRigEditor/Public/IControlRigEditorModule.h b/Engine/Plugins/Experimental/ControlRig/Source/ControlRigEditor/Public/IControlRigEditorModule.h
--- a/Engine/Plugins/Experimental/ControlRig/Source/ControlRigEditor/Public/IControlRigEditorModule.h
+++ b/Engine/Plugins/Experimental/ControlRig/Source/ControlRigEditor/Public/IControlRigEditorModule.h
@@ -15,7 +15,6 @@ class IToolkitHost;
 class UControlRigBlueprint;
 class UControlRigGraphNode;
 class UControlRigGraphSchema;
-class URigUnitEditor_Base;
 class FConnectionDrawingPolicy;
 
 class IControlRigEditorModule : public IModuleInterface, public IHasMenuExtensibility, public IHasToolBarExtensibility
@@ -41,10 +40,7 @@ public:
 	DECLARE_DELEGATE_RetVal_TwoParams(TSharedRef<FExtender>, FControlRigEditorToolbarExtender, const TSharedRef<FUICommandList> /*InCommandList*/, TSharedRef<IControlRigEditor> /*InControlRigEditor*/);
 	virtual TArray<FControlRigEditorToolbarExtender>& GetAllControlRigEditorToolbarExtenders() = 0;
 
-	virtual void RegisterRigUnitEditorClass(FName RigUnitClassName, TSubclassOf<URigUnitEditor_Base> Class) = 0;
-	virtual void UnregisterRigUnitEditorClass(FName RigUnitClassName) = 0;
 	virtual void GetTypeActions(const UControlRigBlueprint* CRB, FBlueprintActionDatabaseRegistrar& ActionRegistrar) = 0;
-	virtual void GetInstanceActions(const UControlRigBlueprint* CRB, FBlueprintActionDatabaseRegistrar& ActionRegistrar) = 0;
 	virtual FConnectionDrawingPolicy* CreateConnectionDrawingPolicy(int32 InBackLayerID, int32 InFrontLayerID, float InZoomFactor, const FSlateRect& InClippingRect, class FSlateWindowElementList& InDrawElements, class UEdGraph* InGraphObj) = 0;
 	virtual void GetNodeContextMenuActions(const UControlRigGraphNode* Node, class UToolMenu* Menu, class UGraphNodeContextMenuContext* Context) const = 0;
 	virtual void GetContextMenuActions(const UControlRigGraphSchema* Schema, class UToolMenu* Menu, class UGraphNodeContextMenuContext* Context) const = 0;
diff --git a/Engine/Plugins/Experimental/ControlRig/Source/ControlRigManipulation/ControlRigManipulation.Build.cs b/Engine/Plugins/Experimental/ControlRig/Source/ControlRigManipulation/ControlRigManipulation.Build.cs
--- a/Engine/Plugins/Experimental/ControlRig/Source/ControlRigManipulation/ControlRigManipulation.Build.cs
+++ b/Engine/Plugins/Experimental/ControlRig/Source/ControlRigManipulation/ControlRigManipulation.Build.cs
@@ -15,6 +15,8 @@ namespace UnrealBuildTool.Rules
                     "CoreUObject",
                     "Engine",
                     "ControlRig",
+					"ControlRigDeveloper",
+                    "UnrealEd",
                 }
             );
         }
diff --git a/Engine/Plugins/Experimental/ControlRig/Source/ControlRigManipulation/Private/ControlRigGizmoActor.cpp b/Engine/Plugins/Experimental/ControlRig/Source/ControlRigManipulation/Private/ControlRigGizmoActor.cpp
--- a/Engine/Plugins/Experimental/ControlRig/Source/ControlRigManipulation/Private/ControlRigGizmoActor.cpp
+++ b/Engine/Plugins/Experimental/ControlRig/Source/ControlRigManipulation/Private/ControlRigGizmoActor.cpp
@@ -100,6 +100,17 @@ bool AControlRigGizmoActor::IsManipulating() const
 	return bManipulating;
 }
 
+void AControlRigGizmoActor::SetGizmoColor(const FLinearColor& InColor)
+{
+	if (StaticMeshComponent && !ColorParameterName.IsNone())
+	{
+		if (UMaterialInstanceDynamic* MaterialInstance = Cast<UMaterialInstanceDynamic>(StaticMeshComponent->GetMaterial(0)))
+		{
+			MaterialInstance->SetVectorParameterValue(ColorParameterName, FVector(InColor));
+		}
+	}
+}
+
 // FControlRigGizmoHelper START
 
 namespace FControlRigGizmoHelper
@@ -173,6 +184,7 @@ namespace FControlRigGizmoHelper
 			}
 			if (CreationParam.StaticMesh.IsValid())
 			{
+				GizmoActor->ColorParameterName = CreationParam.ColorParameterName;
 				UMaterialInstanceDynamic* MaterialInstance = UMaterialInstanceDynamic::Create(CreationParam.Material.Get(), GizmoActor);
 				MaterialInstance->SetVectorParameterValue(CreationParam.ColorParameterName, FVector(CreationParam.Color));
 				MeshComponent->SetMaterial(0, MaterialInstance);
diff --git a/Engine/Plugins/Experimental/ControlRig/Source/ControlRigManipulation/Private/DefaultControlRigManipulationLayer.cpp b/Engine/Plugins/Experimental/ControlRig/Source/ControlRigManipulation/Private/DefaultControlRigManipulationLayer.cpp
--- a/Engine/Plugins/Experimental/ControlRig/Source/ControlRigManipulation/Private/DefaultControlRigManipulationLayer.cpp
+++ b/Engine/Plugins/Experimental/ControlRig/Source/ControlRigManipulation/Private/DefaultControlRigManipulationLayer.cpp
@@ -8,8 +8,15 @@
 #include "ControlRigGizmoLibrary.h"
 #include "IControlRigObjectBinding.h"
 
+#if WITH_EDITOR
+#include "ControlRigBlueprint.h"
+#endif
+
+#define LOCTEXT_NAMESPACE "ControlRigDefaultManipulationLayer"
+
 UDefaultControlRigManipulationLayer::UDefaultControlRigManipulationLayer(const FObjectInitializer& ObjectInitializer)
 	: Super(ObjectInitializer)
+	, InteractionTransaction(nullptr)
 {
 }
 
@@ -19,6 +26,7 @@ void UDefaultControlRigManipulationLayer::CreateLayer()
 
 void UDefaultControlRigManipulationLayer::DestroyLayer()
 {
+	EndTransaction();
 	IControlRigManipulationLayer::DestroyLayer();
 }
 
@@ -27,8 +35,10 @@ void UDefaultControlRigManipulationLayer::AddManipulatableObject(IControlRigMani
 	if (InObject)
 	{
 		IControlRigManipulationLayer::AddManipulatableObject(InObject);
-
-		OnControlRigAdded(static_cast<UControlRig*>(InObject));
+		if (UControlRig* ControlRig = Cast<UControlRig>(InObject))
+		{
+			OnControlRigAdded(ControlRig);
+		}
 	}
 }
 
@@ -36,13 +46,14 @@ void UDefaultControlRigManipulationLayer::RemoveManipulatableObject(IControlRigM
 {
 	if (InObject)
 	{
-		OnControlRigRemoved(static_cast<UControlRig*>(InObject));
-
+		if (UControlRig* ControlRig = Cast<UControlRig>(InObject))
+		{
+			OnControlRigRemoved(ControlRig);
+		}
 		IControlRigManipulationLayer::RemoveManipulatableObject(InObject);
 	}
 }
 
-
 void UDefaultControlRigManipulationLayer::OnControlRigAdded(UControlRig* InControlRig)
 {
 	// bind execution delegate
@@ -61,6 +72,7 @@ void UDefaultControlRigManipulationLayer::OnControlRigAdded(UControlRig* InContr
 	if (MeshComponent)
 	{
 		MeshComponent->OnBoneTransformsFinalized.AddDynamic(this, &UDefaultControlRigManipulationLayer::PostPoseUpdate);
+		MeshComponent->OnAnimInitialized.AddDynamic(this, &UDefaultControlRigManipulationLayer::OnPoseInitialized);
 	}
 }
 
@@ -79,6 +91,7 @@ void UDefaultControlRigManipulationLayer::OnControlRigRemoved(UControlRig* InCon
 			if (MeshComponent)
 			{
 				MeshComponent->OnBoneTransformsFinalized.RemoveDynamic(this, &UDefaultControlRigManipulationLayer::PostPoseUpdate);
+				MeshComponent->OnAnimInitialized.RemoveDynamic(this, &UDefaultControlRigManipulationLayer::OnPoseInitialized);
 			}
 		}
 
@@ -146,7 +159,17 @@ void UDefaultControlRigManipulationLayer::MoveGizmo(AControlRigGizmoActor* Gizmo
 			FTransform NewTransform = CurrentTransform.GetRelativeTransform(ToWorldTransform);
 			Manip->ManipObject->SetControlGlobalTransform(Manip->ControlName, NewTransform);
 			// assumes it's attached to actor
-			GizmoActor->SetGlobalTransform(NewTransform);
+			GizmoActor->SetGlobalTransform(CurrentTransform);
+
+#if WITH_EDITOR
+			if (UControlRig* ControlRig = Cast<UControlRig>(Manip->ManipObject))
+			{
+				if(UControlRigBlueprint* Blueprint = Cast<UControlRigBlueprint>(ControlRig->GetClass()->ClassGeneratedBy))
+				{
+					Blueprint->PropagatePoseFromInstanceToBP(ControlRig);
+				}
+			}
+#endif
 		}
 	}
 }
@@ -156,12 +179,20 @@ bool IsSupportedControlType(const ERigControlType ControlType)
 {
 	switch (ControlType)
 	{
-	case ERigControlType::Position:
-	case ERigControlType::Scale:
-	case ERigControlType::Quat:
-	case ERigControlType::Rotator:
-	case ERigControlType::Transform:
-		return true;
+		case ERigControlType::Float:
+		case ERigControlType::Vector2D:
+		case ERigControlType::Position:
+		case ERigControlType::Scale:
+		case ERigControlType::Rotator:
+		case ERigControlType::Transform:
+		case ERigControlType::TransformNoScale:
+		{
+			return true;
+		}
+		default:
+		{
+			break;
+		}
 	}
 
 	return false;
@@ -191,6 +222,10 @@ void UDefaultControlRigManipulationLayer::GetGizmoCreationParams(TArray<FGizmoAc
 
 		for (const FRigControl& Control : Controls)
 		{
+			if(!Control.bGizmoEnabled)
+			{
+				continue;
+			}
 			if (IsSupportedControlType(Control.ControlType))
 			{
 				FGizmoActorCreationParam Param;
@@ -225,10 +260,83 @@ bool UDefaultControlRigManipulationLayer::ModeSupportedByGizmoActor(const AContr
 		const FRigControl* RigControl = Control->ManipObject->FindControl(Control->ControlName);
 		if (RigControl)
 		{
-			return IsSupportedControlType(RigControl->ControlType);
+			if (IsSupportedControlType(RigControl->ControlType))
+			{
+				switch (InMode)
+				{
+					case FWidget::WM_Rotate:
+					{
+						switch (RigControl->ControlType)
+						{
+							case ERigControlType::Rotator:
+							case ERigControlType::Transform:
+							case ERigControlType::TransformNoScale:
+							{
+								return true;
+							}
+							default:
+							{
+								break;
+							}
+						}
+						break;
+					}
+					case FWidget::WM_Translate:
+					{
+						switch (RigControl->ControlType)
+						{
+							case ERigControlType::Float:
+							case ERigControlType::Vector2D:
+							case ERigControlType::Position:
+							case ERigControlType::Transform:
+							case ERigControlType::TransformNoScale:
+							{
+								return true;
+							}
+							default:
+							{
+								break;
+							}
+						}
+						break;
+					}
+					case FWidget::WM_Scale:
+					{
+						switch (RigControl->ControlType)
+						{
+							case ERigControlType::Scale:
+							case ERigControlType::Transform:
+							{
+								return true;
+							}
+							default:
+							{
+								break;
+							}
+						}
+						break;
+					}
+					case FWidget::WM_TranslateRotateZ:
+					{
+						switch (RigControl->ControlType)
+						{
+							case ERigControlType::Transform:
+							case ERigControlType::TransformNoScale:
+							{
+								return true;
+							}
+							default:
+							{
+								break;
+							}
+						}
+						break;
+					}
+				}
+			}
 		}
 	}
-	return true;
+	return false;
 }
 
 void UDefaultControlRigManipulationLayer::TickGizmo(AControlRigGizmoActor* GizmoActor, const FTransform& ComponentTransform)
@@ -240,6 +348,11 @@ void UDefaultControlRigManipulationLayer::TickGizmo(AControlRigGizmoActor* Gizmo
 		{
 			FTransform Transform = DataPtr->ManipObject->GetControlGlobalTransform(DataPtr->ControlName);
 			GizmoActor->SetActorTransform(Transform * ComponentTransform);
+
+			if (FRigControl* Control = DataPtr->ManipObject->FindControl(DataPtr->ControlName))
+			{
+				GizmoActor->SetGizmoColor(Control->GizmoColor);
+			}
 		}
 	}
 }
@@ -285,6 +398,8 @@ void UDefaultControlRigManipulationLayer::TickManipulatableObjects(float DeltaTi
 		SkeletalMeshComponent->MarkRenderTransformDirty();
 		SkeletalMeshComponent->MarkRenderDynamicDataDirty();
 	}
+
+	PostPoseUpdate();
 }
 
 bool UDefaultControlRigManipulationLayer::GetGlobalTransform(AControlRigGizmoActor* GizmoActor, const FName& ControlName, FTransform& OutTransform) const
@@ -370,13 +485,13 @@ void UDefaultControlRigManipulationLayer::DestroyGizmosActors()
 	FWorldDelegates::OnWorldCleanup.Remove(OnWorldCleanupHandle);
 }
 
-void UDefaultControlRigManipulationLayer::OnControlModified(IControlRigManipulatable* InManipulatable, const FRigControl& InControl)
+void UDefaultControlRigManipulationLayer::OnControlModified(IControlRigManipulatable* InManipulatable, const FRigControl& InControl, EControlRigSetKey InSetKey)
 {
 	if (UControlRig* ControlRig = static_cast<UControlRig*>(InManipulatable))
 	{
+		FTransform ComponentTransform = GetSkeletalMeshComponentTransform();
 		if (AControlRigGizmoActor* const* Actor = GizmoToControlMap.FindKey(InControl.Index))
 		{
-			FTransform ComponentTransform = GetSkeletalMeshComponentTransform();
 			TickGizmo(*Actor, ComponentTransform);
 		}
 	}
@@ -386,7 +501,7 @@ TSharedPtr<IControlRigObjectBinding> UDefaultControlRigManipulationLayer::GetObj
 {
 	for (TWeakObjectPtr<UObject> Manip : ManipulatableObjects)
 	{
-		if (UControlRig* ControlRig = Cast<UControlRig>(Manip))
+		if(UControlRig* ControlRig = Cast<UControlRig>(Manip))
 		{
 			return ControlRig->GetObjectBinding();
 		}
@@ -425,10 +540,24 @@ FTransform	UDefaultControlRigManipulationLayer::GetSkeletalMeshComponentTransfor
 
 void UDefaultControlRigManipulationLayer::BeginTransaction()
 {
+	if (InteractionTransaction != nullptr)
+	{
+		EndTransaction();
+	}
+
+	InteractionTransaction = new FScopedTransaction(LOCTEXT("TransformGizmo", "Transform Gizmo"));
+
 	for (TWeakObjectPtr<UObject> Manip : ManipulatableObjects)
 	{
 		if (UControlRig* ControlRig = Cast<UControlRig>(Manip))
 		{
+			UObject* Blueprint = ControlRig->GetClass()->ClassGeneratedBy;
+			if (Blueprint)
+			{
+				Blueprint->SetFlags(RF_Transactional);
+				Blueprint->Modify();
+			}
+
 			ControlRig->SetFlags(RF_Transactional);
 			ControlRig->Modify();
 		}
@@ -437,6 +566,11 @@ void UDefaultControlRigManipulationLayer::BeginTransaction()
 
 void UDefaultControlRigManipulationLayer::EndTransaction()
 {
+	if (InteractionTransaction != nullptr)
+	{
+		delete InteractionTransaction;
+		InteractionTransaction = nullptr;
+	}
 }
 
 void UDefaultControlRigManipulationLayer::PostPoseUpdate()
@@ -447,4 +581,12 @@ void UDefaultControlRigManipulationLayer::PostPoseUpdate()
 	{
 		TickGizmo(Iter.Key(), ComponentTransform);
 	}
-}
\ No newline at end of file
+}
+
+void UDefaultControlRigManipulationLayer::OnPoseInitialized()
+{
+	// broadcast delegates
+	OnAnimSystemInitialized.Broadcast();
+}
+
+#undef LOCTEXT_NAMESPACE
diff --git a/Engine/Plugins/Experimental/ControlRig/Source/ControlRigManipulation/Public/ControlRigGizmoActor.h b/Engine/Plugins/Experimental/ControlRig/Source/ControlRigManipulation/Public/ControlRigGizmoActor.h
--- a/Engine/Plugins/Experimental/ControlRig/Source/ControlRigManipulation/Public/ControlRigGizmoActor.h
+++ b/Engine/Plugins/Experimental/ControlRig/Source/ControlRigManipulation/Public/ControlRigGizmoActor.h
@@ -58,6 +58,10 @@ public:
 	UPROPERTY(Category = StaticMesh, VisibleAnywhere, BlueprintReadOnly, meta = (ExposeFunctionCategories = "Mesh,Rendering,Physics,Components|StaticMesh", AllowPrivateAccess = "true"))
 	class UStaticMeshComponent* StaticMeshComponent;
 
+	// the name of the color parameter on the material
+	UPROPERTY()
+	FName ColorParameterName;
+
 	UFUNCTION(BlueprintSetter)
 	/** Set the control to be enabled/disabled */
 	virtual void SetEnabled(bool bInEnabled);
@@ -93,6 +97,9 @@ public:
 	/** Called from the edit mode each tick */
 	virtual void TickControl() {};
 
+	/** changes the gizmo color */
+	virtual void SetGizmoColor(const FLinearColor& InColor);
+
 	/** Event called when the transform of this control has changed */
 	UFUNCTION(BlueprintImplementableEvent)
 	void OnTransformChanged(const FTransform& NewTransform);
diff --git a/Engine/Plugins/Experimental/ControlRig/Source/ControlRigManipulation/Public/DefaultControlRigManipulationLayer.h b/Engine/Plugins/Experimental/ControlRig/Source/ControlRigManipulation/Public/DefaultControlRigManipulationLayer.h
--- a/Engine/Plugins/Experimental/ControlRig/Source/ControlRigManipulation/Public/DefaultControlRigManipulationLayer.h
+++ b/Engine/Plugins/Experimental/ControlRig/Source/ControlRigManipulation/Public/DefaultControlRigManipulationLayer.h
@@ -5,11 +5,13 @@
 #include "IControlRigManipulationLayer.h"
 #include "ControlRigGizmoActor.h"
 #include "Units/RigUnitContext.h"
+#include "ScopedTransaction.h"
 #include "DefaultControlRigManipulationLayer.generated.h"
 
 class IControlRigObjectBinding;
 typedef uint32 FControlID;
-
+enum class EControlRigSetKey : uint8;
+	
 // control ID, I should use ID for it
 struct FControlData
 {
@@ -73,6 +75,9 @@ public:
 	AControlRigGizmoActor* GetGizmoFromControlName(const FName& ControlName) const;
 	bool GetGlobalTransform(AControlRigGizmoActor* GizmoActor, const FName& ControlName, FTransform& OutTransform) const;
 
+	// multi delegate for anim system initialized
+	FSimpleMulticastDelegate OnAnimSystemInitialized;
+
 private:
 	// GizmoActor* to ControlData index
 	TMap<AControlRigGizmoActor*, FControlID> GizmoToControlMap;
@@ -89,7 +94,10 @@ private:
 	UFUNCTION()
 	virtual void PostPoseUpdate();
 
-	void OnControlModified(IControlRigManipulatable* InManipulatable, const FRigControl& InControl);
+	UFUNCTION()
+	virtual void OnPoseInitialized();
+
+	void OnControlModified(IControlRigManipulatable* InManipulatable, const FRigControl& InControl, EControlRigSetKey InSetKey);
 	TArray<FDelegateHandle> ControlModifiedDelegateHandles;
 
 	void OnControlRigAdded(UControlRig* InControlRig);
@@ -100,4 +108,6 @@ private:
 	FDelegateHandle OnWorldCleanupHandle;
 	void OnWorldCleanup(UWorld* World, bool bSessionEnded, bool bCleanupResources);
 	UWorld* WorldPtr = nullptr;
+
+	FScopedTransaction* InteractionTransaction = nullptr;
 };
diff --git a/Engine/Plugins/Experimental/GeometryCollectionPlugin/Source/GeometryCollectionEditor/Private/GeometryCollection/GeometryCollectionConversion.cpp b/Engine/Plugins/Experimental/GeometryCollectionPlugin/Source/GeometryCollectionEditor/Private/GeometryCollection/GeometryCollectionConversion.cpp
--- a/Engine/Plugins/Experimental/GeometryCollectionPlugin/Source/GeometryCollectionEditor/Private/GeometryCollection/GeometryCollectionConversion.cpp
+++ b/Engine/Plugins/Experimental/GeometryCollectionPlugin/Source/GeometryCollectionEditor/Private/GeometryCollection/GeometryCollectionConversion.cpp
@@ -330,10 +330,10 @@ void FGeometryCollectionConversion::AppendSkeletalMesh(const USkeletalMesh* Skel
 						{
 							int VertexOffset = VertexBaseIndex + VertexIndex;
 							BoneMap[VertexOffset] = -1;
-							if (const TSkinWeightInfo<false> * SkinWeightInfo = SkinWeightVertexBuffer.GetSkinWeightPtr<false>(VertexIndex))
+							int32 SkeletalBoneIndex = -1;
+							check(SkinWeightVertexBuffer.GetRigidWeightBone(VertexIndex, SkeletalBoneIndex));
+							if (SkeletalBoneIndex > -1)
 							{
-								uint8 SkeletalBoneIndex = -1;
-								check(SkinWeightInfo->GetRigidWeightBone(SkeletalBoneIndex));
 								BoneMap[VertexOffset] = SkeletalBoneIndex + TransformBaseIndex;
 								Vertex[VertexOffset] = Transform[BoneMap[VertexOffset]].ToInverseMatrixWithScale().TransformPosition(PositionVertexBuffer.VertexPosition(VertexIndex));
 							}
diff --git a/Engine/Plugins/Experimental/LiveLinkControlRig/Source/LiveLinkControlRig/Public/LiveLinkRigUnits.h b/Engine/Plugins/Experimental/LiveLinkControlRig/Source/LiveLinkControlRig/Public/LiveLinkRigUnits.h
new file mode 100644
--- /dev/null
+++ b/Engine/Plugins/Experimental/LiveLinkControlRig/Source/LiveLinkControlRig/Public/LiveLinkRigUnits.h
@@ -0,0 +1,110 @@
+// Copyright Epic Games, Inc. All Rights Reserved.
+#pragma once
+#include "CoreMinimal.h"
+#include "ControlRigDefines.h"
+#include "Units/RigUnit.h"
+#include "Roles/LiveLinkAnimationBlueprintStructs.h"
+#include "LiveLinkTypes.h"
+#include "LiveLinkRigUnits.generated.h"
+
+namespace LiveLinkControlRigUtilities
+{
+	ILiveLinkClient* TryGetLiveLinkClient();
+}
+
+USTRUCT(meta = (Abstract, NodeColor = "0.3 0.1 0.1"))
+struct LIVELINKCONTROLRIG_API FRigUnit_LiveLinkBase : public FRigUnit
+{
+	GENERATED_BODY()
+};
+
+USTRUCT(meta = (DisplayName = "Evaluate Live Link Frame (Animation)", Category = "Live Link"))
+struct LIVELINKCONTROLRIG_API FRigUnit_LiveLinkEvaluteFrameAnimation : public FRigUnit_LiveLinkBase
+{
+	GENERATED_BODY()
+
+	/** Execute logic for this rig unit */
+	RIGVM_METHOD()
+	virtual void Execute(const FRigUnitContext& Context) override;
+
+	UPROPERTY(meta = (Input))
+	FName SubjectName;
+
+	UPROPERTY(meta = (Input))
+	bool bDrawDebug;
+
+	UPROPERTY(meta = (Input))
+	FLinearColor DebugColor = FLinearColor::Red;
+
+	UPROPERTY(meta = (Input))
+	FTransform DebugDrawOffset;
+
+	UPROPERTY(meta = (Output))
+	FSubjectFrameHandle SubjectFrame;
+};
+
+USTRUCT(meta = (DisplayName = "Get Transform By Name", Category = "Live Link"))
+struct LIVELINKCONTROLRIG_API FRigUnit_LiveLinkGetTransformByName : public FRigUnit_LiveLinkBase
+{
+	GENERATED_BODY()
+
+	/** Execute logic for this rig unit */
+	RIGVM_METHOD()
+	virtual void Execute(const FRigUnitContext& Context) override;
+
+	UPROPERTY(meta = (Input))
+	FSubjectFrameHandle SubjectFrame;
+
+	UPROPERTY(meta = (Input))
+	FName TransformName;
+
+	UPROPERTY(meta = (Input))
+	EBoneGetterSetterMode Space;
+
+	UPROPERTY(meta = (Output))
+	FTransform Transform;
+};
+
+USTRUCT(meta = (DisplayName = "Get Parameter Value By Name", Category = "Live Link"))
+struct LIVELINKCONTROLRIG_API FRigUnit_LiveLinkGetParameterValueByName : public FRigUnit_LiveLinkBase
+{
+	GENERATED_BODY()
+
+	/** Execute logic for this rig unit */
+	RIGVM_METHOD()
+	virtual void Execute(const FRigUnitContext& Context) override;
+
+	UPROPERTY(meta = (Input))
+	FSubjectFrameHandle SubjectFrame;
+
+	UPROPERTY(meta = (Input))
+	FName ParameterName;
+
+	UPROPERTY(meta = (Output))
+	float Value;
+};
+
+USTRUCT(meta = (DisplayName = "Evaluate Live Link Frame (Transform)", Category = "Live Link"))
+struct LIVELINKCONTROLRIG_API FRigUnit_LiveLinkEvaluteFrameTransform : public FRigUnit_LiveLinkBase
+{
+	GENERATED_BODY()
+
+	/** Execute logic for this rig unit */
+	RIGVM_METHOD()
+	virtual void Execute(const FRigUnitContext& Context) override;
+
+	UPROPERTY(meta = (Input))
+	FName SubjectName;
+
+	UPROPERTY(meta = (Input))
+	bool bDrawDebug;
+
+	UPROPERTY(meta = (Input))
+	FLinearColor DebugColor = FLinearColor::Red;
+
+	UPROPERTY(meta = (Input))
+	FTransform DebugDrawOffset;
+
+	UPROPERTY(meta = (Output))
+	FTransform Transform;
+};
\ No newline at end of file
diff --git a/Engine/Plugins/Experimental/SkeletalReduction/Source/Private/SkeletalMeshReductionPlugin.cpp b/Engine/Plugins/Experimental/SkeletalReduction/Source/Private/SkeletalMeshReductionPlugin.cpp
--- a/Engine/Plugins/Experimental/SkeletalReduction/Source/Private/SkeletalMeshReductionPlugin.cpp
+++ b/Engine/Plugins/Experimental/SkeletalReduction/Source/Private/SkeletalMeshReductionPlugin.cpp
@@ -1380,7 +1380,7 @@ void  FQuadricSkeletalMeshReduction::AddSourceModelInfluences( const FSkeletalMe
 					// Add the non-zero weights 
 					for (int32 b = 0; b < MAX_TOTAL_INFLUENCES; ++b)
 					{
-						uint8 LocalBoneId = SrcRawSkinWeight.InfluenceBones[b];
+						FBoneIndexType LocalBoneId = SrcRawSkinWeight.InfluenceBones[b];
 						uint8 Weight = SrcRawSkinWeight.InfluenceWeights[b];
 
 						checkSlow(LocalBoneId < BoneMap.Num());
diff --git a/Engine/Plugins/FX/Niagara/Source/Niagara/Private/NiagaraDataInterfaceSkeletalMesh.cpp b/Engine/Plugins/FX/Niagara/Source/Niagara/Private/NiagaraDataInterfaceSkeletalMesh.cpp
--- a/Engine/Plugins/FX/Niagara/Source/Niagara/Private/NiagaraDataInterfaceSkeletalMesh.cpp
+++ b/Engine/Plugins/FX/Niagara/Source/Niagara/Private/NiagaraDataInterfaceSkeletalMesh.cpp
@@ -417,7 +417,7 @@ void FSkeletalMeshGpuSpawnStaticBuffers::InitRHI()
 
 	MeshColorBufferSrv = LODRenderData->StaticVertexBuffers.ColorVertexBuffer.GetColorComponentsSRV();
 
-	NumWeights = LODRenderData->SkinWeightVertexBuffer.HasExtraBoneInfluences() ? 8 : 4;
+	NumWeights = LODRenderData->SkinWeightVertexBuffer.GetMaxBoneInfluences();
 
 	uint32 SectionCount = LODRenderData->RenderSections.Num();
 
@@ -1355,8 +1355,8 @@ bool FNDISkeletalMesh_InstanceData::Init(UNiagaraDataInterfaceSkeletalMesh* Inte
 	//-TODO: We should find out if this DI is connected to a GPU emitter or not rather than a blanket accross the system
 	if ( SystemInstance->HasGPUEmitters() )
 	{
-		MeshWeightStrideByte = SkinWeightBuffer->GetStride();
-		MeshSkinWeightBufferSrv = SkinWeightBuffer->GetSRV();
+		MeshWeightStrideByte = SkinWeightBuffer->GetConstantInfluencesVertexStride();
+		MeshSkinWeightBufferSrv = SkinWeightBuffer->GetDataVertexBuffer()->GetSRV();
 		//check(MeshSkinWeightBufferSrv->IsValid()); // not available in this stream
 
 		FSkeletalMeshLODInfo* LODInfo = Mesh->GetLODInfo(LODIndex);
@@ -1366,7 +1366,7 @@ bool FNDISkeletalMesh_InstanceData::Init(UNiagaraDataInterfaceSkeletalMesh* Inte
 		{
 			UE_LOG(LogNiagara, Warning, TEXT("Skeletal Mesh %s has cloth asset on it: spawning from it might not work properly."), *Mesh->GetName());
 		}
-		if (LODData.DoesVertexBufferHaveExtraBoneInfluences())
+		if (LODData.GetVertexBufferMaxBoneInfluences() > MAX_INFLUENCES_PER_STREAM)
 		{
 			UE_LOG(LogNiagara, Warning, TEXT("Skeletal Mesh %s has bones extra influence: spawning from it might not work properly."), *Mesh->GetName());
 		}
diff --git a/Engine/Plugins/MovieScene/SequencerScripting/Source/SequencerScriptingEditor/Private/SequencerTools.cpp b/Engine/Plugins/MovieScene/SequencerScripting/Source/SequencerScriptingEditor/Private/SequencerTools.cpp
--- a/Engine/Plugins/MovieScene/SequencerScripting/Source/SequencerScriptingEditor/Private/SequencerTools.cpp
+++ b/Engine/Plugins/MovieScene/SequencerScripting/Source/SequencerScriptingEditor/Private/SequencerTools.cpp
@@ -16,6 +16,8 @@
 #include "CineCameraActor.h"
 #include "CineCameraComponent.h"
 #include "MovieSceneCommonHelpers.h"
+#include "Components/SkeletalMeshComponent.h"
+#include "Animation/AnimSequenceBase.h"
 
 #define LOCTEXT_NAMESPACE "SequencerTools"
 
@@ -103,8 +105,6 @@ void USequencerToolsFunctionLibrary::CancelMovieRender()
 
 bool USequencerToolsFunctionLibrary::ExportFBX(UWorld* World, ULevelSequence* Sequence, const TArray<FSequencerBindingProxy>& InBindings, UFbxExportOption* OverrideOptions, const FString& InFBXFileName)
 {
-	
-	
 	UnFbx::FFbxExporter* Exporter = UnFbx::FFbxExporter::GetInstance();
 	//Show the fbx export dialog options
 	Exporter->SetExportOptionsOverride(OverrideOptions);
@@ -133,6 +133,60 @@ bool USequencerToolsFunctionLibrary::ExportFBX(UWorld* World, ULevelSequence* Se
 	return bDidExport;
 }
 
+static USkeletalMeshComponent* GetSkelMeshComponent(IMovieScenePlayer* Player, const FSequencerBindingProxy& Binding)
+{
+	FMovieSceneSequenceIDRef Template = MovieSceneSequenceID::Root;
+	for (TWeakObjectPtr<UObject> RuntimeObject : Player->FindBoundObjects(Binding.BindingID, Template))
+	{
+
+		if (AActor* Actor = Cast<AActor>(RuntimeObject.Get()))
+		{
+			for (UActorComponent* Component : Actor->GetComponents())
+			{
+				if (USkeletalMeshComponent* SkeletalMeshComp = Cast<USkeletalMeshComponent>(Component))
+				{
+					return SkeletalMeshComp;
+				}
+			}
+		}
+		else if (USkeletalMeshComponent* SkeletalMeshComponent = Cast<USkeletalMeshComponent>(RuntimeObject.Get()))
+		{
+			if (SkeletalMeshComponent->SkeletalMesh)
+			{
+				return SkeletalMeshComponent;
+			}
+		}
+	}
+	return nullptr;
+}
+
+bool USequencerToolsFunctionLibrary::ExportAnimSequence(UWorld* World, ULevelSequence*  Sequence,  UAnimSequence* AnimSequence, const FSequencerBindingProxy& Binding)
+{
+	UMovieScene* MovieScene = Sequence->GetMovieScene();
+	if (Binding.Sequence != Sequence || !AnimSequence)
+	{
+		return false;
+	}
+	INodeNameAdapter NodeNameAdapter;
+	ALevelSequenceActor* OutActor;
+	FMovieSceneSequencePlaybackSettings Settings;
+	FLevelSequenceCameraSettings CameraSettings;
+	FMovieSceneSequenceIDRef Template = MovieSceneSequenceID::Root;
+	FMovieSceneSequenceTransform RootToLocalTransform;
+	ULevelSequencePlayer* Player = ULevelSequencePlayer::CreateLevelSequencePlayer(World, Sequence, Settings, OutActor);
+	Player->Initialize(Sequence, World->PersistentLevel, Settings, CameraSettings);
+	Player->State.AssignSequence(MovieSceneSequenceID::Root, *Sequence, *Player);
+
+	bool bResult = false;
+	USkeletalMeshComponent* SkeletalMeshComp =  GetSkelMeshComponent(Player, Binding);
+	if (SkeletalMeshComp && SkeletalMeshComp->SkeletalMesh && SkeletalMeshComp->SkeletalMesh->Skeleton)
+	{
+		AnimSequence->SetSkeleton(SkeletalMeshComp->SkeletalMesh->Skeleton);
+		bResult = MovieSceneToolHelpers::ExportToAnimSequence(AnimSequence, MovieScene, Player, SkeletalMeshComp, Template, RootToLocalTransform);
+	}
+	return bResult;
+}
+
 
 
 TArray<FGuid> AddActors(UWorld* World, UMovieSceneSequence* InSequence, UMovieScene* InMovieScene, IMovieScenePlayer* Player, FMovieSceneSequenceIDRef TemplateID,const TArray<TWeakObjectPtr<AActor> >& InActors)
diff --git a/Engine/Plugins/MovieScene/SequencerScripting/Source/SequencerScriptingEditor/Public/SequencerTools.h b/Engine/Plugins/MovieScene/SequencerScripting/Source/SequencerScriptingEditor/Public/SequencerTools.h
--- a/Engine/Plugins/MovieScene/SequencerScripting/Source/SequencerScriptingEditor/Public/SequencerTools.h
+++ b/Engine/Plugins/MovieScene/SequencerScripting/Source/SequencerScriptingEditor/Public/SequencerTools.h
@@ -10,6 +10,8 @@
 #include "SequencerTools.generated.h"
 
 class UFbxExportOption;
+class UAnimSequence;
+class UPoseAsset;
 
 DECLARE_DYNAMIC_DELEGATE_OneParam(FOnRenderMovieStopped, bool, bSuccess);
 
@@ -59,6 +61,18 @@ public:
 	UFUNCTION(BlueprintCallable, Category = "Editor Scripting | Sequencer Tools | FBX")
 	static bool ExportFBX(UWorld* InWorld, ULevelSequence* InSequence, const TArray<FSequencerBindingProxy>& InBindings, UFbxExportOption* OverrideOptions,const FString& InFBXFileName);
 
+	/*
+	 * Export Passed in Binding as an Anim Seqquence.
+	 *
+	 * @InWorld World to export
+	 * @InSequence Sequence to export
+	 * @AnimSequence The AnimSequence to save into.
+	 * @InBinding Binding to export that has a skelmesh component on it
+	 * @InAnimSequenceFilename File to create
+	 */
+	UFUNCTION(BlueprintCallable, Category = "Editor Scripting | Sequencer Tools | Animation")
+	static bool ExportAnimSequence(UWorld* World, ULevelSequence*  Sequence, UAnimSequence* AnimSequence, const FSequencerBindingProxy& Binding);
+
 	/*
 	 * Import Passed in Bindings to FBX
 	 *
@@ -71,5 +85,4 @@ public:
 	UFUNCTION(BlueprintCallable, Category = "Editor Scripting | Sequencer Tools | FBX")
 	static bool ImportFBX(UWorld* InWorld, ULevelSequence* InSequence, const TArray<FSequencerBindingProxy>& InBindings, UMovieSceneUserImportFBXSettings* InImportFBXSettings, const FString&  InImportFilename);
 
-
 };
\ No newline at end of file
diff --git a/Engine/Plugins/Runtime/AudioModulation/Source/AudioModulationEditor/AudioModulationEditor.build.cs b/Engine/Plugins/Runtime/AudioModulation/Source/AudioModulationEditor/AudioModulationEditor.build.cs
--- a/Engine/Plugins/Runtime/AudioModulation/Source/AudioModulationEditor/AudioModulationEditor.build.cs
+++ b/Engine/Plugins/Runtime/AudioModulation/Source/AudioModulationEditor/AudioModulationEditor.build.cs
@@ -23,7 +23,7 @@ public class AudioModulationEditor : ModuleRules
 				"AudioEditor",
 				"AudioModulation",
 				"CurveEditor",
-				"Persona"
+				"EditorWidgets"
 			}
 		);
 
diff --git a/Engine/Plugins/Runtime/AudioModulation/Source/AudioModulationEditor/Private/Editors/ModulationSettingsCurveEditorViewStacked.cpp b/Engine/Plugins/Runtime/AudioModulation/Source/AudioModulationEditor/Private/Editors/ModulationSettingsCurveEditorViewStacked.cpp
--- a/Engine/Plugins/Runtime/AudioModulation/Source/AudioModulationEditor/Private/Editors/ModulationSettingsCurveEditorViewStacked.cpp
+++ b/Engine/Plugins/Runtime/AudioModulation/Source/AudioModulationEditor/Private/Editors/ModulationSettingsCurveEditorViewStacked.cpp
@@ -16,7 +16,7 @@
 
 
 FModCurveEditorModel::FModCurveEditorModel(FRichCurve& InRichCurve, UObject* InOwner, FName InControlName, EModSettingsOutputEditorCurveSource InSource, UCurveFloat* InSharedCurve)
-	: FRichCurveEditorModel(&InRichCurve, InOwner)
+	: FRichCurveEditorModelRaw(&InRichCurve, InOwner)
 	, Output(EModSettingsEditorCurveOutput::Control)
 	, Source(InSource)
 {
@@ -25,7 +25,7 @@ FModCurveEditorModel::FModCurveEditorModel(FRichCurve& InRichCurve, UObject* InO
 }
 
 FModCurveEditorModel::FModCurveEditorModel(FRichCurve& InRichCurve, UObject* InOwner, EModSettingsEditorCurveOutput InOutput, EModSettingsOutputEditorCurveSource InSource, UCurveFloat* InSharedCurve)
-	: FRichCurveEditorModel(&InRichCurve, InOwner)
+	: FRichCurveEditorModelRaw(&InRichCurve, InOwner)
 	, Output(InOutput)
 	, Source(InSource)
 {
@@ -97,7 +97,7 @@ void FModCurveEditorModel::Refresh(EModSettingsEditorCurveOutput InCurveOutput,
 
 		case EModSettingsOutputEditorCurveSource::Expression:
 		{
-			bKeyDrawEnabled = 0;
+			bKeyDrawEnabled = false;
 			ShortDisplayName = FText::Format(LOCTEXT("ModulationOutputCurveExpressionDisplayName", "{0} (Expression)"), ShortNameBase);
 		}
 		break;
@@ -105,7 +105,7 @@ void FModCurveEditorModel::Refresh(EModSettingsEditorCurveOutput InCurveOutput,
 		case EModSettingsOutputEditorCurveSource::Unset:
 		default:
 		{
-			bKeyDrawEnabled = 0;
+			bKeyDrawEnabled = false;
 			ShortDisplayName = FText::Format(LOCTEXT("ModulationOutputCurveUnsetDisplayName", "{0} (Shared - Unset)"), ShortNameBase);
 		}
 	}
diff --git a/Engine/Plugins/Runtime/AudioModulation/Source/AudioModulationEditor/Private/Editors/ModulationSettingsCurveEditorViewStacked.h b/Engine/Plugins/Runtime/AudioModulation/Source/AudioModulationEditor/Private/Editors/ModulationSettingsCurveEditorViewStacked.h
--- a/Engine/Plugins/Runtime/AudioModulation/Source/AudioModulationEditor/Private/Editors/ModulationSettingsCurveEditorViewStacked.h
+++ b/Engine/Plugins/Runtime/AudioModulation/Source/AudioModulationEditor/Private/Editors/ModulationSettingsCurveEditorViewStacked.h
@@ -137,7 +137,7 @@ public:
 };
 
 
-class FModCurveEditorModel : public FRichCurveEditorModel
+class FModCurveEditorModel : public FRichCurveEditorModelRaw
 {
 public:
 
diff --git a/Engine/Plugins/Runtime/Oculus/OculusAvatar/Source/Public/OvrAvatar.cpp b/Engine/Plugins/Runtime/Oculus/OculusAvatar/Source/Public/OvrAvatar.cpp
--- a/Engine/Plugins/Runtime/Oculus/OculusAvatar/Source/Public/OvrAvatar.cpp
+++ b/Engine/Plugins/Runtime/Oculus/OculusAvatar/Source/Public/OvrAvatar.cpp
@@ -1406,7 +1406,7 @@ void UOvrAvatar::LoadMesh(USkeletalMesh* SkeletalMesh, const ovrAvatarMeshAssetD
 	LodRenderData->Sections.SetNumUninitialized(SubmeshCount);
 #else
 	TArray<FColor> ColorArray;
-	TArray<TSkinWeightInfo<true>> InWeights;
+	TArray<FSkinWeightInfo> InWeights;
 	InWeights.AddUninitialized(data->vertexCount);
 	TMap<int32, TArray<int32>> OverlappingVertices;
 
@@ -1612,7 +1612,7 @@ void UOvrAvatar::LoadMesh(USkeletalMesh* SkeletalMesh, const ovrAvatarMeshAssetD
 	}
 #else
 	LodRenderData->StaticVertexBuffers.ColorVertexBuffer.InitFromColorArray(ColorArray);
-	LodRenderData->SkinWeightVertexBuffer.SetHasExtraBoneInfluences(true);
+	LodRenderData->SkinWeightVertexBuffer.SetMaxBoneInfluences(EXTRA_BONE_INFLUENCES);
 	LodRenderData->SkinWeightVertexBuffer = InWeights;
 	LodRenderData->MultiSizeIndexContainer.CreateIndexBuffer(sizeof(uint16_t));
 
@@ -2458,7 +2458,7 @@ void UOvrAvatar::LoadCombinedMesh(USkeletalMesh* SkeletalMesh, const ovrAvatarMe
 	UE_LOG(LogAvatars, Display, TEXT("[Mesh] - Submesh Count: %d"), SubmeshCount);
 
 	TArray<FColor> ColorArray;
-	TArray<TSkinWeightInfo<true>> InWeights;
+	TArray<FSkinWeightInfo> InWeights;
 	InWeights.AddUninitialized(TotalVertCount);
 	TMap<int32, TArray<int32>> OverlappingVertices;
 	bool SumMeshVertexOffsetFlag = false;
@@ -2583,7 +2583,7 @@ void UOvrAvatar::LoadCombinedMesh(USkeletalMesh* SkeletalMesh, const ovrAvatarMe
 	UE_LOG(LogAvatars, Display, TEXT("[Mesh] Loading Indicies"));
 
 	LodRenderData->StaticVertexBuffers.ColorVertexBuffer.InitFromColorArray(ColorArray);
-	LodRenderData->SkinWeightVertexBuffer.SetHasExtraBoneInfluences(true);
+	LodRenderData->SkinWeightVertexBuffer.SetMaxBoneInfluences(EXTRA_BONE_INFLUENCES);
 	LodRenderData->SkinWeightVertexBuffer = InWeights;
 	LodRenderData->MultiSizeIndexContainer.CreateIndexBuffer(sizeof(uint16_t));
 
diff --git a/Engine/Plugins/Runtime/Steam/SteamVR/Source/SteamVRInputDevice/Private/AnimNode_SteamVRInputAnimPose.cpp b/Engine/Plugins/Runtime/Steam/SteamVR/Source/SteamVRInputDevice/Private/AnimNode_SteamVRInputAnimPose.cpp
--- a/Engine/Plugins/Runtime/Steam/SteamVR/Source/SteamVRInputDevice/Private/AnimNode_SteamVRInputAnimPose.cpp
+++ b/Engine/Plugins/Runtime/Steam/SteamVR/Source/SteamVRInputDevice/Private/AnimNode_SteamVRInputAnimPose.cpp
@@ -316,7 +316,7 @@ FSteamVRInputDevice* FAnimNode_SteamVRInputAnimPose::GetSteamVRInputDevice()
 	for (auto MotionController : MotionControllers)
 	{
 		FSteamVRInputDevice* TestSteamVRDevice = static_cast<FSteamVRInputDevice*>(MotionController);
-		if (TestSteamVRDevice != nullptr && !FGenericPlatformMath::IsNaN(TestSteamVRDevice->DeviceSignature) && TestSteamVRDevice->DeviceSignature == 2019)
+		if (TestSteamVRDevice != nullptr  && TestSteamVRDevice->DeviceSignature == 2019)
 		{
 			return TestSteamVRDevice;
 		}
diff --git a/Engine/Plugins/Runtime/Steam/SteamVR/Source/SteamVRInputDevice/Private/SteamVRInputDeviceFunctionLibrary.cpp b/Engine/Plugins/Runtime/Steam/SteamVR/Source/SteamVRInputDevice/Private/SteamVRInputDeviceFunctionLibrary.cpp
--- a/Engine/Plugins/Runtime/Steam/SteamVR/Source/SteamVRInputDevice/Private/SteamVRInputDeviceFunctionLibrary.cpp
+++ b/Engine/Plugins/Runtime/Steam/SteamVR/Source/SteamVRInputDevice/Private/SteamVRInputDeviceFunctionLibrary.cpp
@@ -851,7 +851,7 @@ FSteamVRInputDevice* USteamVRInputDeviceFunctionLibrary::GetSteamVRInputDevice()
 	for (auto MotionController : MotionControllers)
 	{
 		FSteamVRInputDevice* TestSteamVRDevice = static_cast<FSteamVRInputDevice*>(MotionController);
-		if (TestSteamVRDevice != nullptr && !FGenericPlatformMath::IsNaN(TestSteamVRDevice->DeviceSignature) && TestSteamVRDevice->DeviceSignature == 2019)
+		if (TestSteamVRDevice != nullptr && TestSteamVRDevice->DeviceSignature == 2019)
 		{
 			return TestSteamVRDevice;
 		}
diff --git a/Engine/Shaders/Private/GpuSkinCacheComputeShader.usf b/Engine/Shaders/Private/GpuSkinCacheComputeShader.usf
--- a/Engine/Shaders/Private/GpuSkinCacheComputeShader.usf
+++ b/Engine/Shaders/Private/GpuSkinCacheComputeShader.usf
@@ -74,6 +74,7 @@ struct FVertexUnpacked
 	float4	TangentX;
 	float4	TangentZ;
 
+#if !GPUSKIN_UNLIMITED_BONE_INFLUENCE
 	int4	BlendIndices;
 #if GPUSKIN_USE_EXTRA_INFLUENCES
 	int4	BlendIndices2;
@@ -82,6 +83,7 @@ struct FVertexUnpacked
 #if GPUSKIN_USE_EXTRA_INFLUENCES
 	float4	BlendWeights2;
 #endif
+#endif // GPUSKIN_UNLIMITED_BONE_INFLUENCE
 
 #if GPUSKIN_MORPH_BLEND
 	// morph target, added to the Position
@@ -118,8 +120,33 @@ struct FVertexFactoryIntermediates
 
 };
 
-FBoneMatrix CalcBoneMatrix( FVertexUnpacked VertexUnpacked )
+#if GPUSKIN_UNLIMITED_BONE_INFLUENCE
+uint InputWeightIndexSize; // bone index size: 1 - 8bit or 2 - 16bit
+Buffer<uint> InputWeightLookupStream;
+#endif
+Buffer<uint> InputWeightStream;
+
+FBoneMatrix CalcBoneMatrix( FVertexUnpacked VertexUnpacked, uint VertexIndex )
 {
+#if GPUSKIN_UNLIMITED_BONE_INFLUENCE
+	uint BlendOffsetCount = InputWeightLookupStream[VertexIndex];
+	int NumBoneInfluences = BlendOffsetCount & 0xff;
+	int StreamOffset = BlendOffsetCount >> 8;
+	int WeightsOffset = StreamOffset + (InputWeightIndexSize * NumBoneInfluences);
+
+	FBoneMatrix BoneMatrix = FBoneMatrix(float4(0,0,0,0), float4(0,0,0,0), float4(0,0,0,0));
+	for (int InfluenceIdx = 0; InfluenceIdx < NumBoneInfluences; InfluenceIdx++)
+	{
+		int BoneIndexOffset = StreamOffset + (InputWeightIndexSize * InfluenceIdx);
+		int BoneIndex = InputWeightStream[BoneIndexOffset];
+		if (InputWeightIndexSize > 1)
+		{
+			BoneIndex = InputWeightStream[BoneIndexOffset + 1] << 8 | BoneIndex;
+		}
+		float BoneWeight = float(InputWeightStream[WeightsOffset + InfluenceIdx]) / 255.0f;
+		BoneMatrix += BoneWeight * FBoneMatrix(BoneMatrices[BoneIndex*3], BoneMatrices[BoneIndex*3+1], BoneMatrices[BoneIndex*3+2]);
+	}
+#else // GPUSKIN_UNLIMITED_BONE_INFLUENCE
 	FBoneMatrix BoneMatrix = VertexUnpacked.BlendWeights.x * FBoneMatrix(BoneMatrices[VertexUnpacked.BlendIndices.x*3], BoneMatrices[VertexUnpacked.BlendIndices.x*3+1], BoneMatrices[VertexUnpacked.BlendIndices.x*3+2]);
 	BoneMatrix += VertexUnpacked.BlendWeights.y * FBoneMatrix(BoneMatrices[VertexUnpacked.BlendIndices.y*3], BoneMatrices[VertexUnpacked.BlendIndices.y*3+1], BoneMatrices[VertexUnpacked.BlendIndices.y*3+2]);
 	BoneMatrix += VertexUnpacked.BlendWeights.z * FBoneMatrix(BoneMatrices[VertexUnpacked.BlendIndices.z*3], BoneMatrices[VertexUnpacked.BlendIndices.z*3+1], BoneMatrices[VertexUnpacked.BlendIndices.z*3+2]);
@@ -130,6 +157,7 @@ FBoneMatrix CalcBoneMatrix( FVertexUnpacked VertexUnpacked )
 	BoneMatrix += VertexUnpacked.BlendWeights2.z * FBoneMatrix(BoneMatrices[VertexUnpacked.BlendIndices2.z*3], BoneMatrices[VertexUnpacked.BlendIndices2.z*3+1], BoneMatrices[VertexUnpacked.BlendIndices2.z*3+2]);
 	BoneMatrix += VertexUnpacked.BlendWeights2.w * FBoneMatrix(BoneMatrices[VertexUnpacked.BlendIndices2.w*3], BoneMatrices[VertexUnpacked.BlendIndices2.w*3+1], BoneMatrices[VertexUnpacked.BlendIndices2.w*3+2]);
 #endif
+#endif // GPUSKIN_UNLIMITED_BONE_INFLUENCE
 	return BoneMatrix;
 }
 
@@ -144,7 +172,6 @@ uint NumVertices;
 
 uint InputWeightStart;
 uint InputWeightStride;
-Buffer<uint> InputWeightStream;
 
 // see r.SkinCache.Debug, default is 1.0f
 float SkinCacheDebug;
@@ -156,7 +183,23 @@ FVertexUnpacked UnpackedVertex(uint MorphIndex, uint ClothIndex, uint VertexInde
 	Unpacked.TangentX = TangentBias_SkinCache(TangentInputBuffer[2 * VertexIndex + GPUSKIN_VB_OFFSET_TANGENT_X]);
 	Unpacked.TangentZ = TangentBias_SkinCache(TangentInputBuffer[2 * VertexIndex + GPUSKIN_VB_OFFSET_TANGENT_Z]);
 
+#if !GPUSKIN_UNLIMITED_BONE_INFLUENCE
 	uint BlendIndices = InputWeightStream[WeightOffset + GPUSKIN_VB_OFFSET_INFLUENCEBONES];
+#if GPUSKIN_BONE_INDEX_UINT16
+	Unpacked.BlendIndices.x = BlendIndices & 0xffff;
+	Unpacked.BlendIndices.y = BlendIndices >> 16 & 0xffff;
+	BlendIndices = InputWeightStream[WeightOffset + GPUSKIN_VB_OFFSET_INFLUENCEBONES + 1];
+	Unpacked.BlendIndices.z = BlendIndices & 0xffff;
+	Unpacked.BlendIndices.w = BlendIndices >> 16 & 0xffff;
+#if GPUSKIN_USE_EXTRA_INFLUENCES
+	BlendIndices = InputWeightStream[WeightOffset + GPUSKIN_VB_OFFSET_INFLUENCEBONES + 2];
+	Unpacked.BlendIndices2.x = BlendIndices & 0xffff;
+	Unpacked.BlendIndices2.y = BlendIndices >> 16 & 0xffff;
+	BlendIndices = InputWeightStream[WeightOffset + GPUSKIN_VB_OFFSET_INFLUENCEBONES + 3];
+	Unpacked.BlendIndices2.z = BlendIndices & 0xffff;
+	Unpacked.BlendIndices2.w = BlendIndices >> 16 & 0xffff;
+#endif
+#else // GPUSKIN_BONE_INDEX_UINT16
 	Unpacked.BlendIndices.x = BlendIndices & 0xff;
 	Unpacked.BlendIndices.y = BlendIndices >> 8 & 0xff;
 	Unpacked.BlendIndices.z = BlendIndices >> 16 & 0xff;
@@ -168,6 +211,7 @@ FVertexUnpacked UnpackedVertex(uint MorphIndex, uint ClothIndex, uint VertexInde
 	Unpacked.BlendIndices2.z = BlendIndices >> 16 & 0xff;
 	Unpacked.BlendIndices2.w = BlendIndices >> 24 & 0xff;
 #endif
+#endif // GPUSKIN_BONE_INDEX_UINT16
 	uint BlendWeights = InputWeightStream[WeightOffset + GPUSKIN_VB_OFFSET_INFLUENCEWEIGHTS];
 	Unpacked.BlendWeights.x = float(BlendWeights & 0xff) / 255.0f;
 	Unpacked.BlendWeights.y = float(BlendWeights >> 8 & 0xff) / 255.0f;
@@ -180,6 +224,8 @@ FVertexUnpacked UnpackedVertex(uint MorphIndex, uint ClothIndex, uint VertexInde
 	Unpacked.BlendWeights2.z = float(BlendWeights >> 16 & 0xff) / 255.0f;
 	Unpacked.BlendWeights2.w = float(BlendWeights >> 24 & 0xff) / 255.0f;
 #endif
+#endif // GPUSKIN_UNLIMITED_BONE_INFLUENCE
+
 #if GPUSKIN_MORPH_BLEND
 	GetMorphBuffer(MorphIndex, Unpacked.DeltaPosition, Unpacked.DeltaTangentZ);
 #endif
@@ -235,7 +281,7 @@ void SkinCacheUpdateBatchCS(uint3 GroupID : SV_GroupID,
 	FVertexFactoryIntermediates	Intermediates = (FVertexFactoryIntermediates)0;
 	
 #if 1
-	Intermediates.BlendMatrix = CalcBoneMatrix(Unpacked);
+	Intermediates.BlendMatrix = CalcBoneMatrix(Unpacked, InputOffset);
 #else
 	// Test no blend mtx
 	Intermediates.BlendMatrix[0] = float4(1,0,0,0);
diff --git a/Engine/Shaders/Private/GpuSkinCommon.ush b/Engine/Shaders/Private/GpuSkinCommon.ush
--- a/Engine/Shaders/Private/GpuSkinCommon.ush
+++ b/Engine/Shaders/Private/GpuSkinCommon.ush
@@ -16,21 +16,21 @@
 //	Position			2,3					2,3								2,3
 //	UVs[1..N]			N/A					5,1..N							5,(1..N)*2
 
-// TSkinWeightInfo<true>
-// TSkinWeightInfo<false>
-
-// Members of weight array
-//	MEMBER				TSkinWeightInfo<true>		TSkinWeightInfo<false>
-//	InfluenceBones		0,1							0,2	
-//	InfluenceWeights	1,1							2,2	
-
 #include "Common.ush"
 
 #define	GPUSKIN_VB_OFFSET_TANGENT_X			0
 #define	GPUSKIN_VB_OFFSET_TANGENT_Z			1
 
+#ifndef GPUSKIN_BONE_INDEX_UINT16
+#define GPUSKIN_BONE_INDEX_UINT16			0	// 16-bit bone index if 1, otherwise 8-bit bone index
+#endif
+
+#ifndef GPUSKIN_UNLIMITED_BONE_INFLUENCE
+#define GPUSKIN_UNLIMITED_BONE_INFLUENCE	0
+#endif
+
 #define	GPUSKIN_VB_OFFSET_INFLUENCEBONES	0
-#define	GPUSKIN_VB_OFFSET_INFLUENCEWEIGHTS	(1 + GPUSKIN_USE_EXTRA_INFLUENCES)
+#define	GPUSKIN_VB_OFFSET_INFLUENCEWEIGHTS	((1 + GPUSKIN_USE_EXTRA_INFLUENCES) * (1 + GPUSKIN_BONE_INDEX_UINT16))
 
 
 #ifndef GPUSKIN_RWBUFFER_OFFSET_TANGENT_X
