Control Rig: Crash on Opening New ControlRig Asset

We were missing the call to CreateDefaultCommands

diff --git a/Engine/Plugins/Experimental/ControlRig/Source/ControlRigEditor/Private/Editor/ControlRigEditor.cpp b/Engine/Plugins/Experimental/ControlRig/Source/ControlRigEditor/Private/Editor/ControlRigEditor.cpp
--- a/Engine/Plugins/Experimental/ControlRig/Source/ControlRigEditor/Private/Editor/ControlRigEditor.cpp
+++ b/Engine/Plugins/Experimental/ControlRig/Source/ControlRigEditor/Private/Editor/ControlRigEditor.cpp
@@ -185,6 +185,8 @@ void FControlRigEditor::InitControlRigEditor(const EToolkitMode::Type Mode, cons
 	const bool bCreateDefaultToolbar = true;
 	InitAssetEditor(Mode, InitToolkitHost, ControlRigEditorAppName, DummyLayout, bCreateDefaultStandaloneMenu, bCreateDefaultToolbar, ObjectsBeingEdited);
 
+	CreateDefaultCommands();
+
 	UControlRigBlueprint::sCurrentlyOpenedRigBlueprints.AddUnique(InControlRigBlueprint);
 
 	TArray<UBlueprint*> ControlRigBlueprints;
