Copying //UE4/Dev-Anim to Dev-Main (//UE4/Dev-Main)

diff --git a/Engine/Plugins/Animation/LiveLink/Source/LiveLink/Private/LiveLinkClient.cpp b/Engine/Plugins/Animation/LiveLink/Source/LiveLink/Private/LiveLinkClient.cpp
--- a/Engine/Plugins/Animation/LiveLink/Source/LiveLink/Private/LiveLinkClient.cpp
+++ b/Engine/Plugins/Animation/LiveLink/Source/LiveLink/Private/LiveLinkClient.cpp
@@ -1041,6 +1041,12 @@ bool FLiveLinkClient::EvaluateFrameFromSource_AnyThread(const FLiveLinkSubjectKe
 	return false;
 }
 
+//just call our tick
+void FLiveLinkClient::ForceTick()
+{
+	Tick();
+}
+
 bool FLiveLinkClient::EvaluateFrame_AnyThread(FLiveLinkSubjectName InSubjectName, TSubclassOf<ULiveLinkRole> InDesiredRole, FLiveLinkSubjectFrameData& OutFrame)
 {
 	SCOPE_CYCLE_COUNTER(STAT_LiveLink_EvaluateFrame);
@@ -1257,14 +1263,11 @@ void FLiveLinkClient::OnPropertyChanged(FGuid InEntryGuid, const FPropertyChange
 	}
 }
 
-ULiveLinkSourceSettings* FLiveLinkClient::GetSourceSettings(FGuid InEntryGuid) const
+ULiveLinkSourceSettings* FLiveLinkClient::GetSourceSettings(const FGuid& InEntryGuid) const
 {
 	if (const FLiveLinkCollectionSourceItem* SourceItem = Collection->FindSource(InEntryGuid))
 	{
-		if (!SourceItem->IsVirtualSource())
-		{
-			return SourceItem->Setting;
-		}
+		return SourceItem->Setting;
 	}
 	return nullptr;
 }
diff --git a/Engine/Plugins/Animation/LiveLink/Source/LiveLink/Public/LiveLinkClient.h b/Engine/Plugins/Animation/LiveLink/Source/LiveLink/Public/LiveLinkClient.h
--- a/Engine/Plugins/Animation/LiveLink/Source/LiveLink/Public/LiveLinkClient.h
+++ b/Engine/Plugins/Animation/LiveLink/Source/LiveLink/Public/LiveLinkClient.h
@@ -21,7 +21,6 @@ class ULiveLinkSourceSettings;
 class ULiveLinkSubjectBase;
 class FLiveLinkSourceCollection;
 
-
 // Live Link Log Category
 DECLARE_LOG_CATEGORY_EXTERN(LogLiveLink, Log, All);
 
@@ -112,11 +111,14 @@ public:
 	virtual TArray<FLiveLinkTime> GetSubjectFrameTimes(const FLiveLinkSubjectKey& SubjectKey) const override;
 	virtual TArray<FLiveLinkTime> GetSubjectFrameTimes(FLiveLinkSubjectName SubjectName) const override;
 	virtual UObject* GetSubjectSettings(const FLiveLinkSubjectKey& SubjectKey) const override;
+	virtual ULiveLinkSourceSettings* GetSourceSettings(const FGuid& SourceGuid) const override;
+
 
 	virtual bool EvaluateFrameFromSource_AnyThread(const FLiveLinkSubjectKey& SubjectKey, TSubclassOf<ULiveLinkRole> Role, FLiveLinkSubjectFrameData& OutFrame) override;
 	virtual bool EvaluateFrame_AnyThread(FLiveLinkSubjectName SubjectName, TSubclassOf<ULiveLinkRole> Role, FLiveLinkSubjectFrameData& OutFrame) override;
 	virtual bool EvaluateFrameAtWorldTime_AnyThread(FLiveLinkSubjectName SubjectName, double WorldTime, TSubclassOf<ULiveLinkRole> DesiredRole, FLiveLinkSubjectFrameData& OutFrame) override;
 	virtual bool EvaluateFrameAtSceneTime_AnyThread(FLiveLinkSubjectName SubjectName, const FTimecode& SceneTime, TSubclassOf<ULiveLinkRole> DesiredRole, FLiveLinkSubjectFrameData& OutFrame) override;
+	virtual void ForceTick() override;
 
 	virtual FSimpleMulticastDelegate& OnLiveLinkTicked() override;
 	virtual FSimpleMulticastDelegate& OnLiveLinkSourcesChanged() override;
@@ -147,8 +149,6 @@ public:
 
 	TArray<FGuid> GetDisplayableSources() const;
 
-	ULiveLinkSourceSettings* GetSourceSettings(FGuid InEntryGuid) const;
-
 	FLiveLinkSubjectTimeSyncData GetTimeSyncData(FLiveLinkSubjectName SubjectName);
 
 	FText GetSourceMachineName(FGuid EntryGuid) const;
diff --git a/Engine/Plugins/Editor/CurveEditorTools/Source/CurveEditorTools/Private/Filters/CurveEditorFFTFilter.cpp b/Engine/Plugins/Editor/CurveEditorTools/Source/CurveEditorTools/Private/Filters/CurveEditorFFTFilter.cpp
--- a/Engine/Plugins/Editor/CurveEditorTools/Source/CurveEditorTools/Private/Filters/CurveEditorFFTFilter.cpp
+++ b/Engine/Plugins/Editor/CurveEditorTools/Source/CurveEditorTools/Private/Filters/CurveEditorFFTFilter.cpp
@@ -72,12 +72,33 @@ void UCurveEditorFFTFilter::ApplyFilter_Impl(TSharedRef<FCurveEditor> InCurveEdi
 		// Apply the Bake filter to the same segment of the curve as this. This invalidates the original KeyHandles, because it creates and destroys keys.
 		TArray<FKeyHandle> BakedKeyHandles;
 		BakeFilter->ApplyFilter(InCurveEditor, Pair.Key, MakeArrayView(OriginalKeyHandles.GetData(), OriginalKeyHandles.Num()), BakedKeyHandles);
-
 		// Get the positions from the baked curve to feed into the FFT samples.
 		TArray<FKeyPosition> BakedKeyPositions;
-		BakedKeyPositions.SetNumUninitialized(BakedKeyHandles.Num());
+		BakedKeyPositions.SetNumUninitialized(BakedKeyHandles.Num());	
 		CurveModel->GetKeyPositions(BakedKeyHandles, BakedKeyPositions);
-		
+
+		//Bake Filter returns handles unsorted so need to sort. (Second key is last key but they may change so do it manually).
+		TArray<TPair < FKeyHandle, FKeyPosition>> KeyAndPosition;
+		KeyAndPosition.SetNumUninitialized(BakedKeyHandles.Num());
+		int32 Index = 0;
+		for (; Index < BakedKeyHandles.Num(); ++Index)
+		{
+			KeyAndPosition[Index] = TPair<FKeyHandle, FKeyPosition>(BakedKeyHandles[Index], BakedKeyPositions[Index]);
+		}
+
+		Algo::Sort(KeyAndPosition, [](const TPair < FKeyHandle, FKeyPosition> LHS, const TPair < FKeyHandle, FKeyPosition> RHS)
+		{
+			return LHS.Value.InputValue < RHS.Value.InputValue;
+		});
+
+		for (Index = 0; Index < KeyAndPosition.Num(); ++Index)
+		{
+			BakedKeyHandles[Index] = KeyAndPosition[Index].Key;
+			BakedKeyPositions[Index] = KeyAndPosition[Index].Value;
+		}
+
+		//now it's sorted we can just use them.
+
 		TArray<float> FFTSamples;
 		FFTSamples.Reserve(BakedKeyPositions.Num());
 
@@ -112,7 +133,10 @@ void UCurveEditorFFTFilter::ApplyFilter_Impl(TSharedRef<FCurveEditor> InCurveEdi
 		// of handles, as FFT has padded our array (with zeros) up to the next POT.
 		for (int32 KeyHandleIndex = 0; KeyHandleIndex < BakedKeyHandles.Num(); KeyHandleIndex++)
 		{
-			BakedKeyPositions[KeyHandleIndex].OutputValue = FFTSamples[KeyHandleIndex] + ValueOffset;
+			if (!FMath::IsNaN(FFTSamples[KeyHandleIndex]))
+			{
+				BakedKeyPositions[KeyHandleIndex].OutputValue = FFTSamples[KeyHandleIndex] + ValueOffset;
+			}
 		}
 
 		// We re-assign the filtered samples back to our baked curve so that we can then sample it at the original times.
diff --git a/Engine/Plugins/Experimental/ControlRig/Source/ControlRig/Private/AnimNode_ControlRig.cpp b/Engine/Plugins/Experimental/ControlRig/Source/ControlRig/Private/AnimNode_ControlRig.cpp
--- a/Engine/Plugins/Experimental/ControlRig/Source/ControlRig/Private/AnimNode_ControlRig.cpp
+++ b/Engine/Plugins/Experimental/ControlRig/Source/ControlRig/Private/AnimNode_ControlRig.cpp
@@ -68,7 +68,6 @@ void FAnimNode_ControlRig::Update_AnyThread(const FAnimationUpdateContext& Conte
 {
 	DECLARE_SCOPE_HIERARCHICAL_COUNTER_FUNC()
 
-	FAnimNode_ControlRigBase::Update_AnyThread(Context);
 	GetEvaluateGraphExposedInputs().Execute(Context);
 
 	// alpha handlers
@@ -93,7 +92,8 @@ void FAnimNode_ControlRig::Update_AnyThread(const FAnimationUpdateContext& Conte
 	InternalBlendAlpha = FMath::Clamp<float>(InternalBlendAlpha, 0.f, 1.f);
 
 	PropagateInputProperties(Context.AnimInstanceProxy->GetAnimInstanceObject());
-	Source.Update(Context);
+
+	FAnimNode_ControlRigBase::Update_AnyThread(Context);
 
 	TRACE_ANIM_NODE_VALUE(Context, TEXT("Class"), *GetNameSafe(ControlRigClass.Get()));
 }
@@ -104,8 +104,6 @@ void FAnimNode_ControlRig::Initialize_AnyThread(const FAnimationInitializeContex
 
 	FAnimNode_ControlRigBase::Initialize_AnyThread(Context);
 
-	Source.Initialize(Context);
-
 	AlphaBoolBlend.Reinitialize();
 	AlphaScaleBiasClamp.Reinitialize();
 }
@@ -115,7 +113,6 @@ void FAnimNode_ControlRig::CacheBones_AnyThread(const FAnimationCacheBonesContex
 	DECLARE_SCOPE_HIERARCHICAL_COUNTER_FUNC()
 
 	FAnimNode_ControlRigBase::CacheBones_AnyThread(Context);
-	Source.CacheBones(Context);
 
 	FBoneContainer& RequiredBones = Context.AnimInstanceProxy->GetRequiredBones();
 	InputToCurveMappingUIDs.Reset();
@@ -157,9 +154,6 @@ void FAnimNode_ControlRig::Evaluate_AnyThread(FPoseContext & Output)
 {
 	DECLARE_SCOPE_HIERARCHICAL_COUNTER_FUNC()
 
-	// If not playing a montage, just pass through
-	Source.Evaluate(Output);
-
 	// evaluate 
 	FAnimNode_ControlRigBase::Evaluate_AnyThread(Output);
 }
diff --git a/Engine/Plugins/Experimental/ControlRig/Source/ControlRig/Private/AnimNode_ControlRigBase.cpp b/Engine/Plugins/Experimental/ControlRig/Source/ControlRig/Private/AnimNode_ControlRigBase.cpp
--- a/Engine/Plugins/Experimental/ControlRig/Source/ControlRig/Private/AnimNode_ControlRigBase.cpp
+++ b/Engine/Plugins/Experimental/ControlRig/Source/ControlRig/Private/AnimNode_ControlRigBase.cpp
@@ -50,6 +50,7 @@ void FAnimNode_ControlRigBase::Initialize_AnyThread(const FAnimationInitializeCo
 	{
 		//Don't Inititialize the Control Rig here it may have the wrong VM on the CDO
 		SetTargetInstance(ControlRig);
+		ControlRig->RequestInit();
 	}
 }
 
diff --git a/Engine/Plugins/Experimental/ControlRig/Source/ControlRig/Private/Rigs/AdditiveControlRig.cpp b/Engine/Plugins/Experimental/ControlRig/Source/ControlRig/Private/Rigs/AdditiveControlRig.cpp
--- a/Engine/Plugins/Experimental/ControlRig/Source/ControlRig/Private/Rigs/AdditiveControlRig.cpp
+++ b/Engine/Plugins/Experimental/ControlRig/Source/ControlRig/Private/Rigs/AdditiveControlRig.cpp
@@ -97,13 +97,30 @@ void UAdditiveControlRig::CreateRigElements(const FReferenceSkeleton& InReferenc
 	// add control for all bone hierarchy 
 	for (int32 BoneIndex = 0; BoneIndex < BoneHierarchy.Num(); ++BoneIndex)
 	{
+
 		const FRigBone& RigBone = BoneHierarchy[BoneIndex];
 		FName BoneName = RigBone.Name;
 		FName ParentName = RigBone.ParentName;
 		FName SpaceName = GetSpaceName(BoneName);// name conflict?
 		FName ControlName = GetControlName(BoneName); // name conflict?
-		Container->SpaceHierarchy.Add(SpaceName, ERigSpaceType::Bone, BoneName);
-		Container->ControlHierarchy.Add(ControlName, ERigControlType::Transform, GetControlName(ParentName), SpaceName);
+		if (ParentName != NAME_None)
+		{
+			FTransform Transform = BoneHierarchy.GetGlobalTransform(BoneName);
+			FTransform ParentTransform = BoneHierarchy.GetGlobalTransform(ParentName);
+			FTransform LocalTransform = Transform.GetRelativeTransform(ParentTransform);
+			FRigSpace& Space = Container->SpaceHierarchy.Add(SpaceName, ERigSpaceType::Bone, ParentName);
+			Space.InitialTransform = LocalTransform;
+		}
+		else
+		{
+			FTransform Transform = BoneHierarchy.GetGlobalTransform(BoneName);
+			FTransform ParentTransform = FTransform::Identity;
+			FTransform LocalTransform = Transform.GetRelativeTransform(ParentTransform);
+			FRigSpace& Space = Container->SpaceHierarchy.Add(SpaceName, ERigSpaceType::Global, ParentName);
+			Space.InitialTransform = LocalTransform;
+		}
+		Container->ControlHierarchy.Add(ControlName, ERigControlType::Transform, NAME_None, SpaceName);
+
 	}
 
 	Container->Initialize(true);
diff --git a/Engine/Plugins/Experimental/ControlRig/Source/ControlRig/Private/Rigs/RigBoneHierarchy.cpp b/Engine/Plugins/Experimental/ControlRig/Source/ControlRig/Private/Rigs/RigBoneHierarchy.cpp
--- a/Engine/Plugins/Experimental/ControlRig/Source/ControlRig/Private/Rigs/RigBoneHierarchy.cpp
+++ b/Engine/Plugins/Experimental/ControlRig/Source/ControlRig/Private/Rigs/RigBoneHierarchy.cpp
@@ -522,6 +522,15 @@ void FRigBoneHierarchy::ResetTransforms()
 	}
 }
 
+void FRigBoneHierarchy::RecomputeGlobalTransforms()
+{
+	for (int32 Index = 0; Index < Bones.Num(); ++Index)
+	{
+		Bones[Index].GlobalTransform = Bones[Index].InitialTransform;
+		RecalculateGlobalTransform(Bones[Index]);
+	}
+}
+
 int32 FRigBoneHierarchy::GetChildrenRecursive(const int32 InIndex, TArray<int32>& OutChildren, bool bRecursively) const
 {
 	const int32 StartChildIndex = OutChildren.Num();
@@ -575,13 +584,6 @@ bool FRigBoneHierarchy::Select(const FName& InName, bool bSelect)
 
 	if(bSelect)
 	{
-		if (Container)
-		{
-			Container->SpaceHierarchy.ClearSelection();
-			Container->ControlHierarchy.ClearSelection();
-			Container->CurveContainer.ClearSelection();
-		}
-
 		Selection.Add(InName);
 	}
 	else
diff --git a/Engine/Plugins/Experimental/ControlRig/Source/ControlRig/Private/Rigs/RigControlHierarchy.cpp b/Engine/Plugins/Experimental/ControlRig/Source/ControlRig/Private/Rigs/RigControlHierarchy.cpp
--- a/Engine/Plugins/Experimental/ControlRig/Source/ControlRig/Private/Rigs/RigControlHierarchy.cpp
+++ b/Engine/Plugins/Experimental/ControlRig/Source/ControlRig/Private/Rigs/RigControlHierarchy.cpp
@@ -1014,13 +1014,6 @@ bool FRigControlHierarchy::Select(const FName& InName, bool bSelect)
 
 	if(bSelect)
 	{
-		if (Container)
-		{
-			Container->BoneHierarchy.ClearSelection();
-			Container->SpaceHierarchy.ClearSelection();
-			Container->CurveContainer.ClearSelection();
-		}
-
 		Selection.Add(InName);
 	}
 	else
diff --git a/Engine/Plugins/Experimental/ControlRig/Source/ControlRig/Private/Rigs/RigCurveContainer.cpp b/Engine/Plugins/Experimental/ControlRig/Source/ControlRig/Private/Rigs/RigCurveContainer.cpp
--- a/Engine/Plugins/Experimental/ControlRig/Source/ControlRig/Private/Rigs/RigCurveContainer.cpp
+++ b/Engine/Plugins/Experimental/ControlRig/Source/ControlRig/Private/Rigs/RigCurveContainer.cpp
@@ -249,13 +249,6 @@ bool FRigCurveContainer::Select(const FName& InName, bool bSelect)
 
 	if(bSelect)
 	{
-		if (Container)
-		{
-			Container->BoneHierarchy.ClearSelection();
-			Container->SpaceHierarchy.ClearSelection();
-			Container->ControlHierarchy.ClearSelection();
-		}
-
 		Selection.Add(InName);
 	}
 	else
diff --git a/Engine/Plugins/Experimental/ControlRig/Source/ControlRig/Private/Rigs/RigSpaceHierarchy.cpp b/Engine/Plugins/Experimental/ControlRig/Source/ControlRig/Private/Rigs/RigSpaceHierarchy.cpp
--- a/Engine/Plugins/Experimental/ControlRig/Source/ControlRig/Private/Rigs/RigSpaceHierarchy.cpp
+++ b/Engine/Plugins/Experimental/ControlRig/Source/ControlRig/Private/Rigs/RigSpaceHierarchy.cpp
@@ -504,13 +504,6 @@ bool FRigSpaceHierarchy::Select(const FName& InName, bool bSelect)
 
 	if(bSelect)
 	{
-		if (Container)
-		{
-			Container->BoneHierarchy.ClearSelection();
-			Container->ControlHierarchy.ClearSelection();
-			Container->CurveContainer.ClearSelection();
-		}
-
 		Selection.Add(InName);
 	}
 	else
diff --git a/Engine/Plugins/Experimental/ControlRig/Source/ControlRig/Private/Units/Hierarchy/RigUnit_SetControlTransform.h b/Engine/Plugins/Experimental/ControlRig/Source/ControlRig/Private/Units/Hierarchy/RigUnit_SetControlTransform.h
--- a/Engine/Plugins/Experimental/ControlRig/Source/ControlRig/Private/Units/Hierarchy/RigUnit_SetControlTransform.h
+++ b/Engine/Plugins/Experimental/ControlRig/Source/ControlRig/Private/Units/Hierarchy/RigUnit_SetControlTransform.h
@@ -6,7 +6,7 @@
 #include "RigUnit_SetControlTransform.generated.h"
 
 /**
- * SetControlBool is used to perform a change in the hierarchy by setting a single bone's Bool.
+ * SetControlBool is used to perform a change in the hierarchy by setting a single control's bool value.
  */
 USTRUCT(meta=(DisplayName="Set Control Bool", Category="Hierarchy", DocumentationPolicy="Strict", Keywords = "SetControlBool", PrototypeName = "SetControlValue"))
 struct FRigUnit_SetControlBool : public FRigUnitMutable
@@ -39,7 +39,7 @@ struct FRigUnit_SetControlBool : public FRigUnitMutable
 };
 
 /**
- * SetControlFloat is used to perform a change in the hierarchy by setting a single bone's Float.
+ * SetControlFloat is used to perform a change in the hierarchy by setting a single control's float value.
  */
 USTRUCT(meta=(DisplayName="Set Control Float", Category="Hierarchy", DocumentationPolicy="Strict", Keywords = "SetControlFloat", PrototypeName = "SetControlValue"))
 struct FRigUnit_SetControlFloat : public FRigUnitMutable
@@ -79,7 +79,7 @@ struct FRigUnit_SetControlFloat : public FRigUnitMutable
 };
 
 /**
- * SetControlVector2D is used to perform a change in the hierarchy by setting a single bone's Vector2D.
+ * SetControlVector2D is used to perform a change in the hierarchy by setting a single control's Vector2D value.
  */
 USTRUCT(meta=(DisplayName="Set Control Vector2D", Category="Hierarchy", DocumentationPolicy="Strict", Keywords = "SetControlVector2D", PrototypeName = "SetControlValue"))
 struct FRigUnit_SetControlVector2D : public FRigUnitMutable
@@ -119,7 +119,7 @@ struct FRigUnit_SetControlVector2D : public FRigUnitMutable
 };
 
 /**
- * SetControlVector is used to perform a change in the hierarchy by setting a single bone's Vector.
+ * SetControlVector is used to perform a change in the hierarchy by setting a single control's Vector value.
  */
 USTRUCT(meta=(DisplayName="Set Control Vector", Category="Hierarchy", DocumentationPolicy="Strict", Keywords = "SetControlVector", PrototypeName = "SetControlValue"))
 struct FRigUnit_SetControlVector : public FRigUnitMutable
@@ -167,7 +167,7 @@ struct FRigUnit_SetControlVector : public FRigUnitMutable
 };
 
 /**
- * SetControlRotator is used to perform a change in the hierarchy by setting a single bone's Rotator.
+ * SetControlRotator is used to perform a change in the hierarchy by setting a single control's Rotator value.
  */
 USTRUCT(meta=(DisplayName="Set Control Rotator", Category="Hierarchy", DocumentationPolicy="Strict", Keywords = "SetControlRotator", PrototypeName = "SetControlValue"))
 struct FRigUnit_SetControlRotator : public FRigUnitMutable
diff --git a/Engine/Plugins/Experimental/ControlRig/Source/ControlRig/Public/Rigs/RigBoneHierarchy.h b/Engine/Plugins/Experimental/ControlRig/Source/ControlRig/Public/Rigs/RigBoneHierarchy.h
--- a/Engine/Plugins/Experimental/ControlRig/Source/ControlRig/Public/Rigs/RigBoneHierarchy.h
+++ b/Engine/Plugins/Experimental/ControlRig/Source/ControlRig/Public/Rigs/RigBoneHierarchy.h
@@ -160,6 +160,9 @@ struct CONTROLRIG_API FRigBoneHierarchy
 	// resets all of the transforms back to the initial transform
 	void ResetTransforms();
 
+	// recomputes all of the global transforms from local
+	void RecomputeGlobalTransforms();
+
 	// import skeleton
 	TArray<FRigElementKey> ImportSkeleton(const FReferenceSkeleton& InSkeleton, const FName& InNameSpace, bool bReplaceExistingBones, bool bRemoveObsoleteBones, bool bSelectBones, bool bNotify);
 
diff --git a/Engine/Plugins/Experimental/ControlRig/Source/ControlRig/Public/Rigs/RigControlHierarchy.h b/Engine/Plugins/Experimental/ControlRig/Source/ControlRig/Public/Rigs/RigControlHierarchy.h
--- a/Engine/Plugins/Experimental/ControlRig/Source/ControlRig/Public/Rigs/RigControlHierarchy.h
+++ b/Engine/Plugins/Experimental/ControlRig/Source/ControlRig/Public/Rigs/RigControlHierarchy.h
@@ -140,7 +140,7 @@ struct CONTROLRIG_API FRigControl : public FRigElement
 		, bLimitTranslation(false)
 		, bLimitRotation(false)
 		, bLimitScale(false)
-		, bDrawLimits(false)
+		, bDrawLimits(true)
 		, MinimumValue()
 		, MaximumValue()
 		, bGizmoEnabled(true)
@@ -304,6 +304,30 @@ struct CONTROLRIG_API FRigControl : public FRigElement
 		}
 		return FMath::Clamp<float>(Value, Maximum, Minimum);
 	}
+
+	FORCEINLINE_DEBUGGABLE static FProperty* FindPropertyForValueType(ERigControlValueType InValueType)
+	{
+		switch (InValueType)
+		{
+			case ERigControlValueType::Current:
+			{
+				return FRigControl::StaticStruct()->FindPropertyByName(TEXT("Value"));
+			}
+			case ERigControlValueType::Initial:
+			{
+				return FRigControl::StaticStruct()->FindPropertyByName(TEXT("InitialValue"));
+			}
+			case ERigControlValueType::Minimum:
+			{
+				return FRigControl::StaticStruct()->FindPropertyByName(TEXT("MinimumValue"));
+			}
+			case ERigControlValueType::Maximum:
+			{
+				return FRigControl::StaticStruct()->FindPropertyByName(TEXT("MaximumValue"));
+			}
+		}
+		return nullptr;
+	}
 };
 
 USTRUCT(BlueprintType)
diff --git a/Engine/Plugins/Experimental/ControlRig/Source/ControlRig/Public/Rigs/RigHierarchyDefines.h b/Engine/Plugins/Experimental/ControlRig/Source/ControlRig/Public/Rigs/RigHierarchyDefines.h
--- a/Engine/Plugins/Experimental/ControlRig/Source/ControlRig/Public/Rigs/RigHierarchyDefines.h
+++ b/Engine/Plugins/Experimental/ControlRig/Source/ControlRig/Public/Rigs/RigHierarchyDefines.h
@@ -14,12 +14,12 @@ struct FRigHierarchyContainer;
 UENUM(BlueprintType)
 enum class ERigElementType : uint8
 {
-	None UMETA(Hidden),
+	None,
 	Bone = 0x001,
 	Space = 0x002,
 	Control = 0x004,
 	Curve = 0x008,
-	All = Bone | Space | Control | Curve UMETA(Hidden),
+	All = Bone | Space | Control | Curve,
 };
 
 /*
diff --git a/Engine/Plugins/Experimental/ControlRig/Source/ControlRigDeveloper/Private/ControlRigBlueprint.cpp b/Engine/Plugins/Experimental/ControlRig/Source/ControlRigDeveloper/Private/ControlRigBlueprint.cpp
--- a/Engine/Plugins/Experimental/ControlRig/Source/ControlRigDeveloper/Private/ControlRigBlueprint.cpp
+++ b/Engine/Plugins/Experimental/ControlRig/Source/ControlRigDeveloper/Private/ControlRigBlueprint.cpp
@@ -943,6 +943,10 @@ void UControlRigBlueprint::HandleModifiedEvent(ERigVMGraphNotifType InNotifType,
 				ClearTransientControls();
 				RequestAutoVMRecompilation();
 				MarkPackageDirty();
+
+				// This is not necessarily required but due to workflow
+				// expectations we still mark the blueprint as dirty.
+				FBlueprintEditorUtils::MarkBlueprintAsModified(this);
 				break;
 			}
 			case ERigVMGraphNotifType::PinWatchedChanged:
diff --git a/Engine/Plugins/Experimental/ControlRig/Source/ControlRigDeveloper/Private/Graph/ControlRigGraph.cpp b/Engine/Plugins/Experimental/ControlRig/Source/ControlRigDeveloper/Private/Graph/ControlRigGraph.cpp
--- a/Engine/Plugins/Experimental/ControlRig/Source/ControlRigDeveloper/Private/Graph/ControlRigGraph.cpp
+++ b/Engine/Plugins/Experimental/ControlRig/Source/ControlRigDeveloper/Private/Graph/ControlRigGraph.cpp
@@ -178,7 +178,6 @@ void UControlRigGraph::HandleModifiedEvent(ERigVMGraphNotifType InNotifType, URi
 				}
 				else if (URigVMRerouteNode* RerouteModelNode = Cast<URigVMRerouteNode>(ModelNode))
 				{
-					// todo: collapsed or not?
 					UControlRigGraphNode* NewNode = NewObject<UControlRigGraphNode>(this, ModelNode->GetFName());
 					AddNode(NewNode, false);
 
@@ -244,7 +243,8 @@ void UControlRigGraph::HandleModifiedEvent(ERigVMGraphNotifType InNotifType, URi
 				UEdGraphNode* EdNode = FindNodeForModelNodeName(ModelNode->GetFName());
 				if (EdNode)
 				{
-					RemoveNode(EdNode);
+					RemoveNode(EdNode, true);
+					NotifyGraphChanged();
 				}
 			}
 			break;
@@ -275,6 +275,27 @@ void UControlRigGraph::HandleModifiedEvent(ERigVMGraphNotifType InNotifType, URi
 			}
 			break;
 		}
+		case ERigVMGraphNotifType::RerouteCompactnessChanged:
+		{
+			if (URigVMRerouteNode* ModelNode = Cast<URigVMRerouteNode>(InSubject))
+			{
+				UEdGraphNode* EdNode = Cast<UEdGraphNode>(FindNodeForModelNodeName(ModelNode->GetFName()));
+				if (EdNode)
+				{
+					if (UControlRigGraphNode* RigNode = Cast<UControlRigGraphNode>(EdNode))
+					{
+						// start at index 2 (the subpins below the top level value pin)
+						// and hide the pins (or show them if they were hidden previously)
+						for (int32 PinIndex = 2; PinIndex < RigNode->Pins.Num(); PinIndex++)
+						{
+							RigNode->Pins[PinIndex]->bHidden = !ModelNode->GetShowsAsFullNode();
+						}
+						NotifyGraphChanged();
+					}
+				}
+			}
+			break;
+		}
 		case ERigVMGraphNotifType::NodeColorChanged:
 		{
 			if (URigVMNode* ModelNode = Cast<URigVMNode>(InSubject))
@@ -335,6 +356,29 @@ void UControlRigGraph::HandleModifiedEvent(ERigVMGraphNotifType InNotifType, URi
 								SourceRigPin->BreakLinkTo(TargetRigPin);
 							}
 
+							for (int32 LinkedPinIndex = 0; LinkedPinIndex < SourceRigPin->LinkedTo.Num();)
+							{
+								if (SourceRigPin->LinkedTo[LinkedPinIndex])
+								{
+									++LinkedPinIndex;
+								}
+								else
+								{
+									SourceRigPin->LinkedTo.RemoveAtSwap(LinkedPinIndex);
+								}
+							}
+
+							for (int32 LinkedPinIndex = 0; LinkedPinIndex < TargetRigPin->LinkedTo.Num();)
+							{
+								if (TargetRigPin->LinkedTo[LinkedPinIndex])
+								{
+									++LinkedPinIndex;
+								}
+								else
+								{
+									TargetRigPin->LinkedTo.RemoveAtSwap(LinkedPinIndex);
+								}
+							}
 						}
 					}
 				}
diff --git a/Engine/Plugins/Experimental/ControlRig/Source/ControlRigDeveloper/Private/Graph/ControlRigGraphNode.cpp b/Engine/Plugins/Experimental/ControlRig/Source/ControlRigDeveloper/Private/Graph/ControlRigGraphNode.cpp
--- a/Engine/Plugins/Experimental/ControlRig/Source/ControlRigDeveloper/Private/Graph/ControlRigGraphNode.cpp
+++ b/Engine/Plugins/Experimental/ControlRig/Source/ControlRigDeveloper/Private/Graph/ControlRigGraphNode.cpp
@@ -327,7 +327,6 @@ void UControlRigGraphNode::CreateExecutionPins()
 			if (Pair.InputPin != nullptr)
 			{
 				Pair.InputPin->PinFriendlyName = FText::FromName(ModelPin->GetDisplayName());
-				Pair.InputPin->bDisplayAsMutableRef = ModelPin->GetInjectedNodes().Num() > 0;
 			}
 		}
 		if (Pair.OutputPin == nullptr)
@@ -336,7 +335,6 @@ void UControlRigGraphNode::CreateExecutionPins()
 			if (Pair.OutputPin != nullptr)
 			{
 				Pair.OutputPin->PinFriendlyName = FText::FromName(ModelPin->GetDisplayName());
-				Pair.OutputPin->bDisplayAsMutableRef = ModelPin->GetInjectedNodes().Num() > 0;
 			}
 		}
 		// note: no recursion for execution pins
@@ -358,7 +356,6 @@ void UControlRigGraphNode::CreateInputPins(URigVMPin* InParentPin)
 			{
 				Pair.InputPin->PinFriendlyName = FText::FromName(ModelPin->GetDisplayName());
 				Pair.InputPin->bNotConnectable = ModelPin->GetDirection() != ERigVMPinDirection::Input;
-				Pair.InputPin->bDisplayAsMutableRef = ModelPin->GetInjectedNodes().Num() > 0;
 
 				SetupPinDefaultsFromModel(Pair.InputPin);
 
@@ -374,7 +371,7 @@ void UControlRigGraphNode::CreateInputPins(URigVMPin* InParentPin)
 	}
 }
 
-void UControlRigGraphNode::CreateInputOutputPins(URigVMPin* InParentPin)
+void UControlRigGraphNode::CreateInputOutputPins(URigVMPin* InParentPin, bool bHidden)
 {
 	DECLARE_SCOPE_HIERARCHICAL_COUNTER_FUNC()
 
@@ -387,9 +384,9 @@ void UControlRigGraphNode::CreateInputOutputPins(URigVMPin* InParentPin)
 			Pair.InputPin = CreatePin(EGPD_Input, GetPinTypeForModelPin(ModelPin), FName(*ModelPin->GetPinPath()));
 			if (Pair.InputPin != nullptr)
 			{
+				Pair.InputPin->bHidden = bHidden;
 				Pair.InputPin->PinFriendlyName = FText::FromName(ModelPin->GetDisplayName());
 				Pair.InputPin->bNotConnectable = ModelPin->GetDirection() != ERigVMPinDirection::IO;
-				Pair.InputPin->bDisplayAsMutableRef = ModelPin->GetInjectedNodes().Num() > 0;
 
 				SetupPinDefaultsFromModel(Pair.InputPin);
 
@@ -406,9 +403,9 @@ void UControlRigGraphNode::CreateInputOutputPins(URigVMPin* InParentPin)
 			Pair.OutputPin = CreatePin(EGPD_Output, GetPinTypeForModelPin(ModelPin), FName(*ModelPin->GetPinPath()));
 			if (Pair.OutputPin != nullptr)
 			{
+				Pair.OutputPin->bHidden = bHidden;
 				Pair.OutputPin->PinFriendlyName = FText::FromName(ModelPin->GetDisplayName());
 				Pair.OutputPin->bNotConnectable = ModelPin->GetDirection() != ERigVMPinDirection::IO;
-				Pair.OutputPin->bDisplayAsMutableRef = ModelPin->GetInjectedNodes().Num() > 0;
 
 				if (InParentPin != nullptr)
 				{
@@ -424,11 +421,11 @@ void UControlRigGraphNode::CreateInputOutputPins(URigVMPin* InParentPin)
 		{
 			if (!RerouteNode->GetShowsAsFullNode())
 			{
-				continue;
+				bHidden = true;
 			}
 		}
 
-		CreateInputOutputPins(ModelPin);
+		CreateInputOutputPins(ModelPin, bHidden);
 	}
 }
 
@@ -447,7 +444,6 @@ void UControlRigGraphNode::CreateOutputPins(URigVMPin* InParentPin)
 			{
 				Pair.OutputPin->PinFriendlyName = FText::FromName(ModelPin->GetDisplayName());
 				Pair.OutputPin->bNotConnectable = ModelPin->GetDirection() != ERigVMPinDirection::Output;
-				Pair.OutputPin->bDisplayAsMutableRef = ModelPin->GetInjectedNodes().Num() > 0;
 
 				if (InParentPin != nullptr)
 				{
@@ -751,7 +747,7 @@ bool UControlRigGraphNode::ShouldDrawNodeAsControlPointOnly(int32& OutInputPinIn
 	{
 		if (!Reroute->GetShowsAsFullNode())
 		{
-			if (Pins.Num() == 2)
+			if (Pins.Num() >= 2)
 			{
 				OutInputPinIndex = 0;
 				OutOutputPinIndex = 1;
diff --git a/Engine/Plugins/Experimental/ControlRig/Source/ControlRigDeveloper/Private/Graph/ControlRigGraphSchema.cpp b/Engine/Plugins/Experimental/ControlRig/Source/ControlRigDeveloper/Private/Graph/ControlRigGraphSchema.cpp
--- a/Engine/Plugins/Experimental/ControlRig/Source/ControlRigDeveloper/Private/Graph/ControlRigGraphSchema.cpp
+++ b/Engine/Plugins/Experimental/ControlRig/Source/ControlRigDeveloper/Private/Graph/ControlRigGraphSchema.cpp
@@ -52,6 +52,16 @@ bool UControlRigGraphSchema::TryCreateConnection(UEdGraphPin* PinA, UEdGraphPin*
 	}
 #endif
 
+	if (PinA == PinB)
+	{
+		return false;
+	}
+
+	if (PinA->GetOwningNode() == PinB->GetOwningNode())
+	{
+		return false;
+	}
+
 	UControlRigGraphSchema* MutableThis = (UControlRigGraphSchema*)this;
 	MutableThis->LastPinForCompatibleCheck = nullptr;
 
@@ -313,14 +323,14 @@ bool UControlRigGraphSchema::ArePinsCompatible(const UEdGraphPin* PinA, const UE
 	{
 		UControlRigGraphSchema* MutableThis = (UControlRigGraphSchema*)this;
 		MutableThis->LastPinForCompatibleCheck = PinB;
-		MutableThis->bLastPinWasInput = true;
+		MutableThis->bLastPinWasInput = PinB->Direction == EGPD_Input;
 		return true;
 	}
 	if (PinB->PinType.PinCategory == TEXT("REROUTE"))
 	{
 		UControlRigGraphSchema* MutableThis = (UControlRigGraphSchema*)this;
 		MutableThis->LastPinForCompatibleCheck = PinA;
-		MutableThis->bLastPinWasInput = false;
+		MutableThis->bLastPinWasInput = PinA->Direction == EGPD_Input;
 		return true;
 	}
 
@@ -361,5 +371,19 @@ void UControlRigGraphSchema::GetVariablePinTypes(TArray<FEdGraphPinType>& PinTyp
 	PinTypes.Add(FEdGraphPinType(UEdGraphSchema_K2::PC_Struct, FName(NAME_None), TBaseStructure<FLinearColor>::Get(), EPinContainerType::None, false, FEdGraphTerminalType()));
 }
 
+bool UControlRigGraphSchema::SafeDeleteNodeFromGraph(UEdGraph* Graph, UEdGraphNode* Node) const
+{
+	if (UControlRigGraphNode* RigNode = Cast<UControlRigGraphNode>(Node))
+	{
+		UBlueprint* Blueprint = FBlueprintEditorUtils::FindBlueprintForNodeChecked(RigNode);
+		UControlRigBlueprint* RigBlueprint = Cast<UControlRigBlueprint>(Blueprint);
+		if (RigBlueprint != nullptr)
+		{
+			return RigBlueprint->Controller->RemoveNode(RigNode->GetModelNode());
+		}
+	}
+	return false;
+}
+
 
 #undef LOCTEXT_NAMESPACE
\ No newline at end of file
diff --git a/Engine/Plugins/Experimental/ControlRig/Source/ControlRigDeveloper/Public/ControlRigBlueprint.h b/Engine/Plugins/Experimental/ControlRig/Source/ControlRigDeveloper/Public/ControlRigBlueprint.h
--- a/Engine/Plugins/Experimental/ControlRig/Source/ControlRigDeveloper/Public/ControlRigBlueprint.h
+++ b/Engine/Plugins/Experimental/ControlRig/Source/ControlRigDeveloper/Public/ControlRigBlueprint.h
@@ -67,20 +67,25 @@ public:
 	virtual bool ShouldBeMarkedDirtyUponTransaction() const override { return false; }
 
 	/** IInterface_PreviewMeshProvider interface */
+	UFUNCTION(BlueprintCallable, Category = "Control Rig Blueprint")
 	virtual void SetPreviewMesh(USkeletalMesh* PreviewMesh, bool bMarkAsDirty = true) override;
+	
+	UFUNCTION(BlueprintCallable, Category = "Control Rig Blueprint")
 	virtual USkeletalMesh* GetPreviewMesh() const override;
 
-	UFUNCTION()
+	UFUNCTION(BlueprintCallable, Category = "Control Rig Blueprint")
 	void RecompileVM();
 
-	UFUNCTION()
+	UFUNCTION(BlueprintCallable, Category = "Control Rig Blueprint")
 	void RecompileVMIfRequired();
 	
-	UFUNCTION()
+	UFUNCTION(BlueprintCallable, Category = "Control Rig Blueprint")
 	void RequestAutoVMRecompilation();
 
 	void IncrementVMRecompileBracket();
 	void DecrementVMRecompileBracket();
+
+	UFUNCTION(BlueprintCallable, Category = "Control Rig Blueprint")
 	void RequestControlRigInit();
 
 	UPROPERTY(EditAnywhere, BlueprintReadWrite, Category = "VM")
diff --git a/Engine/Plugins/Experimental/ControlRig/Source/ControlRigDeveloper/Public/Graph/ControlRigGraph.h b/Engine/Plugins/Experimental/ControlRig/Source/ControlRigDeveloper/Public/Graph/ControlRigGraph.h
--- a/Engine/Plugins/Experimental/ControlRig/Source/ControlRigDeveloper/Public/Graph/ControlRigGraph.h
+++ b/Engine/Plugins/Experimental/ControlRig/Source/ControlRigDeveloper/Public/Graph/ControlRigGraph.h
@@ -15,6 +15,8 @@ class UControlRig;
 class URigVMController;
 struct FRigCurveContainer;
 
+DECLARE_MULTICAST_DELEGATE_OneParam(FControlRigGraphNodeClicked, UControlRigGraphNode*);
+
 UCLASS()
 class CONTROLRIGDEVELOPER_API UControlRigGraph : public UEdGraph
 {
@@ -82,6 +84,8 @@ private:
 
 	bool bIsSelecting;
 
+	FControlRigGraphNodeClicked OnGraphNodeClicked;
+
 #endif
 #if WITH_EDITORONLY_DATA
 
@@ -103,5 +107,6 @@ private:
 #endif
 	friend class UControlRigGraphNode;
 	friend class FControlRigEditor;
+	friend class SControlRigGraphNode;
 };
 
diff --git a/Engine/Plugins/Experimental/ControlRig/Source/ControlRigDeveloper/Public/Graph/ControlRigGraphNode.h b/Engine/Plugins/Experimental/ControlRig/Source/ControlRigDeveloper/Public/Graph/ControlRigGraphNode.h
--- a/Engine/Plugins/Experimental/ControlRig/Source/ControlRigDeveloper/Public/Graph/ControlRigGraphNode.h
+++ b/Engine/Plugins/Experimental/ControlRig/Source/ControlRigDeveloper/Public/Graph/ControlRigGraphNode.h
@@ -136,7 +136,7 @@ protected:
 	/** Helper function for AllocateDefaultPins */
 	void CreateExecutionPins();
 	void CreateInputPins(URigVMPin* InParentPin = nullptr);
-	void CreateInputOutputPins(URigVMPin* InParentPin = nullptr);
+	void CreateInputOutputPins(URigVMPin* InParentPin = nullptr, bool bHidden = false);
 	void CreateOutputPins(URigVMPin* InParentPin = nullptr);
 
 	/** Copies default values from underlying properties into pin defaults, for editing */
diff --git a/Engine/Plugins/Experimental/ControlRig/Source/ControlRigDeveloper/Public/Graph/ControlRigGraphSchema.h b/Engine/Plugins/Experimental/ControlRig/Source/ControlRigDeveloper/Public/Graph/ControlRigGraphSchema.h
--- a/Engine/Plugins/Experimental/ControlRig/Source/ControlRigDeveloper/Public/Graph/ControlRigGraphSchema.h
+++ b/Engine/Plugins/Experimental/ControlRig/Source/ControlRigDeveloper/Public/Graph/ControlRigGraphSchema.h
@@ -71,6 +71,7 @@ public:
 	virtual void ClearPinWatch(UEdGraphPin const* Pin) const override;
 	virtual void OnPinConnectionDoubleCicked(UEdGraphPin* PinA, UEdGraphPin* PinB, const FVector2D& GraphPosition) const override;
 	virtual bool MarkBlueprintDirtyFromNewNode(UBlueprint* InBlueprint, UEdGraphNode* InEdGraphNode) const override;
+	virtual bool SafeDeleteNodeFromGraph(UEdGraph* Graph, UEdGraphNode* Node) const override;
 
 	virtual bool IsStructEditable(UStruct* InStruct) const;
 
diff --git a/Engine/Plugins/Experimental/ControlRig/Source/ControlRigEditor/Private/ControlRigBlueprintActions.cpp b/Engine/Plugins/Experimental/ControlRig/Source/ControlRigEditor/Private/ControlRigBlueprintActions.cpp
--- a/Engine/Plugins/Experimental/ControlRig/Source/ControlRigEditor/Private/ControlRigBlueprintActions.cpp
+++ b/Engine/Plugins/Experimental/ControlRig/Source/ControlRigEditor/Private/ControlRigBlueprintActions.cpp
@@ -111,7 +111,7 @@ void FControlRigBlueprintActions::ExtendSketalMeshToolMenu()
 						{
 							for (UObject* SelectedObject : SelectedObjects)
 							{
-								FControlRigBlueprintActions::ExecuteCreateControlRigFromSkeletalMesh(SelectedObject);
+								FControlRigBlueprintActions::CreateControlRigFromSkeletalMeshOrSkeleton(SelectedObject);
 							}
 						})
 					);
@@ -121,7 +121,22 @@ void FControlRigBlueprintActions::ExtendSketalMeshToolMenu()
 	}
 }
 
-void FControlRigBlueprintActions::ExecuteCreateControlRigFromSkeletalMesh(UObject* InSelectedObject)
+UControlRigBlueprint* FControlRigBlueprintActions::CreateNewControlRigAsset(const FString& InDesiredPackagePath)
+{
+	FAssetToolsModule& AssetToolsModule = FModuleManager::LoadModuleChecked<FAssetToolsModule>("AssetTools");
+
+	UControlRigBlueprintFactory* Factory = NewObject<UControlRigBlueprintFactory>();
+	Factory->ParentClass = UControlRig::StaticClass();
+
+	FString UniquePackageName;
+	FString UniqueAssetName;
+	AssetToolsModule.Get().CreateUniqueAssetName(InDesiredPackagePath, TEXT(""), UniquePackageName, UniqueAssetName);
+
+	UObject* NewAsset = AssetToolsModule.Get().CreateAsset(*UniqueAssetName, *InDesiredPackagePath, nullptr, Factory);
+	return Cast<UControlRigBlueprint>(NewAsset);
+}
+
+UControlRigBlueprint* FControlRigBlueprintActions::CreateControlRigFromSkeletalMeshOrSkeleton(UObject* InSelectedObject)
 {
 	FScopedTransaction Transaction(LOCTEXT("CreateControlRigFromSkeletalMesh", "Create Control Rig"));
 
@@ -132,11 +147,11 @@ void FControlRigBlueprintActions::ExecuteCreateControlRigFromSkeletalMesh(UObjec
 	{
 		Skeleton = SkeletalMesh->Skeleton;
 	}
-
-	check(Skeleton);
-
-	UControlRigBlueprintFactory* Factory = NewObject<UControlRigBlueprintFactory>();
-	Factory->ParentClass = UControlRig::StaticClass();
+	else if (Skeleton == nullptr)
+	{
+		UE_LOG(LogControlRigEditor, Error, TEXT("CreateControlRigFromSkeletalMeshOrSkeleton: Provided object has to be a SkeletalMesh or Skeleton."));
+		return nullptr;
+	}
 
 	FString PackagePath = InSelectedObject->GetPathName();
 	FString ControlRigName = FString::Printf(TEXT("%s_CtrlRig"), *InSelectedObject->GetName());
@@ -147,15 +162,7 @@ void FControlRigBlueprintActions::ExecuteCreateControlRigFromSkeletalMesh(UObjec
 		PackagePath = PackagePath.Left(LastSlashPos);
 	}
 
-	FAssetToolsModule& AssetToolsModule = FModuleManager::LoadModuleChecked<FAssetToolsModule>("AssetTools");
-
-	FString UniquePackageName;
-	FString UniqueAssetName;
-	AssetToolsModule.Get().CreateUniqueAssetName(PackagePath / ControlRigName, TEXT(""), UniquePackageName, UniqueAssetName);
-
-	UObject* NewAsset = AssetToolsModule.Get().CreateAsset(*UniqueAssetName, *PackagePath, nullptr, Factory);
-
-	UControlRigBlueprint* NewControlRigBlueprint = Cast<UControlRigBlueprint>(NewAsset);
+	UControlRigBlueprint* NewControlRigBlueprint = CreateNewControlRigAsset(PackagePath / ControlRigName);
 	check(NewControlRigBlueprint);
 
 	NewControlRigBlueprint->HierarchyContainer.BoneHierarchy.ImportSkeleton(Skeleton->GetReferenceSkeleton(), NAME_None, true, true, false, false);
@@ -166,6 +173,8 @@ void FControlRigBlueprintActions::ExecuteCreateControlRigFromSkeletalMesh(UObjec
 	{
 		NewControlRigBlueprint->SetPreviewMesh(SkeletalMesh);
 	}
+
+	return NewControlRigBlueprint;
 }
 
 USkeletalMesh* FControlRigBlueprintActions::GetSkeletalMeshFromControlRigBlueprint(UObject* InAsset)
diff --git a/Engine/Plugins/Experimental/ControlRig/Source/ControlRigEditor/Private/ControlRigBlueprintActions.h b/Engine/Plugins/Experimental/ControlRig/Source/ControlRigEditor/Private/ControlRigBlueprintActions.h
--- a/Engine/Plugins/Experimental/ControlRig/Source/ControlRigEditor/Private/ControlRigBlueprintActions.h
+++ b/Engine/Plugins/Experimental/ControlRig/Source/ControlRigEditor/Private/ControlRigBlueprintActions.h
@@ -27,7 +27,9 @@ public:
 	virtual UFactory* GetFactoryForBlueprintType(UBlueprint* InBlueprint) const override;
 
 	static void ExtendSketalMeshToolMenu();
-	static void ExecuteCreateControlRigFromSkeletalMesh(UObject* InSelectedObject);
+
+	static UControlRigBlueprint* CreateNewControlRigAsset(const FString& InDesiredPackagePath);
+	static UControlRigBlueprint* CreateControlRigFromSkeletalMeshOrSkeleton(UObject* InSelectedObject);
 
 	static USkeletalMesh* GetSkeletalMeshFromControlRigBlueprint(UObject* InAsset);
 	static void PostSpawningSkeletalMeshActor(AActor* InSpawnedActor, UObject* InAsset);
diff --git a/Engine/Plugins/Experimental/ControlRig/Source/ControlRigEditor/Private/ControlRigBlueprintFactory.cpp b/Engine/Plugins/Experimental/ControlRig/Source/ControlRigEditor/Private/ControlRigBlueprintFactory.cpp
--- a/Engine/Plugins/Experimental/ControlRig/Source/ControlRigEditor/Private/ControlRigBlueprintFactory.cpp
+++ b/Engine/Plugins/Experimental/ControlRig/Source/ControlRigEditor/Private/ControlRigBlueprintFactory.cpp
@@ -19,6 +19,7 @@
 #include "ClassViewerModule.h"
 #include "ClassViewerFilter.h"
 #include "ControlRigBlueprint.h"
+#include "ControlRigBlueprintActions.h"
 #include "ControlRigBlueprintGeneratedClass.h"
 
 #define LOCTEXT_NAMESPACE "ControlRigBlueprintFactory"
@@ -292,4 +293,14 @@ UObject* UControlRigBlueprintFactory::FactoryCreateNew(UClass* Class, UObject* I
 	return FactoryCreateNew(Class, InParent, Name, Flags, Context, Warn, NAME_None);
 }
 
+UControlRigBlueprint* UControlRigBlueprintFactory::CreateNewControlRigAsset(const FString& InDesiredPackagePath)
+{
+	return FControlRigBlueprintActions::CreateNewControlRigAsset(InDesiredPackagePath);
+}
+
+UControlRigBlueprint* UControlRigBlueprintFactory::CreateControlRigFromSkeletalMeshOrSkeleton(UObject* InSelectedObject)
+{
+	return FControlRigBlueprintActions::CreateControlRigFromSkeletalMeshOrSkeleton(InSelectedObject);
+}
+
 #undef LOCTEXT_NAMESPACE
diff --git a/Engine/Plugins/Experimental/ControlRig/Source/ControlRigEditor/Private/ControlRigBlueprintFactory.h b/Engine/Plugins/Experimental/ControlRig/Source/ControlRigEditor/Private/ControlRigBlueprintFactory.h
--- a/Engine/Plugins/Experimental/ControlRig/Source/ControlRigEditor/Private/ControlRigBlueprintFactory.h
+++ b/Engine/Plugins/Experimental/ControlRig/Source/ControlRigEditor/Private/ControlRigBlueprintFactory.h
@@ -8,6 +8,7 @@
 #include "Engine/Blueprint.h"
 #include "Factories/Factory.h"
 #include "ControlRig.h"
+#include "ControlRigBlueprint.h"
 #include "ControlRigBlueprintFactory.generated.h"
 
 UCLASS(MinimalAPI, HideCategories=Object)
@@ -26,5 +27,20 @@ public:
 	virtual bool ConfigureProperties() override;
 	virtual UObject* FactoryCreateNew(UClass* Class, UObject* InParent, FName Name, EObjectFlags Flags, UObject* Context, FFeedbackContext* Warn, FName CallingContext) override;
 	virtual UObject* FactoryCreateNew(UClass* Class, UObject* InParent, FName Name, EObjectFlags Flags, UObject* Context, FFeedbackContext* Warn) override;
+
+	/**
+	 * Create a new control rig asset within the contents space of the project.
+	 * @param InDesiredPackagePath The package path to use for the control rig asset
+	 */
+	UFUNCTION(BlueprintCallable, Category = "Control Rig")
+	static UControlRigBlueprint* CreateNewControlRigAsset(const FString& InDesiredPackagePath);
+
+	/**
+	 * Create a new control rig asset within the contents space of the project
+	 * based on a skeletal mesh or skeleton object.
+	 * @param The SkeletalMesh / Skeleton object to base the control rig asset on
+	 */
+	UFUNCTION(BlueprintCallable, Category = "Control Rig")
+	static UControlRigBlueprint* CreateControlRigFromSkeletalMeshOrSkeleton(UObject* InSelectedObject);
 };
 
diff --git a/Engine/Plugins/Experimental/ControlRig/Source/ControlRigEditor/Private/ControlRigEditorStyle.h b/Engine/Plugins/Experimental/ControlRig/Source/ControlRigEditor/Private/ControlRigEditorStyle.h
--- a/Engine/Plugins/Experimental/ControlRig/Source/ControlRigEditor/Private/ControlRigEditorStyle.h
+++ b/Engine/Plugins/Experimental/ControlRig/Source/ControlRigEditor/Private/ControlRigEditorStyle.h
@@ -19,6 +19,7 @@ public:
 		: FSlateStyleSet("ControlRigEditorStyle")
 	{
 		const FVector2D Icon10x10(10.0f, 10.0f);
+		const FVector2D Icon14x14(14.0f, 14.0f);
 		const FVector2D Icon16x16(16.0f, 16.0f);
 		const FVector2D Icon20x20(20.0f, 20.0f);
 		const FVector2D Icon24x24(24.0f, 24.0f);
@@ -67,7 +68,10 @@ public:
 			Set("ControlRig.AutoCompileGraph", new IMAGE_BRUSH("Slate/AutoCompile", Icon40x40));
 			Set("ControlRig.AutoCompileGraph.Small", new IMAGE_BRUSH("Slate/AutoCompile", Icon20x20));
 
-			Set("ControlRig.RigUnit.VisualDebug", new IMAGE_BRUSH("Slate/ControlRig_Unit_VisualDebug_40x", Icon40x40));
+			Set("ControlRig.Bug.Dot", new IMAGE_BRUSH("Slate/ControlRig_BugDot_32x", Icon16x16));
+			Set("ControlRig.Bug.Normal", new IMAGE_BRUSH("Slate/ControlRig_Bug_28x", Icon14x14));
+			Set("ControlRig.Bug.Open", new IMAGE_BRUSH("Slate/ControlRig_BugOpen_28x", Icon14x14));
+			Set("ControlRig.Bug.Solid", new IMAGE_BRUSH("Slate/ControlRig_BugSolid_28x", Icon14x14));
 		}
 
 		// Graph styles
diff --git a/Engine/Plugins/Experimental/ControlRig/Source/ControlRigEditor/Private/ControlRigElementDetails.cpp b/Engine/Plugins/Experimental/ControlRig/Source/ControlRigEditor/Private/ControlRigElementDetails.cpp
--- a/Engine/Plugins/Experimental/ControlRig/Source/ControlRigEditor/Private/ControlRigElementDetails.cpp
+++ b/Engine/Plugins/Experimental/ControlRig/Source/ControlRigEditor/Private/ControlRigElementDetails.cpp
@@ -14,6 +14,7 @@
 #include "Editor/SControlRigGizmoNameList.h"
 #include "ControlRigBlueprint.h"
 #include "Graph/ControlRigGraph.h"
+#include "PropertyCustomizationHelpers.h"
 
 #define LOCTEXT_NAMESPACE "ControlRigElementDetails"
 
@@ -484,6 +485,7 @@ void FRigControlDetails::CustomizeDetails(IDetailLayoutBuilder& DetailBuilder)
 		}
 	}
 
+	UEnum* ControlTypeEnum = StaticEnum<ERigControlType>();
 	UEnum* ValueTypeEnum = StaticEnum<ERigControlValueType>();
 	FRigControl& Control = ContainerBeingCustomized->ControlHierarchy[ElementKeyBeingCustomized.Name];
 
@@ -508,6 +510,7 @@ void FRigControlDetails::CustomizeDetails(IDetailLayoutBuilder& DetailBuilder)
 					case ERigControlType::Float:
 					{
 						Control.Value = Control.InitialValue = FRigControlValue::Make<float>(0.f);
+						Control.bLimitTranslation = true;
 						Control.MinimumValue = FRigControlValue::Make<float>(0.f);
 						Control.MaximumValue = FRigControlValue::Make<float>(1.f);
 						break;
@@ -515,6 +518,7 @@ void FRigControlDetails::CustomizeDetails(IDetailLayoutBuilder& DetailBuilder)
 					case ERigControlType::Vector2D:
 					{
 						Control.Value = Control.InitialValue = FRigControlValue::Make<FVector2D>(FVector2D::ZeroVector);
+						Control.bLimitTranslation = true;
 						Control.MinimumValue = FRigControlValue::Make<FVector2D>(FVector2D::ZeroVector);
 						Control.MaximumValue = FRigControlValue::Make<FVector2D>(FVector2D(1.f, 1.f));
 						break;
@@ -559,6 +563,7 @@ void FRigControlDetails::CustomizeDetails(IDetailLayoutBuilder& DetailBuilder)
 					}
 				}
 
+				this->BlueprintBeingCustomized->PropagateHierarchyFromBPToInstances(false, false);
 				this->ContainerBeingCustomized->ControlHierarchy.OnControlUISettingsChanged.Broadcast(this->ContainerBeingCustomized, this->ElementKeyBeingCustomized);
 			}
 		}
@@ -574,6 +579,8 @@ void FRigControlDetails::CustomizeDetails(IDetailLayoutBuilder& DetailBuilder)
 		ControlType->CreatePropertyValueWidget()
 	];
 
+	FString ControlTypeName = ControlTypeEnum->GetDisplayNameTextByValue((int64)Control.ControlType).ToString();
+
 	switch (Control.ControlType)
 	{
 		case ERigControlType::Bool:
@@ -588,12 +595,15 @@ void FRigControlDetails::CustomizeDetails(IDetailLayoutBuilder& DetailBuilder)
 					continue;
 				}
 
-				ControlCategory.AddCustomRow(ValueTypeEnum->GetDisplayNameTextByValue(ValueTypeIndex))
+				FString ValueTypeName = ValueTypeEnum->GetDisplayNameTextByValue(ValueTypeIndex).ToString();
+				FText PropertyLabel = FText::FromString(FString::Printf(TEXT("%s %s"), *ValueTypeName, *ControlTypeName));
+
+				ControlCategory.AddCustomRow(PropertyLabel)
 				.NameContent()
 				.VAlign(VAlign_Top)
 				[
 					SNew(STextBlock)
-					.Text(ValueTypeEnum->GetDisplayNameTextByValue(ValueTypeIndex))
+					.Text(PropertyLabel)
 					.Font(IDetailLayoutBuilder::GetDetailFont())
 				]
 				.ValueContent()
@@ -614,6 +624,8 @@ void FRigControlDetails::CustomizeDetails(IDetailLayoutBuilder& DetailBuilder)
 			for(int64 ValueTypeIndex = 0; ValueTypeIndex < ValueTypeEnum->GetMaxEnumValue(); ValueTypeIndex++)
 			{
 				ERigControlValueType ValueType = (ERigControlValueType)ValueTypeIndex;
+				FString ValueTypeName = ValueTypeEnum->GetDisplayNameTextByValue(ValueTypeIndex).ToString();
+				FText PropertyLabel = FText::FromString(FString::Printf(TEXT("%s %s"), *ValueTypeName, *ControlTypeName));
 
 				if (ValueType == ERigControlValueType::Minimum)
 				{
@@ -641,12 +653,12 @@ void FRigControlDetails::CustomizeDetails(IDetailLayoutBuilder& DetailBuilder)
 				}
 
 				IDetailCategoryBuilder& Category = (ValueType == ERigControlValueType::Minimum || ValueType == ERigControlValueType::Maximum) ? LimitsCategory : ControlCategory;
-				Category.AddCustomRow(ValueTypeEnum->GetDisplayNameTextByValue(ValueTypeIndex))
+				Category.AddCustomRow(PropertyLabel)
 				.NameContent()
 				.VAlign(VAlign_Top)
 				[
 					SNew(STextBlock)
-					.Text(ValueTypeEnum->GetDisplayNameTextByValue(ValueTypeIndex))
+					.Text(PropertyLabel)
 					.Font(IDetailLayoutBuilder::GetDetailFont())
 					.IsEnabled(this, &FRigControlDetails::IsEnabled, ValueType)
 				]
@@ -669,144 +681,17 @@ void FRigControlDetails::CustomizeDetails(IDetailLayoutBuilder& DetailBuilder)
 			break;
 		}
 		case ERigControlType::Vector2D:
-		{
-			for(int64 ValueTypeIndex = 0; ValueTypeIndex < ValueTypeEnum->GetMaxEnumValue(); ValueTypeIndex++)
-			{
-				ERigControlValueType ValueType = (ERigControlValueType)ValueTypeIndex;
-
-				if (ValueType == ERigControlValueType::Minimum)
-				{
-					LimitsCategory.AddCustomRow(FText::FromString(TEXT("Limit")))
-					.NameContent()
-					[
-						SNew(STextBlock)
-						.Text(FText::FromString(TEXT("Limit")))
-						.Font(IDetailLayoutBuilder::GetDetailFont())
-					]
-					.ValueContent()
-					[
-						DetailBuilder.GetProperty(TEXT("bLimitTranslation"))->CreatePropertyValueWidget()
-					];
-
-					LimitsCategory.AddCustomRow(FText::FromString(TEXT("bDrawLimits")))
-					.NameContent()
-					[
-						DetailBuilder.GetProperty(TEXT("bDrawLimits"))->CreatePropertyNameWidget()
-					]
-					.ValueContent()
-					[
-						DetailBuilder.GetProperty(TEXT("bDrawLimits"))->CreatePropertyValueWidget()
-					];
-				}
-
-				IDetailCategoryBuilder& Category = (ValueType == ERigControlValueType::Minimum || ValueType == ERigControlValueType::Maximum) ? LimitsCategory : ControlCategory;
-				Category.AddCustomRow(ValueTypeEnum->GetDisplayNameTextByValue(ValueTypeIndex))
-				.NameContent()
-				.VAlign(VAlign_Top)
-				[
-					SNew(STextBlock)
-					.Text(ValueTypeEnum->GetDisplayNameTextByValue(ValueTypeIndex))
-					.Font(IDetailLayoutBuilder::GetDetailFont())
-					.IsEnabled(this, &FRigControlDetails::IsEnabled, ValueType)
-				]
-				.ValueContent()
-				.MinDesiredWidth(125.0f * 3.0f) // copied from FComponentTransformDetails
-				.MaxDesiredWidth(125.0f * 3.0f)
-				[
-					SNew(SHorizontalBox)
-					+ SHorizontalBox::Slot()
-					.Padding(0.0f, 1.0f, 2.0f, 1.0f)
-					[
-						SNew(SNumericEntryBox<float>)
-						.Font(FEditorStyle::GetFontStyle(TEXT("MenuItem.Font")))
-						.AllowSpin(true)
-						.IsEnabled(this, &FRigControlDetails::IsEnabled, ValueType)
-						.Value(this, &FRigControlDetails::GetComponentValueFloat, ValueType, ERigElementDetailsTransformComponent::TranslationX)
-						.OnValueChanged(this, &FRigControlDetails::SetComponentValueFloat, ValueType, ERigElementDetailsTransformComponent::TranslationX)
-					]
-					.Padding(0.0f, 1.0f, 2.0f, 1.0f)
-					+ SHorizontalBox::Slot()
-					[
-						SNew(SNumericEntryBox<float>)
-						.Font(FEditorStyle::GetFontStyle(TEXT("MenuItem.Font")))
-						.AllowSpin(true)
-						.IsEnabled(this, &FRigControlDetails::IsEnabled, ValueType)
-						.Value(this, &FRigControlDetails::GetComponentValueFloat, ValueType, ERigElementDetailsTransformComponent::TranslationY)
-						.OnValueChanged(this, &FRigControlDetails::SetComponentValueFloat, ValueType, ERigElementDetailsTransformComponent::TranslationY)
-					]
-				];
-			}
-			break;
-		}
 		case ERigControlType::Position:
-		{
-			for(int64 ValueTypeIndex = 0; ValueTypeIndex < ValueTypeEnum->GetMaxEnumValue(); ValueTypeIndex++)
-			{
-				ERigControlValueType ValueType = (ERigControlValueType)ValueTypeIndex;
-
-				if (ValueType == ERigControlValueType::Minimum)
-				{
-					LimitsCategory.AddCustomRow(FText::FromString(TEXT("Limit")))
-					.NameContent()
-					[
-						SNew(STextBlock)
-						.Text(FText::FromString(TEXT("Limit")))
-						.Font(IDetailLayoutBuilder::GetDetailFont())
-					]
-					.ValueContent()
-					[
-						DetailBuilder.GetProperty(TEXT("bLimitTranslation"))->CreatePropertyValueWidget()
-					];
-
-					LimitsCategory.AddCustomRow(FText::FromString(TEXT("bDrawLimits")))
-					.NameContent()
-					[
-						DetailBuilder.GetProperty(TEXT("bDrawLimits"))->CreatePropertyNameWidget()
-					]
-					.ValueContent()
-					[
-						DetailBuilder.GetProperty(TEXT("bDrawLimits"))->CreatePropertyValueWidget()
-					];
-				}
-
-				IDetailCategoryBuilder& Category = (ValueType == ERigControlValueType::Minimum || ValueType == ERigControlValueType::Maximum) ? LimitsCategory : ControlCategory;
-				Category.AddCustomRow(ValueTypeEnum->GetDisplayNameTextByValue(ValueTypeIndex))
-				.NameContent()
-				.VAlign(VAlign_Top)
-				[
-					SNew(STextBlock)
-					.Text(ValueTypeEnum->GetDisplayNameTextByValue(ValueTypeIndex))
-					.Font(IDetailLayoutBuilder::GetDetailFont())
-					.IsEnabled(this, &FRigControlDetails::IsEnabled, ValueType)
-				]
-				.ValueContent()
-				.MinDesiredWidth(125.0f * 3.0f) // copied from FComponentTransformDetails
-				.MaxDesiredWidth(125.0f * 3.0f)
-				[
-					SNew(SVerticalBox)
-					+ SVerticalBox::Slot()
-					[
-						SNew(SVectorInputBox)
-						.AllowSpin(true)
-						.IsEnabled(this, &FRigControlDetails::IsEnabled, ValueType)
-						.bColorAxisLabels(true)
-						.Font(IDetailLayoutBuilder::GetDetailFont())
-						.X(this, &FRigControlDetails::GetComponentValueFloat, ValueType, ERigElementDetailsTransformComponent::TranslationX)
-						.Y(this, &FRigControlDetails::GetComponentValueFloat, ValueType, ERigElementDetailsTransformComponent::TranslationY)
-						.Z(this, &FRigControlDetails::GetComponentValueFloat, ValueType, ERigElementDetailsTransformComponent::TranslationZ)
-						.OnXCommitted(this, &FRigControlDetails::SetComponentValueFloat, ValueType, ERigElementDetailsTransformComponent::TranslationX)
-						.OnYCommitted(this, &FRigControlDetails::SetComponentValueFloat, ValueType, ERigElementDetailsTransformComponent::TranslationY)
-						.OnZCommitted(this, &FRigControlDetails::SetComponentValueFloat, ValueType, ERigElementDetailsTransformComponent::TranslationZ)
-					]
-				];
-			}
-			break;
-		}
 		case ERigControlType::Scale:
+		case ERigControlType::Rotator:
+		case ERigControlType::Transform:
+		case ERigControlType::TransformNoScale:
 		{
 			for(int64 ValueTypeIndex = 0; ValueTypeIndex < ValueTypeEnum->GetMaxEnumValue(); ValueTypeIndex++)
 			{
 				ERigControlValueType ValueType = (ERigControlValueType)ValueTypeIndex;
+				FString ValueTypeName = ValueTypeEnum->GetDisplayNameTextByValue(ValueTypeIndex).ToString();
+				FText PropertyLabel = FText::FromString(FString::Printf(TEXT("%s %s"), *ValueTypeName, *ControlTypeName));
 
 				if (ValueType == ERigControlValueType::Minimum)
 				{
@@ -819,7 +704,7 @@ void FRigControlDetails::CustomizeDetails(IDetailLayoutBuilder& DetailBuilder)
 					]
 					.ValueContent()
 					[
-						DetailBuilder.GetProperty(TEXT("bLimitScale"))->CreatePropertyValueWidget()
+						DetailBuilder.GetProperty(TEXT("bLimitTranslation"))->CreatePropertyValueWidget()
 					];
 
 					LimitsCategory.AddCustomRow(FText::FromString(TEXT("bDrawLimits")))
@@ -833,286 +718,58 @@ void FRigControlDetails::CustomizeDetails(IDetailLayoutBuilder& DetailBuilder)
 					];
 				}
 
-				IDetailCategoryBuilder& Category = (ValueType == ERigControlValueType::Minimum || ValueType == ERigControlValueType::Maximum) ? LimitsCategory : ControlCategory;
-				Category.AddCustomRow(ValueTypeEnum->GetDisplayNameTextByValue(ValueTypeIndex))
-				.NameContent()
-				.VAlign(VAlign_Top)
-				[
-					SNew(STextBlock)
-					.Text(ValueTypeEnum->GetDisplayNameTextByValue(ValueTypeIndex))
-					.Font(IDetailLayoutBuilder::GetDetailFont())
-					.IsEnabled(this, &FRigControlDetails::IsEnabled, ValueType)
-				]
-				.ValueContent()
-				.MinDesiredWidth(125.0f * 3.0f) // copied from FComponentTransformDetails
-				.MaxDesiredWidth(125.0f * 3.0f)
-				[
-					SNew(SVerticalBox)
-					+ SVerticalBox::Slot()
-					[
-						SNew(SVectorInputBox)
-						.AllowSpin(true)
-						.IsEnabled(this, &FRigControlDetails::IsEnabled, ValueType)
-						.bColorAxisLabels(true)
-						.Font(IDetailLayoutBuilder::GetDetailFont())
-						.X(this, &FRigControlDetails::GetComponentValueFloat, ValueType, ERigElementDetailsTransformComponent::ScaleX)
-						.Y(this, &FRigControlDetails::GetComponentValueFloat, ValueType, ERigElementDetailsTransformComponent::ScaleY)
-						.Z(this, &FRigControlDetails::GetComponentValueFloat, ValueType, ERigElementDetailsTransformComponent::ScaleZ)
-						.OnXCommitted(this, &FRigControlDetails::SetComponentValueFloat, ValueType, ERigElementDetailsTransformComponent::ScaleX)
-						.OnYCommitted(this, &FRigControlDetails::SetComponentValueFloat, ValueType, ERigElementDetailsTransformComponent::ScaleY)
-						.OnZCommitted(this, &FRigControlDetails::SetComponentValueFloat, ValueType, ERigElementDetailsTransformComponent::ScaleZ)
-					]
-				];
-			}
-			break;
-		}
-		case ERigControlType::Rotator:
-		{
-			for(int64 ValueTypeIndex = 0; ValueTypeIndex < ValueTypeEnum->GetMaxEnumValue(); ValueTypeIndex++)
-			{
-				ERigControlValueType ValueType = (ERigControlValueType)ValueTypeIndex;
+				uint8* ValuePtr = (uint8*)&Control.GetValue(ValueType);
 
-				if (ValueType == ERigControlValueType::Minimum)
+				const UStruct* ValueStruct = nullptr;
+				switch (Control.ControlType)
 				{
-					LimitsCategory.AddCustomRow(FText::FromString(TEXT("Limit")))
-					.NameContent()
-					[
-						SNew(STextBlock)
-						.Text(FText::FromString(TEXT("Limit")))
-						.Font(IDetailLayoutBuilder::GetDetailFont())
-					]
-					.ValueContent()
-					[
-						DetailBuilder.GetProperty(TEXT("bLimitRotation"))->CreatePropertyValueWidget()
-					];
+					case ERigControlType::Vector2D:
+					{
+						ValueStruct = TBaseStructure<FVector2D>::Get();
+						break;
+					}
+					case ERigControlType::Position:
+					case ERigControlType::Scale:
+					{
+						ValueStruct = TBaseStructure<FVector>::Get();
+						break;
+					}
+					case ERigControlType::Rotator:
+					{
+						ValueStruct = TBaseStructure<FRotator>::Get();
+						break;
+					}
+					case ERigControlType::Transform:
+					{
+						ValueStruct = TBaseStructure<FTransform>::Get();
+						break;
+					}
+					case ERigControlType::TransformNoScale:
+					{
+						ValueStruct = TBaseStructure<FTransformNoScale>::Get();
+						break;
+					}
+					default:
+					{
+						checkNoEntry();
+					}
 				}
 
-				IDetailCategoryBuilder& Category = (ValueType == ERigControlValueType::Minimum || ValueType == ERigControlValueType::Maximum) ? LimitsCategory : ControlCategory;
-				Category.AddCustomRow(ValueTypeEnum->GetDisplayNameTextByValue(ValueTypeIndex))
-				.NameContent()
-				.VAlign(VAlign_Top)
-				[
-					SNew(STextBlock)
-					.Text(ValueTypeEnum->GetDisplayNameTextByValue(ValueTypeIndex))
-					.Font(IDetailLayoutBuilder::GetDetailFont())
-					.IsEnabled(this, &FRigControlDetails::IsEnabled, ValueType)
-				]
-				.ValueContent()
-				.MinDesiredWidth(125.0f * 3.0f) // copied from FComponentTransformDetails
-				.MaxDesiredWidth(125.0f * 3.0f)
-				[
-					SNew(SVerticalBox)
-					+ SVerticalBox::Slot()
-					[
-						SNew(SRotatorInputBox)
-						.AllowSpin(true)
-						.IsEnabled(this, &FRigControlDetails::IsEnabled, ValueType)
-						.bColorAxisLabels(true)
-						.Font(IDetailLayoutBuilder::GetDetailFont())
-						.Roll(this, &FRigControlDetails::GetComponentValueFloat, ValueType, ERigElementDetailsTransformComponent::RotationRoll)
-						.Pitch(this, &FRigControlDetails::GetComponentValueFloat, ValueType, ERigElementDetailsTransformComponent::RotationPitch)
-						.Yaw(this, &FRigControlDetails::GetComponentValueFloat, ValueType, ERigElementDetailsTransformComponent::RotationYaw)
-						.OnRollCommitted(this, &FRigControlDetails::SetComponentValueFloat, ValueType, ERigElementDetailsTransformComponent::RotationRoll)
-						.OnPitchCommitted(this, &FRigControlDetails::SetComponentValueFloat, ValueType, ERigElementDetailsTransformComponent::RotationPitch)
-						.OnYawCommitted(this, &FRigControlDetails::SetComponentValueFloat, ValueType, ERigElementDetailsTransformComponent::RotationYaw)
-					]
-				];
-			}
-			break;
-		}
-		case ERigControlType::Transform:
-		{
-			for(int64 ValueTypeIndex = 0; ValueTypeIndex < ValueTypeEnum->GetMaxEnumValue(); ValueTypeIndex++)
-			{
-				ERigControlValueType ValueType = (ERigControlValueType)ValueTypeIndex;
-
-				if (ValueType == ERigControlValueType::Minimum)
-				{
-					LimitsCategory.AddCustomRow(FText::FromString(TEXT("bLimitTranslation")))
-					.NameContent()
-					[
-						DetailBuilder.GetProperty(TEXT("bLimitTranslation"))->CreatePropertyNameWidget()
-					]
-					.ValueContent()
-					[
-						DetailBuilder.GetProperty(TEXT("bLimitTranslation"))->CreatePropertyValueWidget()
-					];
-
-					LimitsCategory.AddCustomRow(FText::FromString(TEXT("bLimitRotation")))
-					.NameContent()
-					[
-						DetailBuilder.GetProperty(TEXT("bLimitRotation"))->CreatePropertyNameWidget()
-					]
-					.ValueContent()
-					[
-						DetailBuilder.GetProperty(TEXT("bLimitRotation"))->CreatePropertyValueWidget()
-					];
-
-					LimitsCategory.AddCustomRow(FText::FromString(TEXT("bLimitScale")))
-					.NameContent()
-					[
-						DetailBuilder.GetProperty(TEXT("bLimitScale"))->CreatePropertyNameWidget()
-					]
-					.ValueContent()
-					[
-						DetailBuilder.GetProperty(TEXT("bLimitScale"))->CreatePropertyValueWidget()
-					];
-
-					LimitsCategory.AddCustomRow(FText::FromString(TEXT("bDrawLimits")))
-					.NameContent()
-					[
-						DetailBuilder.GetProperty(TEXT("bDrawLimits"))->CreatePropertyNameWidget()
-					]
-					.ValueContent()
-					[
-						DetailBuilder.GetProperty(TEXT("bDrawLimits"))->CreatePropertyValueWidget()
-					];
-				}
+				TSharedPtr<FStructOnScope> StructToDisplay = MakeShareable(new FStructOnScope(ValueStruct, ValuePtr));
 
 				IDetailCategoryBuilder& Category = (ValueType == ERigControlValueType::Minimum || ValueType == ERigControlValueType::Maximum) ? LimitsCategory : ControlCategory;
-				Category.AddCustomRow(ValueTypeEnum->GetDisplayNameTextByValue(ValueTypeIndex))
-				.NameContent()
-				.VAlign(VAlign_Top)
-				[
-					SNew(STextBlock)
-					.Text(ValueTypeEnum->GetDisplayNameTextByValue(ValueTypeIndex))
-					.Font(IDetailLayoutBuilder::GetDetailFont())
-					.IsEnabled(this, &FRigControlDetails::IsEnabled, ValueType)
-				]
-				.ValueContent()
-				.MinDesiredWidth(125.0f * 3.0f) // copied from FComponentTransformDetails
-				.MaxDesiredWidth(125.0f * 3.0f)
-				[
-					SNew(SVerticalBox)
-					+ SVerticalBox::Slot()
-					[
-						SNew(SVectorInputBox)
-						.AllowSpin(true)
-						.IsEnabled(this, &FRigControlDetails::IsEnabled, ValueType)
-						.bColorAxisLabels(true)
-						.Font(IDetailLayoutBuilder::GetDetailFont())
-						.X(this, &FRigControlDetails::GetComponentValueFloat, ValueType, ERigElementDetailsTransformComponent::TranslationX)
-						.Y(this, &FRigControlDetails::GetComponentValueFloat, ValueType, ERigElementDetailsTransformComponent::TranslationY)
-						.Z(this, &FRigControlDetails::GetComponentValueFloat, ValueType, ERigElementDetailsTransformComponent::TranslationZ)
-						.OnXCommitted(this, &FRigControlDetails::SetComponentValueFloat, ValueType, ERigElementDetailsTransformComponent::TranslationX)
-						.OnYCommitted(this, &FRigControlDetails::SetComponentValueFloat, ValueType, ERigElementDetailsTransformComponent::TranslationY)
-						.OnZCommitted(this, &FRigControlDetails::SetComponentValueFloat, ValueType, ERigElementDetailsTransformComponent::TranslationZ)
-					]
-					+ SVerticalBox::Slot()
-					[
-						SNew(SRotatorInputBox)
-						.AllowSpin(true)
-						.IsEnabled(this, &FRigControlDetails::IsEnabled, ValueType)
-						.bColorAxisLabels(true)
-						.Font(IDetailLayoutBuilder::GetDetailFont())
-						.Roll(this, &FRigControlDetails::GetComponentValueFloat, ValueType, ERigElementDetailsTransformComponent::RotationRoll)
-						.Pitch(this, &FRigControlDetails::GetComponentValueFloat, ValueType, ERigElementDetailsTransformComponent::RotationPitch)
-						.Yaw(this, &FRigControlDetails::GetComponentValueFloat, ValueType, ERigElementDetailsTransformComponent::RotationYaw)
-						.OnRollCommitted(this, &FRigControlDetails::SetComponentValueFloat, ValueType, ERigElementDetailsTransformComponent::RotationRoll)
-						.OnPitchCommitted(this, &FRigControlDetails::SetComponentValueFloat, ValueType, ERigElementDetailsTransformComponent::RotationPitch)
-						.OnYawCommitted(this, &FRigControlDetails::SetComponentValueFloat, ValueType, ERigElementDetailsTransformComponent::RotationYaw)
-					]
-					+ SVerticalBox::Slot()
-					[
-						SNew(SVectorInputBox)
-						.AllowSpin(true)
-						.IsEnabled(this, &FRigControlDetails::IsEnabled, ValueType)
-						.bColorAxisLabels(true)
-						.Font(IDetailLayoutBuilder::GetDetailFont())
-						.X(this, &FRigControlDetails::GetComponentValueFloat, ValueType, ERigElementDetailsTransformComponent::ScaleX)
-						.Y(this, &FRigControlDetails::GetComponentValueFloat, ValueType, ERigElementDetailsTransformComponent::ScaleY)
-						.Z(this, &FRigControlDetails::GetComponentValueFloat, ValueType, ERigElementDetailsTransformComponent::ScaleZ)
-						.OnXCommitted(this, &FRigControlDetails::SetComponentValueFloat, ValueType, ERigElementDetailsTransformComponent::ScaleX)
-						.OnYCommitted(this, &FRigControlDetails::SetComponentValueFloat, ValueType, ERigElementDetailsTransformComponent::ScaleY)
-						.OnZCommitted(this, &FRigControlDetails::SetComponentValueFloat, ValueType, ERigElementDetailsTransformComponent::ScaleZ)
-					]
-				];
-			}
-			break;
-		}
-		case ERigControlType::TransformNoScale:
-		{
-			for(int64 ValueTypeIndex = 0; ValueTypeIndex < ValueTypeEnum->GetMaxEnumValue(); ValueTypeIndex++)
-			{
-				ERigControlValueType ValueType = (ERigControlValueType)ValueTypeIndex;
 
-				if (ValueType == ERigControlValueType::Minimum)
-				{
-					LimitsCategory.AddCustomRow(FText::FromString(TEXT("bLimitTranslation")))
-					.NameContent()
-					[
-						DetailBuilder.GetProperty(TEXT("bLimitTranslation"))->CreatePropertyNameWidget()
-					]
-					.ValueContent()
-					[
-						DetailBuilder.GetProperty(TEXT("bLimitTranslation"))->CreatePropertyValueWidget()
-					];
+				IDetailPropertyRow* Row = Category.AddExternalStructure(StructToDisplay);
+				Row->DisplayName(PropertyLabel);
+				Row->ShouldAutoExpand(true);
+				Row->IsEnabled(TAttribute<bool>::Create(TAttribute<bool>::FGetter::CreateSP(this, &FRigControlDetails::IsEnabled, ValueType)));
 
-					LimitsCategory.AddCustomRow(FText::FromString(TEXT("bLimitRotation")))
-					.NameContent()
-					[
-						DetailBuilder.GetProperty(TEXT("bLimitRotation"))->CreatePropertyNameWidget()
-					]
-					.ValueContent()
-					[
-						DetailBuilder.GetProperty(TEXT("bLimitRotation"))->CreatePropertyValueWidget()
-					];
+				FProperty* Property = FRigControl::FindPropertyForValueType(ValueType);
+				FSimpleDelegate OnStructContentsChangedDelegate = FSimpleDelegate::CreateSP(this, &FRigControlDetails::OnStructContentsChanged, Property, DetailBuilder.GetPropertyUtilities());
 
-					LimitsCategory.AddCustomRow(FText::FromString(TEXT("bDrawLimits")))
-					.NameContent()
-					[
-						DetailBuilder.GetProperty(TEXT("bDrawLimits"))->CreatePropertyNameWidget()
-					]
-					.ValueContent()
-					[
-						DetailBuilder.GetProperty(TEXT("bDrawLimits"))->CreatePropertyValueWidget()
-					];
-				}
-
-				IDetailCategoryBuilder& Category = (ValueType == ERigControlValueType::Minimum || ValueType == ERigControlValueType::Maximum) ? LimitsCategory : ControlCategory;
-				Category.AddCustomRow(ValueTypeEnum->GetDisplayNameTextByValue(ValueTypeIndex))
-				.NameContent()
-				.VAlign(VAlign_Top)
-				[
-					SNew(STextBlock)
-					.Text(ValueTypeEnum->GetDisplayNameTextByValue(ValueTypeIndex))
-					.Font(IDetailLayoutBuilder::GetDetailFont())
-					.IsEnabled(this, &FRigControlDetails::IsEnabled, ValueType)
-				]
-				.ValueContent()
-				.MinDesiredWidth(125.0f * 3.0f) // copied from FComponentTransformDetails
-				.MaxDesiredWidth(125.0f * 3.0f)
-				[
-					SNew(SVerticalBox)
-					+ SVerticalBox::Slot()
-					[
-						SNew(SVectorInputBox)
-						.AllowSpin(true)
-						.IsEnabled(this, &FRigControlDetails::IsEnabled, ValueType)
-						.bColorAxisLabels(true)
-						.Font(IDetailLayoutBuilder::GetDetailFont())
-						.X(this, &FRigControlDetails::GetComponentValueFloat, ValueType, ERigElementDetailsTransformComponent::TranslationX)
-						.Y(this, &FRigControlDetails::GetComponentValueFloat, ValueType, ERigElementDetailsTransformComponent::TranslationY)
-						.Z(this, &FRigControlDetails::GetComponentValueFloat, ValueType, ERigElementDetailsTransformComponent::TranslationZ)
-						.OnXCommitted(this, &FRigControlDetails::SetComponentValueFloat, ValueType, ERigElementDetailsTransformComponent::TranslationX)
-						.OnYCommitted(this, &FRigControlDetails::SetComponentValueFloat, ValueType, ERigElementDetailsTransformComponent::TranslationY)
-						.OnZCommitted(this, &FRigControlDetails::SetComponentValueFloat, ValueType, ERigElementDetailsTransformComponent::TranslationZ)
-					]
-					+ SVerticalBox::Slot()
-					[
-						SNew(SRotatorInputBox)
-						.AllowSpin(true)
-						.IsEnabled(this, &FRigControlDetails::IsEnabled, ValueType)
-						.bColorAxisLabels(true)
-						.Font(IDetailLayoutBuilder::GetDetailFont())
-						.Roll(this, &FRigControlDetails::GetComponentValueFloat, ValueType, ERigElementDetailsTransformComponent::RotationRoll)
-						.Pitch(this, &FRigControlDetails::GetComponentValueFloat, ValueType, ERigElementDetailsTransformComponent::RotationPitch)
-						.Yaw(this, &FRigControlDetails::GetComponentValueFloat, ValueType, ERigElementDetailsTransformComponent::RotationYaw)
-						.OnRollCommitted(this, &FRigControlDetails::SetComponentValueFloat, ValueType, ERigElementDetailsTransformComponent::RotationRoll)
-						.OnPitchCommitted(this, &FRigControlDetails::SetComponentValueFloat, ValueType, ERigElementDetailsTransformComponent::RotationPitch)
-						.OnYawCommitted(this, &FRigControlDetails::SetComponentValueFloat, ValueType, ERigElementDetailsTransformComponent::RotationYaw)
-					]
-				];
+				TSharedPtr<IPropertyHandle> Handle = Row->GetPropertyHandle();
+				Handle->SetOnPropertyValueChanged(OnStructContentsChangedDelegate);
+				Handle->SetOnChildPropertyValueChanged(OnStructContentsChangedDelegate);
 			}
 			break;
 		}
@@ -1188,6 +845,12 @@ void FRigControlDetails::CustomizeDetails(IDetailLayoutBuilder& DetailBuilder)
 	}
 }
 
+void FRigControlDetails::OnStructContentsChanged(FProperty* InProperty, const TSharedRef<IPropertyUtilities> PropertyUtilities)
+{
+	FPropertyChangedEvent ChangeEvent(InProperty, EPropertyChangeType::ValueSet);
+	PropertyUtilities->NotifyFinishedChangingProperties(ChangeEvent);
+}
+
 ECheckBoxState FRigControlDetails::GetComponentValueBool(bool bInitial) const
 {
 	if (ContainerBeingCustomized != nullptr && ElementKeyBeingCustomized)
diff --git a/Engine/Plugins/Experimental/ControlRig/Source/ControlRigEditor/Private/ControlRigElementDetails.h b/Engine/Plugins/Experimental/ControlRig/Source/ControlRigEditor/Private/ControlRigElementDetails.h
--- a/Engine/Plugins/Experimental/ControlRig/Source/ControlRigEditor/Private/ControlRigElementDetails.h
+++ b/Engine/Plugins/Experimental/ControlRig/Source/ControlRigEditor/Private/ControlRigElementDetails.h
@@ -111,6 +111,7 @@ public:
 	void SetComponentValueFloat(float InNewValue, ERigControlValueType InValueType, ERigElementDetailsTransformComponent Component);
 	bool IsGizmoEnabled() const;
 	bool IsEnabled(ERigControlValueType InValueType) const;
+	void OnStructContentsChanged(FProperty* InProperty, const TSharedRef<IPropertyUtilities> PropertyUtilities);
 
 	const TArray<TSharedPtr<FString>>& GetGizmoNameList() const;
 	const TArray<TSharedPtr<FString>>& GetControlTypeList() const;
diff --git a/Engine/Plugins/Experimental/ControlRig/Source/ControlRigEditor/Private/EditMode/ControlRigControlsProxy.cpp b/Engine/Plugins/Experimental/ControlRig/Source/ControlRigEditor/Private/EditMode/ControlRigControlsProxy.cpp
--- a/Engine/Plugins/Experimental/ControlRig/Source/ControlRigEditor/Private/EditMode/ControlRigControlsProxy.cpp
+++ b/Engine/Plugins/Experimental/ControlRig/Source/ControlRigEditor/Private/EditMode/ControlRigControlsProxy.cpp
@@ -38,7 +38,7 @@ void UControlRigControlsProxy::PostEditChangeProperty(struct FPropertyChangedEve
 #if WITH_EDITOR
 void UControlRigControlsProxy::PostEditUndo()
 {
-	if (RigControl && ControlRig.IsValid())
+	if (RigControl && ControlRig.IsValid() && ControlRig->GetControlHierarchy().GetIndex(ControlName) != INDEX_NONE)
 	{
 		ControlRig->SelectControl(ControlName, bSelected);
 	}
@@ -76,7 +76,7 @@ void UControlRigTransformControlProxy::ValueChanged()
 #if WITH_EDITOR
 void UControlRigTransformControlProxy::PostEditUndo()
 {
-	if (RigControl && ControlRig.IsValid())
+	if (RigControl && ControlRig.IsValid() && ControlRig->GetControlHierarchy().GetIndex(ControlName) != INDEX_NONE)
 	{
 		ControlRig->SelectControl(ControlName, bSelected);
 		ControlRig->SetControlValue<FTransform>(ControlName, Transform, true, EControlRigSetKey::DoNotCare);
@@ -120,7 +120,7 @@ void UControlRigTransformNoScaleControlProxy::ValueChanged()
 #if WITH_EDITOR
 void UControlRigTransformNoScaleControlProxy::PostEditUndo()
 {
-	if (RigControl && ControlRig.IsValid())
+	if (RigControl && ControlRig.IsValid() && ControlRig->GetControlHierarchy().GetIndex(ControlName) != INDEX_NONE)
 	{
 		ControlRig->SelectControl(ControlName, bSelected);
 		ControlRig->SetControlValue<FTransformNoScale>(ControlName, Transform, true, EControlRigSetKey::DoNotCare);
@@ -165,7 +165,7 @@ void UControlRigFloatControlProxy::ValueChanged()
 #if WITH_EDITOR
 void UControlRigFloatControlProxy::PostEditUndo()
 {
-	if (RigControl && ControlRig.IsValid())
+	if (RigControl && ControlRig.IsValid() && ControlRig->GetControlHierarchy().GetIndex(ControlName) != INDEX_NONE)
 	{
 		ControlRig->SelectControl(ControlName, bSelected);
 		ControlRig->SetControlValue<float>(ControlName, Float, true, EControlRigSetKey::DoNotCare);
@@ -186,7 +186,8 @@ void UControlRigFloatControlProxy::SetKey(const IPropertyHandle& KeyedPropertyHa
 void UControlRigVectorControlProxy::PostEditChangeProperty(struct FPropertyChangedEvent& PropertyChangedEvent)
 {
 	Super::PostEditChangeProperty(PropertyChangedEvent);
-	if (PropertyChangedEvent.Property && PropertyChangedEvent.Property->GetFName() == GET_MEMBER_NAME_CHECKED(UControlRigVectorControlProxy, Vector))
+	if ((PropertyChangedEvent.Property && PropertyChangedEvent.Property->GetFName() == GET_MEMBER_NAME_CHECKED(UControlRigVectorControlProxy, Vector))
+		|| (PropertyChangedEvent.MemberProperty && PropertyChangedEvent.MemberProperty->GetFName() == GET_MEMBER_NAME_CHECKED(UControlRigVectorControlProxy, Vector)))
 	{
 		if (RigControl && ControlRig.IsValid())
 		{
@@ -212,7 +213,7 @@ void UControlRigVectorControlProxy::ValueChanged()
 #if WITH_EDITOR
 void UControlRigVectorControlProxy::PostEditUndo()
 {
-	if (RigControl && ControlRig.IsValid())
+	if (RigControl && ControlRig.IsValid() && ControlRig->GetControlHierarchy().GetIndex(ControlName) != INDEX_NONE)
 	{
 		ControlRig->SelectControl(ControlName, bSelected);
 		ControlRig->SetControlValue<FVector>(ControlName, Vector, true, EControlRigSetKey::DoNotCare);
@@ -258,7 +259,7 @@ void UControlRigVector2DControlProxy::ValueChanged()
 #if WITH_EDITOR
 void UControlRigVector2DControlProxy::PostEditUndo()
 {
-	if (RigControl && ControlRig.IsValid())
+	if (RigControl && ControlRig.IsValid() && ControlRig->GetControlHierarchy().GetIndex(ControlName) != INDEX_NONE)
 	{
 		ControlRig->SelectControl(ControlName, bSelected);
 		ControlRig->SetControlValue<FVector2D>(ControlName, Vector2D, true, EControlRigSetKey::DoNotCare);
@@ -305,7 +306,7 @@ void UControlRigBoolControlProxy::ValueChanged()
 #if WITH_EDITOR
 void UControlRigBoolControlProxy::PostEditUndo()
 {
-	if (RigControl && ControlRig.IsValid())
+	if (RigControl && ControlRig.IsValid() && ControlRig->GetControlHierarchy().GetIndex(ControlName) != INDEX_NONE)
 	{
 		ControlRig->SelectControl(ControlName, bSelected);
 		ControlRig->SetControlValue<bool>(ControlName, Bool, true, EControlRigSetKey::DoNotCare);
diff --git a/Engine/Plugins/Experimental/ControlRig/Source/ControlRigEditor/Private/EditMode/ControlRigEditMode.cpp b/Engine/Plugins/Experimental/ControlRig/Source/ControlRigEditor/Private/EditMode/ControlRigEditMode.cpp
--- a/Engine/Plugins/Experimental/ControlRig/Source/ControlRigEditor/Private/EditMode/ControlRigEditMode.cpp
+++ b/Engine/Plugins/Experimental/ControlRig/Source/ControlRigEditor/Private/EditMode/ControlRigEditMode.cpp
@@ -144,14 +144,12 @@ void FControlRigEditMode::SetObjects_Internal()
 	{
 		//Don't add the WeakControlRig Editing...SelectedObjects.Add(WeakControlRigEditing);
 		WeakControlRigEditing.Get()->DrawInterface = &DrawInterface;
-		if (IsInLevelEditor())
-		{
-			WeakControlRigEditing->Hierarchy.OnElementSelected.RemoveAll(this);
-			WeakControlRigEditing->ControlModified().RemoveAll(this);
 
-			WeakControlRigEditing->Hierarchy.OnElementSelected.AddSP(this, &FControlRigEditMode::OnRigElementSelected);
-			WeakControlRigEditing->ControlModified().AddSP(this, &FControlRigEditMode::OnControlModified);
-		}
+		WeakControlRigEditing->Hierarchy.OnElementSelected.RemoveAll(this);
+		WeakControlRigEditing->ControlModified().RemoveAll(this);
+
+		WeakControlRigEditing->Hierarchy.OnElementSelected.AddSP(this, &FControlRigEditMode::OnRigElementSelected);
+		WeakControlRigEditing->ControlModified().AddSP(this, &FControlRigEditMode::OnControlModified);
 
 		// create default manipulation layer
 		RecreateManipulationLayer();
@@ -203,10 +201,6 @@ void FControlRigEditMode::Exit()
 {
 	if (bIsTransacting)
 	{
-		if (ManipulationLayer)
-		{
-			ManipulationLayer->EndTransaction();
-		}
 
 		if (GEditor)
 		{
@@ -467,26 +461,10 @@ bool FControlRigEditMode::EndTracking(FEditorViewportClient* InViewportClient, F
 	{
 		if (bManipulatorMadeChange)
 		{
-			// One final notify of our manipulators to make sure the property is keyed
-			for (AControlRigGizmoActor* GizmoActor : GizmoActors)
-			{
-				if (GizmoActor->IsManipulating())
-				{
-					GizmoActor->SetManipulating(false);
-				}
-			}
-
 			bManipulatorMadeChange = false;
 			GEditor->EndTransaction();
 
 		}
-
-		if (ManipulationLayer)
-		{
-			ManipulationLayer->EndTransaction();
-		}
-
-
 		bIsTransacting = false;
 		return true;
 	}
@@ -501,15 +479,18 @@ bool FControlRigEditMode::StartTracking(FEditorViewportClient* InViewportClient,
 	if (!bIsTransacting)
 	{
 
-		if (ManipulationLayer != nullptr)
-		{
-			ManipulationLayer->BeginTransaction();
-		}
-		for (AControlRigGizmoActor* GizmoActor : GizmoActors)
+		if (WeakControlRigEditing.IsValid())
 		{
-			GizmoActor->SetManipulating(true);
-		}
+			UObject* Blueprint = WeakControlRigEditing->GetClass()->ClassGeneratedBy;
+			if (Blueprint)
+			{
+				Blueprint->SetFlags(RF_Transactional);
+				Blueprint->Modify();
+			}
 
+			WeakControlRigEditing->SetFlags(RF_Transactional);
+			WeakControlRigEditing->Modify();
+		}
 		bIsTransacting = true;
 		bManipulatorMadeChange = false;
 
@@ -595,8 +576,8 @@ bool FControlRigEditMode::HandleClick(FEditorViewportClient* InViewportClient, H
 
 				if (ControlData)
 				{
-					FScopedTransaction ScopedTransaction(LOCTEXT("SelectControlTransaction", "Select Control"), !GIsTransacting);
-
+					FScopedTransaction ScopedTransaction(LOCTEXT("SelectControlTransaction", "Select Control"), IsInLevelEditor() && !GIsTransacting);
+					
 					const FName& ControlName = ControlData->ControlName;
 					if (Click.IsShiftDown() || Click.IsControlDown())
 					{
@@ -633,8 +614,9 @@ bool FControlRigEditMode::HandleClick(FEditorViewportClient* InViewportClient, H
 	{
 		return true;
 	}
-	FScopedTransaction ScopedTransaction(LOCTEXT("SelectControlTransaction", "Select Control"), !GIsTransacting);
 
+	FScopedTransaction ScopedTransaction(LOCTEXT("SelectControlTransaction", "Select Control"), IsInLevelEditor() &&  !GIsTransacting);
+	
 	// clear selected controls
 	ClearRigElementSelection(FRigElementTypeHelper::ToMask(ERigElementType::All));
 
@@ -805,7 +787,7 @@ bool FControlRigEditMode::InputDelta(FEditorViewportClient* InViewportClient, FV
 
 				if (SelectedRigElementType == ERigElementType::Control)
 				{
-					FTransform NewWorldTransform = OnGetRigElementTransformDelegate.Execute(SelectedRigElements[Index], false) * ComponentTransform;
+					FTransform NewWorldTransform = OnGetRigElementTransformDelegate.Execute(SelectedRigElements[Index], false, true) * ComponentTransform;
 					bool bTransformChanged = false;
 					if (bDoRotation)
 					{
@@ -1075,7 +1057,7 @@ void FControlRigEditMode::RecalcPivotTransform()
 		{
 			if (SelectedRigElements[Index].Type == ERigElementType::Control)
 			{
-				LastTransform = OnGetRigElementTransformDelegate.Execute(SelectedRigElements[Index], false);
+				LastTransform = OnGetRigElementTransformDelegate.Execute(SelectedRigElements[Index], false, true);
 				PivotLocation += LastTransform.GetLocation();
 				++NumSelection;
 			}
@@ -1298,18 +1280,21 @@ bool FControlRigEditMode::IsTransformDelegateAvailable() const
 
 bool FControlRigEditMode::AreRigElementSelectedAndMovable() const
 {
-	if (AreRigElementsSelected(FRigElementTypeHelper::ToMask(ERigElementType::Control)))
+	if (!AreRigElementsSelected(FRigElementTypeHelper::ToMask(ERigElementType::Control)))
 	{
-		return true;
+		return false;
 	}
 
-	if (IsTransformDelegateAvailable())
+	//when in sequencer/level we don't have that delegate so don't check.
+	if (!IsInLevelEditor())
 	{
-		return true;
+		if (!IsTransformDelegateAvailable())
+		{
+			return false;
+		}
 	}
 
-	return false;
-
+	return true;
 }
 
 void FControlRigEditMode::OnRigElementAdded(FRigHierarchyContainer* Container, const FRigElementKey& InKey)
@@ -1353,8 +1338,11 @@ void FControlRigEditMode::OnRigElementSelected(FRigHierarchyContainer* Container
 			// if it's control
 			if (InKey.Type == ERigElementType::Control)
 			{
-				FScopedTransaction ScopedTransaction(LOCTEXT("SelectControlTransaction", "Select Control"), !GIsTransacting );
-				ControlProxy->Modify();
+				FScopedTransaction ScopedTransaction(LOCTEXT("SelectControlTransaction", "Select Control"), IsInLevelEditor() && !GIsTransacting);	
+				if (IsInLevelEditor())
+				{
+					ControlProxy->Modify();
+				}
 				// users may select gizmo and control rig units, so we have to let them go through both of them if they do
 				// first go through gizmo actor
 				AControlRigGizmoActor* GizmoActor = GetGizmoFromControlName(InKey.Name);
diff --git a/Engine/Plugins/Experimental/ControlRig/Source/ControlRigEditor/Private/EditMode/ControlRigEditMode.h b/Engine/Plugins/Experimental/ControlRig/Source/ControlRigEditor/Private/EditMode/ControlRigEditMode.h
--- a/Engine/Plugins/Experimental/ControlRig/Source/ControlRigEditor/Private/EditMode/ControlRigEditMode.h
+++ b/Engine/Plugins/Experimental/ControlRig/Source/ControlRigEditor/Private/EditMode/ControlRigEditMode.h
@@ -33,7 +33,7 @@ class IControlRigManipulatable;
 class ISequencer;
 enum class EControlRigSetKey : uint8;
 
-DECLARE_DELEGATE_RetVal_TwoParams(FTransform, FOnGetRigElementTransform, const FRigElementKey& /*RigElementKey*/, bool /*bLocal*/);
+DECLARE_DELEGATE_RetVal_ThreeParams(FTransform, FOnGetRigElementTransform, const FRigElementKey& /*RigElementKey*/, bool /*bLocal*/, bool /*bOnDebugInstance*/);
 DECLARE_DELEGATE_ThreeParams(FOnSetRigElementTransform, const FRigElementKey& /*RigElementKey*/, const FTransform& /*Transform*/, bool /*bLocal*/);
 DECLARE_DELEGATE_RetVal(TSharedPtr<FUICommandList>, FNewMenuCommandsDelegate);
 
diff --git a/Engine/Plugins/Experimental/ControlRig/Source/ControlRigEditor/Private/EditMode/SControlRigEditModeTools.cpp b/Engine/Plugins/Experimental/ControlRig/Source/ControlRigEditor/Private/EditMode/SControlRigEditModeTools.cpp
--- a/Engine/Plugins/Experimental/ControlRig/Source/ControlRigEditor/Private/EditMode/SControlRigEditModeTools.cpp
+++ b/Engine/Plugins/Experimental/ControlRig/Source/ControlRigEditor/Private/EditMode/SControlRigEditModeTools.cpp
@@ -164,8 +164,12 @@ void SControlRigEditModeTools::SetSequencer(TWeakPtr<ISequencer> InSequencer)
 
 bool SControlRigEditModeTools::IsPropertyKeyable(UClass* InObjectClass, const IPropertyHandle& InPropertyHandle) const
 {
+	if (InObjectClass && InObjectClass->IsChildOf(UControlRigTransformNoScaleControlProxy::StaticClass()) && InPropertyHandle.GetProperty() 
+		&& InPropertyHandle.GetProperty()->GetFName() == GET_MEMBER_NAME_CHECKED(UControlRigTransformControlProxy, Transform)) 
+	{
+		return true;
+	}
 	FCanKeyPropertyParams CanKeyPropertyParams(InObjectClass, InPropertyHandle);
-
 	TSharedPtr<ISequencer> Sequencer = WeakSequencer.Pin();
 	if (Sequencer.IsValid() && Sequencer->CanKeyProperty(CanKeyPropertyParams))
 	{
diff --git a/Engine/Plugins/Experimental/ControlRig/Source/ControlRigEditor/Private/Editor/ControlRigEditor.cpp b/Engine/Plugins/Experimental/ControlRig/Source/ControlRigEditor/Private/Editor/ControlRigEditor.cpp
--- a/Engine/Plugins/Experimental/ControlRig/Source/ControlRigEditor/Private/Editor/ControlRigEditor.cpp
+++ b/Engine/Plugins/Experimental/ControlRig/Source/ControlRigEditor/Private/Editor/ControlRigEditor.cpp
@@ -9,7 +9,6 @@
 #include "SKismetInspector.h"
 #include "Framework/Commands/GenericCommands.h"
 #include "Editor.h"
-#include "Graph/ControlRigGraphNode.h"
 #include "Graph/ControlRigGraph.h"
 #include "BlueprintActionDatabase.h"
 #include "ControlRigBlueprintCommands.h"
@@ -153,6 +152,8 @@ void FControlRigEditor::InitControlRigEditor(const EToolkitMode::Type Mode, cons
 {
 	DECLARE_SCOPE_HIERARCHICAL_COUNTER_FUNC()
 
+	check(InControlRigBlueprint);
+
 	FPersonaModule& PersonaModule = FModuleManager::GetModuleChecked<FPersonaModule>("Persona");
 
 	FPersonaToolkitArgs PersonaToolkitArgs;
@@ -245,7 +246,10 @@ void FControlRigEditor::InitControlRigEditor(const EToolkitMode::Type Mode, cons
 		
 		PersonaToolkit->GetPreviewScene()->SetRemoveAttachedComponentFilter(FOnRemoveAttachedComponentFilter::CreateSP(EditMode, &FControlRigEditMode::CanRemoveFromPreviewScene));
 
-		InControlRigBlueprint->OnModified().AddSP(EditMode, &FControlRigEditMode::HandleModifiedEvent);
+		if (InControlRigBlueprint)
+		{
+			InControlRigBlueprint->OnModified().AddSP(EditMode, &FControlRigEditMode::HandleModifiedEvent);
+		}
 	}
 
 	UpdateControlRig();
@@ -255,15 +259,23 @@ void FControlRigEditor::InitControlRigEditor(const EToolkitMode::Type Mode, cons
 
 	if (ControlRigBlueprints.Num() > 0)
 	{
+		bool bBroughtGraphToFront = false;
 		for(UEdGraph* Graph : ControlRigBlueprints[0]->UbergraphPages)
 		{
 			if (Graph->GetFName().IsEqual(UControlRigGraphSchema::GraphName_ControlRig))
 			{
-				OpenGraphAndBringToFront(Graph, false);
-				break;
+				if (!bBroughtGraphToFront)
+				{
+					OpenGraphAndBringToFront(Graph, false);
+					bBroughtGraphToFront = true;
+				}
 			}
-		}
 
+			if (UControlRigGraph* RigGraph = Cast<UControlRigGraph>(Graph))
+			{
+				RigGraph->OnGraphNodeClicked.AddSP(this, &FControlRigEditor::OnGraphNodeClicked);
+			}
+		}
 	}
 
 	if (InControlRigBlueprint)
@@ -648,6 +660,7 @@ void FControlRigEditor::SetDetailObjects(const TArray<UObject*>& InObjects)
 	ClearDetailObject();
 
 	RigElementInDetailPanel = FRigElementKey();
+	StructToDisplay.Reset();
 	SKismetInspector::FShowDetailsOptions Options;
 	Options.bForceRefresh = true;
 	Inspector->ShowDetailsForObjects(InObjects);
@@ -669,7 +682,7 @@ void FControlRigEditor::SetDetailObject(UObject* Obj)
 			FString StructDefaultValue = StructNode->GetStructDefaultValue();
 			NodeDetailStruct->ImportText(*StructDefaultValue, NodeDetailBuffer.GetData(), nullptr, PPF_None, nullptr, NodeDetailStruct->GetName());
 
-			TSharedPtr<FStructOnScope> StructToDisplay = MakeShareable(new FStructOnScope(NodeDetailStruct, (uint8*)NodeDetailBuffer.GetData()));
+			StructToDisplay = MakeShareable(new FStructOnScope(NodeDetailStruct, (uint8*)NodeDetailBuffer.GetData()));
 			if (StructToDisplay.IsValid())
 			{
 				StructToDisplay->SetPackage(GetControlRigBlueprint()->GetOutermost());
@@ -745,7 +758,7 @@ void FControlRigEditor::SetDetailObject(UObject* Obj)
 	SetDetailObjects(Objects);
 }
 
-void FControlRigEditor::SetDetailStruct(const FRigElementKey& InElement, TSharedPtr<FStructOnScope> StructToDisplay)
+void FControlRigEditor::SetDetailStruct(const FRigElementKey& InElement)
 {
 	if (RigElementInDetailPanel == InElement)
 	{
@@ -753,6 +766,40 @@ void FControlRigEditor::SetDetailStruct(const FRigElementKey& InElement, TShared
 	}
 
 	ClearDetailObject();
+	
+	UControlRigBlueprint* RigBlueprint = Cast<UControlRigBlueprint>(GetBlueprintObj());
+
+	switch (InElement.Type)
+	{
+		case ERigElementType::Bone:
+		{
+			FRigBoneHierarchy& BoneHierarchy = RigBlueprint->HierarchyContainer.BoneHierarchy;
+			StructToDisplay = MakeShareable(new FStructOnScope(FRigBone::StaticStruct(), (uint8*)&(BoneHierarchy[InElement.Name])));
+			break;
+		}
+		case ERigElementType::Control:
+		{
+			FRigControlHierarchy& ControlHierarchy = RigBlueprint->HierarchyContainer.ControlHierarchy;
+			StructToDisplay = MakeShareable(new FStructOnScope(FRigControl::StaticStruct(), (uint8*)&(ControlHierarchy[InElement.Name])));
+			break;
+		}
+		case ERigElementType::Space:
+		{
+			FRigSpaceHierarchy& SpaceHierarchy = RigBlueprint->HierarchyContainer.SpaceHierarchy;
+			StructToDisplay = MakeShareable(new FStructOnScope(FRigSpace::StaticStruct(), (uint8*)&(SpaceHierarchy[InElement.Name])));
+			break;
+		}
+		case ERigElementType::Curve:
+		{
+			FRigCurveContainer& CurveContainer = RigBlueprint->HierarchyContainer.CurveContainer;
+			StructToDisplay = MakeShareable(new FStructOnScope(FRigBone::StaticStruct(), (uint8*)&(CurveContainer[InElement.Name])));
+			break;
+		}
+		default:
+		{
+			break;
+		}
+	}
 
 	RigElementInDetailPanel = InElement;
 	if (StructToDisplay.IsValid())
@@ -808,7 +855,15 @@ void FControlRigEditor::Compile()
 		GetBlueprintObj()->SetObjectBeingDebugged(nullptr);
 
 		FRigElementKey SelectedKey = RigElementInDetailPanel;
-		ClearDetailObject();
+		TArray< TWeakObjectPtr<UObject> > SelectedObjects;
+		if (SelectedKey.IsValid())
+		{
+			ClearDetailObject();
+		}
+		else
+		{
+			SelectedObjects = Inspector->GetSelectedObjects();
+		}
 
 		if (ControlRig)
 		{
@@ -855,35 +910,15 @@ void FControlRigEditor::Compile()
 
 		if(SelectedKey.IsValid())
 		{
-			UControlRigBlueprint* RigBlueprint = Cast<UControlRigBlueprint>(GetBlueprintObj());
-			switch (SelectedKey.Type)
+			SetDetailStruct(SelectedKey);
+		}
+		else if (SelectedObjects.Num() > 0)
+		{
+			for (TWeakObjectPtr<UObject> SelectedObject : SelectedObjects)
 			{
-				case ERigElementType::Bone:
-				{
-					FRigBoneHierarchy& BoneHierarchy = RigBlueprint->HierarchyContainer.BoneHierarchy;
-					SetDetailStruct(SelectedKey, MakeShareable(new FStructOnScope(FRigBone::StaticStruct(), (uint8*)&(BoneHierarchy[SelectedKey.Name]))));
-					break;
-				}
-				case ERigElementType::Control:
-				{
-					FRigControlHierarchy& ControlHierarchy = RigBlueprint->HierarchyContainer.ControlHierarchy;
-					SetDetailStruct(SelectedKey, MakeShareable(new FStructOnScope(FRigControl::StaticStruct(), (uint8*)&(ControlHierarchy[SelectedKey.Name]))));
-					break;
-				}
-				case ERigElementType::Space:
-				{
-					FRigSpaceHierarchy& SpaceHierarchy = RigBlueprint->HierarchyContainer.SpaceHierarchy;
-					SetDetailStruct(SelectedKey, MakeShareable(new FStructOnScope(FRigSpace::StaticStruct(), (uint8*)&(SpaceHierarchy[SelectedKey.Name]))));
-					break;
-				}
-				case ERigElementType::Curve:
-				{
-					FRigCurveContainer& CurveContainer = RigBlueprint->HierarchyContainer.CurveContainer;
-					SetDetailStruct(SelectedKey, MakeShareable(new FStructOnScope(FRigBone::StaticStruct(), (uint8*)&(CurveContainer[SelectedKey.Name]))));
-					break;
-				}
-				default:
+				if (SelectedObject.IsValid())
 				{
+					SetDetailObject(SelectedObject.Get());
 					break;
 				}
 			}
@@ -1074,12 +1109,22 @@ void FControlRigEditor::PostUndo(bool bSuccess)
 {
 	IControlRigEditor::PostUndo(bSuccess);
 	EnsureValidRigElementInDetailPanel();
+
+	if (FControlRigEditMode* EditMode = GetEditMode())
+	{
+		EditMode->RecreateManipulationLayer();
+	}
 }
 
 void FControlRigEditor::PostRedo(bool bSuccess)
 {
 	IControlRigEditor::PostRedo(bSuccess);
 	EnsureValidRigElementInDetailPanel();
+
+	if (FControlRigEditMode* EditMode = GetEditMode())
+	{
+		EditMode->RecreateManipulationLayer();
+	}
 }
 
 void FControlRigEditor::EnsureValidRigElementInDetailPanel()
@@ -1259,6 +1304,18 @@ void FControlRigEditor::HandleModifiedEvent(ERigVMGraphNotifType InNotifType, UR
 						SetDetailObject(Node);
 					}
 
+					// if we used to have a rig unit selected, clear the details panel
+					else if(InGraph->GetSelectNodes().Num() == 0)
+					{
+						if (StructToDisplay.IsValid())
+						{
+							if (StructToDisplay->GetStruct()->IsChildOf(FRigUnit::StaticStruct()))
+							{
+								ClearDetailObject();
+							}
+						}
+					}
+
 					if (!RigGraph->bIsSelecting)
 					{
 						TGuardValue<bool> SelectingGuard(RigGraph->bIsSelecting, true);
@@ -2041,9 +2098,28 @@ void FControlRigEditor::OnNodeTitleCommitted(const FText& NewText, ETextCommit::
 	}
 }
 
-FTransform FControlRigEditor::GetRigElementTransform(const FRigElementKey& InElement, bool bLocal) const
+FTransform FControlRigEditor::GetRigElementTransform(const FRigElementKey& InElement, bool bLocal, bool bOnDebugInstance) const
 {
 	DECLARE_SCOPE_HIERARCHICAL_COUNTER_FUNC()
+
+	if(bOnDebugInstance)
+	{
+		UControlRig* DebuggedControlRig = Cast<UControlRig>(GetBlueprintObj()->GetObjectBeingDebugged());
+		if (DebuggedControlRig == nullptr)
+		{
+			DebuggedControlRig = ControlRig;
+		}
+
+		if (DebuggedControlRig)
+		{
+			if (bLocal)
+			{
+				return DebuggedControlRig->GetHierarchy()->GetLocalTransform(InElement);
+			}
+			return DebuggedControlRig->GetHierarchy()->GetGlobalTransform(InElement);
+		}
+	}
+
 	if (bLocal)
 	{
 		return GetControlRigBlueprint()->HierarchyContainer.GetLocalTransform(InElement);
@@ -2070,7 +2146,7 @@ void FControlRigEditor::SetRigElementTransform(const FRigElementKey& InElement,
 				FRigElementKey ParentKey = ControlRigBP->HierarchyContainer.BoneHierarchy[InElement.Name].GetParentElementKey();
 				if (ParentKey.IsValid())
 				{
-					ParentTransform = GetRigElementTransform(ParentKey, false);
+					ParentTransform = GetRigElementTransform(ParentKey, false, false);
 				}
 				Transform = Transform * ParentTransform;
 				Transform.NormalizeRotation();
@@ -2353,6 +2429,43 @@ void FControlRigEditor::OnRigElementAdded(FRigHierarchyContainer* Container, con
 
 void FControlRigEditor::OnRigElementRemoved(FRigHierarchyContainer* Container, const FRigElementKey& InKey)
 {
+	UControlRigBlueprint* Blueprint = GetControlRigBlueprint();
+	for (UEdGraph* Graph : Blueprint->UbergraphPages)
+	{
+		UControlRigGraph* RigGraph = Cast<UControlRigGraph>(Graph);
+		if (RigGraph == nullptr)
+		{
+			continue;
+		}
+
+		FString RemovedElementName = InKey.Name.ToString();
+		ERigElementType RemovedElementType = InKey.Type;
+
+		for (UEdGraphNode* Node : RigGraph->Nodes)
+		{
+			if (UControlRigGraphNode* RigNode = Cast<UControlRigGraphNode>(Node))
+			{
+				if (URigVMNode* ModelNode = RigNode->GetModelNode())
+				{
+					for (URigVMPin * ModelPin : ModelNode->GetPins())
+					{
+						if ((ModelPin->GetCPPType() == TEXT("FName") && ModelPin->GetCustomWidgetName() == TEXT("BoneName") && RemovedElementType == ERigElementType::Bone) ||
+							(ModelPin->GetCPPType() == TEXT("FName") && ModelPin->GetCustomWidgetName() == TEXT("ControlName") && RemovedElementType == ERigElementType::Control) ||
+							(ModelPin->GetCPPType() == TEXT("FName") && ModelPin->GetCustomWidgetName() == TEXT("SpaceName") && RemovedElementType == ERigElementType::Space) ||
+							(ModelPin->GetCPPType() == TEXT("FName") && ModelPin->GetCustomWidgetName() == TEXT("CurveName") && RemovedElementType == ERigElementType::Curve))
+						{
+							if (ModelPin->GetDefaultValue() == RemovedElementName)
+							{
+								Blueprint->Controller->SetPinDefaultValue(ModelPin->GetPinPath(), FName(NAME_None).ToString());
+							}
+						}
+					}
+				}
+			}
+		}
+	}
+
+
 	OnHierarchyChanged();
 }
 
@@ -2410,6 +2523,11 @@ void FControlRigEditor::OnRigElementSelected(FRigHierarchyContainer* Container,
 		return;
 	}
 
+	if (Container->GetIndex(InKey) == INDEX_NONE)
+	{
+		return;
+	}
+
 	if (InKey.Type == ERigElementType::Bone)
 	{
 		UControlRigSkeletalMeshComponent* EditorSkelComp = Cast<UControlRigSkeletalMeshComponent>(GetPersonaToolkit()->GetPreviewScene()->GetPreviewMeshComponent());
@@ -2432,37 +2550,7 @@ void FControlRigEditor::OnRigElementSelected(FRigHierarchyContainer* Container,
 
 	if (bSelected)
 	{
-		switch (InKey.Type)
-		{
-			case ERigElementType::Bone:
-			{
-				FRigBoneHierarchy& BoneHierarchy = RigBlueprint->HierarchyContainer.BoneHierarchy;
-				SetDetailStruct(InKey, MakeShareable(new FStructOnScope(FRigBone::StaticStruct(), (uint8*)&(BoneHierarchy[InKey.Name]))));
-				break;
-			}
-			case ERigElementType::Space:
-			{
-				FRigSpaceHierarchy& SpaceHierarchy = RigBlueprint->HierarchyContainer.SpaceHierarchy;
-				SetDetailStruct(InKey, MakeShareable(new FStructOnScope(FRigSpace::StaticStruct(), (uint8*)&(SpaceHierarchy[InKey.Name]))));
-				break;
-			}
-			case ERigElementType::Control:
-			{
-				FRigControlHierarchy& ControlHierarchy = RigBlueprint->HierarchyContainer.ControlHierarchy;
-				int32 ControlIndex = ControlHierarchy.GetIndex(InKey.Name);
-				if (ControlIndex != INDEX_NONE)
-				{
-					SetDetailStruct(InKey, MakeShareable(new FStructOnScope(FRigControl::StaticStruct(), (uint8*)&(ControlHierarchy[ControlIndex]))));
-				}
-				break;
-			}
-			case ERigElementType::Curve:
-			{
-				FRigCurveContainer& CurveContainer = RigBlueprint->HierarchyContainer.CurveContainer;
-				SetDetailStruct(InKey, MakeShareable(new FStructOnScope(FRigCurve::StaticStruct(), (uint8*)&(CurveContainer[InKey.Name]))));
-				break;
-			}
-		}
+		SetDetailStruct(InKey);
 	}
 	else
 	{
@@ -3093,6 +3181,17 @@ void FControlRigEditor::HandleOnControlModified(IControlRigManipulatable* Subjec
 	}
 }
 
+void FControlRigEditor::OnGraphNodeClicked(UControlRigGraphNode* InNode)
+{
+	if (InNode)
+	{
+		if (InNode->IsSelectedInEditor())
+		{
+			SetDetailObject(InNode->GetModelNode());
+		}
+	}
+}
+
 void FControlRigEditor::UpdateGraphCompilerErrors()
 {
 	DECLARE_SCOPE_HIERARCHICAL_COUNTER_FUNC()
diff --git a/Engine/Plugins/Experimental/ControlRig/Source/ControlRigEditor/Private/Editor/ControlRigEditor.h b/Engine/Plugins/Experimental/ControlRig/Source/ControlRigEditor/Private/Editor/ControlRigEditor.h
--- a/Engine/Plugins/Experimental/ControlRig/Source/ControlRigEditor/Private/Editor/ControlRigEditor.h
+++ b/Engine/Plugins/Experimental/ControlRig/Source/ControlRigEditor/Private/Editor/ControlRigEditor.h
@@ -17,6 +17,7 @@
 #include "Styling/SlateTypes.h"
 #include "AnimPreviewInstance.h"
 #include "ScopedTransaction.h"
+#include "Graph/ControlRigGraphNode.h"
 
 class UControlRigBlueprint;
 class IPersonaToolkit;
@@ -117,7 +118,7 @@ public:
 
 	void SetDetailObject(UObject* Obj);
 
-	void SetDetailStruct(const FRigElementKey& InElement, TSharedPtr<FStructOnScope> StructToDisplay);
+	void SetDetailStruct(const FRigElementKey& InElement);
 
 	void ClearDetailObject();
 
@@ -252,6 +253,8 @@ private:
 
 	 void HandleOnControlModified(IControlRigManipulatable* Subject, const FRigControl& Control, EControlRigSetKey InSetKey);
 
+	 void OnGraphNodeClicked(UControlRigGraphNode* InNode);
+
 protected:
 
 	/** Toolbox hosting widget */
@@ -288,7 +291,7 @@ protected:
 	TSharedPtr<FUICommandList> HandleOnViewportContextMenuCommandsDelegate();
 
 	/** Bone Selection related */
-	FTransform GetRigElementTransform(const FRigElementKey& InElement, bool bLocal) const;
+	FTransform GetRigElementTransform(const FRigElementKey& InElement, bool bLocal, bool bOnDebugInstance) const;
 	void SetRigElementTransform(const FRigElementKey& InElement, const FTransform& InTransform, bool bLocal);
 	
 	// FNotifyHook Interface
@@ -300,6 +303,7 @@ protected:
 	bool bControlRigEditorInitialized;
 	bool bIsSettingObjectBeingDebugged;
 	FRigElementKey RigElementInDetailPanel;
+	TSharedPtr<FStructOnScope> StructToDisplay;
 
 	TArray<uint8, TAlignedHeapAllocator<16>> NodeDetailBuffer;
 	UScriptStruct* NodeDetailStruct;
diff --git a/Engine/Plugins/Experimental/ControlRig/Source/ControlRigEditor/Private/Editor/ControlRigEditorEditMode.cpp b/Engine/Plugins/Experimental/ControlRig/Source/ControlRigEditor/Private/Editor/ControlRigEditorEditMode.cpp
--- a/Engine/Plugins/Experimental/ControlRig/Source/ControlRigEditor/Private/Editor/ControlRigEditorEditMode.cpp
+++ b/Engine/Plugins/Experimental/ControlRig/Source/ControlRigEditor/Private/Editor/ControlRigEditorEditMode.cpp
@@ -15,12 +15,12 @@ bool FControlRigEditorEditMode::GetCameraTarget(FSphere& OutTarget) const
 	{
 		if (SelectedRigElements[Index].Type == ERigElementType::Bone || SelectedRigElements[Index].Type == ERigElementType::Space)
 		{
-			FTransform Transform = OnGetRigElementTransformDelegate.Execute(SelectedRigElements[Index], false);
+			FTransform Transform = OnGetRigElementTransformDelegate.Execute(SelectedRigElements[Index], false, true);
 			Box += Transform.GetLocation();
 		}
 		else if (SelectedRigElements[Index].Type == ERigElementType::Control)
 		{
-			FTransform Transform = OnGetRigElementTransformDelegate.Execute(SelectedRigElements[Index], false);
+			FTransform Transform = OnGetRigElementTransformDelegate.Execute(SelectedRigElements[Index], false, true);
 			Box += Transform.TransformPosition(FVector::OneVector * 50.f); // we assume 100cm as the base size for a control
 			Box += Transform.TransformPosition(FVector::OneVector * -50.f);
 		}
diff --git a/Engine/Plugins/Experimental/ControlRig/Source/ControlRigEditor/Private/Editor/SRigHierarchy.cpp b/Engine/Plugins/Experimental/ControlRig/Source/ControlRigEditor/Private/Editor/SRigHierarchy.cpp
--- a/Engine/Plugins/Experimental/ControlRig/Source/ControlRigEditor/Private/Editor/SRigHierarchy.cpp
+++ b/Engine/Plugins/Experimental/ControlRig/Source/ControlRigEditor/Private/Editor/SRigHierarchy.cpp
@@ -31,9 +31,11 @@
 #include "ControlRig.h"
 #include "ControlRigEditorStyle.h"
 #include "HAL/PlatformApplicationMisc.h"
+#include "HAL/PlatformTime.h"
 #include "Dialogs/Dialogs.h"
 #include "IPersonaToolkit.h"
 #include "SKismetInspector.h"
+#include "Types/WidgetActiveTimerDelegate.h"
 
 #define LOCTEXT_NAMESPACE "SRigHierarchy"
 
@@ -323,13 +325,14 @@ void SRigHierarchy::Construct(const FArguments& InArgs, TSharedRef<FControlRigEd
 			.Padding(2.0f)
 			.BorderImage(FEditorStyle::GetBrush("SCSEditor.TreePanel"))
 			[
-				SAssignNew(TreeView, STreeView<TSharedPtr<FRigTreeElement>>)
+				SAssignNew(TreeView, SRigHierarchyTreeView)
 				.TreeItemsSource(&RootElements)
 				.SelectionMode(ESelectionMode::Multi)
 				.OnGenerateRow(this, &SRigHierarchy::MakeTableRowWidget)
 				.OnGetChildren(this, &SRigHierarchy::HandleGetChildrenForTree)
 				.OnSelectionChanged(this, &SRigHierarchy::OnSelectionChanged)
 				.OnContextMenuOpening(this, &SRigHierarchy::CreateContextMenu)
+				.OnMouseButtonClick(this, &SRigHierarchy::OnItemClicked)
 				.OnMouseButtonDoubleClick(this, &SRigHierarchy::OnItemDoubleClicked)
 				.HighlightParentNodesForSelection(true)
 				.ItemHeight(24)
@@ -843,6 +846,39 @@ TSharedPtr< SWidget > SRigHierarchy::CreateContextMenu()
 	return MenuBuilder.MakeWidget();
 }
 
+void SRigHierarchy::OnItemClicked(TSharedPtr<FRigTreeElement> InItem)
+{
+	FRigHierarchyContainer* Hierarchy = GetHierarchyContainer();
+	if (Hierarchy->IsSelected(InItem->Key))
+	{
+		if (ControlRigEditor.IsValid())
+		{
+			ControlRigEditor.Pin()->SetDetailStruct(InItem->Key);
+		}
+
+		if (InItem->Key.Type == ERigElementType::Bone)
+		{
+			const FRigBone& Bone = Hierarchy->BoneHierarchy[InItem->Key.Name];
+			if (Bone.Type == ERigBoneType::Imported)
+			{
+				return;
+			}
+		}
+
+		uint32 CurrentCycles = FPlatformTime::Cycles();
+		double SecondsPassed = double(CurrentCycles - TreeView->LastClickCycles) * FPlatformTime::GetSecondsPerCycle();
+		if (SecondsPassed > 0.5f)
+		{
+			RegisterActiveTimer(0.f, FWidgetActiveTimerDelegate::CreateLambda([this](double, float) {
+				HandleRenameItem();
+				return EActiveTimerReturnType::Stop;
+			}));
+		}
+
+		TreeView->LastClickCycles = CurrentCycles;
+	}
+}
+
 void SRigHierarchy::OnItemDoubleClicked(TSharedPtr<FRigTreeElement> InItem)
 {
 	if (TreeView->IsItemExpanded(InItem))
@@ -855,14 +891,13 @@ void SRigHierarchy::OnItemDoubleClicked(TSharedPtr<FRigTreeElement> InItem)
 	}
 }
 
-
 void SRigHierarchy::FillContextMenu(class FMenuBuilder& MenuBuilder)
 {
 	const FControlRigHierarchyCommands& Actions = FControlRigHierarchyCommands::Get();
 	{
 		struct FLocalMenuBuilder
 		{
-			static void FillNewMenu(FMenuBuilder& InSubMenuBuilder, TSharedPtr<STreeView<TSharedPtr<FRigTreeElement>>> InTreeView)
+			static void FillNewMenu(FMenuBuilder& InSubMenuBuilder, TSharedPtr<SRigHierarchyTreeView> InTreeView)
 			{
 				const FControlRigHierarchyCommands& Actions = FControlRigHierarchyCommands::Get();
 
@@ -1334,6 +1369,18 @@ void SRigHierarchy::HandleDuplicateItem()
 					FRigControl Control = Hierarchy->ControlHierarchy[Key.Name];
 					const FName NewName = CreateUniqueName(Key.Name, Key.Type);
 					FRigControl& NewControl = Hierarchy->ControlHierarchy.Add(NewName, Control.ControlType, Control.ParentName, Control.SpaceName, Control.InitialValue, Control.GizmoName, Control.GizmoTransform, Control.GizmoColor);
+
+					// copy additional members
+					NewControl.bAnimatable = Control.bAnimatable;
+					NewControl.PrimaryAxis = Control.PrimaryAxis;
+					NewControl.bLimitTranslation = Control.bLimitTranslation;
+					NewControl.bLimitRotation = Control.bLimitRotation;
+					NewControl.bLimitScale = Control.bLimitScale;
+					NewControl.MinimumValue = Control.MinimumValue;
+					NewControl.MaximumValue = Control.MaximumValue;
+					NewControl.bDrawLimits = Control.bDrawLimits;
+					NewControl.bGizmoEnabled = Control.bGizmoEnabled;
+
 					NewKeys.Add(NewControl.GetElementKey());
 					break;
 				}
@@ -1357,6 +1404,8 @@ void SRigHierarchy::HandleDuplicateItem()
 		{
 			Hierarchy->Select(NewKeys[Index]);
 		}
+		
+		ControlRigBlueprint->PropagateHierarchyFromBPToInstances(false);
 	}
 
 	FSlateApplication::Get().DismissAllMenus();
@@ -1372,11 +1421,14 @@ bool SRigHierarchy::CanRenameItem() const
 /** Delete Item */
 void SRigHierarchy::HandleRenameItem()
 {
+	if (!CanRenameItem())
+	{
+		return;
+	}
+
 	FRigHierarchyContainer* Hierarchy = GetHierarchyContainer();
 	if (Hierarchy)
 	{
-		ClearDetailPanel();
-
 		FScopedTransaction Transaction(LOCTEXT("HierarchyTreeRenameSelected", "Rename selected item from hierarchy"));
 		ControlRigBlueprint->Modify();
 
diff --git a/Engine/Plugins/Experimental/ControlRig/Source/ControlRigEditor/Private/Editor/SRigHierarchy.h b/Engine/Plugins/Experimental/ControlRig/Source/ControlRigEditor/Private/Editor/SRigHierarchy.h
--- a/Engine/Plugins/Experimental/ControlRig/Source/ControlRigEditor/Private/Editor/SRigHierarchy.h
+++ b/Engine/Plugins/Experimental/ControlRig/Source/ControlRigEditor/Private/Editor/SRigHierarchy.h
@@ -107,6 +107,24 @@ struct FRigHierarchyImportSettings
 	USkeletalMesh* Mesh;
 };
 
+class SRigHierarchyTreeView : public STreeView<TSharedPtr<FRigTreeElement>>
+{
+public:
+
+	virtual ~SRigHierarchyTreeView() {}
+
+	virtual FReply OnFocusReceived(const FGeometry& MyGeometry, const FFocusEvent& InFocusEvent) override
+	{
+		FReply Reply = STreeView<TSharedPtr<FRigTreeElement>>::OnFocusReceived(MyGeometry, InFocusEvent);
+
+		LastClickCycles = FPlatformTime::Cycles();
+
+		return Reply;
+	}
+
+	uint32 LastClickCycles = 0;
+};
+
 /** Widget allowing editing of a control rig's structure */
 class SRigHierarchy : public SCompoundWidget, public FEditorUndoClient
 {
@@ -166,6 +184,7 @@ private:
 	void OnSelectionChanged(TSharedPtr<FRigTreeElement> Selection, ESelectInfo::Type SelectInfo);
 
 	TSharedPtr< SWidget > CreateContextMenu();
+	void OnItemClicked(TSharedPtr<FRigTreeElement> InItem);
 	void OnItemDoubleClicked(TSharedPtr<FRigTreeElement> InItem);
 
 	// FEditorUndoClient
@@ -195,7 +214,7 @@ private:
 	void OnFilterTextChanged(const FText& SearchText);
 
 	/** Tree view widget */
-	TSharedPtr<STreeView<TSharedPtr<FRigTreeElement>>> TreeView;
+	TSharedPtr<SRigHierarchyTreeView> TreeView;
 
 	/** Backing array for tree view */
 	TArray<TSharedPtr<FRigTreeElement>> RootElements;
diff --git a/Engine/Plugins/Experimental/ControlRig/Source/ControlRigEditor/Private/Graph/SControlRigGraphNode.cpp b/Engine/Plugins/Experimental/ControlRig/Source/ControlRigEditor/Private/Graph/SControlRigGraphNode.cpp
--- a/Engine/Plugins/Experimental/ControlRig/Source/ControlRigEditor/Private/Graph/SControlRigGraphNode.cpp
+++ b/Engine/Plugins/Experimental/ControlRig/Source/ControlRigEditor/Private/Graph/SControlRigGraphNode.cpp
@@ -33,11 +33,21 @@
 
 #define LOCTEXT_NAMESPACE "SControlRigGraphNode"
 
+const FSlateBrush* SControlRigGraphNode::CachedImg_CR_Pin_Connected = nullptr;
+const FSlateBrush* SControlRigGraphNode::CachedImg_CR_Pin_Disconnected = nullptr;
+
 void SControlRigGraphNode::Construct( const FArguments& InArgs )
 {
+	if (CachedImg_CR_Pin_Connected == nullptr)
+	{
+		static const FName NAME_CR_Pin_Connected("ControlRig.Bug.Solid");
+		static const FName NAME_CR_Pin_Disconnected("ControlRig.Bug.Open");
+		CachedImg_CR_Pin_Connected = FControlRigEditorStyle::Get().GetBrush(NAME_CR_Pin_Connected);
+		CachedImg_CR_Pin_Disconnected = FControlRigEditorStyle::Get().GetBrush(NAME_CR_Pin_Disconnected);
+	}
+
 	check(InArgs._GraphNodeObj);
 	this->GraphNode = InArgs._GraphNodeObj;
-
 	this->SetCursor( EMouseCursor::CardinalCross );
 
  	UControlRigGraphNode* ControlRigGraphNode = InArgs._GraphNodeObj;
@@ -147,7 +157,7 @@ void SControlRigGraphNode::Construct( const FArguments& InArgs )
 	InputOutputTree->Tick(DummyGeometry, 0.f, 0.f);
 	OutputTree->Tick(DummyGeometry, 0.f, 0.f);
 
-	const FSlateBrush* ImageBrush = FControlRigEditorStyle::Get().GetBrush(TEXT("ControlRig.RigUnit.VisualDebug"));
+	const FSlateBrush* ImageBrush = FControlRigEditorStyle::Get().GetBrush(TEXT("ControlRig.Bug.Dot"));
 
 	VisualDebugIndicatorWidget =
 		SNew(SImage)
@@ -259,6 +269,19 @@ void SControlRigGraphNode::AddPin(const TSharedRef<SGraphPin>& PinToAdd)
 			FullPinHorizontalRowWidget->RemoveSlot(LabelAndValueWidget.ToSharedRef());
 		}
 
+		// Customize the look for pins with injected nodes
+		FString NodeName, PinPath;
+		if (URigVMPin::SplitPinPathAtStart(EdPinObj->GetName(), NodeName, PinPath))
+		{
+			if (URigVMPin* ModelPin = ModelNode->FindPin(PinPath))
+			{
+				if (ModelPin->HasInjectedNodes())
+				{
+					PinToAdd->SetCustomPinIcon(CachedImg_CR_Pin_Connected, CachedImg_CR_Pin_Disconnected);
+				}
+			}
+		}
+
 		PinToAdd->SetOwner(SharedThis(this));
 		PinWidgetMap.Add(EdPinObj, PinToAdd);
 		if(EdPinObj->Direction == EGPD_Input)
@@ -277,6 +300,21 @@ const FSlateBrush * SControlRigGraphNode::GetNodeBodyBrush() const
 	return FEditorStyle::GetBrush("Graph.Node.TintedBody");
 }
 
+FReply SControlRigGraphNode::OnMouseButtonDown(const FGeometry& MyGeometry, const FPointerEvent& MouseEvent)
+{
+	FReply Reply = SGraphNode::OnMouseButtonDown(MyGeometry, MouseEvent);
+
+	if (UControlRigGraphNode* RigNode = Cast<UControlRigGraphNode>(GraphNode))
+	{
+		if (UControlRigGraph* RigGraph = Cast<UControlRigGraph>(RigNode->GetGraph()))
+		{
+			RigGraph->OnGraphNodeClicked.Broadcast(RigNode);
+		}
+	}
+
+	return Reply;
+}
+
 bool SControlRigGraphNode::UseLowDetailNodeTitles() const
 {
 	return ParentUseLowDetailNodeTitles();
@@ -889,11 +927,10 @@ TArray<FOverlayWidgetInfo> SControlRigGraphNode::GetOverlayWidgets(bool bSelecte
 
 								if (Widgets.Num() == PreviousNumWidgets)
 								{
-									const FSlateBrush* ImageBrush = FEditorStyle::Get().GetBrush(TEXT("ControlRig.RigUnit.VisualDebug"));
 									FVector2D ImageSize = VisualDebugIndicatorWidget->GetDesiredSize();
 
 									FOverlayWidgetInfo Info;
-									Info.OverlayOffset = FVector2D(WidgetSize.X - (ImageSize.X * 0.8f), -(ImageSize.Y * 0.5f));
+									Info.OverlayOffset = FVector2D(WidgetSize.X - ImageSize.X - 6.f, 6.f);
 									Info.Widget = VisualDebugIndicatorWidget;
 
 									Widgets.Add(Info);
diff --git a/Engine/Plugins/Experimental/ControlRig/Source/ControlRigEditor/Private/Graph/SControlRigGraphNode.h b/Engine/Plugins/Experimental/ControlRig/Source/ControlRigEditor/Private/Graph/SControlRigGraphNode.h
--- a/Engine/Plugins/Experimental/ControlRig/Source/ControlRigEditor/Private/Graph/SControlRigGraphNode.h
+++ b/Engine/Plugins/Experimental/ControlRig/Source/ControlRigEditor/Private/Graph/SControlRigGraphNode.h
@@ -39,6 +39,7 @@ public:
 		TitleAreaWidget = DefaultTitleAreaWidget;
 	}
 	virtual const FSlateBrush * GetNodeBodyBrush() const override;
+	virtual FReply OnMouseButtonDown(const FGeometry& MyGeometry, const FPointerEvent& MouseEvent) override;
 
 	virtual TSharedRef<SWidget> CreateNodeContentArea() override;
 	virtual TSharedPtr<SGraphPin> GetHoveredPin( const FGeometry& MyGeometry, const FPointerEvent& MouseEvent ) const override;
@@ -106,4 +107,7 @@ private:
 	int32 NodeErrorType;
 
 	TSharedPtr<SImage> VisualDebugIndicatorWidget;
+
+	static const FSlateBrush* CachedImg_CR_Pin_Connected;
+	static const FSlateBrush* CachedImg_CR_Pin_Disconnected;
 };
diff --git a/Engine/Plugins/Experimental/ControlRig/Source/ControlRigEditor/Private/Sequencer/ControlRigParameterTrackEditor.cpp b/Engine/Plugins/Experimental/ControlRig/Source/ControlRigEditor/Private/Sequencer/ControlRigParameterTrackEditor.cpp
--- a/Engine/Plugins/Experimental/ControlRig/Source/ControlRigEditor/Private/Sequencer/ControlRigParameterTrackEditor.cpp
+++ b/Engine/Plugins/Experimental/ControlRig/Source/ControlRigEditor/Private/Sequencer/ControlRigParameterTrackEditor.cpp
@@ -625,7 +625,14 @@ void FControlRigParameterTrackEditor::OnSelectionChanged(TArray<UMovieSceneTrack
 	GetSequencer()->GetSelectedKeyAreas(KeyAreas);
 	if (KeyAreas.Num() <= 0)
 	{
-		return;
+		if (ControlRigEditMode)
+		{
+			ControlRig = ControlRigEditMode->GetControlRig();
+			if (ControlRig)
+			{
+				ControlRig->ClearControlSelection();
+			}
+		}
 	}
 
 	FScopedTransaction ScopedTransaction(LOCTEXT("SelectControlTransaction", "Select Control"), !GIsTransacting);
diff --git a/Engine/Plugins/Experimental/ControlRig/Source/ControlRigManipulation/Private/ControlRigGizmoActor.cpp b/Engine/Plugins/Experimental/ControlRig/Source/ControlRigManipulation/Private/ControlRigGizmoActor.cpp
--- a/Engine/Plugins/Experimental/ControlRig/Source/ControlRigManipulation/Private/ControlRigGizmoActor.cpp
+++ b/Engine/Plugins/Experimental/ControlRig/Source/ControlRigManipulation/Private/ControlRigGizmoActor.cpp
@@ -12,7 +12,6 @@ AControlRigGizmoActor::AControlRigGizmoActor(const FObjectInitializer& ObjectIni
 	, bEnabled(true)
 	, bSelected(false)
 	, bHovered(false)
-	, bManipulating(false)
 {
 
 	ActorRootComponent = CreateDefaultSubobject<USceneComponent>(TEXT("SceneComponent0"));
@@ -82,23 +81,6 @@ bool AControlRigGizmoActor::IsHovered() const
 	return bHovered;
 }
 
-void AControlRigGizmoActor::SetManipulating(bool bInManipulating)
-{
-	bool bOldManipulating = bManipulating;
-
-	bManipulating = bInManipulating;
-
-	if(bManipulating != bOldManipulating)
-	{
-		FEditorScriptExecutionGuard Guard;
-		OnManipulatingChanged(bManipulating);
-	}
-}
-
-bool AControlRigGizmoActor::IsManipulating() const
-{
-	return bManipulating;
-}
 
 void AControlRigGizmoActor::SetGizmoColor(const FLinearColor& InColor)
 {
diff --git a/Engine/Plugins/Experimental/ControlRig/Source/ControlRigManipulation/Private/DefaultControlRigManipulationLayer.cpp b/Engine/Plugins/Experimental/ControlRig/Source/ControlRigManipulation/Private/DefaultControlRigManipulationLayer.cpp
--- a/Engine/Plugins/Experimental/ControlRig/Source/ControlRigManipulation/Private/DefaultControlRigManipulationLayer.cpp
+++ b/Engine/Plugins/Experimental/ControlRig/Source/ControlRigManipulation/Private/DefaultControlRigManipulationLayer.cpp
@@ -26,7 +26,6 @@ void UDefaultControlRigManipulationLayer::CreateLayer()
 
 void UDefaultControlRigManipulationLayer::DestroyLayer()
 {
-	EndTransaction();
 	IControlRigManipulationLayer::DestroyLayer();
 }
 
@@ -538,41 +537,6 @@ FTransform	UDefaultControlRigManipulationLayer::GetSkeletalMeshComponentTransfor
 	return MeshComponent ? MeshComponent->GetComponentTransform() : FTransform::Identity;
 }
 
-void UDefaultControlRigManipulationLayer::BeginTransaction()
-{
-	if (InteractionTransaction != nullptr)
-	{
-		EndTransaction();
-	}
-
-	InteractionTransaction = new FScopedTransaction(LOCTEXT("TransformGizmo", "Transform Gizmo"));
-
-	for (TWeakObjectPtr<UObject> Manip : ManipulatableObjects)
-	{
-		if (UControlRig* ControlRig = Cast<UControlRig>(Manip))
-		{
-			UObject* Blueprint = ControlRig->GetClass()->ClassGeneratedBy;
-			if (Blueprint)
-			{
-				Blueprint->SetFlags(RF_Transactional);
-				Blueprint->Modify();
-			}
-
-			ControlRig->SetFlags(RF_Transactional);
-			ControlRig->Modify();
-		}
-	}
-}
-
-void UDefaultControlRigManipulationLayer::EndTransaction()
-{
-	if (InteractionTransaction != nullptr)
-	{
-		delete InteractionTransaction;
-		InteractionTransaction = nullptr;
-	}
-}
-
 void UDefaultControlRigManipulationLayer::PostPoseUpdate()
 {
 	FTransform ComponentTransform = GetSkeletalMeshComponentTransform();
diff --git a/Engine/Plugins/Experimental/ControlRig/Source/ControlRigManipulation/Public/ControlRigGizmoActor.h b/Engine/Plugins/Experimental/ControlRig/Source/ControlRigManipulation/Public/ControlRigGizmoActor.h
--- a/Engine/Plugins/Experimental/ControlRig/Source/ControlRigManipulation/Public/ControlRigGizmoActor.h
+++ b/Engine/Plugins/Experimental/ControlRig/Source/ControlRigManipulation/Public/ControlRigGizmoActor.h
@@ -86,14 +86,6 @@ public:
 	/** Get whether the control is hovered */
 	virtual bool IsHovered() const;
 
-	UFUNCTION(BlueprintSetter)
-	/** Set whether the control is being manipulated */
-	virtual void SetManipulating(bool bInManipulating);
-
-	UFUNCTION(BlueprintGetter)
-	/** Get whether the control is being manipulated */
-	virtual bool IsManipulating() const;
-
 	/** Called from the edit mode each tick */
 	virtual void TickControl() {};
 
@@ -140,9 +132,6 @@ private:
 	UPROPERTY(BlueprintGetter = IsHovered, BlueprintSetter = SetHovered, Category = "ControlRig|Gizmo")
 	uint8 bHovered : 1;
 
-	/** Whether this control is being manipulated */
-	UPROPERTY(BlueprintGetter = IsManipulating, BlueprintSetter = SetManipulating, Category = "ControlRig|Gizmo")
-	uint8 bManipulating : 1;
 };
 
 /**
diff --git a/Engine/Plugins/Experimental/ControlRig/Source/ControlRigManipulation/Public/DefaultControlRigManipulationLayer.h b/Engine/Plugins/Experimental/ControlRig/Source/ControlRigManipulation/Public/DefaultControlRigManipulationLayer.h
--- a/Engine/Plugins/Experimental/ControlRig/Source/ControlRigManipulation/Public/DefaultControlRigManipulationLayer.h
+++ b/Engine/Plugins/Experimental/ControlRig/Source/ControlRigManipulation/Public/DefaultControlRigManipulationLayer.h
@@ -64,10 +64,6 @@ public:
 	USkeletalMeshComponent* GetSkeletalMeshComponent() const;
 	FTransform	GetSkeletalMeshComponentTransform() const;
 
-	// Utility functions
-	void BeginTransaction();
-	void EndTransaction();
-
 	// in this layer, we only care one to one
 	const FControlData* GetControlDataFromGizmo(const AControlRigGizmoActor* GizmoActor) const;
 	// this is slow, and it only finds the first one, and there is no guarantee it will always find the same name if you run this multiple sessions.
diff --git a/Engine/Plugins/FX/Niagara/Shaders/Private/NiagaraDataInterfaceSkeletalMesh.ush b/Engine/Plugins/FX/Niagara/Shaders/Private/NiagaraDataInterfaceSkeletalMesh.ush
--- a/Engine/Plugins/FX/Niagara/Shaders/Private/NiagaraDataInterfaceSkeletalMesh.ush
+++ b/Engine/Plugins/FX/Niagara/Shaders/Private/NiagaraDataInterfaceSkeletalMesh.ush
@@ -2,8 +2,12 @@
 
 #pragma once
 
-#ifndef DISKELMESH_USE_EXTRA_INFLUENCES
-#define DISKELMESH_USE_EXTRA_INFLUENCES	0			// Default is 4 bone, this enables support for 4 & 8 bones
+#define DISKELMESH_BONE_INFLUENCES_DEFAULT		0
+#define DISKELMESH_BONE_INFLUENCES_EXTRA		1
+#define DISKELMESH_BONE_INFLUENCES_UNLIMITED	2
+
+#ifndef DISKELMESH_BONE_INFLUENCES
+#define DISKELMESH_BONE_INFLUENCES	DISKELMESH_BONE_INFLUENCES_DEFAULT
 #endif
 
 struct FDISkelMeshContext
@@ -11,6 +15,7 @@ struct FDISkelMeshContext
 	Buffer<uint>	MeshIndexBuffer;
 	Buffer<float>	MeshVertexBuffer;
 	Buffer<uint>	MeshSkinWeightBuffer;
+	Buffer<uint>	MeshSkinWeightLookupBuffer;
 	Buffer<float4>	MeshCurrBonesBuffer;
 	Buffer<float4>	MeshPrevBonesBuffer;
 	Buffer<float4>	MeshCurrSamplingBonesBuffer;
@@ -24,6 +29,7 @@ struct FDISkelMeshContext
 	uint			MeshTriangleCount;
 	uint			MeshVertexCount;
     uint			MeshWeightStride;
+	uint			MeshSkinWeightIndexSize;
 	uint			MeshNumTexCoord;
 	uint			MeshNumWeights;
 	uint			NumSpecificBones;
@@ -34,6 +40,7 @@ struct FDISkelMeshContext
 	float4x4		InstancePrevTransform;
 	float			InstanceInvDeltaTime;
 	bool			UniformTriangleSamplingEnable;
+	bool			UnlimitedBoneInfluences;
 };
 
 ////////////////////////////////////////////////////////////////////////////////////////////////////////////////////////////////////////////////////////////////
@@ -108,26 +115,62 @@ FDISkelMeshSkinnedVertex DISkelMesh_GetSkinnedVertex(in FDISkelMeshContext DICon
 		float3x4 PrevBoneMatrix;
 		float3x4 CurrBoneMatrix;
 
-#if DISKELMESH_USE_EXTRA_INFLUENCES
-        uint WeightOffset = DIContext.MeshNumWeights == 8 ? 2 : 1;
-
-        int4 BlendIndices = DISkelMesh_UnpackIndices4(DIContext.MeshSkinWeightBuffer[VertexIndex * DIContext.MeshWeightStride]);
-        float4 BlendWeights = DISkelMesh_UnpackWeights4(DIContext.MeshSkinWeightBuffer[VertexIndex * DIContext.MeshWeightStride + WeightOffset]);
-        CurrBoneMatrix = DISkelMesh_GetCurrSkinningMatrix(DIContext, VertexIndex, BlendIndices, BlendWeights);
-        PrevBoneMatrix = DISkelMesh_GetPrevSkinningMatrix(DIContext, VertexIndex, BlendIndices, BlendWeights);
-		if (DIContext.MeshNumWeights == 8)
-        {
-            BlendIndices = DISkelMesh_UnpackIndices4(DIContext.MeshSkinWeightBuffer[VertexIndex * DIContext.MeshWeightStride + 2]);
-            BlendWeights = DISkelMesh_UnpackWeights4(DIContext.MeshSkinWeightBuffer[VertexIndex * DIContext.MeshWeightStride + WeightOffset + 1]);
-			CurrBoneMatrix += DISkelMesh_GetCurrSkinningMatrix(DIContext, VertexIndex, BlendIndices, BlendWeights);
-			PrevBoneMatrix += DISkelMesh_GetPrevSkinningMatrix(DIContext, VertexIndex, BlendIndices, BlendWeights);
-        }
+#if DISKELMESH_BONE_INFLUENCES == DISKELMESH_BONE_INFLUENCES_UNLIMITED
+		BRANCH
+		if (DIContext.UnlimitedBoneInfluences)
+		{
+			uint BlendOffsetCount = DIContext.MeshSkinWeightLookupBuffer[VertexIndex];
+			int NumBoneInfluences = BlendOffsetCount & 0xff;
+			int StreamOffset = BlendOffsetCount >> 8;
+			int WeightsOffset = StreamOffset + (DIContext.MeshSkinWeightIndexSize * NumBoneInfluences);
+
+			PrevBoneMatrix = float3x4(float4(0,0,0,0), float4(0,0,0,0), float4(0,0,0,0));
+			CurrBoneMatrix = float3x4(float4(0,0,0,0), float4(0,0,0,0), float4(0,0,0,0));
+			for (int InfluenceIdx = 0; InfluenceIdx < NumBoneInfluences; InfluenceIdx++)
+			{
+				int BoneIndexOffset = StreamOffset + (DIContext.MeshSkinWeightIndexSize * InfluenceIdx);
+				int BoneIndex = DIContext.MeshSkinWeightBuffer[BoneIndexOffset];
+				if (DIContext.MeshSkinWeightIndexSize > 1)
+				{
+					BoneIndex = DIContext.MeshSkinWeightBuffer[BoneIndexOffset + 1] << 8 | BoneIndex;
+					//@todo-lh: Workaround to fix issue in SPIRVEmitter of DXC; this block must be inside the if branch
+					float BoneWeight = float(DIContext.MeshSkinWeightBuffer[WeightsOffset + InfluenceIdx]) / 255.0f;
+					PrevBoneMatrix += DISkelMesh_GetPrevBoneSkinningMatrix(DIContext, BoneIndex) * BoneWeight;
+					CurrBoneMatrix += DISkelMesh_GetCurrBoneSkinningMatrix(DIContext, BoneIndex) * BoneWeight;
+				}
+				else
+				{
+					//@todo-lh: Workaround to fix issue in SPIRVEmitter of DXC; this block must be inside the if branch
+					float BoneWeight = float(DIContext.MeshSkinWeightBuffer[WeightsOffset + InfluenceIdx]) / 255.0f;
+					PrevBoneMatrix += DISkelMesh_GetPrevBoneSkinningMatrix(DIContext, BoneIndex) * BoneWeight;
+					CurrBoneMatrix += DISkelMesh_GetCurrBoneSkinningMatrix(DIContext, BoneIndex) * BoneWeight;
+				}
+			}	
+		}
+		else
+#endif
+		{
+#if DISKELMESH_BONE_INFLUENCES > DISKELMESH_BONE_INFLUENCES_DEFAULT
+			uint WeightOffset = DIContext.MeshNumWeights == 8 ? 2 : 1;
+
+			int4 BlendIndices = DISkelMesh_UnpackIndices4(DIContext.MeshSkinWeightBuffer[VertexIndex * DIContext.MeshWeightStride]);
+			float4 BlendWeights = DISkelMesh_UnpackWeights4(DIContext.MeshSkinWeightBuffer[VertexIndex * DIContext.MeshWeightStride + WeightOffset]);
+			CurrBoneMatrix = DISkelMesh_GetCurrSkinningMatrix(DIContext, VertexIndex, BlendIndices, BlendWeights);
+			PrevBoneMatrix = DISkelMesh_GetPrevSkinningMatrix(DIContext, VertexIndex, BlendIndices, BlendWeights);
+			if (DIContext.MeshNumWeights == 8)
+			{
+			    BlendIndices = DISkelMesh_UnpackIndices4(DIContext.MeshSkinWeightBuffer[VertexIndex * DIContext.MeshWeightStride + 2]);
+			    BlendWeights = DISkelMesh_UnpackWeights4(DIContext.MeshSkinWeightBuffer[VertexIndex * DIContext.MeshWeightStride + WeightOffset + 1]);
+				CurrBoneMatrix += DISkelMesh_GetCurrSkinningMatrix(DIContext, VertexIndex, BlendIndices, BlendWeights);
+				PrevBoneMatrix += DISkelMesh_GetPrevSkinningMatrix(DIContext, VertexIndex, BlendIndices, BlendWeights);
+			}
 #else
-        int4 BlendIndices = DISkelMesh_UnpackIndices4(DIContext.MeshSkinWeightBuffer[VertexIndex * DIContext.MeshWeightStride]);
-        float4 BlendWeights = DISkelMesh_UnpackWeights4(DIContext.MeshSkinWeightBuffer[VertexIndex * DIContext.MeshWeightStride + 1]);
-        CurrBoneMatrix = DISkelMesh_GetCurrSkinningMatrix(DIContext, VertexIndex, BlendIndices, BlendWeights);
-        PrevBoneMatrix = DISkelMesh_GetPrevSkinningMatrix(DIContext, VertexIndex, BlendIndices, BlendWeights);
+			int4 BlendIndices = DISkelMesh_UnpackIndices4(DIContext.MeshSkinWeightBuffer[VertexIndex * DIContext.MeshWeightStride]);
+			float4 BlendWeights = DISkelMesh_UnpackWeights4(DIContext.MeshSkinWeightBuffer[VertexIndex * DIContext.MeshWeightStride + 1]);
+			CurrBoneMatrix = DISkelMesh_GetCurrSkinningMatrix(DIContext, VertexIndex, BlendIndices, BlendWeights);
+			PrevBoneMatrix = DISkelMesh_GetPrevSkinningMatrix(DIContext, VertexIndex, BlendIndices, BlendWeights);
 #endif
+		}
 
 		CurrPosition = mul(CurrBoneMatrix, float4(CurrPosition, 1.0f)).xyz;
 		PrevPosition = mul(PrevBoneMatrix, float4(PrevPosition, 1.0f)).xyz;
@@ -446,6 +489,7 @@ void DISkelMesh_RandomSpecificSocketBone(in FDISkelMeshContext DIContext, out in
 	Buffer<uint>	MeshIndexBuffer_ ## NAME; \
 	Buffer<float>	MeshVertexBuffer_ ## NAME; \
 	Buffer<uint>	MeshSkinWeightBuffer_ ## NAME; \
+	Buffer<uint>	MeshSkinWeightLookupBuffer_ ## NAME; \
 	Buffer<float4>	MeshCurrBonesBuffer_ ## NAME; \
 	Buffer<float4>	MeshPrevBonesBuffer_ ## NAME; \
 	Buffer<float4>	MeshCurrSamplingBonesBuffer_ ## NAME; \
@@ -459,6 +503,7 @@ void DISkelMesh_RandomSpecificSocketBone(in FDISkelMeshContext DIContext, out in
 	uint			MeshTriangleCount_ ## NAME; \
 	uint			MeshVertexCount_ ## NAME; \
 	uint			MeshWeightStride_ ## NAME; \
+	uint			MeshSkinWeightIndexSize_ ## NAME; \
 	uint			MeshNumTexCoord_ ## NAME; \
 	uint			MeshNumWeights_ ## NAME; \
 	uint			NumSpecificBones_ ## NAME; \
@@ -475,6 +520,7 @@ void DISkelMesh_RandomSpecificSocketBone(in FDISkelMeshContext DIContext, out in
 	DIContext.MeshIndexBuffer = MeshIndexBuffer_ ## NAME; \
 	DIContext.MeshVertexBuffer = MeshVertexBuffer_ ## NAME; \
 	DIContext.MeshSkinWeightBuffer = MeshSkinWeightBuffer_ ## NAME; \
+	DIContext.MeshSkinWeightLookupBuffer = MeshSkinWeightLookupBuffer_ ## NAME; \
 	DIContext.MeshCurrBonesBuffer = MeshCurrBonesBuffer_ ## NAME; \
 	DIContext.MeshPrevBonesBuffer = MeshPrevBonesBuffer_ ## NAME; \
 	DIContext.MeshCurrSamplingBonesBuffer = MeshCurrSamplingBonesBuffer_ ## NAME; \
@@ -494,8 +540,9 @@ void DISkelMesh_RandomSpecificSocketBone(in FDISkelMeshContext DIContext, out in
 	DIContext.NumSpecificSockets = NumSpecificSockets_ ## NAME; \
 	DIContext.SpecificSocketBoneOffset = SpecificSocketBoneOffset_ ## NAME; \
 	DIContext.MeshWeightStride = MeshWeightStride_ ## NAME; \
+	DIContext.MeshSkinWeightIndexSize = MeshSkinWeightIndexSize_ ## NAME; \
 	DIContext.InstanceTransform = InstanceTransform_ ## NAME; \
 	DIContext.InstancePrevTransform = InstancePrevTransform_ ## NAME; \
 	DIContext.InstanceInvDeltaTime = InstanceInvDeltaTime_ ## NAME; \
 	DIContext.UniformTriangleSamplingEnable = EnabledFeatures_ ## NAME & 0x1; \
-
+	DIContext.UnlimitedBoneInfluences = EnabledFeatures_ ## NAME & 0x2; \
diff --git a/Engine/Plugins/FX/Niagara/Source/Niagara/Classes/NiagaraDataInterfaceSkeletalMesh.h b/Engine/Plugins/FX/Niagara/Source/Niagara/Classes/NiagaraDataInterfaceSkeletalMesh.h
--- a/Engine/Plugins/FX/Niagara/Source/Niagara/Classes/NiagaraDataInterfaceSkeletalMesh.h
+++ b/Engine/Plugins/FX/Niagara/Source/Niagara/Classes/NiagaraDataInterfaceSkeletalMesh.h
@@ -425,8 +425,12 @@ struct FNDISkeletalMesh_InstanceData
 	/** True if the mesh we're using allows area weighted sampling on GPU. */
 	uint32 bIsGpuUniformlyDistributedSampling : 1;
 
+	/** True if the mesh we're using is to be rendered in unlimited bone influences mode. */
+	uint32 bUnlimitedBoneInfluences : 1;
 	FRHIShaderResourceView* MeshSkinWeightBufferSrv;
+	FRHIShaderResourceView* MeshSkinWeightLookupBufferSrv;
 	uint32 MeshWeightStrideByte;
+	uint32 MeshSkinWeightIndexSizeByte;
 
 	/** Extra mesh data upload to GPU.*/
 	FSkeletalMeshGpuSpawnStaticBuffers* MeshGpuSpawnStaticBuffers;
@@ -558,6 +562,7 @@ public:
 	static const FString MeshIndexBufferName;
 	static const FString MeshVertexBufferName;
 	static const FString MeshSkinWeightBufferName;
+	static const FString MeshSkinWeightLookupBufferName;
 	static const FString MeshCurrBonesBufferName;
 	static const FString MeshPrevBonesBufferName;
 	static const FString MeshCurrSamplingBonesBufferName;
@@ -571,6 +576,7 @@ public:
 	static const FString MeshTriangleCountName;
 	static const FString MeshVertexCountName;
 	static const FString MeshWeightStrideName;
+	static const FString MeshSkinWeightIndexSizeName;
 	static const FString MeshNumTexCoordName;
 	static const FString MeshNumWeightsName;
 	static const FString NumSpecificBonesName;
@@ -752,10 +758,13 @@ struct FNiagaraDISkeletalMeshPassedDataToRT
 	FSkeletalMeshGpuSpawnStaticBuffers* StaticBuffers;
 	FSkeletalMeshGpuDynamicBufferProxy* DynamicBuffer;
 	FRHIShaderResourceView* MeshSkinWeightBufferSrv;
+	FRHIShaderResourceView* MeshSkinWeightLookupBufferSrv;
 
 	bool bIsGpuUniformlyDistributedSampling;
 
+	bool bUnlimitedBoneInfluences;
 	uint32 MeshWeightStrideByte;
+	uint32 MeshSkinWeightIndexSizeByte;
 	FMatrix Transform;
 	FMatrix PrevTransform;
 	float DeltaSeconds;
diff --git a/Engine/Plugins/FX/Niagara/Source/Niagara/Private/NiagaraDataInterfaceSkeletalMesh.cpp b/Engine/Plugins/FX/Niagara/Source/Niagara/Private/NiagaraDataInterfaceSkeletalMesh.cpp
--- a/Engine/Plugins/FX/Niagara/Source/Niagara/Private/NiagaraDataInterfaceSkeletalMesh.cpp
+++ b/Engine/Plugins/FX/Niagara/Source/Niagara/Private/NiagaraDataInterfaceSkeletalMesh.cpp
@@ -706,6 +706,7 @@ struct FNDISkeletalMeshParametersName
 	FString MeshIndexBufferName;
 	FString MeshVertexBufferName;
 	FString MeshSkinWeightBufferName;
+	FString MeshSkinWeightLookupBufferName;
 	FString MeshCurrBonesBufferName;
 	FString MeshPrevBonesBufferName;
 	FString MeshCurrSamplingBonesBufferName;
@@ -719,6 +720,7 @@ struct FNDISkeletalMeshParametersName
 	FString MeshTriangleCountName;
 	FString MeshVertexCountName;
 	FString MeshWeightStrideName;
+	FString MeshSkinWeightIndexSizeName;
 	FString MeshNumTexCoordName;
 	FString MeshNumWeightsName;
 	FString NumSpecificBonesName;
@@ -736,6 +738,7 @@ static void GetNiagaraDataInterfaceParametersName(FNDISkeletalMeshParametersName
 	Names.MeshIndexBufferName = UNiagaraDataInterfaceSkeletalMesh::MeshIndexBufferName + Suffix;
 	Names.MeshVertexBufferName = UNiagaraDataInterfaceSkeletalMesh::MeshVertexBufferName + Suffix;
 	Names.MeshSkinWeightBufferName = UNiagaraDataInterfaceSkeletalMesh::MeshSkinWeightBufferName + Suffix;
+	Names.MeshSkinWeightLookupBufferName = UNiagaraDataInterfaceSkeletalMesh::MeshSkinWeightLookupBufferName + Suffix;
 	Names.MeshCurrBonesBufferName = UNiagaraDataInterfaceSkeletalMesh::MeshCurrBonesBufferName + Suffix;
 	Names.MeshPrevBonesBufferName = UNiagaraDataInterfaceSkeletalMesh::MeshPrevBonesBufferName + Suffix;
 	Names.MeshCurrSamplingBonesBufferName = UNiagaraDataInterfaceSkeletalMesh::MeshCurrSamplingBonesBufferName + Suffix;
@@ -749,6 +752,7 @@ static void GetNiagaraDataInterfaceParametersName(FNDISkeletalMeshParametersName
 	Names.MeshTriangleCountName = UNiagaraDataInterfaceSkeletalMesh::MeshTriangleCountName + Suffix;
 	Names.MeshVertexCountName = UNiagaraDataInterfaceSkeletalMesh::MeshVertexCountName + Suffix;
 	Names.MeshWeightStrideName = UNiagaraDataInterfaceSkeletalMesh::MeshWeightStrideName + Suffix;
+	Names.MeshSkinWeightIndexSizeName = UNiagaraDataInterfaceSkeletalMesh::MeshSkinWeightIndexSizeName + Suffix;
 	Names.MeshNumTexCoordName = UNiagaraDataInterfaceSkeletalMesh::MeshNumTexCoordName + Suffix;
 	Names.MeshNumWeightsName = UNiagaraDataInterfaceSkeletalMesh::MeshNumWeightsName + Suffix;
 	Names.NumSpecificBonesName = UNiagaraDataInterfaceSkeletalMesh::NumSpecificBonesName + Suffix;
@@ -771,6 +775,7 @@ struct FNiagaraDataInterfaceParametersCS_SkeletalMesh : public FNiagaraDataInter
 		MeshIndexBuffer.Bind(ParameterMap, *ParamNames.MeshIndexBufferName);
 		MeshVertexBuffer.Bind(ParameterMap, *ParamNames.MeshVertexBufferName);
 		MeshSkinWeightBuffer.Bind(ParameterMap, *ParamNames.MeshSkinWeightBufferName);
+		MeshSkinWeightLookupBuffer.Bind(ParameterMap, *ParamNames.MeshSkinWeightLookupBufferName);
 		MeshCurrBonesBuffer.Bind(ParameterMap, *ParamNames.MeshCurrBonesBufferName);
 		MeshPrevBonesBuffer.Bind(ParameterMap, *ParamNames.MeshPrevBonesBufferName);
 		MeshCurrSamplingBonesBuffer.Bind(ParameterMap, *ParamNames.MeshCurrSamplingBonesBufferName);
@@ -784,6 +789,7 @@ struct FNiagaraDataInterfaceParametersCS_SkeletalMesh : public FNiagaraDataInter
 		MeshTriangleCount.Bind(ParameterMap, *ParamNames.MeshTriangleCountName);
 		MeshVertexCount.Bind(ParameterMap, *ParamNames.MeshVertexCountName);
 		MeshWeightStride.Bind(ParameterMap, *ParamNames.MeshWeightStrideName);
+		MeshSkinWeightIndexSize.Bind(ParameterMap, *ParamNames.MeshSkinWeightIndexSizeName);
 		MeshNumTexCoord.Bind(ParameterMap, *ParamNames.MeshNumTexCoordName);
 		MeshNumWeights.Bind(ParameterMap, *ParamNames.MeshNumWeightsName);
 		NumSpecificBones.Bind(ParameterMap, *ParamNames.NumSpecificBonesName);
@@ -801,6 +807,7 @@ struct FNiagaraDataInterfaceParametersCS_SkeletalMesh : public FNiagaraDataInter
 		Ar << MeshIndexBuffer;
 		Ar << MeshVertexBuffer;
 		Ar << MeshSkinWeightBuffer;
+		Ar << MeshSkinWeightLookupBuffer;
 		Ar << MeshCurrBonesBuffer;
 		Ar << MeshPrevBonesBuffer;
 		Ar << MeshCurrSamplingBonesBuffer;
@@ -814,6 +821,7 @@ struct FNiagaraDataInterfaceParametersCS_SkeletalMesh : public FNiagaraDataInter
 		Ar << MeshTriangleCount;
 		Ar << MeshVertexCount;
 		Ar << MeshWeightStride;
+		Ar << MeshSkinWeightIndexSize;
 		Ar << MeshNumTexCoord;
 		Ar << MeshNumWeights;
 		Ar << NumSpecificBones;
@@ -872,10 +880,13 @@ struct FNiagaraDataInterfaceParametersCS_SkeletalMesh : public FNiagaraDataInter
 			}
 
 			SetSRVParameter(RHICmdList, ComputeShaderRHI, MeshSkinWeightBuffer, InstanceData->MeshSkinWeightBufferSrv);
+			SetSRVParameter(RHICmdList, ComputeShaderRHI, MeshSkinWeightLookupBuffer, InstanceData->MeshSkinWeightLookupBufferSrv);
 
 			SetShaderValue(RHICmdList, ComputeShaderRHI, MeshWeightStride, InstanceData->MeshWeightStrideByte/4);
+			SetShaderValue(RHICmdList, ComputeShaderRHI, MeshSkinWeightIndexSize, InstanceData->MeshSkinWeightIndexSizeByte);
 
 			uint32 EnabledFeaturesBits = InstanceData->bIsGpuUniformlyDistributedSampling ? 1 : 0;
+			EnabledFeaturesBits |= (InstanceData->bUnlimitedBoneInfluences ? (1 << 1) : 0);
 
 			FSkeletalMeshGpuDynamicBufferProxy* DynamicBuffers = InstanceData->DynamicBuffer;
 			check(DynamicBuffers);
@@ -930,8 +941,10 @@ struct FNiagaraDataInterfaceParametersCS_SkeletalMesh : public FNiagaraDataInter
 			SetSRVParameter(RHICmdList, ComputeShaderRHI, MeshTriangleSamplerAliasBuffer, FNiagaraRenderer::GetDummyUIntBuffer().SRV);
 
 			SetSRVParameter(RHICmdList, ComputeShaderRHI, MeshSkinWeightBuffer, FNiagaraRenderer::GetDummyUIntBuffer().SRV);
+			SetSRVParameter(RHICmdList, ComputeShaderRHI, MeshSkinWeightLookupBuffer, FNiagaraRenderer::GetDummyUIntBuffer().SRV);
 
 			SetShaderValue(RHICmdList, ComputeShaderRHI, MeshWeightStride, 0);
+			SetShaderValue(RHICmdList, ComputeShaderRHI, MeshSkinWeightIndexSize, 0);
 			SetShaderValue(RHICmdList, ComputeShaderRHI, MeshNumWeights, 0);
 
 			SetSRVParameter(RHICmdList, ComputeShaderRHI, MeshCurrBonesBuffer, FNiagaraRenderer::GetDummyFloat4Buffer().SRV);
@@ -959,6 +972,7 @@ private:
 	FShaderResourceParameter MeshIndexBuffer;
 	FShaderResourceParameter MeshVertexBuffer;
 	FShaderResourceParameter MeshSkinWeightBuffer;
+	FShaderResourceParameter MeshSkinWeightLookupBuffer;
 	FShaderResourceParameter MeshCurrBonesBuffer;
 	FShaderResourceParameter MeshPrevBonesBuffer;
 	FShaderResourceParameter MeshCurrSamplingBonesBuffer;
@@ -972,6 +986,7 @@ private:
 	FShaderParameter MeshTriangleCount;
 	FShaderParameter MeshVertexCount;
 	FShaderParameter MeshWeightStride;
+	FShaderParameter MeshSkinWeightIndexSize;
 	FShaderParameter MeshNumTexCoord;
 	FShaderParameter MeshNumWeights;
 	FShaderParameter NumSpecificBones;
@@ -993,15 +1008,18 @@ void FNiagaraDataInterfaceProxySkeletalMesh::ConsumePerInstanceDataFromGameThrea
 	FNiagaraDataInterfaceProxySkeletalMeshData& Data = SystemInstancesToData.FindOrAdd(Instance);
 
 	Data.bIsGpuUniformlyDistributedSampling = SourceData->bIsGpuUniformlyDistributedSampling;
+	Data.bUnlimitedBoneInfluences = SourceData->bUnlimitedBoneInfluences;
 	Data.DeltaSeconds = SourceData->DeltaSeconds;
 	Data.DynamicBuffer = SourceData->DynamicBuffer;
 	Data.MeshWeightStrideByte = SourceData->MeshWeightStrideByte;
+	Data.MeshSkinWeightIndexSizeByte = SourceData->MeshSkinWeightIndexSizeByte;
 	Data.PrevTransform = SourceData->PrevTransform;
 	Data.StaticBuffers = SourceData->StaticBuffers;
 	Data.Transform = SourceData->Transform;
 
 	// @todo-threadsafety race here. Need to hold a ref to this buffer on the RT
 	Data.MeshSkinWeightBufferSrv = SourceData->MeshSkinWeightBufferSrv;
+	Data.MeshSkinWeightLookupBufferSrv = SourceData->MeshSkinWeightLookupBufferSrv;
 }
 
 //////////////////////////////////////////////////////////////////////////
@@ -1013,15 +1031,18 @@ void UNiagaraDataInterfaceSkeletalMesh::ProvidePerInstanceDataForRenderThread(vo
 	FNDISkeletalMesh_InstanceData* SourceData = static_cast<FNDISkeletalMesh_InstanceData*>(PerInstanceData);
 
 	Data->bIsGpuUniformlyDistributedSampling = SourceData->bIsGpuUniformlyDistributedSampling;
+	Data->bUnlimitedBoneInfluences = SourceData->bUnlimitedBoneInfluences;
 	Data->DeltaSeconds = SourceData->DeltaSeconds;
 	Data->DynamicBuffer = SourceData->MeshGpuSpawnDynamicBuffers;
 	Data->MeshWeightStrideByte = SourceData->MeshWeightStrideByte;
+	Data->MeshSkinWeightIndexSizeByte = SourceData->MeshSkinWeightIndexSizeByte;
 	Data->PrevTransform = SourceData->PrevTransform;
 	Data->StaticBuffers = SourceData->MeshGpuSpawnStaticBuffers;
 	Data->Transform = SourceData->Transform;
 
 	// @todo-threadsafety race here. Need to hold a ref to this buffer on the RT
 	Data->MeshSkinWeightBufferSrv = SourceData->MeshSkinWeightBufferSrv;
+	Data->MeshSkinWeightLookupBufferSrv = SourceData->MeshSkinWeightLookupBufferSrv;
 }
 
 USkeletalMesh* UNiagaraDataInterfaceSkeletalMesh::GetSkeletalMesh(UNiagaraComponent* OwningComponent, TWeakObjectPtr<USceneComponent>& SceneComponent, USkeletalMeshComponent*& FoundSkelComp, FNDISkeletalMesh_InstanceData* InstData)
@@ -1173,7 +1194,9 @@ bool FNDISkeletalMesh_InstanceData::Init(UNiagaraDataInterfaceSkeletalMesh* Inte
 	DeltaSeconds = SystemInstance->GetComponent()->GetWorld()->GetDeltaSeconds();
 	ChangeId = Interface->ChangeId;
 	bIsGpuUniformlyDistributedSampling = false;
+	bUnlimitedBoneInfluences = false;
 	MeshWeightStrideByte = 0;
+	MeshSkinWeightIndexSizeByte = 0;
 	MeshGpuSpawnStaticBuffers = nullptr;
 	MeshGpuSpawnDynamicBuffers = nullptr;
 
@@ -1426,9 +1449,13 @@ bool FNDISkeletalMesh_InstanceData::Init(UNiagaraDataInterfaceSkeletalMesh* Inte
 	//-TODO: We should find out if this DI is connected to a GPU emitter or not rather than a blanket accross the system
 	if ( SystemInstance->HasGPUEmitters() )
 	{
+		GPUSkinBoneInfluenceType BoneInfluenceType = SkinWeightBuffer->GetBoneInfluenceType();
+		bUnlimitedBoneInfluences = (BoneInfluenceType == GPUSkinBoneInfluenceType::UnlimitedBoneInfluence);
 		MeshWeightStrideByte = SkinWeightBuffer->GetConstantInfluencesVertexStride();
+		MeshSkinWeightIndexSizeByte = SkinWeightBuffer->GetBoneIndexByteSize();
 		MeshSkinWeightBufferSrv = SkinWeightBuffer->GetDataVertexBuffer()->GetSRV();
 		//check(MeshSkinWeightBufferSrv->IsValid()); // not available in this stream
+		MeshSkinWeightLookupBufferSrv = SkinWeightBuffer->GetLookupVertexBuffer()->GetSRV();
 
 		FSkeletalMeshLODInfo* LODInfo = Mesh->GetLODInfo(LODIndex);
 		bIsGpuUniformlyDistributedSampling = LODInfo->bSupportUniformlyDistributedSampling && bAllRegionsAreAreaWeighting;
@@ -2001,6 +2028,7 @@ void UNiagaraDataInterfaceSkeletalMesh::ValidateFunction(const FNiagaraFunctionS
 const FString UNiagaraDataInterfaceSkeletalMesh::MeshIndexBufferName(TEXT("MeshIndexBuffer_"));
 const FString UNiagaraDataInterfaceSkeletalMesh::MeshVertexBufferName(TEXT("MeshVertexBuffer_"));
 const FString UNiagaraDataInterfaceSkeletalMesh::MeshSkinWeightBufferName(TEXT("MeshSkinWeightBuffer_"));
+const FString UNiagaraDataInterfaceSkeletalMesh::MeshSkinWeightLookupBufferName(TEXT("MeshSkinWeightLookupBuffer_"));
 const FString UNiagaraDataInterfaceSkeletalMesh::MeshCurrBonesBufferName(TEXT("MeshCurrBonesBuffer_"));
 const FString UNiagaraDataInterfaceSkeletalMesh::MeshPrevBonesBufferName(TEXT("MeshPrevBonesBuffer_"));
 const FString UNiagaraDataInterfaceSkeletalMesh::MeshCurrSamplingBonesBufferName(TEXT("MeshCurrSamplingBonesBuffer_"));
@@ -2014,6 +2042,7 @@ const FString UNiagaraDataInterfaceSkeletalMesh::MeshTriangleMatricesOffsetBuffe
 const FString UNiagaraDataInterfaceSkeletalMesh::MeshTriangleCountName(TEXT("MeshTriangleCount_"));
 const FString UNiagaraDataInterfaceSkeletalMesh::MeshVertexCountName(TEXT("MeshVertexCount_"));
 const FString UNiagaraDataInterfaceSkeletalMesh::MeshWeightStrideName(TEXT("MeshWeightStride_"));
+const FString UNiagaraDataInterfaceSkeletalMesh::MeshSkinWeightIndexSizeName(TEXT("MeshSkinWeightIndexSize_"));
 const FString UNiagaraDataInterfaceSkeletalMesh::MeshNumTexCoordName(TEXT("MeshNumTexCoord_"));
 const FString UNiagaraDataInterfaceSkeletalMesh::MeshNumWeightsName(TEXT("MeshNumWeights_"));
 const FString UNiagaraDataInterfaceSkeletalMesh::NumSpecificBonesName(TEXT("NumSpecificBones_"));
diff --git a/Engine/Plugins/FX/Niagara/Source/NiagaraShader/Private/NiagaraShader.cpp b/Engine/Plugins/FX/Niagara/Source/NiagaraShader/Private/NiagaraShader.cpp
--- a/Engine/Plugins/FX/Niagara/Source/NiagaraShader/Private/NiagaraShader.cpp
+++ b/Engine/Plugins/FX/Niagara/Source/NiagaraShader/Private/NiagaraShader.cpp
@@ -336,6 +336,7 @@ TSharedRef<FShaderCommonCompileJob, ESPMode::ThreadSafe> FNiagaraShaderType::Beg
 	NewJob->Input.EntryPointName = TEXT("SimulateMainComputeCS");
 	NewJob->Input.Environment.SetDefine(TEXT("GPU_SIMULATION"), 1);
 	NewJob->Input.Environment.SetDefine(TEXT("NIAGARA_MAX_GPU_SPAWN_INFOS"), NIAGARA_MAX_GPU_SPAWN_INFOS);
+	NewJob->Input.Environment.SetDefine(TEXT("DISKELMESH_BONE_INFLUENCES"), 0);
 	NewJob->Input.Environment.IncludeVirtualPathToContentsMap.Add(TEXT("/Engine/Generated/NiagaraEmitterInstance.ush"), Script->HlslOutput);
 
 	static const auto UseShaderStagesCVar = IConsoleManager::Get().FindConsoleVariable(TEXT("fx.UseShaderStages"));
diff --git a/Engine/Plugins/Runtime/AnimationBudgetAllocator/Source/AnimationBudgetAllocator/Private/AnimationBudgetAllocator.cpp b/Engine/Plugins/Runtime/AnimationBudgetAllocator/Source/AnimationBudgetAllocator/Private/AnimationBudgetAllocator.cpp
--- a/Engine/Plugins/Runtime/AnimationBudgetAllocator/Source/AnimationBudgetAllocator/Private/AnimationBudgetAllocator.cpp
+++ b/Engine/Plugins/Runtime/AnimationBudgetAllocator/Source/AnimationBudgetAllocator/Private/AnimationBudgetAllocator.cpp
@@ -788,7 +788,7 @@ void FAnimationBudgetAllocator::OnHUDPostRender(AHUD* HUD, UCanvas* Canvas)
 
 			{
 				TMap<AActor*, TArray<int32>> ActorMap;
-				for (int32 ComponentDataIndex : AllSortedComponentData)
+				for (int32 ComponentDataIndex = 0; ComponentDataIndex < AllComponentData.Num(); ++ComponentDataIndex)
 				{
 					FComponentData& ComponentData = AllComponentData[ComponentDataIndex];
 					TArray<int32>& ComponentIndexArray = ActorMap.FindOrAdd(ComponentData.Component->GetOwner());
@@ -797,8 +797,17 @@ void FAnimationBudgetAllocator::OnHUDPostRender(AHUD* HUD, UCanvas* Canvas)
 
 				float LineOffset = 0.0f;
 
-				for(const TPair<AActor*, TArray<int32>>& ActorIndicesPair : ActorMap)
+				for(TPair<AActor*, TArray<int32>>& ActorIndicesPair : ActorMap)
 				{
+					// Sort by significance
+					ActorIndicesPair.Value.Sort([this](int32 Index0, int32 Index1)
+					{
+						const FComponentData& ComponentData0 = AllComponentData[Index0];
+						const FComponentData& ComponentData1 = AllComponentData[Index1];
+ 
+						return ComponentData0.Significance < ComponentData1.Significance;
+					});
+
 					FVector Location = ActorIndicesPair.Key->GetActorLocation();
 
 					FString DebugString;
diff --git a/Engine/Plugins/Runtime/AnimationBudgetAllocator/Source/AnimationBudgetAllocator/Private/AnimationBudgetAllocator.h b/Engine/Plugins/Runtime/AnimationBudgetAllocator/Source/AnimationBudgetAllocator/Private/AnimationBudgetAllocator.h
--- a/Engine/Plugins/Runtime/AnimationBudgetAllocator/Source/AnimationBudgetAllocator/Private/AnimationBudgetAllocator.h
+++ b/Engine/Plugins/Runtime/AnimationBudgetAllocator/Source/AnimationBudgetAllocator/Private/AnimationBudgetAllocator.h
@@ -194,7 +194,11 @@ protected:
 	// All component data
 	TArray<FComponentData> AllComponentData;
 
-	/** All currently tickable component indices sorted by significance, updated each tick */
+	/** 
+	 * All currently tickable component indices sorted by significance, updated each tick.
+	 * Note that this array is not managed, so components can be deleted underneath it. 
+	 * Therefore usage outside of Tick() is not recommended.
+	 */
 	TArray<int32> AllSortedComponentData;
 
 #if WITH_TICK_DEBUG
diff --git a/Engine/Plugins/Runtime/GameplayAbilities/Source/GameplayAbilities/Private/AbilitySystemComponent.cpp b/Engine/Plugins/Runtime/GameplayAbilities/Source/GameplayAbilities/Private/AbilitySystemComponent.cpp
--- a/Engine/Plugins/Runtime/GameplayAbilities/Source/GameplayAbilities/Private/AbilitySystemComponent.cpp
+++ b/Engine/Plugins/Runtime/GameplayAbilities/Source/GameplayAbilities/Private/AbilitySystemComponent.cpp
@@ -52,7 +52,7 @@ UAbilitySystemComponent::UAbilitySystemComponent(const FObjectInitializer& Objec
 	bSuppressGameplayCues = false;
 	bPendingMontageRep = false;
 	bIsNetDirty = true;
-
+	AffectedAnimInstanceTag = NAME_None; 
 	AbilityLastActivatedTime = 0.f;
 
 	ReplicationMode = EGameplayEffectReplicationMode::Full;
diff --git a/Engine/Plugins/Runtime/GameplayAbilities/Source/GameplayAbilities/Private/GameplayAbilityTypes.cpp b/Engine/Plugins/Runtime/GameplayAbilities/Source/GameplayAbilities/Private/GameplayAbilityTypes.cpp
--- a/Engine/Plugins/Runtime/GameplayAbilities/Source/GameplayAbilities/Private/GameplayAbilityTypes.cpp
+++ b/Engine/Plugins/Runtime/GameplayAbilities/Source/GameplayAbilities/Private/GameplayAbilityTypes.cpp
@@ -19,7 +19,8 @@ void FGameplayAbilityActorInfo::InitFromActor(AActor *InOwnerActor, AActor *InAv
 	OwnerActor = InOwnerActor;
 	AvatarActor = InAvatarActor;
 	AbilitySystemComponent = InAbilitySystemComponent;
-
+	AffectedAnimInstanceTag = InAbilitySystemComponent->AffectedAnimInstanceTag; 
+	
 	APlayerController* OldPC = PlayerController.Get();
 
 	// Look for a player controller or pawn in the owner chain.
diff --git a/Engine/Plugins/Runtime/GameplayAbilities/Source/GameplayAbilities/Public/Abilities/GameplayAbilityTypes.h b/Engine/Plugins/Runtime/GameplayAbilities/Source/GameplayAbilities/Public/Abilities/GameplayAbilityTypes.h
--- a/Engine/Plugins/Runtime/GameplayAbilities/Source/GameplayAbilities/Public/Abilities/GameplayAbilityTypes.h
+++ b/Engine/Plugins/Runtime/GameplayAbilities/Source/GameplayAbilities/Public/Abilities/GameplayAbilityTypes.h
@@ -12,13 +12,13 @@
 #include "GameplayPrediction.h"
 #include "GameplayAbilitySpec.h"
 #include "UObject/Package.h"
+#include "Animation/AnimInstance.h"
 #include "Components/SkeletalMeshComponent.h"
 #include "Abilities/GameplayAbilityTargetTypes.h"
 #include "GameplayAbilityTypes.generated.h"
 
 class APlayerController;
 class UAbilitySystemComponent;
-class UAnimInstance;
 class UAnimMontage;
 class UGameplayAbility;
 class UMovementComponent;
@@ -166,12 +166,30 @@ struct GAMEPLAYABILITIES_API FGameplayAbilityActorInfo
 	/** Movement component of the avatar actor. Often null */
 	UPROPERTY(BlueprintReadOnly, Category = "ActorInfo")
 	TWeakObjectPtr<UMovementComponent>	MovementComponent;
-
-	/** Accessor to get the current anim instance from the SkeletalMeshComponent */
+	
+	/** The linked Anim Instance that this component will play montages in. Use NAME_None for the main anim instance. */
+	UPROPERTY(BlueprintReadOnly, Category = "ActorInfo")
+	FName AffectedAnimInstanceTag; 
+	
+	/** Accessor to get the affected anim instance from the SkeletalMeshComponent */
 	UAnimInstance* GetAnimInstance() const
 	{ 
 		const USkeletalMeshComponent* SKMC = SkeletalMeshComponent.Get();
-		return (SKMC ? SKMC->GetAnimInstance() : nullptr);
+
+		if (SKMC)
+		{
+			if (AffectedAnimInstanceTag != NAME_None)
+			{
+				if(UAnimInstance* Instance = SKMC->GetAnimInstance())
+				{
+					return Instance->GetLinkedAnimGraphInstanceByTag(AffectedAnimInstanceTag);
+				}
+			}
+
+			return SKMC->GetAnimInstance();
+		}
+
+		return nullptr;
 	}
 
 	/** Returns true if this actor is locally controlled. Only true for players on the client that owns them */
diff --git a/Engine/Plugins/Runtime/GameplayAbilities/Source/GameplayAbilities/Public/AbilitySystemComponent.h b/Engine/Plugins/Runtime/GameplayAbilities/Source/GameplayAbilities/Public/AbilitySystemComponent.h
--- a/Engine/Plugins/Runtime/GameplayAbilities/Source/GameplayAbilities/Public/AbilitySystemComponent.h
+++ b/Engine/Plugins/Runtime/GameplayAbilities/Source/GameplayAbilities/Public/AbilitySystemComponent.h
@@ -154,6 +154,11 @@ class GAMEPLAYABILITIES_API UAbilitySystemComponent : public UGameplayTasksCompo
 	UPROPERTY(Replicated)
 	TArray<UAttributeSet*>	SpawnedAttributes;
 
+	/** The linked Anim Instance that this component will play montages in. Use NAME_None for the main anim instance. */
+	UPROPERTY(BlueprintReadWrite, EditAnywhere, Category = "Skills")
+	FName AffectedAnimInstanceTag; 
+
+
 	/** Sets the base value of an attribute. Existing active modifiers are NOT cleared and will act upon the new base value. */
 	void SetNumericAttributeBase(const FGameplayAttribute &Attribute, float NewBaseValue);
 
diff --git a/Engine/Shaders/Private/RecomputeTangentsPerVertexPass.usf b/Engine/Shaders/Private/RecomputeTangentsPerVertexPass.usf
--- a/Engine/Shaders/Private/RecomputeTangentsPerVertexPass.usf
+++ b/Engine/Shaders/Private/RecomputeTangentsPerVertexPass.usf
@@ -9,6 +9,118 @@
 // 0:fast and simple / 1:higher quality lighting with extreme mesh distortions
 #define ORTHOGONALIZE 1
 
+#ifndef BLEND_USING_VERTEX_COLOR
+#define BLEND_USING_VERTEX_COLOR 0
+#endif
+
+#if BLEND_USING_VERTEX_COLOR
+Buffer<float4> ColorInputBuffer;
+#endif
+
+float4 TangentToQuaternion(float3 TangentX, float3 TangentZ)
+{
+	// Based off \Engine\Source\Runtime\Core\Public\Math\Quat.h FQuat(const FMatrix& M)
+	// Be sure TangentX and TangentX are normalized
+	float3x3 M;
+	M[1] = cross(TangentZ, TangentX);
+	M[2] = TangentZ;
+	M[0] = cross(M[1], TangentZ.xyz);	// Recompute TangentX in case TangentX & Z aren't perfectly perpendicular to each other
+
+	float4 Quat;
+	float S;
+
+	// Check diagonal (trace)
+	const float Tr = M[0][0] + M[1][1] + M[2][2];
+	if (Tr > 0.0f) 
+	{
+		float InvS = 1.0f / sqrt(Tr + 1.f);
+		Quat.w = 0.5f * (1.f / InvS);
+		S = 0.5f * InvS;
+
+		Quat.x = (M[1][2] - M[2][1]) * S;
+		Quat.y = (M[2][0] - M[0][2]) * S;
+		Quat.z = (M[0][1] - M[1][0]) * S;
+	} 
+	else 
+	{
+		// diagonal is negative
+		int i = 0;
+
+		if (M[1][1] > M[0][0])
+			i = 1;
+
+		if (M[2][2] > M[i][i])
+			i = 2;
+
+		int nxt[3] = { 1, 2, 0 };
+		int j = nxt[i];
+		int k = nxt[j];
+ 
+		S = M[i][i] - M[j][j] - M[k][k] + 1.0f;
+
+		float InvS = 1.0f / sqrt(S);
+
+		float Qt[4];
+		Qt[i] = 0.5f * (1.0f / InvS);
+
+		S = 0.5f * InvS;
+
+		Qt[3] = (M[j][k] - M[k][j]) * S;
+		Qt[j] = (M[i][j] + M[j][i]) * S;
+		Qt[k] = (M[i][k] + M[k][i]) * S;
+
+		Quat.x = Qt[0];
+		Quat.y = Qt[1];
+		Quat.z = Qt[2];
+		Quat.w = Qt[3];
+	}
+
+	return normalize(Quat); 
+}
+
+void QuaternionToTangent(float4 Quat, out float3 TangentX, out float3 TangentZ)
+{
+	// Based off \Engine\Source\Runtime\Core\Public\Math\QuatRotationTranslationMatrix.h FQuatRotationTranslationMatrix(const FQuat& Q, const FVector& Origin)
+	float x2 = Quat.x + Quat.x; float y2 = Quat.y + Quat.y; float z2 = Quat.z + Quat.z;
+	float xx = Quat.x * x2;		float xy = Quat.x * y2;		float xz = Quat.x * z2;
+	float yy = Quat.y * y2;     float yz = Quat.y * z2;		float zz = Quat.z * z2;
+	float wx = Quat.w * x2;     float wy = Quat.w * y2;		float wz = Quat.w * z2;
+
+	TangentX = float3(1.0f - (yy + zz), xy + wz, xz - wy);
+	TangentZ = float3(xz + wy, yz - wx, 1.0f - (xx + yy));
+}
+
+float4 Slerp(float4 Start, float4 End, float T)
+{
+    float DotProduct = dot(Start, End);
+    float DotAbs = abs(DotProduct);
+	
+    float T0, T1;
+    if (DotAbs > 0.9999f)
+    {
+		// Start and End are very close together, do linear interpolation.
+        T0 = 1.0f - T;
+        T1 = T;
+    }
+    else
+    {
+        float Omega = acos(DotAbs);
+        float InvSin = 1.f / sin(Omega);
+        T0 = sin((1.f - T) * Omega) * InvSin;
+        T1 = sin(T * Omega) * InvSin;
+    }
+
+    T1 = DotProduct < 0 ? -T1 : T1;
+
+	return (T0 * Start) + (T1 * End);
+}
+
+float3 Slerp(float3 Start, float3 End, float T)
+{
+	float4 Result = Slerp(float4(Start, 0), float4(End, 0), T);
+	return Result.xyz;
+}
+
 [numthreads(THREADGROUP_SIZEX, 1, 1)]
 void MainCS(
 	uint DispatchThreadId : SV_DispatchThreadID)	// DispatchThreadId = GroupId * int2(dimx,dimy) + GroupThreadId
@@ -46,6 +158,35 @@ void MainCS(
 		uint VertexIndex = DispatchThreadId;
 		uint Offset = VertexIndex + SkinCacheStart;
 
+#if BLEND_USING_VERTEX_COLOR > 0
+		float4 InputTangentX = TangentBias_SkinCache(TangentInputBuffer[2 * Offset + GPUSKIN_RWBUFFER_OFFSET_TANGENT_X]);
+		float4 InputTangentZ = TangentBias_SkinCache(TangentInputBuffer[2 * Offset + GPUSKIN_RWBUFFER_OFFSET_TANGENT_Z]);
+		float BlendFactor = ColorInputBuffer[DispatchThreadId + InputStreamStart].g;
+#if BLEND_USING_VERTEX_COLOR == 1
+		// No blending, choose between input tangents and recompute tangents
+		TangentX = BlendFactor < 0.5f ? InputTangentX.xyz : TangentX;
+		TangentZ = BlendFactor < 0.5f ? InputTangentZ.xyz : TangentZ;
+#elif BLEND_USING_VERTEX_COLOR == 2
+		// Linear interpolation
+		TangentX = lerp(InputTangentX.xyz, TangentX, BlendFactor);
+		TangentZ = lerp(InputTangentZ.xyz, TangentZ, BlendFactor);
+#elif BLEND_USING_VERTEX_COLOR == 3
+		// Vector slerp
+		TangentX = Slerp(InputTangentX.xyz, TangentX, BlendFactor);
+		TangentZ = Slerp(InputTangentZ.xyz, TangentZ, BlendFactor);
+#elif BLEND_USING_VERTEX_COLOR == 4
+		// Convert tangents into quaternion, apply slerp, then convert from quaternion back to tangents
+		float4 InputQuat = TangentToQuaternion(InputTangentX.xyz, InputTangentZ.xyz);
+		float4 Quat = TangentToQuaternion(TangentX, TangentZ);
+		float4 BlendedQuat = Slerp(InputQuat, Quat, BlendFactor);
+		QuaternionToTangent(BlendedQuat, TangentX, TangentZ);
+#endif
+		// No need to realign TangentX and TangentZ to be perfect perpendicular after blending as vertex factory will do so later on
+		TangentX = normalize(TangentX);
+		TangentZ = normalize(TangentZ);
+		Orientation = BlendFactor < 0.5f ? InputTangentZ.w : Orientation;
+#endif
+
 		TangentBufferUAV[2 * Offset + GPUSKIN_RWBUFFER_OFFSET_TANGENT_X] = TangentUnbias_SkinCache(float4(TangentX, 0.0f));
 		TangentBufferUAV[2 * Offset + GPUSKIN_RWBUFFER_OFFSET_TANGENT_Z] = TangentUnbias_SkinCache(float4(TangentZ, Orientation));
 	}
diff --git a/Engine/Source/Developer/RigVMDeveloper/Private/RigVMModel/RigVMController.cpp b/Engine/Source/Developer/RigVMDeveloper/Private/RigVMModel/RigVMController.cpp
--- a/Engine/Source/Developer/RigVMDeveloper/Private/RigVMModel/RigVMController.cpp
+++ b/Engine/Source/Developer/RigVMDeveloper/Private/RigVMModel/RigVMController.cpp
@@ -19,6 +19,7 @@
 URigVMController::URigVMController()
 	: bSuspendNotifications(false)
 	, bReportWarningsAndErrors(true)
+	, bIgnoreRerouteCompactnessChanges(false)
 {
 	SetExecuteContextStruct(FRigVMExecuteContext::StaticStruct());
 }
@@ -582,6 +583,8 @@ URigVMRerouteNode* URigVMController::AddRerouteNodeOnLink(URigVMLink* InLink, bo
 	URigVMPin* SourcePin = InLink->GetSourcePin();
 	URigVMPin* TargetPin = InLink->GetTargetPin();
 
+	TGuardValue<bool> GuardCompactness(bIgnoreRerouteCompactnessChanges, true);
+
 	FRigVMBaseAction Action;
 	if (bUndo)
 	{
@@ -634,6 +637,8 @@ URigVMRerouteNode* URigVMController::AddRerouteNodeOnPin(const FString& InPinPat
 		return nullptr;
 	}
 
+	TGuardValue<bool> GuardCompactness(bIgnoreRerouteCompactnessChanges, true);
+
 	FRigVMBaseAction Action;
 	if (bUndo)
 	{
@@ -1025,6 +1030,8 @@ bool URigVMController::Undo()
 	{
 		return false;
 	}
+
+	TGuardValue<bool> GuardCompactness(bIgnoreRerouteCompactnessChanges, true);
 	return ActionStack->Undo(this);
 }
 
@@ -1034,6 +1041,8 @@ bool URigVMController::Redo()
 	{
 		return false;
 	}
+
+	TGuardValue<bool> GuardCompactness(bIgnoreRerouteCompactnessChanges, true);
 	return ActionStack->Redo(this);
 }
 
@@ -1332,10 +1341,12 @@ bool URigVMController::RemoveNode(URigVMNode* InNode, bool bUndo)
 		return false;
 	}
 
-	FRigVMRemoveNodeAction Action;
+	TGuardValue<bool> GuardCompactness(bIgnoreRerouteCompactnessChanges, true);
+
+	FRigVMBaseAction Action;
 	if (bUndo)
 	{
-		Action = FRigVMRemoveNodeAction(InNode);
+		Action = FRigVMBaseAction();
 		Action.Title = FString::Printf(TEXT("Remove %s Node"), *InNode->GetNodeTitle());
 		ActionStack->BeginAction(Action);
 	}
@@ -1378,6 +1389,11 @@ bool URigVMController::RemoveNode(URigVMNode* InNode, bool bUndo)
 		}
 	}
 
+	if (bUndo)
+	{
+		ActionStack->AddAction(FRigVMRemoveNodeAction(InNode));
+	}
+
 	Graph->Nodes.Remove(InNode);
 	Graph->MarkPackageDirty();
 
@@ -1891,6 +1907,36 @@ bool URigVMController::RenameParameter(const FName& InOldName, const FName& InNe
 	return RenamedNodes.Num() > 0;
 }
 
+void URigVMController::UpdateRerouteNodeAfterChangingLinks(URigVMPin* PinChanged, bool bUndo)
+{
+	if (bIgnoreRerouteCompactnessChanges)
+	{
+		return;
+	}
+
+	if (!IsValidGraph())
+	{
+		return;
+	}
+
+	URigVMRerouteNode* Node = Cast<URigVMRerouteNode>(PinChanged->GetNode());
+	if (Node == nullptr)
+	{
+		return;
+	}
+
+	int32 NbTotalSources = Node->Pins[0]->GetSourceLinks(true /* recursive */).Num();
+	int32 NbTotalTargets = Node->Pins[0]->GetTargetLinks(true /* recursive */).Num();
+	int32 NbToplevelSources = Node->Pins[0]->GetSourceLinks(false /* recursive */).Num();
+	int32 NbToplevelTargets = Node->Pins[0]->GetSourceLinks(false /* recursive */).Num();
+
+	bool bJustTopLevelConnections = (NbTotalSources == NbToplevelSources) && (NbTotalTargets == NbToplevelTargets);
+	bool bOnlyConnectionsOnOneSide = (NbTotalSources == 0) || (NbTotalTargets == 0);
+	bool bShowAsFullNode = (!bJustTopLevelConnections) || bOnlyConnectionsOnOneSide;
+
+	SetRerouteCompactness(Node, bShowAsFullNode, bUndo);
+}
+
 bool URigVMController::SetPinExpansion(const FString& InPinPath, bool bIsExpanded, bool bUndo)
 {
 	if (!IsValidGraph())
@@ -2687,6 +2733,9 @@ bool URigVMController::AddLink(URigVMPin* OutputPin, URigVMPin* InputPin, bool b
 	Graph->MarkPackageDirty();
 	Notify(ERigVMGraphNotifType::LinkAdded, Link);
 
+	UpdateRerouteNodeAfterChangingLinks(OutputPin, bUndo);
+	UpdateRerouteNodeAfterChangingLinks(InputPin, bUndo);
+
 	if (bUndo)
 	{
 		ActionStack->EndAction(Action);
@@ -2756,6 +2805,9 @@ bool URigVMController::BreakLink(URigVMPin* OutputPin, URigVMPin* InputPin, bool
 			Link->Rename(nullptr, GetTransientPackage());
 			Link->MarkPendingKill();
 
+			UpdateRerouteNodeAfterChangingLinks(OutputPin, bUndo);
+			UpdateRerouteNodeAfterChangingLinks(InputPin, bUndo);
+
 			if (bUndo)
 			{
 				ActionStack->EndAction(Action);
diff --git a/Engine/Source/Developer/RigVMDeveloper/Public/RigVMModel/RigVMController.h b/Engine/Source/Developer/RigVMDeveloper/Public/RigVMModel/RigVMController.h
--- a/Engine/Source/Developer/RigVMDeveloper/Public/RigVMModel/RigVMController.h
+++ b/Engine/Source/Developer/RigVMDeveloper/Public/RigVMModel/RigVMController.h
@@ -410,6 +410,7 @@ private:
 	bool BreakLink(URigVMPin* OutputPin, URigVMPin* InputPin, bool bUndo);
 	bool BreakAllLinks(URigVMPin* Pin, bool bAsInput, bool bUndo);
 	void BreakAllLinksRecursive(URigVMPin* Pin, bool bAsInput, bool bTowardsParent, bool bUndo);
+	void UpdateRerouteNodeAfterChangingLinks(URigVMPin* PinChanged, bool bUndo = true);
 	bool SetPinExpansion(URigVMPin* InPin, bool bIsExpanded, bool bUndo = true);
 	void ExpandPinRecursively(URigVMPin* InPin, bool bUndo);
 	bool SetPinIsWatched(URigVMPin* InPin, bool bIsWatched, bool bUndo);
@@ -452,6 +453,7 @@ private:
 
 	bool bSuspendNotifications;
 	bool bReportWarningsAndErrors;
+	bool bIgnoreRerouteCompactnessChanges;
 
 	friend class URigVMGraph;
 	friend class URigVMActionStack;
diff --git a/Engine/Source/Editor/ClothingSystemEditor/ClothingSystemEditor.Build.cs b/Engine/Source/Editor/ClothingSystemEditor/ClothingSystemEditor.Build.cs
--- a/Engine/Source/Editor/ClothingSystemEditor/ClothingSystemEditor.Build.cs
+++ b/Engine/Source/Editor/ClothingSystemEditor/ClothingSystemEditor.Build.cs
@@ -35,7 +35,13 @@ public class ClothingSystemEditor : ModuleRules
                 "ClothingSystemEditorInterface"
             }
 		);
-        
+
+		PrivateIncludePathModuleNames.AddRange(
+            new string[] {
+                "Persona",
+            }
+        );
+
         SetupModulePhysicsSupport(Target);
 	}
 }
diff --git a/Engine/Source/Editor/CurveEditor/Private/CurveEditorSelection.cpp b/Engine/Source/Editor/CurveEditor/Private/CurveEditorSelection.cpp
--- a/Engine/Source/Editor/CurveEditor/Private/CurveEditorSelection.cpp
+++ b/Engine/Source/Editor/CurveEditor/Private/CurveEditorSelection.cpp
@@ -65,14 +65,19 @@ void FCurveEditorSelection::Add(FCurveModelID CurveID, ECurvePointType PointType
 void FCurveEditorSelection::Add(FCurveModelID CurveID, ECurvePointType PointType, TArrayView<const FKeyHandle> Keys)
 {
 	TSharedPtr<FCurveEditor> CurveEditor = WeakCurveEditor.Pin();
-	if (Keys.Num() > 0 && CurveEditor && !CurveEditor->GetCurves()[CurveID]->IsReadOnly())
+	if (Keys.Num() > 0 && CurveEditor)
 	{
-		ChangeSelectionPointType(PointType);
-
-		FKeyHandleSet& SelectedKeys = CurveToSelectedKeys.FindOrAdd(CurveID);
-		for (FKeyHandle Key : Keys)
+		const TUniquePtr<FCurveModel> *CurveModel = CurveEditor->GetCurves().Find(CurveID);
+		if (CurveModel && CurveModel->IsValid() && !(*CurveModel)->IsReadOnly())
 		{
-			SelectedKeys.Add(Key);
+			ChangeSelectionPointType(PointType);
+
+			FKeyHandleSet& SelectedKeys = CurveToSelectedKeys.FindOrAdd(CurveID);
+			for (FKeyHandle Key : Keys)
+			{
+				SelectedKeys.Add(Key);
+			}
+
 		}
 	}
 
diff --git a/Engine/Source/Editor/EditorWidgets/Private/STransportControl.cpp b/Engine/Source/Editor/EditorWidgets/Private/STransportControl.cpp
--- a/Engine/Source/Editor/EditorWidgets/Private/STransportControl.cpp
+++ b/Engine/Source/Editor/EditorWidgets/Private/STransportControl.cpp
@@ -84,7 +84,7 @@ TSharedPtr<SWidget> STransportControl::MakeTransportControlWidget(ETransportCont
 			. ContentPadding(2.0f)
 			. IsFocusable(bAreButtonsFocusable);
 	case ETransportControlWidgetType::Loop:
-		return SNew(SButton)
+		return SAssignNew(LoopButton, SButton)
 			. OnClicked(TransportControlArgs.OnToggleLooping)
 			. Visibility(TransportControlArgs.OnGetLooping.IsBound() ? EVisibility::Visible : EVisibility::Collapsed)
 			. ToolTipText( LOCTEXT("Loop", "Loop") )
@@ -264,12 +264,12 @@ const FSlateBrush* STransportControl::GetLoopStatusIcon() const
 	if (TransportControlArgs.OnGetLooping.IsBound() &&
 		TransportControlArgs.OnGetLooping.Execute())
 	{
-		return BackwardPlayButton.IsValid() && BackwardPlayButton->IsPressed() ? 
+		return LoopButton.IsValid() && LoopButton->IsPressed() ? 
 			&FEditorStyle::Get().GetWidgetStyle<FButtonStyle>("Animation.Loop.Enabled").Pressed : 
 			&FEditorStyle::Get().GetWidgetStyle<FButtonStyle>("Animation.Loop.Enabled").Normal;
 	}
 
-	return BackwardPlayButton.IsValid() && BackwardPlayButton->IsPressed() ? 
+	return LoopButton.IsValid() && LoopButton->IsPressed() ? 
 		&FEditorStyle::Get().GetWidgetStyle<FButtonStyle>("Animation.Loop.Disabled").Pressed : 
 		&FEditorStyle::Get().GetWidgetStyle<FButtonStyle>("Animation.Loop.Disabled").Normal;
 }
diff --git a/Engine/Source/Editor/EditorWidgets/Private/STransportControl.h b/Engine/Source/Editor/EditorWidgets/Private/STransportControl.h
--- a/Engine/Source/Editor/EditorWidgets/Private/STransportControl.h
+++ b/Engine/Source/Editor/EditorWidgets/Private/STransportControl.h
@@ -67,4 +67,5 @@ private:
 	TSharedPtr<SButton> ForwardPlayButton;
 	TSharedPtr<SButton> BackwardPlayButton;
 	TSharedPtr<SButton> RecordButton;
+	TSharedPtr<SButton> LoopButton;
 };
diff --git a/Engine/Source/Editor/GraphEditor/Private/SGraphPin.cpp b/Engine/Source/Editor/GraphEditor/Private/SGraphPin.cpp
--- a/Engine/Source/Editor/GraphEditor/Private/SGraphPin.cpp
+++ b/Engine/Source/Editor/GraphEditor/Private/SGraphPin.cpp
@@ -19,6 +19,7 @@
 #include "ScopedTransaction.h"
 #include "SLevelOfDetailBranchNode.h"
 #include "SPinTypeSelector.h"
+#include "Animation/AnimNodeBase.h"
 
 /////////////////////////////////////////////////////
 // FGraphPinHandle
@@ -69,11 +70,14 @@ SGraphPin::SGraphPin()
 	: GraphPinObj(nullptr)
 	, PinColorModifier(FLinearColor::White)
 	, CachedNodeOffset(FVector2D::ZeroVector)
+	, Custom_Brush_Connected(nullptr)
+	, Custom_Brush_Disconnected(nullptr)
 	, bGraphDataInvalid(false)
 	, bShowLabel(true)
 	, bOnlyShowDefaultValue(false)
 	, bIsMovingLinks(false)
 	, bUsePinColorForText(false)
+
 {
 	IsEditable = true;
 
@@ -103,6 +107,9 @@ SGraphPin::SGraphPin()
 	static const FName NAME_Pin_Background("Graph.Pin.Background");
 	static const FName NAME_Pin_BackgroundHovered("Graph.Pin.BackgroundHovered");
 
+	static const FName NAME_PosePin_Connected("Graph.PosePin.Connected");
+	static const FName NAME_PosePin_Disconnected("Graph.PosePin.Disconnected");
+
 	const EBlueprintPinStyleType StyleType = GetDefault<UGraphEditorSettings>()->DataPinStyle;
 
 	switch(StyleType)
@@ -127,6 +134,9 @@ SGraphPin::SGraphPin()
 	CachedImg_DelegatePin_Connected = FEditorStyle::GetBrush( NAME_DelegatePin_Connected );
 	CachedImg_DelegatePin_Disconnected = FEditorStyle::GetBrush( NAME_DelegatePin_Disconnected );
 
+	CachedImg_PosePin_Connected = FEditorStyle::GetBrush(NAME_PosePin_Connected);
+	CachedImg_PosePin_Disconnected = FEditorStyle::GetBrush(NAME_PosePin_Disconnected);
+
 	CachedImg_SetPin = FEditorStyle::GetBrush(NAME_SetPin);
 	CachedImg_MapPinKey = FEditorStyle::GetBrush(NAME_MapPinKey);
 	CachedImg_MapPinValue = FEditorStyle::GetBrush(NAME_MapPinValue);
@@ -886,6 +896,18 @@ bool SGraphPin::IsConnected() const
 /** @return The brush with which to pain this graph pin's incoming/outgoing bullet point */
 const FSlateBrush* SGraphPin::GetPinIcon() const
 {
+	if (Custom_Brush_Connected && Custom_Brush_Disconnected)
+	{
+		if (IsConnected())
+		{
+			return Custom_Brush_Connected;
+		}
+		else
+		{
+			return Custom_Brush_Disconnected;
+		}
+	}
+
 	if (IsArray())
 	{
 		if (IsConnected())
@@ -928,6 +950,17 @@ const FSlateBrush* SGraphPin::GetPinIcon() const
 	{
 		return CachedImg_MapPinKey;
 	}
+	else if (GraphPinObj->PinType.PinCategory == UEdGraphSchema_K2::PC_Struct && ((GraphPinObj->PinType.PinSubCategoryObject == FPoseLink::StaticStruct()) || (GraphPinObj->PinType.PinSubCategoryObject == FComponentSpacePoseLink::StaticStruct())))
+	{
+		if (IsConnected())
+		{
+			return CachedImg_PosePin_Connected;
+		}
+		else
+		{
+			return CachedImg_PosePin_Disconnected;
+		}
+	}
 	else
 	{
 		if (IsConnected())
@@ -1169,6 +1202,12 @@ EVisibility SGraphPin::GetPinVisiblity() const
 	return EVisibility::Visible;
 }
 
+void SGraphPin::SetCustomPinIcon(const FSlateBrush* InConnectedBrush, const FSlateBrush* InDisconnectedBrush)
+{
+	Custom_Brush_Connected = InConnectedBrush;
+	Custom_Brush_Disconnected = InDisconnectedBrush;
+}
+
 bool SGraphPin::GetIsConnectable() const
 {
 	const bool bCanConnectToPin = !GraphPinObj->bNotConnectable;
diff --git a/Engine/Source/Editor/GraphEditor/Public/SGraphPin.h b/Engine/Source/Editor/GraphEditor/Public/SGraphPin.h
--- a/Engine/Source/Editor/GraphEditor/Public/SGraphPin.h
+++ b/Engine/Source/Editor/GraphEditor/Public/SGraphPin.h
@@ -199,6 +199,10 @@ public:
 	  * but for added safety we track graph data invalidation here:
 	  */
 	void InvalidateGraphData() { bGraphDataInvalid = true; }
+
+	/** Override the visual look of the pin by providing two custom brushes */
+	void SetCustomPinIcon(const FSlateBrush* InConnectedBrush, const FSlateBrush* InDisconnectedBrush);
+
 protected:
 
 	/** If true the graph pin subclass is responsible for setting the IsEnabled delegates for the aspects it cares about. If false, the default value widget enabling is done by the base class */
@@ -295,13 +299,18 @@ protected:
 	const FSlateBrush* CachedImg_Pin_Disconnected;
 	const FSlateBrush* CachedImg_DelegatePin_Connected;
 	const FSlateBrush* CachedImg_DelegatePin_Disconnected;
+	const FSlateBrush* CachedImg_PosePin_Connected;
+	const FSlateBrush* CachedImg_PosePin_Disconnected;
 	const FSlateBrush* CachedImg_SetPin;
 	const FSlateBrush* CachedImg_MapPinKey;
 	const FSlateBrush* CachedImg_MapPinValue;
 
 	const FSlateBrush* CachedImg_Pin_Background;
 	const FSlateBrush* CachedImg_Pin_BackgroundHovered;
-	
+
+	const FSlateBrush* Custom_Brush_Connected;
+	const FSlateBrush* Custom_Brush_Disconnected;
+
 	/** flag indicating that graph data has been deleted by the user */
 	bool bGraphDataInvalid;
 
diff --git a/Engine/Source/Editor/Kismet/Private/SKismetInspector.cpp b/Engine/Source/Editor/Kismet/Private/SKismetInspector.cpp
--- a/Engine/Source/Editor/Kismet/Private/SKismetInspector.cpp
+++ b/Engine/Source/Editor/Kismet/Private/SKismetInspector.cpp
@@ -234,6 +234,11 @@ bool SKismetInspector::IsSelected(UObject* Object) const
 	return false;
 }
 
+const TArray< TWeakObjectPtr<UObject> >& SKismetInspector::GetSelectedObjects() const
+{
+	return SelectedObjects;
+}
+
 FText SKismetInspector::GetContextualEditingWidgetTitle() const
 {
 	FText Title = PropertyViewTitle;
diff --git a/Engine/Source/Editor/Kismet/Public/SKismetInspector.h b/Engine/Source/Editor/Kismet/Public/SKismetInspector.h
--- a/Engine/Source/Editor/Kismet/Public/SKismetInspector.h
+++ b/Engine/Source/Editor/Kismet/Public/SKismetInspector.h
@@ -98,6 +98,9 @@ public:
 	/** @return true if the object is in the selection set. */
 	bool IsSelected(UObject* Object) const;
 
+	/** returns the list of selected objects */
+	const TArray< TWeakObjectPtr<UObject> >& GetSelectedObjects() const;
+
 protected:
 	/** Update the inspector window to show information on the supplied objects */
 	void UpdateFromObjects(const TArray<UObject*>& PropertyObjects, struct FKismetSelectionInfo& SelectionInfo, const FShowDetailsOptions& Options);
diff --git a/Engine/Source/Editor/KismetCompiler/Private/AnimBlueprintCompiler.cpp b/Engine/Source/Editor/KismetCompiler/Private/AnimBlueprintCompiler.cpp
--- a/Engine/Source/Editor/KismetCompiler/Private/AnimBlueprintCompiler.cpp
+++ b/Engine/Source/Editor/KismetCompiler/Private/AnimBlueprintCompiler.cpp
@@ -3223,6 +3223,18 @@ void FAnimBlueprintCompilerContext::FPropertyCopyRecord::ValidateFastPath(UClass
 				InvalidateFastPath();
 				return;
 			}
+
+			// If the property is not one of the following types or categories, then we ignore the fast
+			// path. This matches the supported copyable types in FExposedValueHandler::Initialize
+			if (!(CastField<FBoolProperty>(DestProperty) ||
+				  CastField<FNameProperty>(DestProperty) ||
+				  CastField<FStructProperty>(DestProperty) ||
+				  CastField<FObjectPropertyBase>(DestProperty) ||
+				  (DestProperty->PropertyFlags & CPF_IsPlainOldData) != 0))
+			{
+				InvalidateFastPath();
+				return;
+			}
 		}
 		else
 		{
@@ -3399,22 +3411,28 @@ void FAnimBlueprintCompilerContext::PrecompileFunction(FKismetFunctionContext& C
 {
 	Super::PrecompileFunction(Context, InternalFlags);
 
-	auto CompareEntryPointName =
-	[Function = Context.Function](UEdGraph* InGraph)
+	if(Context.Function)
 	{
-		TArray<UK2Node_FunctionEntry*> EntryPoints;
-		InGraph->GetNodesOfClass<UK2Node_FunctionEntry>(EntryPoints);
-		if(EntryPoints.Num() == 1)
+		auto CompareEntryPointName =
+		[Function = Context.Function](UEdGraph* InGraph)
 		{
-			return EntryPoints[0]->CustomGeneratedFunctionName == Function->GetFName(); 
-		}
-		return true;
-	};
+			if(InGraph)
+			{
+				TArray<UK2Node_FunctionEntry*> EntryPoints;
+				InGraph->GetNodesOfClass<UK2Node_FunctionEntry>(EntryPoints);
+				if(EntryPoints.Num() == 1 && EntryPoints[0])
+				{
+					return EntryPoints[0]->CustomGeneratedFunctionName == Function->GetFName(); 
+				}
+			}
+			return true;
+		};
 
-	if(GeneratedStubGraphs.ContainsByPredicate(CompareEntryPointName))
-	{
-		Context.Function->SetMetaData(FBlueprintMetadata::MD_BlueprintInternalUseOnly, TEXT("true"));
-		Context.Function->SetMetaData(FBlueprintMetadata::MD_AnimBlueprintFunction, TEXT("true"));
+		if(GeneratedStubGraphs.ContainsByPredicate(CompareEntryPointName))
+		{
+			Context.Function->SetMetaData(FBlueprintMetadata::MD_BlueprintInternalUseOnly, TEXT("true"));
+			Context.Function->SetMetaData(FBlueprintMetadata::MD_AnimBlueprintFunction, TEXT("true"));
+		}
 	}
 }
 
@@ -3422,23 +3440,29 @@ void FAnimBlueprintCompilerContext::SetCalculatedMetaDataAndFlags(UFunction* Fun
 {
 	Super::SetCalculatedMetaDataAndFlags(Function, EntryNode, K2Schema);
 
-	auto CompareEntryPointName =
-	[Function](UEdGraph* InGraph)
+	if(Function)
 	{
-		TArray<UK2Node_FunctionEntry*> EntryPoints;
-		InGraph->GetNodesOfClass<UK2Node_FunctionEntry>(EntryPoints);
-		if(EntryPoints.Num() == 1)
+		auto CompareEntryPointName =
+		[Function](UEdGraph* InGraph)
 		{
-			return EntryPoints[0]->CustomGeneratedFunctionName == Function->GetFName(); 
-		}
-		return true;
-	};
+			if(InGraph)
+			{
+				TArray<UK2Node_FunctionEntry*> EntryPoints;
+				InGraph->GetNodesOfClass<UK2Node_FunctionEntry>(EntryPoints);
+				if(EntryPoints.Num() == 1 && EntryPoints[0])
+				{
+					return EntryPoints[0]->CustomGeneratedFunctionName == Function->GetFName(); 
+				}
+			}
+			return true;
+		};
 
-	// Match by name to generated graph's entry points
-	if(GeneratedStubGraphs.ContainsByPredicate(CompareEntryPointName))
-	{
-		Function->SetMetaData(FBlueprintMetadata::MD_BlueprintInternalUseOnly, TEXT("true"));
-		Function->SetMetaData(FBlueprintMetadata::MD_AnimBlueprintFunction, TEXT("true"));
+		// Match by name to generated graph's entry points
+		if(GeneratedStubGraphs.ContainsByPredicate(CompareEntryPointName))
+		{
+			Function->SetMetaData(FBlueprintMetadata::MD_BlueprintInternalUseOnly, TEXT("true"));
+			Function->SetMetaData(FBlueprintMetadata::MD_AnimBlueprintFunction, TEXT("true"));
+		}
 	}
 }
 
diff --git a/Engine/Source/Editor/LevelEditor/Private/SLevelEditor.cpp b/Engine/Source/Editor/LevelEditor/Private/SLevelEditor.cpp
--- a/Engine/Source/Editor/LevelEditor/Private/SLevelEditor.cpp
+++ b/Engine/Source/Editor/LevelEditor/Private/SLevelEditor.cpp
@@ -167,6 +167,7 @@ void SLevelEditor::Construct( const SLevelEditor::FArguments& InArgs)
 
 	// We need to register when modes list changes so that we can refresh the auto generated commands.
 	FEditorModeRegistry::Get().OnRegisteredModesChanged().AddRaw(this, &SLevelEditor::EditorModeCommandsChanged);
+	GLevelEditorModeTools().OnEditorModeIDChanged().AddSP(this, &SLevelEditor::OnEditorModeIdChanged);
 
 	// @todo This is a hack to get this working for now. This won't work with multiple worlds
 	GEditor->GetEditorWorldContext(true).AddRef(World);
@@ -1415,22 +1416,6 @@ void SLevelEditor::ToggleEditorMode( FEditorModeID ModeID )
 	// Find and disable any other 'visible' modes since we only ever allow one of those active at a time.
 	GLevelEditorModeTools().DeactivateOtherVisibleModes(ModeID);
 
-	FLevelEditorModule& LevelEditorModule = FModuleManager::GetModuleChecked<FLevelEditorModule>( "LevelEditor" );
-	TSharedPtr<FTabManager> LevelEditorTabManager = LevelEditorModule.GetLevelEditorTabManager();
-
-	if (GLevelEditorModeTools().IsModeActive(FBuiltinEditorModes::EM_Default) && !GetDefault<UEditorStyleSettings>()->bEnableLegacyEditorModeUI)
-	{
-		TSharedPtr<SDockTab> ToolboxTab = LevelEditorTabManager->FindExistingLiveTab(LevelEditorTabIds::LevelEditorToolBox);
-		if (ToolboxTab.IsValid())
-		{
-			ToolboxTab->RequestCloseTab();
-		}
-	}
-	else
-	{
-		LevelEditorTabManager->InvokeTab(LevelEditorTabIds::LevelEditorToolBox);
-	}
-
 }
 
 bool SLevelEditor::IsModeActive( FEditorModeID ModeID )
@@ -1456,6 +1441,28 @@ void SLevelEditor::EditorModeCommandsChanged()
 	RefreshEditorModeCommands();
 }
 
+void SLevelEditor::OnEditorModeIdChanged(const FEditorModeID& ModeChangedID, bool bIsEnteringMode)
+{
+	FEdMode* EdMode = GLevelEditorModeTools().GetActiveMode(ModeChangedID);
+	if (EdMode && EdMode->UsesToolkits() && EdMode->GetToolkit().IsValid())
+	{
+		FLevelEditorModule& LevelEditorModule = FModuleManager::GetModuleChecked<FLevelEditorModule>("LevelEditor");
+		TSharedPtr<FTabManager> LevelEditorTabManager = LevelEditorModule.GetLevelEditorTabManager();
+
+		if ((ModeChangedID == FBuiltinEditorModes::EM_Default || !GLevelEditorModeTools().ShouldShowModeToolbox()) && !GetDefault<UEditorStyleSettings>()->bEnableLegacyEditorModeUI)
+		{
+			TSharedPtr<SDockTab> ToolboxTab = LevelEditorTabManager->FindExistingLiveTab(LevelEditorTabIds::LevelEditorToolBox);
+			if (ToolboxTab.IsValid())
+			{
+				ToolboxTab->RequestCloseTab();
+			}
+		}
+		else
+		{
+			LevelEditorTabManager->InvokeTab(LevelEditorTabIds::LevelEditorToolBox);
+		}
+	}
+}
 void SLevelEditor::RefreshEditorModeCommands()
 {
 	FLevelEditorModule& LevelEditorModule = FModuleManager::GetModuleChecked<FLevelEditorModule>( "LevelEditor" );
diff --git a/Engine/Source/Editor/LevelEditor/Private/SLevelEditor.h b/Engine/Source/Editor/LevelEditor/Private/SLevelEditor.h
--- a/Engine/Source/Editor/LevelEditor/Private/SLevelEditor.h
+++ b/Engine/Source/Editor/LevelEditor/Private/SLevelEditor.h
@@ -167,6 +167,9 @@ private:
 	/** Editor mode has been added or removed, clears cached command list so it will be rebuilt */
 	void EditorModeCommandsChanged();
 
+	/** Called when a level editor mode is toggled */
+	void OnEditorModeIdChanged(const FEditorModeID& ModeChangedID, bool bIsEnteringMode);
+
 	/** Gets the tabId mapping to an editor mode */
 	static FName GetEditorModeTabId( FEditorModeID ModeID );
 
diff --git a/Engine/Source/Editor/MovieSceneTools/MovieSceneTools.Build.cs b/Engine/Source/Editor/MovieSceneTools/MovieSceneTools.Build.cs
--- a/Engine/Source/Editor/MovieSceneTools/MovieSceneTools.Build.cs
+++ b/Engine/Source/Editor/MovieSceneTools/MovieSceneTools.Build.cs
@@ -40,6 +40,7 @@ public class MovieSceneTools : ModuleRules
                 "Json",
                 "JsonUtilities",
 				"LevelSequence",
+                "LiveLinkInterface",
                 "MessageLog",
 				"MovieScene",
 				"MovieSceneTracks",
diff --git a/Engine/Source/Editor/MovieSceneTools/Private/MovieSceneToolHelpers.cpp b/Engine/Source/Editor/MovieSceneTools/Private/MovieSceneToolHelpers.cpp
--- a/Engine/Source/Editor/MovieSceneTools/Private/MovieSceneToolHelpers.cpp
+++ b/Engine/Source/Editor/MovieSceneTools/Private/MovieSceneToolHelpers.cpp
@@ -62,7 +62,10 @@
 #include "Serialization/ObjectReader.h"
 #include "AnimationRecorder.h"
 #include "Components/SkeletalMeshComponent.h"
-
+#include "ILiveLinkClient.h"
+#include "LiveLinkPresetTypes.h"
+#include "LiveLinkSourceSettings.h"
+#include "Features/IModularFeatures.h"
 /* MovieSceneToolHelpers
  *****************************************************************************/
 
@@ -2065,9 +2068,60 @@ bool MovieSceneToolHelpers::ExportToAnimSequence(UAnimSequence* AnimSequence, UM
 	}
 
 	UnFbx::FLevelSequenceAnimTrackAdapter AnimTrackAdapter(Player, MovieScene, RootToLocalTransform);
-
+	int32 LocalStartFrame = AnimTrackAdapter.GetLocalStartFrame();
+	int32 StartFrame = AnimTrackAdapter.GetStartFrame();
+	int32 AnimationLength = AnimTrackAdapter.GetLength();
+	float FrameRate = AnimTrackAdapter.GetFrameRate();
+	float DeltaTime = 1.0f / FrameRate;
 	FFrameRate SampleRate = MovieScene->GetDisplayRate();
 
+
+	//If we are running with a live link track we need to do a few things.
+	// 1. First test to see if we have one, only way to really do that is to see if we have a source that has the `Sequencer Live Link Track`.  We also evalute the first frame in case we are out of range and the sources aren't created yet.
+	// 2. Make sure Sequencer.AlwaysSendInterpolated.LiveLink is non-zero, and then set it back to zero if it's not.
+	// 3. For each live link sequencer source we need to set the ELiveLinkSourceMode to Latest so that we just get the latest and don't use engine/timecode for any interpolation.
+
+	ILiveLinkClient* LiveLinkClient = nullptr;
+	IModularFeatures& ModularFeatures = IModularFeatures::Get();
+	TOptional<int32> SequencerAlwaysSenedLiveLinkInterpolated;
+	TMap<FGuid, ELiveLinkSourceMode>  SourceAndMode;
+	IConsoleVariable* CVarAlwaysSendInterpolatedLiveLink = IConsoleManager::Get().FindConsoleVariable(TEXT("Sequencer.AlwaysSendInterpolatedLiveLink"));
+	bool bHaveAtLeastOneSource = false;
+	if (ModularFeatures.IsModularFeatureAvailable(ILiveLinkClient::ModularFeatureName))
+	{
+
+		LiveLinkClient = &ModularFeatures.GetModularFeature<ILiveLinkClient>(ILiveLinkClient::ModularFeatureName);
+		if (LiveLinkClient)
+		{
+			//Evaluate the first frame we may be out of range and not have any live link sources created yet.
+			AnimTrackAdapter.UpdateAnimation(StartFrame);
+			TArray<FGuid> Sources = LiveLinkClient->GetSources();
+			for (const FGuid& Guid : Sources)
+			{
+				FText SourceTypeText = LiveLinkClient->GetSourceType(Guid);
+				FString SourceTypeStr = SourceTypeText.ToString();
+				if (SourceTypeStr.Contains(TEXT("Sequencer Live Link")))
+				{
+					bHaveAtLeastOneSource = true;
+					ULiveLinkSourceSettings* Settings = LiveLinkClient->GetSourceSettings(Guid);
+					if (Settings)
+					{
+						if (Settings->Mode != ELiveLinkSourceMode::Latest)
+						{
+							SourceAndMode.Add(Guid, Settings->Mode);
+							Settings->Mode = ELiveLinkSourceMode::Latest;
+						}
+						if (CVarAlwaysSendInterpolatedLiveLink)
+						{
+							SequencerAlwaysSenedLiveLinkInterpolated = CVarAlwaysSendInterpolatedLiveLink->GetInt();
+							CVarAlwaysSendInterpolatedLiveLink->Set(1, ECVF_SetByConsole);
+						}
+					}
+				}
+			}
+		}
+	}
+
 	FAnimRecorderInstance AnimationRecorder;
 	FAnimationRecordingSettings RecordingSettings;
 	RecordingSettings.SampleRate = SampleRate.AsDecimal();
@@ -2078,22 +2132,27 @@ bool MovieSceneToolHelpers::ExportToAnimSequence(UAnimSequence* AnimSequence, UM
 	RecordingSettings.bRemoveRootAnimation = false;
 	RecordingSettings.bCheckDeltaTimeAtBeginning = false;
 	AnimationRecorder.Init(SkelMeshComp, AnimSequence, nullptr, RecordingSettings);
-	AnimationRecorder.BeginRecording();
-
 
-	int32 LocalStartFrame = AnimTrackAdapter.GetLocalStartFrame();
-	int32 StartFrame = AnimTrackAdapter.GetStartFrame();
-	int32 AnimationLength = AnimTrackAdapter.GetLength();
-	float FrameRate = AnimTrackAdapter.GetFrameRate();
-	float DeltaTime = 1.0f / FrameRate;
-	for (int32 FrameCount = 0; FrameCount <= AnimationLength; ++FrameCount)
+	//need to run first time before start record.
+	AnimTrackAdapter.UpdateAnimation(LocalStartFrame);
+	if (LiveLinkClient && bHaveAtLeastOneSource)
 	{
+		LiveLinkClient->ForceTick();
+	}
+	AnimationRecorder.BeginRecording();
 
+	for (int32 FrameCount = 1; FrameCount <= AnimationLength; ++FrameCount)
+	{
 		int32 LocalFrame = LocalStartFrame + FrameCount;
 
 		// This will call UpdateSkelPose on the skeletal mesh component to move bones based on animations in the matinee group
 		AnimTrackAdapter.UpdateAnimation(LocalFrame);
 
+		if (LiveLinkClient && bHaveAtLeastOneSource)
+		{
+			LiveLinkClient->ForceTick();
+		}
+
 		// Update space bases so new animation position has an effect.
 		// @todo - hack - this will be removed at some point (this comment is all over the place by the way in fbx export code).
 		SkelMeshComp->TickAnimation(0.03f, false);
@@ -2111,6 +2170,23 @@ bool MovieSceneToolHelpers::ExportToAnimSequence(UAnimSequence* AnimSequence, UM
 	const bool bShowAnimationAssetCreatedToast = false;
 	AnimationRecorder.FinishRecording(bShowAnimationAssetCreatedToast);
 
+	//now do any sequencer live link cleanup
+	if (LiveLinkClient && bHaveAtLeastOneSource)
+	{
+		for (TPair<FGuid, ELiveLinkSourceMode>& Item: SourceAndMode)
+		{
+			ULiveLinkSourceSettings* Settings = LiveLinkClient->GetSourceSettings(Item.Key);
+			if (Settings)
+			{
+				Settings->Mode = Item.Value;
+			}
+		}
+	}
+
+	if(SequencerAlwaysSenedLiveLinkInterpolated.IsSet() && CVarAlwaysSendInterpolatedLiveLink)
+	{
+		CVarAlwaysSendInterpolatedLiveLink->Set(0, ECVF_SetByConsole);
+	}
 	return true;
 }
 
diff --git a/Engine/Source/Editor/Persona/Private/AnimTimeline/AnimModel_AnimSequenceBase.cpp b/Engine/Source/Editor/Persona/Private/AnimTimeline/AnimModel_AnimSequenceBase.cpp
--- a/Engine/Source/Editor/Persona/Private/AnimTimeline/AnimModel_AnimSequenceBase.cpp
+++ b/Engine/Source/Editor/Persona/Private/AnimTimeline/AnimModel_AnimSequenceBase.cpp
@@ -138,9 +138,13 @@ void FAnimModel_AnimSequenceBase::RefreshNotifyTracks()
 	NotifyRoot->ClearChildren();
 	RootTracks.Add(NotifyRoot.ToSharedRef());
 
-	TSharedRef<FAnimTimelineTrack_NotifiesPanel> NotifyPanel = MakeShared<FAnimTimelineTrack_NotifiesPanel>(SharedThis(this));
-	NotifyRoot->AddChild(NotifyPanel);
-	NotifyRoot->SetAnimNotifyPanel(NotifyPanel);
+	if(!NotifyPanel.IsValid())
+	{
+		NotifyPanel = MakeShared<FAnimTimelineTrack_NotifiesPanel>(SharedThis(this));
+		NotifyRoot->SetAnimNotifyPanel(NotifyPanel.ToSharedRef());
+	}
+
+	NotifyRoot->AddChild(NotifyPanel.ToSharedRef());
 }
 
 void FAnimModel_AnimSequenceBase::RefreshCurveTracks()
diff --git a/Engine/Source/Editor/Persona/Private/AnimTimeline/AnimModel_AnimSequenceBase.h b/Engine/Source/Editor/Persona/Private/AnimTimeline/AnimModel_AnimSequenceBase.h
--- a/Engine/Source/Editor/Persona/Private/AnimTimeline/AnimModel_AnimSequenceBase.h
+++ b/Engine/Source/Editor/Persona/Private/AnimTimeline/AnimModel_AnimSequenceBase.h
@@ -11,6 +11,7 @@ class UAnimSequenceBase;
 class FAnimTimelineTrack_Notifies;
 class FAnimTimelineTrack_Curves;
 class FAnimTimelineTrack;
+class FAnimTimelineTrack_NotifiesPanel;
 enum class EFrameNumberDisplayFormats : uint8;
 
 /** Anim model for an anim sequence base */
@@ -76,6 +77,9 @@ private:
 	/** Root track for notifies */
 	TSharedPtr<FAnimTimelineTrack_Notifies> NotifyRoot;
 
+	/** Legacy notify panel track */
+	TSharedPtr<FAnimTimelineTrack_NotifiesPanel> NotifyPanel;
+
 	/** Root track for curves */
 	TSharedPtr<FAnimTimelineTrack_Curves> CurveRoot;
 
diff --git a/Engine/Source/Editor/Persona/Private/AnimTimeline/AnimTimelineTrack_Curve.cpp b/Engine/Source/Editor/Persona/Private/AnimTimeline/AnimTimelineTrack_Curve.cpp
--- a/Engine/Source/Editor/Persona/Private/AnimTimeline/AnimTimelineTrack_Curve.cpp
+++ b/Engine/Source/Editor/Persona/Private/AnimTimeline/AnimTimelineTrack_Curve.cpp
@@ -356,8 +356,14 @@ FReply FAnimTimelineTrack_Curve::HandleDoubleClicked(const FGeometry& InGeometry
 
 void FAnimTimelineTrack_Curve::HandleCurveChanged()
 {
-	CurveEditor->ZoomToFit(EAxisList::Y);
-	CurveOverlay->BuildBoundsLabels();
+	if(CurveEditor.IsValid())
+	{
+		CurveEditor->ZoomToFit(EAxisList::Y);
+	}
+	if(CurveOverlay.IsValid())
+	{
+		CurveOverlay->BuildBoundsLabels();
+	}
 
 	GetModel()->GetAnimSequenceBase()->MarkRawDataAsModified();
 
@@ -367,6 +373,18 @@ void FAnimTimelineTrack_Curve::HandleCurveChanged()
 	}
 }
 
+void FAnimTimelineTrack_Curve::PostUndoRedo()
+{
+	if(CurveEditor.IsValid())
+	{
+		CurveEditor->ZoomToFit(EAxisList::Y);
+	}
+	if(CurveOverlay.IsValid())
+	{
+		CurveOverlay->BuildBoundsLabels();
+	}
+}
+
 void FAnimTimelineTrack_Curve::HendleEditCurve()
 {
 	TArray<IAnimationEditor::FCurveEditInfo> EditCurveInfo;
diff --git a/Engine/Source/Editor/Persona/Private/AnimTimeline/AnimTimelineTrack_Curve.h b/Engine/Source/Editor/Persona/Private/AnimTimeline/AnimTimelineTrack_Curve.h
--- a/Engine/Source/Editor/Persona/Private/AnimTimeline/AnimTimelineTrack_Curve.h
+++ b/Engine/Source/Editor/Persona/Private/AnimTimeline/AnimTimelineTrack_Curve.h
@@ -4,6 +4,7 @@
 
 #include "AnimTimelineTrack.h"
 #include "Animation/Skeleton.h"
+#include "EditorUndoClient.h"
 
 class SBorder;
 class FCurveEditor;
@@ -11,7 +12,7 @@ struct FRichCurve;
 class SHorizontalBox;
 class SCurveBoundsOverlay;
 
-class FAnimTimelineTrack_Curve : public FAnimTimelineTrack
+class FAnimTimelineTrack_Curve : public FAnimTimelineTrack, public FSelfRegisteringEditorUndoClient
 {
 	ANIMTIMELINE_DECLARE_TRACK(FAnimTimelineTrack_Curve, FAnimTimelineTrack);
 
@@ -26,6 +27,10 @@ public:
 	virtual bool SupportsSelection() const override { return true; }
 	virtual void AddToContextMenu(FMenuBuilder& InMenuBuilder, TSet<FName>& InOutExistingMenuTypes) const override;
 	
+	/** FEditorUndoClient interface */
+	virtual void PostUndo(bool bSuccess) override { PostUndoRedo(); }
+	virtual void PostRedo(bool bSuccess) override { PostUndoRedo(); }
+
 	/** Access the curve we are editing */
 	const TArray<FRichCurve*>& GetCurves() { return Curves; }
 
@@ -66,6 +71,9 @@ protected:
 	/** Edit this curve in the curve editor */
 	void HendleEditCurve();
 
+	/** Handle undo/redo */
+	void PostUndoRedo();
+
 protected:
 	/** The curve we are editing */
 	TArray<FRichCurve*> Curves;
diff --git a/Engine/Source/Editor/Persona/Private/AnimTimeline/AnimTimelineTrack_FloatCurve.cpp b/Engine/Source/Editor/Persona/Private/AnimTimeline/AnimTimelineTrack_FloatCurve.cpp
--- a/Engine/Source/Editor/Persona/Private/AnimTimeline/AnimTimelineTrack_FloatCurve.cpp
+++ b/Engine/Source/Editor/Persona/Private/AnimTimeline/AnimTimelineTrack_FloatCurve.cpp
@@ -35,6 +35,11 @@ FAnimTimelineTrack_FloatCurve::FAnimTimelineTrack_FloatCurve(FFloatCurve& InCurv
 
 TSharedRef<SWidget> FAnimTimelineTrack_FloatCurve::MakeTimelineWidgetContainer()
 {
+	TSharedRef<SWidget> CurveWidget = MakeCurveWidget();
+
+	// zoom to fit now we have a view
+	CurveEditor->ZoomToFit(EAxisList::Y);
+
 	auto ColorLambda = [this]()
 	{
 		if(GetModel()->IsTrackSelected(AsShared()))
@@ -53,7 +58,7 @@ TSharedRef<SWidget> FAnimTimelineTrack_FloatCurve::MakeTimelineWidgetContainer()
 		.BorderImage_Lambda([this](){ return FloatCurve.GetCurveTypeFlag(AACF_Metadata) ? FEditorStyle::GetBrush("Sequencer.Section.SelectedSectionOverlay") : FEditorStyle::GetBrush("AnimTimeline.Outliner.DefaultBorder"); })
 		.BorderBackgroundColor_Lambda(ColorLambda)
 		[
-			MakeCurveWidget()
+			CurveWidget
 		];
 }
 
@@ -63,7 +68,6 @@ TSharedRef<SWidget> FAnimTimelineTrack_FloatCurve::GenerateContainerWidgetForOut
 	TSharedPtr<SHorizontalBox> InnerHorizontalBox;
 	TSharedRef<SWidget> OutlinerWidget = GenerateStandardOutlinerWidget(InRow, false, OuterBorder, InnerHorizontalBox);
 
-	TWeakPtr<FAnimTimelineTrack_FloatCurve> WeakThisTrack = SharedThis(this);
 
 	UAnimMontage* AnimMontage = Cast<UAnimMontage>(GetModel()->GetAnimSequenceBase());
 	bool bChildAnimMontage = AnimMontage && AnimMontage->HasParentAsset();
@@ -77,7 +81,7 @@ TSharedRef<SWidget> FAnimTimelineTrack_FloatCurve::GenerateContainerWidgetForOut
 			SAssignNew(EditableTextLabel, SInlineEditableTextBlock)
 			.IsReadOnly(bChildAnimMontage)
 			.Text(this, &FAnimTimelineTrack_FloatCurve::GetLabel)
-			.IsSelected_Lambda([WeakThisTrack](){ return WeakThisTrack.Pin()->GetModel()->IsTrackSelected(WeakThisTrack.Pin().ToSharedRef()); })
+			.IsSelected_Lambda([this](){ return GetModel()->IsTrackSelected(SharedThis(this)); })
 			.OnTextCommitted(this, &FAnimTimelineTrack_FloatCurve::OnCommitCurveName)
 			.HighlightText(InRow->GetHighlightText())
 		];
@@ -214,6 +218,8 @@ void FAnimTimelineTrack_FloatCurve::OnCommitCurveName(const FText& InText, EText
 
 			FScopedTransaction Transaction(LOCTEXT("CurveEditor_RenameCurve", "Rename Curve"));
 
+			AnimSequenceBase->Modify();
+
 			FSmartName NewSmartName;
 			if (NameMapping->FindSmartName(RequestedName, NewSmartName))
 			{
@@ -258,6 +264,9 @@ void FAnimTimelineTrack_FloatCurve::OnCommitCurveName(const FText& InText, EText
 
 			FloatCurve.Name.UID = NewSmartName.UID;
 			FloatCurve.Name.DisplayName = NewSmartName.DisplayName;
+
+			CurveName = FloatCurve.Name;
+			FullCurveName = FText::FromName(FloatCurve.Name.DisplayName);
 		}
 	}
 }
diff --git a/Engine/Source/Editor/Persona/Private/AnimTimeline/AnimTimelineTrack_NotifiesPanel.cpp b/Engine/Source/Editor/Persona/Private/AnimTimeline/AnimTimelineTrack_NotifiesPanel.cpp
--- a/Engine/Source/Editor/Persona/Private/AnimTimeline/AnimTimelineTrack_NotifiesPanel.cpp
+++ b/Engine/Source/Editor/Persona/Private/AnimTimeline/AnimTimelineTrack_NotifiesPanel.cpp
@@ -30,27 +30,11 @@ FAnimTimelineTrack_NotifiesPanel::FAnimTimelineTrack_NotifiesPanel(const TShared
 
 TSharedRef<SWidget> FAnimTimelineTrack_NotifiesPanel::GenerateContainerWidgetForTimeline()
 {
-	UAnimMontage* AnimMontage = Cast<UAnimMontage>(GetModel()->GetAnimSequenceBase());
-	bool bChildAnimMontage = AnimMontage && AnimMontage->HasParentAsset();
-
-	return SAssignNew(AnimNotifyPanel, SAnimNotifyPanel, GetModel())
-		.IsEnabled(!bChildAnimMontage)
-		.Sequence(GetModel()->GetAnimSequenceBase())
-		.InputMin(this, &FAnimTimelineTrack_NotifiesPanel::GetMinInput)
-		.InputMax(this, &FAnimTimelineTrack_NotifiesPanel::GetMaxInput)
-		.ViewInputMin(this, &FAnimTimelineTrack_NotifiesPanel::GetViewMinInput)
-		.ViewInputMax(this, &FAnimTimelineTrack_NotifiesPanel::GetViewMaxInput)
-		.OnGetScrubValue(this, &FAnimTimelineTrack_NotifiesPanel::GetScrubValue)
-		.OnSelectionChanged(this, &FAnimTimelineTrack_NotifiesPanel::SelectObjects)
-		.OnSetInputViewRange(this, &FAnimTimelineTrack_NotifiesPanel::OnSetInputViewRange)
-		.OnInvokeTab(GetModel()->OnInvokeTab)
-		.OnSnapPosition(&GetModel().Get(), &FAnimModel::Snap)
-		.OnGetTimingNodeVisibility(this, &FAnimTimelineTrack_NotifiesPanel::OnGetTimingNodeVisibility)
-		.OnNotifiesChanged_Lambda([this]()
-		{ 
-			Update();
-			GetModel()->OnTracksChanged().Broadcast(); 
-		});
+	GetAnimNotifyPanel();
+
+	AnimNotifyPanel->Update();
+
+	return AnimNotifyPanel.ToSharedRef();
 }
 
 TSharedRef<SWidget> FAnimTimelineTrack_NotifiesPanel::GenerateContainerWidgetForOutliner(const TSharedRef<SAnimOutlinerItem>& InRow)
@@ -258,7 +242,10 @@ void FAnimTimelineTrack_NotifiesPanel::Update()
 {
 	SetHeight((float)GetModel()->GetAnimSequenceBase()->AnimNotifyTracks.Num() * NotificationTrackHeight);
 	RefreshOutlinerWidget();
-	AnimNotifyPanel->Update();
+	if(AnimNotifyPanel.IsValid())
+	{
+		AnimNotifyPanel->Update();
+	}
 }
 
 void FAnimTimelineTrack_NotifiesPanel::OnCommitTrackName(const FText& InText, ETextCommit::Type CommitInfo, int32 TrackIndexToName)
@@ -279,4 +266,34 @@ EVisibility FAnimTimelineTrack_NotifiesPanel::OnGetTimingNodeVisibility(ETimingE
 	return StaticCastSharedRef<FAnimModel_AnimSequenceBase>(GetModel())->IsNotifiesTimingElementDisplayEnabled(ElementType) ? EVisibility::Visible : EVisibility::Hidden;
 }
 
+TSharedRef<SAnimNotifyPanel> FAnimTimelineTrack_NotifiesPanel::GetAnimNotifyPanel()
+{
+	if(!AnimNotifyPanel.IsValid())
+	{
+		UAnimMontage* AnimMontage = Cast<UAnimMontage>(GetModel()->GetAnimSequenceBase());
+		bool bChildAnimMontage = AnimMontage && AnimMontage->HasParentAsset();
+
+		AnimNotifyPanel = SNew(SAnimNotifyPanel, GetModel())
+			.IsEnabled(!bChildAnimMontage)
+			.Sequence(GetModel()->GetAnimSequenceBase())
+			.InputMin(this, &FAnimTimelineTrack_NotifiesPanel::GetMinInput)
+			.InputMax(this, &FAnimTimelineTrack_NotifiesPanel::GetMaxInput)
+			.ViewInputMin(this, &FAnimTimelineTrack_NotifiesPanel::GetViewMinInput)
+			.ViewInputMax(this, &FAnimTimelineTrack_NotifiesPanel::GetViewMaxInput)
+			.OnGetScrubValue(this, &FAnimTimelineTrack_NotifiesPanel::GetScrubValue)
+			.OnSelectionChanged(this, &FAnimTimelineTrack_NotifiesPanel::SelectObjects)
+			.OnSetInputViewRange(this, &FAnimTimelineTrack_NotifiesPanel::OnSetInputViewRange)
+			.OnInvokeTab(GetModel()->OnInvokeTab)
+			.OnSnapPosition(&GetModel().Get(), &FAnimModel::Snap)
+			.OnGetTimingNodeVisibility(this, &FAnimTimelineTrack_NotifiesPanel::OnGetTimingNodeVisibility)
+			.OnNotifiesChanged_Lambda([this]()
+			{ 
+				Update();
+				GetModel()->OnTracksChanged().Broadcast(); 
+			});	
+	}
+
+	return AnimNotifyPanel.ToSharedRef();
+}
+
 #undef LOCTEXT_NAMESPACE
\ No newline at end of file
diff --git a/Engine/Source/Editor/Persona/Private/AnimTimeline/AnimTimelineTrack_NotifiesPanel.h b/Engine/Source/Editor/Persona/Private/AnimTimeline/AnimTimelineTrack_NotifiesPanel.h
--- a/Engine/Source/Editor/Persona/Private/AnimTimeline/AnimTimelineTrack_NotifiesPanel.h
+++ b/Engine/Source/Editor/Persona/Private/AnimTimeline/AnimTimelineTrack_NotifiesPanel.h
@@ -24,7 +24,7 @@ public:
 	virtual TSharedRef<SWidget> GenerateContainerWidgetForOutliner(const TSharedRef<SAnimOutlinerItem>& InRow) override;
 	virtual bool SupportsFiltering() const override { return false; }
 
-	TSharedPtr<SAnimNotifyPanel> GetAnimNotifyPanel() { return AnimNotifyPanel; }
+	TSharedRef<SAnimNotifyPanel> GetAnimNotifyPanel();
 	
 	void Update();
 
diff --git a/Engine/Source/Editor/Persona/Private/AnimationEditorViewportClient.cpp b/Engine/Source/Editor/Persona/Private/AnimationEditorViewportClient.cpp
--- a/Engine/Source/Editor/Persona/Private/AnimationEditorViewportClient.cpp
+++ b/Engine/Source/Editor/Persona/Private/AnimationEditorViewportClient.cpp
@@ -879,11 +879,24 @@ FText FAnimationViewportClient::GetDisplayInfo(bool bDisplayAllInfo) const
 				FTransform LocalTransform = LocalBoneTransforms[BoneIndex];
 				FTransform ComponentTransform = PreviewMeshComponent->GetDrawTransform(BoneIndex);
 
-				TextValue = ConcatenateLine(TextValue, FText::Format(LOCTEXT("LocalTransform", "Local: {0}"), FText::FromString(LocalTransform.ToHumanReadableString())));
+				auto GetDisplayTransform = [](const FTransform& InTransform) -> FText
+				{
+					FRotator R(InTransform.GetRotation());
+					FVector T(InTransform.GetTranslation());
+					FVector S(InTransform.GetScale3D());
+
+					FString Output = FString::Printf(TEXT("Rotation: X(Roll) %f Y(Pitch)  %f Z(Yaw) %f\r\n"), R.Roll, R.Pitch, R.Yaw);
+					Output += FString::Printf(TEXT("Translation: %f %f %f\r\n"), T.X, T.Y, T.Z);
+					Output += FString::Printf(TEXT("Scale3D: %f %f %f\r\n"), S.X, S.Y, S.Z);
+
+					return FText::FromString(Output);
+				};
+
+				TextValue = ConcatenateLine(TextValue, FText::Format(LOCTEXT("LocalTransform", "Local: {0}"), GetDisplayTransform(LocalTransform)));
 
-				TextValue = ConcatenateLine(TextValue, FText::Format(LOCTEXT("ComponentTransform", "Component: {0}"), FText::FromString(ComponentTransform.ToHumanReadableString())));
+				TextValue = ConcatenateLine(TextValue, FText::Format(LOCTEXT("ComponentTransform", "Component: {0}"), GetDisplayTransform(ComponentTransform)));
 
-				TextValue = ConcatenateLine(TextValue, FText::Format(LOCTEXT("ReferenceTransform", "Reference: {0}"), FText::FromString(ReferenceTransform.ToHumanReadableString())));
+				TextValue = ConcatenateLine(TextValue, FText::Format(LOCTEXT("ReferenceTransform", "Reference: {0}"), GetDisplayTransform(ReferenceTransform)));
 			}
 
 			TextValue = ConcatenateLine(TextValue, FText::Format(LOCTEXT("ApproximateSize", "Approximate Size: {0}x{1}x{2}"),
diff --git a/Engine/Source/Editor/Persona/Private/SAnimNotifyPanel.cpp b/Engine/Source/Editor/Persona/Private/SAnimNotifyPanel.cpp
--- a/Engine/Source/Editor/Persona/Private/SAnimNotifyPanel.cpp
+++ b/Engine/Source/Editor/Persona/Private/SAnimNotifyPanel.cpp
@@ -2562,6 +2562,7 @@ void SAnimNotifyTrack::CreateNewSyncMarkerAtCursor(FString NewSyncMarkerName, UC
 	SyncMarker.MarkerName = FName(*NewSyncMarkerName);
 	SyncMarker.TrackIndex = TrackIndex;
 	SyncMarker.Time = LastClickedTime;
+	SyncMarker.Guid = FGuid::NewGuid();
 
 	Seq->PostEditChange();
 	Seq->MarkPackageDirty();
diff --git a/Engine/Source/Editor/Persona/Private/SAnimSequenceCurveEditor.cpp b/Engine/Source/Editor/Persona/Private/SAnimSequenceCurveEditor.cpp
--- a/Engine/Source/Editor/Persona/Private/SAnimSequenceCurveEditor.cpp
+++ b/Engine/Source/Editor/Persona/Private/SAnimSequenceCurveEditor.cpp
@@ -58,17 +58,23 @@ public:
 		{
 			FTransformCurve& TransformCurve = *(static_cast<FTransformCurve*>(CurveBase));
 			check(CurveIndex < 9);
-			const int32 SubCurveIndex = CurveIndex / 3;
-			switch(SubCurveIndex)
+			const int32 SubCurveIndex = CurveIndex % 3;
+			switch(CurveIndex)
 			{
 			default:
 				check(false);
 				// fall through
 			case 0:
-				return TransformCurve.TranslationCurve.FloatCurves[SubCurveIndex];
 			case 1:
-				return TransformCurve.RotationCurve.FloatCurves[SubCurveIndex];
 			case 2:
+				return TransformCurve.TranslationCurve.FloatCurves[SubCurveIndex];
+			case 3:
+			case 4:
+			case 5:
+				return TransformCurve.RotationCurve.FloatCurves[SubCurveIndex];
+			case 6:
+			case 7:
+			case 8:
 				return TransformCurve.ScaleCurve.FloatCurves[SubCurveIndex];
 			}
 			
diff --git a/Engine/Source/Editor/SequenceRecorder/Private/AnimationRecorder.cpp b/Engine/Source/Editor/SequenceRecorder/Private/AnimationRecorder.cpp
--- a/Engine/Source/Editor/SequenceRecorder/Private/AnimationRecorder.cpp
+++ b/Engine/Source/Editor/SequenceRecorder/Private/AnimationRecorder.cpp
@@ -15,7 +15,6 @@
 #include "AssetRegistryModule.h"
 #include "Framework/Notifications/NotificationManager.h"
 #include "Widgets/Notifications/SNotificationList.h"
-#include "Animation/AnimationSettings.h"
 #include "Animation/AnimationRecordingSettings.h"
 #include "Animation/AnimNotifies/AnimNotify.h"
 #include "Animation/AnimNotifies/AnimNotifyState.h"
@@ -295,16 +294,6 @@ UAnimSequence* FAnimationRecorder::StopRecord(bool bShowMessage)
 
 		FixupNotifies();
 
-		// force anim settings for speed, we dont want any fancy recompression at present
-		UAnimationSettings* AnimationSettings = GetMutableDefault<UAnimationSettings>();
-		TSubclassOf<UAnimCompress> OldDefaultCompressionAlgorithm = AnimationSettings->DefaultCompressionAlgorithm;
-		TEnumAsByte<AnimationCompressionFormat> OldRotationCompressionFormat = AnimationSettings->RotationCompressionFormat;
-		TEnumAsByte<AnimationCompressionFormat> OldTranslationCompressionFormat = AnimationSettings->TranslationCompressionFormat;
-
-		AnimationSettings->DefaultCompressionAlgorithm = UAnimCompress_BitwiseCompressOnly::StaticClass();
-		AnimationSettings->RotationCompressionFormat = ACF_None;
-		AnimationSettings->TranslationCompressionFormat = ACF_None;
-
 		// post-process applies compression etc.
 		// @todo figure out why removing redundant keys is inconsistent
 
@@ -375,11 +364,6 @@ UAnimSequence* FAnimationRecorder::StopRecord(bool bShowMessage)
 		//AnimationObject->RawCurveData.RemoveRedundantKeys();
 		AnimationObject->PostProcessSequence();
 
-		// restore old settings
-		AnimationSettings->DefaultCompressionAlgorithm = OldDefaultCompressionAlgorithm;
-		AnimationSettings->RotationCompressionFormat = OldRotationCompressionFormat;
-		AnimationSettings->TranslationCompressionFormat = OldTranslationCompressionFormat;
-
 		AnimationObject->MarkPackageDirty();
 		
 		// save the package to disk, for convenience and so we can run this in standalone mode
diff --git a/Engine/Source/Editor/SkeletonEditor/Private/SSkeletonTree.cpp b/Engine/Source/Editor/SkeletonEditor/Private/SSkeletonTree.cpp
--- a/Engine/Source/Editor/SkeletonEditor/Private/SSkeletonTree.cpp
+++ b/Engine/Source/Editor/SkeletonEditor/Private/SSkeletonTree.cpp
@@ -5,6 +5,8 @@
 #include "Misc/MessageDialog.h"
 #include "Modules/ModuleManager.h"
 #include "UObject/PropertyPortFlags.h"
+#include "UObject/UObjectGlobals.h"
+#include "UObject/PackageReload.h"
 #include "Framework/Application/SlateApplication.h"
 #include "Textures/SlateIcon.h"
 #include "Framework/Commands/UIAction.h"
@@ -184,6 +186,8 @@ void SSkeletonTree::Construct(const FArguments& InArgs, const TSharedRef<FEditab
 		RegisterOnSelectionChanged(InSkeletonTreeArgs.OnSelectionChanged);
 	}
 
+	FCoreUObjectDelegates::OnPackageReloaded.AddSP(this, &SSkeletonTree::HandlePackageReloaded);
+
 	BoneProxy = NewObject<UBoneProxy>(GetTransientPackage());
 	BoneProxy->SkelMeshComponent = PreviewScene.IsValid() ? PreviewScene.Pin()->GetPreviewMeshComponent() : nullptr;
 	BoneProxy->bIsTickable = true;
@@ -314,6 +318,7 @@ SSkeletonTree::~SSkeletonTree()
 	{
 		EditableSkeleton.Pin()->UnregisterOnSkeletonHierarchyChanged(this);
 	}
+	FCoreUObjectDelegates::OnPackageReloaded.RemoveAll(this);
 }
 PRAGMA_ENABLE_DEPRECATION_WARNINGS
 
@@ -1539,6 +1544,23 @@ void SSkeletonTree::OnFilterTextChanged( const FText& SearchText )
 	ApplyFilter();
 }
 
+void SSkeletonTree::HandlePackageReloaded(const EPackageReloadPhase InPackageReloadPhase, FPackageReloadedEvent* InPackageReloadedEvent)
+{
+	if (InPackageReloadPhase == EPackageReloadPhase::PostPackageFixup)
+	{
+		for (const auto& RepointedObjectPair : InPackageReloadedEvent->GetRepointedObjects())
+		{
+			if (USkeleton* NewObject = Cast<USkeleton>(RepointedObjectPair.Value))
+			{
+				if (&GetEditableSkeletonInternal()->GetSkeleton() == NewObject)
+				{
+					Refresh();
+				}
+			}
+		}
+	}
+}
+
 void SSkeletonTree::Refresh()
 {
 	CreateFromSkeleton();
diff --git a/Engine/Source/Editor/SkeletonEditor/Private/SSkeletonTree.h b/Engine/Source/Editor/SkeletonEditor/Private/SSkeletonTree.h
--- a/Engine/Source/Editor/SkeletonEditor/Private/SSkeletonTree.h
+++ b/Engine/Source/Editor/SkeletonEditor/Private/SSkeletonTree.h
@@ -39,6 +39,8 @@ class SComboButton;
 class UBlendProfile;
 struct FNotificationInfo;
 class IPinnedCommandList;
+class FPackageReloadedEvent;
+enum class EPackageReloadPhase : uint8;
 
 //////////////////////////////////////////////////////////////////////////
 // SSkeletonTree
@@ -315,6 +317,9 @@ private:
 	/** Handle external deselection event */
 	void HandleDeselectAll();
 
+	/** Handle package reloading (might be our skeleton) */
+	void HandlePackageReloaded(const EPackageReloadPhase InPackageReloadPhase, FPackageReloadedEvent* InPackageReloadedEvent);
+
 private:
 	/** Pointer back to the skeleton tree that owns us */
 	TWeakPtr<FEditableSkeleton> EditableSkeleton;
diff --git a/Engine/Source/Editor/UnrealEd/Private/EditorModeManager.cpp b/Engine/Source/Editor/UnrealEd/Private/EditorModeManager.cpp
--- a/Engine/Source/Editor/UnrealEd/Private/EditorModeManager.cpp
+++ b/Engine/Source/Editor/UnrealEd/Private/EditorModeManager.cpp
@@ -950,8 +950,29 @@ bool FEditorModeTools::ShouldShowModeToolbar() const
 
 bool FEditorModeTools::ShouldShowModeToolbox() const
 {
-	// This could ideally ask each active mode if it has any tools but when developing a new mode the toolbox not appearing by default could be hard to understand
-	return !IsModeActive(FBuiltinEditorModes::EM_Default);
+	bool bHasVisibleModes = false;
+	for (const TSharedPtr<FEdMode>& Mode : ActiveModes)
+	{
+		if (Mode->GetModeInfo().bVisible)
+		{
+			bHasVisibleModes = true;
+			break;
+		}
+	}
+
+	if (!bHasVisibleModes)
+	{
+		for (const UEdMode* Mode : ActiveScriptableModes)
+		{
+			if (Mode->GetModeInfo().bVisible)
+			{
+				bHasVisibleModes = true;
+				break;
+			}
+		}
+	}
+
+	return bHasVisibleModes;
 }
 
 void FEditorModeTools::ActivateMode(FEditorModeID InID, bool bToggle)
@@ -1015,7 +1036,7 @@ void FEditorModeTools::ActivateMode(FEditorModeID InID, bool bToggle)
 			}
 
 			// Remove anything that isn't compatible with this mode
-			for (int32 ModeIndex = ActiveModes.Num() - 1; ModeIndex >= 0; --ModeIndex)
+			for (int32 ModeIndex = ActiveModes.Num() - 1; ModeIndex >= 0; ModeIndex--)
 			{
 				const bool bModesAreCompatible = Mode->IsCompatibleWith(ActiveModes[ModeIndex]->GetID()) || ActiveModes[ModeIndex]->IsCompatibleWith(Mode->GetID());
 				if (!bModesAreCompatible)
@@ -1024,7 +1045,7 @@ void FEditorModeTools::ActivateMode(FEditorModeID InID, bool bToggle)
 				}
 			}
 			// Remove anything that isn't compatible with this mode
-			for (int32 ModeIndex = ActiveScriptableModes.Num() - 1; ModeIndex >= 0; --ModeIndex)
+			for (int32 ModeIndex = ActiveScriptableModes.Num() - 1; ModeIndex >= 0; ModeIndex--)
 			{
 				const bool bModesAreCompatible = Mode->IsCompatibleWith(ActiveScriptableModes[ModeIndex]->GetID()) || ActiveScriptableModes[ModeIndex]->IsCompatibleWith(Mode->GetID());
 				if (!bModesAreCompatible)
diff --git a/Engine/Source/Editor/UnrealEd/Private/EditorUndoClient.cpp b/Engine/Source/Editor/UnrealEd/Private/EditorUndoClient.cpp
--- a/Engine/Source/Editor/UnrealEd/Private/EditorUndoClient.cpp
+++ b/Engine/Source/Editor/UnrealEd/Private/EditorUndoClient.cpp
@@ -10,3 +10,11 @@ FEditorUndoClient::~FEditorUndoClient()
 		GEditor->UnregisterForUndo(this);
 	}
 }
+
+FSelfRegisteringEditorUndoClient::FSelfRegisteringEditorUndoClient()
+{
+	if (GEditor)
+	{
+		GEditor->RegisterForUndo(this);
+	}
+}
\ No newline at end of file
diff --git a/Engine/Source/Editor/UnrealEd/Public/EditorUndoClient.h b/Engine/Source/Editor/UnrealEd/Public/EditorUndoClient.h
--- a/Engine/Source/Editor/UnrealEd/Public/EditorUndoClient.h
+++ b/Engine/Source/Editor/UnrealEd/Public/EditorUndoClient.h
@@ -44,3 +44,13 @@ public:
 	/** Return the transaction context for this client */
 	virtual FString GetTransactionContext() const { return FString(); }
 };
+
+/** An undo client that registers itself in its constructor and unregisters itself in its destructor */
+class UNREALED_API FSelfRegisteringEditorUndoClient : public FEditorUndoClient
+{
+public:
+	/** Register in constructor */
+	FSelfRegisteringEditorUndoClient();
+
+	/** FEditorUndoClient already unregisters in its destructor */
+};
\ No newline at end of file
diff --git a/Engine/Source/Runtime/AnimGraphRuntime/Private/RBF/RBFInterpolator.cpp b/Engine/Source/Runtime/AnimGraphRuntime/Private/RBF/RBFInterpolator.cpp
--- a/Engine/Source/Runtime/AnimGraphRuntime/Private/RBF/RBFInterpolator.cpp
+++ b/Engine/Source/Runtime/AnimGraphRuntime/Private/RBF/RBFInterpolator.cpp
@@ -11,11 +11,9 @@
 #pragma warning(push)
 #pragma warning(disable:6294) // Ill-defined for-loop:  initial condition does not satisfy test.  Loop body not executed.
 #endif
-PRAGMA_DEFAULT_VISIBILITY_START
 THIRD_PARTY_INCLUDES_START
 #include <Eigen/LU>
 THIRD_PARTY_INCLUDES_END
-PRAGMA_DEFAULT_VISIBILITY_END
 #if defined(_MSC_VER) && USING_CODE_ANALYSIS
 #pragma warning(pop)
 #endif
diff --git a/Engine/Source/Runtime/AnimGraphRuntime/Public/RBF/RBFInterpolator.h b/Engine/Source/Runtime/AnimGraphRuntime/Public/RBF/RBFInterpolator.h
--- a/Engine/Source/Runtime/AnimGraphRuntime/Public/RBF/RBFInterpolator.h
+++ b/Engine/Source/Runtime/AnimGraphRuntime/Public/RBF/RBFInterpolator.h
@@ -73,7 +73,10 @@ namespace RBFDistanceMetric
 	*/
 	static inline float SwingAngle(const FQuat& A, const FQuat& B, const FVector& TwistAxis)
 	{
-		return ArcLength(A.RotateVector(TwistAxis), B.RotateVector(TwistAxis));
+		FQuat ASwing, BSwing, DummyTwist;
+		A.ToSwingTwist(TwistAxis, ASwing, DummyTwist);
+		B.ToSwingTwist(TwistAxis, BSwing, DummyTwist);
+		return ASwing.AngularDistance(BSwing);
 	}
 
 	/* Returns the twist arc length distance between two quaternions, using a specific 
@@ -150,6 +153,8 @@ class TRBFInterpolator
 public:
 	using WeightFuncT = TFunction<float(const T& A, const T& B)>;
 
+	TRBFInterpolator() = default;
+
 	/* Construct an RBF interpolator, taking in a set of sparse nodes and a symmetric weighing
 	   function that computes the distance between two nodes, and, optionally, smooths 
 	   the distance with a smoothing kernel.
@@ -163,6 +168,11 @@ public:
 		MakeUpperKernel();
 	}
 
+	TRBFInterpolator(const TRBFInterpolator<T>&) = default;
+	TRBFInterpolator(TRBFInterpolator<T>&&) = default;
+	TRBFInterpolator<T>& operator=(const TRBFInterpolator<T>&) = default;
+	TRBFInterpolator<T>& operator=(TRBFInterpolator<T>&&) = default;
+
 	/* Given a value, compute the weight values to use to calculate each node's contribution
 	   to that value's location.
 	*/
@@ -329,6 +339,6 @@ private:
 	}
 
 
-	const TArrayView<T>& Nodes;
+	TArrayView<T> Nodes;
 	WeightFuncT WeightFunc;
 };
diff --git a/Engine/Source/Runtime/CoreUObject/Public/UObject/Field.h b/Engine/Source/Runtime/CoreUObject/Public/UObject/Field.h
--- a/Engine/Source/Runtime/CoreUObject/Public/UObject/Field.h
+++ b/Engine/Source/Runtime/CoreUObject/Public/UObject/Field.h
@@ -164,6 +164,10 @@ public: \
 		new (Mem) TClass(EC_InternalUseOnlyConstructor, TClass::StaticClass()); \
 		return Mem; \
 	} \
+	inline void operator delete(void* InMem) noexcept \
+	{ \
+		FMemory::Free(InMem); \
+	} \
 	friend FArchive &operator<<( FArchive& Ar, ThisClass*& Res ) \
 	{ \
 		return Ar << (FField*&)Res; \
diff --git a/Engine/Source/Runtime/Engine/Classes/Animation/AnimCompress.h b/Engine/Source/Runtime/Engine/Classes/Animation/AnimCompress.h
--- a/Engine/Source/Runtime/Engine/Classes/Animation/AnimCompress.h
+++ b/Engine/Source/Runtime/Engine/Classes/Animation/AnimCompress.h
@@ -283,10 +283,6 @@ class UAnimCompress : public UAnimBoneCompressionCodec
 	UPROPERTY(Category = Compression, EditAnywhere)
 	TEnumAsByte<AnimationCompressionFormat> ScaleCompressionFormat;
 
-	/** Max error for compression of curves using remove redundant keys */
-	UPROPERTY(Category = Compression, EditAnywhere)
-	float MaxCurveError;
-
 public:
 #if WITH_EDITORONLY_DATA
 	/** UAnimBoneCompressionCodec implementation */
diff --git a/Engine/Source/Runtime/Engine/Classes/Animation/AnimInstance.h b/Engine/Source/Runtime/Engine/Classes/Animation/AnimInstance.h
--- a/Engine/Source/Runtime/Engine/Classes/Animation/AnimInstance.h
+++ b/Engine/Source/Runtime/Engine/Classes/Animation/AnimInstance.h
@@ -426,6 +426,14 @@ class ENGINE_API UAnimInstance : public UObject
 	/** Flag to check if created by LinkedAnimGraph in ReinitializeLinkedAnimInstance */
 	uint8 bCreatedByLinkedAnimGraph : 1;
 
+	/** Whether to process notifies from any linked anim instances */
+	UPROPERTY(EditDefaultsOnly, Category = Notifies)
+	uint8 bReceiveNotifiesFromLinkedInstances : 1;
+
+	/** Whether to propagate notifies to any linked anim instances */
+	UPROPERTY(EditDefaultsOnly, Category = Notifies)
+	uint8 bPropagateNotifiesToLinkedInstances : 1;
+
 private:
 	/** True when Montages are being ticked, and Montage Events should be queued. 
 	 * When Montage are being ticked, we queue AnimNotifies and Events. We trigger notifies first, then Montage events. */
@@ -445,8 +453,6 @@ public:
 	// @todo document
 	void MakeMontageTickRecord(FAnimTickRecord& TickRecord, class UAnimMontage* Montage, float CurrentPosition, float PreviousPosition, float MoveDelta, float Weight, TArray<FPassedMarker>& MarkersPassedThisTick, FMarkerTickRecord& MarkerTickRecord);
 
-	bool IsSlotNodeRelevantForNotifies(FName SlotNodeName) const;
-
 	/** Get global weight in AnimGraph for this slot node.
 	* Note: this is the weight of the node, not the weight of any potential montage it is playing. */
 	float GetSlotNodeGlobalWeight(const FName& SlotNodeName) const;
@@ -498,6 +504,22 @@ public:
 	// Can does this anim instance need an update (parallel or not)?
 	bool NeedsUpdate() const;
 
+	/** Get whether to process notifies from any linked anim instances */
+	UFUNCTION(BlueprintPure, Category = "Notifies")
+	bool GetReceiveNotifiesFromLinkedInstances() const { return bReceiveNotifiesFromLinkedInstances; }
+
+	/** Set whether to process notifies from any linked anim instances */
+	UFUNCTION(BlueprintCallable, Category = "Notifies")
+	void SetReceiveNotifiesFromLinkedInstances(bool bSet) { bReceiveNotifiesFromLinkedInstances = bSet; }
+
+	/** Get whether to propagate notifies to any linked anim instances */
+	UFUNCTION(BlueprintPure, Category = "Notifies")
+	bool GetPropagateNotifiesToLinkedInstances() const { return bPropagateNotifiesToLinkedInstances; }
+
+	/** Set whether to propagate notifies to any linked anim instances */
+	UFUNCTION(BlueprintCallable, Category = "Notifies")
+	void SetPropagateNotifiesToLinkedInstances(bool bSet) { bPropagateNotifiesToLinkedInstances = bSet; }
+
 private:
 	// Does this anim instance need immediate update (rather than parallel)?
 	bool NeedsImmediateUpdate(float DeltaSeconds) const;
@@ -529,6 +551,10 @@ public:
 	UFUNCTION(BlueprintImplementableEvent)
 	void BlueprintBeginPlay();
 
+	/** Executed when the all Linked Animation Layers are initialized */
+	UFUNCTION(BlueprintImplementableEvent)
+	void BlueprintLinkedAnimationLayersInitialized();
+
 	bool CanTransitionSignature() const;
 	
 	/*********************************************************************************************
diff --git a/Engine/Source/Runtime/Engine/Classes/Animation/AnimNodeBase.h b/Engine/Source/Runtime/Engine/Classes/Animation/AnimNodeBase.h
--- a/Engine/Source/Runtime/Engine/Classes/Animation/AnimNodeBase.h
+++ b/Engine/Source/Runtime/Engine/Classes/Animation/AnimNodeBase.h
@@ -692,8 +692,8 @@ enum class EPostCopyOperation : uint8
 UENUM()
 enum class ECopyType : uint8
 {
-	// Just copy the memory
-	MemCopy,
+	// For plain old data types, we do a simple memcpy.
+	PlainProperty,
 
 	// Read and write properties using bool property helpers, as source/dest could be bitfield or boolean
 	BoolProperty,
@@ -703,6 +703,9 @@ enum class ECopyType : uint8
 
 	// Read and write properties using object property helpers, as source/dest could be regular/weak/lazy etc.
 	ObjectProperty,
+
+	// FName needs special case because its size changes between editor/compiler and runtime.
+	NameProperty,
 };
 
 
@@ -721,7 +724,7 @@ struct FExposedValueCopyRecord
 		, SourceArrayIndex(0)
 		, bInstanceIsTarget(false)
 		, PostCopyOperation(EPostCopyOperation::None)
-		, CopyType(ECopyType::MemCopy)
+		, CopyType(ECopyType::PlainProperty)
 		, DestProperty(nullptr)
 		, DestArrayIndex(0)
 		, Size(0)
diff --git a/Engine/Source/Runtime/Engine/Classes/Animation/AnimNode_LinkedAnimGraph.h b/Engine/Source/Runtime/Engine/Classes/Animation/AnimNode_LinkedAnimGraph.h
--- a/Engine/Source/Runtime/Engine/Classes/Animation/AnimNode_LinkedAnimGraph.h
+++ b/Engine/Source/Runtime/Engine/Classes/Animation/AnimNode_LinkedAnimGraph.h
@@ -51,6 +51,19 @@ public:
 	// Our node index
 	int32 NodeIndex;
 
+protected:
+	// Inertial blending duration to request next update
+	float PendingBlendDuration;
+
+public:
+	/** Whether named notifies will be received by this linked instance from other instances (outer or other linked instances) */
+	UPROPERTY(EditAnywhere, Category = Settings)
+	uint8 bReceiveNotifiesFromLinkedInstances : 1;
+
+	/** Whether named notifies will be propagated from this linked instance to other instances (outer or other linked instances) */
+	UPROPERTY(EditAnywhere, Category = Settings)
+	uint8 bPropagateNotifiesToLinkedInstances : 1;
+
 	/** Dynamically set the anim class of this linked instance */
 	void SetAnimClass(TSubclassOf<UAnimInstance> InClass, const UAnimInstance* InOwningAnimInstance);
 
@@ -79,9 +92,6 @@ protected:
 	virtual bool NeedsOnInitializeAnimInstance() const override { return true; }
 	// End of FAnimNode_Base interface
 
-	// Inertial blending duration to request next update
-	float PendingBlendDuration;
-
 	// Re-create the linked instances for this node
 	void ReinitializeLinkedAnimInstance(const UAnimInstance* InOwningAnimInstance, UAnimInstance* InNewAnimInstance = nullptr);
 
diff --git a/Engine/Source/Runtime/Engine/Classes/Animation/AnimationSettings.h b/Engine/Source/Runtime/Engine/Classes/Animation/AnimationSettings.h
--- a/Engine/Source/Runtime/Engine/Classes/Animation/AnimationSettings.h
+++ b/Engine/Source/Runtime/Engine/Classes/Animation/AnimationSettings.h
@@ -28,19 +28,6 @@ class ENGINE_API UAnimationSettings : public UDeveloperSettings
 	UPROPERTY(config, EditAnywhere, Category = Compression)
 	TArray<FString> KeyEndEffectorsMatchNameArray;
 
-	UPROPERTY(config, EditAnywhere, Category = Compression, NoClear)
-	TSubclassOf<class UAnimCompress>  DefaultCompressionAlgorithm; 
-
-	UPROPERTY(config, EditAnywhere, Category = Compression)
-	TEnumAsByte<AnimationCompressionFormat> RotationCompressionFormat;
-
-	UPROPERTY(config, EditAnywhere, Category = Compression)
-	TEnumAsByte<AnimationCompressionFormat> TranslationCompressionFormat;
-
-	/** Max error for compression of curves using remove redundant keys */
-	UPROPERTY(config, EditAnywhere, Category = Compression)
-	float MaxCurveError;
-
 	UPROPERTY(config, EditAnywhere, Category = Compression)
 	bool ForceRecompression;
 
@@ -61,14 +48,6 @@ class ENGINE_API UAnimationSettings : public UDeveloperSettings
 	UPROPERTY(config, EditAnywhere, Category = Compression)
 	bool bRaiseMaxErrorToExisting;
 
-	/** If true, then an exhaustive search is used otherwise only a short list of the best methods is tried */
-	UPROPERTY(config, EditAnywhere, Category = Compression)
-	bool bTryExhaustiveSearch;
-
-	/** If true, anim sequences are compressed in segments. This allows for parallel compression of individual sequences. */
-	UPROPERTY(config, EditAnywhere, Category = Compression)
-	bool bEnableSegmenting;
-
 	UPROPERTY(config, EditAnywhere, Category = Performance)
 	bool bEnablePerformanceLog;
 
diff --git a/Engine/Source/Runtime/Engine/Classes/Animation/Skeleton.h b/Engine/Source/Runtime/Engine/Classes/Animation/Skeleton.h
--- a/Engine/Source/Runtime/Engine/Classes/Animation/Skeleton.h
+++ b/Engine/Source/Runtime/Engine/Classes/Animation/Skeleton.h
@@ -24,6 +24,8 @@ class UBlendProfile;
 class URig;
 class USkeletalMesh;
 class USkeletalMeshSocket;
+class FPackageReloadedEvent;
+enum class EPackageReloadPhase : uint8;
 
 /** This is a mapping table between bone in a particular skeletal mesh and bone of this skeleton set. */
 USTRUCT()
@@ -314,6 +316,7 @@ public:
 	ENGINE_API virtual void PreEditUndo() override;
 	ENGINE_API virtual void PostEditUndo() override;
 #endif
+	ENGINE_API virtual void BeginDestroy() override;
 
 	/** Accessor to Reference Skeleton to make data read only */
 	const FReferenceSkeleton& GetReferenceSkeleton() const
@@ -880,6 +883,9 @@ private:
 	void RegenerateGuid();
 	void RegenerateVirtualBoneGuid();
 
+	// Handle skeletons being reloaded via the content browser
+	static void HandlePackageReloaded(const EPackageReloadPhase InPackageReloadPhase, FPackageReloadedEvent* InPackageReloadedEvent);
+
 
 
 
diff --git a/Engine/Source/Runtime/Engine/Classes/Components/SkeletalMeshComponent.h b/Engine/Source/Runtime/Engine/Classes/Components/SkeletalMeshComponent.h
--- a/Engine/Source/Runtime/Engine/Classes/Components/SkeletalMeshComponent.h
+++ b/Engine/Source/Runtime/Engine/Classes/Components/SkeletalMeshComponent.h
@@ -780,7 +780,7 @@ public:
 	 */
 	UFUNCTION(BlueprintCallable, Category="Components|SkeletalMesh", meta=(Keywords = "AnimBlueprint", UnsafeDuringActorConstruction = "true"))
 	class UAnimInstance * GetAnimInstance() const;
-
+	
 	/**
 	 * Returns the active post process instance is one is available. This is set on the mesh that this
 	 * component is using, and is evaluated immediately after the main instance.
@@ -868,6 +868,9 @@ public:
 	UFUNCTION(BlueprintPure, Category = "Components|SkeletalMesh|Animation Blueprint Linking")
 	UAnimInstance* GetLinkedAnimLayerInstanceByClass(TSubclassOf<UAnimInstance> InClass) const;
 
+	/** Calls a function on each of the anim instances that this mesh component hosts, including linked and post-process instances */
+	void ForEachAnimInstance(TFunctionRef<void(UAnimInstance*)> InFunction);
+
 	/** 
 	 * Returns whether there are any valid instances to run, currently this means whether we have
 	 * have an animation instance or a post process instance available to process.
diff --git a/Engine/Source/Runtime/Engine/Private/Animation/AnimBlueprintGeneratedClass.cpp b/Engine/Source/Runtime/Engine/Private/Animation/AnimBlueprintGeneratedClass.cpp
--- a/Engine/Source/Runtime/Engine/Private/Animation/AnimBlueprintGeneratedClass.cpp
+++ b/Engine/Source/Runtime/Engine/Private/Animation/AnimBlueprintGeneratedClass.cpp
@@ -400,7 +400,7 @@ void UAnimBlueprintGeneratedClass::GenerateAnimationBlueprintFunctions()
 		bool bFoundOutput = false;
 #if WITH_EDITOR
 		// In editor we can grab the group from metadata, otherwise we need to wait until CDO post load (LinkFunctionsToDefaultObjectNodes)
-		FText CategoryText = FObjectEditorUtils::GetCategoryText((*It)->GetAssociatedFField());
+		FText CategoryText = FObjectEditorUtils::GetCategoryText(*It);
 		FName Group = CategoryText.IsEmpty() ? NAME_None : FName(*CategoryText.ToString());
 #endif
 		FStructProperty* OutputPoseNodeProperty = nullptr;
diff --git a/Engine/Source/Runtime/Engine/Private/Animation/AnimCompress.cpp b/Engine/Source/Runtime/Engine/Private/Animation/AnimCompress.cpp
--- a/Engine/Source/Runtime/Engine/Private/Animation/AnimCompress.cpp
+++ b/Engine/Source/Runtime/Engine/Private/Animation/AnimCompress.cpp
@@ -11,7 +11,6 @@
 #include "Misc/FeedbackContext.h"
 #include "AnimationCompression.h"
 #include "AnimEncoding.h"
-#include "Animation/AnimationSettings.h"
 
 DEFINE_LOG_CATEGORY(LogAnimationCompression);
 
@@ -20,9 +19,6 @@ UAnimCompress::UAnimCompress(const FObjectInitializer& ObjectInitializer)
 {
 	TranslationCompressionFormat = ACF_None;
 	RotationCompressionFormat = ACF_Float96NoW;
-
-	UAnimationSettings* AnimationSettings = UAnimationSettings::Get();
-	MaxCurveError = AnimationSettings->MaxCurveError;
 }
 
 void FCompressionMemorySummary::GatherPostCompressionStats(const FCompressedAnimSequence& CompressedData, const TArray<FBoneData>& BoneData, const FName AnimFName, double CompressionTime, bool bInPerformedCompression)
@@ -717,8 +713,6 @@ void UAnimCompress::PopulateDDCKey(FArchive& Ar)
 	SCF = (uint8)ScaleCompressionFormat.GetValue();
 
 	Ar << TCF << RCF << SCF;
-
-	Ar << MaxCurveError;
 }
 
 void UAnimCompress::FilterTrivialPositionKeys(
diff --git a/Engine/Source/Runtime/Engine/Private/Animation/AnimInstance.cpp b/Engine/Source/Runtime/Engine/Private/Animation/AnimInstance.cpp
--- a/Engine/Source/Runtime/Engine/Private/Animation/AnimInstance.cpp
+++ b/Engine/Source/Runtime/Engine/Private/Animation/AnimInstance.cpp
@@ -109,6 +109,9 @@ UAnimInstance::UAnimInstance(const FObjectInitializer& ObjectInitializer)
 	bUseMultiThreadedAnimationUpdate = true;
 	bCreatedByLinkedAnimGraph = false;
 	PendingDynamicResetTeleportType = ETeleportType::None;
+
+	bReceiveNotifiesFromLinkedInstances = false;
+	bPropagateNotifiesToLinkedInstances = false;
 }
 
 // this is only used by montage marker based sync
@@ -223,6 +226,8 @@ void UAnimInstance::InitializeAnimation(bool bInDeferRootNodeInitialization)
 	GetProxyOnGameThread<FAnimInstanceProxy>().BindNativeDelegates();
 
 	InitializeGroupedLayers(bInDeferRootNodeInitialization);
+
+	BlueprintLinkedAnimationLayersInitialized();
 }
 
 void UAnimInstance::UninitializeAnimation()
@@ -606,10 +611,13 @@ void UAnimInstance::ParallelUpdateAnimation()
 {
 	GetProxyOnAnyThread<FAnimInstanceProxy>().UpdateAnimation();
 
-	// Tick asset players for this and any linked instances we have
-	for(UAnimInstance* LinkedInstance : GetSkelMeshComponent()->GetLinkedAnimInstances())
+	if(GetSkelMeshComponent()->GetAnimInstance() == this)
 	{
-		LinkedInstance->GetProxyOnAnyThread<FAnimInstanceProxy>().TickAssetPlayerInstances();
+		// If this is the main instance,  Tick asset players for this and any linked instances we have
+		for(UAnimInstance* LinkedInstance : GetSkelMeshComponent()->GetLinkedAnimInstances())
+		{
+			LinkedInstance->GetProxyOnAnyThread<FAnimInstanceProxy>().TickAssetPlayerInstances();
+		}
 	}
 
 	GetProxyOnAnyThread<FAnimInstanceProxy>().TickAssetPlayerInstances();
@@ -1374,30 +1382,48 @@ void UAnimInstance::TriggerSingleAnimNotify(const FAnimNotifyEvent* AnimNotifyEv
 		{
 			// Custom Event based notifies. These will call a AnimNotify_* function on the AnimInstance.
 			const FName FuncName = AnimNotifyEvent->GetNotifyEventName();
-			UFunction* Function = FindFunction(FuncName);
-			if (Function)
+
+			auto NotifyAnimInstance = [this, AnimNotifyEvent, FuncName](UAnimInstance* InAnimInstance)
 			{
-				// if parameter is none, add event
-				if (Function->NumParms == 0)
-				{
-					ProcessEvent(Function, NULL);
-				}
-				else if ((Function->NumParms == 1) && (CastField<FObjectProperty>(Function->PropertyLink) != NULL))
+				check(InAnimInstance);
+
+				if (InAnimInstance == this || InAnimInstance->bReceiveNotifiesFromLinkedInstances)
 				{
-					struct FAnimNotifierHandler_Parms
+					UFunction* Function = InAnimInstance->FindFunction(FuncName);
+					if (Function)
 					{
-						UAnimNotify* Notify;
-					};
+						// if parameter is none, add event
+						if (Function->NumParms == 0)
+						{
+							InAnimInstance->ProcessEvent(Function, nullptr);
+						}
+						else if ((Function->NumParms == 1) && (CastField<FObjectProperty>(Function->PropertyLink) != nullptr))
+						{
+							struct FAnimNotifierHandler_Parms
+							{
+								UAnimNotify* Notify;
+							};
 
-					FAnimNotifierHandler_Parms Parms;
-					Parms.Notify = AnimNotifyEvent->Notify;
-					ProcessEvent(Function, &Parms);
-				}
-				else
-				{
-					// Actor has event, but with different parameters. Print warning
-					UE_LOG(LogAnimNotify, Warning, TEXT("Anim notifier named %s, but the parameter number does not match or not of the correct type"), *FuncName.ToString());
+							FAnimNotifierHandler_Parms Parms;
+							Parms.Notify = AnimNotifyEvent->Notify;
+							InAnimInstance->ProcessEvent(Function, &Parms);
+						}
+						else
+						{
+							// Actor has event, but with different parameters. Print warning
+							UE_LOG(LogAnimNotify, Warning, TEXT("Anim notifier named %s, but the parameter number does not match or not of the correct type"), *FuncName.ToString());
+						}
+					}
 				}
+			};
+			
+			if(bPropagateNotifiesToLinkedInstances)
+			{
+				GetSkelMeshComponent()->ForEachAnimInstance(NotifyAnimInstance);
+			}
+			else
+			{
+				NotifyAnimInstance(this);
 			}
 		}
 	}
@@ -2633,6 +2659,8 @@ void UAnimInstance::PerformLinkedLayerOverlayOperation(TSubclassOf<UAnimInstance
 						{
 							UAnimInstance* NewLinkedInstance = NewObject<UAnimInstance>(MeshComp, ClassToSet);
 							NewLinkedInstance->bCreatedByLinkedAnimGraph = true;
+							NewLinkedInstance->bPropagateNotifiesToLinkedInstances = LayerNode->bPropagateNotifiesToLinkedInstances;
+							NewLinkedInstance->bReceiveNotifiesFromLinkedInstances = LayerNode->bReceiveNotifiesFromLinkedInstances;
 							NewLinkedInstance->InitializeAnimation();
 
 							// Unlink any layer nodes in the new linked instance, as they may have been hooked up to self in InitializeAnimation above.
@@ -2643,8 +2671,8 @@ void UAnimInstance::PerformLinkedLayerOverlayOperation(TSubclassOf<UAnimInstance
 							if(!bInDeferSubGraphInitialization)
 							{
 								// Initialize the correct parts of the linked instance
-							if(LayerNode->LinkedRoot)
-							{
+								if(LayerNode->LinkedRoot)
+								{
 									FAnimInstanceProxy& ThisProxy = GetProxyOnAnyThread<FAnimInstanceProxy>();
 									FAnimInstanceProxy& LinkedProxy = NewLinkedInstance->GetProxyOnAnyThread<FAnimInstanceProxy>();
 									InitializeAndCacheBonesForLinkedRoot(LayerNode, ThisProxy, NewLinkedInstance, LinkedProxy);
@@ -2697,6 +2725,10 @@ void UAnimInstance::PerformLinkedLayerOverlayOperation(TSubclassOf<UAnimInstance
 						for(FAnimNode_LinkedAnimLayer* LayerNode : LayerPair.Value)
 						{
 							LayerNode->SetLinkedLayerInstance(this, NewLinkedInstance);
+
+							// Propagate notify flags. If any nodes have this set then we need to propagate to the group.
+							NewLinkedInstance->bPropagateNotifiesToLinkedInstances |= LayerNode->bPropagateNotifiesToLinkedInstances;
+							NewLinkedInstance->bReceiveNotifiesFromLinkedInstances |= LayerNode->bReceiveNotifiesFromLinkedInstances;
 						}
 
 						if(!bInDeferSubGraphInitialization)
@@ -2706,9 +2738,9 @@ void UAnimInstance::PerformLinkedLayerOverlayOperation(TSubclassOf<UAnimInstance
 
 							// Initialize the correct parts of the linked instance
 							for(FAnimNode_LinkedAnimLayer* LayerNode : LayerPair.Value)
-						{
-							if(LayerNode->LinkedRoot)
 							{
+								if(LayerNode->LinkedRoot)
+								{
 									InitializeAndCacheBonesForLinkedRoot(LayerNode, ThisProxy, NewLinkedInstance, LinkedProxy);
 								}
 							}
diff --git a/Engine/Source/Runtime/Engine/Private/Animation/AnimNodeBase.cpp b/Engine/Source/Runtime/Engine/Private/Animation/AnimNodeBase.cpp
--- a/Engine/Source/Runtime/Engine/Private/Animation/AnimNodeBase.cpp
+++ b/Engine/Source/Runtime/Engine/Private/Animation/AnimNodeBase.cpp
@@ -689,21 +689,26 @@ void FExposedValueHandler::Initialize(UObject* AnimInstanceObject, int32 NodeOff
 				check(CopyRecord.DestProperty.Get());
 			}
 
-			if (FBoolProperty* BoolProperty = CastField<FBoolProperty>(CopyRecord.DestProperty.Get()))
+			if (CastField<FNameProperty>(CopyRecord.DestProperty.Get()))
+			{
+				CopyRecord.CopyType = ECopyType::NameProperty;
+			}
+			else if (CastField<FBoolProperty>(CopyRecord.DestProperty.Get()))
 			{
 				CopyRecord.CopyType = ECopyType::BoolProperty;
 			}
-			else if (FStructProperty* StructProperty = CastField<FStructProperty>(CopyRecord.DestProperty.Get()))
+			else if (CastField<FStructProperty>(CopyRecord.DestProperty.Get()))
 			{
 				CopyRecord.CopyType = ECopyType::StructProperty;
 			}
-			else if (FObjectPropertyBase* ObjectProperty = CastField<FObjectPropertyBase>(CopyRecord.DestProperty.Get()))
+			else if (CastField<FObjectPropertyBase>(CopyRecord.DestProperty.Get()))
 			{
 				CopyRecord.CopyType = ECopyType::ObjectProperty;
 			}
 			else
 			{
-				CopyRecord.CopyType = ECopyType::MemCopy;
+				check(CopyRecord.DestProperty->PropertyFlags & CPF_IsPlainOldData);
+				CopyRecord.CopyType = ECopyType::PlainProperty;
 			}
 		}
 	}
@@ -736,7 +741,7 @@ void FExposedValueHandler::Execute(const FAnimationBaseContext& Context) const
 				switch(CopyRecord.CopyType)
 				{
 				default:
-				case ECopyType::MemCopy:
+				case ECopyType::PlainProperty:
 					FMemory::Memcpy(Dest, Src, CopyRecord.Size);
 					break;
 				case ECopyType::BoolProperty:
@@ -754,6 +759,12 @@ void FExposedValueHandler::Execute(const FAnimationBaseContext& Context) const
 						static_cast<FObjectPropertyBase*>(CopyRecord.DestProperty.Get())->SetObjectPropertyValue(Dest, Value);
 					}
 					break;
+				case ECopyType::NameProperty:
+					{
+						const FName& NameValue = static_cast<FNameProperty*>(SourceProperty)->GetPropertyValue(Src);
+						static_cast<FNameProperty*>(CopyRecord.DestProperty.Get())->SetPropertyValue(Dest, NameValue);
+					}
+					break;
 				}
 			}
 			break;
diff --git a/Engine/Source/Runtime/Engine/Private/Animation/AnimNode_LinkedAnimGraph.cpp b/Engine/Source/Runtime/Engine/Private/Animation/AnimNode_LinkedAnimGraph.cpp
--- a/Engine/Source/Runtime/Engine/Private/Animation/AnimNode_LinkedAnimGraph.cpp
+++ b/Engine/Source/Runtime/Engine/Private/Animation/AnimNode_LinkedAnimGraph.cpp
@@ -35,6 +35,8 @@ FAnimNode_LinkedAnimGraph::FAnimNode_LinkedAnimGraph()
 	, LinkedRoot(nullptr)
 	, NodeIndex(INDEX_NONE)
 	, PendingBlendDuration(-1.0f)
+	, bReceiveNotifiesFromLinkedInstances(false)
+	, bPropagateNotifiesToLinkedInstances(false)
 {
 }
 
@@ -258,6 +260,8 @@ void FAnimNode_LinkedAnimGraph::ReinitializeLinkedAnimInstance(const UAnimInstan
 			// we mark them as created by linked anim graph
 			// this is to know who owns memory instance
 			InstanceToRun->bCreatedByLinkedAnimGraph = true;
+			InstanceToRun->bPropagateNotifiesToLinkedInstances = bPropagateNotifiesToLinkedInstances;
+			InstanceToRun->bReceiveNotifiesFromLinkedInstances = bReceiveNotifiesFromLinkedInstances;
 		}
 
 		SetTargetInstance(InstanceToRun);
diff --git a/Engine/Source/Runtime/Engine/Private/Animation/AnimationSettings.cpp b/Engine/Source/Runtime/Engine/Private/Animation/AnimationSettings.cpp
--- a/Engine/Source/Runtime/Engine/Private/Animation/AnimationSettings.cpp
+++ b/Engine/Source/Runtime/Engine/Private/Animation/AnimationSettings.cpp
@@ -6,17 +6,11 @@
 UAnimationSettings::UAnimationSettings(const FObjectInitializer& ObjectInitializer)
 	: Super(ObjectInitializer)
 	, CompressCommandletVersion(2)	// Bump this up to trigger full recompression. Otherwise only new animations imported will be recompressed.
-	, DefaultCompressionAlgorithm(UAnimCompress_BitwiseCompressOnly::StaticClass())
-	, RotationCompressionFormat(ACF_Float96NoW)
-	, TranslationCompressionFormat(ACF_None)
-	, MaxCurveError(0.f)
 	, ForceRecompression(false)
 	, bOnlyCheckForMissingSkeletalMeshes(false)
 	, bForceBelowThreshold(false)
 	, bFirstRecompressUsingCurrentOrDefault(true)
 	, bRaiseMaxErrorToExisting(false)
-	, bTryExhaustiveSearch(false)
-	, bEnableSegmenting(true)
 	, bEnablePerformanceLog(false)
 	, bTickAnimationOnSkeletalMeshInit(true)
 {
diff --git a/Engine/Source/Runtime/Engine/Private/Animation/Skeleton.cpp b/Engine/Source/Runtime/Engine/Private/Animation/Skeleton.cpp
--- a/Engine/Source/Runtime/Engine/Private/Animation/Skeleton.cpp
+++ b/Engine/Source/Runtime/Engine/Private/Animation/Skeleton.cpp
@@ -7,6 +7,8 @@
 #include "Animation/Skeleton.h"
 #include "UObject/UObjectHash.h"
 #include "UObject/UObjectIterator.h"
+#include "UObject/UObjectGlobals.h"
+#include "UObject/PackageReload.h"
 #include "Misc/MessageDialog.h"
 #include "Modules/ModuleManager.h"
 #include "Components/SkinnedMeshComponent.h"
@@ -111,6 +113,20 @@ USkeleton::USkeleton(const FObjectInitializer& ObjectInitializer)
 	AnimCurveMapping = SmartNames.AddContainer(AnimCurveMappingName);
 
 	AnimCurveUidVersion = 0;
+
+	if (HasAnyFlags(RF_ClassDefaultObject))
+	{
+		FCoreUObjectDelegates::OnPackageReloaded.AddStatic(&USkeleton::HandlePackageReloaded);
+	}
+}
+
+void USkeleton::BeginDestroy()
+{
+	Super::BeginDestroy();
+	if (HasAnyFlags(RF_ClassDefaultObject))
+	{
+		FCoreUObjectDelegates::OnPackageReloaded.RemoveAll(this);
+	}
 }
 
 void USkeleton::PostInitProperties()
@@ -1980,4 +1996,18 @@ const TArray<UAssetUserData*>* USkeleton::GetAssetUserDataArray() const
 	return &AssetUserData;
 }
 
+void USkeleton::HandlePackageReloaded(const EPackageReloadPhase InPackageReloadPhase, FPackageReloadedEvent* InPackageReloadedEvent)
+{
+	if (InPackageReloadPhase == EPackageReloadPhase::PostPackageFixup)
+	{
+		for (const auto& RepointedObjectPair : InPackageReloadedEvent->GetRepointedObjects())
+		{
+			if (USkeleton* NewObject = Cast<USkeleton>(RepointedObjectPair.Value))
+			{
+				NewObject->HandleVirtualBoneChanges(); // Reloading Skeletons can invalidate virtual bones so refresh
+			}
+		}
+	}
+}
+
 #undef LOCTEXT_NAMESPACE 
diff --git a/Engine/Source/Runtime/Engine/Private/Components/SkeletalMeshComponent.cpp b/Engine/Source/Runtime/Engine/Private/Components/SkeletalMeshComponent.cpp
--- a/Engine/Source/Runtime/Engine/Private/Components/SkeletalMeshComponent.cpp
+++ b/Engine/Source/Runtime/Engine/Private/Components/SkeletalMeshComponent.cpp
@@ -836,11 +836,11 @@ void USkeletalMeshComponent::BeginPlay()
 	// Trace the 'first frame' markers
 	TRACE_SKELETAL_MESH_COMPONENT(this);
 
-	if (AnimScriptInstance)
+	ForEachAnimInstance([](UAnimInstance* InAnimInstance)
 	{
-		AnimScriptInstance->NativeBeginPlay();
-		AnimScriptInstance->BlueprintBeginPlay();
-	}
+		InAnimInstance->NativeBeginPlay();
+		InAnimInstance->BlueprintBeginPlay();
+	});
 }
 
 #if WITH_EDITOR
@@ -2292,8 +2292,11 @@ void USkeletalMeshComponent::PostAnimEvaluation(FAnimationEvaluationContext& Eva
 	if (EvaluationContext.AnimInstance && EvaluationContext.AnimInstance->NeedsUpdate())
 	{
 		EvaluationContext.AnimInstance->PostUpdateAnimation();
+	}
 
-		for (UAnimInstance* LinkedInstance : LinkedInstances)
+	for (UAnimInstance* LinkedInstance : LinkedInstances)
+	{
+		if(LinkedInstance->NeedsUpdate())
 		{
 			LinkedInstance->PostUpdateAnimation();
 		}
@@ -2812,6 +2815,27 @@ UAnimInstance* USkeletalMeshComponent::GetLinkedAnimLayerInstanceByClass(TSubcla
 	return nullptr;
 }
 
+void USkeletalMeshComponent::ForEachAnimInstance(TFunctionRef<void(UAnimInstance*)> InFunction)
+{
+	if(AnimScriptInstance)
+	{
+		InFunction(AnimScriptInstance);
+	}
+
+	for(UAnimInstance* LinkedInstance : LinkedInstances)
+	{
+		if(LinkedInstance)
+		{
+			InFunction(LinkedInstance);
+		}
+	}
+
+	if(PostProcessAnimInstance)
+	{
+		InFunction(PostProcessAnimInstance);
+	}
+}
+
 bool USkeletalMeshComponent::HasValidAnimationInstance() const
 {
 	return AnimScriptInstance || PostProcessAnimInstance;
diff --git a/Engine/Source/Runtime/Engine/Private/Components/SkinnedMeshComponent.cpp b/Engine/Source/Runtime/Engine/Private/Components/SkinnedMeshComponent.cpp
--- a/Engine/Source/Runtime/Engine/Private/Components/SkinnedMeshComponent.cpp
+++ b/Engine/Source/Runtime/Engine/Private/Components/SkinnedMeshComponent.cpp
@@ -1371,6 +1371,13 @@ bool USkinnedMeshComponent::GetTwistAndSwingAngleOfDeltaRotationFromRefPose(FNam
 	const int32 BoneIndex = GetBoneIndex(BoneName);
 	const TArray<FTransform>& Transforms = GetComponentSpaceTransforms();
 
+	// detect the case where we don't have a pose yet
+	if (Transforms.Num() == 0)
+	{
+		OutTwistAngle = OutSwingAngle = 0.f;
+		return false;
+	}
+
 	if (BoneIndex != INDEX_NONE && ensureMsgf(BoneIndex < Transforms.Num(), TEXT("Invalid transform access in %s. Index=%d, Num=%d"), *GetPathName(), BoneIndex, Transforms.Num()))
 	{
 		FTransform LocalTransform = GetComponentSpaceTransforms()[BoneIndex];
@@ -1478,9 +1485,19 @@ void USkinnedMeshComponent::SetSkeletalMesh(USkeletalMesh* InSkelMesh, bool bRei
 
 		SkeletalMesh = InSkelMesh;
 
-		for (TWeakObjectPtr<USkinnedMeshComponent>& SlavePoseComponent : SlavePoseComponents)
+		//SlavePoseComponents is an array of weak obj ptrs, so it can contain null elements
+		for (auto Iter = SlavePoseComponents.CreateIterator(); Iter; ++Iter)
 		{
-			SlavePoseComponent->UpdateMasterBoneMap();
+			TWeakObjectPtr<USkinnedMeshComponent> Comp = (*Iter);
+			if (Comp.IsValid() == false)
+			{
+				SlavePoseComponents.RemoveAt(Iter.GetIndex());
+				--Iter;
+			}
+			else
+			{
+				Comp->UpdateMasterBoneMap();
+			}
 		}
 
 		// Don't init anim state if not registered
diff --git a/Engine/Source/Runtime/Engine/Private/GPUSkinCache.cpp b/Engine/Source/Runtime/Engine/Private/GPUSkinCache.cpp
--- a/Engine/Source/Runtime/Engine/Private/GPUSkinCache.cpp
+++ b/Engine/Source/Runtime/Engine/Private/GPUSkinCache.cpp
@@ -107,6 +107,18 @@ FAutoConsoleVariableRef CVarGPUSkinCacheAllowDupedVertesForRecomputeTangents(
 	ECVF_RenderThreadSafe
 );
 
+static int32 GBlendUsingVertexColorForRecomputeTangents = 0;
+FAutoConsoleVariableRef CVarGPUSkinCacheBlendUsingVertexColorForRecomputeTangents(
+	TEXT("r.SkinCache.BlendUsingVertexColorForRecomputeTangents"),
+	GBlendUsingVertexColorForRecomputeTangents,
+	TEXT("0: off (default)\n")
+	TEXT("1: No blending, choose between source and recompute tangents.\n")
+	TEXT("2: Linear interpolation between source and recompute tangents.\n")
+	TEXT("3: Vector slerp between source and recompute tangents.\n")
+	TEXT("4: Convert tangents into quaternion, apply slerp, then convert from quaternion back to tangents (most expensive).\n"),
+	ECVF_RenderThreadSafe
+);
+
 static int32 GGPUSkinCacheFlushCounter = 0;
 
 static inline bool DoesPlatformSupportGPUSkinCache(EShaderPlatform Platform)
@@ -203,6 +215,7 @@ public:
 
 		FShaderResourceViewRHIRef TangentBufferSRV = nullptr;
 		FShaderResourceViewRHIRef UVsBufferSRV = nullptr;
+		FShaderResourceViewRHIRef ColorBufferSRV = nullptr;
 		FShaderResourceViewRHIRef PositionBufferSRV = nullptr;
 
 		// skin weight input
@@ -223,6 +236,7 @@ public:
 		uint32 NumTriangles = 0;
 
 		FRWBuffer* TangentBuffer = nullptr;
+		FRWBuffer* IntermediateTangentBuffer = nullptr;
 		FRWBuffer* PositionBuffer = nullptr;
 		FRWBuffer* PreviousPositionBuffer = nullptr;
 
@@ -337,6 +351,7 @@ public:
 
 		Data.TangentBufferSRV = InSourceVertexFactory->GetTangentsSRV();
 		Data.UVsBufferSRV = InSourceVertexFactory->GetTextureCoordinatesSRV();
+		Data.ColorBufferSRV = InSourceVertexFactory->GetColorComponentsSRV();
 		Data.NumTexCoords = InSourceVertexFactory->GetNumTexCoords();
 		Data.PositionBufferSRV = InSourceVertexFactory->GetPositionsSRV();
 
@@ -403,7 +418,7 @@ protected:
 
 	friend class FGPUSkinCache;
 	friend class FBaseGPUSkinCacheCS;
-	friend class FBaseRecomputeTangents;
+	friend class FBaseRecomputeTangentsPerTriangleShader;
 };
 
 class FBaseGPUSkinCacheCS : public FGlobalShader
@@ -704,7 +719,7 @@ void FGPUSkinCache::TransitionAllToWriteable(FRHICommandList& RHICmdList)
 #endif
 
 /** base of the FRecomputeTangentsPerTrianglePassCS class */
-class FBaseRecomputeTangents : public FGlobalShader
+class FBaseRecomputeTangentsPerTriangleShader : public FGlobalShader
 {
 public:
 	static bool ShouldCompilePermutation(const FGlobalShaderPermutationParameters& Parameters)
@@ -730,10 +745,10 @@ public:
 	FShaderResourceParameter DuplicatedIndices;
 	FShaderResourceParameter DuplicatedIndicesIndices;
 
-	FBaseRecomputeTangents()
+	FBaseRecomputeTangentsPerTriangleShader()
 	{}
 
-	FBaseRecomputeTangents(const ShaderMetaType::CompiledShaderInitializerType& Initializer)
+	FBaseRecomputeTangentsPerTriangleShader(const ShaderMetaType::CompiledShaderInitializerType& Initializer)
 		: FGlobalShader(Initializer)
 	{
 		IntermediateAccumBufferUAV.Bind(Initializer.ParameterMap, TEXT("IntermediateAccumBufferUAV"));
@@ -763,7 +778,7 @@ public:
 		SetShaderValue(RHICmdList, ShaderRHI, NumTriangles, DispatchData.NumTriangles);
 
 		SetSRVParameter(RHICmdList, ShaderRHI, GPUPositionCacheBuffer, DispatchData.GetPositionRWBuffer()->SRV);
-		SetSRVParameter(RHICmdList, ShaderRHI, GPUTangentCacheBuffer, DispatchData.GetTangentRWBuffer()->SRV);
+		SetSRVParameter(RHICmdList, ShaderRHI, GPUTangentCacheBuffer, GBlendUsingVertexColorForRecomputeTangents > 0 ? DispatchData.IntermediateTangentBuffer->SRV : DispatchData.GetTangentRWBuffer()->SRV);
 		SetSRVParameter(RHICmdList, ShaderRHI, UVsInputBuffer, DispatchData.UVsBufferSRV);
 
 		SetShaderValue(RHICmdList, ShaderRHI, SkinCacheStart, DispatchData.OutputStreamStart);
@@ -806,7 +821,7 @@ public:
 
 /** Encapsulates the RecomputeSkinTangents compute shader. */
 template <int Permutation>
-class FRecomputeTangentsPerTrianglePassCS : public FBaseRecomputeTangents
+class FRecomputeTangentsPerTrianglePassCS : public FBaseRecomputeTangentsPerTriangleShader
 {
     constexpr static bool bMergeDuplicatedVerts = (2 == (Permutation & 2));
 	constexpr static bool bFullPrecisionUV = (1 == (Permutation & 1));
@@ -823,7 +838,7 @@ class FRecomputeTangentsPerTrianglePassCS : public FBaseRecomputeTangents
 	}
 
 	FRecomputeTangentsPerTrianglePassCS(const ShaderMetaType::CompiledShaderInitializerType& Initializer)
-		: FBaseRecomputeTangents(Initializer)
+		: FBaseRecomputeTangentsPerTriangleShader(Initializer)
 	{
 	}
 
@@ -837,42 +852,34 @@ IMPLEMENT_SHADER_TYPE(template<>, FRecomputeTangentsPerTrianglePassCS<2>, TEXT("
 IMPLEMENT_SHADER_TYPE(template<>, FRecomputeTangentsPerTrianglePassCS<3>, TEXT("/Engine/Private/RecomputeTangentsPerTrianglePass.usf"), TEXT("MainCS"), SF_Compute);
 
 /** Encapsulates the RecomputeSkinTangentsResolve compute shader. */
-class FRecomputeTangentsPerVertexPassCS : public FGlobalShader
+class FBaseRecomputeTangentsPerVertexShader : public FGlobalShader
 {
-	DECLARE_SHADER_TYPE(FRecomputeTangentsPerVertexPassCS, Global);
-
+public:
 	static bool ShouldCompilePermutation(const FGlobalShaderPermutationParameters& Parameters)
 	{
 		// currently only implemented and tested on Window SM5 (needs Compute, Atomics, SRV for index buffers, UAV for VertexBuffers)
 		return DoesPlatformSupportGPUSkinCache(Parameters.Platform) && IsGPUSkinCacheAvailable(Parameters.Platform);
 	}
 
-	static void ModifyCompilationEnvironment(const FGlobalShaderPermutationParameters& Parameters, FShaderCompilerEnvironment& OutEnvironment)
-	{
-		FGlobalShader::ModifyCompilationEnvironment(Parameters, OutEnvironment);
-		// this pass cannot read the input as it doesn't have the permutation
-		OutEnvironment.SetDefine(TEXT("THREADGROUP_SIZEX"), ThreadGroupSizeX);
-		OutEnvironment.SetDefine(TEXT("GPUSKIN_RWBUFFER_OFFSET_TANGENT_X"), FGPUSkinCache::RWTangentXOffsetInFloats);
-		OutEnvironment.SetDefine(TEXT("GPUSKIN_RWBUFFER_OFFSET_TANGENT_Z"), FGPUSkinCache::RWTangentZOffsetInFloats);
-		OutEnvironment.SetDefine(TEXT("INTERMEDIATE_ACCUM_BUFFER_NUM_INTS"), FGPUSkinCache::IntermediateAccumBufferNumInts);
-	}
-
 	static const uint32 ThreadGroupSizeX = 64;
 
-	FRecomputeTangentsPerVertexPassCS() {}
-
-public:
 	FShaderResourceParameter IntermediateAccumBufferUAV;
 	FShaderResourceParameter TangentBufferUAV;
+	FShaderResourceParameter TangentInputBuffer;
+	FShaderResourceParameter ColorInputBuffer;
 	FShaderParameter SkinCacheStart;
 	FShaderParameter NumVertices;
 	FShaderParameter InputStreamStart;
 
-	FRecomputeTangentsPerVertexPassCS(const ShaderMetaType::CompiledShaderInitializerType& Initializer)
+	FBaseRecomputeTangentsPerVertexShader() {}
+
+	FBaseRecomputeTangentsPerVertexShader(const ShaderMetaType::CompiledShaderInitializerType& Initializer)
 		: FGlobalShader(Initializer)
 	{
 		IntermediateAccumBufferUAV.Bind(Initializer.ParameterMap, TEXT("IntermediateAccumBufferUAV"));
 		TangentBufferUAV.Bind(Initializer.ParameterMap, TEXT("TangentBufferUAV"));
+		TangentInputBuffer.Bind(Initializer.ParameterMap, TEXT("TangentInputBuffer"));
+		ColorInputBuffer.Bind(Initializer.ParameterMap, TEXT("ColorInputBuffer"));
 		SkinCacheStart.Bind(Initializer.ParameterMap, TEXT("SkinCacheStart"));
 		NumVertices.Bind(Initializer.ParameterMap, TEXT("NumVertices"));
 		InputStreamStart.Bind(Initializer.ParameterMap, TEXT("InputStreamStart"));
@@ -893,6 +900,9 @@ public:
 		// UAVs
 		SetUAVParameter(RHICmdList, ShaderRHI, IntermediateAccumBufferUAV, StagingBuffer.UAV);
 		SetUAVParameter(RHICmdList, ShaderRHI, TangentBufferUAV, DispatchData.GetTangentRWBuffer()->UAV);
+
+		SetSRVParameter(RHICmdList, ShaderRHI, TangentInputBuffer, DispatchData.IntermediateTangentBuffer ? DispatchData.IntermediateTangentBuffer->SRV : nullptr);
+		SetSRVParameter(RHICmdList, ShaderRHI, ColorInputBuffer, DispatchData.ColorBufferSRV);
 	}
 
 	void UnsetParameters(FRHICommandList& RHICmdList)
@@ -906,12 +916,41 @@ public:
 	virtual bool Serialize(FArchive& Ar) override
 	{
 		bool bShaderHasOutdatedParameters = FGlobalShader::Serialize(Ar);
-		Ar << IntermediateAccumBufferUAV << TangentBufferUAV << SkinCacheStart << NumVertices << InputStreamStart;
+		Ar << IntermediateAccumBufferUAV << TangentBufferUAV << TangentInputBuffer << ColorInputBuffer << SkinCacheStart << NumVertices << InputStreamStart;
 		return bShaderHasOutdatedParameters;
 	}
 };
 
-IMPLEMENT_SHADER_TYPE(, FRecomputeTangentsPerVertexPassCS, TEXT("/Engine/Private/RecomputeTangentsPerVertexPass.usf"), TEXT("MainCS"), SF_Compute);
+template <int Permutation>
+class FRecomputeTangentsPerVertexPassCS : public FBaseRecomputeTangentsPerVertexShader
+{
+	DECLARE_SHADER_TYPE(FRecomputeTangentsPerVertexPassCS, Global);
+
+	static void ModifyCompilationEnvironment(const FGlobalShaderPermutationParameters& Parameters, FShaderCompilerEnvironment& OutEnvironment)
+	{
+		FGlobalShader::ModifyCompilationEnvironment(Parameters, OutEnvironment);
+		// this pass cannot read the input as it doesn't have the permutation
+		OutEnvironment.SetDefine(TEXT("THREADGROUP_SIZEX"), ThreadGroupSizeX);
+		OutEnvironment.SetDefine(TEXT("GPUSKIN_RWBUFFER_OFFSET_TANGENT_X"), FGPUSkinCache::RWTangentXOffsetInFloats);
+		OutEnvironment.SetDefine(TEXT("GPUSKIN_RWBUFFER_OFFSET_TANGENT_Z"), FGPUSkinCache::RWTangentZOffsetInFloats);
+		OutEnvironment.SetDefine(TEXT("INTERMEDIATE_ACCUM_BUFFER_NUM_INTS"), FGPUSkinCache::IntermediateAccumBufferNumInts);
+		OutEnvironment.SetDefine(TEXT("BLEND_USING_VERTEX_COLOR"), Permutation);
+	}
+
+	FRecomputeTangentsPerVertexPassCS(const ShaderMetaType::CompiledShaderInitializerType& Initializer)
+		: FBaseRecomputeTangentsPerVertexShader(Initializer)
+	{
+	}
+
+	FRecomputeTangentsPerVertexPassCS()
+	{}
+};
+
+IMPLEMENT_SHADER_TYPE(template<>, FRecomputeTangentsPerVertexPassCS<0>, TEXT("/Engine/Private/RecomputeTangentsPerVertexPass.usf"), TEXT("MainCS"), SF_Compute);
+IMPLEMENT_SHADER_TYPE(template<>, FRecomputeTangentsPerVertexPassCS<1>, TEXT("/Engine/Private/RecomputeTangentsPerVertexPass.usf"), TEXT("MainCS"), SF_Compute);
+IMPLEMENT_SHADER_TYPE(template<>, FRecomputeTangentsPerVertexPassCS<2>, TEXT("/Engine/Private/RecomputeTangentsPerVertexPass.usf"), TEXT("MainCS"), SF_Compute);
+IMPLEMENT_SHADER_TYPE(template<>, FRecomputeTangentsPerVertexPassCS<3>, TEXT("/Engine/Private/RecomputeTangentsPerVertexPass.usf"), TEXT("MainCS"), SF_Compute);
+IMPLEMENT_SHADER_TYPE(template<>, FRecomputeTangentsPerVertexPassCS<4>, TEXT("/Engine/Private/RecomputeTangentsPerVertexPass.usf"), TEXT("MainCS"), SF_Compute);
 
 void FGPUSkinCache::DispatchUpdateSkinTangents(FRHICommandListImmediate& RHICmdList, FGPUSkinCacheEntry* Entry, int32 SectionIndex)
 {
@@ -969,21 +1008,21 @@ void FGPUSkinCache::DispatchUpdateSkinTangents(FRHICommandListImmediate& RHICmdL
 
 			bool bFullPrecisionUV = LodData.StaticVertexBuffers.StaticMeshVertexBuffer.GetUseFullPrecisionUVs();
 
-			FBaseRecomputeTangents* Shader = 0;
+			FBaseRecomputeTangentsPerTriangleShader* Shader = 0;
 
 			if (bFullPrecisionUV)
 			{
-				Shader = GAllowDupedVertsForRecomputeTangents ? (FBaseRecomputeTangents*)*ComputeShader11 : (FBaseRecomputeTangents*)*ComputeShader01;
+				Shader = GAllowDupedVertsForRecomputeTangents ? (FBaseRecomputeTangentsPerTriangleShader*)*ComputeShader11 : (FBaseRecomputeTangentsPerTriangleShader*)*ComputeShader01;
 			}
 			else
 			{
-				Shader = GAllowDupedVertsForRecomputeTangents ? (FBaseRecomputeTangents*)*ComputeShader10 : (FBaseRecomputeTangents*)*ComputeShader00;
+				Shader = GAllowDupedVertsForRecomputeTangents ? (FBaseRecomputeTangentsPerTriangleShader*)*ComputeShader10 : (FBaseRecomputeTangentsPerTriangleShader*)*ComputeShader00;
 			}
 
 			check(Shader);
 
 			uint32 NumTriangles = DispatchData.NumTriangles;
-			uint32 ThreadGroupCountValue = FMath::DivideAndRoundUp(NumTriangles, FBaseRecomputeTangents::ThreadGroupSizeX);
+			uint32 ThreadGroupCountValue = FMath::DivideAndRoundUp(NumTriangles, FBaseRecomputeTangentsPerTriangleShader::ThreadGroupSizeX);
 
 			SCOPED_DRAW_EVENTF(RHICmdList, SkinTangents_PerTrianglePass, TEXT("TangentsTri IndexStart=%d Tri=%d BoneInfluenceType=%d UVPrecision=%d"),
 				DispatchData.IndexBufferOffsetValue, DispatchData.NumTriangles, Entry->BoneInfluenceType, bFullPrecisionUV);
@@ -1010,17 +1049,33 @@ void FGPUSkinCache::DispatchUpdateSkinTangents(FRHICommandListImmediate& RHICmdL
 			SCOPED_DRAW_EVENTF(RHICmdList, SkinTangents_PerVertexPass, TEXT("TangentsVertex InputStreamStart=%d, OutputStreamStart=%d, Vert=%d"),
 				DispatchData.InputStreamStart, DispatchData.OutputStreamStart, DispatchData.NumVertices);
 			//#todo-gpuskin Feature level?
-			TShaderMapRef<FRecomputeTangentsPerVertexPassCS> ComputeShader(GetGlobalShaderMap(GMaxRHIFeatureLevel));
+			auto* GlobalShaderMap = GetGlobalShaderMap(GMaxRHIFeatureLevel);
+			TShaderMapRef<FRecomputeTangentsPerVertexPassCS<0>> ComputeShader0(GlobalShaderMap);
+			TShaderMapRef<FRecomputeTangentsPerVertexPassCS<1>> ComputeShader1(GlobalShaderMap);
+			TShaderMapRef<FRecomputeTangentsPerVertexPassCS<2>> ComputeShader2(GlobalShaderMap);
+			TShaderMapRef<FRecomputeTangentsPerVertexPassCS<3>> ComputeShader3(GlobalShaderMap);
+			TShaderMapRef<FRecomputeTangentsPerVertexPassCS<4>> ComputeShader4(GlobalShaderMap);
+			FBaseRecomputeTangentsPerVertexShader* ComputeShader = nullptr;
+			if (GBlendUsingVertexColorForRecomputeTangents == 1)
+				ComputeShader = (FBaseRecomputeTangentsPerVertexShader*)*ComputeShader1;
+			else if (GBlendUsingVertexColorForRecomputeTangents == 2)
+				ComputeShader = (FBaseRecomputeTangentsPerVertexShader*)*ComputeShader2;
+			else if (GBlendUsingVertexColorForRecomputeTangents == 3)
+				ComputeShader = (FBaseRecomputeTangentsPerVertexShader*)*ComputeShader3;
+			else if (GBlendUsingVertexColorForRecomputeTangents == 4)
+				ComputeShader = (FBaseRecomputeTangentsPerVertexShader*)*ComputeShader4;
+			else
+				ComputeShader = (FBaseRecomputeTangentsPerVertexShader*)*ComputeShader0;
 			FRHIComputeShader* ShaderRHI = ComputeShader->GetComputeShader();
 			RHICmdList.SetComputeShader(ComputeShader->GetComputeShader());
 
 			uint32 VertexCount = DispatchData.NumVertices;
-			uint32 ThreadGroupCountValue = FMath::DivideAndRoundUp(VertexCount, FRecomputeTangentsPerVertexPassCS::ThreadGroupSizeX);
+			uint32 ThreadGroupCountValue = FMath::DivideAndRoundUp(VertexCount, ComputeShader->ThreadGroupSizeX);
 
 			RHICmdList.TransitionResource(EResourceTransitionAccess::ERWBarrier, EResourceTransitionPipeline::EComputeToCompute, StagingBuffer.UAV.GetReference());
 
 			ComputeShader->SetParameters(RHICmdList, Entry, DispatchData, StagingBuffer);
-			DispatchComputeShader(RHICmdList, *ComputeShader, ThreadGroupCountValue, 1, 1);
+			DispatchComputeShader(RHICmdList, ComputeShader, ThreadGroupCountValue, 1, 1);
 			ComputeShader->UnsetParameters(RHICmdList);
 		}
 		// todo				RHICmdList.TransitionResource(EResourceTransitionAccess::EReadable, EResourceTransitionPipeline::EComputeToGfx, TangentsBlendBuffer.VertexBufferSRV);
@@ -1310,7 +1365,10 @@ void FGPUSkinCache::DispatchUpdateSkinning(FRHICommandListImmediate& RHICmdList,
 	uint32 CurrentRevision = ShaderData.GetRevisionNumber(false);
 	uint32 PreviousRevision = ShaderData.GetRevisionNumber(true);
 
+	bool bUseIntermediateTangentBuffer = Entry->DispatchData[Section].IndexBuffer && GBlendUsingVertexColorForRecomputeTangents > 0;
+
 	DispatchData.TangentBuffer = DispatchData.PositionTracker.GetTangentBuffer();
+	DispatchData.IntermediateTangentBuffer = DispatchData.PositionTracker.GetIntermediateTangentBuffer();
 
 	DispatchData.PreviousPositionBuffer = DispatchData.PositionTracker.Find(PrevBoneBuffer, PreviousRevision);
 	if (!DispatchData.PreviousPositionBuffer)
@@ -1320,15 +1378,34 @@ void FGPUSkinCache::DispatchUpdateSkinning(FRHICommandListImmediate& RHICmdList,
 		check(DispatchData.PreviousPositionBuffer)
 
 		RHICmdList.SetComputeShader(Shader->GetComputeShader());
-		Shader->SetParameters(RHICmdList, PrevBoneBuffer, Entry, DispatchData, DispatchData.GetPreviousPositionRWBuffer()->UAV, DispatchData.GetTangentRWBuffer() ? DispatchData.GetTangentRWBuffer()->UAV : nullptr);
+
+		if (bUseIntermediateTangentBuffer)
+		{
+			Shader->SetParameters(RHICmdList, PrevBoneBuffer, Entry, DispatchData, DispatchData.GetPreviousPositionRWBuffer()->UAV, DispatchData.IntermediateTangentBuffer ? DispatchData.IntermediateTangentBuffer->UAV : nullptr);
+		}
+		else
+		{
+			Shader->SetParameters(RHICmdList, PrevBoneBuffer, Entry, DispatchData, DispatchData.GetPreviousPositionRWBuffer()->UAV, DispatchData.GetTangentRWBuffer() ? DispatchData.GetTangentRWBuffer()->UAV : nullptr);
+		}
 
 		RHICmdList.TransitionResource(EResourceTransitionAccess::EWritable, EResourceTransitionPipeline::EGfxToCompute, DispatchData.GetPreviousPositionRWBuffer()->UAV.GetReference());
 		BuffersToTransition.Add(DispatchData.GetPreviousPositionRWBuffer()->UAV);
 
-		if (DispatchData.TangentBuffer)
+		if (bUseIntermediateTangentBuffer)
+		{
+			if (DispatchData.IntermediateTangentBuffer)
+			{
+				RHICmdList.TransitionResource(EResourceTransitionAccess::EWritable, EResourceTransitionPipeline::EGfxToCompute, DispatchData.IntermediateTangentBuffer->UAV.GetReference());
+				BuffersToTransition.Add(DispatchData.IntermediateTangentBuffer->UAV);
+			}
+		}
+		else
 		{
-			RHICmdList.TransitionResource(EResourceTransitionAccess::EWritable, EResourceTransitionPipeline::EGfxToCompute, DispatchData.TangentBuffer->UAV.GetReference());
-			BuffersToTransition.Add(DispatchData.TangentBuffer->UAV);
+			if (DispatchData.TangentBuffer)
+			{
+				RHICmdList.TransitionResource(EResourceTransitionAccess::EWritable, EResourceTransitionPipeline::EGfxToCompute, DispatchData.TangentBuffer->UAV.GetReference());
+				BuffersToTransition.Add(DispatchData.TangentBuffer->UAV);
+			}
 		}
 
 		uint32 VertexCountAlign64 = FMath::DivideAndRoundUp(DispatchData.NumVertices, (uint32)64);
@@ -1346,15 +1423,33 @@ void FGPUSkinCache::DispatchUpdateSkinning(FRHICommandListImmediate& RHICmdList,
 		check(DispatchData.PositionBuffer);
 
 		RHICmdList.SetComputeShader(Shader->GetComputeShader());
-		Shader->SetParameters(RHICmdList, BoneBuffer, Entry, DispatchData, DispatchData.GetPositionRWBuffer()->UAV, DispatchData.GetTangentRWBuffer() ? DispatchData.GetTangentRWBuffer()->UAV : nullptr);
+		if (bUseIntermediateTangentBuffer)
+		{
+			Shader->SetParameters(RHICmdList, BoneBuffer, Entry, DispatchData, DispatchData.GetPositionRWBuffer()->UAV, DispatchData.IntermediateTangentBuffer ? DispatchData.IntermediateTangentBuffer->UAV : nullptr);
+		}
+		else
+		{
+			Shader->SetParameters(RHICmdList, BoneBuffer, Entry, DispatchData, DispatchData.GetPositionRWBuffer()->UAV, DispatchData.GetTangentRWBuffer() ? DispatchData.GetTangentRWBuffer()->UAV : nullptr);
+		}
 
 		RHICmdList.TransitionResource(EResourceTransitionAccess::EWritable, EResourceTransitionPipeline::EGfxToCompute, DispatchData.GetPositionRWBuffer()->UAV.GetReference());
 		BuffersToTransition.Add(DispatchData.GetPositionRWBuffer()->UAV);
 
-		if (DispatchData.TangentBuffer)
+		if (bUseIntermediateTangentBuffer)
+		{
+			if (DispatchData.IntermediateTangentBuffer)
+			{
+				RHICmdList.TransitionResource(EResourceTransitionAccess::EWritable, EResourceTransitionPipeline::EGfxToCompute, DispatchData.IntermediateTangentBuffer->UAV.GetReference());
+				BuffersToTransition.Add(DispatchData.IntermediateTangentBuffer->UAV);
+			}
+		}
+		else
 		{
-			RHICmdList.TransitionResource(EResourceTransitionAccess::EWritable, EResourceTransitionPipeline::EGfxToCompute, DispatchData.TangentBuffer->UAV.GetReference());
-			BuffersToTransition.Add(DispatchData.TangentBuffer->UAV);
+			if (DispatchData.TangentBuffer)
+			{
+				RHICmdList.TransitionResource(EResourceTransitionAccess::EWritable, EResourceTransitionPipeline::EGfxToCompute, DispatchData.TangentBuffer->UAV.GetReference());
+				BuffersToTransition.Add(DispatchData.TangentBuffer->UAV);
+			}
 		}
 
 		uint32 VertexCountAlign64 = FMath::DivideAndRoundUp(DispatchData.NumVertices, (uint32)64);
@@ -1383,6 +1478,13 @@ void FGPUSkinCache::FRWBuffersAllocation::RemoveAllFromTransitionArray(TArray<FR
 				InBuffersToTransition.Remove(TangentBuffer->UAV);
 			}
 		}
+		if (auto IntermediateTangentBuffer = GetIntermediateTangentBuffer())
+		{
+			if (IntermediateTangentBuffer->UAV.IsValid())
+			{
+				InBuffersToTransition.Remove(IntermediateTangentBuffer->UAV);
+			}
+		}
 	}
 }
 
@@ -1423,6 +1525,12 @@ bool FGPUSkinCache::IsEntryValid(FGPUSkinCacheEntry* SkinCacheEntry, int32 Secti
 	return SkinCacheEntry->IsSectionValid(Section);
 }
 
+bool FGPUSkinCache::UseIntermediateTangents()
+{
+	int32 RecomputeTangentsMode = GForceRecomputeTangents > 0 ? 1 : GSkinCacheRecomputeTangents;
+	return (RecomputeTangentsMode > 0) && (GBlendUsingVertexColorForRecomputeTangents > 0);
+}
+
 FGPUSkinBatchElementUserData* FGPUSkinCache::InternalGetFactoryUserData(FGPUSkinCacheEntry* Entry, int32 Section)
 {
 	return &Entry->BatchElementsUserData[Section];
diff --git a/Engine/Source/Runtime/Engine/Public/GPUSkinCache.h b/Engine/Source/Runtime/Engine/Public/GPUSkinCache.h
--- a/Engine/Source/Runtime/Engine/Public/GPUSkinCache.h
+++ b/Engine/Source/Runtime/Engine/Public/GPUSkinCache.h
@@ -143,6 +143,8 @@ public:
 
 	static bool IsEntryValid(FGPUSkinCacheEntry* SkinCacheEntry, int32 Section);
 
+	static bool UseIntermediateTangents();
+
 	inline uint64 GetExtraRequiredMemoryAndReset()
 	{
 		uint64 OriginalValue = ExtraRequiredMemory;
@@ -169,6 +171,10 @@ public:
 			if (WithTangents)
 			{
 				Tangents.Initialize(8, NumVertices * 2, PF_R16G16B16A16_SNORM, BUF_Static, TEXT("SkinCacheTangents"));
+				if (FGPUSkinCache::UseIntermediateTangents())
+				{
+					IntermediateTangents.Initialize(8, NumVertices * 2, PF_R16G16B16A16_SNORM, BUF_Static, TEXT("SkinCacheIntermediateTangents"));
+				}
 			}
 		}
 
@@ -181,6 +187,7 @@ public:
 			if (WithTangents)
 			{
 				Tangents.Release();
+				IntermediateTangents.Release();
 			}
 		}
 
@@ -188,7 +195,12 @@ public:
 		{
 			uint64 PositionBufferSize = 4 * 3 * NumVertices * NUM_BUFFERS;
 			uint64 TangentBufferSize = WithTangents ? 2 * 4 * NumVertices : 0;
-			return TangentBufferSize + PositionBufferSize;
+			uint64 IntermediateTangentBufferSize = 0;
+			if (FGPUSkinCache::UseIntermediateTangents())
+			{
+				IntermediateTangentBufferSize = WithTangents ? 2 * 4 * NumVertices : 0;
+			}
+			return TangentBufferSize + IntermediateTangentBufferSize + PositionBufferSize;
 		}
 
 		uint64 GetNumBytes() const
@@ -201,6 +213,11 @@ public:
 			return WithTangents ? &Tangents : nullptr;
 		}
 
+		FRWBuffer* GetIntermediateTangentBuffer()
+		{
+			return WithTangents ? &IntermediateTangents : nullptr;
+		}
+
 		void RemoveAllFromTransitionArray(TArray<FRHIUnorderedAccessView*>& BuffersToTransition);
 
 	private:
@@ -208,6 +225,7 @@ public:
 		FRWBuffer RWBuffers[NUM_BUFFERS];
 
 		FRWBuffer Tangents;
+		FRWBuffer IntermediateTangents;
 		const uint32 NumVertices;
 		const bool WithTangents;
 	};
@@ -254,6 +272,11 @@ public:
 			return Allocation->GetTangentBuffer();
 		}
 
+		FRWBuffer* GetIntermediateTangentBuffer()
+		{
+			return Allocation->GetIntermediateTangentBuffer();
+		}
+
 		void Advance(const FVertexBufferAndSRV& BoneBuffer1, uint32 Revision1, const FVertexBufferAndSRV& BoneBuffer2, uint32 Revision2)
 		{
 			const FVertexBufferAndSRV* InBoneBuffers[2] = { &BoneBuffer1 , &BoneBuffer2 };
diff --git a/Engine/Source/Runtime/Engine/Public/Rendering/SkinWeightVertexBuffer.h b/Engine/Source/Runtime/Engine/Public/Rendering/SkinWeightVertexBuffer.h
--- a/Engine/Source/Runtime/Engine/Public/Rendering/SkinWeightVertexBuffer.h
+++ b/Engine/Source/Runtime/Engine/Public/Rendering/SkinWeightVertexBuffer.h
@@ -446,7 +446,7 @@ public:
 	FSkinWeightRHIInfo CreateRHIBuffer_RenderThread();
 	FSkinWeightRHIInfo CreateRHIBuffer_Async();
 
-	GPUSkinBoneInfluenceType GetBoneInfluenceType() const;
+	ENGINE_API GPUSkinBoneInfluenceType GetBoneInfluenceType() const;
 	void GetVertexInfluenceOffsetCount(uint32 VertexIndex, uint32& VertexWeightOffset, uint32& VertexInfluenceCount) const;
 	ENGINE_API bool GetRigidWeightBone(uint32 VertexIndex, int32& OutBoneIndex) const;
 	uint32 GetBoneIndex(uint32 VertexIndex, uint32 InfluenceIndex) const;
diff --git a/Engine/Source/Runtime/LiveLinkInterface/Public/ILiveLinkClient.h b/Engine/Source/Runtime/LiveLinkInterface/Public/ILiveLinkClient.h
--- a/Engine/Source/Runtime/LiveLinkInterface/Public/ILiveLinkClient.h
+++ b/Engine/Source/Runtime/LiveLinkInterface/Public/ILiveLinkClient.h
@@ -16,6 +16,7 @@ class ILiveLinkSource;
 class ULiveLinkRole;
 struct FLiveLinkSubjectFrameData;
 struct FTimecode;
+class ULiveLinkSourceSettings;
 
 DECLARE_EVENT_OneParam(ILiveLinkClient, FOnLiveLinkSourceChangedDelegate, FGuid /*SourceGuid*/);
 DECLARE_EVENT_OneParam(ILiveLinkClient, FOnLiveLinkSubjectChangedDelegate, FLiveLinkSubjectKey /*SubjectKey*/);
@@ -232,6 +233,11 @@ public:
 	 */
 	virtual TArray<FLiveLinkTime> GetSubjectFrameTimes(const FLiveLinkSubjectKey& SubjectKey) const = 0;
 
+	/**
+	 * Get the Settings of this source.
+	 */
+	virtual ULiveLinkSourceSettings* GetSourceSettings(const FGuid& SourceKey) const = 0;
+
 	/**
 	 * Get the time of all the frames for a source.
 	 * @note Use for debugging purposes.
@@ -283,6 +289,11 @@ public:
 
 	/** Notify when the LiveLinkClient has ticked. */
 	virtual FSimpleMulticastDelegate& OnLiveLinkTicked() = 0;
+	/**
+	* Performs an internal Tick(). This is to be used when we want to run live link outside of the normal engine tick workflow,
+	* for example when we need to export data that requires live link evaluation during the export process.
+	*/
+	virtual void ForceTick() = 0;
 
 	/** Notify when the list of sources has changed. */
 	virtual FSimpleMulticastDelegate& OnLiveLinkSourcesChanged() = 0;
diff --git a/Engine/Source/Runtime/RigVM/Private/RigVMCore/RigVM.cpp b/Engine/Source/Runtime/RigVM/Private/RigVMCore/RigVM.cpp
--- a/Engine/Source/Runtime/RigVM/Private/RigVMCore/RigVM.cpp
+++ b/Engine/Source/Runtime/RigVM/Private/RigVMCore/RigVM.cpp
@@ -73,16 +73,24 @@ void URigVM::Serialize(FArchive& Ar)
 
 	if (Ar.IsLoading())
 	{
-		Instructions.Reset();
-		Functions.Reset();
-		ParametersNameMap.Reset();
-
-		for (int32 Index = 0; Index < Parameters.Num(); Index++)
+		if (WorkMemory.bEncounteredErrorDuringLoad ||
+			LiteralMemory.bEncounteredErrorDuringLoad)
 		{
-			ParametersNameMap.Add(Parameters[Index].Name, Index);
+			Reset();
 		}
+		else
+		{
+			Instructions.Reset();
+			Functions.Reset();
+			ParametersNameMap.Reset();
 
-		InvalidateCachedMemory();
+			for (int32 Index = 0; Index < Parameters.Num(); Index++)
+			{
+				ParametersNameMap.Add(Parameters[Index].Name, Index);
+			}
+
+			InvalidateCachedMemory();
+		}
 	}
 }
 
diff --git a/Engine/Source/Runtime/RigVM/Private/RigVMCore/RigVMMemory.cpp b/Engine/Source/Runtime/RigVM/Private/RigVMCore/RigVMMemory.cpp
--- a/Engine/Source/Runtime/RigVM/Private/RigVMCore/RigVMMemory.cpp
+++ b/Engine/Source/Runtime/RigVM/Private/RigVMCore/RigVMMemory.cpp
@@ -4,6 +4,7 @@
 #include "UObject/AnimObjectVersion.h"
 #include "UObject/PropertyPortFlags.h"
 #include "UObject/Package.h"
+#include "RigVMModule.h"
 
 bool FRigVMRegister::Serialize(FArchive& Ar)
 {
@@ -291,6 +292,7 @@ uint16 FRigVMRegisterOffset::GetElementSize() const
 FRigVMMemoryContainer::FRigVMMemoryContainer(bool bInUseNames)
 	: bUseNameMap(bInUseNames)
 	, MemoryType(ERigVMMemoryType::Work)
+	, bEncounteredErrorDuringLoad(false)
 {
 }
 
@@ -309,6 +311,7 @@ FRigVMMemoryContainer& FRigVMMemoryContainer::operator= (const FRigVMMemoryConta
 	Reset();
 
 	bUseNameMap = InOther.bUseNameMap;
+	bEncounteredErrorDuringLoad = false;
 	Data.Append(InOther.Data);
 	Registers.Append(InOther.Registers);
 	RegisterOffsets.Append(InOther.RegisterOffsets);
@@ -357,6 +360,8 @@ bool FRigVMMemoryContainer::Serialize(FArchive& Ar)
 
 	if (Ar.IsLoading())
 	{
+		bEncounteredErrorDuringLoad = false;
+
 		ScriptStructs.Reset();
 		TArray<FString> ScriptStructPaths;
 		Ar << ScriptStructPaths;
@@ -364,7 +369,16 @@ bool FRigVMMemoryContainer::Serialize(FArchive& Ar)
 		for (const FString& ScriptStructPath : ScriptStructPaths)
 		{
 			UScriptStruct* ScriptStruct = FindObject<UScriptStruct>(nullptr, *ScriptStructPath);
-			ensure(ScriptStruct != nullptr);
+
+			// this might have happened if a given script struct no longer 
+			// exists or cannot be loaded.
+			if (ScriptStruct == nullptr)
+			{
+				FString PackagePath = Ar.GetArchiveName();
+				UE_LOG(LogRigVM, Error, TEXT("Struct '%s' cannot be found. Asset '%s' no longer functional."), *ScriptStructPath, *PackagePath);
+				bEncounteredErrorDuringLoad = true;
+			}
+
 			ScriptStructs.Add(ScriptStruct);
 		}
 
@@ -425,7 +439,7 @@ bool FRigVMMemoryContainer::Serialize(FArchive& Ar)
 					ensure(View.Num() == Register.GetTotalElementCount());
 
 					uint8* DataPtr = &Data[Register.GetWorkByteIndex()];
-					UScriptStruct* ScriptStruct = ScriptStructs[Register.ScriptStructIndex];
+					UScriptStruct* ScriptStruct = GetScriptStruct(Register);
 					for (uint16 ElementIndex = 0; ElementIndex < Register.GetTotalElementCount(); ElementIndex++)
 					{
 						ScriptStruct->ImportText(*View[ElementIndex], DataPtr, nullptr, PPF_None, nullptr, ScriptStruct->GetName());
@@ -436,7 +450,14 @@ bool FRigVMMemoryContainer::Serialize(FArchive& Ar)
 			}
 		}
 
-		UpdateRegisters();
+		if (bEncounteredErrorDuringLoad)
+		{
+			Reset();
+		}
+		else
+		{
+			UpdateRegisters();
+		}
 	}
 	else
 	{
@@ -485,7 +506,7 @@ bool FRigVMMemoryContainer::Serialize(FArchive& Ar)
 				case ERigVMRegisterType::Struct:
 				{
 					uint8* DataPtr = &Data[Register.GetWorkByteIndex()];
-					UScriptStruct* ScriptStruct = ScriptStructs[Register.ScriptStructIndex];
+					UScriptStruct* ScriptStruct = GetScriptStruct(Register);
 
 					TArray<FString> View;
 					for (uint16 ElementIndex = 0; ElementIndex < Register.GetTotalElementCount(); ElementIndex++)
@@ -785,7 +806,14 @@ bool FRigVMMemoryContainer::Destroy(int32 InRegisterIndex, int32 InElementIndex)
 			int32 Count = InElementIndex == INDEX_NONE ? Register.GetTotalElementCount() : 1;
 
 			UScriptStruct* ScriptStruct = GetScriptStruct(InRegisterIndex);
-			ScriptStruct->DestroyStruct(DataPtr, Count);
+			if (ScriptStruct)
+			{
+				ScriptStruct->DestroyStruct(DataPtr, Count);
+			}
+			else
+			{
+				FMemory::Memzero(DataPtr, Register.GetAllocatedBytes());
+			}
 			break;
 		}
 		case ERigVMRegisterType::String:
@@ -1041,7 +1069,7 @@ void FRigVMMemoryContainer::SetRegisterValueFromString(const FRigVMOperand& InOp
 
 		if (Register.ScriptStructIndex != INDEX_NONE)
 		{
-			UScriptStruct* ScriptStruct = ScriptStructs[Register.ScriptStructIndex];
+			UScriptStruct* ScriptStruct = GetScriptStruct(Register);
 			if (ScriptStruct == InCPPTypeObject)
 			{
 				uint8* DataPtr = (uint8*)GetData(Register);
@@ -1103,7 +1131,7 @@ TArray<FString> FRigVMMemoryContainer::GetRegisterValueAsString(const FRigVMOper
 
 		if (Register.ScriptStructIndex != INDEX_NONE)
 		{
-			UScriptStruct* ScriptStruct = ScriptStructs[Register.ScriptStructIndex];
+			UScriptStruct* ScriptStruct = GetScriptStruct(Register);
 			if (ScriptStruct == InCPPTypeObject)
 			{
 				uint8* DataPtr = (uint8*)GetData(Register);
diff --git a/Engine/Source/Runtime/RigVM/Public/RigVMCore/RigVMMemory.h b/Engine/Source/Runtime/RigVM/Public/RigVMCore/RigVMMemory.h
--- a/Engine/Source/Runtime/RigVM/Public/RigVMCore/RigVMMemory.h
+++ b/Engine/Source/Runtime/RigVM/Public/RigVMCore/RigVMMemory.h
@@ -897,6 +897,9 @@ private:
 	UPROPERTY(transient)
 	TMap<FName, int32> NameMap;
 
+	UPROPERTY(transient)
+	bool bEncounteredErrorDuringLoad;
+
 	friend class URigVM;
 	friend class URigVMCompiler;
 };
