Control Rig: Fix crashes with editor open during start PIE

diff --git a/Engine/Plugins/Experimental/ControlRig/Source/ControlRigEditor/Private/Editor/ControlRigEditor.cpp b/Engine/Plugins/Experimental/ControlRig/Source/ControlRigEditor/Private/Editor/ControlRigEditor.cpp
--- a/Engine/Plugins/Experimental/ControlRig/Source/ControlRigEditor/Private/Editor/ControlRigEditor.cpp
+++ b/Engine/Plugins/Experimental/ControlRig/Source/ControlRigEditor/Private/Editor/ControlRigEditor.cpp
@@ -1905,6 +1905,17 @@ void FControlRigEditor::UpdateControlRig()
 
 		if (AnimInstance)
 		{
+			if (ControlRig)
+			{
+				// if this control rig is from a temporary step,
+				// for example the reinstancing class, clear it 
+				// and create a new one!
+				if (ControlRig->GetClass() != Class)
+				{
+					ControlRig = nullptr;
+				}
+			}
+
 			if (ControlRig == nullptr)
 			{
 				ControlRig = NewObject<UControlRig>(EditorSkelComp, Class);
