Reimplementing FProperty changes from //UE4/Main

+ Reimplemented FProperty related fixed from //UE4/Main:
CL#10791312, 10804850, 10851666, 10855122, 10855198, 10942138, 11030611, 11030639, 11032261, 11061515, 11136964,11138881, 11214238, 11214865

diff --git a/Engine/Plugins/Enterprise/DataprepEditor/Source/DataprepCore/Private/Parameterization/DataprepParameterization.cpp b/Engine/Plugins/Enterprise/DataprepEditor/Source/DataprepCore/Private/Parameterization/DataprepParameterization.cpp
--- a/Engine/Plugins/Enterprise/DataprepEditor/Source/DataprepCore/Private/Parameterization/DataprepParameterization.cpp
+++ b/Engine/Plugins/Enterprise/DataprepEditor/Source/DataprepCore/Private/Parameterization/DataprepParameterization.cpp
@@ -163,9 +163,9 @@ namespace DataprepParameterization
 			}
 
 			if ( NumberOfObject == ValueTypeValidationData.Num() )
-			{
-				ValueTypeValidationData.Add( CurrentProperty->GetClass() );
-			}
+				{
+					ValueTypeValidationData.Add( CurrentProperty->GetClass() );
+				}
 		}
 	}
 
@@ -522,6 +522,40 @@ FDataprepParameterizationBinding::FDataprepParameterizationBinding(UDataprepPara
 	}
 }
 
+bool FDataprepParameterizationBinding::Serialize(FArchive& Ar)
+{
+	check(false); // @todo FProperties: if we never hit this, we don't need this
+
+	Ar.UsingCustomVersion(FCoreObjectVersion::GUID);
+
+	UScriptStruct* Struct = FDataprepParameterizationBinding::StaticStruct();
+	if (Ar.IsLoading() || Ar.IsSaving())
+	{
+		Struct->SerializeTaggedProperties(Ar, (uint8*)this, Struct, nullptr);
+	}
+
+#if WITH_EDITORONLY_DATA
+	//Take old data and put it in new data structure
+	if (Ar.IsLoading() && Ar.CustomVer(FCoreObjectVersion::GUID) < FCoreObjectVersion::FProperties)
+	{
+		//if (ValueType_DEPRECATED)
+		//{
+		//	ValueType = FFieldClass::GetNameToFieldClassMap().FindRef(ValueType_DEPRECATED->GetFName());
+		//}
+		//else
+		//{
+		//	ValueType = nullptr;
+		//}
+	}
+	else
+#endif // WITH_EDITORONLY_DATA
+	{
+		ValueTypeValidationData.Serialize(Ar);
+	}
+
+	return true;
+}
+
 bool FDataprepParameterizationBinding::operator==(const FDataprepParameterizationBinding& Other) const
 {
 	// The value type validation data shouldn't matter when comparing binding
@@ -1158,9 +1192,9 @@ void UDataprepParameterization::OnObjectPostEdit(UDataprepParameterizableObject*
 
 			TSharedPtr<FDataprepParameterizationBinding> Binding = BindingsContainer->GetContainingBinding( BindingForModifiedProperty );
 			if ( Binding )
-		{
+			{
 				UpdateParameterizationFromBinding( Binding.ToSharedRef() );
-		}
+			}
 		}
 	}
 }
diff --git a/Engine/Plugins/Enterprise/DataprepEditor/Source/DataprepCore/Private/Parameterization/DataprepParameterization.h b/Engine/Plugins/Enterprise/DataprepEditor/Source/DataprepCore/Private/Parameterization/DataprepParameterization.h
--- a/Engine/Plugins/Enterprise/DataprepEditor/Source/DataprepCore/Private/Parameterization/DataprepParameterization.h
+++ b/Engine/Plugins/Enterprise/DataprepEditor/Source/DataprepCore/Private/Parameterization/DataprepParameterization.h
@@ -60,6 +60,14 @@ struct FValueTypeValidationData
 	}
 };
 
+//template<> struct TStructOpsTypeTraits<FValueTypeValidationData> : public TStructOpsTypeTraitsBase2<FValueTypeValidationData>
+//{
+//	enum
+//	{
+//		WithSerializer = true
+//	};
+//};
+
 /**
  * The parameterization binding is a struct that hold an object and the property path to the parameterized property
  * It also hold a array to validate the that value type of the parameterized property didn't change since it's creation
@@ -93,6 +101,16 @@ struct FDataprepParameterizationBinding
 	// Value Type Validation Array. This is the result of a depth first search on the parametrized property
 	//UPROPERTY() // @todo FProp: do we need this to be an FProp? (is this struct being serialized)
 	FValueTypeValidationData ValueTypeValidationData;
+
+	bool Serialize(FArchive& Ar);
+};
+
+template<> struct TStructOpsTypeTraits<FDataprepParameterizationBinding> : public TStructOpsTypeTraitsBase2<FDataprepParameterizationBinding>
+{
+	enum
+	{
+		WithSerializer = true
+	};
 };
 
 uint32 GetTypeHash(const FDataprepParameterizationBinding& Binding);
diff --git a/Engine/Plugins/Enterprise/VariantManager/Source/VariantManager/Private/SVariantManager.cpp b/Engine/Plugins/Enterprise/VariantManager/Source/VariantManager/Private/SVariantManager.cpp
--- a/Engine/Plugins/Enterprise/VariantManager/Source/VariantManager/Private/SVariantManager.cpp
+++ b/Engine/Plugins/Enterprise/VariantManager/Source/VariantManager/Private/SVariantManager.cpp
@@ -312,10 +312,10 @@ void SVariantManager::Construct(const FArguments& InArgs, TSharedRef<FVariantMan
 					.AutoWidth()
 					.VAlign(VAlign_Center)
 					[
-						SNew(STextBlock)
-						.TextStyle(FEditorStyle::Get(), "NormalText.Important")
-						.Text(LOCTEXT("ActorsText", "Actors"))
-					]
+					SNew(STextBlock)
+					.TextStyle(FEditorStyle::Get(), "NormalText.Important")
+					.Text(LOCTEXT("ActorsText", "Actors"))
+				]
 
 					+SHorizontalBox::Slot()
 					.VAlign(VAlign_Center)
@@ -394,10 +394,10 @@ void SVariantManager::Construct(const FArguments& InArgs, TSharedRef<FVariantMan
 						.AutoWidth()
 						.VAlign(VAlign_Center)
 						[
-							SNew(STextBlock)
-							.TextStyle(FEditorStyle::Get(), "NormalText.Important")
-							.Text(LOCTEXT("PropertiesText", "Properties"))
-						]
+						SNew(STextBlock)
+						.TextStyle(FEditorStyle::Get(), "NormalText.Important")
+						.Text(LOCTEXT("PropertiesText", "Properties"))
+					]
 
 						+SHorizontalBox::Slot()
 						.VAlign(VAlign_Center)
@@ -1438,32 +1438,32 @@ void SVariantManager::CaptureNewPropertiesFromSelectedActors()
 		return;
 	}
 
-	FVariantManagerSelection& Selection = PinnedVariantManager->GetSelection();
-	const TSet<TSharedRef<FVariantManagerActorNode>>& SelectedActorNodes = Selection.GetSelectedActorNodes();
+		FVariantManagerSelection& Selection = PinnedVariantManager->GetSelection();
+		const TSet<TSharedRef<FVariantManagerActorNode>>& SelectedActorNodes = Selection.GetSelectedActorNodes();
 
-	TArray<UVariantObjectBinding*> SelectedBindings;
-	for (const TSharedRef<FVariantManagerActorNode>& ActorNode : SelectedActorNodes)
-	{
-		UVariantObjectBinding* Binding = ActorNode->GetObjectBinding().Get();
-		if (Binding)
+		TArray<UVariantObjectBinding*> SelectedBindings;
+		for (const TSharedRef<FVariantManagerActorNode>& ActorNode : SelectedActorNodes)
 		{
-			SelectedBindings.Add(Binding);
+			UVariantObjectBinding* Binding = ActorNode->GetObjectBinding().Get();
+			if (Binding)
+			{
+				SelectedBindings.Add(Binding);
+			}
 		}
-	}
 
-	int32 NumBindings = SelectedBindings.Num();
+		int32 NumBindings = SelectedBindings.Num();
 	if (NumBindings == 0)
 	{
 		return;
 	}
 
-	FScopedTransaction Transaction(FText::Format(
-		LOCTEXT("ActorNodeCaptureNewPropertiesTransaction", "Capture new properties for {0} actor {0}|plural(one=binding,other=bindings)"),
-		NumBindings
-		));
+		FScopedTransaction Transaction(FText::Format(
+			LOCTEXT("ActorNodeCaptureNewPropertiesTransaction", "Capture new properties for {0} actor {0}|plural(one=binding,other=bindings)"),
+			NumBindings
+			));
 
-	PinnedVariantManager->CaptureNewProperties(SelectedBindings);
-	PinnedVariantManager->GetVariantManagerWidget()->RefreshPropertyList();
+		PinnedVariantManager->CaptureNewProperties(SelectedBindings);
+		PinnedVariantManager->GetVariantManagerWidget()->RefreshPropertyList();
 }
 
 bool SVariantManager::CanCaptureNewPropertiesFromSelectedActors()
@@ -2360,7 +2360,7 @@ void SVariantManager::RefreshPropertyList()
 		if (Left.DisplayOrder == Right.DisplayOrder)
 		{
 			return Left.GetName().LexicalLess(Right.GetName());
-		}
+	}
 
 		return Left.DisplayOrder < Right.DisplayOrder;
 	});
@@ -2945,7 +2945,7 @@ namespace SVariantManagerUtils
 	{
 		FString PathActorToLeafComponent;
 		FString PathTailProperty;
-		AActor* TargetActor = nullptr;
+	AActor* TargetActor = nullptr;
 		USceneComponent* LeafComponent = nullptr;
 	};
 
@@ -3008,7 +3008,7 @@ namespace SVariantManagerUtils
 				else if (UActorComponent* StartingActorComp = Cast<UActorComponent>(TargetObject))
 				{
 					if (StartingActorComp->GetOwner() == TargetActor)
-					{
+	{
 						// Actor components can't be attached to scene components, so we know we're
 						// attached directly to the actor and there are no other components in the chain
 						ComponentPath = StartingActorComp->GetName() + PATH_DELIMITER;
@@ -3016,23 +3016,23 @@ namespace SVariantManagerUtils
 					}
 				}
 			}
-		}
+	}
 
-		// We need to check if its a blueprint actor or not, as we handle blueprint root component
-		// names a little bit differently
-		bool bIsBlueprintGeneratedClass = ((UObject*)TargetActor->GetClass())->IsA(UBlueprintGeneratedClass::StaticClass());
+	// We need to check if its a blueprint actor or not, as we handle blueprint root component
+	// names a little bit differently
+	bool bIsBlueprintGeneratedClass = ((UObject*)TargetActor->GetClass())->IsA(UBlueprintGeneratedClass::StaticClass());
 
 		// Build component path up to TargetActor, if our Object is actually a SceneComponent
 		USceneComponent* ComponentPointer = LeafComponent;
 		while (ComponentPointer)
-		{
+	{
 			USceneComponent* AttachParent = ComponentPointer->GetAttachParent();
 
-			FString ComponentName;
+		FString ComponentName;
 
 			// We're some form of root component, naming is different
-			if (AttachParent == nullptr || AttachParent->GetOwner() != TargetActor)
-			{
+		if (AttachParent == nullptr || AttachParent->GetOwner() != TargetActor)
+		{
 				// Users can rename of the root component for a blueprint generated class, so lets use that.
 				// On other cases, users can't rename root components, and their actual names are always
 				// something like StaticMeshComponent0 or LightComponent0 (even if its class is a UPointLightComponent).
@@ -3074,26 +3074,26 @@ namespace SVariantManagerUtils
 
 			ComponentPath = RootComponentName;
 			PropertyPath += ChildPropertyName;
-		}
-		else
-		{
+			}
+			else
+			{
 			PropertyPath += ChildPropertyName;
-		}
+			}
 
 		Result.TargetActor = TargetActor;
 		Result.LeafComponent = LeafComponent;
 		Result.PathActorToLeafComponent = ComponentPath;
 		Result.PathTailProperty = PropertyPath;
 		return Result;
-	}
+		}
 }
 
 void SVariantManager::OnObjectPropertyChanged(UObject* Object, struct FPropertyChangedEvent& Event)
 {
 	if (!bAutoCaptureProperties || !Object || !Event.Property)
-	{
+		{
 		return;
-	}
+		}
 
 	bool bIsStructProperty = Event.MemberProperty && Event.MemberProperty->IsA(FStructProperty::StaticClass());
 	bool bIsBuiltIn = bIsStructProperty && FVariantManagerUtils::IsBuiltInStructProperty(Event.MemberProperty);
@@ -3133,9 +3133,9 @@ void SVariantManager::OnObjectPropertyChanged(UObject* Object, struct FPropertyC
 	}
 
 	if (PropertyPath.IsEmpty())
-	{
+		{
 		return;
-	}
+		}
 
 	bool bUpdatedSomething = AutoCaptureProperty(VariantManagerPtr.Pin().Get(), TargetActor, PropertyPath, Prop);
 	if (bUpdatedSomething)
@@ -3154,7 +3154,7 @@ namespace SVariantManagerUtils
 		for (PropNode* Node = PropChain.GetHead(); Node; Node = Node->GetNextNode())
 		{
 			Hash = HashCombine(Hash, GetTypeHash(Node->GetValue()));
-	}
+		}
 
 		Hash = HashCombine(Hash, GetTypeHash(Object));
 
@@ -3169,7 +3169,7 @@ void SVariantManager::OnPreObjectPropertyChanged(UObject* Object, const class FE
 	// but in the first firing the root component is Object, and in the second firing, the root component is
 
 	if (!bAutoCaptureProperties || !Object)
-	{
+		{
 		return;
 	}
 
@@ -3251,16 +3251,16 @@ void SVariantManager::OnPreObjectPropertyChanged(UObject* Object, const class FE
 
 				if (!IntermediatePath.EndsWith(PATH_DELIMITER + Category + PATH_DELIMITER) &&
 					LastCategorySegment != PropertyName)
-				{
+	{
 					IntermediatePath += Category + PATH_DELIMITER;
 				}
 			}
-		}
+	}
 
 		if (Node == Tail)
-		{
+	{
 			break;
-		}
+	}
 
 		IntermediatePath += PropertyName + PATH_DELIMITER;
 
@@ -3279,7 +3279,7 @@ void SVariantManager::OnPreObjectPropertyChanged(UObject* Object, const class FE
 			}
 		}
 		else
-		{
+	{
 			bShowCategories = false;
 		}
 	}
diff --git a/Engine/Plugins/Enterprise/VariantManagerContent/Source/VariantManagerContent/Private/PropertyValue.cpp b/Engine/Plugins/Enterprise/VariantManagerContent/Source/VariantManagerContent/Private/PropertyValue.cpp
--- a/Engine/Plugins/Enterprise/VariantManagerContent/Source/VariantManagerContent/Private/PropertyValue.cpp
+++ b/Engine/Plugins/Enterprise/VariantManagerContent/Source/VariantManagerContent/Private/PropertyValue.cpp
@@ -22,6 +22,9 @@
 #include "EdGraphSchema_K2.h"
 #endif
 
+// WARNING: This should always be the last include in any file that needs it (except .generated.h)
+#include "UObject/UndefineUPropertyMacros.h"
+
 #define LOCTEXT_NAMESPACE "PropertyValue"
 
 DEFINE_LOG_CATEGORY(LogVariantContent);
@@ -618,7 +621,7 @@ void UPropertyValue::ApplyDataToResolvedObject()
 	// Modify container component
 	UObject* ContainerObject = nullptr;
 	if (ParentContainerObject && ParentContainerObject->IsA(UActorComponent::StaticClass()))
-	{
+		{
 		ParentContainerObject->SetFlags(RF_Transactional);
 		ParentContainerObject->Modify();
 		ContainerObject = ParentContainerObject;
@@ -1161,8 +1164,8 @@ const TArray<uint8>& UPropertyValue::GetDefaultValue()
 							}
 							else
 							{
-								FMemory::Memcpy(DefaultValue.GetData(), PropertyValuePtr, NumBytes);
-							}
+							FMemory::Memcpy(DefaultValue.GetData(), PropertyValuePtr, NumBytes);
+						}
 
 						}
 
@@ -1729,3 +1732,5 @@ void UPropertyValueVisibility::Serialize(FArchive& Ar)
 }
 
 #undef LOCTEXT_NAMESPACE
+
+#include "UObject/DefineUPropertyMacros.h"
\ No newline at end of file
diff --git a/Engine/Plugins/Enterprise/VariantManagerContent/Source/VariantManagerContent/Public/PropertyValue.h b/Engine/Plugins/Enterprise/VariantManagerContent/Source/VariantManagerContent/Public/PropertyValue.h
--- a/Engine/Plugins/Enterprise/VariantManagerContent/Source/VariantManagerContent/Public/PropertyValue.h
+++ b/Engine/Plugins/Enterprise/VariantManagerContent/Source/VariantManagerContent/Public/PropertyValue.h
@@ -6,7 +6,8 @@
 #include "UObject/ObjectMacros.h"
 #include "UObject/Field.h"
 
-#include "PropertyValue.generated.h"
+// WARNING: This should always be the last include in any file that needs it
+#include "UObject/UndefineUPropertyMacros.h"
 
 #define PATH_DELIMITER TEXT(" / ")
 #define ATTACH_CHILDREN_NAME TEXT("Children")
@@ -16,7 +17,7 @@ DECLARE_MULTICAST_DELEGATE(FOnPropertyApplied);
 
 class UVariantObjectBinding;
 class USCS_Node;
-class FProperty;
+class UProperty;
 
 UENUM()
 enum class EPropertyValueCategory : uint8
@@ -247,4 +248,5 @@ class VARIANTMANAGERCONTENT_API UPropertyValueVisibility : public UPropertyValue
 	// UObject Interface
 	virtual void Serialize(FArchive& Ar) override;
 	//~ End UObject Interface
-};
\ No newline at end of file
+};
+#include "UObject/DefineUPropertyMacros.h"
\ No newline at end of file
diff --git a/Engine/Plugins/Experimental/ControlRig/Source/ControlRig/Private/ControlRigBlueprintGeneratedClass.cpp b/Engine/Plugins/Experimental/ControlRig/Source/ControlRig/Private/ControlRigBlueprintGeneratedClass.cpp
--- a/Engine/Plugins/Experimental/ControlRig/Source/ControlRig/Private/ControlRigBlueprintGeneratedClass.cpp
+++ b/Engine/Plugins/Experimental/ControlRig/Source/ControlRig/Private/ControlRigBlueprintGeneratedClass.cpp
@@ -36,9 +36,9 @@ void UControlRigBlueprintGeneratedClass::Serialize(FArchive& Ar)
 	URigVM* VM = NewObject<URigVM>(GetTransientPackage());
 
 	if (UControlRig* CDO = Cast<UControlRig>(GetDefaultObject(true)))
-	{
-		if (Ar.IsSaving() && CDO->VM)
 		{
+		if (Ar.IsSaving() && CDO->VM)
+			{
 			VM->CopyFrom(CDO->VM);
 		}
 	}
@@ -48,11 +48,11 @@ void UControlRigBlueprintGeneratedClass::Serialize(FArchive& Ar)
 	if (UControlRig* CDO = Cast<UControlRig>(GetDefaultObject(false)))
 	{
 		if (Ar.IsLoading())
-		{
+				{
 			if (CDO->VM == nullptr)
 			{
 				CDO->VM = NewObject<URigVM>(CDO, TEXT("VM"));
-			}
+				}
 			if (CDO->VM->GetOuter() != CDO)
 			{
 				CDO->VM = NewObject<URigVM>(CDO, TEXT("VM"));
@@ -60,4 +60,4 @@ void UControlRigBlueprintGeneratedClass::Serialize(FArchive& Ar)
 			CDO->VM->CopyFrom(VM);
 		}
 	}
-}
\ No newline at end of file
+}
diff --git a/Engine/Plugins/Experimental/ControlRig/Source/ControlRigDeveloper/Private/Graph/ControlRigGraph.cpp b/Engine/Plugins/Experimental/ControlRig/Source/ControlRigDeveloper/Private/Graph/ControlRigGraph.cpp
--- a/Engine/Plugins/Experimental/ControlRig/Source/ControlRigDeveloper/Private/Graph/ControlRigGraph.cpp
+++ b/Engine/Plugins/Experimental/ControlRig/Source/ControlRigDeveloper/Private/Graph/ControlRigGraph.cpp
@@ -96,70 +96,70 @@ void UControlRigGraph::HandleModifiedEvent(ERigVMGraphNotifType InNotifType, URi
 	if (bSuspendModelNotifications)
 	{
 		return;
-	}
+			}
 
 	switch (InNotifType)
-	{
-		case ERigVMGraphNotifType::GraphChanged:
 		{
-			for (URigVMNode* Node : InGraph->GetNodes())
+		case ERigVMGraphNotifType::GraphChanged:
 			{
+			for (URigVMNode* Node : InGraph->GetNodes())
+				{
 				UEdGraphNode* EdNode = FindNodeForModelNodeName(Node->GetFName());
 				if (EdNode != nullptr)
-				{
+					{
 					RemoveNode(EdNode);
 				}
 			}
 			break;
 		}
 		case ERigVMGraphNotifType::NodeSelectionChanged:
-		{
+						{
 			if (bIsSelecting)
-			{
+							{
 				return;
-			}
+							}
 			TGuardValue<bool> SelectionGuard(bIsSelecting, true);
 
 			TSet<const UEdGraphNode*> NodeSelection;
 			for (FName NodeName : InGraph->GetSelectNodes())
-			{
+							{
 				if (UEdGraphNode* EdNode = FindNodeForModelNodeName(NodeName))
-				{
+							{
 					NodeSelection.Add(EdNode);
-				}
-			}
+							}
+						}
 			SelectNodeSet(NodeSelection);
 			break;
-		}
+					}
 		case ERigVMGraphNotifType::NodeAdded:
-		{
-			if (URigVMNode* ModelNode = Cast<URigVMNode>(InSubject))
 			{
+			if (URigVMNode* ModelNode = Cast<URigVMNode>(InSubject))
+				{
 				if (!ModelNode->IsVisibleInUI())
 				{
 					if (URigVMInjectionInfo* Injection = ModelNode->GetInjectionInfo())
-					{
+				{
 						if (URigVMPin* ModelPin = Injection->GetPin())
-						{
+				{
 							URigVMNode* ParentModelNode = ModelPin->GetNode();
 							if (ParentModelNode)
-							{
+					{
 								UEdGraphNode* EdNode = FindNodeForModelNodeName(ParentModelNode->GetFName());
 								if (EdNode)
-								{
+					{
 									if (UControlRigGraphNode* RigNode = Cast<UControlRigGraphNode>(EdNode))
-									{
+					{
 										RigNode->ReconstructNode_Internal(true);
-									}
-								}
-							}
-						}
 					}
-					break;
 				}
+			}
+		}
+			}
+					break;
+		}
 
 				if (URigVMCommentNode* CommentModelNode = Cast<URigVMCommentNode>(ModelNode))
-				{
+		{
 					UEdGraphNode_Comment* NewNode = NewObject<UEdGraphNode_Comment>(this, CommentModelNode->GetFName());
 					AddNode(NewNode, false);
 
@@ -175,7 +175,7 @@ void UControlRigGraph::HandleModifiedEvent(ERigVMGraphNotifType InNotifType, URi
 					NewNode->NodeComment = CommentModelNode->GetCommentText();
 					NewNode->SetFlags(RF_Transactional);
 					NewNode->GetNodesUnderComment();
-				}
+	}
 				else if (URigVMRerouteNode* RerouteModelNode = Cast<URigVMRerouteNode>(ModelNode))
 				{
 					UControlRigGraphNode* NewNode = NewObject<UControlRigGraphNode>(this, ModelNode->GetFName());
@@ -193,10 +193,10 @@ void UControlRigGraph::HandleModifiedEvent(ERigVMGraphNotifType InNotifType, URi
 					NewNode->AllocateDefaultPins();
 
 					if (UEdGraphPin* ValuePin = NewNode->FindPin(ModelNode->FindPin("Value")->GetPinPath()))
-					{
+	{
 						NewNode->SetColorFromModel(GetSchema()->GetPinTypeColor(ValuePin->PinType));
 					}
-				}
+	}
 				else // struct, parameter + variable
 				{
 					UControlRigGraphNode* NewNode = NewObject<UControlRigGraphNode>(this, ModelNode->GetFName());
@@ -226,36 +226,36 @@ void UControlRigGraph::HandleModifiedEvent(ERigVMGraphNotifType InNotifType, URi
 					{
 						URigVMNode* ParentModelNode = ModelPin->GetNode();
 						if (ParentModelNode)
-						{
+					{
 							UEdGraphNode* EdNode = FindNodeForModelNodeName(ParentModelNode->GetFName());
 							if (EdNode)
+						{
+							if (UControlRigGraphNode* RigNode = Cast<UControlRigGraphNode>(EdNode))
 							{
-								if (UControlRigGraphNode* RigNode = Cast<UControlRigGraphNode>(EdNode))
-								{
 									RigNode->ReconstructNode_Internal(true);
 								}
 							}
+							}
 						}
+						break;
 					}
-					break;
-				}
 
 				UEdGraphNode* EdNode = FindNodeForModelNodeName(ModelNode->GetFName());
 				if (EdNode)
-				{
+					{
 					RemoveNode(EdNode, true);
 					NotifyGraphChanged();
 				}
 			}
-			break;
-		}
+						break;
+					}
 		case ERigVMGraphNotifType::NodePositionChanged:
 		{
 			if (URigVMNode* ModelNode = Cast<URigVMNode>(InSubject))
 			{
 				UEdGraphNode* EdNode = FindNodeForModelNodeName(ModelNode->GetFName());
 				if (EdNode)
-				{
+					{
 					EdNode->NodePosX = (int32)ModelNode->GetPosition().X;
 					EdNode->NodePosY = (int32)ModelNode->GetPosition().Y;
 				}
@@ -338,23 +338,23 @@ void UControlRigGraph::HandleModifiedEvent(ERigVMGraphNotifType InNotifType, URi
 					UControlRigGraphNode* SourceRigNode = Cast<UControlRigGraphNode>(FindNodeForModelNodeName(SourcePin->GetNode()->GetFName()));
 					UControlRigGraphNode* TargetRigNode = Cast<UControlRigGraphNode>(FindNodeForModelNodeName(TargetPin->GetNode()->GetFName()));
 
-					if (SourceRigNode != nullptr && TargetRigNode != nullptr)
-					{
+				if (SourceRigNode != nullptr && TargetRigNode != nullptr)
+				{
 						FString SourcePinPath = SourcePin->GetPinPath();
 						FString TargetPinPath = TargetPin->GetPinPath();
-						UEdGraphPin* SourceRigPin = SourceRigNode->FindPin(*SourcePinPath, EGPD_Output);
-						UEdGraphPin* TargetRigPin = TargetRigNode->FindPin(*TargetPinPath, EGPD_Input);
+					UEdGraphPin* SourceRigPin = SourceRigNode->FindPin(*SourcePinPath, EGPD_Output);
+					UEdGraphPin* TargetRigPin = TargetRigNode->FindPin(*TargetPinPath, EGPD_Input);
 
-						if (SourceRigPin != nullptr && TargetRigPin != nullptr)
+					if (SourceRigPin != nullptr && TargetRigPin != nullptr)
+					{
+						if (AddLink)
 						{
-							if (AddLink)
-							{
-								SourceRigPin->MakeLinkTo(TargetRigPin);
-							}
-							else
-							{
-								SourceRigPin->BreakLinkTo(TargetRigPin);
-							}
+							SourceRigPin->MakeLinkTo(TargetRigPin);
+						}
+						else
+						{
+							SourceRigPin->BreakLinkTo(TargetRigPin);
+						}
 
 							for (int32 LinkedPinIndex = 0; LinkedPinIndex < SourceRigPin->LinkedTo.Num();)
 							{
@@ -365,8 +365,8 @@ void UControlRigGraph::HandleModifiedEvent(ERigVMGraphNotifType InNotifType, URi
 								else
 								{
 									SourceRigPin->LinkedTo.RemoveAtSwap(LinkedPinIndex);
-								}
-							}
+					}
+				}
 
 							for (int32 LinkedPinIndex = 0; LinkedPinIndex < TargetRigPin->LinkedTo.Num();)
 							{
@@ -404,16 +404,16 @@ void UControlRigGraph::HandleModifiedEvent(ERigVMGraphNotifType InNotifType, URi
 						}
 					}
 					if (Cast<URigVMParameterNode>(ModelPin->GetNode()))
-					{
+							{
 						if (ModelPin->GetName() == TEXT("Parameter"))
-						{
+								{
 							RigNode->InvalidateNodeTitle();
 							RigNode->ReconstructNode_Internal(true);
+								}
+							}
 						}
-					}
-				}
 				else if (URigVMInjectionInfo* Injection = ModelPin->GetNode()->GetInjectionInfo())
-				{
+						{
 					if (Injection->InputPin != ModelPin->GetRootPin())
 					{
 						if (URigVMPin* InjectionPin = Injection->GetPin())
@@ -429,13 +429,13 @@ void UControlRigGraph::HandleModifiedEvent(ERigVMGraphNotifType InNotifType, URi
 										HostRigNode->ReconstructNode_Internal(true);
 									}
 								}
+							}
+									}
+								}
 							}
 						}
-					}
-				}
-			}
 			break;
-		}
+					}
 		case ERigVMGraphNotifType::PinArraySizeChanged:
 		case ERigVMGraphNotifType::PinDirectionChanged:
 		case ERigVMGraphNotifType::PinTypeChanged:
diff --git a/Engine/Plugins/Experimental/ControlRig/Source/ControlRigDeveloper/Private/Graph/ControlRigGraphNode.cpp b/Engine/Plugins/Experimental/ControlRig/Source/ControlRigDeveloper/Private/Graph/ControlRigGraphNode.cpp
--- a/Engine/Plugins/Experimental/ControlRig/Source/ControlRigDeveloper/Private/Graph/ControlRigGraphNode.cpp
+++ b/Engine/Plugins/Experimental/ControlRig/Source/ControlRigDeveloper/Private/Graph/ControlRigGraphNode.cpp
@@ -50,7 +50,7 @@ FText UControlRigGraphNode::GetNodeTitle(ENodeTitleType::Type TitleType) const
 		}
 	}
 
-	return NodeTitle;
+		return NodeTitle;
 }
 
 void UControlRigGraphNode::ReconstructNode()
@@ -129,16 +129,16 @@ FEdGraphNodeDeprecationResponse UControlRigGraphNode::GetDeprecationResponse(EEd
 		if(URigVMStructNode* StructModelNode = Cast<URigVMStructNode>(ModelNode))
 		{
 			if (UScriptStruct* ScriptStruct = StructModelNode->GetScriptStruct())
-			{
-				FString DeprecatedMetadata;
-				ScriptStruct->GetStringMetaDataHierarchical(UControlRig::DeprecatedMetaName, &DeprecatedMetadata);
-				if (!DeprecatedMetadata.IsEmpty())
-				{
-					FFormatNamedArguments Args;
-					Args.Add(TEXT("DeprecatedMetadata"), FText::FromString(DeprecatedMetadata));
-					Response.MessageText = FText::Format(LOCTEXT("ControlRigGraphNodeDeprecationMessage", "Warning: This node is deprecated from: {DeprecatedMetadata}"), Args);
-				}
-			}
+	{
+		FString DeprecatedMetadata;
+		ScriptStruct->GetStringMetaDataHierarchical(UControlRig::DeprecatedMetaName, &DeprecatedMetadata);
+		if (!DeprecatedMetadata.IsEmpty())
+		{
+			FFormatNamedArguments Args;
+			Args.Add(TEXT("DeprecatedMetadata"), FText::FromString(DeprecatedMetadata));
+			Response.MessageText = FText::Format(LOCTEXT("ControlRigGraphNodeDeprecationMessage", "Warning: This node is deprecated from: {DeprecatedMetadata}"), Args);
+		}
+	}
 		}
 	}
 
@@ -259,9 +259,9 @@ void UControlRigGraphNode::HandleInsertArrayElement(FString InPinPath)
 	if (Blueprint)
 	{
 		if (URigVMPin* ArrayElementPin = GetModelPinFromPinPath(InPinPath))
-		{
-			if (URigVMPin* ArrayPin = ArrayElementPin->GetParentPin())
 			{
+			if (URigVMPin* ArrayPin = ArrayElementPin->GetParentPin())
+				{
 				Blueprint->Controller->OpenUndoBracket(TEXT("Add Array Pin"));
 				FString PinPath = Blueprint->Controller->InsertArrayPin(InPinPath, ArrayElementPin->GetPinIndex() + 1, FString());
 				Blueprint->Controller->SetPinExpansion(InPinPath, true);
@@ -325,18 +325,18 @@ void UControlRigGraphNode::CreateExecutionPins()
 		{
 			Pair.InputPin = CreatePin(EGPD_Input, GetPinTypeForModelPin(ModelPin), FName(*ModelPin->GetPinPath()));
 			if (Pair.InputPin != nullptr)
-			{
+		{
 				Pair.InputPin->PinFriendlyName = FText::FromName(ModelPin->GetDisplayName());
-			}
 		}
+	}
 		if (Pair.OutputPin == nullptr)
-		{
+	{
 			Pair.OutputPin = CreatePin(EGPD_Output, GetPinTypeForModelPin(ModelPin), FName(*ModelPin->GetPinPath()));
 			if (Pair.OutputPin != nullptr)
-			{
+		{
 				Pair.OutputPin->PinFriendlyName = FText::FromName(ModelPin->GetDisplayName());
-			}
 		}
+	}
 		// note: no recursion for execution pins
 	}
 }
@@ -360,12 +360,12 @@ void UControlRigGraphNode::CreateInputPins(URigVMPin* InParentPin)
 				SetupPinDefaultsFromModel(Pair.InputPin);
 
 				if (InParentPin != nullptr)
-				{
+	{
 					PinPair& ParentPair = CachedPins.FindChecked(InParentPin);
 					ParentPair.InputPin->SubPins.Add(Pair.InputPin);
 					Pair.InputPin->ParentPin = ParentPair.InputPin;
 				}
-			}
+		}
 		}
 		CreateInputPins(ModelPin);
 	}
@@ -383,7 +383,7 @@ void UControlRigGraphNode::CreateInputOutputPins(URigVMPin* InParentPin, bool bH
 		{
 			Pair.InputPin = CreatePin(EGPD_Input, GetPinTypeForModelPin(ModelPin), FName(*ModelPin->GetPinPath()));
 			if (Pair.InputPin != nullptr)
-			{
+		{
 				Pair.InputPin->bHidden = bHidden;
 				Pair.InputPin->PinFriendlyName = FText::FromName(ModelPin->GetDisplayName());
 				Pair.InputPin->bNotConnectable = ModelPin->GetDirection() != ERigVMPinDirection::IO;
@@ -391,39 +391,39 @@ void UControlRigGraphNode::CreateInputOutputPins(URigVMPin* InParentPin, bool bH
 				SetupPinDefaultsFromModel(Pair.InputPin);
 
 				if (InParentPin != nullptr)
-				{
+		{
 					PinPair& ParentPair = CachedPins.FindChecked(InParentPin);
 					ParentPair.InputPin->SubPins.Add(Pair.InputPin);
 					Pair.InputPin->ParentPin = ParentPair.InputPin;
-				}
-			}
 		}
+	}
+	}
 		if (Pair.OutputPin == nullptr)
-		{
+	{
 			Pair.OutputPin = CreatePin(EGPD_Output, GetPinTypeForModelPin(ModelPin), FName(*ModelPin->GetPinPath()));
 			if (Pair.OutputPin != nullptr)
-			{
+		{
 				Pair.OutputPin->bHidden = bHidden;
 				Pair.OutputPin->PinFriendlyName = FText::FromName(ModelPin->GetDisplayName());
 				Pair.OutputPin->bNotConnectable = ModelPin->GetDirection() != ERigVMPinDirection::IO;
 
 				if (InParentPin != nullptr)
-				{
+		{
 					PinPair& ParentPair = CachedPins.FindChecked(InParentPin);
 					ParentPair.OutputPin->SubPins.Add(Pair.OutputPin);
 					Pair.OutputPin->ParentPin = ParentPair.OutputPin;
 				}
-			}
 		}
+	}
 
 		// don't recurse on knot / compact reroute nodes
 		if(URigVMRerouteNode* RerouteNode = Cast<URigVMRerouteNode>(GetModelNode()))
-		{
+	{
 			if (!RerouteNode->GetShowsAsFullNode())
-			{
+		{
 				bHidden = true;
-			}
 		}
+	}
 
 		CreateInputOutputPins(ModelPin, bHidden);
 	}
@@ -446,7 +446,7 @@ void UControlRigGraphNode::CreateOutputPins(URigVMPin* InParentPin)
 				Pair.OutputPin->bNotConnectable = ModelPin->GetDirection() != ERigVMPinDirection::Output;
 
 				if (InParentPin != nullptr)
-				{
+	{
 					PinPair& ParentPair = CachedPins.FindChecked(InParentPin);
 					ParentPair.OutputPin->SubPins.Add(Pair.OutputPin);
 					Pair.OutputPin->ParentPin = ParentPair.OutputPin;
@@ -545,14 +545,14 @@ void UControlRigGraphNode::DestroyNode()
 	{
 		UControlRigBlueprint* ControlRigBlueprint = Cast<UControlRigBlueprint>(Graph->GetOuter());
 		if(ControlRigBlueprint)
-		{
+	{
 			BreakAllNodeLinks();
 			if(PropertyName_DEPRECATED.IsValid())
-			{
+		{
 				FControlRigBlueprintUtils::RemoveMemberVariableIfNotUsed(ControlRigBlueprint, PropertyName_DEPRECATED, this);
 			}
 		}
-	}
+		}
 
 	UEdGraphNode::DestroyNode();
 }
@@ -569,7 +569,7 @@ void UControlRigGraphNode::CopyPinDefaultsToModel(UEdGraphPin* Pin, bool bUndo)
 	if (Pin->Direction != EGPD_Input)
 	{
 		return;
-	}
+		}
 
 	const UEdGraphSchema_K2* K2Schema = GetDefault<UEdGraphSchema_K2>();
 	if (URigVMPin* ModelPin = GetModelPinFromPinPath(Pin->GetName()))
@@ -577,16 +577,16 @@ void UControlRigGraphNode::CopyPinDefaultsToModel(UEdGraphPin* Pin, bool bUndo)
 		if (ModelPin->GetSubPins().Num() > 0)
 		{
 			return;
-		}
+	}
 
 		FString DefaultValue = Pin->DefaultValue;
 		if (DefaultValue == FName(NAME_None).ToString() && Pin->PinType.PinSubCategory == UEdGraphSchema_K2::PC_Name)
 		{
 			DefaultValue = FString();
-		}
+	}
 
 		if (ModelPin->GetDefaultValue() != DefaultValue)
-		{
+				{
 			GetBlueprint()->Controller->SetPinDefaultValue(ModelPin->GetPinPath(), DefaultValue, false, true, false);
 		}
 	}
@@ -650,10 +650,10 @@ URigVMPin* UControlRigGraphNode::GetModelPinFromPinPath(const FString& InPinPath
 		FString PinPath = InPinPath.RightChop(ModelNode->GetNodePath().Len() + 1);
 		URigVMPin* ModelPin = ModelNode->FindPin(PinPath);
 		if (ModelPin)
-		{
+	{
 			UControlRigGraphNode* MutableThis = (UControlRigGraphNode*)this;
 			MutableThis->CachedModelPins.Add(InPinPath, ModelPin);
-		}
+	}
 		return ModelPin;
 	}
 	
@@ -667,7 +667,7 @@ void UControlRigGraphNode::SetupPinDefaultsFromModel(UEdGraphPin* Pin)
 	if (Pin->Direction != EGPD_Input)
 	{
 		return;
-	}
+		}
 
 	const UEdGraphSchema_K2* K2Schema = GetDefault<UEdGraphSchema_K2>();
 	if (URigVMPin* ModelPin = GetModelPinFromPinPath(Pin->GetName()))
@@ -675,15 +675,15 @@ void UControlRigGraphNode::SetupPinDefaultsFromModel(UEdGraphPin* Pin)
 		if (ModelPin->GetSubPins().Num() > 0)
 		{
 			return;
-		}
+			}
 
 		FString DefaultValueString = ModelPin->GetDefaultValue();
 		if (DefaultValueString.IsEmpty() && ModelPin->GetCPPType() == TEXT("FName"))
-		{
+	{
 			DefaultValueString = FName(NAME_None).ToString();
-		}
-		K2Schema->GetPinDefaultValuesFromString(Pin->PinType, Pin->GetOwningNodeUnchecked(), DefaultValueString, Pin->DefaultValue, Pin->DefaultObject, Pin->DefaultTextValue);
 	}
+						K2Schema->GetPinDefaultValuesFromString(Pin->PinType, Pin->GetOwningNodeUnchecked(), DefaultValueString, Pin->DefaultValue, Pin->DefaultObject, Pin->DefaultTextValue);
+					}
 }
 
 FText UControlRigGraphNode::GetTooltipText() const
@@ -726,9 +726,9 @@ void UControlRigGraphNode::AutowireNewNode(UEdGraphPin* FromPin)
 			if (Schema->TryCreateConnection(FromPin, Pin))
 			{
 				break;
-			}
+					}
+				}
 		}
-	}
 }
 
 bool UControlRigGraphNode::IsSelectedInEditor() const
@@ -746,14 +746,14 @@ bool UControlRigGraphNode::ShouldDrawNodeAsControlPointOnly(int32& OutInputPinIn
 	if (URigVMRerouteNode* Reroute = Cast<URigVMRerouteNode>(GetModelNode()))
 	{
 		if (!Reroute->GetShowsAsFullNode())
-		{
+	{
 			if (Pins.Num() >= 2)
 			{
 				OutInputPinIndex = 0;
 				OutOutputPinIndex = 1;
 				return true;
 			}
-		}
+	}
 	}
 	return false;
 }
@@ -790,15 +790,15 @@ FEdGraphPinType UControlRigGraphNode::GetPinTypeForModelPin(URigVMPin* InModelPi
 		PinType.PinSubCategoryObject = InModelPin->GetScriptStruct();
 	}
 	else if (InModelPin->GetEnum() != nullptr)
-	{
+		{
 		PinType.PinCategory = UEdGraphSchema_K2::PC_Byte;
 		PinType.PinSubCategoryObject = InModelPin->GetEnum();
-	}
+		}
 
 	if (InModelPin->IsArray())
-	{
+		{
 		PinType.ContainerType = EPinContainerType::Array;
-	}
+		}
 	else
 	{
 		PinType.ContainerType = EPinContainerType::None;
diff --git a/Engine/Plugins/Experimental/ControlRig/Source/ControlRigEditor/Private/EditMode/ControlRigEditMode.cpp b/Engine/Plugins/Experimental/ControlRig/Source/ControlRigEditor/Private/EditMode/ControlRigEditMode.cpp
--- a/Engine/Plugins/Experimental/ControlRig/Source/ControlRigEditor/Private/EditMode/ControlRigEditMode.cpp
+++ b/Engine/Plugins/Experimental/ControlRig/Source/ControlRigEditor/Private/EditMode/ControlRigEditMode.cpp
@@ -122,16 +122,16 @@ void FControlRigEditMode::SetUpDetailPanel()
 {
 	if (IsInLevelEditor())
 	{
-		TArray<TWeakObjectPtr<>> SelectedObjects;
+	TArray<TWeakObjectPtr<>> SelectedObjects;
 		if (WeakControlRigEditing.IsValid())
 		{
 			const TArray<UControlRigControlsProxy*>& Proxies = ControlProxy->GetSelectedProxies();
 			for (UControlRigControlsProxy* Proxy : Proxies)
-			{
+	{
 				SelectedObjects.Add(Proxy);
 			}
-			SelectedObjects.Add(Settings);
-		}
+		SelectedObjects.Add(Settings);
+	}
 		StaticCastSharedPtr<SControlRigEditModeTools>(Toolkit->GetInlineContent())->SetSequencer(WeakSequencer.Pin());
 		StaticCastSharedPtr<SControlRigEditModeTools>(Toolkit->GetInlineContent())->SetDetailsObjects(SelectedObjects);
 	}
@@ -270,8 +270,8 @@ void FControlRigEditMode::Tick(FEditorViewportClient* ViewportClient, float Delt
 						{
 							ControlProxy->AddProxy(SelectedKey.Name, WeakControlRigEditing.Get(), Control);
 						}
-					}
-				}
+			}
+		}
 			}
 		}
 		SetUpDetailPanel();
@@ -394,33 +394,33 @@ void FControlRigEditMode::Render(const FSceneView* View, FViewport* Viewport, FP
 			switch (Instruction.PrimitiveType)
 			{
 				case EControlRigDrawSettings::Points:
+			{
+				for (const FVector& Point : Instruction.Positions)
 				{
-					for (const FVector& Point : Instruction.Positions)
-					{
 						PDI->DrawPoint(InstructionTransform.TransformPosition(Point), Instruction.Color, Instruction.Thickness, SDPG_Foreground);
-					}
-					break;
 				}
+				break;
+			}
 				case EControlRigDrawSettings::Lines:
-				{
-					const TArray<FVector>& Points = Instruction.Positions;
+			{
+				const TArray<FVector>& Points = Instruction.Positions;
 					PDI->AddReserveLines(SDPG_Foreground, Points.Num() / 2, false, Instruction.Thickness > SMALL_NUMBER);
-					for (int32 PointIndex = 0; PointIndex < Points.Num() - 1; PointIndex += 2)
-					{
+				for (int32 PointIndex = 0; PointIndex < Points.Num() - 1; PointIndex += 2)
+				{
 						PDI->DrawLine(InstructionTransform.TransformPosition(Points[PointIndex]), InstructionTransform.TransformPosition(Points[PointIndex + 1]), Instruction.Color, SDPG_Foreground, Instruction.Thickness);
-					}
-					break;
 				}
+				break;
+			}
 				case EControlRigDrawSettings::LineStrip:
-				{
-					const TArray<FVector>& Points = Instruction.Positions;
+			{
+				const TArray<FVector>& Points = Instruction.Positions;
 					PDI->AddReserveLines(SDPG_Foreground, Points.Num() - 1, false, Instruction.Thickness > SMALL_NUMBER);
-					for (int32 PointIndex = 0; PointIndex < Points.Num() - 1; PointIndex++)
-					{
+				for (int32 PointIndex = 0; PointIndex < Points.Num() - 1; PointIndex++)
+				{
 						PDI->DrawLine(InstructionTransform.TransformPosition(Points[PointIndex]), InstructionTransform.TransformPosition(Points[PointIndex + 1]), Instruction.Color, SDPG_Foreground, Instruction.Thickness);
-					}
-					break;
 				}
+				break;
+			}
 			}
 		}
 	}
@@ -464,7 +464,7 @@ bool FControlRigEditMode::EndTracking(FEditorViewportClient* InViewportClient, F
 			bManipulatorMadeChange = false;
 			GEditor->EndTransaction();
 
-		}
+				}
 		bIsTransacting = false;
 		return true;
 	}
@@ -483,10 +483,10 @@ bool FControlRigEditMode::StartTracking(FEditorViewportClient* InViewportClient,
 		{
 			UObject* Blueprint = WeakControlRigEditing->GetClass()->ClassGeneratedBy;
 			if (Blueprint)
-			{
+		{
 				Blueprint->SetFlags(RF_Transactional);
 				Blueprint->Modify();
-			}
+		}
 
 			WeakControlRigEditing->SetFlags(RF_Transactional);
 			WeakControlRigEditing->Modify();
@@ -1247,12 +1247,12 @@ bool FControlRigEditMode::GizmoSelectionOverride(const UPrimitiveComponent* InCo
     //Think we only want to do this in regular editor, in the level editor we are driving selection
 	if (!IsInLevelEditor())
 	{
-		AControlRigGizmoActor* OwnerActor = Cast<AControlRigGizmoActor>(InComponent->GetOwner());
-		if (OwnerActor)
-		{
-			// See if the actor is in a selected unit proxy
-			return OwnerActor->IsSelected();
-		}
+	AControlRigGizmoActor* OwnerActor = Cast<AControlRigGizmoActor>(InComponent->GetOwner());
+	if (OwnerActor)
+	{
+		// See if the actor is in a selected unit proxy
+		return OwnerActor->IsSelected();
+	}
 	}
 
 	return false;
@@ -1289,12 +1289,12 @@ bool FControlRigEditMode::AreRigElementSelectedAndMovable() const
 	if (!IsInLevelEditor())
 	{
 		if (!IsTransformDelegateAvailable())
-		{
+	{
 			return false;
 		}
 	}
 
-	return true;
+		return true;
 }
 
 void FControlRigEditMode::OnRigElementAdded(FRigHierarchyContainer* Container, const FRigElementKey& InKey)
@@ -1359,7 +1359,7 @@ void FControlRigEditMode::OnRigElementSelected(FRigHierarchyContainer* Container
 						{
 							ControlProxy->SelectProxy(InKey.Name, true);
 						}
-					}
+			}
 				}
 				else
 				{
diff --git a/Engine/Plugins/FX/Niagara/Source/NiagaraEditor/Private/NiagaraEditorUtilities.cpp b/Engine/Plugins/FX/Niagara/Source/NiagaraEditor/Private/NiagaraEditorUtilities.cpp
--- a/Engine/Plugins/FX/Niagara/Source/NiagaraEditor/Private/NiagaraEditorUtilities.cpp
+++ b/Engine/Plugins/FX/Niagara/Source/NiagaraEditor/Private/NiagaraEditorUtilities.cpp
@@ -368,9 +368,9 @@ bool FNiagaraEditorUtilities::PODPropertyAppendCompileHash(const void* Container
 		InVisitor->UpdatePOD(*PropertyName, Value);
 		return true;
 	}
-	else if (Property->IsA(FInt16Property::StaticClass()))
+	else if (Property->IsA(FUInt16Property::StaticClass()))
 	{
-		FInt16Property* CastProp = CastFieldChecked<FInt16Property>(Property);
+		FUInt16Property* CastProp = CastFieldChecked<FUInt16Property>(Property);
 		uint16 Value = CastProp->GetPropertyValue_InContainer(Container, 0);
 		InVisitor->UpdatePOD(*PropertyName, Value);
 		return true;
diff --git a/Engine/Source/Editor/BlueprintGraph/Private/CallFunctionHandler.cpp b/Engine/Source/Editor/BlueprintGraph/Private/CallFunctionHandler.cpp
--- a/Engine/Source/Editor/BlueprintGraph/Private/CallFunctionHandler.cpp
+++ b/Engine/Source/Editor/BlueprintGraph/Private/CallFunctionHandler.cpp
@@ -749,9 +749,9 @@ void FKCHandler_CallFunction::Transform(FKismetFunctionContext& Context, UEdGrap
 
 			// Iterate the properties looking for Out Params that are tied to Net Properties.
 			// This is similar to the loop in CreateCallFunction
-			for (TFieldIterator<UProperty> It(Function); It && (It->PropertyFlags & CPF_Parm); ++It)
+			for (TFieldIterator<FProperty> It(Function); It && (It->PropertyFlags & CPF_Parm); ++It)
 			{
-				UProperty* Property = *It;
+				FProperty* Property = *It;
 				if (Property->HasAllPropertyFlags(CPF_OutParm | CPF_ReferenceParm))
 				{
 					for (int32 i = 0; i < RemainingPins.Num(); ++i)
@@ -763,7 +763,7 @@ void FKCHandler_CallFunction::Transform(FKismetFunctionContext& Context, UEdGrap
 							{
 								if (UK2Node_VariableGet* GetPropertyNode = Cast<UK2Node_VariableGet>(PinToTry->GetOwningNode()))
 								{
-									UProperty* ToCheck = GetPropertyNode->GetPropertyForVariable();
+									FProperty* ToCheck = GetPropertyNode->GetPropertyForVariable();
 
 									// We only need to do this work if we actually found a net property.
 									if (UNLIKELY(ToCheck->HasAnyPropertyFlags(CPF_Net)))
@@ -794,7 +794,7 @@ void FKCHandler_CallFunction::Transform(FKismetFunctionContext& Context, UEdGrap
 													if (UClass * UseClass = Blueprint->GeneratedClass)
 													{
 														OwningClass = UseClass;
-														ToCheck = FindFieldChecked<UProperty>(OwningClass, ToCheck->GetFName());
+														ToCheck = FindFieldChecked<FProperty>(OwningClass, ToCheck->GetFName());
 													}
 												}
 											}
diff --git a/Engine/Source/Editor/UnrealEd/Private/Commandlets/GatherTextFromMetadataCommandlet.cpp b/Engine/Source/Editor/UnrealEd/Private/Commandlets/GatherTextFromMetadataCommandlet.cpp
--- a/Engine/Source/Editor/UnrealEd/Private/Commandlets/GatherTextFromMetadataCommandlet.cpp
+++ b/Engine/Source/Editor/UnrealEd/Private/Commandlets/GatherTextFromMetadataCommandlet.cpp
@@ -276,28 +276,28 @@ template <typename TFieldType>
 void GatherTextFromFieldImplementation(TFieldType* const Field, const UGatherTextFromMetaDataCommandlet::FGatherParameters& Arguments, const FName InPlatformName, FLocTextHelper* GatherManifestHelper)
 {
 	for (int32 i = 0; i < Arguments.InputKeys.Num(); ++i)
-	{
-		FFormatNamedArguments PatternArguments;
+		{
+			FFormatNamedArguments PatternArguments;
 		PatternArguments.Add(TEXT("FieldPath"), FText::FromString(Field->GetFullGroupName(false)));
 
 		if (Field->HasMetaData(*Arguments.InputKeys[i]))
-		{
-			const FString& MetaDataValue = Field->GetMetaData(*Arguments.InputKeys[i]);
-			if (!MetaDataValue.IsEmpty())
 			{
+				const FString& MetaDataValue = Field->GetMetaData(*Arguments.InputKeys[i]);
+			if (!MetaDataValue.IsEmpty())
+				{
 				PatternArguments.Add(TEXT("MetaDataValue"), FText::FromString(MetaDataValue));
 
 				const UStruct* FieldOwnerType = Field->GetOwnerStruct();
-				const FString Namespace = Arguments.OutputNamespaces[i];
-				FLocItem LocItem(MetaDataValue);
-				FManifestContext Context;
-				Context.Key = FText::Format(Arguments.OutputKeys[i], PatternArguments).ToString();
+					const FString Namespace = Arguments.OutputNamespaces[i];
+					FLocItem LocItem(MetaDataValue);
+					FManifestContext Context;
+					Context.Key = FText::Format(Arguments.OutputKeys[i], PatternArguments).ToString();
 				Context.SourceLocation = FString::Printf(TEXT("From metadata for key %s of member %s in %s (type: %s, owner: %s)"), *Arguments.InputKeys[i], *Field->GetName(), *Field->GetFullGroupName(true), *Field->GetClass()->GetName(), FieldOwnerType ? *FieldOwnerType->GetName() : TEXT("<null>"));
-				Context.PlatformName = InPlatformName;
-				GatherManifestHelper->AddSourceText(Namespace, LocItem, Context);
+					Context.PlatformName = InPlatformName;
+					GatherManifestHelper->AddSourceText(Namespace, LocItem, Context);
+				}
 			}
 		}
-	}
 }
 
 
diff --git a/Engine/Source/Programs/UnrealHeaderTool/Private/CodeGenerator.cpp b/Engine/Source/Programs/UnrealHeaderTool/Private/CodeGenerator.cpp
--- a/Engine/Source/Programs/UnrealHeaderTool/Private/CodeGenerator.cpp
+++ b/Engine/Source/Programs/UnrealHeaderTool/Private/CodeGenerator.cpp
@@ -166,13 +166,13 @@ namespace
 		if (!Class->HasAnyClassFlags(CLASS_ReplicationDataIsSetUp))
 		{
 			for (TFieldIterator<FProperty> It(Class, EFieldIteratorFlags::ExcludeSuper); It; ++It)
+		{
+			if ((It->PropertyFlags & CPF_Net) != 0)
 			{
-				if ((It->PropertyFlags & CPF_Net) != 0)
-				{
-					return true;
-				}
+				return true;
 			}
 		}
+		}
 
 		return Class->FirstOwnedClassRep < Class->ClassReps.Num();
 	}
@@ -3661,8 +3661,8 @@ void FNativeClassHeaderGenerator::ExportGeneratedStructBodyMacros(FOutputDevice&
 				FString ParameterCast = FString::Printf(TEXT("*(%s*)"), *ExtractedType);
 
 				RigVMStubProlog.Add(FString::Printf(TEXT("%s %s = %sRigVMOperandMemory[%d];"),
-					*VariableType,
-					*ParamNameOriginal,
+				*VariableType,
+				*ParamNameOriginal,
 					*ParameterCast,
 					OperandIndex));
 
@@ -4001,20 +4001,20 @@ void FNativeClassHeaderGenerator::ExportGeneratedStructBodyMacros(FOutputDevice&
 			Out.Log(TEXT("}\r\n"));
 		}
 
-		Out.Log(TEXT("\r\n"));
+			Out.Log(TEXT("\r\n"));
 
 		bool bHasGetMaxArraySize = false;
 		for (const FRigVMParameter& StructMember : StructRigVMInfo->Members)
-		{
-			if (!StructMember.MaxArraySize.IsEmpty())
 			{
+			if (!StructMember.MaxArraySize.IsEmpty())
+				{
 				bHasGetMaxArraySize = true;
 				break;
+				}
 			}
-		}
 
 		if (bHasGetMaxArraySize)
-		{
+			{
 			Out.Logf(TEXT("int32 %s::GetMaxArraySize(const FName& InMemberName, const FRigVMUserDataArray& RigVMUserData)\r\n"), *StructNameCPP);
 			Out.Log(TEXT("{\r\n"));
 			for (const FRigVMParameter& StructMember : StructRigVMInfo->Members)
diff --git a/Engine/Source/Runtime/Engine/Private/BlueprintGeneratedClass.cpp b/Engine/Source/Runtime/Engine/Private/BlueprintGeneratedClass.cpp
--- a/Engine/Source/Runtime/Engine/Private/BlueprintGeneratedClass.cpp
+++ b/Engine/Source/Runtime/Engine/Private/BlueprintGeneratedClass.cpp
@@ -22,6 +22,7 @@
 #include "Misc/ScopeLock.h"
 #include "UObject/CoreObjectVersion.h"
 #include "Net/Core/PushModel/PushModel.h"
+#include "UObject/CoreObjectVersion.h"
 
 #if WITH_EDITOR
 #include "Kismet2/BlueprintEditorUtils.h"
@@ -76,7 +77,7 @@ void UBlueprintGeneratedClass::PostInitProperties()
 
 void UBlueprintGeneratedClass::PostLoad()
 {
-	Super::PostLoad();
+	Super::PostLoad();	
 
 #if WITH_EDITORONLY_DATA
 	UPackage* Package = GetOutermost();
@@ -1521,7 +1522,7 @@ void UBlueprintGeneratedClass::Link(FArchive& Ar, bool bRelinkExistingProperties
 	Super::Link(Ar, bRelinkExistingProperties);
 
 #if USE_UBER_GRAPH_PERSISTENT_FRAME
-	if(UsePersistentUberGraphFrame())
+	if (UsePersistentUberGraphFrame())
 	{
 		if (UberGraphFunction)
 		{
diff --git a/Engine/Source/Runtime/Engine/Private/Net/NetPushModelHelpers.cpp b/Engine/Source/Runtime/Engine/Private/Net/NetPushModelHelpers.cpp
--- a/Engine/Source/Runtime/Engine/Private/Net/NetPushModelHelpers.cpp
+++ b/Engine/Source/Runtime/Engine/Private/Net/NetPushModelHelpers.cpp
@@ -14,7 +14,7 @@ void UNetPushModelHelpers::MarkPropertyDirty(UObject* Object, FName PropertyName
 			// If this is too slow to be practical, we might be able build a lookup from ClassName+PropertyName -> RepIndex.
 			// It would be safest to invalidate the lookup on a hot reload / recompile (if possible), blueprint recompile,
 			// and map change.
-			UProperty* Property = FindField<UProperty>(Class, PropertyName);
+			FProperty* Property = FindField<FProperty>(Class, PropertyName);
 			if (Property == nullptr)
 			{
 				UE_LOG(LogNet, Warning, TEXT("UNetPushModelHelpers::MarkPropertyDirty: Unable to find Property %s in Class %s"), *PropertyName.ToString(), *Class->GetPathName());
diff --git a/Engine/Source/Runtime/Engine/Private/RepLayout.cpp b/Engine/Source/Runtime/Engine/Private/RepLayout.cpp
--- a/Engine/Source/Runtime/Engine/Private/RepLayout.cpp
+++ b/Engine/Source/Runtime/Engine/Private/RepLayout.cpp
@@ -1210,18 +1210,18 @@ static bool CompareParentProperty(
 	const FComparePropertiesSharedParams& SharedParams,
 	FComparePropertiesStackParams& StackParams)
 {
-	const FRepParentCmd& Parent = SharedParams.Parents[ParentIndex];
-	const bool bIsLifetime = EnumHasAnyFlags(Parent.Flags, ERepParentFlags::IsLifetime);
+		const FRepParentCmd& Parent = SharedParams.Parents[ParentIndex];
+		const bool bIsLifetime = EnumHasAnyFlags(Parent.Flags, ERepParentFlags::IsLifetime);
 
-	// Active state of a property applies to *all* connections.
-	// If the property is inactive, we can skip comparing it because we know it won't be sent.
-	// Further, this will keep the last active state of the property in the shadow buffer,
-	// meaning the next time the property becomes active it will be sent to all connections.
+		// Active state of a property applies to *all* connections.
+		// If the property is inactive, we can skip comparing it because we know it won't be sent.
+		// Further, this will keep the last active state of the property in the shadow buffer,
+		// meaning the next time the property becomes active it will be sent to all connections.
 	const bool bIsActive = !SharedParams.RepChangedPropertyTracker || SharedParams.RepChangedPropertyTracker->Parents[ParentIndex].Active;
-	const bool bShouldSkip = !bIsLifetime || !bIsActive || (Parent.Condition == COND_InitialOnly && !SharedParams.bIsInitial);
+		const bool bShouldSkip = !bIsLifetime || !bIsActive || (Parent.Condition == COND_InitialOnly && !SharedParams.bIsInitial);
 
-	if (bShouldSkip)
-	{
+		if (bShouldSkip)
+		{
 		return false;
 	}
 
@@ -1229,26 +1229,26 @@ static bool CompareParentProperty(
 	if (SharedParams.bIsNetworkProfilerActive)
 	{
 		const_cast<TBitArray<>&>(SharedParams.PropertiesCompared)[ParentIndex] = true;
-	}
+		}
 #endif
 
 	const FRepLayoutCmd& Cmd = SharedParams.Cmds[Parent.CmdStart];
 
 	if (EnumHasAnyFlags(SharedParams.Flags, ERepLayoutFlags::IsActor))
-	{
-		if (UNLIKELY(ParentIndex == (int32)AActor::ENetFields_Private::Role))
 		{
+			if (UNLIKELY(ParentIndex == (int32)AActor::ENetFields_Private::Role))
+			{
 			return CompareRoleProperty(SharedParams, StackParams, (int32)AActor::ENetFields_Private::Role, SharedParams.RepState->SavedRole);
-		}
+			}
 		if (UNLIKELY(ParentIndex == (int32)AActor::ENetFields_Private::RemoteRole))
-		{
+			{
 			return CompareRoleProperty(SharedParams, StackParams, (int32)AActor::ENetFields_Private::RemoteRole, SharedParams.RepState->SavedRemoteRole);
+			}
 		}
-	}
-	
+		
 	const int32 NumChanges = StackParams.Changed.Num();
 
-	// Note, Handle - 1 to account for CompareProperties_r incrementing handles.
+		// Note, Handle - 1 to account for CompareProperties_r incrementing handles.
 	CompareProperties_r(SharedParams, StackParams, Parent.CmdStart, Parent.CmdEnd, Cmd.RelativeHandle - 1);
 
 	return !!(StackParams.Changed.Num() - NumChanges);
