Merging //UE4/Main[REDACTED_HANDLE] to Dev-RenderPlat-Staging
#rnx

Note: ImmediatePhysicsSimulation.h still has move/rename problems.  I will fix those up separately.

diff --git a/Engine/Plugins/Experimental/ControlRig/Source/ControlRigEditor/Private/Sequencer/ControlRigParameterTrackEditor.cpp b/Engine/Plugins/Experimental/ControlRig/Source/ControlRigEditor/Private/Sequencer/ControlRigParameterTrackEditor.cpp
new file mode 100644
--- /dev/null
+++ b/Engine/Plugins/Experimental/ControlRig/Source/ControlRigEditor/Private/Sequencer/ControlRigParameterTrackEditor.cpp
@@ -0,0 +1,1310 @@
+// Copyright 1998-2019 Epic Games, Inc. All Rights Reserved.
+
+#include "ControlRigParameterTrackEditor.h"
+#include "Sequencer/MovieSceneControlRigParameterTrack.h"
+#include "Sequencer/MovieSceneControlRigParameterSection.h"
+#include "Framework/Commands/Commands.h"
+#include "Rendering/DrawElements.h"
+#include "Widgets/SBoxPanel.h"
+#include "Framework/MultiBox/MultiBoxBuilder.h"
+#include "GameFramework/Actor.h"
+#include "AssetData.h"
+#include "Modules/ModuleManager.h"
+#include "Layout/WidgetPath.h"
+#include "Framework/Application/MenuStack.h"
+#include "Framework/Application/SlateApplication.h"
+#include "Widgets/Layout/SBox.h"
+#include "SequencerSectionPainter.h"
+#include "Components/SkeletalMeshComponent.h"
+#include "Editor/UnrealEdEngine.h"
+#include "UnrealEdGlobals.h"
+#include "ClassViewerModule.h"
+#include "ClassViewerFilter.h"
+#include "CommonMovieSceneTools.h"
+#include "AssetRegistryModule.h"
+#include "IContentBrowserSingleton.h"
+#include "ContentBrowserModule.h"
+#include "MatineeImportTools.h"
+#include "Matinee/InterpTrackAnimControl.h"
+#include "SequencerUtilities.h"
+#include "ISectionLayoutBuilder.h"
+#include "EditorStyleSet.h"
+#include "MovieSceneTimeHelpers.h"
+#include "Fonts/FontMeasure.h"
+#include "AnimationEditorUtils.h"
+#include "Misc/MessageDialog.h"
+#include "Framework/Notifications/NotificationManager.h"
+#include "Widgets/Notifications/SNotificationList.h"
+#include "Toolkits/AssetEditorManager.h"
+#include "Engine/SCS_Node.h"
+#include "Engine/SimpleConstructionScript.h"
+#include "Engine/Blueprint.h"
+#include "ControlRig.h"
+#include "ControlRigEditMode.h"
+#include "EditorModeManager.h"
+#include "Engine/Selection.h"
+#include "Manipulatable/IControlRigManipulatable.h"
+#include "ControlRigSkeletalMeshBinding.h"
+#include "LevelEditorViewport.h"
+#include "IKeyArea.h"
+
+#define LOCTEXT_NAMESPACE "FControlRigParameterTrackEditor"
+
+
+
+class FControlRigTrackCommands
+	: public TCommands<FControlRigTrackCommands>
+{
+public:
+
+	FControlRigTrackCommands()
+		: TCommands<FControlRigTrackCommands>
+		(
+			"ControlRigTrack",
+			NSLOCTEXT("Contexts", "ControlRigTrack", "ControlRigTrack"),
+			NAME_None,
+			FEditorStyle::GetStyleSetName() // Icon Style Set
+			)
+		, BindingCount(0)
+	{ }
+
+	/** Sets a transform key at the current time for the selected actor */
+	TSharedPtr< FUICommandInfo > AddTransformKey;
+
+	/** Sets a translation key at the current time for the selected actor */
+	TSharedPtr< FUICommandInfo > AddTranslationKey;
+
+	/** Sets a rotation key at the current time for the selected actor */
+	TSharedPtr< FUICommandInfo > AddRotationKey;
+
+	/** Sets a scale key at the current time for the selected actor */
+	TSharedPtr< FUICommandInfo > AddScaleKey;
+
+	/**
+	 * Initialize commands
+	 */
+	virtual void RegisterCommands() override;
+
+	mutable uint32 BindingCount;
+};
+
+void FControlRigTrackCommands::RegisterCommands()
+{
+	UI_COMMAND(AddTransformKey, "Add Transform Key", "Add a transform key at the current time for the selected control rig control.", EUserInterfaceActionType::Button, FInputChord(EKeys::S));
+	UI_COMMAND(AddTranslationKey, "Add Translation Key", "Add a translation key at the current time for the selected control rig control.", EUserInterfaceActionType::Button, FInputChord(EModifierKey::Shift, EKeys::W));
+	UI_COMMAND(AddRotationKey, "Add Rotation Key", "Add a rotation key at the current time for the selected control rig control.", EUserInterfaceActionType::Button, FInputChord(EModifierKey::Shift, EKeys::E));
+	UI_COMMAND(AddScaleKey, "Add Scale Key", "Add a scale key at the current time for the selected control rig control.", EUserInterfaceActionType::Button, FInputChord(EModifierKey::Shift, EKeys::R));
+}
+
+static USkeleton* GetSkeletonFromComponent(UActorComponent* InComponent)
+{
+	USkeletalMeshComponent* SkeletalMeshComp = Cast<USkeletalMeshComponent>(InComponent);
+	if (SkeletalMeshComp && SkeletalMeshComp->SkeletalMesh && SkeletalMeshComp->SkeletalMesh->Skeleton)
+	{
+		// @todo Multiple actors, multiple components
+		return SkeletalMeshComp->SkeletalMesh->Skeleton;
+	}
+
+	return nullptr;
+}
+
+static USkeleton* AcquireSkeletonFromObjectGuid(const FGuid& Guid, UObject** Object, TSharedPtr<ISequencer> SequencerPtr)
+{
+	UObject* BoundObject = SequencerPtr.IsValid() ? SequencerPtr->FindSpawnedObjectOrTemplate(Guid) : nullptr;
+	*Object = BoundObject;
+	if (AActor* Actor = Cast<AActor>(BoundObject))
+	{
+		for (UActorComponent* Component : Actor->GetComponents())
+		{
+			if (USkeleton* Skeleton = GetSkeletonFromComponent(Component))
+			{
+				return Skeleton;
+			}
+		}
+
+		AActor* ActorCDO = Cast<AActor>(Actor->GetClass()->GetDefaultObject());
+		if (ActorCDO)
+		{
+			for (UActorComponent* Component : ActorCDO->GetComponents())
+			{
+				if (USkeleton* Skeleton = GetSkeletonFromComponent(Component))
+				{
+					return Skeleton;
+				}
+			}
+		}
+
+		UBlueprintGeneratedClass* ActorBlueprintGeneratedClass = Cast<UBlueprintGeneratedClass>(Actor->GetClass());
+		if (ActorBlueprintGeneratedClass)
+		{
+			const TArray<USCS_Node*>& ActorBlueprintNodes = ActorBlueprintGeneratedClass->SimpleConstructionScript->GetAllNodes();
+
+			for (USCS_Node* Node : ActorBlueprintNodes)
+			{
+				if (Node->ComponentClass->IsChildOf(USkeletalMeshComponent::StaticClass()))
+				{
+					if (USkeleton* Skeleton = GetSkeletonFromComponent(Node->GetActualComponentTemplate(ActorBlueprintGeneratedClass)))
+					{
+						return Skeleton;
+					}
+				}
+			}
+		}
+	}
+	else if (USkeletalMeshComponent* SkeletalMeshComponent = Cast<USkeletalMeshComponent>(BoundObject))
+	{
+		if (USkeleton* Skeleton = GetSkeletonFromComponent(SkeletalMeshComponent))
+		{
+			return Skeleton;
+		}
+	}
+
+	return nullptr;
+}
+
+FControlRigParameterTrackEditor::FControlRigParameterTrackEditor(TSharedRef<ISequencer> InSequencer)
+	: FKeyframeTrackEditor<UMovieSceneControlRigParameterTrack>(InSequencer), bIsDoingASelection(false)
+{
+	SelectionChangedHandle = InSequencer->GetSelectionChangedTracks().AddRaw(this, &FControlRigParameterTrackEditor::OnSelectionChanged);
+	SequencerChangedHandle = InSequencer->OnMovieSceneDataChanged().AddRaw(this, &FControlRigParameterTrackEditor::OnSequencerDataChanged);
+
+	FControlRigTrackCommands::Register();
+	//register all modified/selections for control rigs
+	UMovieScene* MovieScene = InSequencer->GetFocusedMovieSceneSequence()->GetMovieScene();
+	const TArray<FMovieSceneBinding>& Bindings = MovieScene->GetBindings();
+	for(const FMovieSceneBinding& Binding : Bindings)
+	{
+		UMovieSceneControlRigParameterTrack* Track = Cast<UMovieSceneControlRigParameterTrack>(MovieScene->FindTrack(UMovieSceneControlRigParameterTrack::StaticClass(), Binding.GetObjectGuid(), NAME_None));
+		if (Track && Track->GetControlRig())
+		{
+			Track->GetControlRig()->ControlModified().AddRaw(this, &FControlRigParameterTrackEditor::HandleControlModified);
+			Track->GetControlRig()->ControlSelected().AddRaw(this, &FControlRigParameterTrackEditor::HandleControlSelected);
+		}
+	}
+}
+
+FControlRigParameterTrackEditor::~FControlRigParameterTrackEditor()
+{
+}
+
+void FControlRigParameterTrackEditor::OnRelease()
+{
+	if (GetSequencer().IsValid())
+	{
+		if (SelectionChangedHandle.IsValid())
+		{
+			GetSequencer()->GetSelectionChangedTracks().Remove(SelectionChangedHandle);
+		}
+		if (SequencerChangedHandle.IsValid())
+		{
+			GetSequencer()->OnMovieSceneDataChanged().Remove(SequencerChangedHandle);
+		}
+		if (GetSequencer()->GetFocusedMovieSceneSequence() && GetSequencer()->GetFocusedMovieSceneSequence()->GetMovieScene())
+		{
+			UMovieScene* MovieScene = GetSequencer()->GetFocusedMovieSceneSequence()->GetMovieScene();
+			const TArray<FMovieSceneBinding>& Bindings = MovieScene->GetBindings();
+			for (const FMovieSceneBinding& Binding : Bindings)
+			{
+				UMovieSceneControlRigParameterTrack* Track = Cast<UMovieSceneControlRigParameterTrack>(MovieScene->FindTrack(UMovieSceneControlRigParameterTrack::StaticClass(), Binding.GetObjectGuid(), NAME_None));
+				if (Track && Track->GetControlRig())
+				{
+					Track->GetControlRig()->ControlModified().RemoveAll(this);
+					Track->GetControlRig()->ControlSelected().RemoveAll(this);
+				}
+			}
+		}
+	}
+	FControlRigEditMode* ControlRigEditMode = static_cast<FControlRigEditMode*>(GLevelEditorModeTools().GetActiveMode(FControlRigEditMode::ModeName));
+	if (ControlRigEditMode)
+	{
+		if (GLevelEditorModeTools().HasToolkitHost())
+		{
+			GLevelEditorModeTools().DeactivateMode(FControlRigEditMode::ModeName);
+		}
+
+		ControlRigEditMode->SetObjects(nullptr, FGuid(), nullptr);
+	}
+
+	const FControlRigTrackCommands& Commands = FControlRigTrackCommands::Get();
+	Commands.BindingCount--;
+
+	if (Commands.BindingCount < 1)
+	{
+		FControlRigTrackCommands::Unregister();
+	}
+
+}
+
+TSharedRef<ISequencerTrackEditor> FControlRigParameterTrackEditor::CreateTrackEditor(TSharedRef<ISequencer> InSequencer)
+{
+	return MakeShareable(new FControlRigParameterTrackEditor(InSequencer));
+}
+
+
+bool FControlRigParameterTrackEditor::SupportsType(TSubclassOf<UMovieSceneTrack> Type) const
+{
+	return Type == UMovieSceneControlRigParameterTrack::StaticClass();
+}
+
+
+TSharedRef<ISequencerSection> FControlRigParameterTrackEditor::MakeSectionInterface(UMovieSceneSection& SectionObject, UMovieSceneTrack& Track, FGuid ObjectBinding)
+{
+	check(SupportsType(SectionObject.GetOuter()->GetClass()));
+
+	return MakeShareable(new FControlRigParameterSection(SectionObject, GetSequencer()));
+}
+
+
+void FControlRigParameterTrackEditor::BuildObjectBindingTrackMenu(FMenuBuilder& MenuBuilder, const TArray<FGuid>& ObjectBindings, const UClass* ObjectClass)
+{
+
+	if (ObjectClass->IsChildOf(USkeletalMeshComponent::StaticClass()) || ObjectClass->IsChildOf(AActor::StaticClass()) || ObjectClass->IsChildOf(UChildActorComponent::StaticClass()))
+	{
+		const TSharedPtr<ISequencer> ParentSequencer = GetSequencer();
+		UObject *BoundObject = nullptr;
+		USkeleton* Skeleton = AcquireSkeletonFromObjectGuid(ObjectBindings[0], &BoundObject, ParentSequencer);
+
+		if (Skeleton)
+		{
+			//if there are any other control rigs we don't allow it for now..
+			//mz todo will allow later
+			UMovieScene* MovieScene = GetSequencer()->GetFocusedMovieSceneSequence()->GetMovieScene();
+			UMovieSceneControlRigParameterTrack* ExistingTrack = Cast<UMovieSceneControlRigParameterTrack>(MovieScene->FindTrack(UMovieSceneControlRigParameterTrack::StaticClass(), ObjectBindings[0], NAME_None));
+			if (!ExistingTrack)
+			{
+				UMovieSceneTrack* Track = nullptr;
+
+				MenuBuilder.AddSubMenu(
+					LOCTEXT("AddControlRig", "Animation ControlRig"), NSLOCTEXT("Sequencer", "AddControlRigTooltip", "Adds an animation ControlRig track."),
+					FNewMenuDelegate::CreateRaw(this, &FControlRigParameterTrackEditor::AddControlRigSubMenu, ObjectBindings, Track)
+				);
+			}
+		}
+	}
+}
+
+void FControlRigParameterTrackEditor::AddControlRigSubMenu(FMenuBuilder& MenuBuilder, TArray<FGuid> ObjectBindings, UMovieSceneTrack* Track)
+{
+	/*
+	MenuBuilder.BeginSection(TEXT("ChooseSequence"), LOCTEXT("ChooseSequence", "Choose Sequence"));
+	{
+		FAssetPickerConfig AssetPickerConfig;
+		{
+			AssetPickerConfig.OnAssetSelected = FOnAssetSelected::CreateRaw(this, &FControlRigTrackEditor::OnSequencerAssetSelected, ObjectBindings, Track);
+			AssetPickerConfig.OnAssetEnterPressed = FOnAssetEnterPressed::CreateRaw(this, &FControlRigTrackEditor::OnSequencerAssetEnterPressed, ObjectBindings, Track);
+			AssetPickerConfig.bAllowNullSelection = false;
+			AssetPickerConfig.InitialAssetViewType = EAssetViewType::List;
+			AssetPickerConfig.Filter.bRecursiveClasses = true;
+			AssetPickerConfig.Filter.ClassNames.Add(UControlRigSequence::StaticClass()->GetFName());
+		}
+
+		FContentBrowserModule& ContentBrowserModule = FModuleManager::Get().LoadModuleChecked<FContentBrowserModule>(TEXT("ContentBrowser"));
+
+		TSharedPtr<SBox> MenuEntry = SNew(SBox)
+			.WidthOverride(300.0f)
+			.HeightOverride(300.0f)
+			[
+				ContentBrowserModule.Get().CreateAssetPicker(AssetPickerConfig)
+			];
+
+		MenuBuilder.AddWidget(MenuEntry.ToSharedRef(), FText::GetEmpty(), true);
+	}
+	MenuBuilder.EndSection();
+	*/
+
+
+	const TSharedPtr<ISequencer> ParentSequencer = GetSequencer();
+	UObject *BoundObject = nullptr;
+	//todo support multiple bindings?
+	USkeleton* Skeleton = AcquireSkeletonFromObjectGuid(ObjectBindings[0], &BoundObject, GetSequencer());
+
+	if (Skeleton)
+	{
+		//MenuBuilder.AddSubMenu(
+		//	LOCTEXT("AddControlRigTrack", "ControlRigTrack"), NSLOCTEXT("ControlRig", "AddControlRigTrack", "Adds a Control Rigtrack."),
+		//	FNewMenuDelegate::CreateRaw(this, &FControlRigParameterTrackEditor::AddAnimationSubMenu, BoundObject, ObjectBindings[0], Skeleton)
+		//);
+
+
+		class FControlRigClassFilter : public IClassViewerFilter
+		{
+		public:
+			bool IsClassAllowed(const FClassViewerInitializationOptions& InInitOptions, const UClass* InClass, TSharedRef< FClassViewerFilterFuncs > InFilterFuncs) override
+			{
+				const bool bChildOfObjectClass = InClass->IsChildOf(UControlRig::StaticClass());
+				const bool bMatchesFlags = !InClass->HasAnyClassFlags(CLASS_Hidden | CLASS_HideDropDown | CLASS_Deprecated | CLASS_Abstract);
+				return bChildOfObjectClass && bMatchesFlags;
+			}
+
+			virtual bool IsUnloadedClassAllowed(const FClassViewerInitializationOptions& InInitOptions, const TSharedRef< const IUnloadedBlueprintData > InUnloadedClassData, TSharedRef< FClassViewerFilterFuncs > InFilterFuncs) override
+			{
+				const bool bChildOfObjectClass = InUnloadedClassData->IsChildOf(UControlRig::StaticClass());
+				const bool bMatchesFlags = !InUnloadedClassData->HasAnyClassFlags(CLASS_Hidden | CLASS_HideDropDown | CLASS_Deprecated | CLASS_Abstract);
+				return bChildOfObjectClass && bMatchesFlags;
+			}
+		};
+
+		FClassViewerInitializationOptions Options;
+		Options.bShowUnloadedBlueprints = true;
+		Options.NameTypeToDisplay = EClassViewerNameTypeToDisplay::DisplayName;
+
+		TSharedPtr<FControlRigClassFilter> ClassFilter = MakeShareable(new FControlRigClassFilter);
+		Options.ClassFilter = ClassFilter;
+		Options.bShowNoneOption = false;
+
+		FClassViewerModule& ClassViewerModule = FModuleManager::LoadModuleChecked<FClassViewerModule>("ClassViewer");
+		//	FNewMenuDelegate::CreateRaw(this, &FControlRigParameterTrackEditor::AddAnimationSubMenu, BoundObject, ObjectBindings[0], Skeleton)
+
+
+
+		TSharedRef<SWidget> ClassViewer = ClassViewerModule.CreateClassViewer(Options, FOnClassPicked::CreateRaw(this, &FControlRigParameterTrackEditor::AddControlRig, BoundObject, ObjectBindings[0], Skeleton));
+		MenuBuilder.AddWidget(ClassViewer, FText::GetEmpty(), true);
+
+		/*
+		FAssetPickerConfig AssetPickerConfig;
+		{
+			AssetPickerConfig.OnAssetSelected = FOnAssetSelected::CreateRaw(this, &FControlRigParameterTrackEditor::OnControlRigAssetSelected, ObjectBindings, Skeleton);
+			AssetPickerConfig.OnAssetEnterPressed = FOnAssetEnterPressed::CreateRaw(this, &FControlRigParameterTrackEditor::OnControlRigAssetEnterPressed, ObjectBindings, Skeleton);
+			AssetPickerConfig.bAllowNullSelection = false;
+			AssetPickerConfig.InitialAssetViewType = EAssetViewType::List;
+			AssetPickerConfig.OnShouldFilterAsset = FOnShouldFilterAsset::CreateRaw(this, &FControlRigParameterTrackEditor::ShouldFilterAsset);
+			AssetPickerConfig.Filter.bRecursiveClasses = true;
+			AssetPickerConfig.Filter.ClassNames.Add((UControlRig::StaticClass())->GetFName());
+			AssetPickerConfig.Filter.TagsAndValues.Add(TEXT("Skeleton"), FAssetData(Skeleton).GetExportTextName());
+		}
+
+		FContentBrowserModule& ContentBrowserModule = FModuleManager::Get().LoadModuleChecked<FContentBrowserModule>(TEXT("ContentBrowser"));
+
+		TSharedPtr<SBox> MenuEntry = SNew(SBox)
+			.WidthOverride(300.0f)
+			.HeightOverride(300.f)
+			[
+				ContentBrowserModule.Get().CreateAssetPicker(AssetPickerConfig)
+			];
+
+		MenuBuilder.AddWidget(MenuEntry.ToSharedRef(), FText::GetEmpty(), true);
+		*/
+	}
+
+}
+/*
+bool FControlRigParameterTrackEditor::ShouldFilterAsset(const FAssetData& AssetData)
+{
+	if (AssetData.AssetClass == UControlRig::StaticClass()->GetFName())
+	{
+		return true;
+	}
+	return false;
+}
+
+void FControlRigParameterTrackEditor::OnControlRigAssetSelected(const FAssetData& AssetData, TArray<FGuid> ObjectBindings, USkeleton* Skeleton)
+{
+	FSlateApplication::Get().DismissAllMenus();
+
+	UObject* SelectedObject = AssetData.GetAsset();
+	TSharedPtr<ISequencer> SequencerPtr = GetSequencer();
+
+	if (SelectedObject && SelectedObject->IsA(UControlRig::StaticClass()) && SequencerPtr.IsValid())
+	{
+		UControlRig* ControlRig = CastChecked<UControlRig>(AssetData.GetAsset());
+
+		const FScopedTransaction Transaction(LOCTEXT("AddAnimation_Transaction", "Add Animation"));
+
+		for (FGuid ObjectBinding : ObjectBindings)
+		{
+			UObject* Object = SequencerPtr->FindSpawnedObjectOrTemplate(ObjectBinding);
+			int32 RowIndex = INDEX_NONE;
+			//AnimatablePropertyChanged(FOnKeyProperty::CreateRaw(this, &FControlRigParameterTrackEditor::AddControlRig, BoundObject, ObjectBindings[0], Skeleton)));
+		}
+	}
+}
+
+void FControlRigParameterTrackEditor::OnControlRigAssetEnterPressed(const TArray<FAssetData>& AssetData, TArray<FGuid> ObjectBindings, USkeleton* Skeleton)
+{
+	if (AssetData.Num() > 0)
+	{
+		OnControlRigAssetSelected(AssetData[0].GetAsset(), ObjectBindings, Skeleton);
+	}
+}
+
+*/
+
+void FControlRigParameterTrackEditor::AddControlRig(UClass* InClass, UObject* BoundSkelMesh, FGuid SkelMeshBinding, USkeleton* Skeleton)
+{
+	FSlateApplication::Get().DismissAllMenus();
+	const TSharedPtr<ISequencer> SequencerParent = GetSequencer();
+
+	if (InClass && InClass->IsChildOf(UControlRig::StaticClass()) && SequencerParent.IsValid())
+	{
+		FScopedTransaction AddControlRigTrackTransaction(LOCTEXT("AddControlRigTrack_Transaction", "Add Control Rig Track"));
+
+		UMovieSceneControlRigParameterTrack* Track = Cast<UMovieSceneControlRigParameterTrack>(AddTrack(GetSequencer()->GetFocusedMovieSceneSequence()->GetMovieScene(), SkelMeshBinding, UMovieSceneControlRigParameterTrack::StaticClass(), NAME_None));
+		if (Track)
+		{
+
+			FString ObjectName = (InClass->GetName());
+			ObjectName.RemoveFromEnd(TEXT("_C"));
+
+			UControlRig* ControlRig = NewObject<UControlRig>(Track, InClass, FName(*ObjectName), RF_Transactional);
+
+			ControlRig->SetObjectBinding(MakeShared<FControlRigSkeletalMeshBinding>());
+			ControlRig->GetObjectBinding()->BindToObject(BoundSkelMesh);
+			ControlRig->Initialize();
+			ControlRig->Execute(EControlRigState::Update);
+
+			SequencerParent->NotifyMovieSceneDataChanged(EMovieSceneDataChangeType::MovieSceneStructureItemsChanged);
+
+			Track->Modify();
+			UMovieSceneSection* NewSection = Track->CreateControlRigSection(0, ControlRig);
+
+			//mz todo need to have multiple rigs with same class
+			Track->SetTrackName(FName(*ObjectName));
+			Track->SetDisplayName(FText::FromString(ObjectName));
+
+			GetSequencer()->EmptySelection();
+			GetSequencer()->SelectSection(NewSection);
+			GetSequencer()->ThrobSectionSelection();
+			GetSequencer()->NotifyMovieSceneDataChanged(EMovieSceneDataChangeType::MovieSceneStructureItemAdded);
+
+			FControlRigEditMode* ControlRigEditMode = static_cast<FControlRigEditMode*>(GLevelEditorModeTools().GetActiveMode(FControlRigEditMode::ModeName));
+			if (!ControlRigEditMode)
+			{
+				GLevelEditorModeTools().ActivateMode(FControlRigEditMode::ModeName);
+				ControlRigEditMode = static_cast<FControlRigEditMode*>(GLevelEditorModeTools().GetActiveMode(FControlRigEditMode::ModeName));
+
+			}
+			if (ControlRigEditMode)
+			{
+				ControlRigEditMode->SetObjects(ControlRig, FGuid(), nullptr);
+			}
+			ControlRig->ControlModified().AddRaw(this, &FControlRigParameterTrackEditor::HandleControlModified);
+			ControlRig->ControlSelected().AddRaw(this, &FControlRigParameterTrackEditor::HandleControlSelected);
+		}
+	}
+}
+
+
+void FControlRigParameterTrackEditor::BindCommands(TSharedRef<FUICommandList> SequencerCommandBindings)
+{
+	const FControlRigTrackCommands& Commands = FControlRigTrackCommands::Get();
+
+	CommandBindings = MakeShared<FUICommandList>();
+	CommandBindings->MapAction(
+		Commands.AddTransformKey,
+		FExecuteAction::CreateSP(this, &FControlRigParameterTrackEditor::OnAddTransformKeysForSelectedObjects, EMovieSceneTransformChannel::All),
+		FCanExecuteAction::CreateSP(this, &FControlRigParameterTrackEditor::CanAddTransformKeysForSelectedObjects));
+
+	CommandBindings->MapAction(
+		Commands.AddTranslationKey,
+		FExecuteAction::CreateSP(this, &FControlRigParameterTrackEditor::OnAddTransformKeysForSelectedObjects, EMovieSceneTransformChannel::Translation),
+		FCanExecuteAction::CreateSP(this, &FControlRigParameterTrackEditor::CanAddTransformKeysForSelectedObjects));
+
+	CommandBindings->MapAction(
+		Commands.AddRotationKey,
+		FExecuteAction::CreateSP(this, &FControlRigParameterTrackEditor::OnAddTransformKeysForSelectedObjects, EMovieSceneTransformChannel::Rotation),
+		FCanExecuteAction::CreateSP(this, &FControlRigParameterTrackEditor::CanAddTransformKeysForSelectedObjects));
+
+	CommandBindings->MapAction(
+		Commands.AddScaleKey,
+		FExecuteAction::CreateSP(this, &FControlRigParameterTrackEditor::OnAddTransformKeysForSelectedObjects, EMovieSceneTransformChannel::Scale),
+		FCanExecuteAction::CreateSP(this, &FControlRigParameterTrackEditor::CanAddTransformKeysForSelectedObjects));
+
+	Commands.BindingCount++;
+
+	// Add these bindings to Sequencer
+	SequencerCommandBindings->Append(CommandBindings.ToSharedRef());
+
+	// Also add them to the Curve Editor 
+	GetSequencer()->GetCommandBindings(ESequencerCommandBindings::CurveEditor)->Append(CommandBindings.ToSharedRef());
+}
+
+bool FControlRigParameterTrackEditor::CanAddTransformKeysForSelectedObjects()
+{
+	if (!GetSequencer()->IsAllowedToChange())
+	{
+		return false;
+	}
+
+	FControlRigEditMode* ControlRigEditMode = static_cast<FControlRigEditMode*>(GLevelEditorModeTools().GetActiveMode(FControlRigEditMode::ModeName));
+	if (ControlRigEditMode && ControlRigEditMode->GetControlRig())
+	{
+		UControlRig* ControlRig = ControlRigEditMode->GetControlRig();
+		FString OurName = ControlRig->GetName();
+		FName Name(*OurName);
+		if (TSharedPtr<IControlRigObjectBinding> ObjectBinding = ControlRig->GetObjectBinding())
+		{
+			TArray<FName> ControlNames = ControlRig->CurrentControlSelection();
+			return (ControlNames.Num() > 0);
+		}
+	}
+	return false;
+}
+
+void FControlRigParameterTrackEditor::OnAddTransformKeysForSelectedObjects(EMovieSceneTransformChannel Channel)
+{
+	if (!GetSequencer()->IsAllowedToChange())
+	{
+		return;
+	}
+
+	FControlRigEditMode* ControlRigEditMode = static_cast<FControlRigEditMode*>(GLevelEditorModeTools().GetActiveMode(FControlRigEditMode::ModeName));
+	if (ControlRigEditMode && ControlRigEditMode->GetControlRig())
+	{
+		UControlRig* ControlRig = ControlRigEditMode->GetControlRig();
+		FString OurName = ControlRig->GetName();
+		FName Name(*OurName);
+		if (TSharedPtr<IControlRigObjectBinding> ObjectBinding = ControlRig->GetObjectBinding())
+		{
+			TArray<FName> ControlNames = ControlRig->CurrentControlSelection();
+			for (const FName& ControlName : ControlNames)
+			{
+				USceneComponent* Component = Cast<USkeletalMeshComponent>(ObjectBinding->GetBoundObject());
+				if (Component)
+				{
+					UObject *ActorObject = Component->GetOwner();
+					AddControlKeys(ActorObject, ControlRig, Name, ControlName, Channel, ESequencerKeyMode::ManualKey);
+				}
+			}
+		}
+	}
+}
+
+void FControlRigParameterTrackEditor::OnSequencerDataChanged(EMovieSceneDataChangeType DataChangeType)
+{
+	UMovieScene* MovieScene = GetSequencer()->GetFocusedMovieSceneSequence()->GetMovieScene();
+	FControlRigEditMode* ControlRigEditMode = static_cast<FControlRigEditMode*>(GLevelEditorModeTools().GetActiveMode(FControlRigEditMode::ModeName));
+	
+	//if we have a valid control rig edit mode need to check and see the control rig in that mode is still in a track
+	//if not we get rid of it.
+	if (ControlRigEditMode && ControlRigEditMode->GetControlRig() != nullptr && MovieScene && (DataChangeType == EMovieSceneDataChangeType::MovieSceneStructureItemRemoved ||
+		DataChangeType == EMovieSceneDataChangeType::Unknown))
+	{
+		const TArray<FMovieSceneBinding>& Bindings = MovieScene->GetBindings();
+		for (const FMovieSceneBinding& Binding : Bindings)
+		{
+			UMovieSceneControlRigParameterTrack* Track = Cast<UMovieSceneControlRigParameterTrack>(MovieScene->FindTrack(UMovieSceneControlRigParameterTrack::StaticClass(), Binding.GetObjectGuid(), NAME_None));
+			if (Track && Track->GetControlRig() == ControlRigEditMode->GetControlRig())
+			{
+				return; //just exit out we still have a good track
+			}
+		}
+		//okay no good track so deactive it and delete it's Control Rig and bingings.
+		if (GLevelEditorModeTools().HasToolkitHost())
+		{
+			GLevelEditorModeTools().DeactivateMode(FControlRigEditMode::ModeName);
+		}
+		ControlRigEditMode->SetObjects(nullptr, FGuid(), nullptr);
+	}
+}
+
+void FControlRigParameterTrackEditor::OnSelectionChanged(TArray<UMovieSceneTrack*> InTracks)
+{
+	if (bIsDoingASelection)
+	{
+		return;
+	}
+
+	TGuardValue<bool> Guard(bIsDoingASelection, true);
+	TArray<FString> StringArray;
+	FControlRigEditMode* ControlRigEditMode = static_cast<FControlRigEditMode*>(GLevelEditorModeTools().GetActiveMode(FControlRigEditMode::ModeName));
+	UControlRig* ControlRig = nullptr;
+	//Always clear the control rig(s) in the edit mode.
+	if (ControlRigEditMode)
+	{
+		ControlRig = ControlRigEditMode->GetControlRig();
+		if (ControlRig)
+		{
+			ControlRig->ClearControlSelection();
+		}
+	}
+	TArray<const IKeyArea*> KeyAreas;
+	GetSequencer()->GetSelectedKeyAreas(KeyAreas);
+	if (KeyAreas.Num() <= 0)
+	{
+		return;
+	}
+
+	TMap<UControlRig *, TSet<FName>> RigsAndControls;
+	for (const IKeyArea* KeyArea : KeyAreas)
+	{
+		UMovieSceneControlRigParameterSection* MovieSection = Cast<UMovieSceneControlRigParameterSection>(KeyArea->GetOwningSection());
+		if (MovieSection)
+		{
+			ControlRig = MovieSection->ControlRig;
+			//Only create the edit mode if we have a KeyAra selected and it's not set and we have some boundobjects.
+			if (!ControlRigEditMode )
+			{
+				GLevelEditorModeTools().ActivateMode(FControlRigEditMode::ModeName);
+				ControlRigEditMode = static_cast<FControlRigEditMode*>(GLevelEditorModeTools().GetActiveMode(FControlRigEditMode::ModeName));
+				if (TSharedPtr<IControlRigObjectBinding> ObjectBinding = ControlRig->GetObjectBinding())
+				{
+					if (ControlRigEditMode)
+					{
+						ControlRigEditMode->SetObjects(ControlRig, FGuid(), nullptr);
+					}
+				}
+			}
+			else
+			{
+				if (ControlRigEditMode->GetControlRig() != ControlRig)
+				{
+					ControlRigEditMode->SetObjects(ControlRig, FGuid(), nullptr);
+				}
+			}
+			
+			const FMovieSceneChannelMetaData* MetaData = KeyArea->GetChannel().GetMetaData();
+			if (MetaData)
+			{
+				StringArray.SetNum(0);
+				FString String = MetaData->Name.ToString();
+				String.ParseIntoArray(StringArray, TEXT("."));
+				if (StringArray.Num() > 0)
+				{
+					FName ControlName(*StringArray[0]);
+					RigsAndControls.FindOrAdd(ControlRig).Add(ControlName);
+				}
+			}
+		}
+	}
+	for (TPair<UControlRig *, TSet<FName>> Pair : RigsAndControls)
+	{
+		Pair.Key->ClearControlSelection();
+		for (const FName Name : Pair.Value)
+		{
+			Pair.Key->SelectControl(Name, true);
+		}
+	}
+}
+
+
+void FControlRigParameterTrackEditor::HandleControlSelected(IControlRigManipulatable* ControlRigManp, const FRigControl& Control, bool bSelected)
+{
+	if (!bIsDoingASelection)
+	{
+		TGuardValue<bool> Guard(bIsDoingASelection, true);
+		FString OurName = ControlRigManp->GetName();
+		FName ControlRigName(*OurName);
+		if (TSharedPtr<IControlRigObjectBinding> ObjectBinding = ControlRigManp->GetObjectBinding())
+		{
+			UObject *ActorObject = nullptr;
+			USceneComponent* Component = Cast<USkeletalMeshComponent>(ObjectBinding->GetBoundObject());
+			if (!Component)
+			{
+				return;
+			}
+			ActorObject = Component->GetOwner();
+			bool bCreateTrack = false;
+			bool bCreateHandle = false;
+			FFindOrCreateHandleResult HandleResult = FindOrCreateHandleToObject(ActorObject, bCreateHandle, NAME_None);
+			FGuid ObjectHandle = HandleResult.Handle;
+			if (!ObjectHandle.IsValid())
+			{
+				return;
+			}
+			FFindOrCreateTrackResult TrackResult = FindOrCreateTrackForObject(ObjectHandle, UMovieSceneControlRigParameterTrack::StaticClass(), ControlRigName, bCreateTrack);
+			UMovieSceneControlRigParameterTrack* Track = CastChecked<UMovieSceneControlRigParameterTrack>(TrackResult.Track, ECastCheckedType::NullAllowed);
+			if (Track)
+			{
+				float Weight = 1.0f;
+				TArray<IKeyArea> KeyAreas;
+				TArray<FString> StringArray;
+
+				for (UMovieSceneSection* Section : Track->GetAllSections())
+				{
+					UMovieSceneControlRigParameterSection* ParamSection = Cast<UMovieSceneControlRigParameterSection>(Section);
+					if (ParamSection)
+					{
+						FMovieSceneChannelProxy& ChannelProxy = Section->GetChannelProxy();
+						for (const FMovieSceneChannelEntry& Entry : Section->GetChannelProxy().GetAllEntries())
+						{
+							const FName ChannelTypeName = Entry.GetChannelTypeName();
+
+							// One editor data ptr per channel
+							TArrayView<FMovieSceneChannel* const>        Channels = Entry.GetChannels();
+							TArrayView<const FMovieSceneChannelMetaData> AllMetaData = Entry.GetMetaData();
+
+							for (int32 Index = 0; Index < Channels.Num(); ++Index)
+							{
+								FMovieSceneChannelHandle ChannelHandle = ChannelProxy.MakeHandle(ChannelTypeName, Index);
+								const FMovieSceneChannelMetaData& MetaData = AllMetaData[Index];
+								FString String = MetaData.Name.ToString();
+								String.ParseIntoArray(StringArray, TEXT("."));
+								if (StringArray.Num() > 0)
+								{
+									FName ControlName(*StringArray[0]);
+									if (ControlName == Control.Name)
+									{
+										IKeyArea KeyArea(Section, ChannelHandle);
+										KeyAreas.Add(KeyArea);
+									}
+								}
+							}
+						}
+					}
+				}
+				if (KeyAreas.Num())
+				{
+					GetSequencer()->SelectByKeyAreas(KeyAreas, true, bSelected);
+				}
+			}
+		}
+	}
+}
+
+void FControlRigParameterTrackEditor::HandleControlModified(IControlRigManipulatable* ControlRigManp, const FRigControl& Control)
+{
+	if (GetSequencer().IsValid() && !GetSequencer()->IsAllowedToChange())
+	{
+		return;
+	}
+	FString OurName = ControlRigManp->GetName();
+	FName Name(*OurName);
+	if (TSharedPtr<IControlRigObjectBinding> ObjectBinding = ControlRigManp->GetObjectBinding())
+	{
+		UObject *ActorObject = nullptr;
+		USceneComponent* Component = Cast<USkeletalMeshComponent>(ObjectBinding->GetBoundObject());
+		if (Component)
+		{
+			ActorObject = Component->GetOwner();
+			AddControlKeys(ActorObject, ControlRigManp, Name, Control.Name, EMovieSceneTransformChannel::All, ESequencerKeyMode::AutoKey);
+		}
+	}
+}
+
+void FControlRigParameterTrackEditor::GetControlRigKeys(IControlRigManipulatable* Manip, FName ParameterName, EMovieSceneTransformChannel ChannelsToKey, FGeneratedTrackKeys& OutGeneratedKeys)
+{
+	const TArray<FRigControl>& Controls = Manip->AvailableControls();
+	// If key all is enabled, for a key on all the channels
+	if (GetSequencer()->GetKeyGroupMode() == EKeyGroupMode::KeyAll)
+	{
+		ChannelsToKey = EMovieSceneTransformChannel::All;
+	}
+
+	int32 ChannelIndex = 0;
+	for (const FRigControl& RigControl : Controls)
+	{
+		switch (RigControl.ControlType)
+		{
+		case ERigControlType::Float:
+		{
+			float Val = RigControl.Value.Get<float>();
+			OutGeneratedKeys.Add(FMovieSceneChannelValueSetter::Create<FMovieSceneFloatChannel>(ChannelIndex++, Val, true));
+			break;
+		}
+		case ERigControlType::Position:
+		case ERigControlType::Scale:
+		case ERigControlType::Rotator:
+		{
+			FVector Val = RigControl.Value.Get<FVector>();
+			OutGeneratedKeys.Add(FMovieSceneChannelValueSetter::Create<FMovieSceneFloatChannel>(ChannelIndex++, Val.X, true));
+			OutGeneratedKeys.Add(FMovieSceneChannelValueSetter::Create<FMovieSceneFloatChannel>(ChannelIndex++, Val.Y, true));
+			OutGeneratedKeys.Add(FMovieSceneChannelValueSetter::Create<FMovieSceneFloatChannel>(ChannelIndex++, Val.Z, true));
+			break;
+		}
+
+		case ERigControlType::Transform:
+		{
+			FTransform Val = RigControl.Value.Get<FTransform>();
+			FVector CurrentVector = Val.GetTranslation();
+			bool bKeyX = EnumHasAnyFlags(ChannelsToKey, EMovieSceneTransformChannel::TranslationX);
+			bool bKeyY = EnumHasAnyFlags(ChannelsToKey, EMovieSceneTransformChannel::TranslationY);
+			bool bKeyZ = EnumHasAnyFlags(ChannelsToKey, EMovieSceneTransformChannel::TranslationZ);
+			if (GetSequencer()->GetKeyGroupMode() == EKeyGroupMode::KeyGroup && (bKeyX || bKeyY || bKeyZ))
+			{
+				bKeyX = bKeyY = bKeyZ = true;
+			}
+			OutGeneratedKeys.Add(FMovieSceneChannelValueSetter::Create<FMovieSceneFloatChannel>(ChannelIndex++, CurrentVector.X, bKeyX));
+			OutGeneratedKeys.Add(FMovieSceneChannelValueSetter::Create<FMovieSceneFloatChannel>(ChannelIndex++, CurrentVector.Y, bKeyY));
+			OutGeneratedKeys.Add(FMovieSceneChannelValueSetter::Create<FMovieSceneFloatChannel>(ChannelIndex++, CurrentVector.Z, bKeyZ));
+
+			CurrentVector = Val.GetRotation().Euler();
+			bKeyX = EnumHasAnyFlags(ChannelsToKey, EMovieSceneTransformChannel::RotationX);
+			bKeyY = EnumHasAnyFlags(ChannelsToKey, EMovieSceneTransformChannel::RotationY);
+			bKeyZ = EnumHasAnyFlags(ChannelsToKey, EMovieSceneTransformChannel::RotationZ);
+			if (GetSequencer()->GetKeyGroupMode() == EKeyGroupMode::KeyGroup && (bKeyX || bKeyY || bKeyZ))
+			{
+				bKeyX = bKeyY = bKeyZ = true;
+			}
+			OutGeneratedKeys.Add(FMovieSceneChannelValueSetter::Create<FMovieSceneFloatChannel>(ChannelIndex++, CurrentVector.X, bKeyX));
+			OutGeneratedKeys.Add(FMovieSceneChannelValueSetter::Create<FMovieSceneFloatChannel>(ChannelIndex++, CurrentVector.Y, bKeyY));
+			OutGeneratedKeys.Add(FMovieSceneChannelValueSetter::Create<FMovieSceneFloatChannel>(ChannelIndex++, CurrentVector.Z, bKeyZ));
+
+			CurrentVector = Val.GetScale3D();
+			bKeyX = EnumHasAnyFlags(ChannelsToKey, EMovieSceneTransformChannel::ScaleX);
+			bKeyY = EnumHasAnyFlags(ChannelsToKey, EMovieSceneTransformChannel::ScaleY);
+			bKeyZ = EnumHasAnyFlags(ChannelsToKey, EMovieSceneTransformChannel::ScaleZ);
+			if (GetSequencer()->GetKeyGroupMode() == EKeyGroupMode::KeyGroup && (bKeyX || bKeyY || bKeyZ))
+			{
+				bKeyX = bKeyY = bKeyZ = true;
+			}
+			OutGeneratedKeys.Add(FMovieSceneChannelValueSetter::Create<FMovieSceneFloatChannel>(ChannelIndex++, CurrentVector.X, bKeyX));
+			OutGeneratedKeys.Add(FMovieSceneChannelValueSetter::Create<FMovieSceneFloatChannel>(ChannelIndex++, CurrentVector.Y, bKeyY));
+			OutGeneratedKeys.Add(FMovieSceneChannelValueSetter::Create<FMovieSceneFloatChannel>(ChannelIndex++, CurrentVector.Z, bKeyZ));
+
+			break;
+		}
+		//mz todo the other types
+		}
+	}
+}
+
+FKeyPropertyResult FControlRigParameterTrackEditor::AddKeysToControlRigHandle(UObject *Object, IControlRigManipulatable* Manip,
+	FGuid ObjectHandle, FFrameNumber KeyTime, FGeneratedTrackKeys& GeneratedKeys,
+	ESequencerKeyMode KeyMode, TSubclassOf<UMovieSceneTrack> TrackClass, FName ControlRigName)
+{
+	EAutoChangeMode AutoChangeMode = GetSequencer()->GetAutoChangeMode();
+	EAllowEditsMode AllowEditsMode = GetSequencer()->GetAllowEditsMode();
+
+	bool bCreateTrack =
+		(KeyMode == ESequencerKeyMode::AutoKey && (AutoChangeMode == EAutoChangeMode::AutoTrack || AutoChangeMode == EAutoChangeMode::All)) ||
+		KeyMode == ESequencerKeyMode::ManualKey ||
+		KeyMode == ESequencerKeyMode::ManualKeyForced ||
+		AllowEditsMode == EAllowEditsMode::AllowSequencerEditsOnly;
+
+	bool bCreateSection = false;
+		// we don't do this, maybe revisit if a bug occurs, but currently extends sections on autokey.
+	    //bCreateTrack || (KeyMode == ESequencerKeyMode::AutoKey && (AutoChangeMode != EAutoChangeMode::None));
+
+	// Try to find an existing Track, and if one doesn't exist check the key params and create one if requested.
+
+	FFindOrCreateTrackResult TrackResult = FindOrCreateTrackForObject(ObjectHandle, TrackClass, ControlRigName, bCreateTrack);
+	UMovieSceneControlRigParameterTrack* Track = CastChecked<UMovieSceneControlRigParameterTrack>(TrackResult.Track, ECastCheckedType::NullAllowed);
+
+	bool bTrackCreated = TrackResult.bWasCreated;
+
+	bool bSectionCreated = false;
+	FKeyPropertyResult KeyPropertyResult;
+
+	if (Track)
+	{
+		float Weight = 1.0f;
+
+		UMovieSceneSection* SectionToKey = bCreateSection ? Track->FindOrExtendSection(KeyTime, Weight) : Track->FindSection(KeyTime);
+
+		// If there's no overlapping section to key, create one only if a track was newly created. Otherwise, skip keying altogether
+		// so that the user is forced to create a section to key on.
+		if (bTrackCreated && !SectionToKey)
+		{
+			Track->Modify();
+			SectionToKey = Track->FindOrAddSection(KeyTime, bSectionCreated);
+			if (bSectionCreated && GetSequencer()->GetInfiniteKeyAreas())
+			{
+				SectionToKey->SetRange(TRange<FFrameNumber>::All());
+			}
+		}
+
+		if (SectionToKey && SectionToKey->GetRange().Contains(KeyTime))
+		{
+			if (!bTrackCreated)
+			{
+				ModifyOurGeneratedKeysByCurrentAndWeight(Object, Manip, Track, SectionToKey, KeyTime, GeneratedKeys, Weight);
+			}
+			UMovieSceneControlRigParameterSection* ParamSection = Cast<UMovieSceneControlRigParameterSection>(SectionToKey);
+			if (!ParamSection->GetDoNotKey())
+			{
+				KeyPropertyResult |= AddKeysToSection(SectionToKey, KeyTime, GeneratedKeys, KeyMode);
+			}
+		}
+	}
+
+	KeyPropertyResult.bTrackCreated |= bTrackCreated || bSectionCreated;
+	return KeyPropertyResult;
+}
+
+FKeyPropertyResult FControlRigParameterTrackEditor::AddKeysToControlRig(
+	UObject* Object, IControlRigManipulatable* Manip, FFrameNumber KeyTime, FGeneratedTrackKeys& GeneratedKeys,
+	ESequencerKeyMode KeyMode, TSubclassOf<UMovieSceneTrack> TrackClass, FName PropertyName)
+{
+	FKeyPropertyResult KeyPropertyResult;
+	EAutoChangeMode AutoChangeMode = GetSequencer()->GetAutoChangeMode();
+	EAllowEditsMode AllowEditsMode = GetSequencer()->GetAllowEditsMode();
+	bool bCreateHandle =
+		(KeyMode == ESequencerKeyMode::AutoKey && (AutoChangeMode == EAutoChangeMode::All)) ||
+		KeyMode == ESequencerKeyMode::ManualKey ||
+		KeyMode == ESequencerKeyMode::ManualKeyForced ||
+		AllowEditsMode == EAllowEditsMode::AllowSequencerEditsOnly;
+
+	FFindOrCreateHandleResult HandleResult = FindOrCreateHandleToObject(Object, bCreateHandle, NAME_None);
+	FGuid ObjectHandle = HandleResult.Handle;
+	KeyPropertyResult.bHandleCreated = HandleResult.bWasCreated;
+	if (ObjectHandle.IsValid())
+	{
+		KeyPropertyResult |= AddKeysToControlRigHandle(Object, Manip, ObjectHandle, KeyTime, GeneratedKeys, KeyMode, TrackClass, PropertyName);
+	}
+
+	return KeyPropertyResult;
+}
+
+void FControlRigParameterTrackEditor::AddControlKeys(UObject* InObject, IControlRigManipulatable* Manip, FName ControlRigName, FName RigControlName, EMovieSceneTransformChannel ChannelsToKey, ESequencerKeyMode KeyMode)
+{
+	if (GetSequencer().IsValid() && !GetSequencer()->IsAllowedToChange())
+	{
+		return;
+	}
+	bool bCreateTrack = false;
+	bool bCreateHandle = false;
+	FFindOrCreateHandleResult HandleResult = FindOrCreateHandleToObject(InObject, bCreateHandle, NAME_None);
+	FGuid ObjectHandle = HandleResult.Handle;
+	if (!ObjectHandle.IsValid())
+	{
+		return;
+	}
+	FFindOrCreateTrackResult TrackResult = FindOrCreateTrackForObject(ObjectHandle, UMovieSceneControlRigParameterTrack::StaticClass(), ControlRigName, bCreateTrack);
+	UMovieSceneControlRigParameterTrack* Track = CastChecked<UMovieSceneControlRigParameterTrack>(TrackResult.Track, ECastCheckedType::NullAllowed);
+	if (Track)
+	{
+		FFrameNumber  FrameTime = GetTimeForKey();
+		UMovieSceneSection* Section = Track->FindSection(FrameTime);
+		UMovieSceneControlRigParameterSection* ParamSection = Cast<UMovieSceneControlRigParameterSection>(Section);
+
+		if (ParamSection && ParamSection->GetDoNotKey())
+		{
+			return;
+		}
+	}
+
+	TSharedRef<FGeneratedTrackKeys> GeneratedKeys = MakeShared<FGeneratedTrackKeys>();
+
+	GetControlRigKeys(Manip, RigControlName, ChannelsToKey, *GeneratedKeys);
+
+	auto OnKeyProperty = [=](FFrameNumber Time) -> FKeyPropertyResult
+	{
+		return this->AddKeysToControlRig(InObject, Manip, Time, *GeneratedKeys, KeyMode, UMovieSceneControlRigParameterTrack::StaticClass(), ControlRigName);
+	};
+
+	AnimatablePropertyChanged(FOnKeyProperty::CreateLambda(OnKeyProperty));
+}
+
+bool FControlRigParameterTrackEditor::ModifyOurGeneratedKeysByCurrentAndWeight(UObject *Object, IControlRigManipulatable* Manip, UMovieSceneTrack *Track, UMovieSceneSection* SectionToKey, FFrameNumber KeyTime, FGeneratedTrackKeys& GeneratedTotalKeys, float Weight) const
+{
+	FFrameRate TickResolution = GetSequencer()->GetFocusedTickResolution();
+	FMovieSceneEvaluationTrack EvalTrack = Track->GenerateTrackTemplate();
+
+	FMovieSceneInterrogationData InterrogationData;
+	GetSequencer()->GetEvaluationTemplate().CopyActuators(InterrogationData.GetAccumulator());
+
+	FMovieSceneContext Context(FMovieSceneEvaluationRange(KeyTime, GetSequencer()->GetFocusedTickResolution()));
+	EvalTrack.Interrogate(Context, InterrogationData, Object);
+	const TArray<FRigControl>& Controls = Manip->AvailableControls();
+
+	FMovieSceneChannelProxy& Proxy = SectionToKey->GetChannelProxy();
+	int32 ChannelIndex = 0;
+	for (const FRigControl& RigControl : Controls)
+	{
+		switch (RigControl.ControlType)
+		{
+		case ERigControlType::Float:
+		{
+			for (const FFloatInterrogationData& Val : InterrogationData.Iterate<FFloatInterrogationData>(UMovieSceneControlRigParameterSection::GetFloatInterrogationKey()))
+			{
+				if (Val.ParameterName == RigControl.Name)
+				{
+					GeneratedTotalKeys[ChannelIndex]->ModifyByCurrentAndWeight(Proxy, KeyTime, (void *)&Val.Val, Weight);
+					break;
+				}
+			}
+			++ChannelIndex;
+			break;
+		}
+		case ERigControlType::Position:
+		case ERigControlType::Scale:
+		case ERigControlType::Rotator:
+		{
+			for (const FVectorInterrogationData& Val : InterrogationData.Iterate<FVectorInterrogationData>(UMovieSceneControlRigParameterSection::GetVectorInterrogationKey()))
+			{
+				if (Val.ParameterName == RigControl.Name)
+				{
+					GeneratedTotalKeys[ChannelIndex]->ModifyByCurrentAndWeight(Proxy, KeyTime, (void *)&Val.Val, Weight);
+					GeneratedTotalKeys[ChannelIndex + 1]->ModifyByCurrentAndWeight(Proxy, KeyTime, (void *)&Val.Val, Weight);
+					GeneratedTotalKeys[ChannelIndex + 2]->ModifyByCurrentAndWeight(Proxy, KeyTime, (void *)&Val.Val, Weight);
+
+					break;
+				}
+			}
+			ChannelIndex += 3;
+
+			break;
+		}
+
+		case ERigControlType::Transform:
+		{
+			for (const FTransformInterrogationData& Val : InterrogationData.Iterate<FTransformInterrogationData>(UMovieSceneControlRigParameterSection::GetTransformInterrogationKey()))
+			{
+				if (Val.ParameterName == RigControl.Name)
+				{
+					FVector CurrentPos = Val.Val.GetTranslation();
+					FRotator CurrentRot = Val.Val.GetRotation().Rotator();
+					FVector CurrentScale = Val.Val.GetScale3D();
+					GeneratedTotalKeys[ChannelIndex]->ModifyByCurrentAndWeight(Proxy, KeyTime, (void *)&CurrentPos.X, Weight);
+					GeneratedTotalKeys[ChannelIndex + 1]->ModifyByCurrentAndWeight(Proxy, KeyTime, (void *)&CurrentPos.Y, Weight);
+					GeneratedTotalKeys[ChannelIndex + 2]->ModifyByCurrentAndWeight(Proxy, KeyTime, (void *)&CurrentPos.Z, Weight);
+					GeneratedTotalKeys[ChannelIndex + 3]->ModifyByCurrentAndWeight(Proxy, KeyTime, (void *)&CurrentRot.Roll, Weight);
+					GeneratedTotalKeys[ChannelIndex + 4]->ModifyByCurrentAndWeight(Proxy, KeyTime, (void *)&CurrentRot.Pitch, Weight);
+					GeneratedTotalKeys[ChannelIndex + 5]->ModifyByCurrentAndWeight(Proxy, KeyTime, (void *)&CurrentRot.Yaw, Weight);
+					GeneratedTotalKeys[ChannelIndex + 6]->ModifyByCurrentAndWeight(Proxy, KeyTime, (void *)&CurrentScale.X, Weight);
+					GeneratedTotalKeys[ChannelIndex + 7]->ModifyByCurrentAndWeight(Proxy, KeyTime, (void *)&CurrentScale.Y, Weight);
+					GeneratedTotalKeys[ChannelIndex + 8]->ModifyByCurrentAndWeight(Proxy, KeyTime, (void *)&CurrentScale.Z, Weight);
+
+					break;
+				}
+			}
+			ChannelIndex += 9;
+
+			break;
+		}
+		//mz todo the other types
+		}
+	}
+	return true;
+}
+
+void FControlRigParameterSection::BuildSectionContextMenu(FMenuBuilder& MenuBuilder, const FGuid& InObjectBinding)
+{
+	UMovieSceneControlRigParameterSection* ParameterSection = CastChecked<UMovieSceneControlRigParameterSection>(WeakSection.Get());
+	TSharedPtr<ISequencer> SequencerPtr = WeakSequencer.Pin();
+
+	UControlRig* ControlRig = ParameterSection->ControlRig;
+	if (ControlRig)
+	{
+		const TArray<FRigControl>& Controls = ControlRig->AvailableControls();
+
+		auto MakeUIAction = [=](EMovieSceneTransformChannel ChannelsToToggle)
+		{
+			return FUIAction(
+				FExecuteAction::CreateLambda([=]
+			{
+				FScopedTransaction Transaction(LOCTEXT("SetActiveChannelsTransaction", "Set Active Channels"));
+				ParameterSection->Modify();
+				EMovieSceneTransformChannel Channels = ParameterSection->GetTransformMask().GetChannels();
+
+				if (EnumHasAllFlags(Channels, ChannelsToToggle) || (Channels & ChannelsToToggle) == EMovieSceneTransformChannel::None)
+				{
+					ParameterSection->SetTransformMask(ParameterSection->GetTransformMask().GetChannels() ^ ChannelsToToggle);
+				}
+				else
+				{
+					ParameterSection->SetTransformMask(ParameterSection->GetTransformMask().GetChannels() | ChannelsToToggle);
+				}
+
+				// Restore pre-animated state for the bound objects so that inactive channels will return to their default values.
+				for (TWeakObjectPtr<> WeakObject : SequencerPtr->FindBoundObjects(InObjectBinding, SequencerPtr->GetFocusedTemplateID()))
+				{
+					if (UObject* Object = WeakObject.Get())
+					{
+						SequencerPtr->RestorePreAnimatedState();
+					}
+				}
+
+				SequencerPtr->NotifyMovieSceneDataChanged(EMovieSceneDataChangeType::MovieSceneStructureItemsChanged);
+			}
+				),
+				FCanExecuteAction(),
+				FGetActionCheckState::CreateLambda([=]
+			{
+				EMovieSceneTransformChannel Channels = ParameterSection->GetTransformMask().GetChannels();
+				if (EnumHasAllFlags(Channels, ChannelsToToggle))
+				{
+					return ECheckBoxState::Checked;
+				}
+				else if (EnumHasAnyFlags(Channels, ChannelsToToggle))
+				{
+					return ECheckBoxState::Undetermined;
+				}
+				return ECheckBoxState::Unchecked;
+			})
+				);
+		};
+		auto ToggleControls = [=](int32 Index)
+		{
+			return FUIAction(
+				FExecuteAction::CreateLambda([=]
+			{
+				FScopedTransaction Transaction(LOCTEXT("RigSectionFilter", "Toggle Rig Control Filters"));
+				ParameterSection->Modify();
+				if (Index >= 0)
+				{
+					ParameterSection->SetControlsMask(Index, !ParameterSection->GetControlsMask(Index));
+				}
+				else
+				{
+					ParameterSection->FillControlsMask(!ParameterSection->GetControlsMask(0));
+				}
+				SequencerPtr->NotifyMovieSceneDataChanged(EMovieSceneDataChangeType::MovieSceneStructureItemsChanged);
+
+			}
+				),
+				FCanExecuteAction(),
+				FGetActionCheckState::CreateLambda([=]
+			{
+				TArray<bool> ControlBool = ParameterSection->GetControlsMask();
+				if (Index >= 0)
+				{
+					if (ControlBool[Index])
+					{
+						return ECheckBoxState::Checked;
+
+					}
+					else
+					{
+						return ECheckBoxState::Unchecked;
+					}
+				}
+				else
+				{
+					TOptional<bool> FirstVal;
+					for (bool Val : ControlBool)
+					{
+						if (FirstVal.IsSet())
+						{
+							if (Val != FirstVal)
+							{
+								return ECheckBoxState::Undetermined;
+							}
+						}
+						else
+						{
+							FirstVal = Val;
+						}
+
+					}
+					return (FirstVal.GetValue()) ? ECheckBoxState::Checked : ECheckBoxState::Unchecked;
+				}
+			})
+				);
+		};
+		MenuBuilder.BeginSection(NAME_None, LOCTEXT("RigSectionFilter", "Filter Controls"));
+		{
+			MenuBuilder.AddSubMenu(
+				LOCTEXT("ToggleRigControlsText", "Toggle Rig Controls"), LOCTEXT("ToggleRigControlsText_Tooltip", "Toggle Rig Controls"),
+				FNewMenuDelegate::CreateLambda([=](FMenuBuilder& SubMenuBuilder) {
+				int32 Index = 0;
+				for (const FRigControl& RigControl : Controls)
+				{
+					const FName RigName = RigControl.Name;
+					FText Name = FText::FromName(RigName);
+					FText Text = FText::Format(LOCTEXT("RigControlToggle", "{0}"), Name);
+					FText TooltipText = FText::Format(LOCTEXT("RigControlToggleTooltip", "Toggle Rig Control {0}"), Name);
+					SubMenuBuilder.AddMenuEntry(
+						Text, TooltipText,
+						FSlateIcon(), ToggleControls(Index++), NAME_None, EUserInterfaceActionType::ToggleButton);
+				}
+			}),
+				ToggleControls(-1),
+				NAME_None,
+				EUserInterfaceActionType::ToggleButton);
+
+			MenuBuilder.AddSubMenu(
+				LOCTEXT("AllTranslation", "Translation"), LOCTEXT("AllTranslation_ToolTip", "Causes this section to affect the translation of rig control transforms"),
+				FNewMenuDelegate::CreateLambda([=](FMenuBuilder& SubMenuBuilder) {
+				SubMenuBuilder.AddMenuEntry(
+					LOCTEXT("TranslationX", "X"), LOCTEXT("TranslationX_ToolTip", "Causes this section to affect the X channel of the transform's translation"),
+					FSlateIcon(), MakeUIAction(EMovieSceneTransformChannel::TranslationX), NAME_None, EUserInterfaceActionType::ToggleButton);
+				SubMenuBuilder.AddMenuEntry(
+					LOCTEXT("TranslationY", "Y"), LOCTEXT("TranslationY_ToolTip", "Causes this section to affect the Y channel of the transform's translation"),
+					FSlateIcon(), MakeUIAction(EMovieSceneTransformChannel::TranslationY), NAME_None, EUserInterfaceActionType::ToggleButton);
+				SubMenuBuilder.AddMenuEntry(
+					LOCTEXT("TranslationZ", "Z"), LOCTEXT("TranslationZ_ToolTip", "Causes this section to affect the Z channel of the transform's translation"),
+					FSlateIcon(), MakeUIAction(EMovieSceneTransformChannel::TranslationZ), NAME_None, EUserInterfaceActionType::ToggleButton);
+			}),
+				MakeUIAction(EMovieSceneTransformChannel::Translation),
+				NAME_None,
+				EUserInterfaceActionType::ToggleButton);
+
+			MenuBuilder.AddSubMenu(
+				LOCTEXT("AllRotation", "Rotation"), LOCTEXT("AllRotation_ToolTip", "Causes this section to affect the rotation of the rig control transform"),
+				FNewMenuDelegate::CreateLambda([=](FMenuBuilder& SubMenuBuilder) {
+				SubMenuBuilder.AddMenuEntry(
+					LOCTEXT("RotationX", "Roll (X)"), LOCTEXT("RotationX_ToolTip", "Causes this section to affect the roll (X) channel the transform's rotation"),
+					FSlateIcon(), MakeUIAction(EMovieSceneTransformChannel::RotationX), NAME_None, EUserInterfaceActionType::ToggleButton);
+				SubMenuBuilder.AddMenuEntry(
+					LOCTEXT("RotationY", "Pitch (Y)"), LOCTEXT("RotationY_ToolTip", "Causes this section to affect the pitch (Y) channel the transform's rotation"),
+					FSlateIcon(), MakeUIAction(EMovieSceneTransformChannel::RotationY), NAME_None, EUserInterfaceActionType::ToggleButton);
+				SubMenuBuilder.AddMenuEntry(
+					LOCTEXT("RotationZ", "Yaw (Z)"), LOCTEXT("RotationZ_ToolTip", "Causes this section to affect the yaw (Z) channel the transform's rotation"),
+					FSlateIcon(), MakeUIAction(EMovieSceneTransformChannel::RotationZ), NAME_None, EUserInterfaceActionType::ToggleButton);
+			}),
+				MakeUIAction(EMovieSceneTransformChannel::Rotation),
+				NAME_None,
+				EUserInterfaceActionType::ToggleButton);
+
+			MenuBuilder.AddSubMenu(
+				LOCTEXT("AllScale", "Scale"), LOCTEXT("AllScale_ToolTip", "Causes this section to affect the scale of the rig control transform"),
+				FNewMenuDelegate::CreateLambda([=](FMenuBuilder& SubMenuBuilder) {
+				SubMenuBuilder.AddMenuEntry(
+					LOCTEXT("ScaleX", "X"), LOCTEXT("ScaleX_ToolTip", "Causes this section to affect the X channel of the transform's scale"),
+					FSlateIcon(), MakeUIAction(EMovieSceneTransformChannel::ScaleX), NAME_None, EUserInterfaceActionType::ToggleButton);
+				SubMenuBuilder.AddMenuEntry(
+					LOCTEXT("ScaleY", "Y"), LOCTEXT("ScaleY_ToolTip", "Causes this section to affect the Y channel of the transform's scale"),
+					FSlateIcon(), MakeUIAction(EMovieSceneTransformChannel::ScaleY), NAME_None, EUserInterfaceActionType::ToggleButton);
+				SubMenuBuilder.AddMenuEntry(
+					LOCTEXT("ScaleZ", "Z"), LOCTEXT("ScaleZ_ToolTip", "Causes this section to affect the Z channel of the transform's scale"),
+					FSlateIcon(), MakeUIAction(EMovieSceneTransformChannel::ScaleZ), NAME_None, EUserInterfaceActionType::ToggleButton);
+			}),
+				MakeUIAction(EMovieSceneTransformChannel::Scale),
+				NAME_None,
+				EUserInterfaceActionType::ToggleButton);
+
+			//mz todo h
+			MenuBuilder.AddMenuEntry(
+				LOCTEXT("Weight", "Weight"), LOCTEXT("Weight_ToolTip", "Causes this section to be applied with a user-specified weight curve"),
+				FSlateIcon(), MakeUIAction(EMovieSceneTransformChannel::Weight), NAME_None, EUserInterfaceActionType::ToggleButton);
+		}
+		MenuBuilder.EndSection();
+	}
+}
+
+//mz todo
+bool FControlRigParameterSection::RequestDeleteCategory(const TArray<FName>& CategoryNamePaths)
+{
+	UMovieSceneControlRigParameterSection* ParameterSection = CastChecked<UMovieSceneControlRigParameterSection>(WeakSection.Get());
+	TSharedPtr<ISequencer> SequencerPtr = WeakSequencer.Pin();
+	/*
+	const FScopedTransaction Transaction(LOCTEXT("DeleteTransformCategory", "Delete transform category"));
+
+	if (ParameterSection->TryModify())
+	{
+		FName CategoryName = CategoryNamePaths[CategoryNamePaths.Num() - 1];
+
+		EMovieSceneTransformChannel Channel = ParameterSection->GetTransformMask().GetChannels();
+		EMovieSceneTransformChannel ChannelToRemove = ParameterSection->GetTransformMaskByName(CategoryName).GetChannels();
+
+		Channel = Channel ^ ChannelToRemove;
+
+		ParameterSection->SetTransformMask(Channel);
+
+		SequencerPtr->NotifyMovieSceneDataChanged(EMovieSceneDataChangeType::MovieSceneStructureItemsChanged);
+		return true;
+	}
+	*/
+	return false;
+}
+
+bool FControlRigParameterSection::RequestDeleteKeyArea(const TArray<FName>& KeyAreaNamePaths)
+{
+	UMovieSceneControlRigParameterSection* ParameterSection = CastChecked<UMovieSceneControlRigParameterSection>(WeakSection.Get());
+	TSharedPtr<ISequencer> SequencerPtr = WeakSequencer.Pin();
+
+	/*
+	const FScopedTransaction Transaction(LOCTEXT("DeleteTransformChannel", "Delete transform channel"));
+
+	if (ParameterSection->TryModify())
+	{
+		// Only delete the last key area path which is the channel. ie. TranslationX as opposed to Translation
+		FName KeyAreaName = KeyAreaNamePaths[KeyAreaNamePaths.Num() - 1];
+
+		EMovieSceneTransformChannel Channel = ParameterSection->GetTransformMask().GetChannels();
+		EMovieSceneTransformChannel ChannelToRemove = ParameterSection->GetTransformMaskByName(KeyAreaName).GetChannels();
+
+		Channel = Channel ^ ChannelToRemove;
+
+		ParameterSection->SetTransformMask(Channel);
+
+		SequencerPtr->NotifyMovieSceneDataChanged(EMovieSceneDataChangeType::MovieSceneStructureItemsChanged);
+		return true;
+	}
+	*/
+	return true;
+}
+
+
+#undef LOCTEXT_NAMESPACE
diff --git a/Engine/Plugins/Experimental/SkeletalReduction/Source/Private/SkeletalMeshReductionPlugin.cpp b/Engine/Plugins/Experimental/SkeletalReduction/Source/Private/SkeletalMeshReductionPlugin.cpp
--- a/Engine/Plugins/Experimental/SkeletalReduction/Source/Private/SkeletalMeshReductionPlugin.cpp
+++ b/Engine/Plugins/Experimental/SkeletalReduction/Source/Private/SkeletalMeshReductionPlugin.cpp
@@ -22,7 +22,7 @@
 #include "SkeletalSimplifier.h"
 #include "SkeletalMeshReductionSkinnedMesh.h"
 #include "Stats/StatsMisc.h"
-#include "Assets/ClothingAsset.h"
+#include "ClothingAsset.h"
 #include "Factories/FbxSkeletalMeshImportData.h"
 #include "LODUtilities.h"
 
@@ -100,8 +100,7 @@ public:
 	* @returns true if reduction was successful.
 	*/
 	virtual bool ReduceSkeletalMesh( class USkeletalMesh* SkeletalMesh,
-		                             int32 LODIndex,
-		                             bool bReregisterComponent = true
+		                             int32 LODIndex		                             
 	                                ) override ;
 
 
@@ -170,7 +169,8 @@ private:
 		                         FSkeletalMeshOptimizationSettings Settings,
 								 const FImportantBones& ImportantBones,
 		                         const TArray<FMatrix>& BoneMatrices,
-		                         const int32 LODIndex ) const;
+		                         const int32 LODIndex,
+								 const bool bReducingSourceModel) const;
 
 	/**
 	* Remove the specified section from the mesh.
@@ -204,7 +204,8 @@ private:
 		                                 const FSkeletalMeshLODModel& SrcLODModel,
 		                                 const SkeletalSimplifier::FSkinnedSkeletalMesh& SkinnedMesh,
 		                                 const FReferenceSkeleton& RefSkeleton,
-		                                 FSkeletalMeshLODModel& NewModel ) const;
+		                                 FSkeletalMeshLODModel& NewModel,
+										 const bool bReducingSourceModel) const;
 
 	/**
 	* Add the SourceModelInflunces to the LODModel in the case that alternate weights exists.
@@ -359,10 +360,7 @@ void FSkeletalMeshReduction::ShutdownModule()
 
 }
 
-
-
-
-bool FQuadricSkeletalMeshReduction::ReduceSkeletalMesh( USkeletalMesh* SkeletalMesh, int32 LODIndex, bool bReregisterComponent )
+bool FQuadricSkeletalMeshReduction::ReduceSkeletalMesh( USkeletalMesh* SkeletalMesh, int32 LODIndex )
 {
 	check(SkeletalMesh);
 	check(LODIndex >= 0);
@@ -372,21 +370,8 @@ bool FQuadricSkeletalMeshReduction::ReduceSkeletalMesh( USkeletalMesh* SkeletalM
 	check(SkeletalMeshResource);
 	check(LODIndex <= SkeletalMeshResource->LODModels.Num());
 
-	if (bReregisterComponent)
-	{
-		TComponentReregisterContext<USkinnedMeshComponent> ReregisterContext;
-		SkeletalMesh->ReleaseResources();
-		SkeletalMesh->ReleaseResourcesFence.Wait();
-
-		ReduceSkeletalMesh(*SkeletalMesh, *SkeletalMeshResource, LODIndex);
-
-		SkeletalMesh->PostEditChange();
-		SkeletalMesh->InitResources();
-	}
-	else
-	{
-		ReduceSkeletalMesh(*SkeletalMesh, *SkeletalMeshResource, LODIndex);
-	}
+	FScopedSkeletalMeshPostEditChange ScopedPostEditChange(SkeletalMesh);
+	ReduceSkeletalMesh(*SkeletalMesh, *SkeletalMeshResource, LODIndex);
 
 	return true;
 }
@@ -1541,8 +1526,8 @@ void FQuadricSkeletalMeshReduction::ConvertToFSkeletalMeshLODModel( const int32
 	                                                                const FSkeletalMeshLODModel& SrcLODModel,
 	                                                                const SkeletalSimplifier::FSkinnedSkeletalMesh& SkinnedMesh,
 	                                                                const FReferenceSkeleton& RefSkeleton,
-	                                                                FSkeletalMeshLODModel& NewModel
-                                                                   ) const
+	                                                                FSkeletalMeshLODModel& NewModel,	
+																	const bool bReducingSourceModel) const
 {
 	// We might be re-using this model - so clear it.
 
@@ -1553,11 +1538,14 @@ void FQuadricSkeletalMeshReduction::ConvertToFSkeletalMeshLODModel( const int32
 	FSkeletalMeshData SkeletalMeshData;
 	ExtractFSkeletalData(SkinnedMesh, SkeletalMeshData);
 
-	// Add alternate weight data to the NewModel.  Note, this has to be done before we "BuildSkeletalMesh" to insure
-	// that the bone-based vertex chunking respects the alternate weights.
-	// NB: this only prepares the NewModel, but BuildSkeletalMesh is only 1/2-aware of this, so we will have to do some additional work after.
-	
-	AddSourceModelInfluences(SrcLODModel, SkinnedMesh, NewModel);
+	//if (!bReducingSourceModel)
+	{
+		// Add alternate weight data to the NewModel.  Note, this has to be done before we "BuildSkeletalMesh" to insure
+		// that the bone-based vertex chunking respects the alternate weights.
+		// NB: this only prepares the NewModel, but BuildSkeletalMesh is only 1/2-aware of this, so we will have to do some additional work after.
+
+		AddSourceModelInfluences(SrcLODModel, SkinnedMesh, NewModel);
+	}
 	
 
 
@@ -1576,6 +1564,7 @@ void FQuadricSkeletalMeshReduction::ConvertToFSkeletalMeshLODModel( const int32
 	Options.bComputeNormals = false;
 	Options.bComputeTangents = false;
 	Options.bUseMikkTSpace = true; //Avoid builtin build by specifying true for mikkt space
+	Options.bComputeWeightedNormals = false;
 	Options.OverlappingThresholds.ThresholdPosition = 0.0f;
 	Options.OverlappingThresholds.ThresholdTangentNormal = 0.0f;
 	Options.OverlappingThresholds.ThresholdUV = 0.0f;
@@ -1595,15 +1584,19 @@ void FQuadricSkeletalMeshReduction::ConvertToFSkeletalMeshLODModel( const int32
 		Options
 	);
 
+	//Re-Apply the user section changes, the UserSectionsData is map to original section and should match the builded LODModel
+	NewModel.SyncronizeUserSectionsDataArray();
+
 	// Set texture coordinate count on the new model.
 	NewModel.NumTexCoords = SkeletalMeshData.TexCoordCount;
 
 
+	//if (!bReducingSourceModel)
+	{
+		// Update the alternate weights
 
-	// Update the alternate weights
-	
-	UpdateAlternateWeights(MaxBonesPerVertex, NewModel);
-
+		UpdateAlternateWeights(MaxBonesPerVertex, NewModel);
+	}
 }
 
 bool FQuadricSkeletalMeshReduction::ReduceSkeletalLODModel( const FSkeletalMeshLODModel& SrcModel,
@@ -1613,7 +1606,8 @@ bool FQuadricSkeletalMeshReduction::ReduceSkeletalLODModel( const FSkeletalMeshL
 	                                                        FSkeletalMeshOptimizationSettings Settings,
 															const FImportantBones& ImportantBones,
 	                                                        const TArray<FMatrix>& BoneMatrices,
-	                                                        const int32 LODIndex
+	                                                        const int32 LODIndex,
+															const bool bReducingSourceModel
                                                            ) const
 {
 
@@ -1682,7 +1676,7 @@ bool FQuadricSkeletalMeshReduction::ReduceSkeletalLODModel( const FSkeletalMeshL
 
 		// Convert to SkeletalMeshLODModel. 
 
-		ConvertToFSkeletalMeshLODModel(Settings.MaxBonesPerVertex, SrcModel, SkinnedSkeletalMesh, RefSkeleton, OutSkeletalMeshLODModel);
+		ConvertToFSkeletalMeshLODModel(Settings.MaxBonesPerVertex, SrcModel, SkinnedSkeletalMesh, RefSkeleton, OutSkeletalMeshLODModel, bReducingSourceModel);
 
 		// We may need to do additional simplification if the user specified a hard number limit for verts and
 		// the internal chunking during conversion split some verts.
@@ -1804,13 +1798,26 @@ void FQuadricSkeletalMeshReduction::ReduceSkeletalMesh(USkeletalMesh& SkeletalMe
 	struct FSectionData
 	{
 		uint16 MaterialIndex;
+		int32 MaterialMap;
 		bool bCastShadow;
 		bool bRecomputeTangent;
+		bool bDisabled;
+		int32 GenerateUpToLodIndex;
+		int32 ChunkedParentSectionIndex;
+		int32 OriginalDataSectionIndex;
 	};
 
 	TMap<int32, FSectionData> BackupSectionIndexToSectionData;
+	TMap<int32, FSkelMeshSourceSectionUserData> BackupUserSectionsData;
+	FString BackupLodModelBuildStringID = TEXT("");
 	//Store the section data. Store the source LOD model in case we add a LOD, or the target LOD model in case the LOD already exist
 	FSkeletalMeshLODModel& BackupSectionLODModel = bLODModelAdded ? SkeletalMeshResource.LODModels[Settings.BaseLOD] : SkeletalMeshResource.LODModels[LODIndex];
+	if (!bLODModelAdded)
+	{
+		BackupLodModelBuildStringID = BackupSectionLODModel.BuildStringID;
+		BackupUserSectionsData = BackupSectionLODModel.UserSectionsData;
+	}
+
 	BackupSectionIndexToSectionData.Reserve(BackupSectionLODModel.Sections.Num());
 
 	for (int32 SectionIndex = 0; SectionIndex < BackupSectionLODModel.Sections.Num(); ++SectionIndex)
@@ -1825,23 +1832,27 @@ void FQuadricSkeletalMeshReduction::ReduceSkeletalMesh(USkeletalMesh& SkeletalMe
 		SectionData.MaterialIndex = BackupSectionLODModel.Sections[SectionIndex].MaterialIndex;
 		SectionData.bCastShadow = BackupSectionLODModel.Sections[SectionIndex].bCastShadow;
 		SectionData.bRecomputeTangent = BackupSectionLODModel.Sections[SectionIndex].bRecomputeTangent;
+		SectionData.bDisabled = BackupSectionLODModel.Sections[SectionIndex].bDisabled;
+		SectionData.GenerateUpToLodIndex = BackupSectionLODModel.Sections[SectionIndex].GenerateUpToLodIndex;
+		SectionData.ChunkedParentSectionIndex = BackupSectionLODModel.Sections[SectionIndex].ChunkedParentSectionIndex;
+		SectionData.OriginalDataSectionIndex = BackupSectionLODModel.Sections[SectionIndex].OriginalDataSectionIndex;
+		SectionData.MaterialMap = LODInfo->LODMaterialMap.IsValidIndex(SectionIndex) ? LODInfo->LODMaterialMap[SectionIndex] : INDEX_NONE;
+		if (SectionData.MaterialMap == SectionData.MaterialIndex)
+		{
+			//Remove any override if the value is the same
+			SectionData.MaterialMap = INDEX_NONE;
+		}
 	}
 
-	auto FillClothingData = [&SkeletalMeshResource, &LODIndex, bLODModelAdded](int32 &EnableSectionNumber, TArray<bool> &SectionStatus)
+	auto FillSectionMaterialSlot = [&SkeletalMeshResource, &LODIndex, bLODModelAdded](TArray<int32>& SectionMaterialSlot)
 	{
-		EnableSectionNumber = 0;
-		SectionStatus.Empty();
+		SectionMaterialSlot.Empty();
 		if (!bLODModelAdded && SkeletalMeshResource.LODModels.IsValidIndex(LODIndex))
 		{
 			int32 SectionNumber = SkeletalMeshResource.LODModels[LODIndex].Sections.Num();
-			SectionStatus.Reserve(SectionNumber);
 			for (int32 SectionIndex = 0; SectionIndex < SectionNumber; ++SectionIndex)
 			{
-				SectionStatus.Add(!SkeletalMeshResource.LODModels[LODIndex].Sections[SectionIndex].bDisabled);
-				if (SectionStatus[SectionIndex])
-				{
-					EnableSectionNumber++;
-				}
+				SectionMaterialSlot.Add(SkeletalMeshResource.LODModels[LODIndex].Sections[SectionIndex].OriginalDataSectionIndex);
 			}
 		}
 	};
@@ -1849,31 +1860,20 @@ void FQuadricSkeletalMeshReduction::ReduceSkeletalMesh(USkeletalMesh& SkeletalMe
 	// Unbind any existing clothing assets before we reimport the geometry
 	TArray<ClothingAssetUtils::FClothingAssetMeshBinding> ClothingBindings;
 	//Get a map of enable/disable sections
-	int32 OriginalSectionNumberBeforeReduction = 0;
-	TArray<bool> OriginalSectionEnableBeforeReduction;
+	TArray<int32> OriginalSectionMaterialSlot;
 
 	//Do not play with cloth if the LOD is added
 	if (!bLODModelAdded)
 	{
-		//Store the clothBinding
-		ClothingAssetUtils::GetMeshClothingAssetBindings(&SkeletalMesh, ClothingBindings, LODIndex);
-		FillClothingData(OriginalSectionNumberBeforeReduction, OriginalSectionEnableBeforeReduction);
-		//Unbind the Cloth for this LOD before we reduce it, we will put back the cloth after the reduction, if it still match the sections
-		for (ClothingAssetUtils::FClothingAssetMeshBinding& Binding : ClothingBindings)
-		{
-			if (Binding.LODIndex == LODIndex)
-			{
-				Binding.Asset->UnbindFromSkeletalMesh(&SkeletalMesh, Binding.LODIndex);
-			}
-		}
+		FLODUtilities::UnbindClothingAndBackup(&SkeletalMesh, ClothingBindings, LODIndex);
 	}
 
 	bool bReducingSourceModel = false;
-	//Reducing Base LOD, we need to use the temporary data so it can be iterative
-	if (BaseLOD == LODIndex && SkelResource->OriginalReductionSourceMeshData.IsValidIndex(BaseLOD) && !SkelResource->OriginalReductionSourceMeshData[BaseLOD]->IsEmpty())
+	//Reducing source LOD, we need to use the temporary data so it can be iterative
+	if (BaseLOD == LODIndex && (SkelResource->OriginalReductionSourceMeshData.IsValidIndex(BaseLOD) && !SkelResource->OriginalReductionSourceMeshData[BaseLOD]->IsEmpty()))
 	{
 		TMap<FString, TArray<FMorphTargetDelta>> TempLODMorphTargetData;
-		SkelResource->OriginalReductionSourceMeshData[BaseLOD]->LoadReductionData(*SrcModel, TempLODMorphTargetData);
+		SkelResource->OriginalReductionSourceMeshData[BaseLOD]->LoadReductionData(*SrcModel, TempLODMorphTargetData, &SkeletalMesh);
 		bReducingSourceModel = true;
 	}
 	else
@@ -1955,6 +1955,8 @@ void FQuadricSkeletalMeshReduction::ReduceSkeletalMesh(USkeletalMesh& SkeletalMe
 
 	// Swap out the old model.  
 	FSkeletalMeshImportData RawMesh;
+	ESkeletalMeshGeoImportVersions GeoImportVersion = ESkeletalMeshGeoImportVersions::Before_Versionning;
+	ESkeletalMeshSkinningImportVersions SkinningImportVersion = ESkeletalMeshSkinningImportVersions::Before_Versionning;
 	{
 		FSkeletalMeshLODModel* Old = LODModels[LODIndex];
 
@@ -1962,12 +1964,21 @@ void FQuadricSkeletalMeshReduction::ReduceSkeletalMesh(USkeletalMesh& SkeletalMe
 
 		if (!bReducingSourceModel && Old)
 		{
+			//We need to backup the original RawSkeletalMeshBulkData in case it was an imported LOD
+			if (!bLODModelAdded && !Old->RawSkeletalMeshBulkData.IsEmpty())
+			{
+				Old->RawSkeletalMeshBulkData.LoadRawMesh(RawMesh);
+				GeoImportVersion = Old->RawSkeletalMeshBulkData.GeoImportVersion;
+				SkinningImportVersion = Old->RawSkeletalMeshBulkData.SkinningImportVersion;
+			}
 			delete Old;
 		}
 		else if(bReducingSourceModel)
 		{
 			//In case we reduce the source model we want to keep the original import data
 			SrcModel->RawSkeletalMeshBulkData.LoadRawMesh(RawMesh);
+			GeoImportVersion = SrcModel->RawSkeletalMeshBulkData.GeoImportVersion;
+			SkinningImportVersion = SrcModel->RawSkeletalMeshBulkData.SkinningImportVersion;
 		}
 	}
 
@@ -1975,8 +1986,10 @@ void FQuadricSkeletalMeshReduction::ReduceSkeletalMesh(USkeletalMesh& SkeletalMe
 
 	// Reduce LOD model with SrcMesh
 
-	if (ReduceSkeletalLODModel(*SrcModel, *NewModel, SkeletalMesh.GetImportedBounds(), SkeletalMesh.RefSkeleton, Settings, ImportantBones, RelativeToRefPoseMatrices, LODIndex))
+	if (ReduceSkeletalLODModel(*SrcModel, *NewModel, SkeletalMesh.GetImportedBounds(), SkeletalMesh.RefSkeleton, Settings, ImportantBones, RelativeToRefPoseMatrices, LODIndex, bReducingSourceModel))
 	{
+		FSkeletalMeshLODInfo* ReducedLODInfoPtr = SkeletalMesh.GetLODInfo(LODIndex);
+		check(ReducedLODInfoPtr);
 		// Do any joint-welding / bone removal.
 
 		if (MeshBoneReductionInterface != NULL && MeshBoneReductionInterface->GetBoneReductionData(&SkeletalMesh, LODIndex, BonesToRemove))
@@ -1990,35 +2003,22 @@ void FQuadricSkeletalMeshReduction::ReduceSkeletalMesh(USkeletalMesh& SkeletalMe
 
 		if (bOldLodWasFromFile)
 		{
-			SkeletalMesh.GetLODInfo(LODIndex)->LODMaterialMap.Empty();
+			ReducedLODInfoPtr->LODMaterialMap.Empty();
 		}
-
-		// If base lod has a customized LODMaterialMap and this LOD doesn't (could have if changes are applied instead of freshly generated, copy over the data into new new LOD
-
-		if (SkeletalMesh.GetLODInfo(LODIndex)->LODMaterialMap.Num() == 0 && SkeletalMesh.GetLODInfo(BaseLOD)->LODMaterialMap.Num() != 0)
-		{
-			SkeletalMesh.GetLODInfo(LODIndex)->LODMaterialMap = SkeletalMesh.GetLODInfo(BaseLOD)->LODMaterialMap;
-		}
-		else
-		{
-			// Assuming the reducing step has set all material indices correctly, we double check if something went wrong
-			// make sure we don't have more materials
-			int32 TotalSectionCount = NewModel->Sections.Num();
-			if (SkeletalMesh.GetLODInfo(LODIndex)->LODMaterialMap.Num() > TotalSectionCount)
-			{
-				SkeletalMesh.GetLODInfo(LODIndex)->LODMaterialMap = SkeletalMesh.GetLODInfo(BaseLOD)->LODMaterialMap;
-				// Something went wrong during the reduce step during regenerate 					
-				check(SkeletalMesh.GetLODInfo(BaseLOD)->LODMaterialMap.Num() == TotalSectionCount || SkeletalMesh.GetLODInfo(BaseLOD)->LODMaterialMap.Num() == 0);
-			}
-		}
-
 		// Flag this LOD as having been simplified.
-		SkeletalMesh.GetLODInfo(LODIndex)->bHasBeenSimplified = true;
+		ReducedLODInfoPtr->bHasBeenSimplified = true;
 		SkeletalMesh.bHasBeenSimplified = true;
 		//Restore section data
 		FSkeletalMeshLODModel& ImportedModelLOD = SkeletalMesh.GetImportedModel()->LODModels[LODIndex];
 		TArray<bool> SectionMatched;
 		SectionMatched.AddZeroed(ImportedModelLOD.Sections.Num());
+		
+		if (!bLODModelAdded)
+		{
+			ImportedModelLOD.UserSectionsData = BackupUserSectionsData;
+			ImportedModelLOD.BuildStringID = BackupLodModelBuildStringID;
+		}
+
 		for (auto Kvp : BackupSectionIndexToSectionData)
 		{
 			const int32 SourceSectionIndex = Kvp.Key;
@@ -2031,10 +2031,24 @@ void FQuadricSkeletalMeshReduction::ReduceSkeletalMesh(USkeletalMesh& SkeletalMe
 				{
 					continue;
 				}
-				if (SectionData.MaterialIndex == ImportedModelLOD.Sections[SectionIndex].MaterialIndex)
+				if (SectionData.OriginalDataSectionIndex == ImportedModelLOD.Sections[SectionIndex].OriginalDataSectionIndex)
 				{
+					//Restore the LOD material map for this section
+					if (SectionData.MaterialMap != INDEX_NONE)
+					{
+						while (ReducedLODInfoPtr->LODMaterialMap.Num() <= SectionIndex)
+						{
+							ReducedLODInfoPtr->LODMaterialMap.Add(INDEX_NONE);
+						}
+						ReducedLODInfoPtr->LODMaterialMap[SectionIndex] = SectionData.MaterialMap;
+					}
+					//Restore the LODModel Section
 					ImportedModelLOD.Sections[SectionIndex].bCastShadow = SectionData.bCastShadow;
 					ImportedModelLOD.Sections[SectionIndex].bRecomputeTangent = SectionData.bRecomputeTangent;
+					ImportedModelLOD.Sections[SectionIndex].bDisabled = SectionData.bDisabled;
+					ImportedModelLOD.Sections[SectionIndex].GenerateUpToLodIndex = SectionData.GenerateUpToLodIndex;
+					ImportedModelLOD.Sections[SectionIndex].ChunkedParentSectionIndex = SectionData.ChunkedParentSectionIndex;
+					ImportedModelLOD.Sections[SectionIndex].OriginalDataSectionIndex = SectionData.OriginalDataSectionIndex;
 					SectionMatched[SectionIndex] = true;
 					break;
 				}
@@ -2087,47 +2101,22 @@ void FQuadricSkeletalMeshReduction::ReduceSkeletalMesh(USkeletalMesh& SkeletalMe
 	if (!bLODModelAdded)
 	{
 		//Get the number of enabled section
-		int32 SectionNumberAfterReduction = 0;
-		TArray<bool> SectionEnableAfterReduction;
-		FillClothingData(SectionNumberAfterReduction, SectionEnableAfterReduction);
+		TArray<int32> SectionMaterialSlotAfterReduction;
+		FillSectionMaterialSlot(SectionMaterialSlotAfterReduction);
 
-		//Put back the clothing for this newly reduce LOD only if the section count match.
-		if (ClothingBindings.Num() > 0 && OriginalSectionNumberBeforeReduction == SectionNumberAfterReduction)
+		//Put back the clothing for this newly reduce LOD
+		if (ClothingBindings.Num() > 0)
 		{
-			TArray<int32> RemapSectionIndex;
-			int32 SectionIndexTest = 0;
-			for (int32 SectionIndexRef = 0; SectionIndexRef < OriginalSectionEnableBeforeReduction.Num(); SectionIndexRef++)
-			{
-				int32& RemapValue = RemapSectionIndex.Add_GetRef(INDEX_NONE);
-				if (!OriginalSectionEnableBeforeReduction[SectionIndexRef])
-				{
-					continue;
-				}
-				for (; SectionIndexTest <= SectionIndexRef; SectionIndexTest++)
-				{
-					if (SectionEnableAfterReduction.IsValidIndex(SectionIndexTest) && SectionEnableAfterReduction[SectionIndexTest])
-					{
-						RemapValue = SectionIndexTest++;
-						break;
-					}
-				}
-			}
-
-			for (ClothingAssetUtils::FClothingAssetMeshBinding& Binding : ClothingBindings)
-			{
-				int32 RemapBindingSectionIndex = RemapSectionIndex[Binding.SectionIndex];
-				if (RemapBindingSectionIndex != INDEX_NONE && Binding.LODIndex == LODIndex && NewModel->Sections.IsValidIndex(RemapBindingSectionIndex))
-				{
-					Binding.Asset->BindToSkeletalMesh(&SkeletalMesh, Binding.LODIndex, RemapBindingSectionIndex, Binding.AssetInternalLodIndex, false);
-				}
-			}
+			FLODUtilities::RestoreClothingFromBackup(&SkeletalMesh, ClothingBindings, LODIndex);
 		}
 	}
 
-	if (bReducingSourceModel && RawMesh.Points.Num() > 0)
+	if ((bReducingSourceModel || !bLODModelAdded ) && RawMesh.Points.Num() > 0)
 	{
 		//Put back the original import data, we need it to allow inline reduction and skeletal mesh split workflow
 		SkeletalMeshResource.LODModels[LODIndex].RawSkeletalMeshBulkData.SaveRawMesh(RawMesh);
+		SkeletalMeshResource.LODModels[LODIndex].RawSkeletalMeshBulkData.GeoImportVersion = GeoImportVersion;
+		SkeletalMeshResource.LODModels[LODIndex].RawSkeletalMeshBulkData.SkinningImportVersion = SkinningImportVersion;
 	}
 
 	SkeletalMesh.CalculateRequiredBones(SkeletalMeshResource.LODModels[LODIndex], SkeletalMesh.RefSkeleton, &BonesToRemove);
diff --git a/Engine/Source/Developer/MeshBuilder/MeshBuilder.Build.cs b/Engine/Source/Developer/MeshBuilder/MeshBuilder.Build.cs
--- a/Engine/Source/Developer/MeshBuilder/MeshBuilder.Build.cs
+++ b/Engine/Source/Developer/MeshBuilder/MeshBuilder.Build.cs
@@ -16,17 +16,19 @@ namespace UnrealBuildTool.Rules
                     "Engine",
                     "RenderCore",
                     "MeshDescription",
-                    "RenderCore",
+					"StaticMeshDescription",
                     "MeshDescriptionOperations",
                     "MeshReductionInterface",
                     "RawMesh",
+					"MeshUtilities",
                     "MeshUtilitiesCommon",
+                    "ClothingSystemRuntimeNv",
+                    "MeshBoneReduction",
+                    "SkeletalMeshUtilitiesCommon",
+					"MeshBuilderCommon",
                 }
 			);
 
-			AddEngineThirdPartyPrivateStaticDependencies(Target, "nvTriStrip");
-            AddEngineThirdPartyPrivateStaticDependencies(Target, "ForsythTriOptimizer");
-	        AddEngineThirdPartyPrivateStaticDependencies(Target, "nvTessLib");
             AddEngineThirdPartyPrivateStaticDependencies(Target, "QuadricMeshReduction");
        }
 	}
diff --git a/Engine/Source/Developer/MeshBuilder/Private/SkeletalMeshBuilder.cpp b/Engine/Source/Developer/MeshBuilder/Private/SkeletalMeshBuilder.cpp
new file mode 100644
--- /dev/null
+++ b/Engine/Source/Developer/MeshBuilder/Private/SkeletalMeshBuilder.cpp
@@ -0,0 +1,197 @@
+// Copyright 1998-2019 Epic Games, Inc. All Rights Reserved.
+
+#include "SkeletalMeshBuilder.h"
+#include "Modules/ModuleManager.h"
+#include "MeshBoneReduction.h"
+#include "Engine/EngineTypes.h"
+#include "Engine/SkeletalMesh.h"
+#include "PhysicsEngine/BodySetup.h"
+#include "MeshDescription.h"
+#include "MeshDescriptionOperations.h"
+#include "MeshAttributes.h"
+#include "MeshDescriptionHelper.h"
+#include "MeshBuild.h"
+#include "Rendering/SkeletalMeshModel.h"
+#include "Rendering/SkeletalMeshLODModel.h"
+#include "GPUSkinVertexFactory.h"
+#include "ThirdPartyBuildOptimizationHelper.h"
+#include "Misc/ScopedSlowTask.h"
+#include "LODUtilities.h"
+#include "ClothingAsset.h"
+#include "MeshUtilities.h"
+#include "EditorFramework/AssetImportData.h"
+
+DEFINE_LOG_CATEGORY(LogSkeletalMeshBuilder);
+
+namespace SkeletalMeshBuilderOptimization
+{
+	void CacheOptimizeIndexBuffer(TArray<uint16>& Indices)
+	{
+		BuildOptimizationThirdParty::CacheOptimizeIndexBuffer(Indices);
+	}
+
+	void CacheOptimizeIndexBuffer(TArray<uint32>& Indices)
+	{
+		BuildOptimizationThirdParty::CacheOptimizeIndexBuffer(Indices);
+	}
+}
+
+struct InfluenceMap
+{
+	FORCEINLINE bool operator()(const float& A, const float& B) const
+	{
+		return B > A;
+	}
+};
+
+struct FSkeletalMeshVertInstanceIDAndZ
+{
+	FVertexInstanceID Index;
+	float Z;
+};
+
+FSkeletalMeshBuilder::FSkeletalMeshBuilder()
+{
+
+}
+
+
+bool FSkeletalMeshBuilder::Build(USkeletalMesh* SkeletalMesh, const int32 LODIndex, const bool bRegenDepLODs)
+{
+	check(SkeletalMesh->GetImportedModel());
+	check(SkeletalMesh->GetImportedModel()->LODModels.IsValidIndex(LODIndex));
+	check(SkeletalMesh->GetLODInfo(LODIndex) != nullptr);
+	
+	FSkeletalMeshLODInfo* LODInfo = SkeletalMesh->GetLODInfo(LODIndex);
+	//We want to backup in case the LODModel is regenerated, this data is use to validate in the UI if the ddc must be rebuild
+	const FString BackupBuildStringID = SkeletalMesh->GetImportedModel()->LODModels[LODIndex].BuildStringID;
+
+	const FReferenceSkeleton& RefSkeleton = SkeletalMesh->RefSkeleton;
+
+	FScopedSlowTask SlowTask(6.0f, NSLOCTEXT("SkeltalMeshBuilder", "BuildingSkeletalMeshLOD", "Building skeletal mesh LOD"));
+	SlowTask.MakeDialog();
+
+	//Prevent any PostEdit change during the build
+	FScopedSkeletalMeshPostEditChange ScopedPostEditChange(SkeletalMesh, false, false);
+	// Unbind any existing clothing assets before we reimport the geometry
+	TArray<ClothingAssetUtils::FClothingAssetMeshBinding> ClothingBindings;
+	FLODUtilities::UnbindClothingAndBackup(SkeletalMesh, ClothingBindings, LODIndex);
+
+	FSkeletalMeshImportData SkeletalMeshImportData;
+	float ProgressStepSize = SkeletalMesh->IsReductionActive(LODIndex) ? 0.5f : 1.0f;
+
+	//This scope define where we can use the LODModel, after a reduction the LODModel must be requery since it is a new instance
+	{
+		FSkeletalMeshLODModel& BuildLODModel = SkeletalMesh->GetImportedModel()->LODModels[LODIndex];
+		
+		BuildLODModel.RawSkeletalMeshBulkData.LoadRawMesh(SkeletalMeshImportData);
+
+		TArray<FVector> LODPoints;
+		TArray<SkeletalMeshImportData::FMeshWedge> LODWedges;
+		TArray<SkeletalMeshImportData::FMeshFace> LODFaces;
+		TArray<SkeletalMeshImportData::FVertInfluence> LODInfluences;
+		TArray<int32> LODPointToRawMap;
+		SkeletalMeshImportData.CopyLODImportData(LODPoints, LODWedges, LODFaces, LODInfluences, LODPointToRawMap);
+
+
+		//Build the skeletalmesh using mesh utilities module
+		IMeshUtilities::MeshBuildOptions Options;
+		Options.FillOptions(LODInfo->BuildSettings);
+		IMeshUtilities& MeshUtilities = FModuleManager::Get().LoadModuleChecked<IMeshUtilities>("MeshUtilities");
+
+		// Create skinning streams for NewModel.
+		SlowTask.EnterProgressFrame(1.0f);
+		MeshUtilities.BuildSkeletalMesh(
+			BuildLODModel,
+			RefSkeleton,
+			LODInfluences,
+			LODWedges,
+			LODFaces,
+			LODPoints,
+			LODPointToRawMap,
+			Options
+		);
+
+		//Re-Apply the user section changes, the UserSectionsData is map to original section and should match the builded LODModel
+		BuildLODModel.SyncronizeUserSectionsDataArray();
+
+		// Set texture coordinate count on the new model.
+		BuildLODModel.NumTexCoords = SkeletalMeshImportData.NumTexCoords;
+
+		//Re-apply the morph target
+		SlowTask.EnterProgressFrame(1.0f, NSLOCTEXT("SkeltalMeshBuilder", "RebuildMorphTarget", "Rebuilding morph targets..."));
+		if (SkeletalMeshImportData.MorphTargetNames.Num() > 0)
+		{
+			FLODUtilities::BuildMorphTargets(SkeletalMesh, SkeletalMeshImportData, LODIndex, !Options.bComputeNormals, !Options.bComputeTangents, Options.bUseMikkTSpace);
+		}
+
+		//Re-apply the alternate skinning it must be after the inline reduction
+		SlowTask.EnterProgressFrame(1.0f, NSLOCTEXT("SkeltalMeshBuilder", "RebuildAlternateSkinning", "Rebuilding alternate skinning..."));
+		const TArray<FSkinWeightProfileInfo>& SkinProfiles = SkeletalMesh->GetSkinWeightProfiles();
+		for (int32 SkinProfileIndex = 0; SkinProfileIndex < SkinProfiles.Num(); ++SkinProfileIndex)
+		{
+			const FSkinWeightProfileInfo& ProfileInfo = SkinProfiles[SkinProfileIndex];
+			FLODUtilities::UpdateAlternateSkinWeights(SkeletalMesh, ProfileInfo.Name, LODIndex, Options.OverlappingThresholds, !Options.bComputeNormals, !Options.bComputeTangents, Options.bUseMikkTSpace, Options.bComputeWeightedNormals);
+		}
+
+		
+		FSkeletalMeshUpdateContext UpdateContext;
+		UpdateContext.SkeletalMesh = SkeletalMesh;
+		//We are reduce ourself in this case we reduce ourself from the original data and return true.
+		if (SkeletalMesh->IsReductionActive(LODIndex))
+		{
+			SlowTask.EnterProgressFrame(ProgressStepSize, NSLOCTEXT("SkeltalMeshBuilder", "RegenerateLOD", "Regenerate LOD..."));
+			//Update the original reduction data since we just build a new LODModel.
+			if (LODInfo->ReductionSettings.BaseLOD == LODIndex && SkeletalMesh->GetImportedModel()->OriginalReductionSourceMeshData.IsValidIndex(LODIndex))
+			{
+				//Make the copy of the data only once until the ImportedModel change (re-imported)
+				SkeletalMesh->GetImportedModel()->OriginalReductionSourceMeshData[LODIndex]->EmptyBulkData();
+				TMap<FString, TArray<FMorphTargetDelta>> BaseLODMorphTargetData;
+				BaseLODMorphTargetData.Empty(SkeletalMesh->MorphTargets.Num());
+				for (UMorphTarget *MorphTarget : SkeletalMesh->MorphTargets)
+				{
+					if (!MorphTarget->HasDataForLOD(LODIndex))
+					{
+						continue;
+					}
+					TArray<FMorphTargetDelta>& MorphDeltasArray = BaseLODMorphTargetData.FindOrAdd(MorphTarget->GetFullName());
+					const FMorphTargetLODModel& BaseMorphModel = MorphTarget->MorphLODModels[LODIndex];
+					//Iterate each original morph target source index to fill the NewMorphTargetDeltas array with the TargetMatchData.
+					for (const FMorphTargetDelta& MorphDelta : BaseMorphModel.Vertices)
+					{
+						MorphDeltasArray.Add(MorphDelta);
+					}
+				}
+				//Copy the original SkeletalMesh LODModel
+				SkeletalMesh->GetImportedModel()->OriginalReductionSourceMeshData[LODIndex]->SaveReductionData(BuildLODModel, BaseLODMorphTargetData, SkeletalMesh);
+
+				if (LODIndex == 0)
+				{
+					SkeletalMesh->GetLODInfo(LODIndex)->SourceImportFilename = SkeletalMesh->AssetImportData->GetFirstFilename();
+				}
+			}
+			FLODUtilities::SimplifySkeletalMeshLOD(UpdateContext, LODIndex, false);
+		}
+	}
+
+	FSkeletalMeshLODModel& LODModelAfterReduction = SkeletalMesh->GetImportedModel()->LODModels[LODIndex];
+	
+	//Re-apply the clothing using the UserSectionsData, this will ensure we remap correctly the cloth if the reduction has change the number of sections
+	SlowTask.EnterProgressFrame(1.0f, NSLOCTEXT("SkeltalMeshBuilder", "RebuildClothing", "Rebuilding clothing..."));
+	FLODUtilities::RestoreClothingFromBackup(SkeletalMesh, ClothingBindings, LODIndex);
+
+	LODModelAfterReduction.SyncronizeUserSectionsDataArray();
+	LODModelAfterReduction.NumTexCoords = SkeletalMeshImportData.NumTexCoords;
+	LODModelAfterReduction.BuildStringID = BackupBuildStringID;
+
+	SlowTask.EnterProgressFrame(ProgressStepSize, NSLOCTEXT("SkeltalMeshBuilder", "RegenerateDependentLODs", "Regenerate Dependent LODs..."));
+	if (bRegenDepLODs)
+	{
+		//Regenerate dependent LODs
+		FLODUtilities::RegenerateDependentLODs(SkeletalMesh, LODIndex);
+	}
+
+	SlowTask.EnterProgressFrame();
+
+	return true;
+}
