Control Rig: Fix for NAN ensures coming from Control Rig nodes. Occured when bones that are excluded from the current LOD were having additive transforms applied to them, but these bones were not having their transforms reset.

diff --git a/Engine/Plugins/Experimental/ControlRig/Source/ControlRig/Private/AnimNode_ControlRigBase.cpp b/Engine/Plugins/Experimental/ControlRig/Source/ControlRig/Private/AnimNode_ControlRigBase.cpp
--- a/Engine/Plugins/Experimental/ControlRig/Source/ControlRig/Private/AnimNode_ControlRigBase.cpp
+++ b/Engine/Plugins/Experimental/ControlRig/Source/ControlRig/Private/AnimNode_ControlRigBase.cpp
@@ -91,6 +91,11 @@ void FAnimNode_ControlRigBase::UpdateInput(UControlRig* ControlRig, const FPoseC
 		// first I need to convert to local pose
 		MeshPoses.InitPose(InOutput.Pose);
 
+		// reset transforms here to prevent additive transforms from accumulating to INF
+		// we only update transforms from the mesh pose for bones in the current LOD, 
+		// so the reset here ensures excluded bones are also reset
+		ControlRig->GetBoneHierarchy().ResetTransforms();
+
 		// @re-think - now control rig contains init pose from their default hierarchy and current pose from this instance.
 		// we may need this init pose somewhere (instance refpose)
 		for (auto Iter = ControlRigBoneMapping.CreateConstIterator(); Iter; ++Iter)
