Control Rig: Re-init control rig when LOD changes to fix issues with transform constraint

diff --git a/Engine/Plugins/Experimental/ControlRig/Source/ControlRig/Private/AnimNode_ControlRigBase.cpp b/Engine/Plugins/Experimental/ControlRig/Source/ControlRig/Private/AnimNode_ControlRigBase.cpp
--- a/Engine/Plugins/Experimental/ControlRig/Source/ControlRig/Private/AnimNode_ControlRigBase.cpp
+++ b/Engine/Plugins/Experimental/ControlRig/Source/ControlRig/Private/AnimNode_ControlRigBase.cpp
@@ -299,6 +299,9 @@ void FAnimNode_ControlRigBase::CacheBones_AnyThread(const FAnimationCacheBonesCo
 				ControlRigCurveMapping.Add(CurveNames[Index], Index);
 			}
 		}
+
+		// re-init when LOD changes
+		ControlRig->Execute(EControlRigState::Init);
 	}
 }
 
