Sequencer: Control Rig: Hard to Select Controls in Viewport.
Control Meshes have same priority as other meshes so you need to select the mesh exactly(actually slightly off since it's a5x5 grid), to hit the control  Fixed this by adding a Priority Property to PrimitiveComponents so that we can set that directly when the hit proxies are created.

diff --git a/Engine/Plugins/Experimental/ControlRig/Source/ControlRigManipulation/Private/ControlRigGizmoActor.cpp b/Engine/Plugins/Experimental/ControlRig/Source/ControlRigManipulation/Private/ControlRigGizmoActor.cpp
--- a/Engine/Plugins/Experimental/ControlRig/Source/ControlRigManipulation/Private/ControlRigGizmoActor.cpp
+++ b/Engine/Plugins/Experimental/ControlRig/Source/ControlRigManipulation/Private/ControlRigGizmoActor.cpp
@@ -20,6 +20,9 @@ AControlRigGizmoActor::AControlRigGizmoActor(const FObjectInitializer& ObjectIni
 	StaticMeshComponent->Mobility = EComponentMobility::Movable;
 	StaticMeshComponent->SetGenerateOverlapEvents(false);
 	StaticMeshComponent->bUseDefaultCollision = true;
+#if WITH_EDITORONLY_DATA
+	StaticMeshComponent->HitProxyPriority = HPP_Wireframe;
+#endif
 
 	RootComponent = ActorRootComponent;
 	StaticMeshComponent->SetupAttachment(RootComponent);
diff --git a/Engine/Source/Runtime/Engine/Classes/Components/PrimitiveComponent.h b/Engine/Source/Runtime/Engine/Classes/Components/PrimitiveComponent.h
--- a/Engine/Source/Runtime/Engine/Classes/Components/PrimitiveComponent.h
+++ b/Engine/Source/Runtime/Engine/Classes/Components/PrimitiveComponent.h
@@ -20,6 +20,7 @@
 #include "Engine/TextureStreamingTypes.h"
 #include "AI/Navigation/NavRelevantInterface.h"
 #include "VT/RuntimeVirtualTextureEnum.h"
+#include "HitProxies.h"
 #include "PrimitiveComponent.generated.h"
 
 class AController;
@@ -523,6 +524,12 @@ public:
 	UPROPERTY()
 	TEnumAsByte<EHasCustomNavigableGeometry::Type> bHasCustomNavigableGeometry;
 
+public:
+#if WITH_EDITORONLY_DATA
+		UPROPERTY()
+			TEnumAsByte<enum EHitProxyPriority> HitProxyPriority;
+#endif
+
 private:
 #if WITH_EDITORONLY_DATA
 	UPROPERTY()
diff --git a/Engine/Source/Runtime/Engine/Private/Components/PrimitiveComponent.cpp b/Engine/Source/Runtime/Engine/Private/Components/PrimitiveComponent.cpp
--- a/Engine/Source/Runtime/Engine/Private/Components/PrimitiveComponent.cpp
+++ b/Engine/Source/Runtime/Engine/Private/Components/PrimitiveComponent.cpp
@@ -378,6 +378,7 @@ UPrimitiveComponent::UPrimitiveComponent(const FObjectInitializer& ObjectInitial
 	bReceiveMobileCSMShadows = true;
 #if WITH_EDITORONLY_DATA
 	bEnableAutoLODGeneration = true;
+	HitProxyPriority = HPP_World;
 #endif // WITH_EDITORONLY_DATA
 }
 
diff --git a/Engine/Source/Runtime/Engine/Private/PrimitiveSceneProxy.cpp b/Engine/Source/Runtime/Engine/Private/PrimitiveSceneProxy.cpp
--- a/Engine/Source/Runtime/Engine/Private/PrimitiveSceneProxy.cpp
+++ b/Engine/Source/Runtime/Engine/Private/PrimitiveSceneProxy.cpp
@@ -316,7 +316,11 @@ HHitProxy* FPrimitiveSceneProxy::CreateHitProxies(UPrimitiveComponent* Component
 		}
 		else
 		{
+#if WITH_EDITORONLY_DATA
+			ActorHitProxy = new HActor(Component->GetOwner(), Component, Component->HitProxyPriority);
+#else
 			ActorHitProxy = new HActor(Component->GetOwner(), Component);
+#endif
 		}
 		OutHitProxies.Add(ActorHitProxy);
 		return ActorHitProxy;
diff --git a/Engine/Source/Runtime/Engine/Private/SkeletalMesh.cpp b/Engine/Source/Runtime/Engine/Private/SkeletalMesh.cpp
--- a/Engine/Source/Runtime/Engine/Private/SkeletalMesh.cpp
+++ b/Engine/Source/Runtime/Engine/Private/SkeletalMesh.cpp
@@ -4645,7 +4645,7 @@ HHitProxy* FSkeletalMeshSceneProxy::CreateHitProxies(UPrimitiveComponent* Compon
 					}
 					else
 					{
-						ActorHitProxy = new HActor(Component->GetOwner(), Component, SectionIndex, MaterialIndex);
+						ActorHitProxy = new HActor(Component->GetOwner(), Component, Component->HitProxyPriority, SectionIndex, MaterialIndex);
 					}
 
 					// Set the hitproxy.
diff --git a/Engine/Source/Runtime/Engine/Private/StaticMeshRender.cpp b/Engine/Source/Runtime/Engine/Private/StaticMeshRender.cpp
--- a/Engine/Source/Runtime/Engine/Private/StaticMeshRender.cpp
+++ b/Engine/Source/Runtime/Engine/Private/StaticMeshRender.cpp
@@ -924,7 +924,7 @@ HHitProxy* FStaticMeshSceneProxy::CreateHitProxies(UPrimitiveComponent* Componen
 				}
 				else
 				{
-					ActorHitProxy = new HActor(Component->GetOwner(), Component, SectionIndex, MaterialIndex);
+					ActorHitProxy = new HActor(Component->GetOwner(), Component, Component->HitProxyPriority, SectionIndex, MaterialIndex);
 				}
 
 				FLODInfo::FSectionInfo& Section = LODs[LODIndex].Sections[SectionIndex];
diff --git a/Engine/Source/Runtime/Engine/Public/HitProxies.h b/Engine/Source/Runtime/Engine/Public/HitProxies.h
--- a/Engine/Source/Runtime/Engine/Public/HitProxies.h
+++ b/Engine/Source/Runtime/Engine/Public/HitProxies.h
@@ -7,6 +7,7 @@
 #pragma once
 
 #include "CoreMinimal.h"
+#include "UObject/ObjectMacros.h"
 #include "Templates/RefCounting.h"
 #include "GenericPlatform/ICursor.h"
 
@@ -18,6 +19,7 @@ class FReferenceCollector;
  * HPP_Wireframe - the priority of items that are drawn in wireframe, such as volumes
  * HPP_UI - the priority of the UI components such as the translation widget
  */
+UENUM()
 enum EHitProxyPriority
 {
 	HPP_World = 0,
