UE-88506: Crash on duplicate name

When creating a control rig for a skeletal mesh or skeleton twice through the contextual menu the second one fails / crashes.

Merging

//UE4/Dev-Anim/Engine/Plugins/Experimental/ControlRig/Source/ControlRigEditor/Private/ControlRigBlueprintActions.cpp

to //UE4/Release-4.25/Engine/Plugins/Experimental/ControlRig/Source/ControlRigEditor/Private/ControlRigBlueprintActions.cpp

diff --git a/Engine/Plugins/Experimental/ControlRig/Source/ControlRigEditor/Private/ControlRigBlueprintActions.cpp b/Engine/Plugins/Experimental/ControlRig/Source/ControlRigEditor/Private/ControlRigBlueprintActions.cpp
--- a/Engine/Plugins/Experimental/ControlRig/Source/ControlRigEditor/Private/ControlRigBlueprintActions.cpp
+++ b/Engine/Plugins/Experimental/ControlRig/Source/ControlRigEditor/Private/ControlRigBlueprintActions.cpp
@@ -132,7 +132,7 @@ UControlRigBlueprint* FControlRigBlueprintActions::CreateNewControlRigAsset(cons
 	FString UniqueAssetName;
 	AssetToolsModule.Get().CreateUniqueAssetName(InDesiredPackagePath, TEXT(""), UniquePackageName, UniqueAssetName);
 
-	UObject* NewAsset = AssetToolsModule.Get().CreateAsset(*UniqueAssetName, *InDesiredPackagePath, nullptr, Factory);
+	UObject* NewAsset = AssetToolsModule.Get().CreateAsset(*UniqueAssetName, *UniquePackageName, nullptr, Factory);
 	return Cast<UControlRigBlueprint>(NewAsset);
 }
 
@@ -163,7 +163,10 @@ UControlRigBlueprint* FControlRigBlueprintActions::CreateControlRigFromSkeletalM
 	}
 
 	UControlRigBlueprint* NewControlRigBlueprint = CreateNewControlRigAsset(PackagePath / ControlRigName);
-	check(NewControlRigBlueprint);
+	if (NewControlRigBlueprint == nullptr)
+	{
+		return nullptr;
+	}
 
 	NewControlRigBlueprint->HierarchyContainer.BoneHierarchy.ImportSkeleton(Skeleton->GetReferenceSkeleton(), NAME_None, true, true, false, false);
 	NewControlRigBlueprint->HierarchyContainer.CurveContainer.ImportCurvesFromSkeleton(Skeleton, NAME_None, true, false, false);
