UE-88396: Selection state in Blueprint editor tool bar not maintained after compile

Merging

//UE4/Dev-Anim/Engine/...

to //UE4/Release-4.25/Engine/...

diff --git a/Engine/Plugins/Experimental/ControlRig/Source/ControlRigEditor/Private/Editor/ControlRigEditor.cpp b/Engine/Plugins/Experimental/ControlRig/Source/ControlRigEditor/Private/Editor/ControlRigEditor.cpp
--- a/Engine/Plugins/Experimental/ControlRig/Source/ControlRigEditor/Private/Editor/ControlRigEditor.cpp
+++ b/Engine/Plugins/Experimental/ControlRig/Source/ControlRigEditor/Private/Editor/ControlRigEditor.cpp
@@ -659,6 +659,20 @@ void FControlRigEditor::SetDetailObjects(const TArray<UObject*>& InObjects)
 {
 	ClearDetailObject();
 
+	if (InObjects.Num() == 1)
+	{
+		if (InObjects[0]->GetClass()->GetDefaultObject() == InObjects[0])
+		{
+			EditClassDefaults_Clicked();
+			return;
+		}
+		else if (InObjects[0] == GetBlueprintObj())
+		{
+			EditGlobalOptions_Clicked();
+			return;
+		}
+	}
+
 	RigElementInDetailPanel = FRigElementKey();
 	StructToDisplay.Reset();
 	SKismetInspector::FShowDetailsOptions Options;
@@ -823,6 +837,8 @@ void FControlRigEditor::ClearDetailObject()
 
 	Inspector->ShowDetailsForObjects(TArray<UObject*>());
 	Inspector->ShowSingleStruct(TSharedPtr<FStructOnScope>());
+
+	SetUISelectionState(FBlueprintEditor::SelectionState_Graph);
 }
 
 
diff --git a/Engine/Source/Editor/Kismet/Private/SKismetInspector.cpp b/Engine/Source/Editor/Kismet/Private/SKismetInspector.cpp
--- a/Engine/Source/Editor/Kismet/Private/SKismetInspector.cpp
+++ b/Engine/Source/Editor/Kismet/Private/SKismetInspector.cpp
@@ -540,6 +540,8 @@ void SKismetInspector::UpdateFromSingleStruct(const TSharedPtr<FStructOnScope>&
 {
 	if (StructureDetailsView.IsValid())
 	{
+		SelectedObjects.Empty();
+
 		// Update our context-sensitive editing widget
 		ContextualEditingBorderWidget->SetContent(StructureDetailsView->GetWidget().ToSharedRef());
 	}
