Sequencer: Control Rig: Fix for leaving saving/leaving/restoring sequence. Need to recreate control rig on proxy.

diff --git a/Engine/Plugins/Experimental/ControlRig/Source/ControlRig/Private/Sequencer/ControlRigLayerInstance.cpp b/Engine/Plugins/Experimental/ControlRig/Source/ControlRig/Private/Sequencer/ControlRigLayerInstance.cpp
--- a/Engine/Plugins/Experimental/ControlRig/Source/ControlRig/Private/Sequencer/ControlRigLayerInstance.cpp
+++ b/Engine/Plugins/Experimental/ControlRig/Source/ControlRig/Private/Sequencer/ControlRigLayerInstance.cpp
@@ -82,6 +82,11 @@ void UControlRigLayerInstance::AddControlRigTrack(int32 ControlRigID, UControlRi
 	GetProxyOnGameThread<FControlRigLayerInstanceProxy>().AddControlRigTrack(ControlRigID, InControlRig);
 }
 
+bool UControlRigLayerInstance::HasControlRigTrack(int32 ControlRigID)
+{
+	return GetProxyOnGameThread<FControlRigLayerInstanceProxy>().HasControlRigTrack(ControlRigID);
+}
+
 void UControlRigLayerInstance::UpdateControlRigTrack(int32 ControlRigID, float Weight, const FControlRigIOSettings& InputSettings, bool bExecute)
 {
 	GetProxyOnGameThread<FControlRigLayerInstanceProxy>().UpdateControlRigTrack(ControlRigID, Weight, InputSettings, bExecute);
diff --git a/Engine/Plugins/Experimental/ControlRig/Source/ControlRig/Private/Sequencer/ControlRigLayerInstanceProxy.cpp b/Engine/Plugins/Experimental/ControlRig/Source/ControlRig/Private/Sequencer/ControlRigLayerInstanceProxy.cpp
--- a/Engine/Plugins/Experimental/ControlRig/Source/ControlRig/Private/Sequencer/ControlRigLayerInstanceProxy.cpp
+++ b/Engine/Plugins/Experimental/ControlRig/Source/ControlRig/Private/Sequencer/ControlRigLayerInstanceProxy.cpp
@@ -148,6 +148,12 @@ void FControlRigLayerInstanceProxy::AddControlRigTrack(int32 ControlRigID, UCont
 	Node->Initialize_AnyThread(FAnimationInitializeContext(this));
 }
 
+bool FControlRigLayerInstanceProxy::HasControlRigTrack(int32 ControlRigID)
+{
+	FAnimNode_ControlRig_ExternalSource* Node = FindControlRigNode(ControlRigID);
+	return (Node != nullptr) ? true : false;
+}
+
 void FControlRigLayerInstanceProxy::ResetControlRigTracks()
 {
 	SequencerToControlRigNodeMap.Reset();
diff --git a/Engine/Plugins/Experimental/ControlRig/Source/ControlRig/Private/Sequencer/ControlRigLayerInstanceProxy.h b/Engine/Plugins/Experimental/ControlRig/Source/ControlRig/Private/Sequencer/ControlRigLayerInstanceProxy.h
--- a/Engine/Plugins/Experimental/ControlRig/Source/ControlRig/Private/Sequencer/ControlRigLayerInstanceProxy.h
+++ b/Engine/Plugins/Experimental/ControlRig/Source/ControlRig/Private/Sequencer/ControlRigLayerInstanceProxy.h
@@ -80,6 +80,7 @@ public:
 	void AddControlRigTrack(int32 ControlRigID, UControlRig* InControlRig);
 	void UpdateControlRigTrack(int32 ControlRigID, float Weight, const FControlRigIOSettings& InputSettings, bool bExecute);
 	void RemoveControlRigTrack(int32 ControlRigID);
+	bool HasControlRigTrack(int32 ControlRigID);
 	void ResetControlRigTracks();
 
 	/** Sequencer AnimInstance Interface */
diff --git a/Engine/Plugins/Experimental/ControlRig/Source/ControlRig/Private/Sequencer/MovieSceneControlRigParameterTemplate.cpp b/Engine/Plugins/Experimental/ControlRig/Source/ControlRig/Private/Sequencer/MovieSceneControlRigParameterTemplate.cpp
--- a/Engine/Plugins/Experimental/ControlRig/Source/ControlRig/Private/Sequencer/MovieSceneControlRigParameterTemplate.cpp
+++ b/Engine/Plugins/Experimental/ControlRig/Source/ControlRig/Private/Sequencer/MovieSceneControlRigParameterTemplate.cpp
@@ -370,7 +370,7 @@ namespace FControlRigBindingHelper
 			bool bWasCreated = false;
 			if (UControlRigLayerInstance* AnimInstance = FAnimCustomInstanceHelper::BindToSkeletalMeshComponent<UControlRigLayerInstance>(SkeletalMeshComponent, bWasCreated))
 			{
-				if (bWasCreated)
+				if (bWasCreated || !AnimInstance->HasControlRigTrack(ControlRig->GetUniqueID()))
 				{
 					AnimInstance->RecalcRequiredBones();
 					AnimInstance->AddControlRigTrack(ControlRig->GetUniqueID(), ControlRig);
diff --git a/Engine/Plugins/Experimental/ControlRig/Source/ControlRig/Public/Sequencer/ControlRigLayerInstance.h b/Engine/Plugins/Experimental/ControlRig/Source/ControlRig/Public/Sequencer/ControlRigLayerInstance.h
--- a/Engine/Plugins/Experimental/ControlRig/Source/ControlRig/Public/Sequencer/ControlRigLayerInstance.h
+++ b/Engine/Plugins/Experimental/ControlRig/Source/ControlRig/Public/Sequencer/ControlRigLayerInstance.h
@@ -24,6 +24,7 @@ public:
 	void AddControlRigTrack(int32 ControlRigID, UControlRig* InControlRig);
 	void UpdateControlRigTrack(int32 ControlRigID, float Weight, const FControlRigIOSettings& InputSettings, bool bExecute);
 	void RemoveControlRigTrack(int32 ControlRigID);
+	bool HasControlRigTrack(int32 ControlRigID);
 	void ResetControlRigTracks();
 
 	/** Sequencer AnimInstance Interface */
