Sequencer:Control Rig: Puffin fixes: Animatable checks, Vector2 keying fix, another Animatable check, and need to recreate channels on load after control rig is initialized. Also turn off display axes.

diff --git a/Engine/Plugins/Experimental/ControlRig/Source/ControlRig/Private/Sequencer/MovieSceneControlRigParameterSection.cpp b/Engine/Plugins/Experimental/ControlRig/Source/ControlRig/Private/Sequencer/MovieSceneControlRigParameterSection.cpp
--- a/Engine/Plugins/Experimental/ControlRig/Source/ControlRig/Private/Sequencer/MovieSceneControlRigParameterSection.cpp
+++ b/Engine/Plugins/Experimental/ControlRig/Source/ControlRig/Private/Sequencer/MovieSceneControlRigParameterSection.cpp
@@ -138,10 +138,10 @@ struct FParameterVectorChannelEditorData
 			MetaData[3].bCanCollapseToTrack = false;
 		}
 
-		ExternalValues[0].OnGetExternalValue = [InControlRig, InName,NumChannels](UObject& InObject, FTrackInstancePropertyBindings* Bindings) { return ExtractChannelX(InObject, Bindings, InControlRig, InName, NumChannels); };
-		ExternalValues[1].OnGetExternalValue = [InControlRig, InName,NumChannels](UObject& InObject, FTrackInstancePropertyBindings* Bindings) { return ExtractChannelY(InObject, Bindings, InControlRig, InName, NumChannels); };
-		ExternalValues[2].OnGetExternalValue = [InControlRig, InName,NumChannels](UObject& InObject, FTrackInstancePropertyBindings* Bindings) { return ExtractChannelZ(InObject, Bindings, InControlRig, InName, NumChannels); };
-		ExternalValues[3].OnGetExternalValue = [InControlRig, InName,NumChannels](UObject& InObject, FTrackInstancePropertyBindings* Bindings) { return ExtractChannelW(InObject, Bindings,InControlRig, InName, NumChannels); };
+		ExternalValues[0].OnGetExternalValue = [InControlRig, InName,NumChannels](UObject& InObject, FTrackInstancePropertyBindings* Bindings) { return ExtractChannelX(InObject, InControlRig, InName, NumChannels); };
+		ExternalValues[1].OnGetExternalValue = [InControlRig, InName,NumChannels](UObject& InObject, FTrackInstancePropertyBindings* Bindings) { return ExtractChannelY(InObject, InControlRig, InName, NumChannels); };
+		ExternalValues[2].OnGetExternalValue = [InControlRig, InName,NumChannels](UObject& InObject, FTrackInstancePropertyBindings* Bindings) { return ExtractChannelZ(InObject, InControlRig, InName, NumChannels); };
+		ExternalValues[3].OnGetExternalValue = [InControlRig, InName,NumChannels](UObject& InObject, FTrackInstancePropertyBindings* Bindings) { return ExtractChannelW(InObject, InControlRig, InName, NumChannels); };
 
 		ExternalValues[0].OnGetCurrentValueAndWeight = [InName, NumChannels](UObject* Object, UMovieSceneSection*  SectionToKey, FFrameNumber KeyTime, FFrameRate TickResolution, FMovieSceneRootEvaluationTemplateInstance& RootTemplate,
 			float& OutValue, float& OutWeight) { GetChannelValueAndWeight(InName, NumChannels, 0, Object, SectionToKey, KeyTime, TickResolution, RootTemplate, OutValue, OutWeight); };
@@ -154,7 +154,7 @@ struct FParameterVectorChannelEditorData
 
 	}
 
-	static FVector4 GetPropertyValue(UControlRig* ControlRig, FName ParameterName, UObject& InObject, FTrackInstancePropertyBindings& Bindings, int32 NumChannels)
+	static FVector4 GetPropertyValue(UControlRig* ControlRig, FName ParameterName, UObject& InObject,int32 NumChannels)
 	{
 		if (ControlRig)
 		{
@@ -181,21 +181,21 @@ struct FParameterVectorChannelEditorData
 		return FVector4();
 	}
 
-	static TOptional<float> ExtractChannelX(UObject& InObject, FTrackInstancePropertyBindings* Bindings, UControlRig* ControlRig, FName ParameterName, int32 NumChannels)
+	static TOptional<float> ExtractChannelX(UObject& InObject, UControlRig* ControlRig, FName ParameterName, int32 NumChannels)
 	{
-		return Bindings ? GetPropertyValue(ControlRig, ParameterName, InObject, *Bindings, NumChannels).X : TOptional<float>();
+		return GetPropertyValue(ControlRig, ParameterName, InObject, NumChannels).X;
 	}
-	static TOptional<float> ExtractChannelY(UObject& InObject, FTrackInstancePropertyBindings* Bindings, UControlRig* ControlRig, FName ParameterName, int32 NumChannels)
+	static TOptional<float> ExtractChannelY(UObject& InObject, UControlRig* ControlRig, FName ParameterName, int32 NumChannels)
 	{
-		return Bindings ? GetPropertyValue(ControlRig, ParameterName, InObject, *Bindings, NumChannels).Y : TOptional<float>();
+		return GetPropertyValue(ControlRig, ParameterName, InObject, NumChannels).Y;
 	}
-	static TOptional<float> ExtractChannelZ(UObject& InObject, FTrackInstancePropertyBindings* Bindings, UControlRig* ControlRig, FName ParameterName, int32 NumChannels)
+	static TOptional<float> ExtractChannelZ(UObject& InObject, UControlRig* ControlRig, FName ParameterName, int32 NumChannels)
 	{
-		return Bindings ? GetPropertyValue(ControlRig, ParameterName, InObject, *Bindings, NumChannels).Z : TOptional<float>();
+		return GetPropertyValue(ControlRig, ParameterName, InObject, NumChannels).Z;
 	}
-	static TOptional<float> ExtractChannelW(UObject& InObject, FTrackInstancePropertyBindings* Bindings, UControlRig* ControlRig, FName ParameterName, int32 NumChannels)
+	static TOptional<float> ExtractChannelW(UObject& InObject, UControlRig* ControlRig, FName ParameterName, int32 NumChannels)
 	{
-		return Bindings ? GetPropertyValue(ControlRig, ParameterName, InObject, *Bindings, NumChannels).W : TOptional<float>();
+		return GetPropertyValue(ControlRig, ParameterName, InObject, NumChannels).W;
 	}
 
 	static void GetChannelValueAndWeight(FName ParameterName, int32 NumChannels, int32 Index, UObject* Object, UMovieSceneSection*  SectionToKey, FFrameNumber KeyTime, FFrameRate TickResolution, FMovieSceneRootEvaluationTemplateInstance& RootTemplate,
@@ -951,6 +951,11 @@ void UMovieSceneControlRigParameterSection::ReconstructChannelProxy(bool bForce)
 			int32 ChannelIndex = 0;
 			for (const FRigControl& RigControl : SortedControls)
 			{
+
+				if (!RigControl.bAnimatable)
+				{
+					continue;
+				}
 #if WITH_EDITOR
 				switch (RigControl.ControlType)
 				{
diff --git a/Engine/Plugins/Experimental/ControlRig/Source/ControlRig/Private/Sequencer/MovieSceneControlRigParameterTemplate.cpp b/Engine/Plugins/Experimental/ControlRig/Source/ControlRig/Private/Sequencer/MovieSceneControlRigParameterTemplate.cpp
--- a/Engine/Plugins/Experimental/ControlRig/Source/ControlRig/Private/Sequencer/MovieSceneControlRigParameterTemplate.cpp
+++ b/Engine/Plugins/Experimental/ControlRig/Source/ControlRig/Private/Sequencer/MovieSceneControlRigParameterTemplate.cpp
@@ -759,7 +759,7 @@ struct TControlRigParameterActuatorVector2D : TMovieSceneBlendingActuator<FContr
 		FVector2DInterrogationData Data;
 		Data.Val = InValue.Value;
 		Data.ParameterName = ParameterName;
-		InterrogationData.Add(FVector2DInterrogationData(Data), UMovieSceneControlRigParameterSection::GetVectorInterrogationKey());
+		InterrogationData.Add(FVector2DInterrogationData(Data), UMovieSceneControlRigParameterSection::GetVector2DInterrogationKey());
 	}
 
 	UControlRig *ControlRig;
diff --git a/Engine/Plugins/Experimental/ControlRig/Source/ControlRig/Private/Sequencer/MovieSceneControlRigParameterTrack.cpp b/Engine/Plugins/Experimental/ControlRig/Source/ControlRig/Private/Sequencer/MovieSceneControlRigParameterTrack.cpp
--- a/Engine/Plugins/Experimental/ControlRig/Source/ControlRig/Private/Sequencer/MovieSceneControlRigParameterTrack.cpp
+++ b/Engine/Plugins/Experimental/ControlRig/Source/ControlRig/Private/Sequencer/MovieSceneControlRigParameterTrack.cpp
@@ -59,6 +59,11 @@ UMovieSceneSection* UMovieSceneControlRigParameterTrack::CreateNewSection()
 		
 		for (const FRigControl& RigControl : SortedControls)
 		{
+			if (!RigControl.bAnimatable)
+			{
+				continue;
+			}
+
 			switch (RigControl.ControlType)
 			{
 				case ERigControlType::Float:
@@ -407,6 +412,17 @@ void UMovieSceneControlRigParameterTrack::PostLoad()
 	{
 		ControlRig->Initialize();
 		ControlRig->CreateRigControlsForCurveContainer();
+		for (UMovieSceneSection * Section: Sections)
+		{
+			if (Section)
+			{
+				UMovieSceneControlRigParameterSection* CRSection = Cast<UMovieSceneControlRigParameterSection>(Section);
+				if (CRSection)
+				{
+					CRSection->ReconstructChannelProxy(true);
+				}
+			}
+		}
 	}
 }
 
diff --git a/Engine/Plugins/Experimental/ControlRig/Source/ControlRigEditor/Private/EditMode/ControlRigControlsProxy.cpp b/Engine/Plugins/Experimental/ControlRig/Source/ControlRigEditor/Private/EditMode/ControlRigControlsProxy.cpp
--- a/Engine/Plugins/Experimental/ControlRig/Source/ControlRigEditor/Private/EditMode/ControlRigControlsProxy.cpp
+++ b/Engine/Plugins/Experimental/ControlRig/Source/ControlRigEditor/Private/EditMode/ControlRigControlsProxy.cpp
@@ -233,7 +233,8 @@ void UControlRigVectorControlProxy::SetKey(const IPropertyHandle& KeyedPropertyH
 void UControlRigVector2DControlProxy::PostEditChangeProperty(struct FPropertyChangedEvent& PropertyChangedEvent)
 {
 	Super::PostEditChangeProperty(PropertyChangedEvent);
-	if (PropertyChangedEvent.Property && PropertyChangedEvent.Property->GetFName() == GET_MEMBER_NAME_CHECKED(UControlRigVector2DControlProxy, Vector2D))
+	if ((PropertyChangedEvent.Property && PropertyChangedEvent.Property->GetFName() == GET_MEMBER_NAME_CHECKED(UControlRigVector2DControlProxy, Vector2D))
+		|| ((PropertyChangedEvent.MemberProperty && PropertyChangedEvent.MemberProperty->GetFName() == GET_MEMBER_NAME_CHECKED(UControlRigVector2DControlProxy, Vector2D))))
 	{
 		if (RigControl && ControlRig.IsValid())
 		{
@@ -341,48 +342,48 @@ void UControlRigDetailPanelControlProxies::AddProxy(const FName& Name, UControlR
 	UControlRigControlsProxy* Proxy = FindProxy(Name);
 	if (!Proxy && RigControl != nullptr)
 	{
-		switch(RigControl->ControlType)
+		switch (RigControl->ControlType)
 		{
-			case ERigControlType::Transform:
-			{
-				Proxy = NewObject<UControlRigTransformControlProxy>(GetTransientPackage(), NAME_None);
-				break;
+		case ERigControlType::Transform:
+		{
+			Proxy = NewObject<UControlRigTransformControlProxy>(GetTransientPackage(), NAME_None);
+			break;
 
-			}
-			case ERigControlType::TransformNoScale:
-			{
-				Proxy = NewObject<UControlRigTransformNoScaleControlProxy>(GetTransientPackage(), NAME_None);
-				break;
+		}
+		case ERigControlType::TransformNoScale:
+		{
+			Proxy = NewObject<UControlRigTransformNoScaleControlProxy>(GetTransientPackage(), NAME_None);
+			break;
 
-			}
-			case ERigControlType::Float:
-			{
-				Proxy = NewObject<UControlRigFloatControlProxy>(GetTransientPackage(), NAME_None);
-				break;
+		}
+		case ERigControlType::Float:
+		{
+			Proxy = NewObject<UControlRigFloatControlProxy>(GetTransientPackage(), NAME_None);
+			break;
 
-			}
-			case ERigControlType::Position:
-			case ERigControlType::Rotator:
-			case ERigControlType::Scale:
-			{
-				Proxy = NewObject<UControlRigVectorControlProxy>(GetTransientPackage(), NAME_None);
-				break;
+		}
+		case ERigControlType::Position:
+		case ERigControlType::Rotator:
+		case ERigControlType::Scale:
+		{
+			Proxy = NewObject<UControlRigVectorControlProxy>(GetTransientPackage(), NAME_None);
+			break;
 
-			}
-			case ERigControlType::Vector2D:
-			{
-				Proxy = NewObject<UControlRigVector2DControlProxy>(GetTransientPackage(), NAME_None);
-				break;
+		}
+		case ERigControlType::Vector2D:
+		{
+			Proxy = NewObject<UControlRigVector2DControlProxy>(GetTransientPackage(), NAME_None);
+			break;
 
-			}
-			case ERigControlType::Bool:
-			{
-				Proxy = NewObject<UControlRigBoolControlProxy>(GetTransientPackage(), NAME_None);
-				break;
+		}
+		case ERigControlType::Bool:
+		{
+			Proxy = NewObject<UControlRigBoolControlProxy>(GetTransientPackage(), NAME_None);
+			break;
 
-			}
-			default:
-				break;
+		}
+		default:
+			break;
 		}
 		if (Proxy)
 		{
@@ -410,7 +411,7 @@ void UControlRigDetailPanelControlProxies::RemoveProxy(const FName& Name)
 
 void UControlRigDetailPanelControlProxies::RemoveAllProxies()
 {
-	for (TPair<FName, UControlRigControlsProxy * >Pair: AllProxies)
+	for (TPair<FName, UControlRigControlsProxy * >Pair : AllProxies)
 	{
 		UControlRigControlsProxy* ExistingProxy = Pair.Value;
 		ExistingProxy->Rename(nullptr, GetTransientPackage(), REN_ForceNoResetLoaders);
@@ -427,7 +428,7 @@ void UControlRigDetailPanelControlProxies::RecreateAllProxies(UControlRig* InCon
 	for (const FRigControl& RigControl : Controls)
 	{
 		FRigControl* PRigControl = const_cast<FRigControl*>(&RigControl);
-		if(PRigControl->bGizmoEnabled)
+		if(PRigControl->bGizmoEnabled && PRigControl->bAnimatable)
 		{
 			AddProxy(RigControl.Name, InControlRig, PRigControl);
 		}
diff --git a/Engine/Plugins/Experimental/ControlRig/Source/ControlRigEditor/Private/EditMode/ControlRigEditModeSettings.h b/Engine/Plugins/Experimental/ControlRig/Source/ControlRigEditor/Private/EditMode/ControlRigEditModeSettings.h
--- a/Engine/Plugins/Experimental/ControlRig/Source/ControlRigEditor/Private/EditMode/ControlRigEditModeSettings.h
+++ b/Engine/Plugins/Experimental/ControlRig/Source/ControlRigEditor/Private/EditMode/ControlRigEditModeSettings.h
@@ -16,7 +16,7 @@ class UControlRigEditModeSettings : public UObject
 	UControlRigEditModeSettings()
 		: bDisplayHierarchy(false)
 		, bHideManipulators(false)
-		, bDisplayAxesOnSelection(true)
+		, bDisplayAxesOnSelection(false)
 		, AxisScale(10.f)
 		, bCoordSystemPerWidgetMode(true)
 		, bOnlySelectRigControls(false)
diff --git a/Engine/Plugins/Experimental/ControlRig/Source/ControlRigEditor/Private/Sequencer/ControlRigParameterTrackEditor.cpp b/Engine/Plugins/Experimental/ControlRig/Source/ControlRigEditor/Private/Sequencer/ControlRigParameterTrackEditor.cpp
--- a/Engine/Plugins/Experimental/ControlRig/Source/ControlRigEditor/Private/Sequencer/ControlRigParameterTrackEditor.cpp
+++ b/Engine/Plugins/Experimental/ControlRig/Source/ControlRigEditor/Private/Sequencer/ControlRigParameterTrackEditor.cpp
@@ -55,10 +55,7 @@
 #include "Sequencer/ControlRigSortedControls.h"
 #include "TransformNoScale.h"
 
-
-
 #define LOCTEXT_NAMESPACE "FControlRigParameterTrackEditor"
-
 static USkeleton* GetSkeletonFromComponent(UActorComponent* InComponent)
 {
 	USkeletalMeshComponent* SkeletalMeshComp = Cast<USkeletalMeshComponent>(InComponent);
@@ -866,8 +863,11 @@ void FControlRigParameterTrackEditor::GetControlRigKeys(IControlRigManipulatable
 	int32 BoolChannelIndex = 0;
 	for (const FRigControl& RigControl : Controls)
 	{
+		if (!RigControl.bAnimatable)
+		{
+			continue;
+		}
 		bool bSetKey = RigControl.Name == ParameterName;
-
 		switch (RigControl.ControlType)
 		{
 		case ERigControlType::Bool:
@@ -959,7 +959,7 @@ void FControlRigParameterTrackEditor::GetControlRigKeys(IControlRigManipulatable
 
 FKeyPropertyResult FControlRigParameterTrackEditor::AddKeysToControlRigHandle(USceneComponent *InSceneComp, IControlRigManipulatable* Manip,
 	FGuid ObjectHandle, FFrameNumber KeyTime, FGeneratedTrackKeys& GeneratedKeys,
-	ESequencerKeyMode KeyMode, TSubclassOf<UMovieSceneTrack> TrackClass, FName ControlRigName)
+	ESequencerKeyMode KeyMode, TSubclassOf<UMovieSceneTrack> TrackClass, FName ControlRigName, FName RigControlName)
 {
 	EAutoChangeMode AutoChangeMode = GetSequencer()->GetAutoChangeMode();
 	EAllowEditsMode AllowEditsMode = GetSequencer()->GetAllowEditsMode();
@@ -1006,7 +1006,7 @@ FKeyPropertyResult FControlRigParameterTrackEditor::AddKeysToControlRigHandle(US
 		{
 			if (!bTrackCreated)
 			{
-				ModifyOurGeneratedKeysByCurrentAndWeight(InSceneComp, Manip, Track, SectionToKey, KeyTime, GeneratedKeys, Weight);
+				ModifyOurGeneratedKeysByCurrentAndWeight(InSceneComp, Manip, RigControlName, Track, SectionToKey, KeyTime, GeneratedKeys, Weight);
 			}
 			UMovieSceneControlRigParameterSection* ParamSection = Cast<UMovieSceneControlRigParameterSection>(SectionToKey);
 			if (!ParamSection->GetDoNotKey())
@@ -1022,7 +1022,7 @@ FKeyPropertyResult FControlRigParameterTrackEditor::AddKeysToControlRigHandle(US
 
 FKeyPropertyResult FControlRigParameterTrackEditor::AddKeysToControlRig(
 	USceneComponent *InSceneComp, IControlRigManipulatable* Manip, FFrameNumber KeyTime, FGeneratedTrackKeys& GeneratedKeys,
-	ESequencerKeyMode KeyMode, TSubclassOf<UMovieSceneTrack> TrackClass, FName PropertyName)
+	ESequencerKeyMode KeyMode, TSubclassOf<UMovieSceneTrack> TrackClass, FName ControlRigName, FName RigControlName)
 {
 	FKeyPropertyResult KeyPropertyResult;
 	EAutoChangeMode AutoChangeMode = GetSequencer()->GetAutoChangeMode();
@@ -1038,7 +1038,7 @@ FKeyPropertyResult FControlRigParameterTrackEditor::AddKeysToControlRig(
 	KeyPropertyResult.bHandleCreated = HandleResult.bWasCreated;
 	if (ObjectHandle.IsValid())
 	{
-		KeyPropertyResult |= AddKeysToControlRigHandle(InSceneComp, Manip, ObjectHandle, KeyTime, GeneratedKeys, KeyMode, TrackClass, PropertyName);
+		KeyPropertyResult |= AddKeysToControlRigHandle(InSceneComp, Manip, ObjectHandle, KeyTime, GeneratedKeys, KeyMode, TrackClass, ControlRigName, RigControlName);
 	}
 
 	return KeyPropertyResult;
@@ -1079,7 +1079,7 @@ void FControlRigParameterTrackEditor::AddControlKeys(USceneComponent *InSceneCom
 
 	auto OnKeyProperty = [=](FFrameNumber Time) -> FKeyPropertyResult
 	{
-		return this->AddKeysToControlRig(InSceneComp, Manip, Time, *GeneratedKeys, KeyMode, UMovieSceneControlRigParameterTrack::StaticClass(), ControlRigName);
+		return this->AddKeysToControlRig(InSceneComp, Manip, Time, *GeneratedKeys, KeyMode, UMovieSceneControlRigParameterTrack::StaticClass(), ControlRigName, RigControlName);
 	};
 
 	AnimatablePropertyChanged(FOnKeyProperty::CreateLambda(OnKeyProperty));
@@ -1087,7 +1087,7 @@ void FControlRigParameterTrackEditor::AddControlKeys(USceneComponent *InSceneCom
 }
 
 
-bool FControlRigParameterTrackEditor::ModifyOurGeneratedKeysByCurrentAndWeight(UObject *Object, IControlRigManipulatable* Manip, UMovieSceneTrack *Track, UMovieSceneSection* SectionToKey, FFrameNumber KeyTime, FGeneratedTrackKeys& GeneratedTotalKeys, float Weight) const
+bool FControlRigParameterTrackEditor::ModifyOurGeneratedKeysByCurrentAndWeight(UObject *Object, IControlRigManipulatable* Manip, FName RigControlName, UMovieSceneTrack *Track, UMovieSceneSection* SectionToKey, FFrameNumber KeyTime, FGeneratedTrackKeys& GeneratedTotalKeys, float Weight) const
 {
 	FFrameRate TickResolution = GetSequencer()->GetFocusedTickResolution();
 	FMovieSceneEvaluationTrack EvalTrack = Track->GenerateTrackTemplate();
@@ -1104,50 +1104,49 @@ bool FControlRigParameterTrackEditor::ModifyOurGeneratedKeysByCurrentAndWeight(U
 	FChannelMapInfo* pChannelIndex = nullptr;
 	for (const FRigControl& RigControl : Controls)
 	{
+		if (!RigControl.bAnimatable)
+		{
+			continue;
+		}
 		switch (RigControl.ControlType)
 		{
 			case ERigControlType::Float:
 			{
 				for (const FFloatInterrogationData& Val : InterrogationData.Iterate<FFloatInterrogationData>(UMovieSceneControlRigParameterSection::GetFloatInterrogationKey()))
 				{
-					if (Val.ParameterName == RigControl.Name)
+					if ((RigControlName == Val.ParameterName) && (RigControlName == RigControl.Name))
 					{
 						pChannelIndex = Section->ControlChannelMap.Find(RigControl.Name);
 						if (pChannelIndex)
 						{
 							ChannelIndex = pChannelIndex->ChannelIndex;
+							GeneratedTotalKeys[ChannelIndex]->ModifyByCurrentAndWeight(Proxy, KeyTime, (void *)&Val.Val, Weight);
 						}
-						
-						GeneratedTotalKeys[ChannelIndex]->ModifyByCurrentAndWeight(Proxy, KeyTime, (void *)&Val.Val, Weight);
 						break;
 					}
 				}
-				++ChannelIndex;
 				break;
 			}
 			case ERigControlType::Bool:
 			{
-				//no blending of bools
-				++ChannelIndex;
+				
 				break;
 			}
 			case ERigControlType::Vector2D:
 			{
 				for (const FVector2DInterrogationData& Val : InterrogationData.Iterate<FVector2DInterrogationData>(UMovieSceneControlRigParameterSection::GetVector2DInterrogationKey()))
 				{
-					if (Val.ParameterName == RigControl.Name)
+					if ((RigControlName == Val.ParameterName) && (RigControlName == RigControl.Name))
 					{
 						pChannelIndex = Section->ControlChannelMap.Find(RigControl.Name);
 						if (pChannelIndex)
 						{
 							ChannelIndex = pChannelIndex->ChannelIndex;
+							GeneratedTotalKeys[ChannelIndex]->ModifyByCurrentAndWeight(Proxy, KeyTime, (void *)&Val.Val.X, Weight);
+							GeneratedTotalKeys[ChannelIndex + 1]->ModifyByCurrentAndWeight(Proxy, KeyTime, (void *)&Val.Val.Y, Weight);
 						}
-						GeneratedTotalKeys[ChannelIndex]->ModifyByCurrentAndWeight(Proxy, KeyTime, (void *)&Val.Val.X, Weight);
-						GeneratedTotalKeys[ChannelIndex + 1]->ModifyByCurrentAndWeight(Proxy, KeyTime, (void *)&Val.Val.Y, Weight);
-
 						break;
 					}
-					ChannelIndex += 2;
 				}
 				break;
 			}
@@ -1157,22 +1156,19 @@ bool FControlRigParameterTrackEditor::ModifyOurGeneratedKeysByCurrentAndWeight(U
 			{
 				for (const FVectorInterrogationData& Val : InterrogationData.Iterate<FVectorInterrogationData>(UMovieSceneControlRigParameterSection::GetVectorInterrogationKey()))
 				{
-					if (Val.ParameterName == RigControl.Name)
+					if ((RigControlName == Val.ParameterName) && (RigControlName == RigControl.Name))
 					{
 						pChannelIndex = Section->ControlChannelMap.Find(RigControl.Name);
 						if (pChannelIndex)
 						{
 							ChannelIndex = pChannelIndex->ChannelIndex;
+							GeneratedTotalKeys[ChannelIndex]->ModifyByCurrentAndWeight(Proxy, KeyTime, (void *)&Val.Val.X, Weight);
+							GeneratedTotalKeys[ChannelIndex + 1]->ModifyByCurrentAndWeight(Proxy, KeyTime, (void *)&Val.Val.Y, Weight);
+							GeneratedTotalKeys[ChannelIndex + 2]->ModifyByCurrentAndWeight(Proxy, KeyTime, (void *)&Val.Val.Z, Weight);
 						}
-						GeneratedTotalKeys[ChannelIndex]->ModifyByCurrentAndWeight(Proxy, KeyTime, (void *)&Val.Val.X, Weight);
-						GeneratedTotalKeys[ChannelIndex + 1]->ModifyByCurrentAndWeight(Proxy, KeyTime, (void *)&Val.Val.Y, Weight);
-						GeneratedTotalKeys[ChannelIndex + 2]->ModifyByCurrentAndWeight(Proxy, KeyTime, (void *)&Val.Val.Z, Weight);
-
 						break;
 					}
-					ChannelIndex += 3;
 				}
-
 				break;
 			}
 
@@ -1181,39 +1177,31 @@ bool FControlRigParameterTrackEditor::ModifyOurGeneratedKeysByCurrentAndWeight(U
 			{
 				for (const FTransformInterrogationData& Val : InterrogationData.Iterate<FTransformInterrogationData>(UMovieSceneControlRigParameterSection::GetTransformInterrogationKey()))
 				{
-					if (Val.ParameterName == RigControl.Name)
+					if ((RigControlName == Val.ParameterName) && (RigControlName == RigControl.Name))
 					{
 						pChannelIndex = Section->ControlChannelMap.Find(RigControl.Name);
 						if (pChannelIndex)
 						{
 							ChannelIndex = pChannelIndex->ChannelIndex;
-						}
-						FVector CurrentPos = Val.Val.GetTranslation();
-						FRotator CurrentRot = Val.Val.GetRotation().Rotator();
-						GeneratedTotalKeys[ChannelIndex]->ModifyByCurrentAndWeight(Proxy, KeyTime, (void *)&CurrentPos.X, Weight);
-						GeneratedTotalKeys[ChannelIndex + 1]->ModifyByCurrentAndWeight(Proxy, KeyTime, (void *)&CurrentPos.Y, Weight);
-						GeneratedTotalKeys[ChannelIndex + 2]->ModifyByCurrentAndWeight(Proxy, KeyTime, (void *)&CurrentPos.Z, Weight);
-						GeneratedTotalKeys[ChannelIndex + 3]->ModifyByCurrentAndWeight(Proxy, KeyTime, (void *)&CurrentRot.Roll, Weight);
-						GeneratedTotalKeys[ChannelIndex + 4]->ModifyByCurrentAndWeight(Proxy, KeyTime, (void *)&CurrentRot.Pitch, Weight);
-						GeneratedTotalKeys[ChannelIndex + 5]->ModifyByCurrentAndWeight(Proxy, KeyTime, (void *)&CurrentRot.Yaw, Weight);
-						if (RigControl.ControlType == ERigControlType::Transform)
-						{
-							FVector CurrentScale = Val.Val.GetScale3D();
-							GeneratedTotalKeys[ChannelIndex + 6]->ModifyByCurrentAndWeight(Proxy, KeyTime, (void *)&CurrentScale.X, Weight);
-							GeneratedTotalKeys[ChannelIndex + 7]->ModifyByCurrentAndWeight(Proxy, KeyTime, (void *)&CurrentScale.Y, Weight);
-							GeneratedTotalKeys[ChannelIndex + 8]->ModifyByCurrentAndWeight(Proxy, KeyTime, (void *)&CurrentScale.Z, Weight);
+							FVector CurrentPos = Val.Val.GetTranslation();
+							FRotator CurrentRot = Val.Val.GetRotation().Rotator();
+							GeneratedTotalKeys[ChannelIndex]->ModifyByCurrentAndWeight(Proxy, KeyTime, (void *)&CurrentPos.X, Weight);
+							GeneratedTotalKeys[ChannelIndex + 1]->ModifyByCurrentAndWeight(Proxy, KeyTime, (void *)&CurrentPos.Y, Weight);
+							GeneratedTotalKeys[ChannelIndex + 2]->ModifyByCurrentAndWeight(Proxy, KeyTime, (void *)&CurrentPos.Z, Weight);
+							GeneratedTotalKeys[ChannelIndex + 3]->ModifyByCurrentAndWeight(Proxy, KeyTime, (void *)&CurrentRot.Roll, Weight);
+							GeneratedTotalKeys[ChannelIndex + 4]->ModifyByCurrentAndWeight(Proxy, KeyTime, (void *)&CurrentRot.Pitch, Weight);
+							GeneratedTotalKeys[ChannelIndex + 5]->ModifyByCurrentAndWeight(Proxy, KeyTime, (void *)&CurrentRot.Yaw, Weight);
+							if (RigControl.ControlType == ERigControlType::Transform)
+							{
+								FVector CurrentScale = Val.Val.GetScale3D();
+								GeneratedTotalKeys[ChannelIndex + 6]->ModifyByCurrentAndWeight(Proxy, KeyTime, (void *)&CurrentScale.X, Weight);
+								GeneratedTotalKeys[ChannelIndex + 7]->ModifyByCurrentAndWeight(Proxy, KeyTime, (void *)&CurrentScale.Y, Weight);
+								GeneratedTotalKeys[ChannelIndex + 8]->ModifyByCurrentAndWeight(Proxy, KeyTime, (void *)&CurrentScale.Z, Weight);
+							}
 						}
 						break;
 					}
 				}
-				if (RigControl.ControlType == ERigControlType::Transform)
-				{
-					ChannelIndex += 9;
-				}
-				else
-				{
-					ChannelIndex += 6;
-				}
 				break;
 			}
 		}
@@ -1469,5 +1457,4 @@ bool FControlRigParameterSection::RequestDeleteKeyArea(const TArray<FName>& KeyA
 	return true;
 }
 
-
 #undef LOCTEXT_NAMESPACE
diff --git a/Engine/Plugins/Experimental/ControlRig/Source/ControlRigEditor/Private/Sequencer/ControlRigParameterTrackEditor.h b/Engine/Plugins/Experimental/ControlRig/Source/ControlRigEditor/Private/Sequencer/ControlRigParameterTrackEditor.h
--- a/Engine/Plugins/Experimental/ControlRig/Source/ControlRigEditor/Private/Sequencer/ControlRigParameterTrackEditor.h
+++ b/Engine/Plugins/Experimental/ControlRig/Source/ControlRigEditor/Private/Sequencer/ControlRigParameterTrackEditor.h
@@ -88,10 +88,10 @@ public:
 	void GetControlRigKeys(IControlRigManipulatable* Manip, FName ParameterName, EMovieSceneTransformChannel ChannelsToKey, FGeneratedTrackKeys& OutGeneratedKeys);
 	FKeyPropertyResult AddKeysToControlRig(
 		USceneComponent *InSceneComp, IControlRigManipulatable* Manip, FFrameNumber KeyTime, FGeneratedTrackKeys& GeneratedKeys,
-		ESequencerKeyMode KeyMode, TSubclassOf<UMovieSceneTrack> TrackClass, FName PropertyName);
+		ESequencerKeyMode KeyMode, TSubclassOf<UMovieSceneTrack> TrackClass, FName ControlRigName, FName RigControlName);
 	FKeyPropertyResult AddKeysToControlRigHandle(USceneComponent *InSceneComp, IControlRigManipulatable* Manip,
 		FGuid ObjectHandle, FFrameNumber KeyTime, FGeneratedTrackKeys& GeneratedKeys,
-		ESequencerKeyMode KeyMode, TSubclassOf<UMovieSceneTrack> TrackClass, FName PropertyName);
+		ESequencerKeyMode KeyMode, TSubclassOf<UMovieSceneTrack> TrackClass, FName ControlRigName, FName RigControlName);
 	/**
 	 * Modify the passed in Generated Keys by the current tracks values and weight at the passed in time.
 
@@ -102,7 +102,7 @@ public:
 	 * @param InOutGeneratedTrackKeys The Keys we need to modify. We change these values.
 	 * @param Weight The weight we need to modify the values by.
 	 */
-	bool ModifyOurGeneratedKeysByCurrentAndWeight(UObject* Object, IControlRigManipulatable* Manip, UMovieSceneTrack *Track, UMovieSceneSection* SectionToKey, FFrameNumber Time, FGeneratedTrackKeys& InOutGeneratedTotalKeys, float Weight) const;
+	bool ModifyOurGeneratedKeysByCurrentAndWeight(UObject* Object, IControlRigManipulatable* Manip, FName RigControlName, UMovieSceneTrack *Track, UMovieSceneSection* SectionToKey, FFrameNumber Time, FGeneratedTrackKeys& InOutGeneratedTotalKeys, float Weight) const;
 
 	/*
 	bool ShouldFilterAsset(const FAssetData& AssetData);
