Level Editor: Animate mode doesn't undo moving actors.  This was because it would only transact moving control rig controls, we now check to see if we are selected if not then we don't transact and so it defaults to the normal transaction.

diff --git a/Engine/Plugins/Experimental/ControlRig/Source/ControlRigEditor/Private/EditMode/ControlRigEditMode.cpp b/Engine/Plugins/Experimental/ControlRig/Source/ControlRigEditor/Private/EditMode/ControlRigEditMode.cpp
--- a/Engine/Plugins/Experimental/ControlRig/Source/ControlRigEditor/Private/EditMode/ControlRigEditMode.cpp
+++ b/Engine/Plugins/Experimental/ControlRig/Source/ControlRigEditor/Private/EditMode/ControlRigEditMode.cpp
@@ -490,8 +490,16 @@ bool FControlRigEditMode::StartTracking(FEditorViewportClient* InViewportClient,
 
 			WeakControlRigEditing->SetFlags(RF_Transactional);
 			WeakControlRigEditing->Modify();
+			//in level editor only transact if we have at least one control selected, in editor we only select CR stuff so always transact
+			if (IsInLevelEditor())
+			{
+				bIsTransacting = (AreRigElementSelectedAndMovable() && ManipulationLayer != nullptr);
+			}
+			else
+			{
+				bIsTransacting = true;
+			}
 		}
-		bIsTransacting = true;
 		bManipulatorMadeChange = false;
 
 		return bIsTransacting;
