Sequencer: Control Rig: Fix for crash when compiling CR with same on in level sequence. Don't cache ptr, just use it from the UProperty on the section since that get's updated correctly.

diff --git a/Engine/Plugins/Experimental/ControlRig/Source/ControlRig/Private/Sequencer/MovieSceneControlRigParameterTemplate.cpp b/Engine/Plugins/Experimental/ControlRig/Source/ControlRig/Private/Sequencer/MovieSceneControlRigParameterTemplate.cpp
--- a/Engine/Plugins/Experimental/ControlRig/Source/ControlRig/Private/Sequencer/MovieSceneControlRigParameterTemplate.cpp
+++ b/Engine/Plugins/Experimental/ControlRig/Source/ControlRig/Private/Sequencer/MovieSceneControlRigParameterTemplate.cpp
@@ -641,9 +641,8 @@ FMovieSceneControlRigParameterTemplate::FMovieSceneControlRigParameterTemplate(c
 
 struct TControlRigParameterActuatorFloat : TMovieSceneBlendingActuator<FControlRigTrackTokenFloat>
 {
-	TControlRigParameterActuatorFloat(FMovieSceneAnimTypeID& InAnimID, UControlRig* InControlRig, const FName& InParameterName, const UMovieSceneControlRigParameterSection* InSection)
+	TControlRigParameterActuatorFloat(FMovieSceneAnimTypeID& InAnimID, const FName& InParameterName, const UMovieSceneControlRigParameterSection* InSection)
 		: TMovieSceneBlendingActuator<FControlRigTrackTokenFloat>(FMovieSceneBlendingActuatorID(InAnimID))
-		, ControlRig(InControlRig)
 		, ParameterName(InParameterName)
 		, SectionData(InSection)
 	{}
@@ -651,9 +650,11 @@ struct TControlRigParameterActuatorFloat : TMovieSceneBlendingActuator<FControlR
 
 	FControlRigTrackTokenFloat RetrieveCurrentValue(UObject* InObject, IMovieScenePlayer* Player) const
 	{
-		if (ControlRig)
+		const UMovieSceneControlRigParameterSection* Section = nullptr;
+		Section = SectionData.Get();
+		if (Section && Section->ControlRig)
 		{
-			FRigControl* RigControl = ControlRig->FindControl(ParameterName);
+			FRigControl* RigControl = Section->ControlRig->FindControl(ParameterName);
 			if (RigControl && RigControl->ControlType == ERigControlType::Float)
 			{
 				float Val = RigControl->Value.Get<float>();
@@ -674,17 +675,16 @@ struct TControlRigParameterActuatorFloat : TMovieSceneBlendingActuator<FControlR
 		{
 			bWasDoNotKey = Section->GetDoNotKey();
 			Section->SetDoNotKey(true);
-		}
-		if (ControlRig)
-		{
-			FRigControl* RigControl = ControlRig->FindControl(ParameterName);
-			if (RigControl && RigControl->ControlType == ERigControlType::Float)
+		
+			if (Section->ControlRig)
 			{
-				ControlRig->SetControlValue<float>(ParameterName, InFinalValue.Value, true, EControlRigSetKey::Never);
+				FRigControl* RigControl = Section->ControlRig->FindControl(ParameterName);
+				if (RigControl && RigControl->ControlType == ERigControlType::Float)
+				{
+					Section->ControlRig->SetControlValue<float>(ParameterName, InFinalValue.Value, true, EControlRigSetKey::Never);
+				}
 			}
-		}
-		if (Section)
-		{
+
 			Section->SetDoNotKey(bWasDoNotKey);
 		}
 	}
@@ -698,7 +698,6 @@ struct TControlRigParameterActuatorFloat : TMovieSceneBlendingActuator<FControlR
 		InterrogationData.Add(FFloatInterrogationData(Data), UMovieSceneControlRigParameterSection::GetFloatInterrogationKey());
 	}
 
-	UControlRig *ControlRig;
 	FName ParameterName;
 	TWeakObjectPtr<const UMovieSceneControlRigParameterSection> SectionData;
 
@@ -707,9 +706,8 @@ struct TControlRigParameterActuatorFloat : TMovieSceneBlendingActuator<FControlR
 
 struct TControlRigParameterActuatorVector2D : TMovieSceneBlendingActuator<FControlRigTrackTokenVector2D>
 {
-	TControlRigParameterActuatorVector2D(FMovieSceneAnimTypeID& InAnimID, UControlRig* InControlRig, const FName& InParameterName, const UMovieSceneControlRigParameterSection* InSection)
+	TControlRigParameterActuatorVector2D(FMovieSceneAnimTypeID& InAnimID,  const FName& InParameterName, const UMovieSceneControlRigParameterSection* InSection)
 		: TMovieSceneBlendingActuator<FControlRigTrackTokenVector2D>(FMovieSceneBlendingActuatorID(InAnimID))
-		, ControlRig(InControlRig)
 		, ParameterName(InParameterName)
 		, SectionData(InSection)
 	{}
@@ -718,9 +716,10 @@ struct TControlRigParameterActuatorVector2D : TMovieSceneBlendingActuator<FContr
 
 	FControlRigTrackTokenVector2D RetrieveCurrentValue(UObject* InObject, IMovieScenePlayer* Player) const
 	{
-		if (ControlRig)
+		const UMovieSceneControlRigParameterSection* Section = nullptr;
+		Section = SectionData.Get();
 		{
-			FRigControl* RigControl = ControlRig->FindControl(ParameterName);
+			FRigControl* RigControl = Section->ControlRig->FindControl(ParameterName);
 			if (RigControl && (RigControl->ControlType == ERigControlType::Vector2D))
 			{
 				FVector2D Val = RigControl->Value.Get<FVector2D>();
@@ -740,17 +739,16 @@ struct TControlRigParameterActuatorVector2D : TMovieSceneBlendingActuator<FContr
 		{
 			bWasDoNotKey = Section->GetDoNotKey();
 			Section->SetDoNotKey(true);
-		}
-		if (ControlRig)
-		{
-			FRigControl* RigControl = ControlRig->FindControl(ParameterName);
-			if (RigControl && (RigControl->ControlType == ERigControlType::Vector2D))
+	
+			if (Section->ControlRig)
 			{
-				ControlRig->SetControlValue<FVector2D>(ParameterName, InFinalValue.Value, true, EControlRigSetKey::Never);
+				FRigControl* RigControl = Section->ControlRig->FindControl(ParameterName);
+				if (RigControl && (RigControl->ControlType == ERigControlType::Vector2D))
+				{
+					Section->ControlRig->SetControlValue<FVector2D>(ParameterName, InFinalValue.Value, true, EControlRigSetKey::Never);
+				}
 			}
-		}
-		if (Section)
-		{
+
 			Section->SetDoNotKey(bWasDoNotKey);
 		}
 	}
@@ -762,7 +760,6 @@ struct TControlRigParameterActuatorVector2D : TMovieSceneBlendingActuator<FContr
 		InterrogationData.Add(FVector2DInterrogationData(Data), UMovieSceneControlRigParameterSection::GetVector2DInterrogationKey());
 	}
 
-	UControlRig *ControlRig;
 	FName ParameterName;
 	TWeakObjectPtr<const UMovieSceneControlRigParameterSection> SectionData;
 
@@ -771,9 +768,8 @@ struct TControlRigParameterActuatorVector2D : TMovieSceneBlendingActuator<FContr
 
 struct TControlRigParameterActuatorVector : TMovieSceneBlendingActuator<FControlRigTrackTokenVector>
 {
-	TControlRigParameterActuatorVector(FMovieSceneAnimTypeID& InAnimID, UControlRig* InControlRig, const FName& InParameterName, const UMovieSceneControlRigParameterSection* InSection)
+	TControlRigParameterActuatorVector(FMovieSceneAnimTypeID& InAnimID, const FName& InParameterName, const UMovieSceneControlRigParameterSection* InSection)
 		: TMovieSceneBlendingActuator<FControlRigTrackTokenVector>(FMovieSceneBlendingActuatorID(InAnimID))
-		, ControlRig(InControlRig)
 		, ParameterName(InParameterName)
 		, SectionData(InSection)
 	{}
@@ -782,9 +778,11 @@ struct TControlRigParameterActuatorVector : TMovieSceneBlendingActuator<FControl
 
 	FControlRigTrackTokenVector RetrieveCurrentValue(UObject* InObject, IMovieScenePlayer* Player) const
 	{
-		if (ControlRig)
+		const UMovieSceneControlRigParameterSection* Section = nullptr;
+		Section = SectionData.Get();
+		if (Section->ControlRig)
 		{
-			FRigControl* RigControl = ControlRig->FindControl(ParameterName);
+			FRigControl* RigControl = Section->ControlRig->FindControl(ParameterName);
 			if (RigControl && (RigControl->ControlType == ERigControlType::Position || RigControl->ControlType == ERigControlType::Scale || RigControl->ControlType == ERigControlType::Rotator))
 			{
 				FVector Val = RigControl->Value.Get<FVector>();
@@ -804,17 +802,15 @@ struct TControlRigParameterActuatorVector : TMovieSceneBlendingActuator<FControl
 		{
 			bWasDoNotKey = Section->GetDoNotKey();
 			Section->SetDoNotKey(true);
-		}
-		if (ControlRig)
-		{
-			FRigControl* RigControl = ControlRig->FindControl(ParameterName);
-			if (RigControl && (RigControl->ControlType == ERigControlType::Position || RigControl->ControlType == ERigControlType::Scale || RigControl->ControlType == ERigControlType::Rotator))
+
+			if (Section->ControlRig)
 			{
-				ControlRig->SetControlValue<FVector>(ParameterName, InFinalValue.Value, true, EControlRigSetKey::Never);
+				FRigControl* RigControl = Section->ControlRig->FindControl(ParameterName);
+				if (RigControl && (RigControl->ControlType == ERigControlType::Position || RigControl->ControlType == ERigControlType::Scale || RigControl->ControlType == ERigControlType::Rotator))
+				{
+					Section->ControlRig->SetControlValue<FVector>(ParameterName, InFinalValue.Value, true, EControlRigSetKey::Never);
+				}
 			}
-		}
-		if (Section)
-		{
 			Section->SetDoNotKey(bWasDoNotKey);
 		}
 	}
@@ -826,7 +822,6 @@ struct TControlRigParameterActuatorVector : TMovieSceneBlendingActuator<FControl
 		InterrogationData.Add(FVectorInterrogationData(Data), UMovieSceneControlRigParameterSection::GetVectorInterrogationKey());
 	}
 
-	UControlRig *ControlRig;
 	FName ParameterName;
 	TWeakObjectPtr<const UMovieSceneControlRigParameterSection> SectionData;
 
@@ -836,18 +831,20 @@ struct TControlRigParameterActuatorVector : TMovieSceneBlendingActuator<FControl
 
 struct TControlRigParameterActuatorTransform : TMovieSceneBlendingActuator<FControlRigTrackTokenTransform>
 {
-	TControlRigParameterActuatorTransform(FMovieSceneAnimTypeID& InAnimID, UControlRig* InControlRig, const FName& InParameterName, const UMovieSceneControlRigParameterSection* InSection)
+	TControlRigParameterActuatorTransform(FMovieSceneAnimTypeID& InAnimID, const FName& InParameterName, const UMovieSceneControlRigParameterSection* InSection)
 		: TMovieSceneBlendingActuator<FControlRigTrackTokenTransform>(FMovieSceneBlendingActuatorID(InAnimID))
-		, ControlRig(InControlRig)
 		, ParameterName(InParameterName)
 		, SectionData(InSection)
 	{}
 
 	FControlRigTrackTokenTransform RetrieveCurrentValue(UObject* InObject, IMovieScenePlayer* Player) const
 	{
-		if (ControlRig)
+		const UMovieSceneControlRigParameterSection* Section = nullptr;
+
+		Section = SectionData.Get();
+		if (Section->ControlRig)
 		{
-			FRigControl* RigControl = ControlRig->FindControl(ParameterName);
+			FRigControl* RigControl = Section->ControlRig->FindControl(ParameterName);
 			if (RigControl && RigControl->ControlType == ERigControlType::Transform)
 			{
 				FTransform Val = RigControl->Value.Get<FTransform>();
@@ -882,22 +879,21 @@ struct TControlRigParameterActuatorTransform : TMovieSceneBlendingActuator<FCont
 		{
 			bWasDoNotKey = Section->GetDoNotKey();
 			Section->SetDoNotKey(true);
-		}
-		if (ControlRig)
-		{
-			FRigControl* RigControl = ControlRig->FindControl(ParameterName);
-			if (RigControl && RigControl->ControlType == ERigControlType::Transform)
-			{
-				ControlRig->SetControlValue<FTransform>(ParameterName, InFinalValue.Value,true, EControlRigSetKey::Never);
-			}
-			else if (RigControl && RigControl->ControlType == ERigControlType::TransformNoScale)
+		
+			if (Section->ControlRig)
 			{
-				FTransformNoScale NoScale = InFinalValue.Value;
-				ControlRig->SetControlValue<FTransformNoScale>(ParameterName, NoScale, true, EControlRigSetKey::Never);
+				FRigControl* RigControl = Section->ControlRig->FindControl(ParameterName);
+				if (RigControl && RigControl->ControlType == ERigControlType::Transform)
+				{
+					Section->ControlRig->SetControlValue<FTransform>(ParameterName, InFinalValue.Value,true, EControlRigSetKey::Never);
+				}
+				else if (RigControl && RigControl->ControlType == ERigControlType::TransformNoScale)
+				{
+					FTransformNoScale NoScale = InFinalValue.Value;
+					Section->ControlRig->SetControlValue<FTransformNoScale>(ParameterName, NoScale, true, EControlRigSetKey::Never);
+				}
 			}
-		}
-		if (Section)
-		{
+
 			Section->SetDoNotKey(bWasDoNotKey);
 		}
 	}
@@ -908,7 +904,6 @@ struct TControlRigParameterActuatorTransform : TMovieSceneBlendingActuator<FCont
 		Data.ParameterName = ParameterName;
 		InterrogationData.Add(FTransformInterrogationData(Data), UMovieSceneControlRigParameterSection::GetTransformInterrogationKey());
 	}
-	UControlRig *ControlRig;
 	FName ParameterName;
 	TWeakObjectPtr<const UMovieSceneControlRigParameterSection> SectionData;
 
@@ -980,7 +975,7 @@ void FMovieSceneControlRigParameterTemplate::Evaluate(const FMovieSceneEvaluatio
 
 			if (!ExecutionTokens.GetBlendingAccumulator().FindActuator< FControlRigTrackTokenFloat>(ActuatorTypeID))
 			{
-				ExecutionTokens.GetBlendingAccumulator().DefineActuator(ActuatorTypeID, MakeShared <TControlRigParameterActuatorFloat>(AnimTypeID, Section->ControlRig, ScalarNameAndValue.ParameterName,Section));
+				ExecutionTokens.GetBlendingAccumulator().DefineActuator(ActuatorTypeID, MakeShared <TControlRigParameterActuatorFloat>(AnimTypeID, ScalarNameAndValue.ParameterName,Section));
 			}
 			ExecutionTokens.BlendToken(ActuatorTypeID, TBlendableToken<FControlRigTrackTokenFloat>(ScalarNameAndValue.Value, Section->GetBlendType().Get(), Weight));
 		}
@@ -995,7 +990,7 @@ void FMovieSceneControlRigParameterTemplate::Evaluate(const FMovieSceneEvaluatio
 
 			if (!ExecutionTokens.GetBlendingAccumulator().FindActuator< FControlRigTrackTokenVector>(ActuatorTypeID))
 			{
-				ExecutionTokens.GetBlendingAccumulator().DefineActuator(ActuatorTypeID, MakeShared <TControlRigParameterActuatorVector>(AnimTypeID, Section->ControlRig, VectorNameAndValue.ParameterName,Section));
+				ExecutionTokens.GetBlendingAccumulator().DefineActuator(ActuatorTypeID, MakeShared <TControlRigParameterActuatorVector>(AnimTypeID,  VectorNameAndValue.ParameterName,Section));
 			}
 			VectorData.Set(0, VectorNameAndValue.Value.X);
 			VectorData.Set(1, VectorNameAndValue.Value.Y);
@@ -1014,7 +1009,7 @@ void FMovieSceneControlRigParameterTemplate::Evaluate(const FMovieSceneEvaluatio
 
 			if (!ExecutionTokens.GetBlendingAccumulator().FindActuator< FControlRigTrackTokenVector2D>(ActuatorTypeID))
 			{
-				ExecutionTokens.GetBlendingAccumulator().DefineActuator(ActuatorTypeID, MakeShared <TControlRigParameterActuatorVector2D>(AnimTypeID, Section->ControlRig, Vector2DNameAndValue.ParameterName, Section));
+				ExecutionTokens.GetBlendingAccumulator().DefineActuator(ActuatorTypeID, MakeShared <TControlRigParameterActuatorVector2D>(AnimTypeID,  Vector2DNameAndValue.ParameterName, Section));
 			}
 			Vector2DData.Set(0, Vector2DNameAndValue.Value.X);
 			Vector2DData.Set(1, Vector2DNameAndValue.Value.Y);
@@ -1032,7 +1027,7 @@ void FMovieSceneControlRigParameterTemplate::Evaluate(const FMovieSceneEvaluatio
 
 			if (!ExecutionTokens.GetBlendingAccumulator().FindActuator< FControlRigTrackTokenTransform>(ActuatorTypeID))
 			{
-				ExecutionTokens.GetBlendingAccumulator().DefineActuator(ActuatorTypeID, MakeShared <TControlRigParameterActuatorTransform>(AnimTypeID, Section->ControlRig, TransformNameAndValue.ParameterName, Section));
+				ExecutionTokens.GetBlendingAccumulator().DefineActuator(ActuatorTypeID, MakeShared <TControlRigParameterActuatorTransform>(AnimTypeID,  TransformNameAndValue.ParameterName, Section));
 			}
 
 			FTransform Transform(TransformNameAndValue.Rotation, TransformNameAndValue.Translation, TransformNameAndValue.Scale);
@@ -1233,7 +1228,7 @@ void FMovieSceneControlRigParameterTemplate::Interrogate(const FMovieSceneContex
 
 			if (!Container.GetAccumulator().FindActuator< FControlRigTrackTokenFloat>(ActuatorTypeID))
 			{
-				Container.GetAccumulator().DefineActuator(ActuatorTypeID, MakeShared <TControlRigParameterActuatorFloat>(AnimTypeID, Section->ControlRig, ScalarNameAndValue.ParameterName, Section));
+				Container.GetAccumulator().DefineActuator(ActuatorTypeID, MakeShared <TControlRigParameterActuatorFloat>(AnimTypeID, ScalarNameAndValue.ParameterName, Section));
 			}
 			Container.GetAccumulator().BlendToken(FMovieSceneEvaluationOperand(), ActuatorTypeID, FMovieSceneEvaluationScope(), Context,TBlendableToken<FControlRigTrackTokenFloat>(ScalarNameAndValue.Value, Section->GetBlendType().Get(), Weight));
 		}
@@ -1246,7 +1241,7 @@ void FMovieSceneControlRigParameterTemplate::Interrogate(const FMovieSceneContex
 
 			if (!Container.GetAccumulator().FindActuator< FControlRigTrackTokenVector>(ActuatorTypeID))
 			{
-				Container.GetAccumulator().DefineActuator(ActuatorTypeID, MakeShared <TControlRigParameterActuatorVector2D>(AnimTypeID, Section->ControlRig, Vector2DNameAndValue.ParameterName, Section));
+				Container.GetAccumulator().DefineActuator(ActuatorTypeID, MakeShared <TControlRigParameterActuatorVector2D>(AnimTypeID, Vector2DNameAndValue.ParameterName, Section));
 			}
 			Vector2DData.Set(0, Vector2DNameAndValue.Value.X);
 			Vector2DData.Set(1, Vector2DNameAndValue.Value.Y);
@@ -1262,7 +1257,7 @@ void FMovieSceneControlRigParameterTemplate::Interrogate(const FMovieSceneContex
 
 			if (!Container.GetAccumulator().FindActuator< FControlRigTrackTokenVector>(ActuatorTypeID))
 			{
-				Container.GetAccumulator().DefineActuator(ActuatorTypeID, MakeShared <TControlRigParameterActuatorVector>(AnimTypeID, Section->ControlRig, VectorNameAndValue.ParameterName, Section));
+				Container.GetAccumulator().DefineActuator(ActuatorTypeID, MakeShared <TControlRigParameterActuatorVector>(AnimTypeID, VectorNameAndValue.ParameterName, Section));
 			}
 			VectorData.Set(0, VectorNameAndValue.Value.X);
 			VectorData.Set(1, VectorNameAndValue.Value.Y);
@@ -1282,7 +1277,7 @@ void FMovieSceneControlRigParameterTemplate::Interrogate(const FMovieSceneContex
 
 			if (!Container.GetAccumulator().FindActuator< FControlRigTrackTokenTransform>(ActuatorTypeID))
 			{
-				Container.GetAccumulator().DefineActuator(ActuatorTypeID, MakeShared <TControlRigParameterActuatorTransform>(AnimTypeID, Section->ControlRig, TransformNameAndValue.ParameterName, Section));
+				Container.GetAccumulator().DefineActuator(ActuatorTypeID, MakeShared <TControlRigParameterActuatorTransform>(AnimTypeID, TransformNameAndValue.ParameterName, Section));
 			}
 
 			FTransform Transform(TransformNameAndValue.Rotation, TransformNameAndValue.Translation, TransformNameAndValue.Scale);
