Sequencer: Control Rig: Add ability to select bones directly with Additive Control Rig. Also some fixes pushed up from Puffin Testing.

diff --git a/Engine/Plugins/Experimental/ControlRig/Source/ControlRig/Private/Sequencer/MovieSceneControlRigParameterSection.cpp b/Engine/Plugins/Experimental/ControlRig/Source/ControlRig/Private/Sequencer/MovieSceneControlRigParameterSection.cpp
--- a/Engine/Plugins/Experimental/ControlRig/Source/ControlRig/Private/Sequencer/MovieSceneControlRigParameterSection.cpp
+++ b/Engine/Plugins/Experimental/ControlRig/Source/ControlRig/Private/Sequencer/MovieSceneControlRigParameterSection.cpp
@@ -8,8 +8,10 @@
 #include "Sequencer/ControlRigBindingTemplate.h"
 #include "Sequencer/MovieSceneControlRigInstanceData.h"
 #include "Channels/MovieSceneChannelProxy.h"
+#include "Sequencer/MovieSceneControlRigParameterTrack.h"
 #include "Sequencer/ControlRigSortedControls.h"
 
+
 #define LOCTEXT_NAMESPACE "MovieSceneControlParameterRigSection"
 
 #if WITH_EDITOR
@@ -678,6 +680,11 @@ void UMovieSceneControlRigParameterSection::Serialize(FArchive& Ar)
 void UMovieSceneControlRigParameterSection::PostEditImport()
 {
 	Super::PostEditImport();
+	if (UMovieSceneControlRigParameterTrack* Track = Cast< UMovieSceneControlRigParameterTrack>(GetOuter()))
+	{
+		ControlRig = Track->GetControlRig();
+	}
+	ReconstructChannelProxy(true);
 }
 
 void UMovieSceneControlRigParameterSection::PostLoad()
diff --git a/Engine/Plugins/Experimental/ControlRig/Source/ControlRig/Private/Sequencer/MovieSceneControlRigParameterTemplate.cpp b/Engine/Plugins/Experimental/ControlRig/Source/ControlRig/Private/Sequencer/MovieSceneControlRigParameterTemplate.cpp
--- a/Engine/Plugins/Experimental/ControlRig/Source/ControlRig/Private/Sequencer/MovieSceneControlRigParameterTemplate.cpp
+++ b/Engine/Plugins/Experimental/ControlRig/Source/ControlRig/Private/Sequencer/MovieSceneControlRigParameterTemplate.cpp
@@ -375,7 +375,7 @@ namespace FControlRigBindingHelper
 					AnimInstance->RecalcRequiredBones();
 					AnimInstance->AddControlRigTrack(ControlRig->GetUniqueID(), ControlRig);
 					ControlRig->CreateRigControlsForCurveContainer();
-					ControlRig->Execute(EControlRigState::Init);
+					ControlRig->Initialize();
 				}
 			}
 		}
@@ -601,7 +601,7 @@ struct FControlRigParameterExecutionToken : IMovieSceneExecutionToken
 				if (!ControlRig->GetObjectBinding()->GetBoundObject())
 				{
 					ControlRig->GetObjectBinding()->BindToObject(BoundSkelMeshes[0].Get());
-					ControlRig->Execute(EControlRigState::Init);
+					ControlRig->Initialize();
 				}
 				FControlRigBindingHelper::BindToSequencerInstance(ControlRig);
 				//MZ TODO HANDLE BOOLS AND OTHER NON BLENDABLES
diff --git a/Engine/Plugins/Experimental/ControlRig/Source/ControlRig/Private/Sequencer/MovieSceneControlRigParameterTrack.cpp b/Engine/Plugins/Experimental/ControlRig/Source/ControlRig/Private/Sequencer/MovieSceneControlRigParameterTrack.cpp
--- a/Engine/Plugins/Experimental/ControlRig/Source/ControlRig/Private/Sequencer/MovieSceneControlRigParameterTrack.cpp
+++ b/Engine/Plugins/Experimental/ControlRig/Source/ControlRig/Private/Sequencer/MovieSceneControlRigParameterTrack.cpp
@@ -158,6 +158,14 @@ bool UMovieSceneControlRigParameterTrack::HasSection(const UMovieSceneSection& S
 void UMovieSceneControlRigParameterTrack::AddSection(UMovieSceneSection& Section)
 {
 	Sections.Add(&Section);
+	if (UMovieSceneControlRigParameterSection* CRSection = Cast<UMovieSceneControlRigParameterSection>(&Section))
+	{
+		if (CRSection->ControlRig != ControlRig)
+		{
+			CRSection->ControlRig = ControlRig;
+			CRSection->ReconstructChannelProxy(true);
+		}
+	}
 }
 
 void UMovieSceneControlRigParameterTrack::RemoveSection(UMovieSceneSection& Section)
@@ -408,7 +416,7 @@ UMovieSceneSection* UMovieSceneControlRigParameterTrack::GetSectionToKey() const
 void UMovieSceneControlRigParameterTrack::PostLoad()
 {
 	Super::PostLoad();
-	if (ControlRig)
+	if (ControlRig && !ControlRig->HasAnyFlags(RF_NeedLoad | RF_NeedPostLoad | RF_NeedInitialization))
 	{
 		ControlRig->Initialize();
 		ControlRig->CreateRigControlsForCurveContainer();
diff --git a/Engine/Plugins/Experimental/ControlRig/Source/ControlRigEditor/Private/EditMode/ControlRigEditMode.cpp b/Engine/Plugins/Experimental/ControlRig/Source/ControlRigEditor/Private/EditMode/ControlRigEditMode.cpp
--- a/Engine/Plugins/Experimental/ControlRig/Source/ControlRigEditor/Private/EditMode/ControlRigEditMode.cpp
+++ b/Engine/Plugins/Experimental/ControlRig/Source/ControlRigEditor/Private/EditMode/ControlRigEditMode.cpp
@@ -40,6 +40,8 @@
 #include "Framework/Notifications/NotificationManager.h"
 #include "Widgets/Notifications/SNotificationList.h"
 #include "RigVMModel/RigVMController.h"
+#include "Rigs/AdditiveControlRig.h"
+#include "EngineUtils.h"
 
 FName FControlRigEditMode::ModeName("EditMode.ControlRig");
 
@@ -57,6 +59,12 @@ enum class ETransformComponent
 	Scale
 };
 
+namespace ControlRigSelectionConstants
+{
+	/** Distance to trace for physics bodies */
+	static const float BodyTraceDistance = 100000.0f;
+}
+
 FControlRigEditMode::FControlRigEditMode()
 	: bIsTransacting(false)
 	, bManipulatorMadeChange(false)
@@ -607,9 +615,50 @@ bool FControlRigEditMode::HandleClick(FEditorViewportClient* InViewportClient, H
 
 				return true;
 			}
+			else
+			{ 
+				//if we have an additive control rig active select the control based upon the selected bone.
+				UAdditiveControlRig* AdditiveControlRig = Cast<UAdditiveControlRig>(WeakControlRigEditing);
+				if (AdditiveControlRig)
+				{
+					if (USkeletalMeshComponent* RigMeshComp = Cast<USkeletalMeshComponent>(AdditiveControlRig->GetObjectBinding()->GetBoundObject()))
+					{
+						const USkeletalMeshComponent* SkelComp = Cast<USkeletalMeshComponent>(ActorHitProxy->PrimComponent);
+
+						if (SkelComp == RigMeshComp)
+						{
+							FHitResult Result(1.0f);
+							bool bHit = RigMeshComp->LineTraceComponent(Result, Click.GetOrigin(), Click.GetOrigin() + Click.GetDirection() * ControlRigSelectionConstants::BodyTraceDistance, FCollisionQueryParams(NAME_None, FCollisionQueryParams::GetUnknownStatId(), true));
+
+							if (bHit)
+							{
+								FName ControlName(*(Result.BoneName.ToString() + TEXT("_CONTROL")));
+								if (WeakControlRigEditing->FindControl(ControlName))
+								{
+									FScopedTransaction ScopedTransaction(LOCTEXT("SelectControlTransaction", "Select Control"), IsInLevelEditor() && !GIsTransacting);
+
+									if (Click.IsShiftDown() || Click.IsControlDown())
+									{
+										SetRigElementSelection(ERigElementType::Control, ControlName, true);
+									}
+									else
+									{
+										ClearRigElementSelection(FRigElementTypeHelper::ToMask(ERigElementType::Control));
+										SetRigElementSelection(ERigElementType::Control, ControlName, true);
+									}
+									return true;
+								}
+							}
+						}
+					}
+				}
+			}
 		}
 	}
 
+		
+	
+
 	// for now we show this menu all the time if body is selected
 	// if we want some global menu, we'll have to move this
 	if (Click.GetKey() == EKeys::RightMouseButton)
diff --git a/Engine/Plugins/Experimental/ControlRig/Source/ControlRigEditor/Private/Sequencer/ControlRigParameterTrackEditor.cpp b/Engine/Plugins/Experimental/ControlRig/Source/ControlRigEditor/Private/Sequencer/ControlRigParameterTrackEditor.cpp
--- a/Engine/Plugins/Experimental/ControlRig/Source/ControlRigEditor/Private/Sequencer/ControlRigParameterTrackEditor.cpp
+++ b/Engine/Plugins/Experimental/ControlRig/Source/ControlRigEditor/Private/Sequencer/ControlRigParameterTrackEditor.cpp
@@ -1114,7 +1114,7 @@ bool FControlRigParameterTrackEditor::ModifyOurGeneratedKeysByCurrentAndWeight(U
 			{
 				for (const FFloatInterrogationData& Val : InterrogationData.Iterate<FFloatInterrogationData>(UMovieSceneControlRigParameterSection::GetFloatInterrogationKey()))
 				{
-					if ((RigControlName == Val.ParameterName) && (RigControlName == RigControl.Name))
+					if ((Val.ParameterName == RigControl.Name))
 					{
 						pChannelIndex = Section->ControlChannelMap.Find(RigControl.Name);
 						if (pChannelIndex)
@@ -1136,7 +1136,7 @@ bool FControlRigParameterTrackEditor::ModifyOurGeneratedKeysByCurrentAndWeight(U
 			{
 				for (const FVector2DInterrogationData& Val : InterrogationData.Iterate<FVector2DInterrogationData>(UMovieSceneControlRigParameterSection::GetVector2DInterrogationKey()))
 				{
-					if ((RigControlName == Val.ParameterName) && (RigControlName == RigControl.Name))
+					if ((Val.ParameterName == RigControl.Name))
 					{
 						pChannelIndex = Section->ControlChannelMap.Find(RigControl.Name);
 						if (pChannelIndex)
@@ -1156,7 +1156,7 @@ bool FControlRigParameterTrackEditor::ModifyOurGeneratedKeysByCurrentAndWeight(U
 			{
 				for (const FVectorInterrogationData& Val : InterrogationData.Iterate<FVectorInterrogationData>(UMovieSceneControlRigParameterSection::GetVectorInterrogationKey()))
 				{
-					if ((RigControlName == Val.ParameterName) && (RigControlName == RigControl.Name))
+					if ((Val.ParameterName == RigControl.Name))
 					{
 						pChannelIndex = Section->ControlChannelMap.Find(RigControl.Name);
 						if (pChannelIndex)
@@ -1177,7 +1177,8 @@ bool FControlRigParameterTrackEditor::ModifyOurGeneratedKeysByCurrentAndWeight(U
 			{
 				for (const FTransformInterrogationData& Val : InterrogationData.Iterate<FTransformInterrogationData>(UMovieSceneControlRigParameterSection::GetTransformInterrogationKey()))
 				{
-					if ((RigControlName == Val.ParameterName) && (RigControlName == RigControl.Name))
+
+					if ((Val.ParameterName == RigControl.Name))
 					{
 						pChannelIndex = Section->ControlChannelMap.Find(RigControl.Name);
 						if (pChannelIndex)
