Sequencer: Control Rig: Blending sections wasn't working as expected.

diff --git a/Engine/Plugins/Experimental/ControlRig/Source/ControlRig/Private/Sequencer/MovieSceneControlRigParameterSection.cpp b/Engine/Plugins/Experimental/ControlRig/Source/ControlRig/Private/Sequencer/MovieSceneControlRigParameterSection.cpp
--- a/Engine/Plugins/Experimental/ControlRig/Source/ControlRig/Private/Sequencer/MovieSceneControlRigParameterSection.cpp
+++ b/Engine/Plugins/Experimental/ControlRig/Source/ControlRig/Private/Sequencer/MovieSceneControlRigParameterSection.cpp
@@ -1029,6 +1029,21 @@ void UMovieSceneControlRigParameterSection::ReconstructChannelProxy(bool bForce)
 					{
 						if (RigControl.Name == Vector.ParameterName)
 						{
+							if (RigControl.ControlType == ERigControlType::Scale)
+							{
+								if (BlendType == EMovieSceneBlendType::Additive)
+								{
+									Vector.XCurve.SetDefault(0.0f);
+									Vector.YCurve.SetDefault(0.0f);
+									Vector.ZCurve.SetDefault(0.0f);
+								}
+								else
+								{
+									Vector.XCurve.SetDefault(1.0f);
+									Vector.YCurve.SetDefault(1.0f);
+									Vector.ZCurve.SetDefault(1.0f);
+								}
+							}
 							ControlChannelMap.Add(Vector.ParameterName, FChannelMapInfo(Index, ChannelIndex));
 							bool bEnabled = ControlsMask[Index++];
 							FString ParameterString = Vector.ParameterName.ToString();
@@ -1106,6 +1121,18 @@ void UMovieSceneControlRigParameterSection::ReconstructChannelProxy(bool bForce)
 
 							if (RigControl.ControlType == ERigControlType::Transform)
 							{
+								if (BlendType == EMovieSceneBlendType::Additive)
+								{
+									Transform.Scale[0].SetDefault(0.0f);
+									Transform.Scale[1].SetDefault(0.0f);
+									Transform.Scale[2].SetDefault(0.0f);
+								}
+								else
+								{
+									Transform.Scale[0].SetDefault(1.0f);
+									Transform.Scale[1].SetDefault(1.0f);
+									Transform.Scale[2].SetDefault(1.0f);
+								}
 								Channels.Add(Transform.Scale[0], EditorData.MetaData[6], EditorData.ExternalValues[6]);
 								Channels.Add(Transform.Scale[1], EditorData.MetaData[7], EditorData.ExternalValues[7]);
 								Channels.Add(Transform.Scale[2], EditorData.MetaData[8], EditorData.ExternalValues[8]);
@@ -1237,7 +1264,22 @@ void UMovieSceneControlRigParameterSection::ReconstructChannelProxy(bool bForce)
 				}
 #endif
 			}
+		
+#if WITH_EDITOR
+			FMovieSceneChannelMetaData      MetaData;
+			MetaData.SetIdentifiers("Weight", NSLOCTEXT("MovieSceneTransformSection", "Weight", "Weight"));
+			MetaData.bEnabled = EnumHasAllFlags(TransformMask.GetChannels(), EMovieSceneTransformChannel::Weight);
+			MetaData.SortOrder = ChannelIndex++;
+			MetaData.bCanCollapseToTrack = false;
+			TMovieSceneExternalValue<float> ExVal;
+			Channels.Add(Weight, MetaData, ExVal);
+#else
+			Channels.Add(Weight);
+
+#endif
 		}
+
+
 		ChannelProxy = MakeShared<FMovieSceneChannelProxy>(MoveTemp(Channels));
 	}
 }
@@ -1272,4 +1314,16 @@ FMovieSceneInterrogationKey UMovieSceneControlRigParameterSection::GetTransformI
 	return TypeID;
 }
 
+float UMovieSceneControlRigParameterSection::GetTotalWeightValue(FFrameTime InTime) const
+{
+	float WeightVal = EvaluateEasing(InTime);
+	if (EnumHasAllFlags(TransformMask.GetChannels(), EMovieSceneTransformChannel::Weight))
+	{
+		float ManualWeightVal = 1.f;
+		Weight.Evaluate(InTime, ManualWeightVal);
+		WeightVal *= ManualWeightVal;
+	}
+	return WeightVal;
+}
+
 #undef LOCTEXT_NAMESPACE 
diff --git a/Engine/Plugins/Experimental/ControlRig/Source/ControlRig/Private/Sequencer/MovieSceneControlRigParameterTemplate.cpp b/Engine/Plugins/Experimental/ControlRig/Source/ControlRig/Private/Sequencer/MovieSceneControlRigParameterTemplate.cpp
--- a/Engine/Plugins/Experimental/ControlRig/Source/ControlRig/Private/Sequencer/MovieSceneControlRigParameterTemplate.cpp
+++ b/Engine/Plugins/Experimental/ControlRig/Source/ControlRig/Private/Sequencer/MovieSceneControlRigParameterTemplate.cpp
@@ -932,16 +932,13 @@ void FMovieSceneControlRigParameterTemplate::Evaluate(const FMovieSceneEvaluatio
 		static TMovieSceneAnimTypeIDContainer<FString> Vector2DAnimTypeIDsByName;
 		static TMovieSceneAnimTypeIDContainer<FString> TransformAnimTypeIDsByName;
 
-
-		float Weight = 1.f;
-
-		//float Weight = EvaluateEasing(Context.GetTime());
-	//if (EnumHasAllFlags(TemplateData.Mask.GetChannels(), EMovieSceneTransformChannel::Weight))
-	//{
-	//	float ManualWeight = 1.f;
-	//	TemplateData.ManualWeight.Evaluate(Context.GetTime(), ManualWeight);
-	//	Weight *= ManualWeight;
-	//}
+		float Weight = EvaluateEasing(Context.GetTime());
+		if (EnumHasAllFlags(Section->TransformMask.GetChannels(), EMovieSceneTransformChannel::Weight))
+		{
+			float ManualWeight = 1.f;
+			Section->Weight.Evaluate(Context.GetTime(), ManualWeight);
+			Weight *= ManualWeight;
+		}
 
 		//Do Bool straight up no blending
 		bool bWasDoNotKey = false;
diff --git a/Engine/Plugins/Experimental/ControlRig/Source/ControlRig/Public/Sequencer/MovieSceneControlRigParameterSection.h b/Engine/Plugins/Experimental/ControlRig/Source/ControlRig/Public/Sequencer/MovieSceneControlRigParameterSection.h
--- a/Engine/Plugins/Experimental/ControlRig/Source/ControlRig/Public/Sequencer/MovieSceneControlRigParameterSection.h
+++ b/Engine/Plugins/Experimental/ControlRig/Source/ControlRig/Public/Sequencer/MovieSceneControlRigParameterSection.h
@@ -227,6 +227,8 @@ protected:
 	virtual void Serialize(FArchive& Ar) override;
 	virtual void PostEditImport() override;
 	virtual void PostLoad() override;
+	virtual float GetTotalWeightValue(FFrameTime InTime) const override;
+
 
 	// When true we do not set a key on the section, since it will be set because we changed the value
 	// We need this because control rig notifications are set on every change even when just changing sequencer time
