Missed one file from previous checkin

diff --git a/Engine/Plugins/Experimental/ControlRig/Source/ControlRig/Private/AnimNode_ControlRig.cpp b/Engine/Plugins/Experimental/ControlRig/Source/ControlRig/Private/AnimNode_ControlRig.cpp
--- a/Engine/Plugins/Experimental/ControlRig/Source/ControlRig/Private/AnimNode_ControlRig.cpp
+++ b/Engine/Plugins/Experimental/ControlRig/Source/ControlRig/Private/AnimNode_ControlRig.cpp
@@ -303,7 +303,7 @@ void FAnimNode_ControlRig::InitializeProperties(const UObject* InSourceInstance,
 		FName& SourceName = SourcePropertyNames[Idx];
 		UClass* SourceClass = InSourceInstance->GetClass();
 
-		FProperty* SourceProperty = FindField<FProperty>(SourceClass, SourceName);
+		FProperty* SourceProperty = FindFProperty<FProperty>(SourceClass, SourceName);
 		SourceProperties.Add(SourceProperty);
 		DestProperties.Add(nullptr);
 	}
