Control Rig: Ensure to initialize the control rig when the anim node initializes.

diff --git a/Engine/Plugins/Experimental/ControlRig/Source/ControlRig/Private/AnimNode_ControlRig.cpp b/Engine/Plugins/Experimental/ControlRig/Source/ControlRig/Private/AnimNode_ControlRig.cpp
--- a/Engine/Plugins/Experimental/ControlRig/Source/ControlRig/Private/AnimNode_ControlRig.cpp
+++ b/Engine/Plugins/Experimental/ControlRig/Source/ControlRig/Private/AnimNode_ControlRig.cpp
@@ -30,6 +30,7 @@ void FAnimNode_ControlRig::OnInitializeAnimInstance(const FAnimInstanceProxy* In
 	if (ControlRigClass)
 	{
 		ControlRig = NewObject<UControlRig>(InAnimInstance->GetOwningComponent(), ControlRigClass);
+		ControlRig->Initialize(true);
 		ControlRig->RequestInit();
 	}
 
