Control Rig: Avoid multiple initializations in Control Rig Editor

diff --git a/Engine/Plugins/Experimental/ControlRig/Source/ControlRig/Private/ControlRig.cpp b/Engine/Plugins/Experimental/ControlRig/Source/ControlRig/Private/ControlRig.cpp
--- a/Engine/Plugins/Experimental/ControlRig/Source/ControlRig/Private/ControlRig.cpp
+++ b/Engine/Plugins/Experimental/ControlRig/Source/ControlRig/Private/ControlRig.cpp
@@ -200,7 +200,11 @@ void UControlRig::Execute(const EControlRigState InState)
 	if (bRequiresInitExecution)
 	{
 		bRequiresInitExecution = false;
-		Execute(EControlRigState::Init);
+
+		if (InState != EControlRigState::Init)
+		{
+			Execute(EControlRigState::Init);
+		}
 	}
 
 	FRigUnitContext Context;
diff --git a/Engine/Plugins/Experimental/ControlRig/Source/ControlRig/Private/Sequencer/ControlRigLayerInstanceProxy.cpp b/Engine/Plugins/Experimental/ControlRig/Source/ControlRig/Private/Sequencer/ControlRigLayerInstanceProxy.cpp
--- a/Engine/Plugins/Experimental/ControlRig/Source/ControlRig/Private/Sequencer/ControlRigLayerInstanceProxy.cpp
+++ b/Engine/Plugins/Experimental/ControlRig/Source/ControlRig/Private/Sequencer/ControlRigLayerInstanceProxy.cpp
@@ -16,9 +16,14 @@ void FControlRigLayerInstanceProxy::Initialize(UAnimInstance* InAnimInstance)
 
 void FControlRigLayerInstanceProxy::CacheBones()
 {
-	FAnimationCacheBonesContext Context(this);
-	check(CurrentRoot);
-	CurrentRoot->CacheBones_AnyThread(Context);
+	if (bBoneCachesInvalidated)
+	{
+		FAnimationCacheBonesContext Context(this);
+		check(CurrentRoot);
+		CurrentRoot->CacheBones_AnyThread(Context);
+
+		bBoneCachesInvalidated = false;
+	}
 }
 
 bool FControlRigLayerInstanceProxy::Evaluate(FPoseContext& Output)
diff --git a/Engine/Plugins/Experimental/ControlRig/Source/ControlRigEditor/Private/Editor/ControlRigEditor.cpp b/Engine/Plugins/Experimental/ControlRig/Source/ControlRigEditor/Private/Editor/ControlRigEditor.cpp
--- a/Engine/Plugins/Experimental/ControlRig/Source/ControlRigEditor/Private/Editor/ControlRigEditor.cpp
+++ b/Engine/Plugins/Experimental/ControlRig/Source/ControlRigEditor/Private/Editor/ControlRigEditor.cpp
@@ -1239,6 +1239,11 @@ void FControlRigEditor::OnActiveTabChanged( TSharedPtr<SDockTab> PreviouslyActiv
 void FControlRigEditor::OnAnimInitialized()
 {
 	UControlRigSkeletalMeshComponent* EditorSkelComp = Cast<UControlRigSkeletalMeshComponent>(GetPersonaToolkit()->GetPreviewScene()->GetPreviewMeshComponent());
+	if (EditorSkelComp)
+	{
+		EditorSkelComp->bRequiredBonesUpToDateDuringTick = 0;
+	}
+
 	UControlRigLayerInstance* AnimInstance = Cast<UControlRigLayerInstance>(EditorSkelComp->GetAnimInstance());
 	if (AnimInstance && ControlRig)
 	{
diff --git a/Engine/Source/Editor/UnrealEd/Classes/Animation/DebugSkelMeshComponent.h b/Engine/Source/Editor/UnrealEd/Classes/Animation/DebugSkelMeshComponent.h
--- a/Engine/Source/Editor/UnrealEd/Classes/Animation/DebugSkelMeshComponent.h
+++ b/Engine/Source/Editor/UnrealEd/Classes/Animation/DebugSkelMeshComponent.h
@@ -229,6 +229,9 @@ class UNREALED_API UDebugSkelMeshComponent : public USkeletalMeshComponent
 	UPROPERTY(transient)
 	bool bClothCullBackface;
 
+	UPROPERTY(transient)
+	uint32 bRequiredBonesUpToDateDuringTick : 1;
+
 	/* Bounds computed from cloth. */
 	FBoxSphereBounds CachedClothBounds;
 
diff --git a/Engine/Source/Editor/UnrealEd/Private/Animation/DebugSkelMeshComponent.cpp b/Engine/Source/Editor/UnrealEd/Private/Animation/DebugSkelMeshComponent.cpp
--- a/Engine/Source/Editor/UnrealEd/Private/Animation/DebugSkelMeshComponent.cpp
+++ b/Engine/Source/Editor/UnrealEd/Private/Animation/DebugSkelMeshComponent.cpp
@@ -795,7 +795,10 @@ void UDebugSkelMeshComponent::TickComponent(float DeltaTime, enum ELevelTick Tic
 	bCachedMaterialParameterIndicesAreDirty = true;
 	
 	// Force retargeting data to be re-cached to take into account skeleton edits.
-	bRequiredBonesUpToDate = false;
+	if (bRequiredBonesUpToDateDuringTick)
+	{
+		bRequiredBonesUpToDate = false;
+	}
 
 	Super::TickComponent(DeltaTime, TickType, ThisTickFunction);
 
