Control Rig: Re-enable deprecated units for RigVM

diff --git a/Engine/Plugins/Experimental/ControlRig/Source/ControlRig/Private/Units/Deprecated/RigUnit_AimConstraint.cpp b/Engine/Plugins/Experimental/ControlRig/Source/ControlRig/Private/Units/Deprecated/RigUnit_AimConstraint.cpp
--- a/Engine/Plugins/Experimental/ControlRig/Source/ControlRig/Private/Units/Deprecated/RigUnit_AimConstraint.cpp
+++ b/Engine/Plugins/Experimental/ControlRig/Source/ControlRig/Private/Units/Deprecated/RigUnit_AimConstraint.cpp
@@ -5,10 +5,12 @@
 #include "HelperUtil.h"
 #include "AnimationCoreLibrary.h"
 
-void FRigUnit_AimConstraint::Execute(const FRigUnitContext& Context)
+FRigUnit_AimConstraint_Execute()
 {
     DECLARE_SCOPE_HIERARCHICAL_COUNTER_RIGUNIT()
 
+	TArray<FConstraintData>& ConstraintData = WorkData.ConstraintData;
+
 	if (Context.State == EControlRigState::Init)
 	{
 		ConstraintData.Reset();
diff --git a/Engine/Plugins/Experimental/ControlRig/Source/ControlRig/Private/Units/Deprecated/RigUnit_AimConstraint.h b/Engine/Plugins/Experimental/ControlRig/Source/ControlRig/Private/Units/Deprecated/RigUnit_AimConstraint.h
--- a/Engine/Plugins/Experimental/ControlRig/Source/ControlRig/Private/Units/Deprecated/RigUnit_AimConstraint.h
+++ b/Engine/Plugins/Experimental/ControlRig/Source/ControlRig/Private/Units/Deprecated/RigUnit_AimConstraint.h
@@ -49,15 +49,24 @@ struct FAimTarget
 	FVector AlignVector;
 };
 
+USTRUCT()
+struct FRigUnit_AimConstraint_WorkData
+{
+	GENERATED_BODY()
+
+	// note that Targets.Num () != ConstraintData.Num()
+	UPROPERTY()
+	TArray<FConstraintData>	ConstraintData;
+};
+
 USTRUCT(meta=(DisplayName="Aim Constraint", Category="Transforms", Deprecated = "4.23.0"))
 struct FRigUnit_AimConstraint : public FRigUnitMutable
 {
 	GENERATED_BODY()
 
+	RIGVM_METHOD()
 	virtual void Execute(const FRigUnitContext& Context) override;
 
-private:
-
 	UPROPERTY(EditAnywhere, Category = FRigUnit_AimConstraint, meta = (Input))
 	FName Joint;
 
@@ -84,5 +93,6 @@ private:
 	TArray<FAimTarget> UpTargets;
 
 	// note that Targets.Num () != ConstraintData.Num()
-	TArray<FConstraintData>	ConstraintData;
+	UPROPERTY()
+	FRigUnit_AimConstraint_WorkData WorkData;
 };
diff --git a/Engine/Plugins/Experimental/ControlRig/Source/ControlRig/Private/Units/Math/RigUnit_MathBool.cpp b/Engine/Plugins/Experimental/ControlRig/Source/ControlRig/Private/Units/Math/RigUnit_MathBool.cpp
--- a/Engine/Plugins/Experimental/ControlRig/Source/ControlRig/Private/Units/Math/RigUnit_MathBool.cpp
+++ b/Engine/Plugins/Experimental/ControlRig/Source/ControlRig/Private/Units/Math/RigUnit_MathBool.cpp
@@ -3,6 +3,14 @@
 #include "Units/Math/RigUnit_MathBool.h"
 #include "Units/RigUnitContext.h"
 
+FRigUnit_MathBoolConstTrue_Execute()
+{
+}
+
+FRigUnit_MathBoolConstFalse_Execute()
+{
+}
+
 FRigUnit_MathBoolNot_Execute()
 {
     DECLARE_SCOPE_HIERARCHICAL_COUNTER_RIGUNIT()
diff --git a/Engine/Plugins/Experimental/ControlRig/Source/ControlRig/Private/Units/Math/RigUnit_MathBool.h b/Engine/Plugins/Experimental/ControlRig/Source/ControlRig/Private/Units/Math/RigUnit_MathBool.h
--- a/Engine/Plugins/Experimental/ControlRig/Source/ControlRig/Private/Units/Math/RigUnit_MathBool.h
+++ b/Engine/Plugins/Experimental/ControlRig/Source/ControlRig/Private/Units/Math/RigUnit_MathBool.h
@@ -74,6 +74,9 @@ struct FRigUnit_MathBoolConstTrue : public FRigUnit_MathBoolConstant
 	{
 		Value = true;
 	}
+
+	RIGVM_METHOD()
+	virtual void Execute(const FRigUnitContext& Context) override;
 };
 
 /**
@@ -88,6 +91,9 @@ struct FRigUnit_MathBoolConstFalse : public FRigUnit_MathBoolConstant
 	{
 		Value = false;
 	}
+
+	RIGVM_METHOD()
+	virtual void Execute(const FRigUnitContext& Context) override;
 };
 
 /**
diff --git a/Engine/Plugins/Experimental/ControlRig/Source/ControlRig/Private/Units/Math/RigUnit_MathQuaternion.cpp b/Engine/Plugins/Experimental/ControlRig/Source/ControlRig/Private/Units/Math/RigUnit_MathQuaternion.cpp
--- a/Engine/Plugins/Experimental/ControlRig/Source/ControlRig/Private/Units/Math/RigUnit_MathQuaternion.cpp
+++ b/Engine/Plugins/Experimental/ControlRig/Source/ControlRig/Private/Units/Math/RigUnit_MathQuaternion.cpp
@@ -147,3 +147,7 @@ FRigUnit_MathQuaternionSwingTwist_Execute()
 	FVector NormalizedAxis = TwistAxis.GetSafeNormal();
 	Input.ToSwingTwist(NormalizedAxis, Swing, Twist);
 }
+
+FRigUnit_MathQuaternionRotationOrder_Execute()
+{
+}
\ No newline at end of file
diff --git a/Engine/Plugins/Experimental/ControlRig/Source/ControlRig/Private/Units/Math/RigUnit_MathQuaternion.h b/Engine/Plugins/Experimental/ControlRig/Source/ControlRig/Private/Units/Math/RigUnit_MathQuaternion.h
--- a/Engine/Plugins/Experimental/ControlRig/Source/ControlRig/Private/Units/Math/RigUnit_MathQuaternion.h
+++ b/Engine/Plugins/Experimental/ControlRig/Source/ControlRig/Private/Units/Math/RigUnit_MathQuaternion.h
@@ -508,4 +508,8 @@ struct FRigUnit_MathQuaternionRotationOrder : public FRigUnit_MathBase
 
 	UPROPERTY(meta = (Input, Output))
 	EControlRigRotationOrder RotationOrder;
+
+
+	RIGVM_METHOD()
+	virtual void Execute(const FRigUnitContext& Context) override;
 };
\ No newline at end of file
diff --git a/Engine/Source/Developer/RigVMDeveloper/Private/RigVMModel/RigVMController.cpp b/Engine/Source/Developer/RigVMDeveloper/Private/RigVMModel/RigVMController.cpp
--- a/Engine/Source/Developer/RigVMDeveloper/Private/RigVMModel/RigVMController.cpp
+++ b/Engine/Source/Developer/RigVMDeveloper/Private/RigVMModel/RigVMController.cpp
@@ -1956,11 +1956,6 @@ bool URigVMController::SetPinExpansion(const FString& InPinPath, bool bIsExpande
 
 bool URigVMController::SetPinExpansion(URigVMPin* InPin, bool bIsExpanded, bool bUndo)
 {
-	if (!IsValidPinForGraph(InPin))
-	{
-		return false;
-	}
-
 	if (InPin->GetSubPins().Num() == 0)
 	{
 		return false;
