Control Rig: Re-enable parameters temporarily to unblock TechAnim

diff --git a/Engine/Plugins/Experimental/ControlRig/Source/ControlRigEditor/Private/ControlRigEditorModule.cpp b/Engine/Plugins/Experimental/ControlRig/Source/ControlRigEditor/Private/ControlRigEditorModule.cpp
--- a/Engine/Plugins/Experimental/ControlRig/Source/ControlRigEditor/Private/ControlRigEditorModule.cpp
+++ b/Engine/Plugins/Experimental/ControlRig/Source/ControlRigEditor/Private/ControlRigEditorModule.cpp
@@ -351,7 +351,6 @@ void FControlRigEditorModule::GetTypeActions(const UControlRigBlueprint* CRB, FB
 			InActionRegistrar.AddBlueprintAction(InActionKey, NodeSpawnerSetter);
 		}
 
-		/*
 		static void AddParameterActions(UClass* InActionKey, FBlueprintActionDatabaseRegistrar& InActionRegistrar, const FEdGraphPinType& PinType, const FString& InCategory)
 		{
 			static const FString CategoryDelimiter(TEXT("|"));
@@ -362,14 +361,13 @@ void FControlRigEditorModule::GetTypeActions(const UControlRigBlueprint* CRB, FB
 			if (PinType.PinCategory == UEdGraphSchema_K2::PC_Struct)
 			{
 				if (UScriptStruct* Struct = Cast<UScriptStruct>(PinType.PinSubCategoryObject.Get()))
-	{
+				{
 					MenuDesc = FText::FromString(Struct->GetName());
 					ToolTip = MenuDesc;
-	}
-
+				}
 			}
 			else
-	{
+			{
 				MenuDesc = UEdGraphSchema_K2::GetCategoryText(PinType.PinCategory, true);
 				ToolTip = UEdGraphSchema_K2::GetCategoryText(PinType.PinCategory, false);
 			}
@@ -379,24 +377,24 @@ void FControlRigEditorModule::GetTypeActions(const UControlRigBlueprint* CRB, FB
 			check(NodeSpawnerGetter != nullptr);
 			InActionRegistrar.AddBlueprintAction(InActionKey, NodeSpawnerGetter);
 
-			UBlueprintNodeSpawner* NodeSpawnerSetter = UControlRigParameterNodeSpawner::CreateFromPinType(PinType, false, MenuDesc, NodeCategory, ToolTip);
-			check(NodeSpawnerSetter != nullptr);
-			InActionRegistrar.AddBlueprintAction(InActionKey, NodeSpawnerSetter);
-	}
-		*/
+			// let's disable setters for now
+			//UBlueprintNodeSpawner* NodeSpawnerSetter = UControlRigParameterNodeSpawner::CreateFromPinType(PinType, false, MenuDesc, NodeCategory, ToolTip);
+			//check(NodeSpawnerSetter != nullptr);
+			//InActionRegistrar.AddBlueprintAction(InActionKey, NodeSpawnerSetter);
+		}
 	};
 
 	FString VariableCategory = LOCTEXT("NewVariable", "New Variable").ToString();
-	//FString ParameterCategory = LOCTEXT("NewParameter", "New Parameter").ToString();
+	FString ParameterCategory = LOCTEXT("NewParameter", "New Parameter").ToString();
 	for (const FEdGraphPinType& PinType: PinTypes)
 	{
 		Local::AddVariableActions(ActionKey, ActionRegistrar, PinType, VariableCategory);
 		// let's disable parameters in the UI for now.
-		//Local::AddParameterActions(ActionKey, ActionRegistrar, PinType, ParameterCategory);
+		Local::AddParameterActions(ActionKey, ActionRegistrar, PinType, ParameterCategory);
 	}
 
 	// add support for names as parameters
-	//Local::AddParameterActions(ActionKey, ActionRegistrar, FEdGraphPinType(UEdGraphSchema_K2::PC_Name, FName(NAME_None), nullptr, EPinContainerType::None, false, FEdGraphTerminalType()), ParameterCategory);
+	Local::AddParameterActions(ActionKey, ActionRegistrar, FEdGraphPinType(UEdGraphSchema_K2::PC_Name, FName(NAME_None), nullptr, EPinContainerType::None, false, FEdGraphTerminalType()), ParameterCategory);
 
 	UBlueprintNodeSpawner* RerouteNodeSpawner = UControlRigRerouteNodeSpawner::CreateGeneric(
 		LOCTEXT("RerouteSpawnerDesc", "New Reroute Node"),
