Control Rig: Fix static analysis warning

diff --git a/Engine/Plugins/Experimental/ControlRig/Source/ControlRigEditor/Private/Editor/ControlRigEditor.cpp b/Engine/Plugins/Experimental/ControlRig/Source/ControlRigEditor/Private/Editor/ControlRigEditor.cpp
--- a/Engine/Plugins/Experimental/ControlRig/Source/ControlRigEditor/Private/Editor/ControlRigEditor.cpp
+++ b/Engine/Plugins/Experimental/ControlRig/Source/ControlRigEditor/Private/Editor/ControlRigEditor.cpp
@@ -1242,16 +1242,16 @@ void FControlRigEditor::OnAnimInitialized()
 	if (EditorSkelComp)
 	{
 		EditorSkelComp->bRequiredBonesUpToDateDuringTick = 0;
-	}
 
-	UControlRigLayerInstance* AnimInstance = Cast<UControlRigLayerInstance>(EditorSkelComp->GetAnimInstance());
-	if (AnimInstance && ControlRig)
-	{
-		// update control rig data to anim instance since animation system has been reinitialized
-		FInputBlendPose Filter;
-		AnimInstance->ResetControlRigTracks();
-		AnimInstance->AddControlRigTrack(0, ControlRig);
-		AnimInstance->UpdateControlRigTrack(0, 1.0f, FControlRigIOSettings::MakeEnabled(), bExecutionControlRig);
+		UControlRigLayerInstance* AnimInstance = Cast<UControlRigLayerInstance>(EditorSkelComp->GetAnimInstance());
+		if (AnimInstance && ControlRig)
+		{
+			// update control rig data to anim instance since animation system has been reinitialized
+			FInputBlendPose Filter;
+			AnimInstance->ResetControlRigTracks();
+			AnimInstance->AddControlRigTrack(0, ControlRig);
+			AnimInstance->UpdateControlRigTrack(0, 1.0f, FControlRigIOSettings::MakeEnabled(), bExecutionControlRig);
+		}
 	}
 }
 
