UE-91377 - Ambiguous search when changing material of instanced mesh

Changed the name of the default object within ControlRigBlueprint to RigVMModel.

diff --git a/Engine/Plugins/Experimental/ControlRig/Source/ControlRigDeveloper/Private/ControlRigBlueprint.cpp b/Engine/Plugins/Experimental/ControlRig/Source/ControlRigDeveloper/Private/ControlRigBlueprint.cpp
--- a/Engine/Plugins/Experimental/ControlRig/Source/ControlRigDeveloper/Private/ControlRigBlueprint.cpp
+++ b/Engine/Plugins/Experimental/ControlRig/Source/ControlRigDeveloper/Private/ControlRigBlueprint.cpp
@@ -44,7 +44,7 @@ UControlRigBlueprint::UControlRigBlueprint(const FObjectInitializer& ObjectIniti
 	bVMRecompilationRequired = false;
 	VMRecompilationBracket = 0;
 
-	Model = ObjectInitializer.CreateDefaultSubobject<URigVMGraph>(this, TEXT("Model"));
+	Model = ObjectInitializer.CreateDefaultSubobject<URigVMGraph>(this, TEXT("RigVMModel"));
 	Controller = nullptr;
 }
 
