Control Rig: Ensure the CR editor creates a mode manager, otherwise it gets crashy when constructing the preview viewport.

#rnx

diff --git a/Engine/Plugins/Experimental/ControlRig/Source/ControlRigEditor/Private/Editor/ControlRigEditor.cpp b/Engine/Plugins/Experimental/ControlRig/Source/ControlRigEditor/Private/Editor/ControlRigEditor.cpp
--- a/Engine/Plugins/Experimental/ControlRig/Source/ControlRigEditor/Private/Editor/ControlRigEditor.cpp
+++ b/Engine/Plugins/Experimental/ControlRig/Source/ControlRigEditor/Private/Editor/ControlRigEditor.cpp
@@ -71,6 +71,8 @@
 #include "UnrealExporter.h"
 #include "ControlRigElementDetails.h"
 #include "PropertyEditorModule.h"
+#include "AssetEditorModeManager.h"
+#include "IPersonaEditorModeManager.h"
 
 #define LOCTEXT_NAMESPACE "ControlRigEditor"
 
@@ -1404,6 +1406,17 @@ void FControlRigEditor::HandleControlRigInitializedEvent(UControlRig* InControlR
 	UpdateGraphCompilerErrors();
 }
 
+
+void FControlRigEditor::CreateEditorModeManager()
+{
+	TSharedPtr<FAssetEditorModeManager> ModeManager = MakeShareable(FModuleManager::LoadModuleChecked<FPersonaModule>("Persona").CreatePersonaEditorModeManager());
+
+	PRAGMA_DISABLE_DEPRECATION_WARNINGS
+	SetAssetEditorModeManager(ModeManager.Get());
+	PRAGMA_ENABLE_DEPRECATION_WARNINGS
+}
+
+
 void FControlRigEditor::Tick(float DeltaTime)
 {
 	FBlueprintEditor::Tick(DeltaTime);
diff --git a/Engine/Plugins/Experimental/ControlRig/Source/ControlRigEditor/Private/Editor/ControlRigEditor.h b/Engine/Plugins/Experimental/ControlRig/Source/ControlRigEditor/Private/Editor/ControlRigEditor.h
--- a/Engine/Plugins/Experimental/ControlRig/Source/ControlRigEditor/Private/Editor/ControlRigEditor.h
+++ b/Engine/Plugins/Experimental/ControlRig/Source/ControlRigEditor/Private/Editor/ControlRigEditor.h
@@ -177,6 +177,9 @@ protected:
 	void HandleVMCompiledEvent(UBlueprint* InBlueprint, URigVM* InVM);
 	void HandleControlRigInitializedEvent(UControlRig* InControlRig, const EControlRigState InState);
 
+	// FBaseToolKit overrides
+	void CreateEditorModeManager() override;
+
 	// FGCObject Interface
 	virtual void AddReferencedObjects( FReferenceCollector& Collector ) override;
 
