Fix Incremental #fyi louise.rasmussen

diff --git a/Engine/Plugins/Experimental/ControlRig/Source/ControlRigEditor/Private/ControlRigEditorStyle.h b/Engine/Plugins/Experimental/ControlRig/Source/ControlRigEditor/Private/ControlRigEditorStyle.h
--- a/Engine/Plugins/Experimental/ControlRig/Source/ControlRigEditor/Private/ControlRigEditorStyle.h
+++ b/Engine/Plugins/Experimental/ControlRig/Source/ControlRigEditor/Private/ControlRigEditorStyle.h
@@ -34,17 +34,8 @@ public:
 		}
 
 		// Edit mode styles
-
-		if (FCoreStyle::IsStarshipStyle())
-		{
-			Set("ControlRigEditMode", new IMAGE_BRUSH_SVG("Slate/animation", Icon40x40));
-			Set("ControlRigEditMode.Small", new IMAGE_BRUSH_SVG("Slate/animation", Icon20x20));
-		}
-		else
-		{
-			Set("ControlRigEditMode", new IMAGE_BRUSH("Slate/ControlRigMode_40x", Icon40x40));
-			Set("ControlRigEditMode.Small", new IMAGE_BRUSH("Slate/ControlRigMode_40x", Icon20x20));
-		}
+		Set("ControlRigEditMode", new IMAGE_BRUSH_SVG("Slate/animation", Icon40x40));
+		Set("ControlRigEditMode.Small", new IMAGE_BRUSH_SVG("Slate/animation", Icon20x20));
 
 		// Sequencer styles
 		{
