[NO NON-PII COMMIT MESSAGE]

diff --git a/Engine/Plugins/Experimental/ControlRig/Source/ControlRigEditor/Private/Editor/ControlRigEditor.cpp b/Engine/Plugins/Experimental/ControlRig/Source/ControlRigEditor/Private/Editor/ControlRigEditor.cpp
--- a/Engine/Plugins/Experimental/ControlRig/Source/ControlRigEditor/Private/Editor/ControlRigEditor.cpp
+++ b/Engine/Plugins/Experimental/ControlRig/Source/ControlRigEditor/Private/Editor/ControlRigEditor.cpp
@@ -244,10 +244,7 @@ void FControlRigEditor::InitControlRigEditor(const EToolkitMode::Type Mode, cons
 		
 		PersonaToolkit->GetPreviewScene()->SetRemoveAttachedComponentFilter(FOnRemoveAttachedComponentFilter::CreateSP(EditMode, &FControlRigEditMode::CanRemoveFromPreviewScene));
 
-		if (InControlRigBlueprint)
-		{
-			InControlRigBlueprint->OnModified().AddSP(EditMode, &FControlRigEditMode::HandleModifiedEvent);
-		}
+		InControlRigBlueprint->OnModified().AddSP(EditMode, &FControlRigEditMode::HandleModifiedEvent);
 	}
 
 	UpdateControlRig();
@@ -276,7 +273,6 @@ void FControlRigEditor::InitControlRigEditor(const EToolkitMode::Type Mode, cons
 		}
 	}
 
-	if (InControlRigBlueprint)
 	{
 		if (URigVMGraph* Model = InControlRigBlueprint->Model)
 		{
