Fixing another global editor mode manager shutdown lifecycle crash.

#rnx

diff --git a/Engine/Plugins/Experimental/ControlRig/Source/ControlRigEditor/Private/EditMode/ControlRigEditMode.cpp b/Engine/Plugins/Experimental/ControlRig/Source/ControlRigEditor/Private/EditMode/ControlRigEditMode.cpp
--- a/Engine/Plugins/Experimental/ControlRig/Source/ControlRigEditor/Private/EditMode/ControlRigEditMode.cpp
+++ b/Engine/Plugins/Experimental/ControlRig/Source/ControlRigEditor/Private/EditMode/ControlRigEditMode.cpp
@@ -90,9 +90,6 @@ FControlRigEditMode::~FControlRigEditMode()
 {	
 	CommandBindings = nullptr;
 
-	GLevelEditorModeTools().OnWidgetModeChanged().RemoveAll(this);
-	GLevelEditorModeTools().OnCoordSystemChanged().RemoveAll(this);
-
 #if WITH_EDITOR
 	if (GEditor)
 	{
