Rename TAssetPtr uses to TSoftObjectPtr in engine code, this was originally changed in 4.18 but not properly deprecated at the time

diff --git a/Engine/Plugins/Experimental/ControlRig/Source/ControlRig/Public/ControlRig.h b/Engine/Plugins/Experimental/ControlRig/Source/ControlRig/Public/ControlRig.h
--- a/Engine/Plugins/Experimental/ControlRig/Source/ControlRig/Public/ControlRig.h
+++ b/Engine/Plugins/Experimental/ControlRig/Source/ControlRig/Public/ControlRig.h
@@ -248,7 +248,7 @@ private:
 	FRigHierarchyContainer Hierarchy;
 
 	UPROPERTY()
-	TAssetPtr<UControlRigGizmoLibrary> GizmoLibrary;
+	TSoftObjectPtr<UControlRigGizmoLibrary> GizmoLibrary;
 
 	/** Runtime object binding */
 	TSharedPtr<IControlRigObjectBinding> ObjectBinding;
diff --git a/Engine/Plugins/Experimental/ControlRig/Source/ControlRig/Public/ControlRigGizmoLibrary.h b/Engine/Plugins/Experimental/ControlRig/Source/ControlRig/Public/ControlRigGizmoLibrary.h
--- a/Engine/Plugins/Experimental/ControlRig/Source/ControlRig/Public/ControlRigGizmoLibrary.h
+++ b/Engine/Plugins/Experimental/ControlRig/Source/ControlRig/Public/ControlRigGizmoLibrary.h
@@ -26,7 +26,7 @@ struct CONTROLRIG_API FControlRigGizmoDefinition
 	FName GizmoName;
 
 	UPROPERTY(EditAnywhere, Category = "Gizmo")
-	TAssetPtr<UStaticMesh> StaticMesh;
+	TSoftObjectPtr<UStaticMesh> StaticMesh;
 
 	UPROPERTY(EditAnywhere, Category = "Gizmo")
 	FTransform Transform;
@@ -45,7 +45,7 @@ public:
 	FControlRigGizmoDefinition DefaultGizmo;
 
 	UPROPERTY(EditAnywhere, Category = "GizmoLibrary")
-	TAssetPtr<UMaterial> DefaultMaterial;
+	TSoftObjectPtr<UMaterial> DefaultMaterial;
 
 	UPROPERTY(EditAnywhere, Category = "GizmoLibrary")
 	FName MaterialColorParameter;
diff --git a/Engine/Plugins/Experimental/ControlRig/Source/ControlRig/Public/Settings/ControlRigSettings.h b/Engine/Plugins/Experimental/ControlRig/Source/ControlRig/Public/Settings/ControlRigSettings.h
--- a/Engine/Plugins/Experimental/ControlRig/Source/ControlRig/Public/Settings/ControlRigSettings.h
+++ b/Engine/Plugins/Experimental/ControlRig/Source/ControlRig/Public/Settings/ControlRigSettings.h
@@ -25,7 +25,7 @@ public:
 
 #if WITH_EDITORONLY_DATA
 	UPROPERTY(EditAnywhere, config, Category = DefaultGizmo)
-	TAssetPtr<UControlRigGizmoLibrary> DefaultGizmoLibrary;
+	TSoftObjectPtr<UControlRigGizmoLibrary> DefaultGizmoLibrary;
 #endif
 
 	static UControlRigSettings * Get() { return CastChecked<UControlRigSettings>(UControlRigSettings::StaticClass()->GetDefaultObject()); }
diff --git a/Engine/Plugins/Experimental/ControlRig/Source/ControlRigDeveloper/Public/ControlRigBlueprint.h b/Engine/Plugins/Experimental/ControlRig/Source/ControlRigDeveloper/Public/ControlRigBlueprint.h
--- a/Engine/Plugins/Experimental/ControlRig/Source/ControlRigDeveloper/Public/ControlRigBlueprint.h
+++ b/Engine/Plugins/Experimental/ControlRig/Source/ControlRigDeveloper/Public/ControlRigBlueprint.h
@@ -120,7 +120,7 @@ public:
 
 #if WITH_EDITORONLY_DATA
 	UPROPERTY(EditAnywhere, config, Category = DefaultGizmo)
-	TAssetPtr<UControlRigGizmoLibrary> GizmoLibrary;
+	TSoftObjectPtr<UControlRigGizmoLibrary> GizmoLibrary;
 #endif
 
 	UPROPERTY(transient, VisibleAnywhere, Category = "VM", meta = (DisplayName = "VM Statistics", DisplayAfter = "VMCompileSettings"))
diff --git a/Engine/Plugins/Experimental/ControlRig/Source/ControlRigManipulation/Public/ControlRigGizmoActor.h b/Engine/Plugins/Experimental/ControlRig/Source/ControlRigManipulation/Public/ControlRigGizmoActor.h
--- a/Engine/Plugins/Experimental/ControlRig/Source/ControlRigManipulation/Public/ControlRigGizmoActor.h
+++ b/Engine/Plugins/Experimental/ControlRig/Source/ControlRigManipulation/Public/ControlRigGizmoActor.h
@@ -35,8 +35,8 @@ struct FGizmoActorCreationParam
 	FTransform	SpawnTransform;
 	FTransform  GizmoTransform;
 	FTransform  MeshTransform;
-	TAssetPtr<UStaticMesh> StaticMesh;
-	TAssetPtr<UMaterial> Material;
+	TSoftObjectPtr<UStaticMesh> StaticMesh;
+	TSoftObjectPtr<UMaterial> Material;
 	FName ColorParameterName;
 	FLinearColor Color;
 };
diff --git a/Engine/Source/Editor/UnrealEd/Classes/Editor/PropertyEditorTestObject.h b/Engine/Source/Editor/UnrealEd/Classes/Editor/PropertyEditorTestObject.h
--- a/Engine/Source/Editor/UnrealEd/Classes/Editor/PropertyEditorTestObject.h
+++ b/Engine/Source/Editor/UnrealEd/Classes/Editor/PropertyEditorTestObject.h
@@ -283,7 +283,7 @@ class UPropertyEditorTestObject : public UObject
 	TSubclassOf<UTexture> SubclassOfWithDisallowed;
 
 	UPROPERTY(EditAnywhere, Category=Classes, meta=(AllowedClasses="StaticMesh,  SkeletalMesh	"))
-	TAssetPtr<UObject> AssetPointerWithAllowedAndWhitespace;
+	TSoftObjectPtr<UObject> AssetPointerWithAllowedAndWhitespace;
 
 	// Integer
 	UPROPERTY(EditAnywhere, Category=ArraysOfProperties)
