ControlRig: Deprecate Blend(Transform) node

diff --git a/Engine/Plugins/Experimental/ControlRig/Source/ControlRig/Private/Units/Deprecated/RigUnit_BlendTransform.h b/Engine/Plugins/Experimental/ControlRig/Source/ControlRig/Private/Units/Deprecated/RigUnit_BlendTransform.h
--- a/Engine/Plugins/Experimental/ControlRig/Source/ControlRig/Private/Units/Deprecated/RigUnit_BlendTransform.h
+++ b/Engine/Plugins/Experimental/ControlRig/Source/ControlRig/Private/Units/Deprecated/RigUnit_BlendTransform.h
@@ -21,7 +21,7 @@ struct FBlendTarget
 	{}
 };
 
-USTRUCT(meta = (DisplayName = "Blend(Transform)", Category = "Blend"))
+USTRUCT(meta = (DisplayName = "Blend(Transform)", Category = "Blend", Deprecated = "5.0.0"))
 struct FRigUnit_BlendTransform : public FRigUnit
 {
 	GENERATED_BODY()
