Control Rig & Animation Runtime: Fix deprecation version and deprecation warning for CL14376957: Animation Runtime: Deprecated FAnimationRuntime::BlendTransformsByWeight

diff --git a/Engine/Plugins/Experimental/ControlRig/Source/ControlRig/Private/Units/Deprecated/RigUnit_BlendTransform.cpp b/Engine/Plugins/Experimental/ControlRig/Source/ControlRig/Private/Units/Deprecated/RigUnit_BlendTransform.cpp
--- a/Engine/Plugins/Experimental/ControlRig/Source/ControlRig/Private/Units/Deprecated/RigUnit_BlendTransform.cpp
+++ b/Engine/Plugins/Experimental/ControlRig/Source/ControlRig/Private/Units/Deprecated/RigUnit_BlendTransform.cpp
@@ -38,8 +38,11 @@ FRigUnit_BlendTransform_Execute()
 				BlendTransform.Add(Source);
 				BlendWeights.Add(SourceWeight);
 			}
-
+			
+			PRAGMA_DISABLE_DEPRECATION_WARNINGS
+			// deprecated node calls a deprecated function, suppress warning for this case
 			FAnimationRuntime::BlendTransformsByWeight(Result, BlendTransform, BlendWeights);
+			PRAGMA_ENABLE_DEPRECATION_WARNINGS
 		}
 	}
 
diff --git a/Engine/Source/Runtime/Engine/Public/AnimationRuntime.h b/Engine/Source/Runtime/Engine/Public/AnimationRuntime.h
--- a/Engine/Source/Runtime/Engine/Public/AnimationRuntime.h
+++ b/Engine/Source/Runtime/Engine/Public/AnimationRuntime.h
@@ -361,7 +361,7 @@ public:
 	 *
 	 * @param OutTransform : result
 	 */
-	UE_DEPRECATED(4.26, "Please use BlendTransform() for weighted blending")
+	UE_DEPRECATED(5.0.0, "Please use BlendTransform() for weighted blending")
 	static void BlendTransformsByWeight(FTransform& OutTransform, const TArray<FTransform>& Transforms, const TArray<float>& Weights);
 
 	/** 
