Edigrate of Back out changelist 14362516 from Release-4.26, author Jack.Cai

diff --git a/Engine/Plugins/Experimental/ControlRig/Source/ControlRig/Private/ControlRig.cpp b/Engine/Plugins/Experimental/ControlRig/Source/ControlRig/Private/ControlRig.cpp
--- a/Engine/Plugins/Experimental/ControlRig/Source/ControlRig/Private/ControlRig.cpp
+++ b/Engine/Plugins/Experimental/ControlRig/Source/ControlRig/Private/ControlRig.cpp
@@ -1165,7 +1165,6 @@ FRigControlValue UControlRig::GetControlValueFromGlobalTransform(const FName& In
 		}
 
 		FTransform Transform = InGlobalTransform.GetRelativeTransform(ParentTransform);
-		Transform = Transform.GetRelativeTransform(Control->OffsetTransform);
 
 		FRigControlValue PreviousVal = Control->Value;
 		Control->SetValueFromTransform(Transform, ERigControlValueType::Current);
@@ -1283,7 +1282,7 @@ void UControlRig::UpdateAvailableControls()
 	AvailableControlsOverride.Append(TransientControls);
 }
 
-FName UControlRig::AddTransientControl(URigVMPin* InPin, FRigElementKey SpaceKey, FTransform OffsetTransform)
+FName UControlRig::AddTransientControl(URigVMPin* InPin, FRigElementKey SpaceKey)
 {
 	if (InPin == nullptr)
 	{
@@ -1325,8 +1324,6 @@ FName UControlRig::AddTransientControl(URigVMPin* InPin, FRigElementKey SpaceKey
 		}
 	}
 
-	NewControl.OffsetTransform = OffsetTransform;
-
 	TransientControls.Add(NewControl);
 
 	SetTransientControlValue(InPin);
diff --git a/Engine/Plugins/Experimental/ControlRig/Source/ControlRig/Private/Units/Highlevel/Hierarchy/RigUnit_ModifyTransforms.h b/Engine/Plugins/Experimental/ControlRig/Source/ControlRig/Private/Units/Highlevel/Hierarchy/RigUnit_ModifyTransforms.h
--- a/Engine/Plugins/Experimental/ControlRig/Source/ControlRig/Private/Units/Highlevel/Hierarchy/RigUnit_ModifyTransforms.h
+++ b/Engine/Plugins/Experimental/ControlRig/Source/ControlRig/Private/Units/Highlevel/Hierarchy/RigUnit_ModifyTransforms.h
@@ -97,51 +97,24 @@ struct FRigUnit_ModifyTransforms : public FRigUnit_HighlevelBaseMutable
 
 			if (ItemToModify.IsValidIndex(Index))
 			{
-				if (Mode == EControlRigModifyBoneMode::AdditiveLocal || Mode == EControlRigModifyBoneMode::OverrideLocal)
-				{
-					if (const FRigHierarchyContainer* Container = (const FRigHierarchyContainer*)InUserContext)
-					{
-						if (ItemToModify[Index].Item.IsValid())
-						{
-							return Container->GetParentKey(ItemToModify[Index].Item);
-						}
-					}
-				}
-			}
-		}
-		return FRigElementKey();
-	}
-
-	virtual FTransform DetermineOffsetTransformForPin(const FString& InPinPath, void* InUserContext) const override
-	{
-		if (InPinPath.StartsWith(TEXT("ItemToModify")))
-		{
-			int32 Index = INDEX_NONE;
-			FString Left, Middle, Right;
-			if (InPinPath.Replace(TEXT("["), TEXT(".")).Split(TEXT("."), &Left, &Middle))
-			{
-				if (Middle.Replace(TEXT("]"), TEXT(".")).Split(TEXT("."), &Left, &Right))
+				if (Mode == EControlRigModifyBoneMode::AdditiveLocal)
 				{
-					Index = FCString::Atoi(*Left);
+					return ItemToModify[Index].Item;
 				}
-			}
 
-			if (ItemToModify.IsValidIndex(Index))
-			{
-				if (Mode == EControlRigModifyBoneMode::AdditiveLocal)
+				if (Mode == EControlRigModifyBoneMode::OverrideLocal)
 				{
 					if (const FRigHierarchyContainer* Container = (const FRigHierarchyContainer*)InUserContext)
 					{
 						if (ItemToModify[Index].Item.IsValid())
 						{
-							return Container->GetLocalTransform(ItemToModify[Index].Item);
+							return Container->GetParentKey(ItemToModify[Index].Item);
 						}
 					}
 				}
 			}
 		}
-
-		return FTransform::Identity;
+		return FRigElementKey();
 	}
 
 	RIGVM_METHOD()
diff --git a/Engine/Plugins/Experimental/ControlRig/Source/ControlRig/Public/ControlRig.h b/Engine/Plugins/Experimental/ControlRig/Source/ControlRig/Public/ControlRig.h
--- a/Engine/Plugins/Experimental/ControlRig/Source/ControlRig/Public/ControlRig.h
+++ b/Engine/Plugins/Experimental/ControlRig/Source/ControlRig/Public/ControlRig.h
@@ -454,7 +454,7 @@ private:
 	void UpdateAvailableControls();
 
 	/** Remove a transient / temporary control used to interact with a pin */
-	FName AddTransientControl(URigVMPin* InPin, FRigElementKey SpaceKey = FRigElementKey(), FTransform OffsetTransform = FTransform::Identity);
+	FName AddTransientControl(URigVMPin* InPin, FRigElementKey SpaceKey = FRigElementKey());
 
 	/** Sets the value of a transient control based on a pin */
 	bool SetTransientControlValue(URigVMPin* InPin);
diff --git a/Engine/Plugins/Experimental/ControlRig/Source/ControlRig/Public/Units/RigUnit.h b/Engine/Plugins/Experimental/ControlRig/Source/ControlRig/Public/Units/RigUnit.h
--- a/Engine/Plugins/Experimental/ControlRig/Source/ControlRig/Public/Units/RigUnit.h
+++ b/Engine/Plugins/Experimental/ControlRig/Source/ControlRig/Public/Units/RigUnit.h
@@ -28,8 +28,6 @@ struct CONTROLRIG_API FRigUnit : public FRigVMStruct
 	virtual void Execute(const FRigUnitContext& Context) {}
 
 	virtual FRigElementKey DetermineSpaceForPin(const FString& InPinPath, void* InUserContext) const { return FRigElementKey(); }
-
-	virtual FTransform DetermineOffsetTransformForPin(const FString& InPinPath, void* InUserContext) const { return FTransform::Identity; }
 };
 
 /** Base class for all rig units that can change data */
diff --git a/Engine/Plugins/Experimental/ControlRig/Source/ControlRigDeveloper/Private/ControlRigBlueprint.cpp b/Engine/Plugins/Experimental/ControlRig/Source/ControlRigDeveloper/Private/ControlRigBlueprint.cpp
--- a/Engine/Plugins/Experimental/ControlRig/Source/ControlRigDeveloper/Private/ControlRigBlueprint.cpp
+++ b/Engine/Plugins/Experimental/ControlRig/Source/ControlRigDeveloper/Private/ControlRigBlueprint.cpp
@@ -523,7 +523,6 @@ FName UControlRigBlueprint::AddTransientControl(URigVMPin* InPin)
 	UControlRig* CDO = Cast<UControlRig>(RigClass->GetDefaultObject(true /* create if needed */));
 
 	FRigElementKey SpaceKey;
-	FTransform OffsetTransform = FTransform::Identity;
 	if (URigVMStructNode* StructNode = Cast<URigVMStructNode>(InPin->GetPinForLink()->GetNode()))
 	{
 		if (TSharedPtr<FStructOnScope> DefaultStructScope = StructNode->ConstructStructInstance())
@@ -534,9 +533,8 @@ FName UControlRigBlueprint::AddTransientControl(URigVMPin* InPin)
 			FString Left, Right;
 
 			if (URigVMPin::SplitPinPathAtStart(PinPath, Left, Right))
-			{
+		{
 				SpaceKey = DefaultStruct->DetermineSpaceForPin(Right, &HierarchyContainer);
-				OffsetTransform = DefaultStruct->DetermineOffsetTransformForPin(Right, &HierarchyContainer);
 			}
 		}
 	}
@@ -549,7 +547,7 @@ FName UControlRigBlueprint::AddTransientControl(URigVMPin* InPin)
 		UControlRig* InstancedControlRig = Cast<UControlRig>(ArchetypeInstance);
 		if (InstancedControlRig)
 		{
-			FName ControlName = InstancedControlRig->AddTransientControl(InPin, SpaceKey, OffsetTransform);
+			FName ControlName = InstancedControlRig->AddTransientControl(InPin, SpaceKey);
 			if (ReturnName == NAME_None)
 			{
 				ReturnName = ControlName;
@@ -559,13 +557,6 @@ FName UControlRigBlueprint::AddTransientControl(URigVMPin* InPin)
 
 	if (ReturnName != NAME_None)
 	{
-		// de-select all elements so that they don't trigger "cleartransientcontrol()" in "OnElementAdded => OnHierachyChanged"
-		TArray<FRigElementKey> SelectedElements = HierarchyContainer.CurrentSelection();
-		for (const FRigElementKey& SelectedElement : SelectedElements)
-		{
-			HierarchyContainer.OnElementSelected.Broadcast(&HierarchyContainer, SelectedElement, false);
-		}
-
 		HierarchyContainer.OnElementAdded.Broadcast(&HierarchyContainer, FRigElementKey(ReturnName, ERigElementType::Control));
 		HierarchyContainer.OnElementSelected.Broadcast(&HierarchyContainer, FRigElementKey(ReturnName, ERigElementType::Control), true);
 	}
@@ -1553,10 +1544,10 @@ void UControlRigBlueprint::PropagateHierarchyFromBPToInstances(bool bInitializeC
 				DefaultObject->Hierarchy.Initialize(false);
 			}
 
-			TArray<UObject*> ArchetypeInstances;
-			DefaultObject->GetArchetypeInstances(ArchetypeInstances);
-			for (UObject* ArchetypeInstance : ArchetypeInstances)
-			{
+						TArray<UObject*> ArchetypeInstances;
+						DefaultObject->GetArchetypeInstances(ArchetypeInstances);
+						for (UObject* ArchetypeInstance : ArchetypeInstances)
+						{
 				if (UControlRig* InstanceRig = Cast<UControlRig>(ArchetypeInstance))
 				{
 					InstanceRig->Hierarchy = HierarchyContainer;
@@ -1789,7 +1780,7 @@ void UControlRigBlueprint::HandleOnElementSelected(FRigHierarchyContainer* InCon
 		{
 			if (FRigControl* Control = RigBeingDebugged->FindControl(InKey.Name))
 			{
-				if (!Control->bIsTransientControl && bSelected)
+				if (!Control->bIsTransientControl)
 				{
 					ClearTransientControls();
 				}
