Control Rig & Animation Runtime: Deprecate Blend(Transform) Node and FAnimationRuntime::BlendTransformsByWeight in 4.26

diff --git a/Engine/Plugins/Experimental/ControlRig/Source/ControlRig/Private/Units/Deprecated/RigUnit_BlendTransform.h b/Engine/Plugins/Experimental/ControlRig/Source/ControlRig/Private/Units/Deprecated/RigUnit_BlendTransform.h
--- a/Engine/Plugins/Experimental/ControlRig/Source/ControlRig/Private/Units/Deprecated/RigUnit_BlendTransform.h
+++ b/Engine/Plugins/Experimental/ControlRig/Source/ControlRig/Private/Units/Deprecated/RigUnit_BlendTransform.h
@@ -21,7 +21,7 @@ struct FBlendTarget
 	{}
 };
 
-USTRUCT(meta = (DisplayName = "Blend(Transform)", Category = "Blend", Deprecated = "5.0.0"))
+USTRUCT(meta = (DisplayName = "Blend(Transform)", Category = "Blend", Deprecated = "4.26.0"))
 struct FRigUnit_BlendTransform : public FRigUnit
 {
 	GENERATED_BODY()
diff --git a/Engine/Source/Runtime/Engine/Public/AnimationRuntime.h b/Engine/Source/Runtime/Engine/Public/AnimationRuntime.h
--- a/Engine/Source/Runtime/Engine/Public/AnimationRuntime.h
+++ b/Engine/Source/Runtime/Engine/Public/AnimationRuntime.h
@@ -428,7 +428,7 @@ public:
 	 *
 	 * @param OutTransform : result
 	 */
-	UE_DEPRECATED(5.0.0, "Please use BlendTransform() for weighted blending")
+	UE_DEPRECATED(4.26.0, "Please use BlendTransform() for weighted blending")
 	static void BlendTransformsByWeight(FTransform& OutTransform, const TArray<FTransform>& Transforms, const TArray<float>& Weights);
 
 	/** 
