ControlRig: Remove stray default constructor declaration. UControlRig already implements the UObjet initializing constructor.

#rnx

diff --git a/Engine/Plugins/Experimental/ControlRig/Source/ControlRig/Public/ControlRig.h b/Engine/Plugins/Experimental/ControlRig/Source/ControlRig/Public/ControlRig.h
--- a/Engine/Plugins/Experimental/ControlRig/Source/ControlRig/Public/ControlRig.h
+++ b/Engine/Plugins/Experimental/ControlRig/Source/ControlRig/Public/ControlRig.h
@@ -70,8 +70,6 @@ private:
 	bool bAccumulateTime;
 
 public:
-	UControlRig();
-
 	virtual void Serialize(FArchive& Ar) override;
 #if WITH_EDITOR
 	virtual void PostEditUndo() override;
