ControlRig: Adding filtering options to the RigHierarchy view

diff --git a/Engine/Plugins/Experimental/ControlRig/Source/ControlRigEditor/Private/ControlRigHierarchyCommands.cpp b/Engine/Plugins/Experimental/ControlRig/Source/ControlRigEditor/Private/ControlRigHierarchyCommands.cpp
--- a/Engine/Plugins/Experimental/ControlRig/Source/ControlRigEditor/Private/ControlRigHierarchyCommands.cpp
+++ b/Engine/Plugins/Experimental/ControlRig/Source/ControlRigEditor/Private/ControlRigHierarchyCommands.cpp
@@ -26,6 +26,12 @@ void FControlRigHierarchyCommands::RegisterCommands()
 	UI_COMMAND(ControlBoneTransform, "Control Bone Transform", "Sets the bone transform using a gizmo", EUserInterfaceActionType::Button, FInputChord());
 	UI_COMMAND(ControlSpaceTransform, "Control Space Transform", "Sets the space transform using a gizmo", EUserInterfaceActionType::Button, FInputChord());
 	UI_COMMAND(Unparent, "Unparent", "Unparents the selected elements from the hierarchy", EUserInterfaceActionType::Button, FInputChord(EKeys::P, EModifierKey::Shift));
+	UI_COMMAND(FilteringFlattensHierarchy, "Filtering Flattens Hierarchy", "Whether to keep the hierarchy or flatten it when searching for tree items", EUserInterfaceActionType::ToggleButton, FInputChord());
+	UI_COMMAND(HideParentsWhenFiltering, "Hide Parents When Filtering", "Whether to show parent items grayed out, or hide them entirely when filtering", EUserInterfaceActionType::ToggleButton, FInputChord());
+	UI_COMMAND(ShowImportedBones, "Show imported bones", "Whether to show or hide imported bones", EUserInterfaceActionType::ToggleButton, FInputChord());
+	UI_COMMAND(ShowBones, "Show bones", "Whether to show or hide bones", EUserInterfaceActionType::ToggleButton, FInputChord());
+	UI_COMMAND(ShowControls, "Show controls", "Whether to show or hide controls", EUserInterfaceActionType::ToggleButton, FInputChord());
+	UI_COMMAND(ShowSpaces, "Show spaces", "Whether to show or hide spaces", EUserInterfaceActionType::ToggleButton, FInputChord());
 }
 
 #undef LOCTEXT_NAMESPACE
diff --git a/Engine/Plugins/Experimental/ControlRig/Source/ControlRigEditor/Private/ControlRigHierarchyCommands.h b/Engine/Plugins/Experimental/ControlRig/Source/ControlRigEditor/Private/ControlRigHierarchyCommands.h
--- a/Engine/Plugins/Experimental/ControlRig/Source/ControlRigEditor/Private/ControlRigHierarchyCommands.h
+++ b/Engine/Plugins/Experimental/ControlRig/Source/ControlRigEditor/Private/ControlRigHierarchyCommands.h
@@ -78,6 +78,25 @@ public:
 	/* Unparents the selected elements from the hierarchy */
 	TSharedPtr<FUICommandInfo> Unparent;
 
+	/* Flatten hierarchy on filter */
+	TSharedPtr< FUICommandInfo > FilteringFlattensHierarchy;
+
+	/* Hide parents on filter */
+	TSharedPtr< FUICommandInfo > HideParentsWhenFiltering;
+
+	/* Show imported bones */
+	TSharedPtr< FUICommandInfo > ShowImportedBones;
+
+	/* Show bones */
+	TSharedPtr< FUICommandInfo > ShowBones;
+
+	/* Show controls */
+	TSharedPtr< FUICommandInfo > ShowControls;
+
+	/* Show spaces */
+	TSharedPtr< FUICommandInfo > ShowSpaces;
+
+
 	/**
 	 * Initialize commands
 	 */
diff --git a/Engine/Plugins/Experimental/ControlRig/Source/ControlRigEditor/Private/Editor/SRigHierarchy.cpp b/Engine/Plugins/Experimental/ControlRig/Source/ControlRigEditor/Private/Editor/SRigHierarchy.cpp
--- a/Engine/Plugins/Experimental/ControlRig/Source/ControlRigEditor/Private/Editor/SRigHierarchy.cpp
+++ b/Engine/Plugins/Experimental/ControlRig/Source/ControlRigEditor/Private/Editor/SRigHierarchy.cpp
@@ -333,6 +333,40 @@ void SRigHierarchy::Construct(const FArguments& InArgs, TSharedRef<FControlRigEd
 					SNew(SHorizontalBox)
 					.Visibility(this, &SRigHierarchy::IsSearchbarVisible)
 
+					+ SHorizontalBox::Slot()
+					.AutoWidth()
+					.VAlign(VAlign_Center)
+					.Padding(0.0f, 0.0f, 2.0f, 0.0f)
+					[
+						SNew(SComboButton)
+						.Visibility(EVisibility::Visible)
+						.ComboButtonStyle(FEditorStyle::Get(), "GenericFilters.ComboButtonStyle")
+						.ForegroundColor(FLinearColor::White)
+						.ContentPadding(0.0f)
+						.OnGetMenuContent(this, &SRigHierarchy::CreateFilterMenu)
+						.ButtonContent()
+						[
+							SNew(SHorizontalBox)
+							+ SHorizontalBox::Slot()
+							.AutoWidth()
+							.VAlign(VAlign_Center)
+							[
+								SNew(STextBlock)
+								.TextStyle(FEditorStyle::Get(), "GenericFilters.TextStyle")
+								.Font(FEditorStyle::Get().GetFontStyle("FontAwesome.9"))
+								.Text(FText::FromString(FString(TEXT("\xf0b0"))) /*fa-filter*/)
+							]
+							+ SHorizontalBox::Slot()
+							.AutoWidth()
+							.Padding(2, 0, 0, 0)
+							.VAlign(VAlign_Center)
+							[
+								SNew(STextBlock)
+								.TextStyle(FEditorStyle::Get(), "GenericFilters.TextStyle")
+								.Text(LOCTEXT("FilterMenuLabel", "Options"))
+							]
+						]
+					]
 					+SHorizontalBox::Slot()
 					.VAlign(VAlign_Center)
 					.Padding(3.0f, 1.0f)
@@ -380,6 +414,12 @@ void SRigHierarchy::Construct(const FArguments& InArgs, TSharedRef<FControlRigEd
 		*/
 	];
 
+	bFlattenHierarchyOnFilter = false;
+	bHideParentsOnFilter = false;
+	bShowImportedBones = true;
+	bShowBones = true;
+	bShowControls = true;
+	bShowSpaces = true;
 	bIsChangingRigHierarchy = false;
 	RefreshTreeView();
 
@@ -473,6 +513,42 @@ void SRigHierarchy::BindCommands()
 		Commands.Unparent,
 		FExecuteAction::CreateSP(this, &SRigHierarchy::HandleUnparent),
 		FCanExecuteAction::CreateSP(this, &SRigHierarchy::IsMultiSelected));
+
+	CommandList->MapAction(
+		Commands.FilteringFlattensHierarchy,
+		FExecuteAction::CreateLambda([this]() { bFlattenHierarchyOnFilter = !bFlattenHierarchyOnFilter; RefreshTreeView(); }),
+		FCanExecuteAction(),
+		FIsActionChecked::CreateLambda([this]() { return bFlattenHierarchyOnFilter; }));
+
+	CommandList->MapAction(
+		Commands.HideParentsWhenFiltering,
+		FExecuteAction::CreateLambda([this]() { bHideParentsOnFilter = !bHideParentsOnFilter; RefreshTreeView(); }),
+		FCanExecuteAction(),
+		FIsActionChecked::CreateLambda([this]() { return bHideParentsOnFilter; }));
+
+	CommandList->MapAction(
+		Commands.ShowImportedBones,
+		FExecuteAction::CreateLambda([this]() { bShowImportedBones = !bShowImportedBones; RefreshTreeView(); }),
+		FCanExecuteAction(),
+		FIsActionChecked::CreateLambda([this]() { return bShowImportedBones; }));
+	
+	CommandList->MapAction(
+		Commands.ShowBones,
+		FExecuteAction::CreateLambda([this]() { bShowBones = !bShowBones; RefreshTreeView(); }),
+		FCanExecuteAction(),
+		FIsActionChecked::CreateLambda([this]() { return bShowBones; }));
+
+	CommandList->MapAction(
+		Commands.ShowControls,
+		FExecuteAction::CreateLambda([this]() { bShowControls = !bShowControls; RefreshTreeView(); }),
+		FCanExecuteAction(),
+		FIsActionChecked::CreateLambda([this]() { return bShowControls; }));
+	
+	CommandList->MapAction(
+		Commands.ShowSpaces,
+		FExecuteAction::CreateLambda([this]() { bShowSpaces = !bShowSpaces; RefreshTreeView(); }),
+		FCanExecuteAction(),
+		FIsActionChecked::CreateLambda([this]() { return bShowSpaces; }));
 }
 
 FReply SRigHierarchy::OnKeyDown(const FGeometry& MyGeometry, const FKeyEvent& InKeyEvent)
@@ -555,17 +631,31 @@ void SRigHierarchy::RefreshTreeView()
 		FRigControlHierarchy& ControlHierarchy = Container->ControlHierarchy;
 		FRigSpaceHierarchy& SpaceHierarchy = Container->SpaceHierarchy;
 
-		for (const FRigBone& Element : BoneHierarchy)
+		if (bShowBones)
 		{
-			AddBoneElement(Element);
+			for (const FRigBone& Element : BoneHierarchy)
+			{
+				if (bShowImportedBones || Element.Type != ERigBoneType::Imported)
+				{
+					AddBoneElement(Element);
+				}
+			}
 		}
-		for (const FRigControl& Element : ControlHierarchy)
+
+		if (bShowControls)
 		{
-			AddControlElement(Element);
+			for (const FRigControl& Element : ControlHierarchy)
+			{
+				AddControlElement(Element);
+			}
 		}
-		for (const FRigSpace& Element : SpaceHierarchy)
+
+		if (bShowSpaces)
 		{
-			AddSpaceElement(Element);
+			for (const FRigSpace& Element : SpaceHierarchy)
+			{
+				AddSpaceElement(Element);
+			}
 		}
 
 		for (const auto& Pair : ElementMap)
@@ -697,7 +787,7 @@ TSharedPtr<FRigTreeElement> SRigHierarchy::FindElement(const FRigElementKey& InE
 	return TSharedPtr<FRigTreeElement>();
 }
 
-void SRigHierarchy::AddElement(FRigElementKey InKey, FRigElementKey InParentKey)
+void SRigHierarchy::AddElement(FRigElementKey InKey, FRigElementKey InParentKey, const bool bIgnoreTextFilter)
 {
 	if(ElementMap.Contains(InKey))
 	{
@@ -705,7 +795,7 @@ void SRigHierarchy::AddElement(FRigElementKey InKey, FRigElementKey InParentKey)
 	}
 
 	FString FilteredString = FilterText.ToString();
-	if (FilteredString.IsEmpty() || !InKey.IsValid())
+	if (bIgnoreTextFilter || FilteredString.IsEmpty() || !InKey.IsValid())
 	{
 		TSharedPtr<FRigTreeElement> NewItem = MakeShared<FRigTreeElement>(InKey, SharedThis(this));
 
@@ -750,7 +840,7 @@ void SRigHierarchy::AddSpacerElement()
 	AddElement(FRigElementKey(), FRigElementKey());
 }
 
-void SRigHierarchy::AddBoneElement(FRigBone InBone)
+void SRigHierarchy::AddBoneElement(FRigBone InBone, const bool bIgnoreTextFilter)
 {
 	if (ElementMap.Contains(InBone.GetElementKey()))
 	{
@@ -762,22 +852,25 @@ void SRigHierarchy::AddBoneElement(FRigBone InBone)
 	const FRigControlHierarchy& ControlHierarchy = Container->ControlHierarchy;
 	const FRigSpaceHierarchy& SpaceHierarchy = Container->SpaceHierarchy;
 
-	AddElement(InBone.GetElementKey());
+	AddElement(InBone.GetElementKey(), FRigElementKey(), bIgnoreTextFilter);
 
-	FRigElementKey ParentKey;
-	if(InBone.ParentIndex != INDEX_NONE)
+	if (ElementMap.Contains(InBone.GetElementKey()))
 	{
-		AddBoneElement(BoneHierarchy[InBone.ParentIndex]);
-		ParentKey = BoneHierarchy[InBone.ParentIndex].GetElementKey();
-	}
+		FRigElementKey ParentKey;
+		if (InBone.ParentIndex != INDEX_NONE)
+		{
+			AddBoneElement(BoneHierarchy[InBone.ParentIndex], !bFlattenHierarchyOnFilter);
+			ParentKey = BoneHierarchy[InBone.ParentIndex].GetElementKey();
+		}
 
-	if (ParentKey.IsValid())
-	{
-		ReparentElement(InBone.GetElementKey(), ParentKey);
+		if (ParentKey.IsValid())
+		{
+			ReparentElement(InBone.GetElementKey(), ParentKey);
+		}
 	}
 }
 
-void SRigHierarchy::AddControlElement(FRigControl InControl)
+void SRigHierarchy::AddControlElement(FRigControl InControl, const bool bIgnoreTextFilter)
 {
 	if (ElementMap.Contains(InControl.GetElementKey()))
 	{
@@ -788,27 +881,30 @@ void SRigHierarchy::AddControlElement(FRigControl InControl)
 	const FRigControlHierarchy& ControlHierarchy = Container->ControlHierarchy;
 	const FRigSpaceHierarchy& SpaceHierarchy = Container->SpaceHierarchy;
 
-	AddElement(InControl.GetElementKey());
+	AddElement(InControl.GetElementKey(), FRigElementKey(), bIgnoreTextFilter);
 
-	FRigElementKey ParentKey;
-	if(InControl.SpaceIndex != INDEX_NONE)
-	{
-		AddSpaceElement(SpaceHierarchy[InControl.SpaceIndex]);
-		ParentKey = SpaceHierarchy[InControl.SpaceIndex].GetElementKey();
-	}
-	else if(InControl.ParentIndex != INDEX_NONE)
+	if (ElementMap.Contains(InControl.GetElementKey()))
 	{
-		AddControlElement(ControlHierarchy[InControl.ParentIndex]);
-		ParentKey = ControlHierarchy[InControl.ParentIndex].GetElementKey();
-	}
+		FRigElementKey ParentKey;
+		if (InControl.SpaceIndex != INDEX_NONE)
+		{
+			AddSpaceElement(SpaceHierarchy[InControl.SpaceIndex], !bFlattenHierarchyOnFilter);
+			ParentKey = SpaceHierarchy[InControl.SpaceIndex].GetElementKey();
+		}
+		else if (InControl.ParentIndex != INDEX_NONE)
+		{
+			AddControlElement(ControlHierarchy[InControl.ParentIndex], !bFlattenHierarchyOnFilter);
+			ParentKey = ControlHierarchy[InControl.ParentIndex].GetElementKey();
+		}
 
-	if (ParentKey.IsValid())
-	{
-		ReparentElement(InControl.GetElementKey(), ParentKey);
+		if (ParentKey.IsValid())
+		{
+			ReparentElement(InControl.GetElementKey(), ParentKey);
+		}
 	}
 }
 
-void SRigHierarchy::AddSpaceElement(FRigSpace InSpace)
+void SRigHierarchy::AddSpaceElement(FRigSpace InSpace, const bool bIgnoreTextFilter)
 {
 	if (ElementMap.Contains(InSpace.GetElementKey()))
 	{
@@ -820,41 +916,44 @@ void SRigHierarchy::AddSpaceElement(FRigSpace InSpace)
 	const FRigControlHierarchy& ControlHierarchy = Container->ControlHierarchy;
 	const FRigSpaceHierarchy& SpaceHierarchy = Container->SpaceHierarchy;
 
-	AddElement(InSpace.GetElementKey());
+	AddElement(InSpace.GetElementKey(), FRigElementKey(), bIgnoreTextFilter);
 
-	FRigElementKey ParentKey;
-	if(InSpace.ParentIndex != INDEX_NONE)
+	if (ElementMap.Contains(InSpace.GetElementKey()))
 	{
-		switch(InSpace.SpaceType)
+		FRigElementKey ParentKey;
+		if (InSpace.ParentIndex != INDEX_NONE)
 		{
-			case ERigSpaceType::Bone:
+			switch (InSpace.SpaceType)
 			{
-				AddBoneElement(BoneHierarchy[InSpace.ParentIndex]);
-				ParentKey = BoneHierarchy[InSpace.ParentIndex].GetElementKey();
-				break;
-			}
-			case ERigSpaceType::Control:
-			{
-				AddControlElement(ControlHierarchy[InSpace.ParentIndex]);
-				ParentKey = ControlHierarchy[InSpace.ParentIndex].GetElementKey();
-				break;
-			}
-			case ERigSpaceType::Space:
-			{
-				AddSpaceElement(SpaceHierarchy[InSpace.ParentIndex]);
-				ParentKey = SpaceHierarchy[InSpace.ParentIndex].GetElementKey();
-				break;
-			}
-			default:
-			{
-				break;
+				case ERigSpaceType::Bone:
+				{
+					AddBoneElement(BoneHierarchy[InSpace.ParentIndex], !bFlattenHierarchyOnFilter);
+					ParentKey = BoneHierarchy[InSpace.ParentIndex].GetElementKey();
+					break;
+				}
+				case ERigSpaceType::Control:
+				{
+					AddControlElement(ControlHierarchy[InSpace.ParentIndex], !bFlattenHierarchyOnFilter);
+					ParentKey = ControlHierarchy[InSpace.ParentIndex].GetElementKey();
+					break;
+				}
+				case ERigSpaceType::Space:
+				{
+					AddSpaceElement(SpaceHierarchy[InSpace.ParentIndex], !bFlattenHierarchyOnFilter);
+					ParentKey = SpaceHierarchy[InSpace.ParentIndex].GetElementKey();
+					break;
+				}
+				default:
+				{
+					break;
+				}
 			}
 		}
-	}
 
-	if (ParentKey.IsValid())
-	{
-		ReparentElement(InSpace.GetElementKey(), ParentKey);
+		if (ParentKey.IsValid())
+		{
+			ReparentElement(InSpace.GetElementKey(), ParentKey);
+		}
 	}
 }
 
@@ -871,7 +970,7 @@ void SRigHierarchy::ReparentElement(FRigElementKey InKey, FRigElementKey InParen
 		return;
 	}
 
-	if (!FilterText.IsEmpty())
+	if (!FilterText.IsEmpty() && bFlattenHierarchyOnFilter)
 	{
 		return;
 	}
@@ -1035,6 +1134,37 @@ void SRigHierarchy::ClearDetailPanel() const
 	ControlRigEditor.Pin()->ClearDetailObject();
 }
 
+TSharedRef< SWidget > SRigHierarchy::CreateFilterMenu()
+{
+	const FControlRigHierarchyCommands& Actions = FControlRigHierarchyCommands::Get();
+
+	const bool CloseAfterSelection = true;
+	FMenuBuilder MenuBuilder(CloseAfterSelection, CommandList);
+
+	MenuBuilder.BeginSection("FilterOptions", LOCTEXT("OptionsMenuHeading", "Options"));
+	{
+		MenuBuilder.AddMenuEntry(Actions.FilteringFlattensHierarchy);
+		//MenuBuilder.AddMenuEntry(Actions.HideParentsWhenFiltering);
+	}
+	MenuBuilder.EndSection();
+
+	MenuBuilder.BeginSection("FilterBones", LOCTEXT("BonesMenuHeading", "Bones"));
+	{
+		MenuBuilder.AddMenuEntry(Actions.ShowImportedBones);
+		MenuBuilder.AddMenuEntry(Actions.ShowBones);
+	}
+	MenuBuilder.EndSection();
+
+	MenuBuilder.BeginSection("FilterControls", LOCTEXT("ControlsMenuHeading", "Controls"));
+	{
+		MenuBuilder.AddMenuEntry(Actions.ShowControls);
+		MenuBuilder.AddMenuEntry(Actions.ShowSpaces);
+	}
+	MenuBuilder.EndSection();
+
+	return MenuBuilder.MakeWidget();
+}
+
 TSharedPtr< SWidget > SRigHierarchy::CreateContextMenu()
 {
 	const bool CloseAfterSelection = true;
diff --git a/Engine/Plugins/Experimental/ControlRig/Source/ControlRigEditor/Private/Editor/SRigHierarchy.h b/Engine/Plugins/Experimental/ControlRig/Source/ControlRigEditor/Private/Editor/SRigHierarchy.h
--- a/Engine/Plugins/Experimental/ControlRig/Source/ControlRigEditor/Private/Editor/SRigHierarchy.h
+++ b/Engine/Plugins/Experimental/ControlRig/Source/ControlRigEditor/Private/Editor/SRigHierarchy.h
@@ -211,6 +211,7 @@ private:
 	/** Set Selection Changed */
 	void OnSelectionChanged(TSharedPtr<FRigTreeElement> Selection, ESelectInfo::Type SelectInfo);
 
+	TSharedRef< SWidget > CreateFilterMenu();
 	TSharedPtr< SWidget > CreateContextMenu();
 	void OnItemClicked(TSharedPtr<FRigTreeElement> InItem);
 	void OnItemDoubleClicked(TSharedPtr<FRigTreeElement> InItem);
@@ -233,6 +234,24 @@ private:
 	/** Our owning control rig editor */
 	TWeakPtr<FControlRigEditor> ControlRigEditor;
 
+	/** Flatten when text filtering is active */
+	bool bFlattenHierarchyOnFilter;
+
+	/** Hide parents when text filtering is active */
+	bool bHideParentsOnFilter;
+
+	/** Whether or not to show imported bones in the hierarchy */
+	bool bShowImportedBones;
+
+	/** Whether or not to show bones in the hierarchy */
+	bool bShowBones;
+
+	/** Whether or not to show controls in the hierarchy */
+	bool bShowControls;
+
+	/** Whether or not to show spaces in the hierarchy */
+	bool bShowSpaces;
+
 	/** Search box widget */
 	TSharedPtr<SSearchBox> FilterBox;
 	FText FilterText;
@@ -300,11 +319,11 @@ private:
 	void HandleRefreshEditorFromBlueprint(UControlRigBlueprint* InBlueprint);
 
 	static TSharedPtr<FRigTreeElement> FindElement(const FRigElementKey& InElementKey, TSharedPtr<FRigTreeElement> CurrentItem);
-	void AddElement(FRigElementKey InKey, FRigElementKey InParentKey = FRigElementKey());
+	void AddElement(FRigElementKey InKey, FRigElementKey InParentKey = FRigElementKey(), const bool bIgnoreTextFilter = false);
 	void AddSpacerElement();
-	void AddBoneElement(FRigBone InBone);
-	void AddControlElement(FRigControl InControl);
-	void AddSpaceElement(FRigSpace InSpace);
+	void AddBoneElement(FRigBone InBone, const bool bIgnoreTextFilter = false);
+	void AddControlElement(FRigControl InControl, const bool bIgnoreTextFilter = false);
+	void AddSpaceElement(FRigSpace InSpace, const bool bIgnoreTextFilter = false);
 	void ReparentElement(FRigElementKey InKey, FRigElementKey InParentKey);
 
 public:
