ControlRig: Adding search box to the Execution Stack

diff --git a/Engine/Plugins/Experimental/ControlRig/Source/ControlRigEditor/Private/ControlRigHierarchyCommands.cpp b/Engine/Plugins/Experimental/ControlRig/Source/ControlRigEditor/Private/ControlRigHierarchyCommands.cpp
--- a/Engine/Plugins/Experimental/ControlRig/Source/ControlRigEditor/Private/ControlRigHierarchyCommands.cpp
+++ b/Engine/Plugins/Experimental/ControlRig/Source/ControlRigEditor/Private/ControlRigHierarchyCommands.cpp
@@ -28,10 +28,10 @@ void FControlRigHierarchyCommands::RegisterCommands()
 	UI_COMMAND(Unparent, "Unparent", "Unparents the selected elements from the hierarchy", EUserInterfaceActionType::Button, FInputChord(EKeys::P, EModifierKey::Shift));
 	UI_COMMAND(FilteringFlattensHierarchy, "Filtering Flattens Hierarchy", "Whether to keep the hierarchy or flatten it when searching for tree items", EUserInterfaceActionType::ToggleButton, FInputChord());
 	UI_COMMAND(HideParentsWhenFiltering, "Hide Parents When Filtering", "Whether to show parent items grayed out, or hide them entirely when filtering", EUserInterfaceActionType::ToggleButton, FInputChord());
-	UI_COMMAND(ShowImportedBones, "Show imported bones", "Whether to show or hide imported bones", EUserInterfaceActionType::ToggleButton, FInputChord());
-	UI_COMMAND(ShowBones, "Show bones", "Whether to show or hide bones", EUserInterfaceActionType::ToggleButton, FInputChord());
-	UI_COMMAND(ShowControls, "Show controls", "Whether to show or hide controls", EUserInterfaceActionType::ToggleButton, FInputChord());
-	UI_COMMAND(ShowSpaces, "Show spaces", "Whether to show or hide spaces", EUserInterfaceActionType::ToggleButton, FInputChord());
+	UI_COMMAND(ShowImportedBones, "Show Imported Bones", "Whether to show or hide imported bones", EUserInterfaceActionType::ToggleButton, FInputChord());
+	UI_COMMAND(ShowBones, "Show Bones", "Whether to show or hide bones", EUserInterfaceActionType::ToggleButton, FInputChord());
+	UI_COMMAND(ShowControls, "Show Controls", "Whether to show or hide controls", EUserInterfaceActionType::ToggleButton, FInputChord());
+	UI_COMMAND(ShowSpaces, "Show Spaces", "Whether to show or hide spaces", EUserInterfaceActionType::ToggleButton, FInputChord());
 }
 
 #undef LOCTEXT_NAMESPACE
diff --git a/Engine/Plugins/Experimental/ControlRig/Source/ControlRigEditor/Private/Editor/SControlRigStackView.cpp b/Engine/Plugins/Experimental/ControlRig/Source/ControlRigEditor/Private/Editor/SControlRigStackView.cpp
--- a/Engine/Plugins/Experimental/ControlRig/Source/ControlRigEditor/Private/Editor/SControlRigStackView.cpp
+++ b/Engine/Plugins/Experimental/ControlRig/Source/ControlRigEditor/Private/Editor/SControlRigStackView.cpp
@@ -4,6 +4,7 @@
 #include "Widgets/Layout/SBorder.h"
 #include "Widgets/SBoxPanel.h"
 #include "Widgets/Layout/SScrollBar.h"
+#include "Widgets/Input/SSearchBox.h"
 #include "ControlRigEditor.h"
 #include "ControlRigEditorStyle.h"
 #include "ControlRigStackCommands.h"
@@ -154,26 +155,55 @@ void SControlRigStackView::Construct( const FArguments& InArgs, TSharedRef<FCont
 	BindCommands();
 
 	ChildSlot
+	[
+		SNew(SVerticalBox)
+		+SVerticalBox::Slot()
+		.AutoHeight()
+		.VAlign(VAlign_Top)
+		.Padding(0.0f)
 		[
-			SNew(SVerticalBox)
-			+ SVerticalBox::Slot()
-			.Padding(0.0f, 0.0f)
+			SNew(SBorder)
+			.Padding(0.0f)
+			.BorderImage(FEditorStyle::GetBrush("DetailsView.CategoryTop"))
+			.BorderBackgroundColor(FLinearColor(0.6f, 0.6f, 0.6f, 1.0f))
 			[
-				SNew(SBorder)
-				.Padding(2.0f)
-				.BorderImage(FEditorStyle::GetBrush("SCSEditor.TreePanel"))
+				SNew(SVerticalBox)
+				+SVerticalBox::Slot()
+				.AutoHeight()
+				.VAlign(VAlign_Top)
 				[
-					SAssignNew(TreeView, STreeView<TSharedPtr<FRigStackEntry>>)
-					.TreeItemsSource(&Operators)
-					.SelectionMode(ESelectionMode::Multi)
-					.OnGenerateRow(this, &SControlRigStackView::MakeTableRowWidget)
-					.OnGetChildren(this, &SControlRigStackView::HandleGetChildrenForTree)
-					.OnSelectionChanged(this, &SControlRigStackView::OnSelectionChanged)
-					.OnContextMenuOpening(this, &SControlRigStackView::CreateContextMenu)
-					.ItemHeight(28)
+					SNew(SHorizontalBox)
+					//.Visibility(this, &SRigHierarchy::IsSearchbarVisible)
+					+SHorizontalBox::Slot()
+					.AutoWidth()
+					+SHorizontalBox::Slot()
+					.VAlign(VAlign_Center)
+					.Padding(3.0f, 1.0f)
+					[
+						SAssignNew(FilterBox, SSearchBox)
+						.OnTextChanged(this, &SControlRigStackView::OnFilterTextChanged)
+					]
 				]
 			]
-		];
+		]
+		+SVerticalBox::Slot()
+		.Padding(0.0f, 0.0f)
+		[
+			SNew(SBorder)
+			.Padding(2.0f)
+			.BorderImage(FEditorStyle::GetBrush("SCSEditor.TreePanel"))
+			[
+				SAssignNew(TreeView, STreeView<TSharedPtr<FRigStackEntry>>)
+				.TreeItemsSource(&Operators)
+				.SelectionMode(ESelectionMode::Multi)
+				.OnGenerateRow(this, &SControlRigStackView::MakeTableRowWidget)
+				.OnGetChildren(this, &SControlRigStackView::HandleGetChildrenForTree)
+				.OnSelectionChanged(this, &SControlRigStackView::OnSelectionChanged)
+				.OnContextMenuOpening(this, &SControlRigStackView::CreateContextMenu)
+				.ItemHeight(28)
+			]
+		]
+	];
 
 	RefreshTreeView(nullptr);
 
@@ -286,9 +316,12 @@ void SControlRigStackView::RefreshTreeView(URigVM* InVM)
 			{
 				Label = FString::Printf(TEXT("%s(...)"), *Left);
 			}
-			
-			TSharedPtr<FRigStackEntry> NewEntry = MakeShared<FRigStackEntry>(Operators.Num(), ERigStackEntry::Operator, InstructionIndex, Instruction.OpCode, Label);
-			Operators.Add(NewEntry);
+
+			if (FilterText.IsEmpty() || Label.Contains(FilterText.ToString()))
+			{ 
+				TSharedPtr<FRigStackEntry> NewEntry = MakeShared<FRigStackEntry>(InstructionIndex, ERigStackEntry::Operator, InstructionIndex, Instruction.OpCode, Label);
+				Operators.Add(NewEntry);
+			}
 		}
 
 		// fill the children from the log
@@ -375,6 +408,12 @@ void SControlRigStackView::OnVMCompiled(UBlueprint* InCompiledBlueprint, URigVM*
 	}
 }
 
+void SControlRigStackView::OnFilterTextChanged(const FText& SearchText)
+{
+	FilterText = SearchText;
+	RefreshTreeView(ControlRigEditor.Pin()->ControlRig->GetVM());
+}
+
 void SControlRigStackView::HandleModifiedEvent(ERigVMGraphNotifType InNotifType, URigVMGraph* InGraph, UObject* InSubject)
 {
 	if (bSuspendModelNotifications)
diff --git a/Engine/Plugins/Experimental/ControlRig/Source/ControlRigEditor/Private/Editor/SControlRigStackView.h b/Engine/Plugins/Experimental/ControlRig/Source/ControlRigEditor/Private/Editor/SControlRigStackView.h
--- a/Engine/Plugins/Experimental/ControlRig/Source/ControlRigEditor/Private/Editor/SControlRigStackView.h
+++ b/Engine/Plugins/Experimental/ControlRig/Source/ControlRigEditor/Private/Editor/SControlRigStackView.h
@@ -91,6 +91,12 @@ private:
 
 	void OnVMCompiled(UBlueprint* InCompiledBlueprint, URigVM* InCompiledVM);
 
+	/** Search box widget */
+	TSharedPtr<SSearchBox> FilterBox;
+	FText FilterText;
+	void OnFilterTextChanged(const FText& SearchText);
+
+
 	bool bSuspendModelNotifications;
 	bool bSuspendControllerSelection;
 	void HandleModifiedEvent(ERigVMGraphNotifType InNotifType, URigVMGraph* InGraph, UObject* InSubject);
