Control Rig: Fix crash when creating empty rig (due to bad merge)

diff --git a/Engine/Plugins/Experimental/ControlRig/Source/ControlRigEditor/Private/Editor/ControlRigEditor.cpp b/Engine/Plugins/Experimental/ControlRig/Source/ControlRigEditor/Private/Editor/ControlRigEditor.cpp
--- a/Engine/Plugins/Experimental/ControlRig/Source/ControlRigEditor/Private/Editor/ControlRigEditor.cpp
+++ b/Engine/Plugins/Experimental/ControlRig/Source/ControlRigEditor/Private/Editor/ControlRigEditor.cpp
@@ -171,7 +171,7 @@ void FControlRigEditor::InitControlRigEditor(const EToolkitMode::Type Mode, cons
 
 	FPersonaToolkitArgs PersonaToolkitArgs;
 	PersonaToolkitArgs.OnPreviewSceneCreated = FOnPreviewSceneCreated::FDelegate::CreateSP(this, &FControlRigEditor::HandlePreviewSceneCreated);
-	PersonaToolkit = PersonaModule.CreatePersonaToolkit(InControlRigBlueprint, PersonaToolkitArgs, InControlRigBlueprint->GetPreviewMesh()->Skeleton);
+	PersonaToolkit = PersonaModule.CreatePersonaToolkit(InControlRigBlueprint, PersonaToolkitArgs);
 
 	// set delegate prior to setting mesh
 	// otherwise, you don't get delegate
