Control Rig: Export Rig Unit Symbols, required for Control Rig Nativization Intern Research Project. As the project progresses, I will try to export more rig units and eventually the project will support most of the rig units.

diff --git a/Engine/Plugins/Experimental/ControlRig/Source/ControlRig/Private/Units/Math/RigUnit_MathFloat.h b/Engine/Plugins/Experimental/ControlRig/Source/ControlRig/Private/Units/Math/RigUnit_MathFloat.h
--- a/Engine/Plugins/Experimental/ControlRig/Source/ControlRig/Private/Units/Math/RigUnit_MathFloat.h
+++ b/Engine/Plugins/Experimental/ControlRig/Source/ControlRig/Private/Units/Math/RigUnit_MathFloat.h
@@ -745,7 +745,7 @@ struct FRigUnit_MathFloatRad : public FRigUnit_MathFloatUnaryOp
  * Returns the sinus value of the given value (in radians)
  */
 USTRUCT(meta=(DisplayName="Sin", PrototypeName="Sin"))
-struct FRigUnit_MathFloatSin : public FRigUnit_MathFloatUnaryOp
+struct CONTROLRIG_API FRigUnit_MathFloatSin : public FRigUnit_MathFloatUnaryOp
 {
 	GENERATED_BODY()
 
diff --git a/Engine/Plugins/Experimental/ControlRig/Source/ControlRig/Private/Units/Simulation/RigUnit_Timeline.h b/Engine/Plugins/Experimental/ControlRig/Source/ControlRig/Private/Units/Simulation/RigUnit_Timeline.h
--- a/Engine/Plugins/Experimental/ControlRig/Source/ControlRig/Private/Units/Simulation/RigUnit_Timeline.h
+++ b/Engine/Plugins/Experimental/ControlRig/Source/ControlRig/Private/Units/Simulation/RigUnit_Timeline.h
@@ -9,7 +9,7 @@
  * Simulates a time value - can act as a timeline playing back
  */
 USTRUCT(meta=(DisplayName="Accumulated Time", Keywords="Playback,Pause,Timeline"))
-struct FRigUnit_Timeline : public FRigUnit_SimBase
+struct CONTROLRIG_API FRigUnit_Timeline : public FRigUnit_SimBase
 {
 	GENERATED_BODY()
 	
diff --git a/Engine/Plugins/Experimental/ControlRig/Source/ControlRig/Public/Units/Hierarchy/RigUnit_SetTransform.h b/Engine/Plugins/Experimental/ControlRig/Source/ControlRig/Public/Units/Hierarchy/RigUnit_SetTransform.h
--- a/Engine/Plugins/Experimental/ControlRig/Source/ControlRig/Public/Units/Hierarchy/RigUnit_SetTransform.h
+++ b/Engine/Plugins/Experimental/ControlRig/Source/ControlRig/Public/Units/Hierarchy/RigUnit_SetTransform.h
@@ -9,7 +9,7 @@
  * SetTransform is used to set a single transform on hierarchy.
  */
 USTRUCT(meta=(DisplayName="Set Transform", Category="Hierarchy", DocumentationPolicy = "Strict", Keywords="SetBoneTransform,SetControlTransform,SetInitialTransform,SetSpaceTransform", Varying))
-struct FRigUnit_SetTransform : public FRigUnitMutable
+struct CONTROLRIG_API FRigUnit_SetTransform : public FRigUnitMutable
 {
 	GENERATED_BODY()
 
