Control Rig: Control detail values do not get saved on first save

diff --git a/Engine/Plugins/Experimental/ControlRig/Source/ControlRigEditor/Private/Editor/ControlRigEditor.cpp b/Engine/Plugins/Experimental/ControlRig/Source/ControlRigEditor/Private/Editor/ControlRigEditor.cpp
--- a/Engine/Plugins/Experimental/ControlRig/Source/ControlRigEditor/Private/Editor/ControlRigEditor.cpp
+++ b/Engine/Plugins/Experimental/ControlRig/Source/ControlRigEditor/Private/Editor/ControlRigEditor.cpp
@@ -1319,12 +1319,22 @@ void FControlRigEditor::SaveAsset_Execute()
 {
 	LastDebuggedRig = GetCustomDebugObjectLabel(GetBlueprintObj()->GetObjectBeingDebugged());
 	FBlueprintEditor::SaveAsset_Execute();
+
+	// Save the new state of the hierarchy in the default object, so that it has the correct values on load
+	UControlRigBlueprint* RigBlueprint = Cast<UControlRigBlueprint>(GetBlueprintObj());
+	UControlRig* CDO = ControlRig->GetClass()->GetDefaultObject<UControlRig>();
+	CDO->Hierarchy = RigBlueprint->HierarchyContainer;
 }
 
 void FControlRigEditor::SaveAssetAs_Execute()
 {
 	LastDebuggedRig = GetCustomDebugObjectLabel(GetBlueprintObj()->GetObjectBeingDebugged());
 	FBlueprintEditor::SaveAssetAs_Execute();
+
+	// Save the new state of the hierarchy in the default object, so that it has the correct values on load
+	UControlRigBlueprint* RigBlueprint = Cast<UControlRigBlueprint>(GetBlueprintObj());
+	UControlRig* CDO = ControlRig->GetClass()->GetDefaultObject<UControlRig>();
+	CDO->Hierarchy = RigBlueprint->HierarchyContainer;
 }
 
 FName FControlRigEditor::GetToolkitFName() const
