Control Rig: Controls retain pose after switching from Backwards to Forward

diff --git a/Engine/Plugins/Experimental/ControlRig/Source/ControlRigEditor/Private/Editor/ControlRigEditor.cpp b/Engine/Plugins/Experimental/ControlRig/Source/ControlRigEditor/Private/Editor/ControlRigEditor.cpp
--- a/Engine/Plugins/Experimental/ControlRig/Source/ControlRigEditor/Private/Editor/ControlRigEditor.cpp
+++ b/Engine/Plugins/Experimental/ControlRig/Source/ControlRigEditor/Private/Editor/ControlRigEditor.cpp
@@ -600,6 +600,13 @@ void FControlRigEditor::SetEventQueue(EControlRigEditorEventQueue InEventQueue)
 		{
 			ToggleSetupMode();
 		}
+
+		// Reset transforms only for setup and forward solve to not inturrupt any animation that might be playing
+		if (InEventQueue == EControlRigEditorEventQueue::Setup ||
+			InEventQueue == EControlRigEditorEventQueue::Update)
+		{
+			ControlRig->Hierarchy.ResetTransforms();
+		}
 	}
 }
 
