Control Rig: Add PropertyAccessEditor dependency

diff --git a/Engine/Plugins/Experimental/ControlRig/ControlRig.uplugin b/Engine/Plugins/Experimental/ControlRig/ControlRig.uplugin
--- a/Engine/Plugins/Experimental/ControlRig/ControlRig.uplugin
+++ b/Engine/Plugins/Experimental/ControlRig/ControlRig.uplugin
@@ -19,6 +19,10 @@
 		{
 			"Name": "LevelSequenceEditor",
 			"Enabled": true
+		},
+		{
+			"Name": "PropertyAccessEditor",
+			"Enabled": true
 		}
 	],
 	"Modules" :
