Control Rig: Ensure to always show MyBlueprint tab

diff --git a/Engine/Plugins/Experimental/ControlRig/Source/ControlRigEditor/Private/Editor/ControlRigEditor.cpp b/Engine/Plugins/Experimental/ControlRig/Source/ControlRigEditor/Private/Editor/ControlRigEditor.cpp
--- a/Engine/Plugins/Experimental/ControlRig/Source/ControlRigEditor/Private/Editor/ControlRigEditor.cpp
+++ b/Engine/Plugins/Experimental/ControlRig/Source/ControlRigEditor/Private/Editor/ControlRigEditor.cpp
@@ -75,6 +75,7 @@
 #include "BlueprintCompilationManager.h"
 #include "AssetEditorModeManager.h"
 #include "IPersonaEditorModeManager.h"
+#include "BlueprintEditorTabs.h"
 
 #define LOCTEXT_NAMESPACE "ControlRigEditor"
 
@@ -240,6 +241,9 @@ void FControlRigEditor::InitControlRigEditor(const EToolkitMode::Type Mode, cons
 	// Activate the initial mode (which will populate with a real layout)
 	SetCurrentMode(FControlRigEditorModes::ControlRigEditorMode);
 
+	// Always show the myblueprint tab
+	GetTabManager()->TryInvokeTab(FTabId(FBlueprintEditorTabs::MyBlueprintID));
+
 	// Activate our edit mode
 	GetEditorModeManager().SetDefaultMode(FControlRigEditorEditMode::ModeName);
 	GetEditorModeManager().ActivateMode(FControlRigEditorEditMode::ModeName);
