Editor UI: Added SetNodePosition(...) to GraphSchema, called when Straighthen pin connections.
Control Rig: Override GraphSchema::SetNodePosition(...) to use RigVMController to perform the set action, UEdGraph is updated through URigVMController::Notify(...)

diff --git a/Engine/Plugins/Experimental/ControlRig/Source/ControlRigDeveloper/Private/Graph/ControlRigGraphSchema.cpp b/Engine/Plugins/Experimental/ControlRig/Source/ControlRigDeveloper/Private/Graph/ControlRigGraphSchema.cpp
--- a/Engine/Plugins/Experimental/ControlRig/Source/ControlRigDeveloper/Private/Graph/ControlRigGraphSchema.cpp
+++ b/Engine/Plugins/Experimental/ControlRig/Source/ControlRigDeveloper/Private/Graph/ControlRigGraphSchema.cpp
@@ -267,6 +267,20 @@ bool UControlRigGraphSchema::IsStructEditable(UStruct* InStruct) const
 	return false;
 }
 
+void UControlRigGraphSchema::SetNodePosition(UEdGraphNode* Node, const FVector2D& Position) const
+{
+	if (UControlRigGraphNode* RigNode = Cast<UControlRigGraphNode>(Node))
+	{
+		UBlueprint* Blueprint = FBlueprintEditorUtils::FindBlueprintForNodeChecked(RigNode);
+		UControlRigBlueprint* RigBlueprint = Cast<UControlRigBlueprint>(Blueprint);
+		if (RigBlueprint != nullptr)
+		{
+			RigBlueprint->Controller->SetNodePosition(RigNode->GetModelNode(), Position, true, false); 
+			return;
+		}
+	}
+}
+
 UControlRigGraphNode* UControlRigGraphSchema::CreateGraphNode(UControlRigGraph* InGraph, const FName& InPropertyName) const
 {
 	const bool bSelectNewNode = true;
diff --git a/Engine/Plugins/Experimental/ControlRig/Source/ControlRigDeveloper/Public/Graph/ControlRigGraphSchema.h b/Engine/Plugins/Experimental/ControlRig/Source/ControlRigDeveloper/Public/Graph/ControlRigGraphSchema.h
--- a/Engine/Plugins/Experimental/ControlRig/Source/ControlRigDeveloper/Public/Graph/ControlRigGraphSchema.h
+++ b/Engine/Plugins/Experimental/ControlRig/Source/ControlRigDeveloper/Public/Graph/ControlRigGraphSchema.h
@@ -76,6 +76,7 @@ public:
 	virtual bool RequestVariableDropOnPanel(UEdGraph* InGraph, FProperty* InVariableToDrop, const FVector2D& InDropPosition, const FVector2D& InScreenPosition) override;
 	virtual bool RequestVariableDropOnPin(UEdGraph* InGraph, FProperty* InVariableToDrop, UEdGraphPin* InPin, const FVector2D& InDropPosition, const FVector2D& InScreenPosition) override;
 	virtual bool IsStructEditable(UStruct* InStruct) const;
+	virtual void SetNodePosition(UEdGraphNode* Node, const FVector2D& Position) const override;
 
 	/** Create a graph node for a rig */
 	UControlRigGraphNode* CreateGraphNode(UControlRigGraph* InGraph, const FName& InPropertyName) const;
diff --git a/Engine/Source/Editor/GraphEditor/Private/SGraphPanel.cpp b/Engine/Source/Editor/GraphEditor/Private/SGraphPanel.cpp
--- a/Engine/Source/Editor/GraphEditor/Private/SGraphPanel.cpp
+++ b/Engine/Source/Editor/GraphEditor/Private/SGraphPanel.cpp
@@ -1255,9 +1255,15 @@ private:
 				AlignmentDelta += (Node->NodePosY + Pins.SrcPin->GetNodeOffset().Y) - (NodeToPins.Key->NodePosY + Pins.DstPin->GetNodeOffset().Y);
 			}
 
-			NodeToPins.Key->Modify();
-			NodeToPins.Key->NodePosY += AlignmentDelta / NodeToPins.Value.Num();
+			UEdGraph* GraphObj = NodeToPins.Key->GetGraph();
 
+			check(GraphObj);
+
+			const UEdGraphSchema* Schema = GraphObj->GetSchema();
+
+			float NewNodePosY = NodeToPins.Key->NodePosY + (AlignmentDelta / NodeToPins.Value.Num());
+			Schema->SetNodePosition(NodeToPins.Key, FVector2D(NodeToPins.Key->NodePosX, NewNodePosY));
+	
 			VisitedNodes.Add(Node);
 			VisitedNodes.Add(NodeToPins.Key);
 			
diff --git a/Engine/Source/Runtime/Engine/Classes/EdGraph/EdGraphSchema.h b/Engine/Source/Runtime/Engine/Classes/EdGraph/EdGraphSchema.h
--- a/Engine/Source/Runtime/Engine/Classes/EdGraph/EdGraphSchema.h
+++ b/Engine/Source/Runtime/Engine/Classes/EdGraph/EdGraphSchema.h
@@ -911,6 +911,13 @@ class ENGINE_API UEdGraphSchema : public UObject
 	 */
 	virtual UEdGraphNode* CreateSubstituteNode(UEdGraphNode* Node, const UEdGraph* Graph, FObjectInstancingGraph* InstanceGraph, TSet<FName>& InOutExtraNames) const { return nullptr; }
 
+	/**
+	 * Sets a node's position.
+	 *
+	 * @param	Node			The node to set
+	 * @param	Position		The target position
+	 */
+	virtual void SetNodePosition(UEdGraphNode* Node, const FVector2D& Position) const;
 
 	/**
 	 * Returns the currently selected graph node count
diff --git a/Engine/Source/Runtime/Engine/Private/EdGraph/EdGraphSchema.cpp b/Engine/Source/Runtime/Engine/Private/EdGraph/EdGraphSchema.cpp
--- a/Engine/Source/Runtime/Engine/Private/EdGraph/EdGraphSchema.cpp
+++ b/Engine/Source/Runtime/Engine/Private/EdGraph/EdGraphSchema.cpp
@@ -785,6 +785,14 @@ void UEdGraphSchema::ReconstructNode(UEdGraphNode& TargetNode, bool bIsBatchRequ
 #endif	//#if WITH_EDITOR
 }
 
+void UEdGraphSchema::SetNodePosition(UEdGraphNode* Node, const FVector2D& Position) const
+{
+	check(Node);
+	Node->Modify();
+	Node->NodePosX = Position.X;
+	Node->NodePosY = Position.Y;
+}
+
 void UEdGraphSchema::GetGraphDisplayInformation(const UEdGraph& Graph, /*out*/ FGraphDisplayInfo& DisplayInfo) const
 {
 	DisplayInfo.PlainName = FText::FromString( Graph.GetName() );
