Control Rig: Change default Control Type to Euler transform

diff --git a/Engine/Plugins/Experimental/ControlRig/Source/ControlRigEditor/Private/Editor/SRigHierarchy.cpp b/Engine/Plugins/Experimental/ControlRig/Source/ControlRigEditor/Private/Editor/SRigHierarchy.cpp
--- a/Engine/Plugins/Experimental/ControlRig/Source/ControlRigEditor/Private/Editor/SRigHierarchy.cpp
+++ b/Engine/Plugins/Experimental/ControlRig/Source/ControlRigEditor/Private/Editor/SRigHierarchy.cpp
@@ -1628,15 +1628,15 @@ void SRigHierarchy::HandleNewItem(ERigElementType InElementType)
 
 					if(!ParentKey)
 					{
-						Hierarchy->ControlHierarchy.Add(NewElementName, ERigControlType::Transform);
+						Hierarchy->ControlHierarchy.Add(NewElementName, ERigControlType::EulerTransform);
 					}
 					else if (ParentKey.Type == ERigElementType::Space)
 					{
-						Hierarchy->ControlHierarchy.Add(NewElementName, ERigControlType::Transform, NAME_None, ParentKey.Name);
+						Hierarchy->ControlHierarchy.Add(NewElementName, ERigControlType::EulerTransform, NAME_None, ParentKey.Name);
 					}
 					else if (ParentKey.Type == ERigElementType::Control)
 					{
-						Hierarchy->ControlHierarchy.Add(NewElementName, ERigControlType::Transform, ParentKey.Name);
+						Hierarchy->ControlHierarchy.Add(NewElementName, ERigControlType::EulerTransform, ParentKey.Name);
 					}
 					break;
 				}
