Control Rig: Control Rig Palette - adding one node to Favorites adds all of nodes to it.

diff --git a/Engine/Plugins/Experimental/ControlRig/Source/ControlRigEditor/Private/Graph/NodeSpawners/ControlRigBranchNodeSpawner.cpp b/Engine/Plugins/Experimental/ControlRig/Source/ControlRigEditor/Private/Graph/NodeSpawners/ControlRigBranchNodeSpawner.cpp
--- a/Engine/Plugins/Experimental/ControlRig/Source/ControlRigEditor/Private/Graph/NodeSpawners/ControlRigBranchNodeSpawner.cpp
+++ b/Engine/Plugins/Experimental/ControlRig/Source/ControlRigEditor/Private/Graph/NodeSpawners/ControlRigBranchNodeSpawner.cpp
@@ -34,7 +34,7 @@ UControlRigBranchNodeSpawner* UControlRigBranchNodeSpawner::CreateGeneric(const
 
 FBlueprintNodeSignature UControlRigBranchNodeSpawner::GetSpawnerSignature() const
 {
-	return FBlueprintNodeSignature(NodeClass);
+	return FBlueprintNodeSignature(FString("RigUnit=" + DefaultMenuSignature.MenuName.ToString()));
 }
 
 FBlueprintActionUiSpec UControlRigBranchNodeSpawner::GetUiSpec(FBlueprintActionContext const& Context, FBindingSet const& Bindings) const
diff --git a/Engine/Plugins/Experimental/ControlRig/Source/ControlRigEditor/Private/Graph/NodeSpawners/ControlRigEnumNodeSpawner.cpp b/Engine/Plugins/Experimental/ControlRig/Source/ControlRigEditor/Private/Graph/NodeSpawners/ControlRigEnumNodeSpawner.cpp
--- a/Engine/Plugins/Experimental/ControlRig/Source/ControlRigEditor/Private/Graph/NodeSpawners/ControlRigEnumNodeSpawner.cpp
+++ b/Engine/Plugins/Experimental/ControlRig/Source/ControlRigEditor/Private/Graph/NodeSpawners/ControlRigEnumNodeSpawner.cpp
@@ -37,7 +37,7 @@ UControlRigEnumNodeSpawner* UControlRigEnumNodeSpawner::CreateForEnum(UEnum* InE
 
 FBlueprintNodeSignature UControlRigEnumNodeSpawner::GetSpawnerSignature() const
 {
-	return FBlueprintNodeSignature(NodeClass);
+	return FBlueprintNodeSignature(FString("RigUnit=" + Enum->GetFName().ToString()));
 }
 
 FBlueprintActionUiSpec UControlRigEnumNodeSpawner::GetUiSpec(FBlueprintActionContext const& Context, FBindingSet const& Bindings) const
diff --git a/Engine/Plugins/Experimental/ControlRig/Source/ControlRigEditor/Private/Graph/NodeSpawners/ControlRigPrototypeNodeSpawner.cpp b/Engine/Plugins/Experimental/ControlRig/Source/ControlRigEditor/Private/Graph/NodeSpawners/ControlRigPrototypeNodeSpawner.cpp
--- a/Engine/Plugins/Experimental/ControlRig/Source/ControlRigEditor/Private/Graph/NodeSpawners/ControlRigPrototypeNodeSpawner.cpp
+++ b/Engine/Plugins/Experimental/ControlRig/Source/ControlRigEditor/Private/Graph/NodeSpawners/ControlRigPrototypeNodeSpawner.cpp
@@ -67,7 +67,7 @@ void UControlRigPrototypeNodeSpawner::Prime()
 
 FBlueprintNodeSignature UControlRigPrototypeNodeSpawner::GetSpawnerSignature() const
 {
-	return FBlueprintNodeSignature(NodeClass);
+	return FBlueprintNodeSignature(FString("RigUnit=" + PrototypeNotation.ToString()));
 }
 
 FBlueprintActionUiSpec UControlRigPrototypeNodeSpawner::GetUiSpec(FBlueprintActionContext const& Context, FBindingSet const& Bindings) const
diff --git a/Engine/Plugins/Experimental/ControlRig/Source/ControlRigEditor/Private/Graph/NodeSpawners/ControlRigUnitNodeSpawner.cpp b/Engine/Plugins/Experimental/ControlRig/Source/ControlRigEditor/Private/Graph/NodeSpawners/ControlRigUnitNodeSpawner.cpp
--- a/Engine/Plugins/Experimental/ControlRig/Source/ControlRigEditor/Private/Graph/NodeSpawners/ControlRigUnitNodeSpawner.cpp
+++ b/Engine/Plugins/Experimental/ControlRig/Source/ControlRigEditor/Private/Graph/NodeSpawners/ControlRigUnitNodeSpawner.cpp
@@ -76,7 +76,7 @@ void UControlRigUnitNodeSpawner::Prime()
 
 FBlueprintNodeSignature UControlRigUnitNodeSpawner::GetSpawnerSignature() const
 {
-	return FBlueprintNodeSignature(NodeClass);
+	return FBlueprintNodeSignature(FString("RigUnit=" + StructTemplate->GetFName().ToString()));
 }
 
 FBlueprintActionUiSpec UControlRigUnitNodeSpawner::GetUiSpec(FBlueprintActionContext const& Context, FBindingSet const& Bindings) const
diff --git a/Engine/Plugins/Experimental/ControlRig/Source/ControlRigEditor/Private/Graph/NodeSpawners/ControlRigVariableNodeSpawner.cpp b/Engine/Plugins/Experimental/ControlRig/Source/ControlRigEditor/Private/Graph/NodeSpawners/ControlRigVariableNodeSpawner.cpp
--- a/Engine/Plugins/Experimental/ControlRig/Source/ControlRigEditor/Private/Graph/NodeSpawners/ControlRigVariableNodeSpawner.cpp
+++ b/Engine/Plugins/Experimental/ControlRig/Source/ControlRigEditor/Private/Graph/NodeSpawners/ControlRigVariableNodeSpawner.cpp
@@ -63,7 +63,7 @@ bool UControlRigVariableNodeSpawner::IsTemplateNodeFilteredOut(FBlueprintActionF
 
 FBlueprintNodeSignature UControlRigVariableNodeSpawner::GetSpawnerSignature() const
 {
-	return FBlueprintNodeSignature(NodeClass);
+	return FBlueprintNodeSignature(FString("RigUnit=" + ExternalVariable.Name.ToString()));
 }
 
 FBlueprintActionUiSpec UControlRigVariableNodeSpawner::GetUiSpec(FBlueprintActionContext const& Context, FBindingSet const& Bindings) const
