Control Rig: Paste to current

diff --git a/Engine/Plugins/Experimental/ControlRig/Source/ControlRig/Private/Rigs/RigHierarchyContainer.cpp b/Engine/Plugins/Experimental/ControlRig/Source/ControlRig/Private/Rigs/RigHierarchyContainer.cpp
--- a/Engine/Plugins/Experimental/ControlRig/Source/ControlRig/Private/Rigs/RigHierarchyContainer.cpp
+++ b/Engine/Plugins/Experimental/ControlRig/Source/ControlRig/Private/Rigs/RigHierarchyContainer.cpp
@@ -880,7 +880,7 @@ public:
 	}
 };
 
-TArray<FRigElementKey> FRigHierarchyContainer::ImportFromText(const FString& InContent, ERigHierarchyImportMode InImportMode, bool bSelectNewElements)
+TArray<FRigElementKey> FRigHierarchyContainer::ImportFromText(const FString& InContent, ERigHierarchyImportMode InImportMode, bool bSelectNewElements, const TArray<FRigElementKey>& InSelection)
 {
 	TArray<FRigElementKey> PastedKeys;
 
@@ -964,7 +964,15 @@ TArray<FRigElementKey> FRigHierarchyContainer::ImportFromText(const FString& InC
 		Elements.Add(NewElement);
 	}
 
-	TArray<FRigElementKey> Selection = CurrentSelection();
+	TArray<FRigElementKey> Selection = InSelection;
+	if (InSelection.IsEmpty())
+	{
+		Selection = CurrentSelection();
+	}
+	else
+	{
+		bSelectNewElements = false;
+	}
 
 	switch (InImportMode)
 	{
@@ -1058,8 +1066,10 @@ TArray<FRigElementKey> FRigHierarchyContainer::ImportFromText(const FString& InC
 			}
 			break;
 		}
-		case ERigHierarchyImportMode::ReplaceLocalTransform:
-		case ERigHierarchyImportMode::ReplaceGlobalTransform:
+		case ERigHierarchyImportMode::ReplaceInitialLocalTransform:
+		case ERigHierarchyImportMode::ReplaceInitialGlobalTransform:
+		case ERigHierarchyImportMode::ReplaceCurrentLocalTransform:
+		case ERigHierarchyImportMode::ReplaceCurrentGlobalTransform:
 		{
 			if (Selection.Num() == 0)
 			{
@@ -1076,7 +1086,7 @@ TArray<FRigElementKey> FRigHierarchyContainer::ImportFromText(const FString& InC
 
 			for (int32 Index = 0; Index < Data.Types.Num(); Index++)
 			{
-				if(InImportMode == ERigHierarchyImportMode::ReplaceLocalTransform)
+				if(InImportMode == ERigHierarchyImportMode::ReplaceInitialLocalTransform)
 				{
 					Data.LocalTransforms[Index].NormalizeRotation();
 					switch (Selection[Index].Type)
@@ -1098,7 +1108,7 @@ TArray<FRigElementKey> FRigHierarchyContainer::ImportFromText(const FString& InC
 						}
 					}
 				}
-				else
+				else if (InImportMode == ERigHierarchyImportMode::ReplaceInitialGlobalTransform)
 				{
 					Data.GlobalTransforms[Index].NormalizeRotation();
 					switch (Selection[Index].Type)
@@ -1121,6 +1131,16 @@ TArray<FRigElementKey> FRigHierarchyContainer::ImportFromText(const FString& InC
 						}
 					}
 				}
+				else if (InImportMode == ERigHierarchyImportMode::ReplaceCurrentLocalTransform)
+				{
+					Data.LocalTransforms[Index].NormalizeRotation();
+					SetLocalTransform(Selection[Index], Data.LocalTransforms[Index]);
+				}
+				else if (InImportMode == ERigHierarchyImportMode::ReplaceCurrentGlobalTransform)
+				{
+					Data.GlobalTransforms[Index].NormalizeRotation();
+					SetGlobalTransform(Selection[Index], Data.GlobalTransforms[Index]);
+				}
 				PastedKeys.Add(Selection[Index]);
 			}
 			break;
diff --git a/Engine/Plugins/Experimental/ControlRig/Source/ControlRig/Public/Rigs/RigHierarchyContainer.h b/Engine/Plugins/Experimental/ControlRig/Source/ControlRig/Public/Rigs/RigHierarchyContainer.h
--- a/Engine/Plugins/Experimental/ControlRig/Source/ControlRig/Public/Rigs/RigHierarchyContainer.h
+++ b/Engine/Plugins/Experimental/ControlRig/Source/ControlRig/Public/Rigs/RigHierarchyContainer.h
@@ -23,9 +23,13 @@ enum class ERigHierarchyImportMode : uint8
 
 	Replace,
 
-	ReplaceLocalTransform,
+	ReplaceInitialLocalTransform,
 
-	ReplaceGlobalTransform,
+	ReplaceInitialGlobalTransform,
+
+	ReplaceCurrentLocalTransform,
+
+	ReplaceCurrentGlobalTransform,
 
 	/** MAX - invalid */
 	Max UMETA(Hidden),
@@ -323,7 +327,7 @@ public:
 
 	FString ExportSelectionToText() const;
 	FString ExportToText(const TArray<FRigElementKey>& InSelection) const;
-	TArray<FRigElementKey> ImportFromText(const FString& InContent, ERigHierarchyImportMode InImportMode = ERigHierarchyImportMode::Append, bool bSelectNewElements = true);
+	TArray<FRigElementKey> ImportFromText(const FString& InContent, ERigHierarchyImportMode InImportMode = ERigHierarchyImportMode::Append, bool bSelectNewElements = true, const TArray<FRigElementKey>& InSelection = TArray<FRigElementKey>());
 
 	TArray<FRigElementKey> DuplicateItems(const TArray<FRigElementKey>& InSelection, bool bSelectNewElements = true);
 	TArray<FRigElementKey> MirrorItems(const TArray<FRigElementKey>& InSelection, const FRigMirrorSettings& InSettings, bool bSelectNewElements = true);
diff --git a/Engine/Plugins/Experimental/ControlRig/Source/ControlRigEditor/Private/Editor/SRigHierarchy.cpp b/Engine/Plugins/Experimental/ControlRig/Source/ControlRigEditor/Private/Editor/SRigHierarchy.cpp
--- a/Engine/Plugins/Experimental/ControlRig/Source/ControlRigEditor/Private/Editor/SRigHierarchy.cpp
+++ b/Engine/Plugins/Experimental/ControlRig/Source/ControlRigEditor/Private/Editor/SRigHierarchy.cpp
@@ -1839,7 +1839,7 @@ void SRigHierarchy::HandlePasteItems()
 		Hierarchy->ImportFromText(Content, ERigHierarchyImportMode::Append, true);
 	}
 
-	ControlRigBlueprint->PropagateHierarchyFromBPToInstances(true);
+	//ControlRigBlueprint->PropagateHierarchyFromBPToInstances(true);
 	ControlRigEditor.Pin()->OnHierarchyChanged();
 	RefreshTreeView();
 }
@@ -1854,10 +1854,14 @@ void SRigHierarchy::HandlePasteLocalTransforms()
 		FScopedTransaction Transaction(LOCTEXT("HierarchyTreePasteLocal", "Pasted local transforms."));
 		ControlRigBlueprint->Modify();
 
-		Hierarchy->ImportFromText(Content, ERigHierarchyImportMode::ReplaceLocalTransform, true);
-	}
+		Hierarchy->ImportFromText(Content, ERigHierarchyImportMode::ReplaceCurrentLocalTransform, true);
 
-	ControlRigBlueprint->PropagateHierarchyFromBPToInstances(false);
+		FRigHierarchyContainer* DebuggedHierarchy = GetDebuggedHierarchyContainer();
+		if(DebuggedHierarchy && DebuggedHierarchy != Hierarchy)
+		{
+			DebuggedHierarchy->ImportFromText(Content, ERigHierarchyImportMode::ReplaceCurrentLocalTransform, false, Hierarchy->CurrentSelection());
+		}
+	}
 }
 
 void SRigHierarchy::HandlePasteGlobalTransforms()
@@ -1870,10 +1874,14 @@ void SRigHierarchy::HandlePasteGlobalTransforms()
 		FScopedTransaction Transaction(LOCTEXT("HierarchyTreePasteGlobal", "Pasted global transforms."));
 		ControlRigBlueprint->Modify();
 
-		Hierarchy->ImportFromText(Content, ERigHierarchyImportMode::ReplaceGlobalTransform, true);
-	}
+		Hierarchy->ImportFromText(Content, ERigHierarchyImportMode::ReplaceCurrentGlobalTransform, true);
 
-	ControlRigBlueprint->PropagateHierarchyFromBPToInstances(false);
+		FRigHierarchyContainer* DebuggedHierarchy = GetDebuggedHierarchyContainer();
+		if(DebuggedHierarchy && DebuggedHierarchy != Hierarchy)
+		{
+			DebuggedHierarchy->ImportFromText(Content, ERigHierarchyImportMode::ReplaceCurrentGlobalTransform, false, Hierarchy->CurrentSelection());
+		}
+	}
 }
 
 FRigHierarchyContainer* SRigHierarchy::GetHierarchyContainer() const
@@ -1888,13 +1896,19 @@ FRigHierarchyContainer* SRigHierarchy::GetHierarchyContainer() const
 
 FRigHierarchyContainer* SRigHierarchy::GetDebuggedHierarchyContainer() const
 {
-	if (UControlRig* CurrentRig = ControlRigEditor.Pin()->ControlRig)
+	if (ControlRigBlueprint.IsValid())
 	{
-		return (FRigHierarchyContainer*)CurrentRig->GetHierarchy();
+		if (UControlRig* DebuggedRig = Cast<UControlRig>(ControlRigBlueprint->GetObjectBeingDebugged()))
+		{
+			return (FRigHierarchyContainer*)DebuggedRig->GetHierarchy();
+		}
 	}
-	else if (UControlRig* DebuggedRig = Cast<UControlRig>(ControlRigBlueprint->GetObjectBeingDebugged()))
+	if (ControlRigEditor.IsValid())
 	{
-		return (FRigHierarchyContainer*)DebuggedRig->GetHierarchy();
+		if (UControlRig* CurrentRig = ControlRigEditor.Pin()->ControlRig)
+		{
+			return (FRigHierarchyContainer*)CurrentRig->GetHierarchy();
+		}
 	}
 
 	return GetHierarchyContainer();
