CIS: Remove cyclic include for SControlRigGraphBinVariableBinding.h

diff --git a/Engine/Plugins/Experimental/ControlRig/Source/ControlRigEditor/Private/Graph/SControlRigGraphPinVariableBinding.h b/Engine/Plugins/Experimental/ControlRig/Source/ControlRigEditor/Private/Graph/SControlRigGraphPinVariableBinding.h
--- a/Engine/Plugins/Experimental/ControlRig/Source/ControlRigEditor/Private/Graph/SControlRigGraphPinVariableBinding.h
+++ b/Engine/Plugins/Experimental/ControlRig/Source/ControlRigEditor/Private/Graph/SControlRigGraphPinVariableBinding.h
@@ -9,7 +9,6 @@
 #include "RigVMModel/RigVMPin.h"
 #include "IPropertyAccessEditor.h"
 #include "ControlRigBlueprint.h"
-#include "SControlRigGraphPinVariableBinding.h"
 
 class SControlRigVariableBinding : public SCompoundWidget
 {
