Control Rig: Transparent background for the pin tree to match everyone else.

#rnx

diff --git a/Engine/Plugins/Experimental/ControlRig/Source/ControlRigEditor/Private/Graph/SControlRigGraphNode.cpp b/Engine/Plugins/Experimental/ControlRig/Source/ControlRigEditor/Private/Graph/SControlRigGraphNode.cpp
--- a/Engine/Plugins/Experimental/ControlRig/Source/ControlRigEditor/Private/Graph/SControlRigGraphNode.cpp
+++ b/Engine/Plugins/Experimental/ControlRig/Source/ControlRigEditor/Private/Graph/SControlRigGraphNode.cpp
@@ -440,6 +440,12 @@ class SControlRigExpanderArrow : public SExpanderArrow
 	bool bLeftAligned;
 };
 
+static TAutoConsoleVariable<int32> CVarEnableShowBackground(
+    TEXT("controlrig.ShowBackground"),
+    0,
+    TEXT("Enable unsupported Experimental prototype Modeling Tools"));
+
+
 class SControlRigPinTreeRow : public STableRow<URigVMPin*>
 {
 public:
@@ -456,6 +462,17 @@ public:
 		STableRow<URigVMPin*>::Construct(STableRow<URigVMPin*>::FArguments(), InOwnerTableView);
 	}
 
+	const FSlateBrush* GetBorder() const
+	{
+		bool bShowBG = (CVarEnableShowBackground.GetValueOnGameThread() > 0);
+
+		if (bShowBG)
+			return STableRow<URigVMPin*>::GetBorder();
+		else
+		// We want a transparent background.
+		return FCoreStyle::Get().GetBrush("NoBrush");
+	}
+
 	virtual void ConstructChildren( ETableViewMode::Type InOwnerTableMode, const TAttribute<FMargin>& InPadding, const TSharedRef<SWidget>& InContent ) override
 	{
 		const UGraphEditorSettings* Settings = GetDefault<UGraphEditorSettings>();
