Control Rig: Add python access to Control Offsets

diff --git a/Engine/Plugins/Experimental/ControlRig/Source/ControlRigDeveloper/Private/ControlRigHierarchyModifier.cpp b/Engine/Plugins/Experimental/ControlRig/Source/ControlRigDeveloper/Private/ControlRigHierarchyModifier.cpp
--- a/Engine/Plugins/Experimental/ControlRig/Source/ControlRigDeveloper/Private/ControlRigHierarchyModifier.cpp
+++ b/Engine/Plugins/Experimental/ControlRig/Source/ControlRigDeveloper/Private/ControlRigHierarchyModifier.cpp
@@ -199,6 +199,19 @@ FTransform UControlRigHierarchyModifier::GetControlValueTransform(const FRigElem
 	return FTransform::Identity;
 }
 
+FTransform UControlRigHierarchyModifier::GetControlOffsetTransform(const FRigElementKey& InKey)
+{
+	if (Container != nullptr && InKey.Type == ERigElementType::Control)
+	{
+		int32 Index = Container->GetIndex(InKey);
+		if (Index != INDEX_NONE)
+		{
+			return Container->ControlHierarchy.Controls[Index].OffsetTransform;
+		}
+	}
+	return FTransform::Identity;
+}
+
 void UControlRigHierarchyModifier::SetControlValueBool(const FRigElementKey& InKey, bool InValue, ERigControlValueType InValueType)
 {
 	if (Container != nullptr && InKey.Type == ERigElementType::Control)
@@ -283,6 +296,18 @@ void UControlRigHierarchyModifier::SetControlValueTransform(const FRigElementKey
 	}
 }
 
+void UControlRigHierarchyModifier::SetControlOffsetTransform(const FRigElementKey& InKey, FTransform InValue)
+{
+	if (Container != nullptr && InKey.Type == ERigElementType::Control)
+	{
+		int32 Index = Container->GetIndex(InKey);
+		if (Index != INDEX_NONE)
+		{
+			Container->ControlHierarchy.Controls[Index].OffsetTransform = InValue;
+		}
+	}
+}
+
 FRigElementKey UControlRigHierarchyModifier::AddSpace
 (
 	const FName& InNewName,
diff --git a/Engine/Plugins/Experimental/ControlRig/Source/ControlRigDeveloper/Public/ControlRigHierarchyModifier.h b/Engine/Plugins/Experimental/ControlRig/Source/ControlRigDeveloper/Public/ControlRigHierarchyModifier.h
--- a/Engine/Plugins/Experimental/ControlRig/Source/ControlRigDeveloper/Public/ControlRigHierarchyModifier.h
+++ b/Engine/Plugins/Experimental/ControlRig/Source/ControlRigDeveloper/Public/ControlRigHierarchyModifier.h
@@ -86,6 +86,10 @@ public:
 	UFUNCTION(BlueprintCallable, Category = "Hierarchy")
 	FTransform GetControlValueTransform(const FRigElementKey& InKey, ERigControlValueType InValueType = ERigControlValueType::Initial);
 
+	// returns a control's offset transform
+	UFUNCTION(BlueprintCallable, Category = "Hierarchy")
+		FTransform GetControlOffsetTransform(const FRigElementKey& InKey);
+
 	// Sets a control value
 	UFUNCTION(BlueprintCallable, Category = "Hierarchy")
 	void SetControlValueBool(const FRigElementKey& InKey, bool InValue, ERigControlValueType InValueType = ERigControlValueType::Initial);
@@ -114,6 +118,10 @@ public:
 	UFUNCTION(BlueprintCallable, Category = "Hierarchy")
 	void SetControlValueTransform(const FRigElementKey& InKey, FTransform InValue, ERigControlValueType InValueType = ERigControlValueType::Initial);
 
+	// Sets a control's offset transform
+	UFUNCTION(BlueprintCallable, Category = "Hierarchy")
+	void SetControlOffsetTransform(const FRigElementKey& InKey, FTransform InValue);
+
 	// Adds a new single space
 	UFUNCTION(BlueprintCallable, Category = "Hierarchy")
 	FRigElementKey AddSpace(
