Control Rig: Make all Rig Unit execute function export

diff --git a/Engine/Plugins/Experimental/ControlRig/Source/ControlRig/Private/Units/Animation/RigUnit_AnimBase.h b/Engine/Plugins/Experimental/ControlRig/Source/ControlRig/Private/Units/Animation/RigUnit_AnimBase.h
--- a/Engine/Plugins/Experimental/ControlRig/Source/ControlRig/Private/Units/Animation/RigUnit_AnimBase.h
+++ b/Engine/Plugins/Experimental/ControlRig/Source/ControlRig/Private/Units/Animation/RigUnit_AnimBase.h
@@ -6,7 +6,7 @@
 #include "RigUnit_AnimBase.generated.h"
 
 USTRUCT(meta=(Abstract, Category = "Animation", NodeColor = "0.05 0.05 0.25"))
-struct FRigUnit_AnimBase : public FRigUnit
+struct CONTROLRIG_API FRigUnit_AnimBase : public FRigUnit
 {
 	GENERATED_BODY()
 };
diff --git a/Engine/Plugins/Experimental/ControlRig/Source/ControlRig/Private/Units/Animation/RigUnit_AnimEasing.h b/Engine/Plugins/Experimental/ControlRig/Source/ControlRig/Private/Units/Animation/RigUnit_AnimEasing.h
--- a/Engine/Plugins/Experimental/ControlRig/Source/ControlRig/Private/Units/Animation/RigUnit_AnimEasing.h
+++ b/Engine/Plugins/Experimental/ControlRig/Source/ControlRig/Private/Units/Animation/RigUnit_AnimEasing.h
@@ -10,7 +10,7 @@
  * A constant value of an easing type
  */
 USTRUCT(meta = (DisplayName = "EaseType", Keywords = "Constant"))
-struct FRigUnit_AnimEasingType : public FRigUnit_AnimBase
+struct CONTROLRIG_API FRigUnit_AnimEasingType : public FRigUnit_AnimBase
 {
 	GENERATED_BODY()
 
@@ -30,7 +30,7 @@ struct FRigUnit_AnimEasingType : public FRigUnit_AnimBase
  * Returns the eased version of the input value
  */
 USTRUCT(meta=(DisplayName="Ease", Keywords="Easing,Profile,Smooth,Cubic"))
-struct FRigUnit_AnimEasing : public FRigUnit_AnimBase
+struct CONTROLRIG_API FRigUnit_AnimEasing : public FRigUnit_AnimBase
 {
 	GENERATED_BODY()
 	
diff --git a/Engine/Plugins/Experimental/ControlRig/Source/ControlRig/Private/Units/Animation/RigUnit_AnimEvalRichCurve.h b/Engine/Plugins/Experimental/ControlRig/Source/ControlRig/Private/Units/Animation/RigUnit_AnimEvalRichCurve.h
--- a/Engine/Plugins/Experimental/ControlRig/Source/ControlRig/Private/Units/Animation/RigUnit_AnimEvalRichCurve.h
+++ b/Engine/Plugins/Experimental/ControlRig/Source/ControlRig/Private/Units/Animation/RigUnit_AnimEvalRichCurve.h
@@ -10,7 +10,7 @@
  * Evaluates the provided curve. Values are normalized between 0 and 1
  */
 USTRUCT(meta=(DisplayName="Evaluate Curve", Keywords="Curve,Profile"))
-struct FRigUnit_AnimEvalRichCurve : public FRigUnit_AnimBase
+struct CONTROLRIG_API FRigUnit_AnimEvalRichCurve : public FRigUnit_AnimBase
 {
 	GENERATED_BODY()
 	
diff --git a/Engine/Plugins/Experimental/ControlRig/Source/ControlRig/Private/Units/Animation/RigUnit_AnimRichCurve.h b/Engine/Plugins/Experimental/ControlRig/Source/ControlRig/Private/Units/Animation/RigUnit_AnimRichCurve.h
--- a/Engine/Plugins/Experimental/ControlRig/Source/ControlRig/Private/Units/Animation/RigUnit_AnimRichCurve.h
+++ b/Engine/Plugins/Experimental/ControlRig/Source/ControlRig/Private/Units/Animation/RigUnit_AnimRichCurve.h
@@ -10,7 +10,7 @@
  * Provides a constant curve to be used for multiple curve evaluations
  */
 USTRUCT(meta=(DisplayName="Curve", Keywords="Curve,Profile"))
-struct FRigUnit_AnimRichCurve : public FRigUnit_AnimBase
+struct CONTROLRIG_API FRigUnit_AnimRichCurve : public FRigUnit_AnimBase
 {
 	GENERATED_BODY()
 	
diff --git a/Engine/Plugins/Experimental/ControlRig/Source/ControlRig/Private/Units/Animation/RigUnit_GetDeltaTime.h b/Engine/Plugins/Experimental/ControlRig/Source/ControlRig/Private/Units/Animation/RigUnit_GetDeltaTime.h
--- a/Engine/Plugins/Experimental/ControlRig/Source/ControlRig/Private/Units/Animation/RigUnit_GetDeltaTime.h
+++ b/Engine/Plugins/Experimental/ControlRig/Source/ControlRig/Private/Units/Animation/RigUnit_GetDeltaTime.h
@@ -9,7 +9,7 @@
  * Returns the time gone by from the previous evaluation
  */
 USTRUCT(meta=(DisplayName="Delta Time", Varying))
-struct FRigUnit_GetDeltaTime : public FRigUnit_AnimBase
+struct CONTROLRIG_API FRigUnit_GetDeltaTime : public FRigUnit_AnimBase
 {
 	GENERATED_BODY()
 	
diff --git a/Engine/Plugins/Experimental/ControlRig/Source/ControlRig/Private/Units/Animation/RigUnit_GetWorldTime.h b/Engine/Plugins/Experimental/ControlRig/Source/ControlRig/Private/Units/Animation/RigUnit_GetWorldTime.h
--- a/Engine/Plugins/Experimental/ControlRig/Source/ControlRig/Private/Units/Animation/RigUnit_GetWorldTime.h
+++ b/Engine/Plugins/Experimental/ControlRig/Source/ControlRig/Private/Units/Animation/RigUnit_GetWorldTime.h
@@ -9,7 +9,7 @@
  * Returns the current time (year, month, day, hour, minute)
  */
 USTRUCT(meta = (DisplayName = "Now", Keywords = "Time,Clock"))
-struct FRigUnit_GetWorldTime : public FRigUnit_AnimBase
+struct CONTROLRIG_API FRigUnit_GetWorldTime : public FRigUnit_AnimBase
 {
 	GENERATED_BODY()
 
diff --git a/Engine/Plugins/Experimental/ControlRig/Source/ControlRig/Private/Units/Animation/RigUnit_TimeConversion.h b/Engine/Plugins/Experimental/ControlRig/Source/ControlRig/Private/Units/Animation/RigUnit_TimeConversion.h
--- a/Engine/Plugins/Experimental/ControlRig/Source/ControlRig/Private/Units/Animation/RigUnit_TimeConversion.h
+++ b/Engine/Plugins/Experimental/ControlRig/Source/ControlRig/Private/Units/Animation/RigUnit_TimeConversion.h
@@ -9,7 +9,7 @@
  * Converts frames to seconds based on the current frame rate
  */
 USTRUCT(meta=(DisplayName="Frames to Seconds", Varying))
-struct FRigUnit_FramesToSeconds : public FRigUnit_AnimBase
+struct CONTROLRIG_API FRigUnit_FramesToSeconds : public FRigUnit_AnimBase
 {
 	GENERATED_BODY()
 	
@@ -32,7 +32,7 @@ struct FRigUnit_FramesToSeconds : public FRigUnit_AnimBase
  * Converts seconds to frames based on the current frame rate
  */
 USTRUCT(meta=(DisplayName="Seconds to Frames", Varying))
-struct FRigUnit_SecondsToFrames : public FRigUnit_AnimBase
+struct CONTROLRIG_API FRigUnit_SecondsToFrames : public FRigUnit_AnimBase
 {
 	GENERATED_BODY()
 	
diff --git a/Engine/Plugins/Experimental/ControlRig/Source/ControlRig/Private/Units/Collision/RigUnit_WorldCollision.h b/Engine/Plugins/Experimental/ControlRig/Source/ControlRig/Private/Units/Collision/RigUnit_WorldCollision.h
--- a/Engine/Plugins/Experimental/ControlRig/Source/ControlRig/Private/Units/Collision/RigUnit_WorldCollision.h
+++ b/Engine/Plugins/Experimental/ControlRig/Source/ControlRig/Private/Units/Collision/RigUnit_WorldCollision.h
@@ -10,7 +10,7 @@
  * Sweeps a sphere against the world and return the first blocking hit using a specific channel
  */
 USTRUCT(meta=(DisplayName="Sphere Trace", Category="Collision", DocumentationPolicy = "Strict", Keywords="Sweep,Raytrace,Collision,Collide,Trace", Varying, NodeColor = "0.2 0.4 0.7"))
-struct FRigUnit_SphereTraceWorld : public FRigUnit
+struct CONTROLRIG_API FRigUnit_SphereTraceWorld : public FRigUnit
 {
 	GENERATED_BODY()
 
diff --git a/Engine/Plugins/Experimental/ControlRig/Source/ControlRig/Private/Units/Core/RigUnit_Name.h b/Engine/Plugins/Experimental/ControlRig/Source/ControlRig/Private/Units/Core/RigUnit_Name.h
--- a/Engine/Plugins/Experimental/ControlRig/Source/ControlRig/Private/Units/Core/RigUnit_Name.h
+++ b/Engine/Plugins/Experimental/ControlRig/Source/ControlRig/Private/Units/Core/RigUnit_Name.h
@@ -6,7 +6,7 @@
 #include "RigUnit_Name.generated.h"
 
 USTRUCT(meta=(Abstract, Category="Core|Name", NodeColor = "0.7 0.05 0.5"))
-struct FRigUnit_NameBase : public FRigUnit
+struct CONTROLRIG_API FRigUnit_NameBase : public FRigUnit
 {
 	GENERATED_BODY()
 };
@@ -15,7 +15,7 @@ struct FRigUnit_NameBase : public FRigUnit
  * Concatenates two strings together to make a new string
  */
 USTRUCT(meta = (DisplayName = "Concat", Keywords = "Add,+,Combine,Merge,Append"))
-struct FRigUnit_NameConcat : public FRigUnit_NameBase
+struct CONTROLRIG_API FRigUnit_NameConcat : public FRigUnit_NameBase
 {
 	GENERATED_BODY()
 
@@ -41,7 +41,7 @@ struct FRigUnit_NameConcat : public FRigUnit_NameBase
  * Returns the left or right most characters from the string chopping the given number of characters from the start or the end
  */
 USTRUCT(meta = (DisplayName = "Chop", Keywords = "Truncate,-,Remove,Subtract,Split"))
-struct FRigUnit_NameTruncate : public FRigUnit_NameBase
+struct CONTROLRIG_API FRigUnit_NameTruncate : public FRigUnit_NameBase
 {
 	GENERATED_BODY()
 
@@ -79,7 +79,7 @@ struct FRigUnit_NameTruncate : public FRigUnit_NameBase
  * Replace all occurrences of a substring in this string
  */
 USTRUCT(meta = (DisplayName = "Replace", Keywords = "Search,Emplace,Find"))
-struct FRigUnit_NameReplace : public FRigUnit_NameBase
+struct CONTROLRIG_API FRigUnit_NameReplace : public FRigUnit_NameBase
 {
 	GENERATED_BODY()
 
@@ -108,7 +108,7 @@ struct FRigUnit_NameReplace : public FRigUnit_NameBase
  * Tests whether this string ends with given string
  */
 USTRUCT(meta = (DisplayName = "Ends With", Keywords = "Right"))
-struct FRigUnit_EndsWith : public FRigUnit_NameBase
+struct CONTROLRIG_API FRigUnit_EndsWith : public FRigUnit_NameBase
 {
 	GENERATED_BODY()
 
@@ -135,7 +135,7 @@ struct FRigUnit_EndsWith : public FRigUnit_NameBase
  * Tests whether this string starts with given string
  */
 USTRUCT(meta = (DisplayName = "Starts With", Keywords = "Left"))
-struct FRigUnit_StartsWith : public FRigUnit_NameBase
+struct CONTROLRIG_API FRigUnit_StartsWith : public FRigUnit_NameBase
 {
 	GENERATED_BODY()
 
@@ -162,7 +162,7 @@ struct FRigUnit_StartsWith : public FRigUnit_NameBase
  * Returns true or false if a given name exists in another given name
  */
 USTRUCT(meta = (DisplayName = "Contains", Keywords = "Contains,Find,Has,Search"))
-struct FRigUnit_Contains : public FRigUnit_NameBase
+struct CONTROLRIG_API FRigUnit_Contains : public FRigUnit_NameBase
 {
 	GENERATED_BODY()
 
diff --git a/Engine/Plugins/Experimental/ControlRig/Source/ControlRig/Private/Units/Debug/RigUnit_DebugBase.h b/Engine/Plugins/Experimental/ControlRig/Source/ControlRig/Private/Units/Debug/RigUnit_DebugBase.h
--- a/Engine/Plugins/Experimental/ControlRig/Source/ControlRig/Private/Units/Debug/RigUnit_DebugBase.h
+++ b/Engine/Plugins/Experimental/ControlRig/Source/ControlRig/Private/Units/Debug/RigUnit_DebugBase.h
@@ -6,13 +6,13 @@
 #include "RigUnit_DebugBase.generated.h"
 
 USTRUCT(meta=(Abstract, Category="Debug", NodeColor = "0.25 0.25 0.05"))
-struct FRigUnit_DebugBase : public FRigUnit
+struct CONTROLRIG_API FRigUnit_DebugBase : public FRigUnit
 {
 	GENERATED_BODY()
 };
 
 USTRUCT(meta=(Abstract, Category="Debug", NodeColor = "0.25 0.25 0.05"))
-struct FRigUnit_DebugBaseMutable : public FRigUnitMutable
+struct CONTROLRIG_API FRigUnit_DebugBaseMutable : public FRigUnitMutable
 {
 	GENERATED_BODY()
 };
diff --git a/Engine/Plugins/Experimental/ControlRig/Source/ControlRig/Private/Units/Debug/RigUnit_DebugBezier.h b/Engine/Plugins/Experimental/ControlRig/Source/ControlRig/Private/Units/Debug/RigUnit_DebugBezier.h
--- a/Engine/Plugins/Experimental/ControlRig/Source/ControlRig/Private/Units/Debug/RigUnit_DebugBezier.h
+++ b/Engine/Plugins/Experimental/ControlRig/Source/ControlRig/Private/Units/Debug/RigUnit_DebugBezier.h
@@ -7,7 +7,7 @@
 #include "RigUnit_DebugBezier.generated.h"
 
 USTRUCT(meta=(DisplayName="Draw Bezier", Deprecated = "4.25"))
-struct FRigUnit_DebugBezier : public FRigUnit_DebugBaseMutable
+struct CONTROLRIG_API FRigUnit_DebugBezier : public FRigUnit_DebugBaseMutable
 {
 	GENERATED_BODY()
 
@@ -55,7 +55,7 @@ struct FRigUnit_DebugBezier : public FRigUnit_DebugBaseMutable
 };
 
 USTRUCT(meta=(DisplayName="Draw Bezier", Deprecated = "4.25"))
-struct FRigUnit_DebugBezierItemSpace : public FRigUnit_DebugBaseMutable
+struct CONTROLRIG_API FRigUnit_DebugBezierItemSpace : public FRigUnit_DebugBaseMutable
 {
 	GENERATED_BODY()
 
diff --git a/Engine/Plugins/Experimental/ControlRig/Source/ControlRig/Private/Units/Debug/RigUnit_DebugHierarchy.h b/Engine/Plugins/Experimental/ControlRig/Source/ControlRig/Private/Units/Debug/RigUnit_DebugHierarchy.h
--- a/Engine/Plugins/Experimental/ControlRig/Source/ControlRig/Private/Units/Debug/RigUnit_DebugHierarchy.h
+++ b/Engine/Plugins/Experimental/ControlRig/Source/ControlRig/Private/Units/Debug/RigUnit_DebugHierarchy.h
@@ -10,7 +10,7 @@
  * Draws vectors on each bone in the viewport across the entire hierarchy
  */
 USTRUCT(meta=(DisplayName="Draw Hierarchy"))
-struct FRigUnit_DebugHierarchy : public FRigUnit_DebugBaseMutable
+struct CONTROLRIG_API FRigUnit_DebugHierarchy : public FRigUnit_DebugBaseMutable
 {
 	GENERATED_BODY()
 
diff --git a/Engine/Plugins/Experimental/ControlRig/Source/ControlRig/Private/Units/Debug/RigUnit_DebugLine.h b/Engine/Plugins/Experimental/ControlRig/Source/ControlRig/Private/Units/Debug/RigUnit_DebugLine.h
--- a/Engine/Plugins/Experimental/ControlRig/Source/ControlRig/Private/Units/Debug/RigUnit_DebugLine.h
+++ b/Engine/Plugins/Experimental/ControlRig/Source/ControlRig/Private/Units/Debug/RigUnit_DebugLine.h
@@ -6,7 +6,7 @@
 #include "RigUnit_DebugLine.generated.h"
 
 USTRUCT(meta=(DisplayName="Draw Line", Deprecated = "4.25"))
-struct FRigUnit_DebugLine : public FRigUnit_DebugBaseMutable
+struct CONTROLRIG_API FRigUnit_DebugLine : public FRigUnit_DebugBaseMutable
 {
 	GENERATED_BODY()
 
@@ -48,7 +48,7 @@ struct FRigUnit_DebugLine : public FRigUnit_DebugBaseMutable
  * Draws a line in the viewport given a start and end vector
  */
 USTRUCT(meta=(DisplayName="Draw Line"))
-struct FRigUnit_DebugLineItemSpace : public FRigUnit_DebugBaseMutable
+struct CONTROLRIG_API FRigUnit_DebugLineItemSpace : public FRigUnit_DebugBaseMutable
 {
 	GENERATED_BODY()
 
diff --git a/Engine/Plugins/Experimental/ControlRig/Source/ControlRig/Private/Units/Debug/RigUnit_DebugLineStrip.h b/Engine/Plugins/Experimental/ControlRig/Source/ControlRig/Private/Units/Debug/RigUnit_DebugLineStrip.h
--- a/Engine/Plugins/Experimental/ControlRig/Source/ControlRig/Private/Units/Debug/RigUnit_DebugLineStrip.h
+++ b/Engine/Plugins/Experimental/ControlRig/Source/ControlRig/Private/Units/Debug/RigUnit_DebugLineStrip.h
@@ -6,7 +6,7 @@
 #include "RigUnit_DebugLineStrip.generated.h"
 
 USTRUCT(meta=(DisplayName="Draw Line Strip", Deprecated = "4.25"))
-struct FRigUnit_DebugLineStrip : public FRigUnit_DebugBaseMutable
+struct CONTROLRIG_API FRigUnit_DebugLineStrip : public FRigUnit_DebugBaseMutable
 {
 	GENERATED_BODY()
 
@@ -44,7 +44,7 @@ struct FRigUnit_DebugLineStrip : public FRigUnit_DebugBaseMutable
  * Draws a line strip in the viewport given any number of points
  */
 USTRUCT(meta=(DisplayName="Draw Line Strip"))
-struct FRigUnit_DebugLineStripItemSpace : public FRigUnit_DebugBaseMutable
+struct CONTROLRIG_API FRigUnit_DebugLineStripItemSpace : public FRigUnit_DebugBaseMutable
 {
 	GENERATED_BODY()
 
diff --git a/Engine/Plugins/Experimental/ControlRig/Source/ControlRig/Private/Units/Debug/RigUnit_DebugPoint.h b/Engine/Plugins/Experimental/ControlRig/Source/ControlRig/Private/Units/Debug/RigUnit_DebugPoint.h
--- a/Engine/Plugins/Experimental/ControlRig/Source/ControlRig/Private/Units/Debug/RigUnit_DebugPoint.h
+++ b/Engine/Plugins/Experimental/ControlRig/Source/ControlRig/Private/Units/Debug/RigUnit_DebugPoint.h
@@ -19,7 +19,7 @@ enum class ERigUnitDebugPointMode : uint8
 };
 
 USTRUCT(meta=(DisplayName="Draw Point In Place", Keywords="Draw Vector", Deprecated = "4.25.0"))
-struct FRigUnit_DebugPoint : public FRigUnit_DebugBase
+struct CONTROLRIG_API FRigUnit_DebugPoint : public FRigUnit_DebugBase
 {
 	GENERATED_BODY()
 
@@ -63,7 +63,7 @@ struct FRigUnit_DebugPoint : public FRigUnit_DebugBase
 };
 
 USTRUCT(meta=(DisplayName="Draw Point", Keywords="Draw Vector", Deprecated = "4.25.0"))
-struct FRigUnit_DebugPointMutable : public FRigUnit_DebugBaseMutable
+struct CONTROLRIG_API FRigUnit_DebugPointMutable : public FRigUnit_DebugBaseMutable
 {
 	GENERATED_BODY()
 
diff --git a/Engine/Plugins/Experimental/ControlRig/Source/ControlRig/Private/Units/Debug/RigUnit_DebugPrimitives.h b/Engine/Plugins/Experimental/ControlRig/Source/ControlRig/Private/Units/Debug/RigUnit_DebugPrimitives.h
--- a/Engine/Plugins/Experimental/ControlRig/Source/ControlRig/Private/Units/Debug/RigUnit_DebugPrimitives.h
+++ b/Engine/Plugins/Experimental/ControlRig/Source/ControlRig/Private/Units/Debug/RigUnit_DebugPrimitives.h
@@ -6,7 +6,7 @@
 #include "RigUnit_DebugPrimitives.generated.h"
 
 USTRUCT(meta=(DisplayName="Draw Rectangle", Keywords="Draw Square", Deprecated = "4.25"))
-struct FRigUnit_DebugRectangle : public FRigUnit_DebugBaseMutable
+struct CONTROLRIG_API FRigUnit_DebugRectangle : public FRigUnit_DebugBaseMutable
 {
 	GENERATED_BODY()
 
@@ -49,7 +49,7 @@ struct FRigUnit_DebugRectangle : public FRigUnit_DebugBaseMutable
  * Draws vectors on each bone in the viewport across the entire hierarchy
  */
 USTRUCT(meta=(DisplayName="Draw Rectangle", Keywords="Draw Square"))
-struct FRigUnit_DebugRectangleItemSpace : public FRigUnit_DebugBaseMutable
+struct CONTROLRIG_API FRigUnit_DebugRectangleItemSpace : public FRigUnit_DebugBaseMutable
 {
 	GENERATED_BODY()
 
@@ -91,7 +91,7 @@ struct FRigUnit_DebugRectangleItemSpace : public FRigUnit_DebugBaseMutable
 
 
 USTRUCT(meta=(DisplayName="Draw Arc", Keywords="Draw Ellipse, Draw Circle", Deprecated = "4.25"))
-struct FRigUnit_DebugArc : public FRigUnit_DebugBaseMutable
+struct CONTROLRIG_API FRigUnit_DebugArc : public FRigUnit_DebugBaseMutable
 {
 	GENERATED_BODY()
 
@@ -145,7 +145,7 @@ struct FRigUnit_DebugArc : public FRigUnit_DebugBaseMutable
  * Draws an arc in the viewport, can take in different min and max degrees
  */
 USTRUCT(meta=(DisplayName="Draw Arc", Keywords="Draw Ellipse, Draw Circle"))
-struct FRigUnit_DebugArcItemSpace : public FRigUnit_DebugBaseMutable
+struct CONTROLRIG_API FRigUnit_DebugArcItemSpace : public FRigUnit_DebugBaseMutable
 {
 	GENERATED_BODY()
 
diff --git a/Engine/Plugins/Experimental/ControlRig/Source/ControlRig/Private/Units/Debug/RigUnit_DebugTransform.h b/Engine/Plugins/Experimental/ControlRig/Source/ControlRig/Private/Units/Debug/RigUnit_DebugTransform.h
--- a/Engine/Plugins/Experimental/ControlRig/Source/ControlRig/Private/Units/Debug/RigUnit_DebugTransform.h
+++ b/Engine/Plugins/Experimental/ControlRig/Source/ControlRig/Private/Units/Debug/RigUnit_DebugTransform.h
@@ -22,7 +22,7 @@ enum class ERigUnitDebugTransformMode : uint8
 };
 
 USTRUCT(meta=(DisplayName="Draw Transform In Place", Deprecated = "4.25.0"))
-struct FRigUnit_DebugTransform : public FRigUnit_DebugBase
+struct CONTROLRIG_API FRigUnit_DebugTransform : public FRigUnit_DebugBase
 {
 	GENERATED_BODY()
 
@@ -65,7 +65,7 @@ struct FRigUnit_DebugTransform : public FRigUnit_DebugBase
 };
 
 USTRUCT(meta=(DisplayName="Draw Transform", Deprecated = "4.25"))
-struct FRigUnit_DebugTransformMutable : public FRigUnit_DebugBaseMutable
+struct CONTROLRIG_API FRigUnit_DebugTransformMutable : public FRigUnit_DebugBaseMutable
 {
 	GENERATED_BODY()
 
@@ -111,7 +111,7 @@ struct FRigUnit_DebugTransformMutable : public FRigUnit_DebugBaseMutable
  * Given a transform, will draw a point, axis, or a box in the viewport
  */
 USTRUCT(meta=(DisplayName="Draw Transform"))
-struct FRigUnit_DebugTransformMutableItemSpace : public FRigUnit_DebugBaseMutable
+struct CONTROLRIG_API FRigUnit_DebugTransformMutableItemSpace : public FRigUnit_DebugBaseMutable
 {
 	GENERATED_BODY()
 
@@ -155,7 +155,7 @@ struct FRigUnit_DebugTransformMutableItemSpace : public FRigUnit_DebugBaseMutabl
 };
 
 USTRUCT()
-struct FRigUnit_DebugTransformArrayMutable_WorkData
+struct CONTROLRIG_API FRigUnit_DebugTransformArrayMutable_WorkData
 {
 	GENERATED_BODY()
 		
@@ -164,7 +164,7 @@ struct FRigUnit_DebugTransformArrayMutable_WorkData
 };
 
 USTRUCT(meta=(DisplayName="Draw Transform Array", Deprecated = "4.25.0"))
-struct FRigUnit_DebugTransformArrayMutable : public FRigUnit_DebugBaseMutable
+struct CONTROLRIG_API FRigUnit_DebugTransformArrayMutable : public FRigUnit_DebugBaseMutable
 {
  	GENERATED_BODY()
 
diff --git a/Engine/Plugins/Experimental/ControlRig/Source/ControlRig/Private/Units/Debug/RigUnit_ProfilingBracket.h b/Engine/Plugins/Experimental/ControlRig/Source/ControlRig/Private/Units/Debug/RigUnit_ProfilingBracket.h
--- a/Engine/Plugins/Experimental/ControlRig/Source/ControlRig/Private/Units/Debug/RigUnit_ProfilingBracket.h
+++ b/Engine/Plugins/Experimental/ControlRig/Source/ControlRig/Private/Units/Debug/RigUnit_ProfilingBracket.h
@@ -9,7 +9,7 @@
  * Starts a profiling timer for debugging, used in conjunction with End Profiling Timer
  */
 USTRUCT(meta=(DisplayName="Start Profiling Timer", Keywords="Measure,BeginProfiling,Profile"))
-struct FRigUnit_StartProfilingTimer : public FRigUnit_DebugBaseMutable
+struct CONTROLRIG_API FRigUnit_StartProfilingTimer : public FRigUnit_DebugBaseMutable
 {
 	GENERATED_BODY()
 
@@ -25,7 +25,7 @@ struct FRigUnit_StartProfilingTimer : public FRigUnit_DebugBaseMutable
  * Ends an existing profiling timer for debugging, used in conjunction with Start Profiling Timer
  */
 USTRUCT(meta = (DisplayName = "End Profiling Timer", Keywords = "Measure,StopProfiling,Meter,Profile"))
-struct FRigUnit_EndProfilingTimer : public FRigUnit_DebugBaseMutable
+struct CONTROLRIG_API FRigUnit_EndProfilingTimer : public FRigUnit_DebugBaseMutable
 {
 	GENERATED_BODY()
 
diff --git a/Engine/Plugins/Experimental/ControlRig/Source/ControlRig/Private/Units/Debug/RigUnit_VisualDebug.h b/Engine/Plugins/Experimental/ControlRig/Source/ControlRig/Private/Units/Debug/RigUnit_VisualDebug.h
--- a/Engine/Plugins/Experimental/ControlRig/Source/ControlRig/Private/Units/Debug/RigUnit_VisualDebug.h
+++ b/Engine/Plugins/Experimental/ControlRig/Source/ControlRig/Private/Units/Debug/RigUnit_VisualDebug.h
@@ -19,7 +19,7 @@ enum class ERigUnitVisualDebugPointMode : uint8
 };
 
 USTRUCT(meta=(DisplayName = "Visual Debug Vector", PrototypeName = "VisualDebug", Keywords = "Draw,Point", Deprecated = "4.25", Varying))
-struct FRigUnit_VisualDebugVector : public FRigUnit_DebugBase
+struct CONTROLRIG_API FRigUnit_VisualDebugVector : public FRigUnit_DebugBase
 {
 	GENERATED_BODY()
 
@@ -63,7 +63,7 @@ struct FRigUnit_VisualDebugVector : public FRigUnit_DebugBase
  * Debug draw parameters for a Point or Vector given a vector
  */
 USTRUCT(meta=(DisplayName = "Visual Debug Vector", PrototypeName = "VisualDebug", Keywords = "Draw,Point", Varying))
-struct FRigUnit_VisualDebugVectorItemSpace : public FRigUnit_DebugBase
+struct CONTROLRIG_API FRigUnit_VisualDebugVectorItemSpace : public FRigUnit_DebugBase
 {
 	GENERATED_BODY()
 
@@ -109,7 +109,7 @@ struct FRigUnit_VisualDebugVectorItemSpace : public FRigUnit_DebugBase
 };
 
 USTRUCT(meta = (DisplayName = "Visual Debug Quat", PrototypeName = "VisualDebug", Keywords = "Draw,Rotation", Deprecated = "4.25", Varying))
-struct FRigUnit_VisualDebugQuat : public FRigUnit_DebugBase
+struct CONTROLRIG_API FRigUnit_VisualDebugQuat : public FRigUnit_DebugBase
 {
 	GENERATED_BODY()
 
@@ -145,7 +145,7 @@ struct FRigUnit_VisualDebugQuat : public FRigUnit_DebugBase
  * Debug draw parameters for an Axis given a quaternion
  */
 USTRUCT(meta = (DisplayName = "Visual Debug Quat", PrototypeName = "VisualDebug", Keywords = "Draw,Rotation", Varying))
-struct FRigUnit_VisualDebugQuatItemSpace : public FRigUnit_DebugBase
+struct CONTROLRIG_API FRigUnit_VisualDebugQuatItemSpace : public FRigUnit_DebugBase
 {
 	GENERATED_BODY()
 
@@ -183,7 +183,7 @@ struct FRigUnit_VisualDebugQuatItemSpace : public FRigUnit_DebugBase
 };
 
 USTRUCT(meta=(DisplayName="Visual Debug Transform", PrototypeName = "VisualDebug", Keywords = "Draw,Axes", Deprecated = "4.25", Varying))
-struct FRigUnit_VisualDebugTransform : public FRigUnit_DebugBase
+struct CONTROLRIG_API FRigUnit_VisualDebugTransform : public FRigUnit_DebugBase
 {
 	GENERATED_BODY()
 
@@ -219,7 +219,7 @@ struct FRigUnit_VisualDebugTransform : public FRigUnit_DebugBase
  * Debug draw parameters for an Axis given a transform
  */
 USTRUCT(meta=(DisplayName="Visual Debug Transform", PrototypeName = "VisualDebug", Keywords = "Draw,Axes", Varying))
-struct FRigUnit_VisualDebugTransformItemSpace : public FRigUnit_DebugBase
+struct CONTROLRIG_API FRigUnit_VisualDebugTransformItemSpace : public FRigUnit_DebugBase
 {
 	GENERATED_BODY()
 
diff --git a/Engine/Plugins/Experimental/ControlRig/Source/ControlRig/Private/Units/Deprecated/Math/RigUnit_Converter.h b/Engine/Plugins/Experimental/ControlRig/Source/ControlRig/Private/Units/Deprecated/Math/RigUnit_Converter.h
--- a/Engine/Plugins/Experimental/ControlRig/Source/ControlRig/Private/Units/Deprecated/Math/RigUnit_Converter.h
+++ b/Engine/Plugins/Experimental/ControlRig/Source/ControlRig/Private/Units/Deprecated/Math/RigUnit_Converter.h
@@ -7,7 +7,7 @@
 #include "RigUnit_Converter.generated.h"
 
 USTRUCT(meta = (DisplayName = "ConvertToEulerTransform", Category = "Math|Convert", NodeColor = "0.1 0.1 0.7", Deprecated="4.23.0"))
-struct FRigUnit_ConvertTransform : public FRigUnit
+struct CONTROLRIG_API FRigUnit_ConvertTransform : public FRigUnit
 {
 	GENERATED_BODY()
 
@@ -22,7 +22,7 @@ struct FRigUnit_ConvertTransform : public FRigUnit
 };
 
 USTRUCT(meta = (DisplayName = "ConvertToTransform", Category = "Math|Convert", NodeColor = "0.1 0.1 0.7", Deprecated="4.23.0"))
-struct FRigUnit_ConvertEulerTransform : public FRigUnit
+struct CONTROLRIG_API FRigUnit_ConvertEulerTransform : public FRigUnit
 {
 	GENERATED_BODY()
 
@@ -37,7 +37,7 @@ struct FRigUnit_ConvertEulerTransform : public FRigUnit
 };
 
 USTRUCT(meta = (DisplayName = "ConvertToQuaternion", Category = "Math|Convert", NodeColor = "0.1 0.1 0.7", Deprecated="4.23.0"))
-struct FRigUnit_ConvertRotation : public FRigUnit
+struct CONTROLRIG_API FRigUnit_ConvertRotation : public FRigUnit
 {
 	GENERATED_BODY()
 
@@ -52,13 +52,13 @@ struct FRigUnit_ConvertRotation : public FRigUnit
 };
 
 USTRUCT(meta = (DisplayName = "ConvertToQuaternionDeprecated", Category = "Math|Convert", NodeColor = "0.1 0.1 0.7", Deprecated="4.23.0"))
-struct FRigUnit_ConvertVectorRotation : public FRigUnit_ConvertRotation
+struct CONTROLRIG_API FRigUnit_ConvertVectorRotation : public FRigUnit_ConvertRotation
 {
 	GENERATED_BODY()
 };
 
 USTRUCT(meta = (DisplayName = "ConvertToRotation", Category = "Math|Convert", NodeColor = "0.1 0.1 0.7", Deprecated="4.23.0"))
-struct FRigUnit_ConvertQuaternion: public FRigUnit
+struct CONTROLRIG_API FRigUnit_ConvertQuaternion: public FRigUnit
 {
 	GENERATED_BODY()
 
@@ -73,7 +73,7 @@ struct FRigUnit_ConvertQuaternion: public FRigUnit
 };
 
 USTRUCT(meta = (DisplayName = "ConvertVectorToRotation", Category = "Math|Convert", NodeColor = "0.1 0.1 0.7", Deprecated="4.23.0"))
-struct FRigUnit_ConvertVectorToRotation: public FRigUnit
+struct CONTROLRIG_API FRigUnit_ConvertVectorToRotation: public FRigUnit
 {
 	GENERATED_BODY()
 
@@ -88,7 +88,7 @@ struct FRigUnit_ConvertVectorToRotation: public FRigUnit
 };
 
 USTRUCT(meta = (DisplayName = "ConvertVectorToQuaternion", Category = "Math|Convert", NodeColor = "0.1 0.1 0.7", Deprecated="4.23.0"))
-struct FRigUnit_ConvertVectorToQuaternion: public FRigUnit
+struct CONTROLRIG_API FRigUnit_ConvertVectorToQuaternion: public FRigUnit
 {
 	GENERATED_BODY()
 
@@ -104,7 +104,7 @@ struct FRigUnit_ConvertVectorToQuaternion: public FRigUnit
 
 
 USTRUCT(meta = (DisplayName = "ConvertRotationToVector", Category = "Math|Convert", NodeColor = "0.1 0.1 0.7", Deprecated="4.23.0"))
-struct FRigUnit_ConvertRotationToVector: public FRigUnit
+struct CONTROLRIG_API FRigUnit_ConvertRotationToVector: public FRigUnit
 {
 	GENERATED_BODY()
 
@@ -119,7 +119,7 @@ struct FRigUnit_ConvertRotationToVector: public FRigUnit
 };
 
 USTRUCT(meta = (DisplayName = "ConvertQuaternionToVector", Category = "Math|Convert", NodeColor = "0.1 0.1 0.7", Deprecated="4.23.0"))
-struct FRigUnit_ConvertQuaternionToVector: public FRigUnit
+struct CONTROLRIG_API FRigUnit_ConvertQuaternionToVector: public FRigUnit
 {
 	GENERATED_BODY()
 
@@ -134,7 +134,7 @@ struct FRigUnit_ConvertQuaternionToVector: public FRigUnit
 };
 
 USTRUCT(meta = (DisplayName = "ToSwingAndTwist", Category = "Math|Transform", NodeColor = "0.1 0.1 0.7", Deprecated="4.23.0"))
-struct FRigUnit_ToSwingAndTwist : public FRigUnit
+struct CONTROLRIG_API FRigUnit_ToSwingAndTwist : public FRigUnit
 {
 	GENERATED_BODY()
 
diff --git a/Engine/Plugins/Experimental/ControlRig/Source/ControlRig/Private/Units/Deprecated/Math/RigUnit_Float.h b/Engine/Plugins/Experimental/ControlRig/Source/ControlRig/Private/Units/Deprecated/Math/RigUnit_Float.h
--- a/Engine/Plugins/Experimental/ControlRig/Source/ControlRig/Private/Units/Deprecated/Math/RigUnit_Float.h
+++ b/Engine/Plugins/Experimental/ControlRig/Source/ControlRig/Private/Units/Deprecated/Math/RigUnit_Float.h
@@ -8,7 +8,7 @@
 
 /** Two args and a result of float type */
 USTRUCT(meta=(Abstract, NodeColor = "0.1 0.7 0.1", Deprecated="4.23.0"))
-struct FRigUnit_BinaryFloatOp : public FRigUnit
+struct CONTROLRIG_API FRigUnit_BinaryFloatOp : public FRigUnit
 {
 	GENERATED_BODY()
 
@@ -23,7 +23,7 @@ struct FRigUnit_BinaryFloatOp : public FRigUnit
 };
 
 USTRUCT(meta=(DisplayName="Multiply", Category="Math|Float", Keywords="*", Deprecated="4.23.0"))
-struct FRigUnit_Multiply_FloatFloat : public FRigUnit_BinaryFloatOp
+struct CONTROLRIG_API FRigUnit_Multiply_FloatFloat : public FRigUnit_BinaryFloatOp
 {
 	GENERATED_BODY()
 
@@ -32,7 +32,7 @@ struct FRigUnit_Multiply_FloatFloat : public FRigUnit_BinaryFloatOp
 };
 
 USTRUCT(meta=(DisplayName="Add", Category="Math|Float", Keywords = "+,Sum", Deprecated="4.23.0"))
-struct FRigUnit_Add_FloatFloat : public FRigUnit_BinaryFloatOp
+struct CONTROLRIG_API FRigUnit_Add_FloatFloat : public FRigUnit_BinaryFloatOp
 {
 	GENERATED_BODY()
 
@@ -41,7 +41,7 @@ struct FRigUnit_Add_FloatFloat : public FRigUnit_BinaryFloatOp
 };
 
 USTRUCT(meta=(DisplayName="Subtract", Category="Math|Float", Keywords = "-", Deprecated="4.23.0"))
-struct FRigUnit_Subtract_FloatFloat : public FRigUnit_BinaryFloatOp
+struct CONTROLRIG_API FRigUnit_Subtract_FloatFloat : public FRigUnit_BinaryFloatOp
 {
 	GENERATED_BODY()
 
@@ -50,7 +50,7 @@ struct FRigUnit_Subtract_FloatFloat : public FRigUnit_BinaryFloatOp
 };
 
 USTRUCT(meta=(DisplayName="Divide", Category="Math|Float", Keywords = "/", Deprecated="4.23.0"))
-struct FRigUnit_Divide_FloatFloat : public FRigUnit_BinaryFloatOp
+struct CONTROLRIG_API FRigUnit_Divide_FloatFloat : public FRigUnit_BinaryFloatOp
 {
 	GENERATED_BODY()
 
@@ -60,7 +60,7 @@ struct FRigUnit_Divide_FloatFloat : public FRigUnit_BinaryFloatOp
 
 /** Two args and a result of float type */
 USTRUCT(meta = (DisplayName = "Clamp", Category = "Math|Float", NodeColor = "0.1 0.7 0.1", Deprecated="4.23.0"))
-struct FRigUnit_Clamp_Float: public FRigUnit
+struct CONTROLRIG_API FRigUnit_Clamp_Float: public FRigUnit
 {
 	GENERATED_BODY()
 
@@ -82,7 +82,7 @@ struct FRigUnit_Clamp_Float: public FRigUnit
 
 /** Two args and a result of float type */
 USTRUCT(meta = (DisplayName = "MapRange", Category = "Math|Float", Deprecated="4.23.0"))
-struct FRigUnit_MapRange_Float: public FRigUnit
+struct CONTROLRIG_API FRigUnit_MapRange_Float: public FRigUnit
 {
 	GENERATED_BODY()
 
diff --git a/Engine/Plugins/Experimental/ControlRig/Source/ControlRig/Private/Units/Deprecated/Math/RigUnit_Quaternion.h b/Engine/Plugins/Experimental/ControlRig/Source/ControlRig/Private/Units/Deprecated/Math/RigUnit_Quaternion.h
--- a/Engine/Plugins/Experimental/ControlRig/Source/ControlRig/Private/Units/Deprecated/Math/RigUnit_Quaternion.h
+++ b/Engine/Plugins/Experimental/ControlRig/Source/ControlRig/Private/Units/Deprecated/Math/RigUnit_Quaternion.h
@@ -8,7 +8,7 @@
 
 /** Two args and a result of Quaternion type */
 USTRUCT(meta=(Abstract, NodeColor = "0.1 0.7 0.1", Deprecated="4.23.0"))
-struct FRigUnit_BinaryQuaternionOp : public FRigUnit
+struct CONTROLRIG_API FRigUnit_BinaryQuaternionOp : public FRigUnit
 {
 	GENERATED_BODY()
 
@@ -29,7 +29,7 @@ struct FRigUnit_BinaryQuaternionOp : public FRigUnit
 };
 
 USTRUCT(meta=(DisplayName="Multiply(Quaternion)", Category="Math|Quaternion", Deprecated="4.23.0"))
-struct FRigUnit_MultiplyQuaternion : public FRigUnit_BinaryQuaternionOp
+struct CONTROLRIG_API FRigUnit_MultiplyQuaternion : public FRigUnit_BinaryQuaternionOp
 {
 	GENERATED_BODY()
 
@@ -39,7 +39,7 @@ struct FRigUnit_MultiplyQuaternion : public FRigUnit_BinaryQuaternionOp
 
 /** Two args and a result of Quaternion type */
 USTRUCT(meta=(Abstract, NodeColor = "0.1 0.7 0.1", Deprecated="4.23.0"))
-struct FRigUnit_UnaryQuaternionOp : public FRigUnit
+struct CONTROLRIG_API FRigUnit_UnaryQuaternionOp : public FRigUnit
 {
 	GENERATED_BODY()
 
@@ -51,7 +51,7 @@ struct FRigUnit_UnaryQuaternionOp : public FRigUnit
 };
 
 USTRUCT(meta = (DisplayName = "Inverse(Quaternion)", Category = "Math|Quaternion", Deprecated="4.23.0"))
-struct FRigUnit_InverseQuaterion: public FRigUnit_UnaryQuaternionOp
+struct CONTROLRIG_API FRigUnit_InverseQuaterion: public FRigUnit_UnaryQuaternionOp
 {
 	GENERATED_BODY()
 
@@ -60,7 +60,7 @@ struct FRigUnit_InverseQuaterion: public FRigUnit_UnaryQuaternionOp
 };
 
 USTRUCT(meta = (DisplayName = "To Axis And Angle(Quaternion)", Category = "Math|Quaternion", NodeColor = "0.1 0.7 0.1", Deprecated="4.23.0"))
-struct FRigUnit_QuaternionToAxisAndAngle : public FRigUnit
+struct CONTROLRIG_API FRigUnit_QuaternionToAxisAndAngle : public FRigUnit
 {
 	GENERATED_BODY()
 
@@ -78,7 +78,7 @@ struct FRigUnit_QuaternionToAxisAndAngle : public FRigUnit
 };
 
 USTRUCT(meta = (DisplayName = "From Axis And Angle(Quaternion)", Category = "Math|Quaternion", NodeColor = "0.1 0.7 0.1", Deprecated="4.23.0"))
-struct FRigUnit_QuaternionFromAxisAndAngle : public FRigUnit
+struct CONTROLRIG_API FRigUnit_QuaternionFromAxisAndAngle : public FRigUnit
 {
 	GENERATED_BODY()
 
@@ -102,7 +102,7 @@ struct FRigUnit_QuaternionFromAxisAndAngle : public FRigUnit
 };
 
 USTRUCT(meta = (DisplayName = "Get Angle Around Axis", Category = "Math|Quaternion", NodeColor = "0.1 0.7 0.1", Deprecated="4.23.0"))
-struct FRigUnit_QuaternionToAngle : public FRigUnit
+struct CONTROLRIG_API FRigUnit_QuaternionToAngle : public FRigUnit
 {
 	GENERATED_BODY()
 
diff --git a/Engine/Plugins/Experimental/ControlRig/Source/ControlRig/Private/Units/Deprecated/Math/RigUnit_Transform.h b/Engine/Plugins/Experimental/ControlRig/Source/ControlRig/Private/Units/Deprecated/Math/RigUnit_Transform.h
--- a/Engine/Plugins/Experimental/ControlRig/Source/ControlRig/Private/Units/Deprecated/Math/RigUnit_Transform.h
+++ b/Engine/Plugins/Experimental/ControlRig/Source/ControlRig/Private/Units/Deprecated/Math/RigUnit_Transform.h
@@ -8,7 +8,7 @@
 
 /** Two args and a result of Transform type */
 USTRUCT(meta=(Abstract, NodeColor = "0.1 0.7 0.1", Deprecated="4.23.0"))
-struct FRigUnit_BinaryTransformOp : public FRigUnit
+struct CONTROLRIG_API FRigUnit_BinaryTransformOp : public FRigUnit
 {
 	GENERATED_BODY()
 
@@ -23,7 +23,7 @@ struct FRigUnit_BinaryTransformOp : public FRigUnit
 };
 
 USTRUCT(meta=(DisplayName="Multiply(Transform)", Category="Math|Transform", Deprecated="4.23.0"))
-struct FRigUnit_MultiplyTransform : public FRigUnit_BinaryTransformOp
+struct CONTROLRIG_API FRigUnit_MultiplyTransform : public FRigUnit_BinaryTransformOp
 {
 	GENERATED_BODY()
 
@@ -32,7 +32,7 @@ struct FRigUnit_MultiplyTransform : public FRigUnit_BinaryTransformOp
 };
 
 USTRUCT(meta = (DisplayName = "GetRelativeTransform", Category = "Math|Transform", Deprecated="4.23.0"))
-struct FRigUnit_GetRelativeTransform : public FRigUnit_BinaryTransformOp
+struct CONTROLRIG_API FRigUnit_GetRelativeTransform : public FRigUnit_BinaryTransformOp
 {
 	GENERATED_BODY()
 
diff --git a/Engine/Plugins/Experimental/ControlRig/Source/ControlRig/Private/Units/Deprecated/Math/RigUnit_Vector.h b/Engine/Plugins/Experimental/ControlRig/Source/ControlRig/Private/Units/Deprecated/Math/RigUnit_Vector.h
--- a/Engine/Plugins/Experimental/ControlRig/Source/ControlRig/Private/Units/Deprecated/Math/RigUnit_Vector.h
+++ b/Engine/Plugins/Experimental/ControlRig/Source/ControlRig/Private/Units/Deprecated/Math/RigUnit_Vector.h
@@ -8,7 +8,7 @@
 
 /** Two args and a result of Vector type */
 USTRUCT(meta=(Abstract, NodeColor = "0.1 0.7 0.1", Deprecated="4.23.0"))
-struct FRigUnit_BinaryVectorOp : public FRigUnit
+struct CONTROLRIG_API FRigUnit_BinaryVectorOp : public FRigUnit
 {
 	GENERATED_BODY()
 
@@ -23,7 +23,7 @@ struct FRigUnit_BinaryVectorOp : public FRigUnit
 };
 
 USTRUCT(meta=(DisplayName="Multiply(Vector)", Category="Math|Vector", Deprecated="4.23.0"))
-struct FRigUnit_Multiply_VectorVector : public FRigUnit_BinaryVectorOp
+struct CONTROLRIG_API FRigUnit_Multiply_VectorVector : public FRigUnit_BinaryVectorOp
 {
 	GENERATED_BODY()
 
@@ -32,7 +32,7 @@ struct FRigUnit_Multiply_VectorVector : public FRigUnit_BinaryVectorOp
 };
 
 USTRUCT(meta=(DisplayName="Add(Vector)", Category="Math|Vector", Deprecated="4.23.0"))
-struct FRigUnit_Add_VectorVector : public FRigUnit_BinaryVectorOp
+struct CONTROLRIG_API FRigUnit_Add_VectorVector : public FRigUnit_BinaryVectorOp
 {
 	GENERATED_BODY()
 
@@ -41,7 +41,7 @@ struct FRigUnit_Add_VectorVector : public FRigUnit_BinaryVectorOp
 };
 
 USTRUCT(meta=(DisplayName="Subtract(Vector)", Category="Math|Vector", Deprecated="4.23.0"))
-struct FRigUnit_Subtract_VectorVector : public FRigUnit_BinaryVectorOp
+struct CONTROLRIG_API FRigUnit_Subtract_VectorVector : public FRigUnit_BinaryVectorOp
 {
 	GENERATED_BODY()
 
@@ -50,7 +50,7 @@ struct FRigUnit_Subtract_VectorVector : public FRigUnit_BinaryVectorOp
 };
 
 USTRUCT(meta=(DisplayName="Divide(Vector)", Category="Math|Vector", Deprecated="4.23.0"))
-struct FRigUnit_Divide_VectorVector : public FRigUnit_BinaryVectorOp
+struct CONTROLRIG_API FRigUnit_Divide_VectorVector : public FRigUnit_BinaryVectorOp
 {
 	GENERATED_BODY()
 
@@ -59,7 +59,7 @@ struct FRigUnit_Divide_VectorVector : public FRigUnit_BinaryVectorOp
 };
 
 USTRUCT(meta = (DisplayName = "Distance", Category = "Math|Vector", Deprecated="4.23.0", NodeColor = "0.1 0.7 0.1"))
-struct FRigUnit_Distance_VectorVector : public FRigUnit
+struct CONTROLRIG_API FRigUnit_Distance_VectorVector : public FRigUnit
 {
 	GENERATED_BODY()
 
diff --git a/Engine/Plugins/Experimental/ControlRig/Source/ControlRig/Private/Units/Deprecated/RigUnit_AimConstraint.h b/Engine/Plugins/Experimental/ControlRig/Source/ControlRig/Private/Units/Deprecated/RigUnit_AimConstraint.h
--- a/Engine/Plugins/Experimental/ControlRig/Source/ControlRig/Private/Units/Deprecated/RigUnit_AimConstraint.h
+++ b/Engine/Plugins/Experimental/ControlRig/Source/ControlRig/Private/Units/Deprecated/RigUnit_AimConstraint.h
@@ -50,7 +50,7 @@ struct FAimTarget
 };
 
 USTRUCT()
-struct FRigUnit_AimConstraint_WorkData
+struct CONTROLRIG_API FRigUnit_AimConstraint_WorkData
 {
 	GENERATED_BODY()
 
@@ -60,7 +60,7 @@ struct FRigUnit_AimConstraint_WorkData
 };
 
 USTRUCT(meta=(DisplayName="Aim Constraint", Category="Transforms", Deprecated = "4.23.0"))
-struct FRigUnit_AimConstraint : public FRigUnitMutable
+struct CONTROLRIG_API FRigUnit_AimConstraint : public FRigUnitMutable
 {
 	GENERATED_BODY()
 
diff --git a/Engine/Plugins/Experimental/ControlRig/Source/ControlRig/Private/Units/Deprecated/RigUnit_ApplyFK.h b/Engine/Plugins/Experimental/ControlRig/Source/ControlRig/Private/Units/Deprecated/RigUnit_ApplyFK.h
--- a/Engine/Plugins/Experimental/ControlRig/Source/ControlRig/Private/Units/Deprecated/RigUnit_ApplyFK.h
+++ b/Engine/Plugins/Experimental/ControlRig/Source/ControlRig/Private/Units/Deprecated/RigUnit_ApplyFK.h
@@ -26,7 +26,7 @@ enum class EApplyTransformMode : uint8
 };
 
 USTRUCT(meta=(DisplayName="Apply FK", Category="Transforms", Deprecated = "4.23.0"))
-struct FRigUnit_ApplyFK : public FRigUnitMutable
+struct CONTROLRIG_API FRigUnit_ApplyFK : public FRigUnitMutable
 {
 	GENERATED_BODY()
 
diff --git a/Engine/Plugins/Experimental/ControlRig/Source/ControlRig/Private/Units/Deprecated/RigUnit_BlendTransform.h b/Engine/Plugins/Experimental/ControlRig/Source/ControlRig/Private/Units/Deprecated/RigUnit_BlendTransform.h
--- a/Engine/Plugins/Experimental/ControlRig/Source/ControlRig/Private/Units/Deprecated/RigUnit_BlendTransform.h
+++ b/Engine/Plugins/Experimental/ControlRig/Source/ControlRig/Private/Units/Deprecated/RigUnit_BlendTransform.h
@@ -22,7 +22,7 @@ struct FBlendTarget
 };
 
 USTRUCT(meta = (DisplayName = "Blend(Transform)", Category = "Blend", Deprecated = "4.26.0"))
-struct FRigUnit_BlendTransform : public FRigUnit
+struct CONTROLRIG_API FRigUnit_BlendTransform : public FRigUnit
 {
 	GENERATED_BODY()
 
diff --git a/Engine/Plugins/Experimental/ControlRig/Source/ControlRig/Private/Units/Deprecated/RigUnit_GetJointTransform.h b/Engine/Plugins/Experimental/ControlRig/Source/ControlRig/Private/Units/Deprecated/RigUnit_GetJointTransform.h
--- a/Engine/Plugins/Experimental/ControlRig/Source/ControlRig/Private/Units/Deprecated/RigUnit_GetJointTransform.h
+++ b/Engine/Plugins/Experimental/ControlRig/Source/ControlRig/Private/Units/Deprecated/RigUnit_GetJointTransform.h
@@ -9,7 +9,7 @@
 #include "RigUnit_GetJointTransform.generated.h"
 
 USTRUCT(meta=(DisplayName="Get Joint Transform", Category="Transforms", Deprecated = "4.23.0"))
-struct FRigUnit_GetJointTransform : public FRigUnitMutable
+struct CONTROLRIG_API FRigUnit_GetJointTransform : public FRigUnitMutable
 {
 	GENERATED_BODY()
 
diff --git a/Engine/Plugins/Experimental/ControlRig/Source/ControlRig/Private/Units/Drawing/RigUnit_DrawContainer.h b/Engine/Plugins/Experimental/ControlRig/Source/ControlRig/Private/Units/Drawing/RigUnit_DrawContainer.h
--- a/Engine/Plugins/Experimental/ControlRig/Source/ControlRig/Private/Units/Drawing/RigUnit_DrawContainer.h
+++ b/Engine/Plugins/Experimental/ControlRig/Source/ControlRig/Private/Units/Drawing/RigUnit_DrawContainer.h
@@ -9,7 +9,7 @@
  * Get Imported Draw Container curve transform and color
  */
 USTRUCT(meta=(DisplayName="Get Draw Instruction", Category="Drawing", NodeColor = "0.25 0.25 0.05", Keywords = "Curve,Shape"))
-struct FRigUnit_DrawContainerGetInstruction : public FRigUnit
+struct CONTROLRIG_API FRigUnit_DrawContainerGetInstruction : public FRigUnit
 {
 	GENERATED_BODY()
 
@@ -37,7 +37,7 @@ struct FRigUnit_DrawContainerGetInstruction : public FRigUnit
  * Set Imported Draw Container curve color
  */
 USTRUCT(meta = (DisplayName = "Set Draw Color", Category = "Drawing", NodeColor = "0.25 0.25 0.05", Keywords = "Curve,Shape"))
-struct FRigUnit_DrawContainerSetColor: public FRigUnitMutable
+struct CONTROLRIG_API FRigUnit_DrawContainerSetColor: public FRigUnitMutable
 {
 	GENERATED_BODY()
 
@@ -61,7 +61,7 @@ struct FRigUnit_DrawContainerSetColor: public FRigUnitMutable
  * Set Imported Draw Container curve thickness
  */
 USTRUCT(meta = (DisplayName = "Set Draw Thickness", Category = "Drawing", NodeColor = "0.25 0.25 0.05", Keywords = "Curve,Shape"))
-struct FRigUnit_DrawContainerSetThickness : public FRigUnitMutable
+struct CONTROLRIG_API FRigUnit_DrawContainerSetThickness : public FRigUnitMutable
 {
 	GENERATED_BODY()
 
@@ -85,7 +85,7 @@ struct FRigUnit_DrawContainerSetThickness : public FRigUnitMutable
  * Set Imported Draw Container curve transform
  */
 USTRUCT(meta = (DisplayName = "Set Draw Transform", Category = "Drawing", NodeColor = "0.25 0.25 0.05", Keywords = "Curve,Shape"))
-struct FRigUnit_DrawContainerSetTransform : public FRigUnitMutable
+struct CONTROLRIG_API FRigUnit_DrawContainerSetTransform : public FRigUnitMutable
 {
 	GENERATED_BODY()
 
diff --git a/Engine/Plugins/Experimental/ControlRig/Source/ControlRig/Private/Units/Execution/RigUnit_Collection.h b/Engine/Plugins/Experimental/ControlRig/Source/ControlRig/Private/Units/Execution/RigUnit_Collection.h
--- a/Engine/Plugins/Experimental/ControlRig/Source/ControlRig/Private/Units/Execution/RigUnit_Collection.h
+++ b/Engine/Plugins/Experimental/ControlRig/Source/ControlRig/Private/Units/Execution/RigUnit_Collection.h
@@ -7,13 +7,13 @@
 #include "RigUnit_Collection.generated.h"
 
 USTRUCT(meta = (Abstract, NodeColor = "0.7 0.05 0.5", Category = "Collections"))
-struct FRigUnit_CollectionBase : public FRigUnit
+struct CONTROLRIG_API FRigUnit_CollectionBase : public FRigUnit
 {
 	GENERATED_BODY()
 };
 
 USTRUCT(meta = (Abstract, NodeColor = "0.7 0.05 0.5", Category = "Collections"))
-struct FRigUnit_CollectionBaseMutable : public FRigUnitMutable
+struct CONTROLRIG_API FRigUnit_CollectionBaseMutable : public FRigUnitMutable
 {
 	GENERATED_BODY()
 };
@@ -23,7 +23,7 @@ struct FRigUnit_CollectionBaseMutable : public FRigUnitMutable
  * Chains can refer to bone chains or chains within a control hierarchy.
  */
 USTRUCT(meta=(DisplayName="Item Chain", Keywords="Bone,Joint,Collection", Varying))
-struct FRigUnit_CollectionChain : public FRigUnit_CollectionBase
+struct CONTROLRIG_API FRigUnit_CollectionChain : public FRigUnit_CollectionBase
 {
 	GENERATED_BODY()
 
@@ -61,7 +61,7 @@ struct FRigUnit_CollectionChain : public FRigUnit_CollectionBase
  * The name search is case sensitive.
  */
 USTRUCT(meta = (DisplayName = "Item Name Search", Keywords = "Bone,Joint,Collection,Filter", Varying))
-struct FRigUnit_CollectionNameSearch : public FRigUnit_CollectionBase
+struct CONTROLRIG_API FRigUnit_CollectionNameSearch : public FRigUnit_CollectionBase
 {
 	GENERATED_BODY()
 
@@ -96,7 +96,7 @@ struct FRigUnit_CollectionNameSearch : public FRigUnit_CollectionBase
  * of a provided parent item. Returns an empty collection for an invalid parent item.
  */
 USTRUCT(meta = (DisplayName = "Children", Keywords = "Bone,Joint,Collection,Filter,Parent", Varying))
-struct FRigUnit_CollectionChildren : public FRigUnit_CollectionBase
+struct CONTROLRIG_API FRigUnit_CollectionChildren : public FRigUnit_CollectionBase
 {
 	GENERATED_BODY()
 
@@ -138,7 +138,7 @@ struct FRigUnit_CollectionChildren : public FRigUnit_CollectionBase
  * Replaces all names within the collection
  */
 USTRUCT(meta = (DisplayName = "Replace Items", Keywords = "Replace,Find", Varying))
-struct FRigUnit_CollectionReplaceItems : public FRigUnit_CollectionBase
+struct CONTROLRIG_API FRigUnit_CollectionReplaceItems : public FRigUnit_CollectionBase
 {
 	GENERATED_BODY()
 
@@ -178,7 +178,7 @@ struct FRigUnit_CollectionReplaceItems : public FRigUnit_CollectionBase
  * Returns a collection provided a specific list of items.
  */
 USTRUCT(meta = (DisplayName = "Items", Keywords = "Collection,Array", Varying))
-struct FRigUnit_CollectionItems : public FRigUnit_CollectionBase
+struct CONTROLRIG_API FRigUnit_CollectionItems : public FRigUnit_CollectionBase
 {
 	GENERATED_BODY()
 
@@ -202,7 +202,7 @@ struct FRigUnit_CollectionItems : public FRigUnit_CollectionBase
  * (the combination of all items from both A and B).
  */
 USTRUCT(meta = (DisplayName = "Union", Keywords = "Combine,Add,Merge,Collection,Hierarchy"))
-struct FRigUnit_CollectionUnion : public FRigUnit_CollectionBase
+struct CONTROLRIG_API FRigUnit_CollectionUnion : public FRigUnit_CollectionBase
 {
 	GENERATED_BODY()
 
@@ -228,7 +228,7 @@ struct FRigUnit_CollectionUnion : public FRigUnit_CollectionBase
  * (the items present in both A and B).
  */
 USTRUCT(meta = (DisplayName = "Intersection", Keywords = "Combine,Merge,Collection,Hierarchy"))
-struct FRigUnit_CollectionIntersection : public FRigUnit_CollectionBase
+struct CONTROLRIG_API FRigUnit_CollectionIntersection : public FRigUnit_CollectionBase
 {
 	GENERATED_BODY()
 
@@ -254,7 +254,7 @@ struct FRigUnit_CollectionIntersection : public FRigUnit_CollectionBase
  * (the items present in A but not in B).
  */
 USTRUCT(meta = (DisplayName = "Difference", Keywords = "Collection,Exclude,Subtract"))
-struct FRigUnit_CollectionDifference : public FRigUnit_CollectionBase
+struct CONTROLRIG_API FRigUnit_CollectionDifference : public FRigUnit_CollectionBase
 {
 	GENERATED_BODY()
 
@@ -279,7 +279,7 @@ struct FRigUnit_CollectionDifference : public FRigUnit_CollectionBase
  * Returns the collection in reverse order
  */
 USTRUCT(meta = (DisplayName = "Reverse", Keywords = "Direction,Order,Reverse"))
-struct FRigUnit_CollectionReverse : public FRigUnit_CollectionBase
+struct CONTROLRIG_API FRigUnit_CollectionReverse : public FRigUnit_CollectionBase
 {
 	GENERATED_BODY()
 
@@ -301,7 +301,7 @@ struct FRigUnit_CollectionReverse : public FRigUnit_CollectionBase
  * Returns the number of elements in a collection
  */
 USTRUCT(meta = (DisplayName = "Count", Keywords = "Collection,Array,Count,Num,Length,Size"))
-struct FRigUnit_CollectionCount : public FRigUnit_CollectionBase
+struct CONTROLRIG_API FRigUnit_CollectionCount : public FRigUnit_CollectionBase
 {
 	GENERATED_BODY()
 
@@ -325,7 +325,7 @@ struct FRigUnit_CollectionCount : public FRigUnit_CollectionBase
  * Returns a single item within a collection by index
  */
 USTRUCT(meta = (DisplayName = "Item At Index", Keywords = "Item,GetIndex,AtIndex,At,ForIndex,[]"))
-struct FRigUnit_CollectionItemAtIndex : public FRigUnit_CollectionBase
+struct CONTROLRIG_API FRigUnit_CollectionItemAtIndex : public FRigUnit_CollectionBase
 {
 	GENERATED_BODY()
 
@@ -353,7 +353,7 @@ struct FRigUnit_CollectionItemAtIndex : public FRigUnit_CollectionBase
  * Given a collection of items, execute iteratively across all items in a given collection
  */
 USTRUCT(meta=(DisplayName="For Each Item", Keywords="Collection,Loop,Iterate"))
-struct FRigUnit_CollectionLoop : public FRigUnit_CollectionBaseMutable
+struct CONTROLRIG_API FRigUnit_CollectionLoop : public FRigUnit_CollectionBaseMutable
 {
 	GENERATED_BODY()
 
diff --git a/Engine/Plugins/Experimental/ControlRig/Source/ControlRig/Private/Units/Execution/RigUnit_ForLoop.h b/Engine/Plugins/Experimental/ControlRig/Source/ControlRig/Private/Units/Execution/RigUnit_ForLoop.h
--- a/Engine/Plugins/Experimental/ControlRig/Source/ControlRig/Private/Units/Execution/RigUnit_ForLoop.h
+++ b/Engine/Plugins/Experimental/ControlRig/Source/ControlRig/Private/Units/Execution/RigUnit_ForLoop.h
@@ -10,7 +10,7 @@
  * Given a count, execute iteratively until the count is up
  */
 USTRUCT(meta=(DisplayName="For Loop", Category="Execution", TitleColor="1 0 0", NodeColor="0.1 0.1 0.1", Keywords="Iterate"))
-struct FRigUnit_ForLoopCount : public FRigUnitMutable
+struct CONTROLRIG_API FRigUnit_ForLoopCount : public FRigUnitMutable
 {
 	GENERATED_BODY()
 
diff --git a/Engine/Plugins/Experimental/ControlRig/Source/ControlRig/Private/Units/Execution/RigUnit_Hierarchy.h b/Engine/Plugins/Experimental/ControlRig/Source/ControlRig/Private/Units/Execution/RigUnit_Hierarchy.h
--- a/Engine/Plugins/Experimental/ControlRig/Source/ControlRig/Private/Units/Execution/RigUnit_Hierarchy.h
+++ b/Engine/Plugins/Experimental/ControlRig/Source/ControlRig/Private/Units/Execution/RigUnit_Hierarchy.h
@@ -7,7 +7,7 @@
 #include "RigUnit_Hierarchy.generated.h"
 
 USTRUCT(meta = (Abstract, NodeColor = "0.7 0.05 0.5", Category = "Hierarchy"))
-struct FRigUnit_HierarchyBase : public FRigUnit
+struct CONTROLRIG_API FRigUnit_HierarchyBase : public FRigUnit
 {
 	GENERATED_BODY()
 };
@@ -16,7 +16,7 @@ struct FRigUnit_HierarchyBase : public FRigUnit
  * Returns the item's parent
  */
 USTRUCT(meta=(DisplayName="Get Parent", Keywords="Child,Parent,Root,Up,Top", Varying))
-struct FRigUnit_HierarchyGetParent : public FRigUnit_HierarchyBase
+struct CONTROLRIG_API FRigUnit_HierarchyGetParent : public FRigUnit_HierarchyBase
 {
 	GENERATED_BODY()
 
@@ -48,7 +48,7 @@ struct FRigUnit_HierarchyGetParent : public FRigUnit_HierarchyBase
  * Returns the item's parents
  */
 USTRUCT(meta=(DisplayName="Get Parents", Keywords="Chain,Parents,Hierarchy", Varying))
-struct FRigUnit_HierarchyGetParents : public FRigUnit_HierarchyBase
+struct CONTROLRIG_API FRigUnit_HierarchyGetParents : public FRigUnit_HierarchyBase
 {
 	GENERATED_BODY()
 
@@ -89,7 +89,7 @@ struct FRigUnit_HierarchyGetParents : public FRigUnit_HierarchyBase
  * Returns the item's children
  */
 USTRUCT(meta=(DisplayName="Get Children", Keywords="Chain,Children,Hierarchy", Deprecated = "4.25.0", Varying))
-struct FRigUnit_HierarchyGetChildren : public FRigUnit_HierarchyBase
+struct CONTROLRIG_API FRigUnit_HierarchyGetChildren : public FRigUnit_HierarchyBase
 {
 	GENERATED_BODY()
 
@@ -130,7 +130,7 @@ struct FRigUnit_HierarchyGetChildren : public FRigUnit_HierarchyBase
  * Returns the item's siblings
  */
 USTRUCT(meta=(DisplayName="Get Siblings", Keywords="Chain,Siblings,Hierarchy", Varying))
-struct FRigUnit_HierarchyGetSiblings : public FRigUnit_HierarchyBase
+struct CONTROLRIG_API FRigUnit_HierarchyGetSiblings : public FRigUnit_HierarchyBase
 {
 	GENERATED_BODY()
 
diff --git a/Engine/Plugins/Experimental/ControlRig/Source/ControlRig/Private/Units/Execution/RigUnit_Item.h b/Engine/Plugins/Experimental/ControlRig/Source/ControlRig/Private/Units/Execution/RigUnit_Item.h
--- a/Engine/Plugins/Experimental/ControlRig/Source/ControlRig/Private/Units/Execution/RigUnit_Item.h
+++ b/Engine/Plugins/Experimental/ControlRig/Source/ControlRig/Private/Units/Execution/RigUnit_Item.h
@@ -7,13 +7,13 @@
 #include "RigUnit_Item.generated.h"
 
 USTRUCT(meta = (Abstract, NodeColor = "0.7 0.05 0.5", Category = "Collections"))
-struct FRigUnit_ItemBase : public FRigUnit
+struct CONTROLRIG_API FRigUnit_ItemBase : public FRigUnit
 {
 	GENERATED_BODY()
 };
 
 USTRUCT(meta = (Abstract, NodeColor = "0.7 0.05 0.5", Category = "Collections"))
-struct FRigUnit_ItemBaseMutable : public FRigUnitMutable
+struct CONTROLRIG_API FRigUnit_ItemBaseMutable : public FRigUnitMutable
 {
 	GENERATED_BODY()
 };
@@ -22,7 +22,7 @@ struct FRigUnit_ItemBaseMutable : public FRigUnitMutable
  * Returns true or false if a given item exists
  */
 USTRUCT(meta=(DisplayName="Item Exists", Keywords=""))
-struct FRigUnit_ItemExists : public FRigUnit_ItemBase
+struct CONTROLRIG_API FRigUnit_ItemExists : public FRigUnit_ItemBase
 {
 	GENERATED_BODY()
 
@@ -51,7 +51,7 @@ struct FRigUnit_ItemExists : public FRigUnit_ItemBase
  * Replaces the text within the name of the item
  */
 USTRUCT(meta=(DisplayName="Item Replace", Keywords="Replace,Name"))
-struct FRigUnit_ItemReplace : public FRigUnit_ItemBase
+struct CONTROLRIG_API FRigUnit_ItemReplace : public FRigUnit_ItemBase
 {
 	GENERATED_BODY()
 
diff --git a/Engine/Plugins/Experimental/ControlRig/Source/ControlRig/Private/Units/Execution/RigUnit_SequenceExecution.h b/Engine/Plugins/Experimental/ControlRig/Source/ControlRig/Private/Units/Execution/RigUnit_SequenceExecution.h
--- a/Engine/Plugins/Experimental/ControlRig/Source/ControlRig/Private/Units/Execution/RigUnit_SequenceExecution.h
+++ b/Engine/Plugins/Experimental/ControlRig/Source/ControlRig/Private/Units/Execution/RigUnit_SequenceExecution.h
@@ -10,7 +10,7 @@
  * Allows for a single execution pulse to trigger a series of events in order.
  */
 USTRUCT(meta=(DisplayName="Sequence", Category="Execution", TitleColor="1 0 0", NodeColor="0.1 0.1 0.1"))
-struct FRigUnit_SequenceExecution : public FRigUnit
+struct CONTROLRIG_API FRigUnit_SequenceExecution : public FRigUnit
 {
 	GENERATED_BODY()
 
diff --git a/Engine/Plugins/Experimental/ControlRig/Source/ControlRig/Private/Units/Hierarchy/RigUnit_BoneName.h b/Engine/Plugins/Experimental/ControlRig/Source/ControlRig/Private/Units/Hierarchy/RigUnit_BoneName.h
--- a/Engine/Plugins/Experimental/ControlRig/Source/ControlRig/Private/Units/Hierarchy/RigUnit_BoneName.h
+++ b/Engine/Plugins/Experimental/ControlRig/Source/ControlRig/Private/Units/Hierarchy/RigUnit_BoneName.h
@@ -9,7 +9,7 @@
  * The Item node is used to share a specific item across the graph
  */
 USTRUCT(meta=(DisplayName="Item", Category="Hierarchy", DocumentationPolicy = "Strict", Constant))
-struct FRigUnit_Item : public FRigUnit
+struct CONTROLRIG_API FRigUnit_Item : public FRigUnit
 {
 	GENERATED_BODY()
 
@@ -30,7 +30,7 @@ struct FRigUnit_Item : public FRigUnit
  * BoneName is used to represent a bone name in the graph
  */
 USTRUCT(meta=(DisplayName="Bone Name", Category="Hierarchy", DocumentationPolicy = "Strict", Constant, Deprecated = "4.25"))
-struct FRigUnit_BoneName : public FRigUnit
+struct CONTROLRIG_API FRigUnit_BoneName : public FRigUnit
 {
 	GENERATED_BODY()
 
@@ -51,7 +51,7 @@ struct FRigUnit_BoneName : public FRigUnit
  * SpaceName is used to represent a Space name in the graph
  */
 USTRUCT(meta=(DisplayName="Space Name", Category="Hierarchy", DocumentationPolicy = "Strict", Deprecated = "4.25"))
-struct FRigUnit_SpaceName : public FRigUnit
+struct CONTROLRIG_API FRigUnit_SpaceName : public FRigUnit
 {
 	GENERATED_BODY()
 
@@ -72,7 +72,7 @@ struct FRigUnit_SpaceName : public FRigUnit
  * ControlName is used to represent a Control name in the graph
  */
 USTRUCT(meta=(DisplayName="Control Name", Category="Hierarchy", DocumentationPolicy = "Strict", Deprecated = "4.25"))
-struct FRigUnit_ControlName : public FRigUnit
+struct CONTROLRIG_API FRigUnit_ControlName : public FRigUnit
 {
 	GENERATED_BODY()
 
diff --git a/Engine/Plugins/Experimental/ControlRig/Source/ControlRig/Private/Units/Hierarchy/RigUnit_GetBoneTransform.h b/Engine/Plugins/Experimental/ControlRig/Source/ControlRig/Private/Units/Hierarchy/RigUnit_GetBoneTransform.h
--- a/Engine/Plugins/Experimental/ControlRig/Source/ControlRig/Private/Units/Hierarchy/RigUnit_GetBoneTransform.h
+++ b/Engine/Plugins/Experimental/ControlRig/Source/ControlRig/Private/Units/Hierarchy/RigUnit_GetBoneTransform.h
@@ -9,7 +9,7 @@
  * GetBoneTransform is used to retrieve a single transform from a hierarchy.
  */
 USTRUCT(meta=(DisplayName="Get Transform", Category="Hierarchy", DocumentationPolicy = "Strict", Keywords="GetBoneTransform", Varying, Deprecated = "4.25"))
-struct FRigUnit_GetBoneTransform : public FRigUnit
+struct CONTROLRIG_API FRigUnit_GetBoneTransform : public FRigUnit
 {
 	GENERATED_BODY()
 
diff --git a/Engine/Plugins/Experimental/ControlRig/Source/ControlRig/Private/Units/Hierarchy/RigUnit_GetControlInitialTransform.h b/Engine/Plugins/Experimental/ControlRig/Source/ControlRig/Private/Units/Hierarchy/RigUnit_GetControlInitialTransform.h
--- a/Engine/Plugins/Experimental/ControlRig/Source/ControlRig/Private/Units/Hierarchy/RigUnit_GetControlInitialTransform.h
+++ b/Engine/Plugins/Experimental/ControlRig/Source/ControlRig/Private/Units/Hierarchy/RigUnit_GetControlInitialTransform.h
@@ -9,7 +9,7 @@
  * GetControlTransform is used to retrieve a single transform from a hierarchy.
  */
 USTRUCT(meta=(DisplayName="Get Control Initial Transform", Category="Controls", DocumentationPolicy = "Strict", Keywords="GetControlInitialTransform", Deprecated = "4.25"))
-struct FRigUnit_GetControlInitialTransform : public FRigUnit
+struct CONTROLRIG_API FRigUnit_GetControlInitialTransform : public FRigUnit
 {
 	GENERATED_BODY()
 
diff --git a/Engine/Plugins/Experimental/ControlRig/Source/ControlRig/Private/Units/Hierarchy/RigUnit_GetControlTransform.h b/Engine/Plugins/Experimental/ControlRig/Source/ControlRig/Private/Units/Hierarchy/RigUnit_GetControlTransform.h
--- a/Engine/Plugins/Experimental/ControlRig/Source/ControlRig/Private/Units/Hierarchy/RigUnit_GetControlTransform.h
+++ b/Engine/Plugins/Experimental/ControlRig/Source/ControlRig/Private/Units/Hierarchy/RigUnit_GetControlTransform.h
@@ -9,7 +9,7 @@
  * GetControlBool is used to retrieve a single Bool from a hierarchy.
  */
 USTRUCT(meta=(DisplayName="Get Control Bool", Category="Controls", DocumentationPolicy = "Strict", Keywords="GetControlBool", PrototypeName = "GetControlValue", Varying))
-struct FRigUnit_GetControlBool : public FRigUnit
+struct CONTROLRIG_API FRigUnit_GetControlBool : public FRigUnit
 {
 	GENERATED_BODY()
 
@@ -40,7 +40,7 @@ struct FRigUnit_GetControlBool : public FRigUnit
  * GetControlFloat is used to retrieve a single Float from a hierarchy.
  */
 USTRUCT(meta=(DisplayName="Get Control Float", Category="Controls", DocumentationPolicy = "Strict", Keywords="GetControlFloat", PrototypeName = "GetControlValue", Varying))
-struct FRigUnit_GetControlFloat : public FRigUnit
+struct CONTROLRIG_API FRigUnit_GetControlFloat : public FRigUnit
 {
 	GENERATED_BODY()
 
@@ -81,7 +81,7 @@ struct FRigUnit_GetControlFloat : public FRigUnit
  * GetControlFloat is used to retrieve a single Integer from a hierarchy.
  */
 USTRUCT(meta=(DisplayName="Get Control Integer", Category="Controls", DocumentationPolicy = "Strict", Keywords="GetControlInteger,GetControlEnum", PrototypeName = "GetControlValue", Varying))
-struct FRigUnit_GetControlInteger : public FRigUnit
+struct CONTROLRIG_API FRigUnit_GetControlInteger : public FRigUnit
 {
 	GENERATED_BODY()
 
@@ -122,7 +122,7 @@ struct FRigUnit_GetControlInteger : public FRigUnit
  * GetControlVector2D is used to retrieve a single Vector2D from a hierarchy.
  */
 USTRUCT(meta=(DisplayName="Get Control Vector2D", Category="Controls", DocumentationPolicy = "Strict", Keywords="GetControlVector2D", PrototypeName = "GetControlValue", Varying))
-struct FRigUnit_GetControlVector2D : public FRigUnit
+struct CONTROLRIG_API FRigUnit_GetControlVector2D : public FRigUnit
 {
 	GENERATED_BODY()
 
@@ -163,7 +163,7 @@ struct FRigUnit_GetControlVector2D : public FRigUnit
  * GetControlVector is used to retrieve a single Vector from a hierarchy.
  */
 USTRUCT(meta=(DisplayName="Get Control Vector", Category="Controls", DocumentationPolicy = "Strict", Keywords="GetControlVector", PrototypeName = "GetControlValue", Varying))
-struct FRigUnit_GetControlVector : public FRigUnit
+struct CONTROLRIG_API FRigUnit_GetControlVector : public FRigUnit
 {
 	GENERATED_BODY()
 
@@ -212,7 +212,7 @@ struct FRigUnit_GetControlVector : public FRigUnit
  * GetControlRotator is used to retrieve a single Rotator from a hierarchy.
  */
 USTRUCT(meta=(DisplayName="Get Control Rotator", Category="Controls", DocumentationPolicy = "Strict", Keywords="GetControlRotator", PrototypeName = "GetControlValue", Varying))
-struct FRigUnit_GetControlRotator : public FRigUnit
+struct CONTROLRIG_API FRigUnit_GetControlRotator : public FRigUnit
 {
 	GENERATED_BODY()
 
@@ -262,7 +262,7 @@ struct FRigUnit_GetControlRotator : public FRigUnit
  * GetControlTransform is used to retrieve a single transform from a hierarchy.
  */
 USTRUCT(meta=(DisplayName="Get Control Transform", Category="Controls", DocumentationPolicy = "Strict", Keywords="GetControlTransform", PrototypeName = "GetControlValue", Varying))
-struct FRigUnit_GetControlTransform : public FRigUnit
+struct CONTROLRIG_API FRigUnit_GetControlTransform : public FRigUnit
 {
 	GENERATED_BODY()
 
diff --git a/Engine/Plugins/Experimental/ControlRig/Source/ControlRig/Private/Units/Hierarchy/RigUnit_GetCurveValue.h b/Engine/Plugins/Experimental/ControlRig/Source/ControlRig/Private/Units/Hierarchy/RigUnit_GetCurveValue.h
--- a/Engine/Plugins/Experimental/ControlRig/Source/ControlRig/Private/Units/Hierarchy/RigUnit_GetCurveValue.h
+++ b/Engine/Plugins/Experimental/ControlRig/Source/ControlRig/Private/Units/Hierarchy/RigUnit_GetCurveValue.h
@@ -9,7 +9,7 @@
  * GetCurveValue is used to retrieve a single float from a Curve
  */
 USTRUCT(meta=(DisplayName="Get Curve Value", Category="Curve", Keywords="GetCurveValue,float", Varying))
-struct FRigUnit_GetCurveValue : public FRigUnit
+struct CONTROLRIG_API FRigUnit_GetCurveValue : public FRigUnit
 {
 	GENERATED_BODY()
 
diff --git a/Engine/Plugins/Experimental/ControlRig/Source/ControlRig/Private/Units/Hierarchy/RigUnit_GetInitialBoneTransform.h b/Engine/Plugins/Experimental/ControlRig/Source/ControlRig/Private/Units/Hierarchy/RigUnit_GetInitialBoneTransform.h
--- a/Engine/Plugins/Experimental/ControlRig/Source/ControlRig/Private/Units/Hierarchy/RigUnit_GetInitialBoneTransform.h
+++ b/Engine/Plugins/Experimental/ControlRig/Source/ControlRig/Private/Units/Hierarchy/RigUnit_GetInitialBoneTransform.h
@@ -9,7 +9,7 @@
  * GetInitialBoneTransform is used to retrieve a single transform from a hierarchy.
  */
 USTRUCT(meta=(DisplayName="Get Initial Transform", Category="Hierarchy", DocumentationPolicy = "Strict", Keywords="GetInitialBoneTransform", Varying, Deprecated = "4.25"))
-struct FRigUnit_GetInitialBoneTransform : public FRigUnit
+struct CONTROLRIG_API FRigUnit_GetInitialBoneTransform : public FRigUnit
 {
 	GENERATED_BODY()
 
diff --git a/Engine/Plugins/Experimental/ControlRig/Source/ControlRig/Private/Units/Hierarchy/RigUnit_GetRelativeBoneTransform.h b/Engine/Plugins/Experimental/ControlRig/Source/ControlRig/Private/Units/Hierarchy/RigUnit_GetRelativeBoneTransform.h
--- a/Engine/Plugins/Experimental/ControlRig/Source/ControlRig/Private/Units/Hierarchy/RigUnit_GetRelativeBoneTransform.h
+++ b/Engine/Plugins/Experimental/ControlRig/Source/ControlRig/Private/Units/Hierarchy/RigUnit_GetRelativeBoneTransform.h
@@ -9,7 +9,7 @@
  * GetBoneTransform is used to retrieve a single transform from a hierarchy.
  */
 USTRUCT(meta=(DisplayName="Get Relative Transform", Category="Hierarchy", DocumentationPolicy = "Strict", Keywords = "GetRelativeBoneTransform", Deprecated = "4.25", Varying))
-struct FRigUnit_GetRelativeBoneTransform : public FRigUnit
+struct CONTROLRIG_API FRigUnit_GetRelativeBoneTransform : public FRigUnit
 {
 	GENERATED_BODY()
 
diff --git a/Engine/Plugins/Experimental/ControlRig/Source/ControlRig/Private/Units/Hierarchy/RigUnit_GetRelativeTransform.h b/Engine/Plugins/Experimental/ControlRig/Source/ControlRig/Private/Units/Hierarchy/RigUnit_GetRelativeTransform.h
--- a/Engine/Plugins/Experimental/ControlRig/Source/ControlRig/Private/Units/Hierarchy/RigUnit_GetRelativeTransform.h
+++ b/Engine/Plugins/Experimental/ControlRig/Source/ControlRig/Private/Units/Hierarchy/RigUnit_GetRelativeTransform.h
@@ -9,7 +9,7 @@
  * GetRelativeTransform is used to retrieve a single transform from a hierarchy in the space of another transform
  */
 USTRUCT(meta=(DisplayName="Get Relative Transform", Category="Hierarchy", DocumentationPolicy = "Strict", Keywords = "Offset,Local", Varying))
-struct FRigUnit_GetRelativeTransformForItem : public FRigUnit
+struct CONTROLRIG_API FRigUnit_GetRelativeTransformForItem : public FRigUnit
 {
 	GENERATED_BODY()
 
diff --git a/Engine/Plugins/Experimental/ControlRig/Source/ControlRig/Private/Units/Hierarchy/RigUnit_GetSpaceTransform.h b/Engine/Plugins/Experimental/ControlRig/Source/ControlRig/Private/Units/Hierarchy/RigUnit_GetSpaceTransform.h
--- a/Engine/Plugins/Experimental/ControlRig/Source/ControlRig/Private/Units/Hierarchy/RigUnit_GetSpaceTransform.h
+++ b/Engine/Plugins/Experimental/ControlRig/Source/ControlRig/Private/Units/Hierarchy/RigUnit_GetSpaceTransform.h
@@ -9,7 +9,7 @@
  * GetSpaceTransform is used to retrieve a single transform from a hierarchy.
  */
 USTRUCT(meta=(DisplayName="Get Space Transform", Category="Spaces", DocumentationPolicy = "Strict", Keywords="GetSpaceTransform", Varying, Deprecated = "4.25"))
-struct FRigUnit_GetSpaceTransform : public FRigUnit
+struct CONTROLRIG_API FRigUnit_GetSpaceTransform : public FRigUnit
 {
 	GENERATED_BODY()
 
diff --git a/Engine/Plugins/Experimental/ControlRig/Source/ControlRig/Private/Units/Hierarchy/RigUnit_GetTransform.h b/Engine/Plugins/Experimental/ControlRig/Source/ControlRig/Private/Units/Hierarchy/RigUnit_GetTransform.h
--- a/Engine/Plugins/Experimental/ControlRig/Source/ControlRig/Private/Units/Hierarchy/RigUnit_GetTransform.h
+++ b/Engine/Plugins/Experimental/ControlRig/Source/ControlRig/Private/Units/Hierarchy/RigUnit_GetTransform.h
@@ -9,7 +9,7 @@
  * GetTransform is used to retrieve a single transform from a hierarchy.
  */
 USTRUCT(meta=(DisplayName="Get Transform", Category="Hierarchy", DocumentationPolicy = "Strict", Keywords="GetBoneTransform,GetControlTransform,GetInitialTransform,GetSpaceTransform,GetTransform", Varying))
-struct FRigUnit_GetTransform : public FRigUnit
+struct CONTROLRIG_API FRigUnit_GetTransform : public FRigUnit
 {
 	GENERATED_BODY()
 
diff --git a/Engine/Plugins/Experimental/ControlRig/Source/ControlRig/Private/Units/Hierarchy/RigUnit_OffsetTransform.h b/Engine/Plugins/Experimental/ControlRig/Source/ControlRig/Private/Units/Hierarchy/RigUnit_OffsetTransform.h
--- a/Engine/Plugins/Experimental/ControlRig/Source/ControlRig/Private/Units/Hierarchy/RigUnit_OffsetTransform.h
+++ b/Engine/Plugins/Experimental/ControlRig/Source/ControlRig/Private/Units/Hierarchy/RigUnit_OffsetTransform.h
@@ -9,7 +9,7 @@
  * Offset Transform is used to add an offset to an existing transform in the hierarchy. The offset is post multiplied.
  */
 USTRUCT(meta=(DisplayName="Offset Transform", Category="Hierarchy", DocumentationPolicy = "Strict", Keywords = "Offset,Relative,AddBoneTransform", Varying))
-struct FRigUnit_OffsetTransformForItem : public FRigUnitMutable
+struct CONTROLRIG_API FRigUnit_OffsetTransformForItem : public FRigUnitMutable
 {
 	GENERATED_BODY()
 
diff --git a/Engine/Plugins/Experimental/ControlRig/Source/ControlRig/Private/Units/Hierarchy/RigUnit_ParentSwitchConstraint.h b/Engine/Plugins/Experimental/ControlRig/Source/ControlRig/Private/Units/Hierarchy/RigUnit_ParentSwitchConstraint.h
--- a/Engine/Plugins/Experimental/ControlRig/Source/ControlRig/Private/Units/Hierarchy/RigUnit_ParentSwitchConstraint.h
+++ b/Engine/Plugins/Experimental/ControlRig/Source/ControlRig/Private/Units/Hierarchy/RigUnit_ParentSwitchConstraint.h
@@ -11,7 +11,7 @@
  * and allowing to switch between the parent in question.
  */
 USTRUCT(meta=(DisplayName="Parent Switch Constraint", Category="Hierarchy", DocumentationPolicy="Strict", Keywords = "SpaceSwitch"))
-struct FRigUnit_ParentSwitchConstraint : public FRigUnitMutable
+struct CONTROLRIG_API FRigUnit_ParentSwitchConstraint : public FRigUnitMutable
 {
 	GENERATED_BODY()
 
diff --git a/Engine/Plugins/Experimental/ControlRig/Source/ControlRig/Private/Units/Hierarchy/RigUnit_ProjectTransformToNewParent.h b/Engine/Plugins/Experimental/ControlRig/Source/ControlRig/Private/Units/Hierarchy/RigUnit_ProjectTransformToNewParent.h
--- a/Engine/Plugins/Experimental/ControlRig/Source/ControlRig/Private/Units/Hierarchy/RigUnit_ProjectTransformToNewParent.h
+++ b/Engine/Plugins/Experimental/ControlRig/Source/ControlRig/Private/Units/Hierarchy/RigUnit_ProjectTransformToNewParent.h
@@ -14,7 +14,7 @@
  * and set the new parent to the head as well (not-initial == current).
  */
 USTRUCT(meta=(DisplayName="Project to new Parent", Category="Hierarchy", DocumentationPolicy = "Strict", Keywords = "ProjectTransformToNewParent,Relative,Reparent,Offset", Varying))
-struct FRigUnit_ProjectTransformToNewParent : public FRigUnit
+struct CONTROLRIG_API FRigUnit_ProjectTransformToNewParent : public FRigUnit
 {
 	GENERATED_BODY()
 
diff --git a/Engine/Plugins/Experimental/ControlRig/Source/ControlRig/Private/Units/Hierarchy/RigUnit_PropagateTransform.h b/Engine/Plugins/Experimental/ControlRig/Source/ControlRig/Private/Units/Hierarchy/RigUnit_PropagateTransform.h
--- a/Engine/Plugins/Experimental/ControlRig/Source/ControlRig/Private/Units/Hierarchy/RigUnit_PropagateTransform.h
+++ b/Engine/Plugins/Experimental/ControlRig/Source/ControlRig/Private/Units/Hierarchy/RigUnit_PropagateTransform.h
@@ -10,7 +10,7 @@
  * from its local - as well as propagating that change onto the children.
  */
 USTRUCT(meta=(DisplayName="Propagate Transform", Category="Hierarchy", DocumentationPolicy = "Strict", Keywords = "PropagateToChildren,RecomputeGlobal,RecalculateGlobal", Varying))
-struct FRigUnit_PropagateTransform : public FRigUnitMutable
+struct CONTROLRIG_API FRigUnit_PropagateTransform : public FRigUnitMutable
 {
 	GENERATED_BODY()
 
diff --git a/Engine/Plugins/Experimental/ControlRig/Source/ControlRig/Private/Units/Hierarchy/RigUnit_SendEvent.h b/Engine/Plugins/Experimental/ControlRig/Source/ControlRig/Private/Units/Hierarchy/RigUnit_SendEvent.h
--- a/Engine/Plugins/Experimental/ControlRig/Source/ControlRig/Private/Units/Hierarchy/RigUnit_SendEvent.h
+++ b/Engine/Plugins/Experimental/ControlRig/Source/ControlRig/Private/Units/Hierarchy/RigUnit_SendEvent.h
@@ -9,7 +9,7 @@
  * SendEvent is used to notify the engine / editor of a change that happend within the Control Rig.
  */
 USTRUCT(meta=(DisplayName="Send Event", Category="Hierarchy", DocumentationPolicy="Strict", Keywords = "SendEvent", PrototypeName = "Event,Notify,Notification"))
-struct FRigUnit_SendEvent : public FRigUnitMutable
+struct CONTROLRIG_API FRigUnit_SendEvent : public FRigUnitMutable
 {
 	GENERATED_BODY()
 
diff --git a/Engine/Plugins/Experimental/ControlRig/Source/ControlRig/Private/Units/Hierarchy/RigUnit_SetBoneInitialTransform.h b/Engine/Plugins/Experimental/ControlRig/Source/ControlRig/Private/Units/Hierarchy/RigUnit_SetBoneInitialTransform.h
--- a/Engine/Plugins/Experimental/ControlRig/Source/ControlRig/Private/Units/Hierarchy/RigUnit_SetBoneInitialTransform.h
+++ b/Engine/Plugins/Experimental/ControlRig/Source/ControlRig/Private/Units/Hierarchy/RigUnit_SetBoneInitialTransform.h
@@ -10,7 +10,7 @@
  * SetBoneInitialTransform is used to perform a change in the hierarchy by setting a single bone's transform.
  */
 USTRUCT(meta=(DisplayName="Set Initial Bone Transform", Category="Setup", DocumentationPolicy="Strict", Keywords = "SetBoneInitialTransform", Deprecated="4.25"))
-struct FRigUnit_SetBoneInitialTransform : public FRigUnitMutable
+struct CONTROLRIG_API FRigUnit_SetBoneInitialTransform : public FRigUnitMutable
 {
 	GENERATED_BODY()
 
diff --git a/Engine/Plugins/Experimental/ControlRig/Source/ControlRig/Private/Units/Hierarchy/RigUnit_SetBoneRotation.h b/Engine/Plugins/Experimental/ControlRig/Source/ControlRig/Private/Units/Hierarchy/RigUnit_SetBoneRotation.h
--- a/Engine/Plugins/Experimental/ControlRig/Source/ControlRig/Private/Units/Hierarchy/RigUnit_SetBoneRotation.h
+++ b/Engine/Plugins/Experimental/ControlRig/Source/ControlRig/Private/Units/Hierarchy/RigUnit_SetBoneRotation.h
@@ -10,7 +10,7 @@
  * SetBoneRotation is used to perform a change in the hierarchy by setting a single bone's rotation.
  */
 USTRUCT(meta=(DisplayName="Set Rotation", Category="Hierarchy", DocumentationPolicy="Strict", Keywords = "SetBoneRotation", Deprecated="4.25"))
-struct FRigUnit_SetBoneRotation : public FRigUnitMutable
+struct CONTROLRIG_API FRigUnit_SetBoneRotation : public FRigUnitMutable
 {
 	GENERATED_BODY()
 
diff --git a/Engine/Plugins/Experimental/ControlRig/Source/ControlRig/Private/Units/Hierarchy/RigUnit_SetBoneTransform.h b/Engine/Plugins/Experimental/ControlRig/Source/ControlRig/Private/Units/Hierarchy/RigUnit_SetBoneTransform.h
--- a/Engine/Plugins/Experimental/ControlRig/Source/ControlRig/Private/Units/Hierarchy/RigUnit_SetBoneTransform.h
+++ b/Engine/Plugins/Experimental/ControlRig/Source/ControlRig/Private/Units/Hierarchy/RigUnit_SetBoneTransform.h
@@ -10,7 +10,7 @@
  * SetBoneTransform is used to perform a change in the hierarchy by setting a single bone's transform.
  */
 USTRUCT(meta=(DisplayName="Set Transform", Category="Hierarchy", DocumentationPolicy="Strict", Keywords = "SetBoneTransform", Deprecated="4.25"))
-struct FRigUnit_SetBoneTransform : public FRigUnitMutable
+struct CONTROLRIG_API FRigUnit_SetBoneTransform : public FRigUnitMutable
 {
 	GENERATED_BODY()
 
diff --git a/Engine/Plugins/Experimental/ControlRig/Source/ControlRig/Private/Units/Hierarchy/RigUnit_SetBoneTranslation.h b/Engine/Plugins/Experimental/ControlRig/Source/ControlRig/Private/Units/Hierarchy/RigUnit_SetBoneTranslation.h
--- a/Engine/Plugins/Experimental/ControlRig/Source/ControlRig/Private/Units/Hierarchy/RigUnit_SetBoneTranslation.h
+++ b/Engine/Plugins/Experimental/ControlRig/Source/ControlRig/Private/Units/Hierarchy/RigUnit_SetBoneTranslation.h
@@ -10,7 +10,7 @@
  * SetBoneTranslation is used to perform a change in the hierarchy by setting a single bone's Translation.
  */
 USTRUCT(meta=(DisplayName="Set Translation", Category="Hierarchy", DocumentationPolicy="Strict", Keywords = "SetBoneTranslation,SetPosition,SetLocation,SetBonePosition,SetBoneLocation", Deprecated="4.25"))
-struct FRigUnit_SetBoneTranslation : public FRigUnitMutable
+struct CONTROLRIG_API FRigUnit_SetBoneTranslation : public FRigUnitMutable
 {
 	GENERATED_BODY()
 
diff --git a/Engine/Plugins/Experimental/ControlRig/Source/ControlRig/Private/Units/Hierarchy/RigUnit_SetControlColor.h b/Engine/Plugins/Experimental/ControlRig/Source/ControlRig/Private/Units/Hierarchy/RigUnit_SetControlColor.h
--- a/Engine/Plugins/Experimental/ControlRig/Source/ControlRig/Private/Units/Hierarchy/RigUnit_SetControlColor.h
+++ b/Engine/Plugins/Experimental/ControlRig/Source/ControlRig/Private/Units/Hierarchy/RigUnit_SetControlColor.h
@@ -9,7 +9,7 @@
  * SetControlColor is used to change the gizmo's color on a control at runtime
  */
 USTRUCT(meta=(DisplayName="Set Control Color", Category="Hierarchy", DocumentationPolicy="Strict", Keywords = "SetControlColor,SetGizmoColor", PrototypeName = "SetControlColor"))
-struct FRigUnit_SetControlColor : public FRigUnitMutable
+struct CONTROLRIG_API FRigUnit_SetControlColor : public FRigUnitMutable
 {
 	GENERATED_BODY()
 
diff --git a/Engine/Plugins/Experimental/ControlRig/Source/ControlRig/Private/Units/Hierarchy/RigUnit_SetControlOffset.h b/Engine/Plugins/Experimental/ControlRig/Source/ControlRig/Private/Units/Hierarchy/RigUnit_SetControlOffset.h
--- a/Engine/Plugins/Experimental/ControlRig/Source/ControlRig/Private/Units/Hierarchy/RigUnit_SetControlOffset.h
+++ b/Engine/Plugins/Experimental/ControlRig/Source/ControlRig/Private/Units/Hierarchy/RigUnit_SetControlOffset.h
@@ -9,7 +9,7 @@
  * SetControlOffset is used to perform a change in the hierarchy by setting a single control's transform.
  */
 USTRUCT(meta=(DisplayName="Set Control Offset", Category="Setup", DocumentationPolicy="Strict", Keywords = "SetControlOffset,Initial,InitialTransform,SetInitialTransform,SetInitialControlTransform", PrototypeName = "SetControlValue"))
-struct FRigUnit_SetControlOffset : public FRigUnitMutable
+struct CONTROLRIG_API FRigUnit_SetControlOffset : public FRigUnitMutable
 {
 	GENERATED_BODY()
 
diff --git a/Engine/Plugins/Experimental/ControlRig/Source/ControlRig/Private/Units/Hierarchy/RigUnit_SetControlTransform.h b/Engine/Plugins/Experimental/ControlRig/Source/ControlRig/Private/Units/Hierarchy/RigUnit_SetControlTransform.h
--- a/Engine/Plugins/Experimental/ControlRig/Source/ControlRig/Private/Units/Hierarchy/RigUnit_SetControlTransform.h
+++ b/Engine/Plugins/Experimental/ControlRig/Source/ControlRig/Private/Units/Hierarchy/RigUnit_SetControlTransform.h
@@ -9,7 +9,7 @@
  * SetControlBool is used to perform a change in the hierarchy by setting a single control's bool value.
  */
 USTRUCT(meta=(DisplayName="Set Control Bool", Category="Hierarchy", DocumentationPolicy="Strict", Keywords = "SetControlBool,SetGizmoBool", PrototypeName = "SetControlValue"))
-struct FRigUnit_SetControlBool : public FRigUnitMutable
+struct CONTROLRIG_API FRigUnit_SetControlBool : public FRigUnitMutable
 {
 	GENERATED_BODY()
 
@@ -39,7 +39,7 @@ struct FRigUnit_SetControlBool : public FRigUnitMutable
 };
 
 USTRUCT()
-struct FRigUnit_SetMultiControlBool_Entry
+struct CONTROLRIG_API FRigUnit_SetMultiControlBool_Entry
 {
 	GENERATED_BODY()
 
@@ -62,7 +62,7 @@ struct FRigUnit_SetMultiControlBool_Entry
  * SetMultiControlBool is used to perform a change in the hierarchy by setting multiple controls' bool value.
  */
 USTRUCT(meta = (DisplayName = "Set Multiple Controls Bool", Category = "Hierarchy", DocumentationPolicy = "Strict", Keywords = "SetMultipleControlsBool,SetControlBool,SetGizmoBool", PrototypeName = "SetMultiControlValue"))
-struct FRigUnit_SetMultiControlBool : public FRigUnitMutable
+struct CONTROLRIG_API FRigUnit_SetMultiControlBool : public FRigUnitMutable
 {
 	GENERATED_BODY()
 
@@ -89,7 +89,7 @@ struct FRigUnit_SetMultiControlBool : public FRigUnitMutable
  * SetControlFloat is used to perform a change in the hierarchy by setting a single control's float value.
  */
 USTRUCT(meta=(DisplayName="Set Control Float", Category="Hierarchy", DocumentationPolicy="Strict", Keywords = "SetControlFloat,SetGizmoFloat", PrototypeName = "SetControlValue"))
-struct FRigUnit_SetControlFloat : public FRigUnitMutable
+struct CONTROLRIG_API FRigUnit_SetControlFloat : public FRigUnitMutable
 {
 	GENERATED_BODY()
 
@@ -126,7 +126,7 @@ struct FRigUnit_SetControlFloat : public FRigUnitMutable
 };
 
 USTRUCT()
-struct FRigUnit_SetMultiControlFloat_Entry
+struct CONTROLRIG_API FRigUnit_SetMultiControlFloat_Entry
 { 
 	GENERATED_BODY()
 
@@ -150,7 +150,7 @@ struct FRigUnit_SetMultiControlFloat_Entry
  * SetMultiControlFloat is used to perform a change in the hierarchy by setting multiple controls' float value.
  */
 USTRUCT(meta = (DisplayName = "Set Multiple Controls Float", Category = "Hierarchy", DocumentationPolicy = "Strict", Keywords = "SetMultipleControlsFloat,SetControlFloat,SetGizmoFloat", PrototypeName = "SetMultiControlValue"))
-struct FRigUnit_SetMultiControlFloat : public FRigUnitMutable
+struct CONTROLRIG_API FRigUnit_SetMultiControlFloat : public FRigUnitMutable
 {
 	GENERATED_BODY()
 
@@ -186,7 +186,7 @@ struct FRigUnit_SetMultiControlFloat : public FRigUnitMutable
  * SetControlInteger is used to perform a change in the hierarchy by setting a single control's int32 value.
  */
 USTRUCT(meta=(DisplayName="Set Control Integer", Category="Hierarchy", DocumentationPolicy="Strict", Keywords = "SetControlInteger,SetGizmoInteger", PrototypeName = "SetControlValue"))
-struct FRigUnit_SetControlInteger : public FRigUnitMutable
+struct CONTROLRIG_API FRigUnit_SetControlInteger : public FRigUnitMutable
 {
 	GENERATED_BODY()
 
@@ -223,7 +223,7 @@ struct FRigUnit_SetControlInteger : public FRigUnitMutable
 };
 
 USTRUCT()
-struct FRigUnit_SetMultiControlInteger_Entry
+struct CONTROLRIG_API FRigUnit_SetMultiControlInteger_Entry
 {
 	GENERATED_BODY()
 
@@ -247,7 +247,7 @@ struct FRigUnit_SetMultiControlInteger_Entry
  * SetMultiControlInteger is used to perform a change in the hierarchy by setting multiple controls' integer value.
  */
 USTRUCT(meta = (DisplayName = "Set Multiple Controls Integer", Category = "Hierarchy", DocumentationPolicy = "Strict", Keywords = "SetMultipleControlsInteger,SetControlInteger,SetGizmoInteger", PrototypeName = "SetMultiControlValue"))
-struct FRigUnit_SetMultiControlInteger : public FRigUnitMutable
+struct CONTROLRIG_API FRigUnit_SetMultiControlInteger : public FRigUnitMutable
 {
 	GENERATED_BODY()
 
@@ -281,7 +281,7 @@ struct FRigUnit_SetMultiControlInteger : public FRigUnitMutable
  * SetControlVector2D is used to perform a change in the hierarchy by setting a single control's Vector2D value.
  */
 USTRUCT(meta=(DisplayName="Set Control Vector2D", Category="Hierarchy", DocumentationPolicy="Strict", Keywords = "SetControlVector2D,SetGizmoVector2D", PrototypeName = "SetControlValue"))
-struct FRigUnit_SetControlVector2D : public FRigUnitMutable
+struct CONTROLRIG_API FRigUnit_SetControlVector2D : public FRigUnitMutable
 {
 	GENERATED_BODY()
 
@@ -318,7 +318,7 @@ struct FRigUnit_SetControlVector2D : public FRigUnitMutable
 };
 
 USTRUCT()
-struct FRigUnit_SetMultiControlVector2D_Entry
+struct CONTROLRIG_API FRigUnit_SetMultiControlVector2D_Entry
 {
 	GENERATED_BODY()
 
@@ -343,7 +343,7 @@ struct FRigUnit_SetMultiControlVector2D_Entry
  * SetMultiControlVector2D is used to perform a change in the hierarchy by setting multiple controls' vector2D value.
  */
 USTRUCT(meta = (DisplayName = "Set Multiple Controls Vector2D", Category = "Hierarchy", DocumentationPolicy = "Strict", Keywords = "SetMultipleControlsVector2D,SetControlVector2D,SetGizmoVector2D", PrototypeName = "SetMultiControlValue"))
-struct FRigUnit_SetMultiControlVector2D : public FRigUnitMutable
+struct CONTROLRIG_API FRigUnit_SetMultiControlVector2D : public FRigUnitMutable
 {
 	GENERATED_BODY()
 
@@ -377,7 +377,7 @@ struct FRigUnit_SetMultiControlVector2D : public FRigUnitMutable
  * SetControlVector is used to perform a change in the hierarchy by setting a single control's Vector value.
  */
 USTRUCT(meta=(DisplayName="Set Control Vector", Category="Hierarchy", DocumentationPolicy="Strict", Keywords = "SetControlVector,SetGizmoVector", PrototypeName = "SetControlValue", Deprecated = "4.25"))
-struct FRigUnit_SetControlVector : public FRigUnitMutable
+struct CONTROLRIG_API FRigUnit_SetControlVector : public FRigUnitMutable
 {
 	GENERATED_BODY()
 
@@ -425,7 +425,7 @@ struct FRigUnit_SetControlVector : public FRigUnitMutable
  * SetControlRotator is used to perform a change in the hierarchy by setting a single control's Rotator value.
  */
 USTRUCT(meta=(DisplayName="Set Control Rotator", Category="Hierarchy", DocumentationPolicy="Strict", Keywords = "SetControlRotator,SetGizmoRotator", PrototypeName = "SetControlValue"))
-struct FRigUnit_SetControlRotator : public FRigUnitMutable
+struct CONTROLRIG_API FRigUnit_SetControlRotator : public FRigUnitMutable
 {
 	GENERATED_BODY()
 
@@ -470,7 +470,7 @@ struct FRigUnit_SetControlRotator : public FRigUnitMutable
 };
 
 USTRUCT()
-struct FRigUnit_SetMultiControlRotator_Entry
+struct CONTROLRIG_API FRigUnit_SetMultiControlRotator_Entry
 {
 	GENERATED_BODY()
 
@@ -504,7 +504,7 @@ struct FRigUnit_SetMultiControlRotator_Entry
  * SetMultiControlRotator is used to perform a change in the hierarchy by setting multiple controls' rotator value.
  */
 USTRUCT(meta = (DisplayName = "Set Multiple Controls Rotator", Category = "Hierarchy", DocumentationPolicy = "Strict", Keywords = "SetMultipleControlsRotator,SetControlRotator,SetGizmoRotator", PrototypeName = "SetMultiControlValue"))
-struct FRigUnit_SetMultiControlRotator : public FRigUnitMutable
+struct CONTROLRIG_API FRigUnit_SetMultiControlRotator : public FRigUnitMutable
 {
 	GENERATED_BODY()
 
@@ -538,7 +538,7 @@ struct FRigUnit_SetMultiControlRotator : public FRigUnitMutable
  * SetControlTransform is used to perform a change in the hierarchy by setting a single control's transform.
  */
 USTRUCT(meta=(DisplayName="Set Control Transform", Category="Hierarchy", DocumentationPolicy="Strict", Keywords = "SetControlTransform,SetGizmoTransform", PrototypeName = "SetControlValue", Deprecated = "4.25"))
-struct FRigUnit_SetControlTransform : public FRigUnitMutable
+struct CONTROLRIG_API FRigUnit_SetControlTransform : public FRigUnitMutable
 {
 	GENERATED_BODY()
 
diff --git a/Engine/Plugins/Experimental/ControlRig/Source/ControlRig/Private/Units/Hierarchy/RigUnit_SetControlVisibility.h b/Engine/Plugins/Experimental/ControlRig/Source/ControlRig/Private/Units/Hierarchy/RigUnit_SetControlVisibility.h
--- a/Engine/Plugins/Experimental/ControlRig/Source/ControlRig/Private/Units/Hierarchy/RigUnit_SetControlVisibility.h
+++ b/Engine/Plugins/Experimental/ControlRig/Source/ControlRig/Private/Units/Hierarchy/RigUnit_SetControlVisibility.h
@@ -9,7 +9,7 @@
  * SetControlVisibility is used to change the gizmo visibility on a control at runtime
  */
 USTRUCT(meta=(DisplayName="Set Control Visibility", Category="Hierarchy", DocumentationPolicy="Strict", Keywords = "SetControlVisibility,Visibility,Hide,Show,Hidden,Visible,SetGizmoVisibility", PrototypeName = "SetControlVisibility"))
-struct FRigUnit_SetControlVisibility : public FRigUnitMutable
+struct CONTROLRIG_API FRigUnit_SetControlVisibility : public FRigUnitMutable
 {
 	GENERATED_BODY()
 
diff --git a/Engine/Plugins/Experimental/ControlRig/Source/ControlRig/Private/Units/Hierarchy/RigUnit_SetRelativeBoneTransform.h b/Engine/Plugins/Experimental/ControlRig/Source/ControlRig/Private/Units/Hierarchy/RigUnit_SetRelativeBoneTransform.h
--- a/Engine/Plugins/Experimental/ControlRig/Source/ControlRig/Private/Units/Hierarchy/RigUnit_SetRelativeBoneTransform.h
+++ b/Engine/Plugins/Experimental/ControlRig/Source/ControlRig/Private/Units/Hierarchy/RigUnit_SetRelativeBoneTransform.h
@@ -10,7 +10,7 @@
  * SetRelativeBoneTransform is used to perform a change in the hierarchy by setting a single bone's transform.
  */
 USTRUCT(meta=(DisplayName="Set Relative Transform", Category="Hierarchy", DocumentationPolicy="Strict", Keywords = "SetRelativeBoneTransform", Deprecated = "4.25"))
-struct FRigUnit_SetRelativeBoneTransform : public FRigUnitMutable
+struct CONTROLRIG_API FRigUnit_SetRelativeBoneTransform : public FRigUnitMutable
 {
 	GENERATED_BODY()
 
diff --git a/Engine/Plugins/Experimental/ControlRig/Source/ControlRig/Private/Units/Hierarchy/RigUnit_SetRelativeTransform.h b/Engine/Plugins/Experimental/ControlRig/Source/ControlRig/Private/Units/Hierarchy/RigUnit_SetRelativeTransform.h
--- a/Engine/Plugins/Experimental/ControlRig/Source/ControlRig/Private/Units/Hierarchy/RigUnit_SetRelativeTransform.h
+++ b/Engine/Plugins/Experimental/ControlRig/Source/ControlRig/Private/Units/Hierarchy/RigUnit_SetRelativeTransform.h
@@ -9,7 +9,7 @@
  * SetRelativeTransform is used to set a single transform from a hierarchy in the space of another item
  */
 USTRUCT(meta=(DisplayName="Set Relative Transform", Category="Hierarchy", DocumentationPolicy = "Strict", Keywords = "Offset,Local", Varying))
-struct FRigUnit_SetRelativeTransformForItem : public FRigUnitMutable
+struct CONTROLRIG_API FRigUnit_SetRelativeTransformForItem : public FRigUnitMutable
 {
 	GENERATED_BODY()
 
diff --git a/Engine/Plugins/Experimental/ControlRig/Source/ControlRig/Private/Units/Hierarchy/RigUnit_SetSpaceInitialTransform.h b/Engine/Plugins/Experimental/ControlRig/Source/ControlRig/Private/Units/Hierarchy/RigUnit_SetSpaceInitialTransform.h
--- a/Engine/Plugins/Experimental/ControlRig/Source/ControlRig/Private/Units/Hierarchy/RigUnit_SetSpaceInitialTransform.h
+++ b/Engine/Plugins/Experimental/ControlRig/Source/ControlRig/Private/Units/Hierarchy/RigUnit_SetSpaceInitialTransform.h
@@ -10,7 +10,7 @@
  * SetSpaceInitialTransform is used to perform a change in the hierarchy by setting a single space's initial transform.
  */
 USTRUCT(meta=(DisplayName="Set Initial Space Transform", Category="Setup", DocumentationPolicy="Strict", Keywords = "SetSpaceInitialTransform", Deprecated="4.25"))
-struct FRigUnit_SetSpaceInitialTransform : public FRigUnitMutable
+struct CONTROLRIG_API FRigUnit_SetSpaceInitialTransform : public FRigUnitMutable
 {
 	GENERATED_BODY()
 
diff --git a/Engine/Plugins/Experimental/ControlRig/Source/ControlRig/Private/Units/Hierarchy/RigUnit_SetSpaceTransform.h b/Engine/Plugins/Experimental/ControlRig/Source/ControlRig/Private/Units/Hierarchy/RigUnit_SetSpaceTransform.h
--- a/Engine/Plugins/Experimental/ControlRig/Source/ControlRig/Private/Units/Hierarchy/RigUnit_SetSpaceTransform.h
+++ b/Engine/Plugins/Experimental/ControlRig/Source/ControlRig/Private/Units/Hierarchy/RigUnit_SetSpaceTransform.h
@@ -10,7 +10,7 @@
  * SetSpaceTransform is used to perform a change in the hierarchy by setting a single space's transform.
  */
 USTRUCT(meta=(DisplayName="Set Space", Category="Hierarchy", DocumentationPolicy="Strict", Keywords = "SetSpaceTransform", Deprecated="4.25"))
-struct FRigUnit_SetSpaceTransform : public FRigUnitMutable
+struct CONTROLRIG_API FRigUnit_SetSpaceTransform : public FRigUnitMutable
 {
 	GENERATED_BODY()
 
diff --git a/Engine/Plugins/Experimental/ControlRig/Source/ControlRig/Private/Units/Hierarchy/RigUnit_WorldSpace.h b/Engine/Plugins/Experimental/ControlRig/Source/ControlRig/Private/Units/Hierarchy/RigUnit_WorldSpace.h
--- a/Engine/Plugins/Experimental/ControlRig/Source/ControlRig/Private/Units/Hierarchy/RigUnit_WorldSpace.h
+++ b/Engine/Plugins/Experimental/ControlRig/Source/ControlRig/Private/Units/Hierarchy/RigUnit_WorldSpace.h
@@ -9,7 +9,7 @@
  * Converts a transform from rig (global) space to world space
  */
 USTRUCT(meta=(DisplayName="To World", MenuDescSuffix="(Transform)", Category="Hierarchy", DocumentationPolicy = "Strict", Keywords="Global,Local,World,Actor,ComponentSpace,FromRig", Varying))
-struct FRigUnit_ToWorldSpace_Transform : public FRigUnit
+struct CONTROLRIG_API FRigUnit_ToWorldSpace_Transform : public FRigUnit
 {
 	GENERATED_BODY()
 
@@ -38,7 +38,7 @@ struct FRigUnit_ToWorldSpace_Transform : public FRigUnit
  * Converts a transform from world space to rig (global) space
  */
 USTRUCT(meta=(DisplayName="From World", MenuDescSuffix="(Transform)", Category="Hierarchy", DocumentationPolicy = "Strict", Keywords="Global,Local,World,Actor,ComponentSpace,ToRig", Varying))
-struct FRigUnit_ToRigSpace_Transform : public FRigUnit
+struct CONTROLRIG_API FRigUnit_ToRigSpace_Transform : public FRigUnit
 {
 	GENERATED_BODY()
 
@@ -67,7 +67,7 @@ struct FRigUnit_ToRigSpace_Transform : public FRigUnit
  * Converts a position / location from rig (global) space to world space
  */
 USTRUCT(meta=(DisplayName="To World", MenuDescSuffix="(Vector)", Category="Hierarchy", DocumentationPolicy = "Strict", Keywords="Global,Local,World,Actor,ComponentSpace,FromRig", Varying))
-struct FRigUnit_ToWorldSpace_Location : public FRigUnit
+struct CONTROLRIG_API FRigUnit_ToWorldSpace_Location : public FRigUnit
 {
 	GENERATED_BODY()
 
@@ -96,7 +96,7 @@ struct FRigUnit_ToWorldSpace_Location : public FRigUnit
  * Converts a position / location from world space to rig (global) space
  */
 USTRUCT(meta=(DisplayName="From World", MenuDescSuffix="(Vector)", Category="Hierarchy", DocumentationPolicy = "Strict", Keywords="Global,Local,World,Actor,ComponentSpace,ToRig", Varying))
-struct FRigUnit_ToRigSpace_Location : public FRigUnit
+struct CONTROLRIG_API FRigUnit_ToRigSpace_Location : public FRigUnit
 {
 	GENERATED_BODY()
 
@@ -125,7 +125,7 @@ struct FRigUnit_ToRigSpace_Location : public FRigUnit
  * Converts a rotation from rig (global) space to world space
  */
 USTRUCT(meta=(DisplayName="To World", MenuDescSuffix="(Quat)", Category="Hierarchy", DocumentationPolicy = "Strict", Keywords="Global,Local,World,Actor,ComponentSpace,FromRig", Varying))
-struct FRigUnit_ToWorldSpace_Rotation : public FRigUnit
+struct CONTROLRIG_API FRigUnit_ToWorldSpace_Rotation : public FRigUnit
 {
 	GENERATED_BODY()
 
@@ -154,7 +154,7 @@ struct FRigUnit_ToWorldSpace_Rotation : public FRigUnit
  * Converts a rotation from world space to rig (global) space
  */
 USTRUCT(meta=(DisplayName="From World", MenuDescSuffix="(Quat)", Category="Hierarchy", DocumentationPolicy = "Strict", Keywords="Global,Local,World,Actor,ComponentSpace,ToRig", Varying))
-struct FRigUnit_ToRigSpace_Rotation : public FRigUnit
+struct CONTROLRIG_API FRigUnit_ToRigSpace_Rotation : public FRigUnit
 {
 	GENERATED_BODY()
 
diff --git a/Engine/Plugins/Experimental/ControlRig/Source/ControlRig/Private/Units/Highlevel/Harmonics/RigUnit_BoneHarmonics.h b/Engine/Plugins/Experimental/ControlRig/Source/ControlRig/Private/Units/Highlevel/Harmonics/RigUnit_BoneHarmonics.h
--- a/Engine/Plugins/Experimental/ControlRig/Source/ControlRig/Private/Units/Highlevel/Harmonics/RigUnit_BoneHarmonics.h
+++ b/Engine/Plugins/Experimental/ControlRig/Source/ControlRig/Private/Units/Highlevel/Harmonics/RigUnit_BoneHarmonics.h
@@ -7,7 +7,7 @@
 #include "RigUnit_BoneHarmonics.generated.h"
 
 USTRUCT()
-struct FRigUnit_BoneHarmonics_BoneTarget
+struct CONTROLRIG_API FRigUnit_BoneHarmonics_BoneTarget
 {
 	GENERATED_BODY()
 
@@ -32,7 +32,7 @@ struct FRigUnit_BoneHarmonics_BoneTarget
 };
 
 USTRUCT()
-struct FRigUnit_Harmonics_TargetItem
+struct CONTROLRIG_API FRigUnit_Harmonics_TargetItem
 {
 	GENERATED_BODY()
 
@@ -57,7 +57,7 @@ struct FRigUnit_Harmonics_TargetItem
 };
 
 USTRUCT()
-struct FRigUnit_BoneHarmonics_WorkData
+struct CONTROLRIG_API FRigUnit_BoneHarmonics_WorkData
 {
 	GENERATED_BODY()
 
@@ -77,7 +77,7 @@ struct FRigUnit_BoneHarmonics_WorkData
  * Performs point based simulation
  */
 USTRUCT(meta=(DisplayName="Harmonics", Keywords="Sin,Wave", Deprecated = 4.25))
-struct FRigUnit_BoneHarmonics : public FRigUnit_HighlevelBaseMutable
+struct CONTROLRIG_API FRigUnit_BoneHarmonics : public FRigUnit_HighlevelBaseMutable
 {
 	GENERATED_BODY()
 	
@@ -146,7 +146,7 @@ struct FRigUnit_BoneHarmonics : public FRigUnit_HighlevelBaseMutable
  * Performs point based simulation
  */
 USTRUCT(meta=(DisplayName="Harmonics", Keywords="Sin,Wave"))
-struct FRigUnit_ItemHarmonics : public FRigUnit_HighlevelBaseMutable
+struct CONTROLRIG_API FRigUnit_ItemHarmonics : public FRigUnit_HighlevelBaseMutable
 {
 	GENERATED_BODY()
 	
diff --git a/Engine/Plugins/Experimental/ControlRig/Source/ControlRig/Private/Units/Highlevel/Harmonics/RigUnit_ChainHarmonics.h b/Engine/Plugins/Experimental/ControlRig/Source/ControlRig/Private/Units/Highlevel/Harmonics/RigUnit_ChainHarmonics.h
--- a/Engine/Plugins/Experimental/ControlRig/Source/ControlRig/Private/Units/Highlevel/Harmonics/RigUnit_ChainHarmonics.h
+++ b/Engine/Plugins/Experimental/ControlRig/Source/ControlRig/Private/Units/Highlevel/Harmonics/RigUnit_ChainHarmonics.h
@@ -8,7 +8,7 @@
 #include "RigUnit_ChainHarmonics.generated.h"
 
 USTRUCT()
-struct FRigUnit_ChainHarmonics_Reach
+struct CONTROLRIG_API FRigUnit_ChainHarmonics_Reach
 {
 	GENERATED_BODY()
 
@@ -42,7 +42,7 @@ struct FRigUnit_ChainHarmonics_Reach
 };
 
 USTRUCT()
-struct FRigUnit_ChainHarmonics_Wave
+struct CONTROLRIG_API FRigUnit_ChainHarmonics_Wave
 {
 	GENERATED_BODY()
 
@@ -84,7 +84,7 @@ struct FRigUnit_ChainHarmonics_Wave
 };
 
 USTRUCT()
-struct FRigUnit_ChainHarmonics_Pendulum
+struct CONTROLRIG_API FRigUnit_ChainHarmonics_Pendulum
 {
 	GENERATED_BODY()
 
@@ -138,7 +138,7 @@ struct FRigUnit_ChainHarmonics_Pendulum
 };
 
 USTRUCT()
-struct FRigUnit_ChainHarmonics_WorkData
+struct CONTROLRIG_API FRigUnit_ChainHarmonics_WorkData
 {
 	GENERATED_BODY()
 
@@ -176,7 +176,7 @@ struct FRigUnit_ChainHarmonics_WorkData
 };
 
 USTRUCT(meta=(DisplayName="ChainHarmonics", Deprecated = "4.25"))
-struct FRigUnit_ChainHarmonics : public FRigUnit_HighlevelBaseMutable
+struct CONTROLRIG_API FRigUnit_ChainHarmonics : public FRigUnit_HighlevelBaseMutable
 {
 	GENERATED_BODY()
 
@@ -232,7 +232,7 @@ struct FRigUnit_ChainHarmonics : public FRigUnit_HighlevelBaseMutable
  * Performs chain based simulation
  */
 USTRUCT(meta=(DisplayName="Chain Harmonics"))
-struct FRigUnit_ChainHarmonicsPerItem : public FRigUnit_HighlevelBaseMutable
+struct CONTROLRIG_API FRigUnit_ChainHarmonicsPerItem : public FRigUnit_HighlevelBaseMutable
 {
 	GENERATED_BODY()
 
diff --git a/Engine/Plugins/Experimental/ControlRig/Source/ControlRig/Private/Units/Highlevel/Hierarchy/RigUnit_AimBone.h b/Engine/Plugins/Experimental/ControlRig/Source/ControlRig/Private/Units/Highlevel/Hierarchy/RigUnit_AimBone.h
--- a/Engine/Plugins/Experimental/ControlRig/Source/ControlRig/Private/Units/Highlevel/Hierarchy/RigUnit_AimBone.h
+++ b/Engine/Plugins/Experimental/ControlRig/Source/ControlRig/Private/Units/Highlevel/Hierarchy/RigUnit_AimBone.h
@@ -6,7 +6,7 @@
 #include "RigUnit_AimBone.generated.h"
 
 USTRUCT()
-struct FRigUnit_AimBone_Target
+struct CONTROLRIG_API FRigUnit_AimBone_Target
 {
 	GENERATED_BODY()
 
@@ -51,7 +51,7 @@ struct FRigUnit_AimBone_Target
 };
 
 USTRUCT()
-struct FRigUnit_AimItem_Target
+struct CONTROLRIG_API FRigUnit_AimItem_Target
 {
 	GENERATED_BODY()
 
@@ -96,7 +96,7 @@ struct FRigUnit_AimItem_Target
 };
 
 USTRUCT()
-struct FRigUnit_AimBone_DebugSettings
+struct CONTROLRIG_API FRigUnit_AimBone_DebugSettings
 {
 	GENERATED_BODY()
 
@@ -131,7 +131,7 @@ struct FRigUnit_AimBone_DebugSettings
  * Note: This node operates in world space!
  */
 USTRUCT(meta = (DisplayName = "Aim Math", Category = "Hierarchy", Keywords = "Lookat"))
-struct FRigUnit_AimBoneMath : public FRigUnit_HighlevelBase
+struct CONTROLRIG_API FRigUnit_AimBoneMath : public FRigUnit_HighlevelBase
 {
 	GENERATED_BODY()
 
@@ -197,7 +197,7 @@ struct FRigUnit_AimBoneMath : public FRigUnit_HighlevelBase
  * Note: This node operates in global space!
  */
 USTRUCT(meta=(DisplayName="Aim", Category="Hierarchy", Keywords="Lookat", Deprecated = "4.25"))
-struct FRigUnit_AimBone : public FRigUnit_HighlevelBaseMutable
+struct CONTROLRIG_API FRigUnit_AimBone : public FRigUnit_HighlevelBaseMutable
 {
 	GENERATED_BODY()
 
@@ -283,7 +283,7 @@ struct FRigUnit_AimBone : public FRigUnit_HighlevelBaseMutable
  * Note: This node operates in global space!
  */
 USTRUCT(meta=(DisplayName="Aim", Category="Hierarchy", Keywords="Lookat"))
-struct FRigUnit_AimItem: public FRigUnit_HighlevelBaseMutable
+struct CONTROLRIG_API FRigUnit_AimItem: public FRigUnit_HighlevelBaseMutable
 {
 	GENERATED_BODY()
 
diff --git a/Engine/Plugins/Experimental/ControlRig/Source/ControlRig/Private/Units/Highlevel/Hierarchy/RigUnit_CCDIK.h b/Engine/Plugins/Experimental/ControlRig/Source/ControlRig/Private/Units/Highlevel/Hierarchy/RigUnit_CCDIK.h
--- a/Engine/Plugins/Experimental/ControlRig/Source/ControlRig/Private/Units/Highlevel/Hierarchy/RigUnit_CCDIK.h
+++ b/Engine/Plugins/Experimental/ControlRig/Source/ControlRig/Private/Units/Highlevel/Hierarchy/RigUnit_CCDIK.h
@@ -7,7 +7,7 @@
 #include "RigUnit_CCDIK.generated.h"
 
 USTRUCT()
-struct FRigUnit_CCDIK_RotationLimit
+struct CONTROLRIG_API FRigUnit_CCDIK_RotationLimit
 {
 	GENERATED_BODY()
 
@@ -31,7 +31,7 @@ struct FRigUnit_CCDIK_RotationLimit
 };
 
 USTRUCT()
-struct FRigUnit_CCDIK_RotationLimitPerItem
+struct CONTROLRIG_API FRigUnit_CCDIK_RotationLimitPerItem
 {
 	GENERATED_BODY()
 
@@ -55,7 +55,7 @@ struct FRigUnit_CCDIK_RotationLimitPerItem
 };
 
 USTRUCT()
-struct FRigUnit_CCDIK_WorkData
+struct CONTROLRIG_API FRigUnit_CCDIK_WorkData
 {
 	GENERATED_BODY()
 
@@ -86,7 +86,7 @@ struct FRigUnit_CCDIK_WorkData
  * For now this node supports single effector chains only.
  */
 USTRUCT(meta=(DisplayName="CCDIK", Category="Hierarchy", Keywords="N-Bone,IK", Deprecated = "4.25"))
-struct FRigUnit_CCDIK : public FRigUnit_HighlevelBaseMutable
+struct CONTROLRIG_API FRigUnit_CCDIK : public FRigUnit_HighlevelBaseMutable
 {
 	GENERATED_BODY()
 
@@ -176,7 +176,7 @@ struct FRigUnit_CCDIK : public FRigUnit_HighlevelBaseMutable
  * For now this node supports single effector chains only.
  */
 USTRUCT(meta=(DisplayName="CCDIK", Category="Hierarchy", Keywords="N-Bone,IK"))
-struct FRigUnit_CCDIKPerItem : public FRigUnit_HighlevelBaseMutable
+struct CONTROLRIG_API FRigUnit_CCDIKPerItem : public FRigUnit_HighlevelBaseMutable
 {
 	GENERATED_BODY()
 
diff --git a/Engine/Plugins/Experimental/ControlRig/Source/ControlRig/Private/Units/Highlevel/Hierarchy/RigUnit_DistributeRotation.h b/Engine/Plugins/Experimental/ControlRig/Source/ControlRig/Private/Units/Highlevel/Hierarchy/RigUnit_DistributeRotation.h
--- a/Engine/Plugins/Experimental/ControlRig/Source/ControlRig/Private/Units/Highlevel/Hierarchy/RigUnit_DistributeRotation.h
+++ b/Engine/Plugins/Experimental/ControlRig/Source/ControlRig/Private/Units/Highlevel/Hierarchy/RigUnit_DistributeRotation.h
@@ -7,7 +7,7 @@
 #include "RigUnit_DistributeRotation.generated.h"
 
 USTRUCT()
-struct FRigUnit_DistributeRotation_Rotation
+struct CONTROLRIG_API FRigUnit_DistributeRotation_Rotation
 {
 	GENERATED_BODY()
 
@@ -31,7 +31,7 @@ struct FRigUnit_DistributeRotation_Rotation
 };
 
 USTRUCT()
-struct FRigUnit_DistributeRotation_WorkData
+struct CONTROLRIG_API FRigUnit_DistributeRotation_WorkData
 {
 	GENERATED_BODY()
 
@@ -57,7 +57,7 @@ struct FRigUnit_DistributeRotation_WorkData
  * Note: This node adds rotation in local space of each bone!
  */
 USTRUCT(meta=(DisplayName="Distribute Rotation", Category="Hierarchy", Keywords="TwistBones", Deprecated = "4.25"))
-struct FRigUnit_DistributeRotation : public FRigUnit_HighlevelBaseMutable
+struct CONTROLRIG_API FRigUnit_DistributeRotation : public FRigUnit_HighlevelBaseMutable
 {
 	GENERATED_BODY()
 
@@ -129,7 +129,7 @@ struct FRigUnit_DistributeRotation : public FRigUnit_HighlevelBaseMutable
  * Note: This node adds rotation in local space of each item!
  */
 USTRUCT(meta=(DisplayName="Distribute Rotation", Category="Hierarchy", Keywords="TwistBones"))
-struct FRigUnit_DistributeRotationForCollection : public FRigUnit_HighlevelBaseMutable
+struct CONTROLRIG_API FRigUnit_DistributeRotationForCollection : public FRigUnit_HighlevelBaseMutable
 {
 	GENERATED_BODY()
 
diff --git a/Engine/Plugins/Experimental/ControlRig/Source/ControlRig/Private/Units/Highlevel/Hierarchy/RigUnit_FABRIK.h b/Engine/Plugins/Experimental/ControlRig/Source/ControlRig/Private/Units/Highlevel/Hierarchy/RigUnit_FABRIK.h
--- a/Engine/Plugins/Experimental/ControlRig/Source/ControlRig/Private/Units/Highlevel/Hierarchy/RigUnit_FABRIK.h
+++ b/Engine/Plugins/Experimental/ControlRig/Source/ControlRig/Private/Units/Highlevel/Hierarchy/RigUnit_FABRIK.h
@@ -7,7 +7,7 @@
 #include "RigUnit_FABRIK.generated.h"
 
 USTRUCT()
-struct FRigUnit_FABRIK_WorkData
+struct CONTROLRIG_API FRigUnit_FABRIK_WorkData
 {
 	GENERATED_BODY()
 
@@ -32,7 +32,7 @@ struct FRigUnit_FABRIK_WorkData
  * For now this node supports single effector chains only.
  */
 USTRUCT(meta=(DisplayName="Basic FABRIK", Category="Hierarchy", Keywords="N-Bone,IK", Deprecated = "4.25"))
-struct FRigUnit_FABRIK : public FRigUnit_HighlevelBaseMutable
+struct CONTROLRIG_API FRigUnit_FABRIK : public FRigUnit_HighlevelBaseMutable
 {
 	GENERATED_BODY()
 
@@ -102,7 +102,7 @@ struct FRigUnit_FABRIK : public FRigUnit_HighlevelBaseMutable
  * For now this node supports single effector chains only.
  */
 USTRUCT(meta=(DisplayName="Basic FABRIK", Category="Hierarchy", Keywords="N-Bone,IK"))
-struct FRigUnit_FABRIKPerItem : public FRigUnit_HighlevelBaseMutable
+struct CONTROLRIG_API FRigUnit_FABRIKPerItem : public FRigUnit_HighlevelBaseMutable
 {
 	GENERATED_BODY()
 
diff --git a/Engine/Plugins/Experimental/ControlRig/Source/ControlRig/Private/Units/Highlevel/Hierarchy/RigUnit_FitChainToCurve.h b/Engine/Plugins/Experimental/ControlRig/Source/ControlRig/Private/Units/Highlevel/Hierarchy/RigUnit_FitChainToCurve.h
--- a/Engine/Plugins/Experimental/ControlRig/Source/ControlRig/Private/Units/Highlevel/Hierarchy/RigUnit_FitChainToCurve.h
+++ b/Engine/Plugins/Experimental/ControlRig/Source/ControlRig/Private/Units/Highlevel/Hierarchy/RigUnit_FitChainToCurve.h
@@ -14,7 +14,7 @@ enum class EControlRigCurveAlignment : uint8
 };
 
 USTRUCT()
-struct FRigUnit_FitChainToCurve_Rotation
+struct CONTROLRIG_API FRigUnit_FitChainToCurve_Rotation
 {
 	GENERATED_BODY()
 
@@ -38,7 +38,7 @@ struct FRigUnit_FitChainToCurve_Rotation
 };
 
 USTRUCT()
-struct FRigUnit_FitChainToCurve_DebugSettings
+struct CONTROLRIG_API FRigUnit_FitChainToCurve_DebugSettings
 {
 	GENERATED_BODY()
 
@@ -83,7 +83,7 @@ struct FRigUnit_FitChainToCurve_DebugSettings
 };
 
 USTRUCT()
-struct FRigUnit_FitChainToCurve_WorkData
+struct CONTROLRIG_API FRigUnit_FitChainToCurve_WorkData
 {
 	GENERATED_BODY()
 
@@ -128,7 +128,7 @@ struct FRigUnit_FitChainToCurve_WorkData
  * Additionally provides rotational control matching the features of the Distribute Rotation node.
  */
 USTRUCT(meta=(DisplayName="Fit Chain on Curve", Category="Hierarchy", Keywords="Fit,Resample,Bezier", Deprecated = "4.25"))
-struct FRigUnit_FitChainToCurve : public FRigUnit_HighlevelBaseMutable
+struct CONTROLRIG_API FRigUnit_FitChainToCurve : public FRigUnit_HighlevelBaseMutable
 {
 	GENERATED_BODY()
 
@@ -252,7 +252,7 @@ struct FRigUnit_FitChainToCurve : public FRigUnit_HighlevelBaseMutable
  * Additionally provides rotational control matching the features of the Distribute Rotation node.
  */
 USTRUCT(meta=(DisplayName="Fit Chain on Curve", Category="Hierarchy", Keywords="Fit,Resample,Bezier"))
-struct FRigUnit_FitChainToCurvePerItem : public FRigUnit_HighlevelBaseMutable
+struct CONTROLRIG_API FRigUnit_FitChainToCurvePerItem : public FRigUnit_HighlevelBaseMutable
 {
 	GENERATED_BODY()
 
diff --git a/Engine/Plugins/Experimental/ControlRig/Source/ControlRig/Private/Units/Highlevel/Hierarchy/RigUnit_ModifyBoneTransforms.h b/Engine/Plugins/Experimental/ControlRig/Source/ControlRig/Private/Units/Highlevel/Hierarchy/RigUnit_ModifyBoneTransforms.h
--- a/Engine/Plugins/Experimental/ControlRig/Source/ControlRig/Private/Units/Highlevel/Hierarchy/RigUnit_ModifyBoneTransforms.h
+++ b/Engine/Plugins/Experimental/ControlRig/Source/ControlRig/Private/Units/Highlevel/Hierarchy/RigUnit_ModifyBoneTransforms.h
@@ -8,7 +8,7 @@
 #include "RigUnit_ModifyBoneTransforms.generated.h"
 
 USTRUCT()
-struct FRigUnit_ModifyBoneTransforms_PerBone
+struct CONTROLRIG_API FRigUnit_ModifyBoneTransforms_PerBone
 {
 	GENERATED_BODY()
 
@@ -31,7 +31,7 @@ struct FRigUnit_ModifyBoneTransforms_PerBone
 };
 
 USTRUCT()
-struct FRigUnit_ModifyBoneTransforms_WorkData : public FRigUnit_ModifyTransforms_WorkData
+struct CONTROLRIG_API FRigUnit_ModifyBoneTransforms_WorkData : public FRigUnit_ModifyTransforms_WorkData
 {
 	GENERATED_BODY()
 };
@@ -40,7 +40,7 @@ struct FRigUnit_ModifyBoneTransforms_WorkData : public FRigUnit_ModifyTransforms
  * ModifyBonetransforms is used to perform a change in the hierarchy by setting one or more bones' transforms.
  */
 USTRUCT(meta=(DisplayName="Modify Transforms", Category="Hierarchy", DocumentationPolicy="Strict", Keywords = "ModifyBone", Deprecated = "4.25"))
-struct FRigUnit_ModifyBoneTransforms : public FRigUnit_HighlevelBaseMutable
+struct CONTROLRIG_API FRigUnit_ModifyBoneTransforms : public FRigUnit_HighlevelBaseMutable
 {
 	GENERATED_BODY()
 
diff --git a/Engine/Plugins/Experimental/ControlRig/Source/ControlRig/Private/Units/Highlevel/Hierarchy/RigUnit_ModifyTransforms.h b/Engine/Plugins/Experimental/ControlRig/Source/ControlRig/Private/Units/Highlevel/Hierarchy/RigUnit_ModifyTransforms.h
--- a/Engine/Plugins/Experimental/ControlRig/Source/ControlRig/Private/Units/Highlevel/Hierarchy/RigUnit_ModifyTransforms.h
+++ b/Engine/Plugins/Experimental/ControlRig/Source/ControlRig/Private/Units/Highlevel/Hierarchy/RigUnit_ModifyTransforms.h
@@ -32,7 +32,7 @@ enum class EControlRigModifyBoneMode : uint8
 };
 
 USTRUCT()
-struct FRigUnit_ModifyTransforms_PerItem
+struct CONTROLRIG_API FRigUnit_ModifyTransforms_PerItem
 {
 	GENERATED_BODY()
 
@@ -55,7 +55,7 @@ struct FRigUnit_ModifyTransforms_PerItem
 };
 
 USTRUCT()
-struct FRigUnit_ModifyTransforms_WorkData
+struct CONTROLRIG_API FRigUnit_ModifyTransforms_WorkData
 {
 	GENERATED_BODY()
 
@@ -67,7 +67,7 @@ struct FRigUnit_ModifyTransforms_WorkData
  * Modify Transforms is used to perform a change in the hierarchy by setting one or more bones' transforms
  */
 USTRUCT(meta=(DisplayName="Modify Transforms", Category="Hierarchy", DocumentationPolicy="Strict", Keywords = "ModifyBone"))
-struct FRigUnit_ModifyTransforms : public FRigUnit_HighlevelBaseMutable
+struct CONTROLRIG_API FRigUnit_ModifyTransforms : public FRigUnit_HighlevelBaseMutable
 {
 	GENERATED_BODY()
 
diff --git a/Engine/Plugins/Experimental/ControlRig/Source/ControlRig/Private/Units/Highlevel/Hierarchy/RigUnit_MultiFABRIK.h b/Engine/Plugins/Experimental/ControlRig/Source/ControlRig/Private/Units/Highlevel/Hierarchy/RigUnit_MultiFABRIK.h
--- a/Engine/Plugins/Experimental/ControlRig/Source/ControlRig/Private/Units/Highlevel/Hierarchy/RigUnit_MultiFABRIK.h
+++ b/Engine/Plugins/Experimental/ControlRig/Source/ControlRig/Private/Units/Highlevel/Hierarchy/RigUnit_MultiFABRIK.h
@@ -9,7 +9,7 @@
 /* This is each element WIP transform
  * while iterating, we use this as bone transform data
  */
-struct FRigUnit_MultiFABRIK_BoneWorkingData
+struct CONTROLRIG_API FRigUnit_MultiFABRIK_BoneWorkingData
 {
 	FName BoneName;
 	FCachedRigElement BoneIndex; // bone index of hierarchy
@@ -31,7 +31,7 @@ struct FRigUnit_MultiFABRIK_BoneWorkingData
 /* 
  * Describes each chain. This contains list of chain and effector location data and what's desired max length 
  */
-struct FRigUnit_MultiFABRIK_Chain
+struct CONTROLRIG_API FRigUnit_MultiFABRIK_Chain
 {
 	TArray<FFABRIKChainLink> Link;
 	FVector EffectorLocation; // effector location - based on where you're this location could change
@@ -56,7 +56,7 @@ struct FRigUnit_MultiFABRIK_Chain
  * Contains groups of Chains 
  * This also contains base bone tree index of the root
  */
-struct FRigUnit_MultiFABRIK_ChainGroup
+struct CONTROLRIG_API FRigUnit_MultiFABRIK_ChainGroup
 {
 	// current chain data
 	FRigUnit_MultiFABRIK_Chain Chain;
@@ -79,7 +79,7 @@ struct FRigUnit_MultiFABRIK_ChainGroup
 };
 
 USTRUCT()
-struct FRigUnit_MultiFABRIK_WorkData
+struct CONTROLRIG_API FRigUnit_MultiFABRIK_WorkData
 {
 	GENERATED_BODY()
 
@@ -98,7 +98,7 @@ struct FRigUnit_MultiFABRIK_WorkData
 };
 
 USTRUCT()
-struct FRigUnit_MultiFABRIK_EndEffector
+struct CONTROLRIG_API FRigUnit_MultiFABRIK_EndEffector
 {
 	GENERATED_BODY()
 	/**
@@ -119,7 +119,7 @@ struct FRigUnit_MultiFABRIK_EndEffector
  * For now this node supports single effector chains only.
  */
 USTRUCT(meta=(DisplayName="Multi Effector FABRIK", Category="Hierarchy", Keywords="Multi, Effector, N-Chain,IK"))
-struct FRigUnit_MultiFABRIK : public FRigUnit_HighlevelBaseMutable
+struct CONTROLRIG_API FRigUnit_MultiFABRIK : public FRigUnit_HighlevelBaseMutable
 {
 	GENERATED_BODY()
 
diff --git a/Engine/Plugins/Experimental/ControlRig/Source/ControlRig/Private/Units/Highlevel/Hierarchy/RigUnit_SlideChain.h b/Engine/Plugins/Experimental/ControlRig/Source/ControlRig/Private/Units/Highlevel/Hierarchy/RigUnit_SlideChain.h
--- a/Engine/Plugins/Experimental/ControlRig/Source/ControlRig/Private/Units/Highlevel/Hierarchy/RigUnit_SlideChain.h
+++ b/Engine/Plugins/Experimental/ControlRig/Source/ControlRig/Private/Units/Highlevel/Hierarchy/RigUnit_SlideChain.h
@@ -7,7 +7,7 @@
 #include "RigUnit_SlideChain.generated.h"
 
 USTRUCT()
-struct FRigUnit_SlideChain_WorkData
+struct CONTROLRIG_API FRigUnit_SlideChain_WorkData
 {
 	GENERATED_BODY()
 
@@ -36,7 +36,7 @@ struct FRigUnit_SlideChain_WorkData
  * Slides an existing chain along itself with control over extrapolation.
  */
 USTRUCT(meta=(DisplayName="Slide Chain", Category="Hierarchy", Keywords="Fit,Refit", Deprecated = "4.25"))
-struct FRigUnit_SlideChain: public FRigUnit_HighlevelBaseMutable
+struct CONTROLRIG_API FRigUnit_SlideChain: public FRigUnit_HighlevelBaseMutable
 {
 	GENERATED_BODY()
 
@@ -84,7 +84,7 @@ struct FRigUnit_SlideChain: public FRigUnit_HighlevelBaseMutable
  * Slides an existing chain along itself with control over extrapolation.
  */
 USTRUCT(meta=(DisplayName="Slide Chain", Category="Hierarchy", Keywords="Fit,Refit"))
-struct FRigUnit_SlideChainPerItem: public FRigUnit_HighlevelBaseMutable
+struct CONTROLRIG_API FRigUnit_SlideChainPerItem: public FRigUnit_HighlevelBaseMutable
 {
 	GENERATED_BODY()
 
diff --git a/Engine/Plugins/Experimental/ControlRig/Source/ControlRig/Private/Units/Highlevel/Hierarchy/RigUnit_SpringIK.h b/Engine/Plugins/Experimental/ControlRig/Source/ControlRig/Private/Units/Highlevel/Hierarchy/RigUnit_SpringIK.h
--- a/Engine/Plugins/Experimental/ControlRig/Source/ControlRig/Private/Units/Highlevel/Hierarchy/RigUnit_SpringIK.h
+++ b/Engine/Plugins/Experimental/ControlRig/Source/ControlRig/Private/Units/Highlevel/Hierarchy/RigUnit_SpringIK.h
@@ -7,7 +7,7 @@
 #include "RigUnit_SpringIK.generated.h"
 
 USTRUCT()
-struct FRigUnit_SpringIK_DebugSettings
+struct CONTROLRIG_API FRigUnit_SpringIK_DebugSettings
 {
 	GENERATED_BODY()
 
@@ -45,7 +45,7 @@ struct FRigUnit_SpringIK_DebugSettings
 };
 
 USTRUCT()
-struct FRigUnit_SpringIK_WorkData
+struct CONTROLRIG_API FRigUnit_SpringIK_WorkData
 {
 	GENERATED_BODY()
 
@@ -73,7 +73,7 @@ struct FRigUnit_SpringIK_WorkData
  * Note: This node operates in world space!
  */
 USTRUCT(meta=(DisplayName="Spring IK", Category="Hierarchy", Keywords="IK"))
-struct FRigUnit_SpringIK : public FRigUnit_HighlevelBaseMutable
+struct CONTROLRIG_API FRigUnit_SpringIK : public FRigUnit_HighlevelBaseMutable
 {
 	GENERATED_BODY()
 
diff --git a/Engine/Plugins/Experimental/ControlRig/Source/ControlRig/Private/Units/Highlevel/Hierarchy/RigUnit_TransformConstraint.h b/Engine/Plugins/Experimental/ControlRig/Source/ControlRig/Private/Units/Highlevel/Hierarchy/RigUnit_TransformConstraint.h
--- a/Engine/Plugins/Experimental/ControlRig/Source/ControlRig/Private/Units/Highlevel/Hierarchy/RigUnit_TransformConstraint.h
+++ b/Engine/Plugins/Experimental/ControlRig/Source/ControlRig/Private/Units/Highlevel/Hierarchy/RigUnit_TransformConstraint.h
@@ -37,7 +37,7 @@ struct FConstraintTarget
 };
 
 USTRUCT()
-struct FRigUnit_TransformConstraint_WorkData
+struct CONTROLRIG_API FRigUnit_TransformConstraint_WorkData
 {
 	GENERATED_BODY()
 
@@ -50,7 +50,7 @@ struct FRigUnit_TransformConstraint_WorkData
 };
 
 USTRUCT(meta=(DisplayName="Transform Constraint", Category="Transforms", Deprecated = "4.25"))
-struct FRigUnit_TransformConstraint : public FRigUnit_HighlevelBaseMutable
+struct CONTROLRIG_API FRigUnit_TransformConstraint : public FRigUnit_HighlevelBaseMutable
 {
 	GENERATED_BODY()
 
@@ -117,7 +117,7 @@ private:
  * Constrains an item's transform to multiple items' transforms
  */
 USTRUCT(meta=(DisplayName="Transform Constraint", Category="Transforms", Keywords = "Parent,Orient,Scale"))
-struct FRigUnit_TransformConstraintPerItem : public FRigUnit_HighlevelBaseMutable
+struct CONTROLRIG_API FRigUnit_TransformConstraintPerItem : public FRigUnit_HighlevelBaseMutable
 {
 	GENERATED_BODY()
 
diff --git a/Engine/Plugins/Experimental/ControlRig/Source/ControlRig/Private/Units/Highlevel/Hierarchy/RigUnit_TwistBones.h b/Engine/Plugins/Experimental/ControlRig/Source/ControlRig/Private/Units/Highlevel/Hierarchy/RigUnit_TwistBones.h
--- a/Engine/Plugins/Experimental/ControlRig/Source/ControlRig/Private/Units/Highlevel/Hierarchy/RigUnit_TwistBones.h
+++ b/Engine/Plugins/Experimental/ControlRig/Source/ControlRig/Private/Units/Highlevel/Hierarchy/RigUnit_TwistBones.h
@@ -7,7 +7,7 @@
 #include "RigUnit_TwistBones.generated.h"
 
 USTRUCT()
-struct FRigUnit_TwistBones_WorkData
+struct CONTROLRIG_API FRigUnit_TwistBones_WorkData
 {
 	GENERATED_BODY()
 
@@ -26,7 +26,7 @@ struct FRigUnit_TwistBones_WorkData
  * Creates a gradient of twist rotation along a chain.
  */
 USTRUCT(meta=(DisplayName="Twist Bones", Category="Hierarchy", Keywords="TwistBones", Deprecated = "4.25"))
-struct FRigUnit_TwistBones : public FRigUnit_HighlevelBaseMutable
+struct CONTROLRIG_API FRigUnit_TwistBones : public FRigUnit_HighlevelBaseMutable
 {
 	GENERATED_BODY()
 
@@ -100,7 +100,7 @@ struct FRigUnit_TwistBones : public FRigUnit_HighlevelBaseMutable
  * Creates a gradient of twist rotation along a collection of items.
  */
 USTRUCT(meta=(DisplayName="Twist Bones", Category="Hierarchy", Keywords="TwistBones"))
-struct FRigUnit_TwistBonesPerItem : public FRigUnit_HighlevelBaseMutable
+struct CONTROLRIG_API FRigUnit_TwistBonesPerItem : public FRigUnit_HighlevelBaseMutable
 {
 	GENERATED_BODY()
 
diff --git a/Engine/Plugins/Experimental/ControlRig/Source/ControlRig/Private/Units/Highlevel/Hierarchy/RigUnit_TwoBoneIKSimple.h b/Engine/Plugins/Experimental/ControlRig/Source/ControlRig/Private/Units/Highlevel/Hierarchy/RigUnit_TwoBoneIKSimple.h
--- a/Engine/Plugins/Experimental/ControlRig/Source/ControlRig/Private/Units/Highlevel/Hierarchy/RigUnit_TwoBoneIKSimple.h
+++ b/Engine/Plugins/Experimental/ControlRig/Source/ControlRig/Private/Units/Highlevel/Hierarchy/RigUnit_TwoBoneIKSimple.h
@@ -6,7 +6,7 @@
 #include "RigUnit_TwoBoneIKSimple.generated.h"
 
 USTRUCT()
-struct FRigUnit_TwoBoneIKSimple_DebugSettings
+struct CONTROLRIG_API FRigUnit_TwoBoneIKSimple_DebugSettings
 {
 	GENERATED_BODY()
 
@@ -41,7 +41,7 @@ struct FRigUnit_TwoBoneIKSimple_DebugSettings
  * Note: This node operates in world space!
  */
 USTRUCT(meta=(DisplayName="Basic IK", Category="Hierarchy", Keywords="TwoBone,IK", Deprecated = "4.25"))
-struct FRigUnit_TwoBoneIKSimple : public FRigUnit_HighlevelBaseMutable
+struct CONTROLRIG_API FRigUnit_TwoBoneIKSimple : public FRigUnit_HighlevelBaseMutable
 {
 	GENERATED_BODY()
 
@@ -208,7 +208,7 @@ struct FRigUnit_TwoBoneIKSimple : public FRigUnit_HighlevelBaseMutable
  * Note: This node operates in world space!
  */
 USTRUCT(meta=(DisplayName="Basic IK", Category="Hierarchy", Keywords="TwoBone,IK"))
-struct FRigUnit_TwoBoneIKSimplePerItem : public FRigUnit_HighlevelBaseMutable
+struct CONTROLRIG_API FRigUnit_TwoBoneIKSimplePerItem : public FRigUnit_HighlevelBaseMutable
 {
 	GENERATED_BODY()
 
@@ -375,7 +375,7 @@ struct FRigUnit_TwoBoneIKSimplePerItem : public FRigUnit_HighlevelBaseMutable
  * Note: This node operates in world space!
  */
 USTRUCT(meta = (DisplayName = "Basic IK Positions", Category = "Hierarchy", Keywords = "TwoBone,IK"))
-struct FRigUnit_TwoBoneIKSimpleVectors : public FRigUnit_HighlevelBase
+struct CONTROLRIG_API FRigUnit_TwoBoneIKSimpleVectors : public FRigUnit_HighlevelBase
 {
 	GENERATED_BODY()
 
@@ -453,7 +453,7 @@ struct FRigUnit_TwoBoneIKSimpleVectors : public FRigUnit_HighlevelBase
  * Note: This node operates in world space!
  */
 USTRUCT(meta = (DisplayName = "Basic IK Transforms", Category = "Hierarchy", Keywords = "TwoBone,IK"))
-struct FRigUnit_TwoBoneIKSimpleTransforms : public FRigUnit_HighlevelBase
+struct CONTROLRIG_API FRigUnit_TwoBoneIKSimpleTransforms : public FRigUnit_HighlevelBase
 {
 	GENERATED_BODY()
 
diff --git a/Engine/Plugins/Experimental/ControlRig/Source/ControlRig/Private/Units/Math/RigUnit_MathBase.h b/Engine/Plugins/Experimental/ControlRig/Source/ControlRig/Private/Units/Math/RigUnit_MathBase.h
--- a/Engine/Plugins/Experimental/ControlRig/Source/ControlRig/Private/Units/Math/RigUnit_MathBase.h
+++ b/Engine/Plugins/Experimental/ControlRig/Source/ControlRig/Private/Units/Math/RigUnit_MathBase.h
@@ -6,7 +6,7 @@
 #include "RigUnit_MathBase.generated.h"
 
 USTRUCT(meta=(Abstract, NodeColor = "0.05 0.25 0.05"))
-struct FRigUnit_MathBase : public FRigUnit
+struct CONTROLRIG_API FRigUnit_MathBase : public FRigUnit
 {
 	GENERATED_BODY()
 };
diff --git a/Engine/Plugins/Experimental/ControlRig/Source/ControlRig/Private/Units/Math/RigUnit_MathBool.h b/Engine/Plugins/Experimental/ControlRig/Source/ControlRig/Private/Units/Math/RigUnit_MathBool.h
--- a/Engine/Plugins/Experimental/ControlRig/Source/ControlRig/Private/Units/Math/RigUnit_MathBool.h
+++ b/Engine/Plugins/Experimental/ControlRig/Source/ControlRig/Private/Units/Math/RigUnit_MathBool.h
@@ -6,13 +6,13 @@
 #include "RigUnit_MathBool.generated.h"
 
 USTRUCT(meta=(Abstract, Category="Math|Boolean", MenuDescSuffix="(Bool)"))
-struct FRigUnit_MathBoolBase : public FRigUnit_MathBase
+struct CONTROLRIG_API FRigUnit_MathBoolBase : public FRigUnit_MathBase
 {
 	GENERATED_BODY()
 };
 
 USTRUCT(meta=(Abstract))
-struct FRigUnit_MathBoolConstant : public FRigUnit_MathBoolBase
+struct CONTROLRIG_API FRigUnit_MathBoolConstant : public FRigUnit_MathBoolBase
 {
 	GENERATED_BODY()
 
@@ -26,7 +26,7 @@ struct FRigUnit_MathBoolConstant : public FRigUnit_MathBoolBase
 };
 
 USTRUCT(meta=(Abstract))
-struct FRigUnit_MathBoolUnaryOp : public FRigUnit_MathBoolBase
+struct CONTROLRIG_API FRigUnit_MathBoolUnaryOp : public FRigUnit_MathBoolBase
 {
 	GENERATED_BODY()
 
@@ -43,7 +43,7 @@ struct FRigUnit_MathBoolUnaryOp : public FRigUnit_MathBoolBase
 };
 
 USTRUCT(meta=(Abstract))
-struct FRigUnit_MathBoolBinaryOp : public FRigUnit_MathBoolBase
+struct CONTROLRIG_API FRigUnit_MathBoolBinaryOp : public FRigUnit_MathBoolBase
 {
 	GENERATED_BODY()
 
@@ -66,7 +66,7 @@ struct FRigUnit_MathBoolBinaryOp : public FRigUnit_MathBoolBase
  * Returns true
  */
 USTRUCT(meta=(DisplayName="True", Keywords="Yes"))
-struct FRigUnit_MathBoolConstTrue : public FRigUnit_MathBoolConstant
+struct CONTROLRIG_API FRigUnit_MathBoolConstTrue : public FRigUnit_MathBoolConstant
 {
 	GENERATED_BODY()
 	
@@ -83,7 +83,7 @@ struct FRigUnit_MathBoolConstTrue : public FRigUnit_MathBoolConstant
  * Returns false
  */
 USTRUCT(meta=(DisplayName="False", Keywords="No"))
-struct FRigUnit_MathBoolConstFalse : public FRigUnit_MathBoolConstant
+struct CONTROLRIG_API FRigUnit_MathBoolConstFalse : public FRigUnit_MathBoolConstant
 {
 	GENERATED_BODY()
 	
@@ -100,7 +100,7 @@ struct FRigUnit_MathBoolConstFalse : public FRigUnit_MathBoolConstant
  * Returns true if the condition is false
  */
 USTRUCT(meta=(DisplayName="Not", PrototypeName="Not", Keywords="!"))
-struct FRigUnit_MathBoolNot : public FRigUnit_MathBoolUnaryOp
+struct CONTROLRIG_API FRigUnit_MathBoolNot : public FRigUnit_MathBoolUnaryOp
 {
 	GENERATED_BODY()
 
@@ -112,7 +112,7 @@ struct FRigUnit_MathBoolNot : public FRigUnit_MathBoolUnaryOp
  * Returns true if both conditions are true
  */
 USTRUCT(meta=(DisplayName="And", PrototypeName="And", Keywords="&&"))
-struct FRigUnit_MathBoolAnd : public FRigUnit_MathBoolBinaryOp
+struct CONTROLRIG_API FRigUnit_MathBoolAnd : public FRigUnit_MathBoolBinaryOp
 {
 	GENERATED_BODY()
 
@@ -124,7 +124,7 @@ struct FRigUnit_MathBoolAnd : public FRigUnit_MathBoolBinaryOp
  * Returns true if both conditions are false
  */
 USTRUCT(meta=(DisplayName="Nand", PrototypeName="Nand"))
-struct FRigUnit_MathBoolNand : public FRigUnit_MathBoolBinaryOp
+struct CONTROLRIG_API FRigUnit_MathBoolNand : public FRigUnit_MathBoolBinaryOp
 {
 	GENERATED_BODY()
 
@@ -136,7 +136,7 @@ struct FRigUnit_MathBoolNand : public FRigUnit_MathBoolBinaryOp
  * Returns true if one of the conditions is true
  */
 USTRUCT(meta=(DisplayName="Or", PrototypeName="Or", Keywords="||"))
-struct FRigUnit_MathBoolOr : public FRigUnit_MathBoolBinaryOp
+struct CONTROLRIG_API FRigUnit_MathBoolOr : public FRigUnit_MathBoolBinaryOp
 {
 	GENERATED_BODY()
 
@@ -148,7 +148,7 @@ struct FRigUnit_MathBoolOr : public FRigUnit_MathBoolBinaryOp
  * Returns true if the value A equals B
  */
 USTRUCT(meta=(DisplayName="Equals", PrototypeName="Equals", Keywords="Same,=="))
-struct FRigUnit_MathBoolEquals : public FRigUnit_MathBoolBase
+struct CONTROLRIG_API FRigUnit_MathBoolEquals : public FRigUnit_MathBoolBase
 {
 	GENERATED_BODY()
 	RIGVM_METHOD()
@@ -174,7 +174,7 @@ struct FRigUnit_MathBoolEquals : public FRigUnit_MathBoolBase
  * Returns true if the value A does not equal B
  */
 USTRUCT(meta=(DisplayName="Not Equals", PrototypeName="NotEquals", Keywords="Different,!=,Xor"))
-struct FRigUnit_MathBoolNotEquals : public FRigUnit_MathBoolBase
+struct CONTROLRIG_API FRigUnit_MathBoolNotEquals : public FRigUnit_MathBoolBase
 {
 	GENERATED_BODY()
 	RIGVM_METHOD()
diff --git a/Engine/Plugins/Experimental/ControlRig/Source/ControlRig/Private/Units/Math/RigUnit_MathColor.h b/Engine/Plugins/Experimental/ControlRig/Source/ControlRig/Private/Units/Math/RigUnit_MathColor.h
--- a/Engine/Plugins/Experimental/ControlRig/Source/ControlRig/Private/Units/Math/RigUnit_MathColor.h
+++ b/Engine/Plugins/Experimental/ControlRig/Source/ControlRig/Private/Units/Math/RigUnit_MathColor.h
@@ -7,13 +7,13 @@
 #include "RigUnit_MathColor.generated.h"
 
 USTRUCT(meta=(Abstract, Category="Math|Color", MenuDescSuffix="(Color)"))
-struct FRigUnit_MathColorBase : public FRigUnit_MathBase
+struct CONTROLRIG_API FRigUnit_MathColorBase : public FRigUnit_MathBase
 {
 	GENERATED_BODY()
 };
 
 USTRUCT(meta=(Abstract))
-struct FRigUnit_MathColorBinaryOp : public FRigUnit_MathColorBase
+struct CONTROLRIG_API FRigUnit_MathColorBinaryOp : public FRigUnit_MathColorBase
 {
 	GENERATED_BODY()
 
@@ -36,7 +36,7 @@ struct FRigUnit_MathColorBinaryOp : public FRigUnit_MathColorBase
  * Makes a vector from a single float
  */
 USTRUCT(meta=(DisplayName="From Float", PrototypeName="FromFloat", Keywords="Make,Construct"))
-struct FRigUnit_MathColorFromFloat : public FRigUnit_MathColorBase
+struct CONTROLRIG_API FRigUnit_MathColorFromFloat : public FRigUnit_MathColorBase
 {
 	GENERATED_BODY()
 	RIGVM_METHOD()
@@ -59,7 +59,7 @@ struct FRigUnit_MathColorFromFloat : public FRigUnit_MathColorBase
  * Returns the sum of the two values
  */
 USTRUCT(meta=(DisplayName="Add", PrototypeName="Add", Keywords="Sum,+"))
-struct FRigUnit_MathColorAdd : public FRigUnit_MathColorBinaryOp
+struct CONTROLRIG_API FRigUnit_MathColorAdd : public FRigUnit_MathColorBinaryOp
 {
 	GENERATED_BODY()
 
@@ -71,7 +71,7 @@ struct FRigUnit_MathColorAdd : public FRigUnit_MathColorBinaryOp
  * Returns the difference of the two values
  */
 USTRUCT(meta=(DisplayName="Subtract", PrototypeName="Subtract", Keywords="-"))
-struct FRigUnit_MathColorSub : public FRigUnit_MathColorBinaryOp
+struct CONTROLRIG_API FRigUnit_MathColorSub : public FRigUnit_MathColorBinaryOp
 {
 	GENERATED_BODY()
 
@@ -83,7 +83,7 @@ struct FRigUnit_MathColorSub : public FRigUnit_MathColorBinaryOp
  * Returns the product of the two values
  */
 USTRUCT(meta=(DisplayName="Multiply", PrototypeName="Multiply", Keywords="Product,*"))
-struct FRigUnit_MathColorMul : public FRigUnit_MathColorBinaryOp
+struct CONTROLRIG_API FRigUnit_MathColorMul : public FRigUnit_MathColorBinaryOp
 {
 	GENERATED_BODY()
 
@@ -100,7 +100,7 @@ struct FRigUnit_MathColorMul : public FRigUnit_MathColorBinaryOp
  * Linearly interpolates between A and B using the ratio T
  */
 USTRUCT(meta=(DisplayName="Interpolate", PrototypeName="Interpolate", Keywords="Lerp,Mix,Blend"))
-struct FRigUnit_MathColorLerp : public FRigUnit_MathColorBase
+struct CONTROLRIG_API FRigUnit_MathColorLerp : public FRigUnit_MathColorBase
 {
 	GENERATED_BODY()
 
diff --git a/Engine/Plugins/Experimental/ControlRig/Source/ControlRig/Private/Units/Math/RigUnit_MathFloat.h b/Engine/Plugins/Experimental/ControlRig/Source/ControlRig/Private/Units/Math/RigUnit_MathFloat.h
--- a/Engine/Plugins/Experimental/ControlRig/Source/ControlRig/Private/Units/Math/RigUnit_MathFloat.h
+++ b/Engine/Plugins/Experimental/ControlRig/Source/ControlRig/Private/Units/Math/RigUnit_MathFloat.h
@@ -6,13 +6,13 @@
 #include "RigUnit_MathFloat.generated.h"
 
 USTRUCT(meta=(Abstract, Category="Math|Float", MenuDescSuffix="(Float)"))
-struct FRigUnit_MathFloatBase : public FRigUnit_MathBase
+struct CONTROLRIG_API FRigUnit_MathFloatBase : public FRigUnit_MathBase
 {
 	GENERATED_BODY()
 };
 
 USTRUCT(meta=(Abstract))
-struct FRigUnit_MathFloatConstant : public FRigUnit_MathFloatBase
+struct CONTROLRIG_API FRigUnit_MathFloatConstant : public FRigUnit_MathFloatBase
 {
 	GENERATED_BODY()
 
@@ -26,7 +26,7 @@ struct FRigUnit_MathFloatConstant : public FRigUnit_MathFloatBase
 };
 
 USTRUCT(meta=(Abstract))
-struct FRigUnit_MathFloatUnaryOp : public FRigUnit_MathFloatBase
+struct CONTROLRIG_API FRigUnit_MathFloatUnaryOp : public FRigUnit_MathFloatBase
 {
 	GENERATED_BODY()
 
@@ -43,7 +43,7 @@ struct FRigUnit_MathFloatUnaryOp : public FRigUnit_MathFloatBase
 };
 
 USTRUCT(meta=(Abstract))
-struct FRigUnit_MathFloatBinaryOp : public FRigUnit_MathFloatBase
+struct CONTROLRIG_API FRigUnit_MathFloatBinaryOp : public FRigUnit_MathFloatBase
 {
 	GENERATED_BODY()
 
@@ -66,7 +66,7 @@ struct FRigUnit_MathFloatBinaryOp : public FRigUnit_MathFloatBase
  * Returns PI
  */
 USTRUCT(meta=(DisplayName="Pi"))
-struct FRigUnit_MathFloatConstPi : public FRigUnit_MathFloatConstant
+struct CONTROLRIG_API FRigUnit_MathFloatConstPi : public FRigUnit_MathFloatConstant
 {
 	GENERATED_BODY()
 	
@@ -78,7 +78,7 @@ struct FRigUnit_MathFloatConstPi : public FRigUnit_MathFloatConstant
  * Returns PI * 0.5
  */
 USTRUCT(meta=(DisplayName="Half Pi"))
-struct FRigUnit_MathFloatConstHalfPi : public FRigUnit_MathFloatConstant
+struct CONTROLRIG_API FRigUnit_MathFloatConstHalfPi : public FRigUnit_MathFloatConstant
 {
 	GENERATED_BODY()
 
@@ -90,7 +90,7 @@ struct FRigUnit_MathFloatConstHalfPi : public FRigUnit_MathFloatConstant
  * Returns PI * 2.0
  */
 USTRUCT(meta=(DisplayName="Two Pi"))
-struct FRigUnit_MathFloatConstTwoPi : public FRigUnit_MathFloatConstant
+struct CONTROLRIG_API FRigUnit_MathFloatConstTwoPi : public FRigUnit_MathFloatConstant
 {
 	GENERATED_BODY()
 
@@ -102,7 +102,7 @@ struct FRigUnit_MathFloatConstTwoPi : public FRigUnit_MathFloatConstant
  * Returns the sum of the two values
  */
 USTRUCT(meta=(DisplayName="Add", PrototypeName="Add", Keywords="Sum,+"))
-struct FRigUnit_MathFloatAdd : public FRigUnit_MathFloatBinaryOp
+struct CONTROLRIG_API FRigUnit_MathFloatAdd : public FRigUnit_MathFloatBinaryOp
 {
 	GENERATED_BODY()
 
@@ -114,7 +114,7 @@ struct FRigUnit_MathFloatAdd : public FRigUnit_MathFloatBinaryOp
  * Returns the difference of the two values
  */
 USTRUCT(meta=(DisplayName="Subtract", PrototypeName="Subtract", Keywords="-"))
-struct FRigUnit_MathFloatSub : public FRigUnit_MathFloatBinaryOp
+struct CONTROLRIG_API FRigUnit_MathFloatSub : public FRigUnit_MathFloatBinaryOp
 {
 	GENERATED_BODY()
 
@@ -126,7 +126,7 @@ struct FRigUnit_MathFloatSub : public FRigUnit_MathFloatBinaryOp
  * Returns the product of the two values
  */
 USTRUCT(meta=(DisplayName="Multiply", PrototypeName="Multiply", Keywords="Product,*"))
-struct FRigUnit_MathFloatMul : public FRigUnit_MathFloatBinaryOp
+struct CONTROLRIG_API FRigUnit_MathFloatMul : public FRigUnit_MathFloatBinaryOp
 {
 	GENERATED_BODY()
 
@@ -144,7 +144,7 @@ struct FRigUnit_MathFloatMul : public FRigUnit_MathFloatBinaryOp
  * Returns the division of the two values
  */
 USTRUCT(meta=(DisplayName="Divide", PrototypeName="Divide", Keywords="Division,Divisor,/"))
-struct FRigUnit_MathFloatDiv : public FRigUnit_MathFloatBinaryOp
+struct CONTROLRIG_API FRigUnit_MathFloatDiv : public FRigUnit_MathFloatBinaryOp
 {
 	GENERATED_BODY()
 
@@ -161,7 +161,7 @@ struct FRigUnit_MathFloatDiv : public FRigUnit_MathFloatBinaryOp
  * Returns the modulo of the two values
  */
 USTRUCT(meta=(DisplayName="Modulo", PrototypeName="Modulo", Keywords="%,fmod"))
-struct FRigUnit_MathFloatMod : public FRigUnit_MathFloatBinaryOp
+struct CONTROLRIG_API FRigUnit_MathFloatMod : public FRigUnit_MathFloatBinaryOp
 {
 	GENERATED_BODY()
 
@@ -179,7 +179,7 @@ struct FRigUnit_MathFloatMod : public FRigUnit_MathFloatBinaryOp
  * Returns the smaller of the two values
  */
 USTRUCT(meta=(DisplayName="Minimum", PrototypeName="Minimum"))
-struct FRigUnit_MathFloatMin : public FRigUnit_MathFloatBinaryOp
+struct CONTROLRIG_API FRigUnit_MathFloatMin : public FRigUnit_MathFloatBinaryOp
 {
 	GENERATED_BODY()
 
@@ -191,7 +191,7 @@ struct FRigUnit_MathFloatMin : public FRigUnit_MathFloatBinaryOp
  * Returns the larger of the two values
  */
 USTRUCT(meta=(DisplayName="Maximum", PrototypeName="Maximum"))
-struct FRigUnit_MathFloatMax : public FRigUnit_MathFloatBinaryOp
+struct CONTROLRIG_API FRigUnit_MathFloatMax : public FRigUnit_MathFloatBinaryOp
 {
 	GENERATED_BODY()
 
@@ -203,7 +203,7 @@ struct FRigUnit_MathFloatMax : public FRigUnit_MathFloatBinaryOp
  * Returns the value of A raised to the power of B.
  */
 USTRUCT(meta=(DisplayName="Power", PrototypeName="Power"))
-struct FRigUnit_MathFloatPow : public FRigUnit_MathFloatBinaryOp
+struct CONTROLRIG_API FRigUnit_MathFloatPow : public FRigUnit_MathFloatBinaryOp
 {
 	GENERATED_BODY()
 
@@ -221,7 +221,7 @@ struct FRigUnit_MathFloatPow : public FRigUnit_MathFloatBinaryOp
  * Returns the square root of the given value
  */
 USTRUCT(meta=(DisplayName="Sqrt", PrototypeName="Sqrt", Keywords="Root,Square"))
-struct FRigUnit_MathFloatSqrt : public FRigUnit_MathFloatUnaryOp
+struct CONTROLRIG_API FRigUnit_MathFloatSqrt : public FRigUnit_MathFloatUnaryOp
 {
 	GENERATED_BODY()
 
@@ -233,7 +233,7 @@ struct FRigUnit_MathFloatSqrt : public FRigUnit_MathFloatUnaryOp
  * Returns the negative value
  */
 USTRUCT(meta=(DisplayName="Negate", PrototypeName="Negate", Keywords="-,Abs"))
-struct FRigUnit_MathFloatNegate : public FRigUnit_MathFloatUnaryOp
+struct CONTROLRIG_API FRigUnit_MathFloatNegate : public FRigUnit_MathFloatUnaryOp
 {
 	GENERATED_BODY()
 
@@ -245,7 +245,7 @@ struct FRigUnit_MathFloatNegate : public FRigUnit_MathFloatUnaryOp
  * Returns the absolute (positive) value
  */
 USTRUCT(meta=(DisplayName="Absolute", PrototypeName="Absolute", Keywords="Abs,Neg"))
-struct FRigUnit_MathFloatAbs : public FRigUnit_MathFloatUnaryOp
+struct CONTROLRIG_API FRigUnit_MathFloatAbs : public FRigUnit_MathFloatUnaryOp
 {
 	GENERATED_BODY()
 
@@ -257,7 +257,7 @@ struct FRigUnit_MathFloatAbs : public FRigUnit_MathFloatUnaryOp
  * Returns the closest lower full number (integer) of the value
  */
 USTRUCT(meta=(DisplayName="Floor", PrototypeName="Floor", Keywords="Round"))
-struct FRigUnit_MathFloatFloor : public FRigUnit_MathFloatBase
+struct CONTROLRIG_API FRigUnit_MathFloatFloor : public FRigUnit_MathFloatBase
 {
 	GENERATED_BODY()
 
@@ -284,7 +284,7 @@ struct FRigUnit_MathFloatFloor : public FRigUnit_MathFloatBase
  * Returns the closest higher full number (integer) of the value
  */
 USTRUCT(meta=(DisplayName="Ceiling", PrototypeName="Ceiling", Keywords="Round"))
-struct FRigUnit_MathFloatCeil : public FRigUnit_MathFloatBase
+struct CONTROLRIG_API FRigUnit_MathFloatCeil : public FRigUnit_MathFloatBase
 {
 	GENERATED_BODY()
 
@@ -311,7 +311,7 @@ struct FRigUnit_MathFloatCeil : public FRigUnit_MathFloatBase
  * Returns the closest higher full number (integer) of the value
  */
 USTRUCT(meta=(DisplayName="Round", PrototypeName="Round"))
-struct FRigUnit_MathFloatRound : public FRigUnit_MathFloatBase
+struct CONTROLRIG_API FRigUnit_MathFloatRound : public FRigUnit_MathFloatBase
 {
 	GENERATED_BODY()
 
@@ -338,7 +338,7 @@ struct FRigUnit_MathFloatRound : public FRigUnit_MathFloatBase
  * Returns the float cast to an int (this uses floor)
  */
 USTRUCT(meta=(DisplayName="To Int", PrototypeName="Convert"))
-struct FRigUnit_MathFloatToInt : public FRigUnit_MathFloatBase
+struct CONTROLRIG_API FRigUnit_MathFloatToInt : public FRigUnit_MathFloatBase
 {
 	GENERATED_BODY()
 
@@ -362,7 +362,7 @@ struct FRigUnit_MathFloatToInt : public FRigUnit_MathFloatBase
  * Returns the sign of the value (+1 for >= 0.f, -1 for < 0.f)
  */
 USTRUCT(meta=(DisplayName="Sign", PrototypeName="Sign"))
-struct FRigUnit_MathFloatSign : public FRigUnit_MathFloatUnaryOp
+struct CONTROLRIG_API FRigUnit_MathFloatSign : public FRigUnit_MathFloatUnaryOp
 {
 	GENERATED_BODY()
 
@@ -374,7 +374,7 @@ struct FRigUnit_MathFloatSign : public FRigUnit_MathFloatUnaryOp
  * Clamps the given value within the range provided by minimum and maximum
  */
 USTRUCT(meta=(DisplayName="Clamp", PrototypeName="Clamp", Keywords="Range,Remap"))
-struct FRigUnit_MathFloatClamp : public FRigUnit_MathFloatBase
+struct CONTROLRIG_API FRigUnit_MathFloatClamp : public FRigUnit_MathFloatBase
 {
 	GENERATED_BODY()
 
@@ -403,7 +403,7 @@ struct FRigUnit_MathFloatClamp : public FRigUnit_MathFloatBase
  * Linearly interpolates between A and B using the ratio T
  */
 USTRUCT(meta=(DisplayName="Interpolate", PrototypeName="Interpolate", Keywords="Lerp,Mix,Blend"))
-struct FRigUnit_MathFloatLerp : public FRigUnit_MathFloatBase
+struct CONTROLRIG_API FRigUnit_MathFloatLerp : public FRigUnit_MathFloatBase
 {
 	GENERATED_BODY()
 
@@ -433,7 +433,7 @@ struct FRigUnit_MathFloatLerp : public FRigUnit_MathFloatBase
  * Remaps the given value from a source range to a target range.
  */
 USTRUCT(meta=(DisplayName="Remap", PrototypeName="Remap", Keywords="Rescale,Scale"))
-struct FRigUnit_MathFloatRemap : public FRigUnit_MathFloatBase
+struct CONTROLRIG_API FRigUnit_MathFloatRemap : public FRigUnit_MathFloatBase
 {
 	GENERATED_BODY()
 
@@ -474,7 +474,7 @@ struct FRigUnit_MathFloatRemap : public FRigUnit_MathFloatBase
  * Returns true if the value A equals B
  */
 USTRUCT(meta=(DisplayName="Equals", PrototypeName="Equals", Keywords="Same,=="))
-struct FRigUnit_MathFloatEquals : public FRigUnit_MathFloatBase
+struct CONTROLRIG_API FRigUnit_MathFloatEquals : public FRigUnit_MathFloatBase
 {
 	GENERATED_BODY()
 	RIGVM_METHOD()
@@ -500,7 +500,7 @@ struct FRigUnit_MathFloatEquals : public FRigUnit_MathFloatBase
  * Returns true if the value A does not equal B
  */
 USTRUCT(meta=(DisplayName="Not Equals", PrototypeName="NotEquals", Keywords="Different,!="))
-struct FRigUnit_MathFloatNotEquals : public FRigUnit_MathFloatBase
+struct CONTROLRIG_API FRigUnit_MathFloatNotEquals : public FRigUnit_MathFloatBase
 {
 	GENERATED_BODY()
 	RIGVM_METHOD()
@@ -526,7 +526,7 @@ struct FRigUnit_MathFloatNotEquals : public FRigUnit_MathFloatBase
  * Returns true if the value A is greater than B
  */
 USTRUCT(meta=(DisplayName="Greater", PrototypeName="Greater", Keywords="Larger,Bigger,>"))
-struct FRigUnit_MathFloatGreater : public FRigUnit_MathFloatBase
+struct CONTROLRIG_API FRigUnit_MathFloatGreater : public FRigUnit_MathFloatBase
 {
 	GENERATED_BODY()
 
@@ -553,7 +553,7 @@ struct FRigUnit_MathFloatGreater : public FRigUnit_MathFloatBase
  * Returns true if the value A is less than B
  */
 USTRUCT(meta=(DisplayName="Less", PrototypeName="Less", Keywords="Smaller,<"))
-struct FRigUnit_MathFloatLess : public FRigUnit_MathFloatBase
+struct CONTROLRIG_API FRigUnit_MathFloatLess : public FRigUnit_MathFloatBase
 {
 	GENERATED_BODY()
 	
@@ -580,7 +580,7 @@ struct FRigUnit_MathFloatLess : public FRigUnit_MathFloatBase
  * Returns true if the value A is greater than or equal to B
  */
 USTRUCT(meta=(DisplayName="Greater Equal", PrototypeName="GreaterEqual", Keywords="Larger,Bigger,>="))
-struct FRigUnit_MathFloatGreaterEqual : public FRigUnit_MathFloatBase
+struct CONTROLRIG_API FRigUnit_MathFloatGreaterEqual : public FRigUnit_MathFloatBase
 {
 	GENERATED_BODY()
 
@@ -607,7 +607,7 @@ struct FRigUnit_MathFloatGreaterEqual : public FRigUnit_MathFloatBase
  * Returns true if the value A is less than or equal to B
  */
 USTRUCT(meta=(DisplayName="Less Equal", PrototypeName="LessEqual", Keywords="Smaller,<="))
-struct FRigUnit_MathFloatLessEqual : public FRigUnit_MathFloatBase
+struct CONTROLRIG_API FRigUnit_MathFloatLessEqual : public FRigUnit_MathFloatBase
 {
 	GENERATED_BODY()
 
@@ -634,7 +634,7 @@ struct FRigUnit_MathFloatLessEqual : public FRigUnit_MathFloatBase
  * Returns true if the value is nearly zero
  */
 USTRUCT(meta=(DisplayName="Is Nearly Zero", PrototypeName="IsNearlyZero", Keywords="AlmostZero,0"))
-struct FRigUnit_MathFloatIsNearlyZero : public FRigUnit_MathFloatBase
+struct CONTROLRIG_API FRigUnit_MathFloatIsNearlyZero : public FRigUnit_MathFloatBase
 {
 	GENERATED_BODY()
 	
@@ -661,7 +661,7 @@ struct FRigUnit_MathFloatIsNearlyZero : public FRigUnit_MathFloatBase
  * Returns true if the value A is almost equal to B
  */
 USTRUCT(meta=(DisplayName="Is Nearly Equal", PrototypeName="IsNearlyEqual", Keywords="AlmostEqual,=="))
-struct FRigUnit_MathFloatIsNearlyEqual : public FRigUnit_MathFloatBase
+struct CONTROLRIG_API FRigUnit_MathFloatIsNearlyEqual : public FRigUnit_MathFloatBase
 {
 	GENERATED_BODY()
 
@@ -691,7 +691,7 @@ struct FRigUnit_MathFloatIsNearlyEqual : public FRigUnit_MathFloatBase
  * Return one of the two values based on the condition
  */
 USTRUCT(meta=(DisplayName="Select", PrototypeName="Select", Keywords="Pick,If", Deprecated = "4.26.0"))
-struct FRigUnit_MathFloatSelectBool : public FRigUnit_MathFloatBase
+struct CONTROLRIG_API FRigUnit_MathFloatSelectBool : public FRigUnit_MathFloatBase
 {
 	GENERATED_BODY()
 
@@ -721,7 +721,7 @@ struct FRigUnit_MathFloatSelectBool : public FRigUnit_MathFloatBase
  * Returns the degrees of a given value in radians
  */
 USTRUCT(meta=(DisplayName="Degrees", PrototypeName="Degrees"))
-struct FRigUnit_MathFloatDeg : public FRigUnit_MathFloatUnaryOp
+struct CONTROLRIG_API FRigUnit_MathFloatDeg : public FRigUnit_MathFloatUnaryOp
 {
 	GENERATED_BODY()
 
@@ -733,7 +733,7 @@ struct FRigUnit_MathFloatDeg : public FRigUnit_MathFloatUnaryOp
  * Returns the radians of a given value in degrees
  */
 USTRUCT(meta=(DisplayName="Radians", PrototypeName="Radians"))
-struct FRigUnit_MathFloatRad : public FRigUnit_MathFloatUnaryOp
+struct CONTROLRIG_API FRigUnit_MathFloatRad : public FRigUnit_MathFloatUnaryOp
 {
 	GENERATED_BODY()
 
@@ -757,7 +757,7 @@ struct CONTROLRIG_API FRigUnit_MathFloatSin : public FRigUnit_MathFloatUnaryOp
  * Returns the cosinus value of the given value (in radians)
  */
 USTRUCT(meta=(DisplayName="Cos", PrototypeName="Cos"))
-struct FRigUnit_MathFloatCos : public FRigUnit_MathFloatUnaryOp
+struct CONTROLRIG_API FRigUnit_MathFloatCos : public FRigUnit_MathFloatUnaryOp
 {
 	GENERATED_BODY()
 
@@ -769,7 +769,7 @@ struct FRigUnit_MathFloatCos : public FRigUnit_MathFloatUnaryOp
  * Returns the tangens value of the given value (in radians)
  */
 USTRUCT(meta=(DisplayName="Tan", PrototypeName="Tan"))
-struct FRigUnit_MathFloatTan : public FRigUnit_MathFloatUnaryOp
+struct CONTROLRIG_API FRigUnit_MathFloatTan : public FRigUnit_MathFloatUnaryOp
 {
 	GENERATED_BODY()
 
@@ -781,7 +781,7 @@ struct FRigUnit_MathFloatTan : public FRigUnit_MathFloatUnaryOp
  * Returns the inverse sinus value (in radians) of the given value
  */
 USTRUCT(meta=(DisplayName="Asin", PrototypeName="Asin", Keywords="Arcsin"))
-struct FRigUnit_MathFloatAsin : public FRigUnit_MathFloatUnaryOp
+struct CONTROLRIG_API FRigUnit_MathFloatAsin : public FRigUnit_MathFloatUnaryOp
 {
 	GENERATED_BODY()
 
@@ -793,7 +793,7 @@ struct FRigUnit_MathFloatAsin : public FRigUnit_MathFloatUnaryOp
  * Returns the inverse cosinus value (in radians) of the given value
  */
 USTRUCT(meta=(DisplayName="Acos", PrototypeName="Acos", Keywords="Arccos"))
-struct FRigUnit_MathFloatAcos : public FRigUnit_MathFloatUnaryOp
+struct CONTROLRIG_API FRigUnit_MathFloatAcos : public FRigUnit_MathFloatUnaryOp
 {
 	GENERATED_BODY()
 
@@ -805,7 +805,7 @@ struct FRigUnit_MathFloatAcos : public FRigUnit_MathFloatUnaryOp
  * Returns the inverse tangens value (in radians) of the given value
  */
 USTRUCT(meta=(DisplayName="Atan", PrototypeName="Atan", Keywords="Arctan"))
-struct FRigUnit_MathFloatAtan : public FRigUnit_MathFloatUnaryOp
+struct CONTROLRIG_API FRigUnit_MathFloatAtan : public FRigUnit_MathFloatUnaryOp
 {
 	GENERATED_BODY()
 
@@ -817,7 +817,7 @@ struct FRigUnit_MathFloatAtan : public FRigUnit_MathFloatUnaryOp
  * Computes the angles alpha, beta and gamma (in radians) between the three sides A, B and C
  */
 USTRUCT(meta = (DisplayName = "Law Of Cosine"))
-struct FRigUnit_MathFloatLawOfCosine : public FRigUnit_MathFloatBase
+struct CONTROLRIG_API FRigUnit_MathFloatLawOfCosine : public FRigUnit_MathFloatBase
 {
 	GENERATED_BODY()
 
diff --git a/Engine/Plugins/Experimental/ControlRig/Source/ControlRig/Private/Units/Math/RigUnit_MathInt.h b/Engine/Plugins/Experimental/ControlRig/Source/ControlRig/Private/Units/Math/RigUnit_MathInt.h
--- a/Engine/Plugins/Experimental/ControlRig/Source/ControlRig/Private/Units/Math/RigUnit_MathInt.h
+++ b/Engine/Plugins/Experimental/ControlRig/Source/ControlRig/Private/Units/Math/RigUnit_MathInt.h
@@ -6,13 +6,13 @@
 #include "RigUnit_MathInt.generated.h"
 
 USTRUCT(meta=(Abstract, Category="Math|Int", MenuDescSuffix="(Int)"))
-struct FRigUnit_MathIntBase : public FRigUnit_MathBase
+struct CONTROLRIG_API FRigUnit_MathIntBase : public FRigUnit_MathBase
 {
 	GENERATED_BODY()
 };
 
 USTRUCT(meta=(Abstract))
-struct FRigUnit_MathIntUnaryOp : public FRigUnit_MathIntBase
+struct CONTROLRIG_API FRigUnit_MathIntUnaryOp : public FRigUnit_MathIntBase
 {
 	GENERATED_BODY()
 
@@ -29,7 +29,7 @@ struct FRigUnit_MathIntUnaryOp : public FRigUnit_MathIntBase
 };
 
 USTRUCT(meta=(Abstract))
-struct FRigUnit_MathIntBinaryOp : public FRigUnit_MathIntBase
+struct CONTROLRIG_API FRigUnit_MathIntBinaryOp : public FRigUnit_MathIntBase
 {
 	GENERATED_BODY()
 
@@ -52,7 +52,7 @@ struct FRigUnit_MathIntBinaryOp : public FRigUnit_MathIntBase
  * Returns the sum of the two values
  */
 USTRUCT(meta=(DisplayName="Add", PrototypeName="Add", Keywords="Sum,+"))
-struct FRigUnit_MathIntAdd : public FRigUnit_MathIntBinaryOp
+struct CONTROLRIG_API FRigUnit_MathIntAdd : public FRigUnit_MathIntBinaryOp
 {
 	GENERATED_BODY()
 
@@ -64,7 +64,7 @@ struct FRigUnit_MathIntAdd : public FRigUnit_MathIntBinaryOp
  * Returns the difference of the two values
  */
 USTRUCT(meta=(DisplayName="Subtract", PrototypeName="Subtract", Keywords="-"))
-struct FRigUnit_MathIntSub : public FRigUnit_MathIntBinaryOp
+struct CONTROLRIG_API FRigUnit_MathIntSub : public FRigUnit_MathIntBinaryOp
 {
 	GENERATED_BODY()
 
@@ -76,7 +76,7 @@ struct FRigUnit_MathIntSub : public FRigUnit_MathIntBinaryOp
  * Returns the product of the two values
  */
 USTRUCT(meta=(DisplayName="Multiply", PrototypeName="Multiply", Keywords="Product,*"))
-struct FRigUnit_MathIntMul : public FRigUnit_MathIntBinaryOp
+struct CONTROLRIG_API FRigUnit_MathIntMul : public FRigUnit_MathIntBinaryOp
 {
 	GENERATED_BODY()
 
@@ -94,7 +94,7 @@ struct FRigUnit_MathIntMul : public FRigUnit_MathIntBinaryOp
  * Returns the division of the two values
  */
 USTRUCT(meta=(DisplayName="Divide", PrototypeName="Divide", Keywords="Division,Divisor,/"))
-struct FRigUnit_MathIntDiv : public FRigUnit_MathIntBinaryOp
+struct CONTROLRIG_API FRigUnit_MathIntDiv : public FRigUnit_MathIntBinaryOp
 {
 	GENERATED_BODY()
 
@@ -111,7 +111,7 @@ struct FRigUnit_MathIntDiv : public FRigUnit_MathIntBinaryOp
  * Returns the modulo of the two values
  */
 USTRUCT(meta=(DisplayName="Modulo", PrototypeName="Modulo", Keywords="%,fmod"))
-struct FRigUnit_MathIntMod : public FRigUnit_MathIntBinaryOp
+struct CONTROLRIG_API FRigUnit_MathIntMod : public FRigUnit_MathIntBinaryOp
 {
 	GENERATED_BODY()
 
@@ -129,7 +129,7 @@ struct FRigUnit_MathIntMod : public FRigUnit_MathIntBinaryOp
  * Returns the smaller of the two values
  */
 USTRUCT(meta=(DisplayName="Minimum", PrototypeName="Minimum"))
-struct FRigUnit_MathIntMin : public FRigUnit_MathIntBinaryOp
+struct CONTROLRIG_API FRigUnit_MathIntMin : public FRigUnit_MathIntBinaryOp
 {
 	GENERATED_BODY()
 
@@ -141,7 +141,7 @@ struct FRigUnit_MathIntMin : public FRigUnit_MathIntBinaryOp
  * Returns the larger of the two values
  */
 USTRUCT(meta=(DisplayName="Maximum", PrototypeName="Maximum"))
-struct FRigUnit_MathIntMax : public FRigUnit_MathIntBinaryOp
+struct CONTROLRIG_API FRigUnit_MathIntMax : public FRigUnit_MathIntBinaryOp
 {
 	GENERATED_BODY()
 
@@ -153,7 +153,7 @@ struct FRigUnit_MathIntMax : public FRigUnit_MathIntBinaryOp
  * Returns the value of A raised to the power of B.
  */
 USTRUCT(meta=(DisplayName="Power", PrototypeName="Power"))
-struct FRigUnit_MathIntPow : public FRigUnit_MathIntBinaryOp
+struct CONTROLRIG_API FRigUnit_MathIntPow : public FRigUnit_MathIntBinaryOp
 {
 	GENERATED_BODY()
 
@@ -171,7 +171,7 @@ struct FRigUnit_MathIntPow : public FRigUnit_MathIntBinaryOp
  * Returns the negative value
  */
 USTRUCT(meta=(DisplayName="Negate", PrototypeName="Negate", Keywords="-,Abs"))
-struct FRigUnit_MathIntNegate : public FRigUnit_MathIntUnaryOp
+struct CONTROLRIG_API FRigUnit_MathIntNegate : public FRigUnit_MathIntUnaryOp
 {
 	GENERATED_BODY()
 
@@ -183,7 +183,7 @@ struct FRigUnit_MathIntNegate : public FRigUnit_MathIntUnaryOp
  * Returns the absolute (positive) value
  */
 USTRUCT(meta=(DisplayName="Absolute", PrototypeName="Absolute", Keywords="Abs,Neg"))
-struct FRigUnit_MathIntAbs : public FRigUnit_MathIntUnaryOp
+struct CONTROLRIG_API FRigUnit_MathIntAbs : public FRigUnit_MathIntUnaryOp
 {
 	GENERATED_BODY()
 
@@ -195,7 +195,7 @@ struct FRigUnit_MathIntAbs : public FRigUnit_MathIntUnaryOp
  * Returns the int cast to a float
  */
 USTRUCT(meta=(DisplayName="To Float", PrototypeName="Convert"))
-struct FRigUnit_MathIntToFloat : public FRigUnit_MathIntBase
+struct CONTROLRIG_API FRigUnit_MathIntToFloat : public FRigUnit_MathIntBase
 {
 	GENERATED_BODY()
 
@@ -219,7 +219,7 @@ struct FRigUnit_MathIntToFloat : public FRigUnit_MathIntBase
  * Returns the sign of the value (+1 for >= 0, -1 for < 0)
  */
 USTRUCT(meta=(DisplayName="Sign", PrototypeName="Sign"))
-struct FRigUnit_MathIntSign : public FRigUnit_MathIntUnaryOp
+struct CONTROLRIG_API FRigUnit_MathIntSign : public FRigUnit_MathIntUnaryOp
 {
 	GENERATED_BODY()
 
@@ -231,7 +231,7 @@ struct FRigUnit_MathIntSign : public FRigUnit_MathIntUnaryOp
  * Clamps the given value within the range provided by minimum and maximum
  */
 USTRUCT(meta=(DisplayName="Clamp", PrototypeName="Clamp", Keywords="Range,Remap"))
-struct FRigUnit_MathIntClamp : public FRigUnit_MathIntBase
+struct CONTROLRIG_API FRigUnit_MathIntClamp : public FRigUnit_MathIntBase
 {
 	GENERATED_BODY()
 
@@ -260,7 +260,7 @@ struct FRigUnit_MathIntClamp : public FRigUnit_MathIntBase
  * Returns true if the value A equals B
  */
 USTRUCT(meta=(DisplayName="Equals", PrototypeName="Equals", Keywords="Same,=="))
-struct FRigUnit_MathIntEquals : public FRigUnit_MathIntBase
+struct CONTROLRIG_API FRigUnit_MathIntEquals : public FRigUnit_MathIntBase
 {
 	GENERATED_BODY()
 	RIGVM_METHOD()
@@ -286,7 +286,7 @@ struct FRigUnit_MathIntEquals : public FRigUnit_MathIntBase
  * Returns true if the value A does not equal B
  */
 USTRUCT(meta=(DisplayName="Not Equals", PrototypeName="NotEquals", Keywords="Different,!="))
-struct FRigUnit_MathIntNotEquals : public FRigUnit_MathIntBase
+struct CONTROLRIG_API FRigUnit_MathIntNotEquals : public FRigUnit_MathIntBase
 {
 	GENERATED_BODY()
 	RIGVM_METHOD()
@@ -312,7 +312,7 @@ struct FRigUnit_MathIntNotEquals : public FRigUnit_MathIntBase
  * Returns true if the value A is greater than B
  */
 USTRUCT(meta=(DisplayName="Greater", PrototypeName="Greater", Keywords="Larger,Bigger,>"))
-struct FRigUnit_MathIntGreater : public FRigUnit_MathIntBase
+struct CONTROLRIG_API FRigUnit_MathIntGreater : public FRigUnit_MathIntBase
 {
 	GENERATED_BODY()
 
@@ -339,7 +339,7 @@ struct FRigUnit_MathIntGreater : public FRigUnit_MathIntBase
  * Returns true if the value A is less than B
  */
 USTRUCT(meta=(DisplayName="Less", PrototypeName="Less", Keywords="Smaller,<"))
-struct FRigUnit_MathIntLess : public FRigUnit_MathIntBase
+struct CONTROLRIG_API FRigUnit_MathIntLess : public FRigUnit_MathIntBase
 {
 	GENERATED_BODY()
 	
@@ -366,7 +366,7 @@ struct FRigUnit_MathIntLess : public FRigUnit_MathIntBase
  * Returns true if the value A is greater than or equal to B
  */
 USTRUCT(meta=(DisplayName="Greater Equal", PrototypeName="GreaterEqual", Keywords="Larger,Bigger,>="))
-struct FRigUnit_MathIntGreaterEqual : public FRigUnit_MathIntBase
+struct CONTROLRIG_API FRigUnit_MathIntGreaterEqual : public FRigUnit_MathIntBase
 {
 	GENERATED_BODY()
 
@@ -393,7 +393,7 @@ struct FRigUnit_MathIntGreaterEqual : public FRigUnit_MathIntBase
  * Returns true if the value A is less than or equal to B
  */
 USTRUCT(meta=(DisplayName="Less Equal", PrototypeName="LessEqual", Keywords="Smaller,<="))
-struct FRigUnit_MathIntLessEqual : public FRigUnit_MathIntBase
+struct CONTROLRIG_API FRigUnit_MathIntLessEqual : public FRigUnit_MathIntBase
 {
 	GENERATED_BODY()
 
diff --git a/Engine/Plugins/Experimental/ControlRig/Source/ControlRig/Private/Units/Math/RigUnit_MathQuaternion.h b/Engine/Plugins/Experimental/ControlRig/Source/ControlRig/Private/Units/Math/RigUnit_MathQuaternion.h
--- a/Engine/Plugins/Experimental/ControlRig/Source/ControlRig/Private/Units/Math/RigUnit_MathQuaternion.h
+++ b/Engine/Plugins/Experimental/ControlRig/Source/ControlRig/Private/Units/Math/RigUnit_MathQuaternion.h
@@ -7,13 +7,13 @@
 #include "RigUnit_MathQuaternion.generated.h"
 
 USTRUCT(meta=(Abstract, Category="Math|Quaternion", MenuDescSuffix="(Quaternion)"))
-struct FRigUnit_MathQuaternionBase : public FRigUnit_MathBase
+struct CONTROLRIG_API FRigUnit_MathQuaternionBase : public FRigUnit_MathBase
 {
 	GENERATED_BODY()
 };
 
 USTRUCT(meta=(Abstract))
-struct FRigUnit_MathQuaternionUnaryOp : public FRigUnit_MathQuaternionBase
+struct CONTROLRIG_API FRigUnit_MathQuaternionUnaryOp : public FRigUnit_MathQuaternionBase
 {
 	GENERATED_BODY()
 
@@ -30,7 +30,7 @@ struct FRigUnit_MathQuaternionUnaryOp : public FRigUnit_MathQuaternionBase
 };
 
 USTRUCT(meta=(Abstract))
-struct FRigUnit_MathQuaternionBinaryOp : public FRigUnit_MathQuaternionBase
+struct CONTROLRIG_API FRigUnit_MathQuaternionBinaryOp : public FRigUnit_MathQuaternionBase
 {
 	GENERATED_BODY()
 
@@ -53,7 +53,7 @@ struct FRigUnit_MathQuaternionBinaryOp : public FRigUnit_MathQuaternionBase
  * Makes a quaternion from an axis and an angle in radians
  */
 USTRUCT(meta=(DisplayName="From Axis And Angle", PrototypeName="FromAxisAndAngle", Keywords="Make,Construct"))
-struct FRigUnit_MathQuaternionFromAxisAndAngle : public FRigUnit_MathQuaternionBase
+struct CONTROLRIG_API FRigUnit_MathQuaternionFromAxisAndAngle : public FRigUnit_MathQuaternionBase
 {
 	GENERATED_BODY()
 
@@ -81,7 +81,7 @@ struct FRigUnit_MathQuaternionFromAxisAndAngle : public FRigUnit_MathQuaternionB
  * Makes a quaternion from euler values in degrees
  */
 USTRUCT(meta=(DisplayName="From Euler", PrototypeName="FromEuler", Keywords="Make,Construct"))
-struct FRigUnit_MathQuaternionFromEuler : public FRigUnit_MathQuaternionBase
+struct CONTROLRIG_API FRigUnit_MathQuaternionFromEuler : public FRigUnit_MathQuaternionBase
 {
 	GENERATED_BODY()
 
@@ -109,7 +109,7 @@ struct FRigUnit_MathQuaternionFromEuler : public FRigUnit_MathQuaternionBase
  * Makes a quaternion from a rotator
  */
 USTRUCT(meta=(DisplayName="From Rotator", PrototypeName="FromRotator", Keywords="Make,Construct"))
-struct FRigUnit_MathQuaternionFromRotator : public FRigUnit_MathQuaternionBase
+struct CONTROLRIG_API FRigUnit_MathQuaternionFromRotator : public FRigUnit_MathQuaternionBase
 {
 	GENERATED_BODY()
 	
@@ -133,7 +133,7 @@ struct FRigUnit_MathQuaternionFromRotator : public FRigUnit_MathQuaternionBase
  * Makes a quaternion from two vectors, representing the shortest rotation between the two vectors.
  */
 USTRUCT(meta=(DisplayName="From Two Vectors", PrototypeName="FromTwoVectors", Keywords="Make,Construct"))
-struct FRigUnit_MathQuaternionFromTwoVectors : public FRigUnit_MathQuaternionBase
+struct CONTROLRIG_API FRigUnit_MathQuaternionFromTwoVectors : public FRigUnit_MathQuaternionBase
 {
 	GENERATED_BODY()
 	
@@ -160,7 +160,7 @@ struct FRigUnit_MathQuaternionFromTwoVectors : public FRigUnit_MathQuaternionBas
  * Retrieves the axis and angle of a quaternion in radians
  */
 USTRUCT(meta=(DisplayName="To Axis And Angle", PrototypeName="ToAxisAndAngle", Keywords="Make,Construct,GetAxis,GetAngle"))
-struct FRigUnit_MathQuaternionToAxisAndAngle : public FRigUnit_MathQuaternionBase
+struct CONTROLRIG_API FRigUnit_MathQuaternionToAxisAndAngle : public FRigUnit_MathQuaternionBase
 {
 	GENERATED_BODY()
 
@@ -188,7 +188,7 @@ struct FRigUnit_MathQuaternionToAxisAndAngle : public FRigUnit_MathQuaternionBas
  * Scales a quaternion's angle
  */
 USTRUCT(meta=(DisplayName="Scale", PrototypeName="Scale", Keywords="Multiply,Angle,Scale", Constant))
-struct FRigUnit_MathQuaternionScale : public FRigUnit_MathQuaternionBase
+struct CONTROLRIG_API FRigUnit_MathQuaternionScale : public FRigUnit_MathQuaternionBase
 {
 	GENERATED_BODY()
 
@@ -212,7 +212,7 @@ struct FRigUnit_MathQuaternionScale : public FRigUnit_MathQuaternionBase
  * Retrieves the euler angles in degrees
  */
 USTRUCT(meta=(DisplayName="To Euler", PrototypeName="ToEuler", Keywords="Make,Construct"))
-struct FRigUnit_MathQuaternionToEuler : public FRigUnit_MathQuaternionBase
+struct CONTROLRIG_API FRigUnit_MathQuaternionToEuler : public FRigUnit_MathQuaternionBase
 {
 	GENERATED_BODY()
 
@@ -240,7 +240,7 @@ struct FRigUnit_MathQuaternionToEuler : public FRigUnit_MathQuaternionBase
  * Retrieves the rotator
  */
 USTRUCT(meta=(DisplayName="To Rotator", PrototypeName="ToRotator", Keywords="Make,Construct"))
-struct FRigUnit_MathQuaternionToRotator : public FRigUnit_MathQuaternionBase
+struct CONTROLRIG_API FRigUnit_MathQuaternionToRotator : public FRigUnit_MathQuaternionBase
 {
 	GENERATED_BODY()
 
@@ -264,7 +264,7 @@ struct FRigUnit_MathQuaternionToRotator : public FRigUnit_MathQuaternionBase
  * Returns the product of the two values
  */
 USTRUCT(meta=(DisplayName="Multiply", PrototypeName="Multiply", Keywords="Product,*"))
-struct FRigUnit_MathQuaternionMul : public FRigUnit_MathQuaternionBinaryOp
+struct CONTROLRIG_API FRigUnit_MathQuaternionMul : public FRigUnit_MathQuaternionBinaryOp
 {
 	GENERATED_BODY()
 	RIGVM_METHOD()
@@ -275,7 +275,7 @@ struct FRigUnit_MathQuaternionMul : public FRigUnit_MathQuaternionBinaryOp
  * Returns the negative value
  */
 USTRUCT(meta=(DisplayName="Inverse", PrototypeName="Inverse"))
-struct FRigUnit_MathQuaternionInverse : public FRigUnit_MathQuaternionUnaryOp
+struct CONTROLRIG_API FRigUnit_MathQuaternionInverse : public FRigUnit_MathQuaternionUnaryOp
 {
 	GENERATED_BODY()
 	RIGVM_METHOD()
@@ -286,7 +286,7 @@ struct FRigUnit_MathQuaternionInverse : public FRigUnit_MathQuaternionUnaryOp
  * Linearly interpolates between A and B using the ratio T
  */
 USTRUCT(meta=(DisplayName="Interpolate", PrototypeName="Interpolate", Keywords="Lerp,Mix,Blend,Slerp,SphericalInterpolate"))
-struct FRigUnit_MathQuaternionSlerp : public FRigUnit_MathQuaternionBase
+struct CONTROLRIG_API FRigUnit_MathQuaternionSlerp : public FRigUnit_MathQuaternionBase
 {
 	GENERATED_BODY()
 
@@ -316,7 +316,7 @@ struct FRigUnit_MathQuaternionSlerp : public FRigUnit_MathQuaternionBase
  * Returns true if the value A equals B
  */
 USTRUCT(meta=(DisplayName="Equals", PrototypeName="Equals", Keywords="Same,=="))
-struct FRigUnit_MathQuaternionEquals : public FRigUnit_MathQuaternionBase
+struct CONTROLRIG_API FRigUnit_MathQuaternionEquals : public FRigUnit_MathQuaternionBase
 {
 	GENERATED_BODY()
 
@@ -343,7 +343,7 @@ struct FRigUnit_MathQuaternionEquals : public FRigUnit_MathQuaternionBase
  * Returns true if the value A does not equal B
  */
 USTRUCT(meta=(DisplayName="Not Equals", PrototypeName="NotEquals", Keywords="Different,!="))
-struct FRigUnit_MathQuaternionNotEquals : public FRigUnit_MathQuaternionBase
+struct CONTROLRIG_API FRigUnit_MathQuaternionNotEquals : public FRigUnit_MathQuaternionBase
 {
 	GENERATED_BODY()
 
@@ -370,7 +370,7 @@ struct FRigUnit_MathQuaternionNotEquals : public FRigUnit_MathQuaternionBase
  * Return one of the two values based on the condition
  */
 USTRUCT(meta=(DisplayName="Select", PrototypeName="Select", Keywords="Pick,If", Deprecated = "4.26.0"))
-struct FRigUnit_MathQuaternionSelectBool : public FRigUnit_MathQuaternionBase
+struct CONTROLRIG_API FRigUnit_MathQuaternionSelectBool : public FRigUnit_MathQuaternionBase
 {
 	GENERATED_BODY()
 
@@ -400,7 +400,7 @@ struct FRigUnit_MathQuaternionSelectBool : public FRigUnit_MathQuaternionBase
  * Returns the dot product between two quaternions
  */
 USTRUCT(meta=(DisplayName="Dot", PrototypeName="Dot,|"))
-struct FRigUnit_MathQuaternionDot : public FRigUnit_MathQuaternionBase
+struct CONTROLRIG_API FRigUnit_MathQuaternionDot : public FRigUnit_MathQuaternionBase
 {
 	GENERATED_BODY()
 
@@ -427,7 +427,7 @@ struct FRigUnit_MathQuaternionDot : public FRigUnit_MathQuaternionBase
  * Returns the normalized quaternion
  */
 USTRUCT(meta=(DisplayName="Unit", PrototypeName="Unit", Keywords="Normalize"))
-struct FRigUnit_MathQuaternionUnit : public FRigUnit_MathQuaternionUnaryOp
+struct CONTROLRIG_API FRigUnit_MathQuaternionUnit : public FRigUnit_MathQuaternionUnaryOp
 {
 	GENERATED_BODY()
 	RIGVM_METHOD()
@@ -438,7 +438,7 @@ struct FRigUnit_MathQuaternionUnit : public FRigUnit_MathQuaternionUnaryOp
  * Rotates a given vector by the quaternion
  */
 USTRUCT(meta=(DisplayName="Rotate", PrototypeName="Multiply", Keywords="Transform"))
-struct FRigUnit_MathQuaternionRotateVector : public FRigUnit_MathQuaternionBase
+struct CONTROLRIG_API FRigUnit_MathQuaternionRotateVector : public FRigUnit_MathQuaternionBase
 {
 	GENERATED_BODY()
 
@@ -465,7 +465,7 @@ struct FRigUnit_MathQuaternionRotateVector : public FRigUnit_MathQuaternionBase
  * Rotates a given vector by the quaternion
  */
 USTRUCT(meta = (DisplayName = "Axis", PrototypeName = "Axis", Keywords = "GetAxis,xAxis,yAxis,zAxis"))
-struct FRigUnit_MathQuaternionGetAxis: public FRigUnit_MathQuaternionBase
+struct CONTROLRIG_API FRigUnit_MathQuaternionGetAxis: public FRigUnit_MathQuaternionBase
 {
 	GENERATED_BODY()
 
@@ -494,7 +494,7 @@ struct FRigUnit_MathQuaternionGetAxis: public FRigUnit_MathQuaternionBase
  * Computes the swing and twist components of a quaternion
  */
 USTRUCT(meta = (DisplayName = "To Swing & Twist"))
-struct FRigUnit_MathQuaternionSwingTwist : public FRigUnit_MathQuaternionBase
+struct CONTROLRIG_API FRigUnit_MathQuaternionSwingTwist : public FRigUnit_MathQuaternionBase
 {
 	GENERATED_BODY()
 
@@ -524,7 +524,7 @@ struct FRigUnit_MathQuaternionSwingTwist : public FRigUnit_MathQuaternionBase
  * Enum of possible rotation orders
  */
 USTRUCT(meta = (DisplayName = "Rotation Order", Category = "Math|Rotation", Constant))
-struct FRigUnit_MathQuaternionRotationOrder : public FRigUnit_MathBase
+struct CONTROLRIG_API FRigUnit_MathQuaternionRotationOrder : public FRigUnit_MathBase
 {
 	GENERATED_BODY()
 
diff --git a/Engine/Plugins/Experimental/ControlRig/Source/ControlRig/Private/Units/Math/RigUnit_MathRBFInterpolate.h b/Engine/Plugins/Experimental/ControlRig/Source/ControlRig/Private/Units/Math/RigUnit_MathRBFInterpolate.h
--- a/Engine/Plugins/Experimental/ControlRig/Source/ControlRig/Private/Units/Math/RigUnit_MathRBFInterpolate.h
+++ b/Engine/Plugins/Experimental/ControlRig/Source/ControlRig/Private/Units/Math/RigUnit_MathRBFInterpolate.h
@@ -9,8 +9,8 @@
 
 enum class EControlRigState : uint8;
 template<typename T> class TRBFInterpolator;
-struct FRigUnit_MathRBFQuatWeightFunctor;
-struct FRigUnit_MathRBFVectorWeightFunctor;
+struct CONTROLRIG_API FRigUnit_MathRBFQuatWeightFunctor;
+struct CONTROLRIG_API FRigUnit_MathRBFVectorWeightFunctor;
 
 
 /** Function to use for each target falloff */
@@ -47,7 +47,7 @@ enum class ERBFVectorDistanceType : uint8
 
 
 USTRUCT()
-struct FRigUnit_MathRBFInterpolateQuatWorkData
+struct CONTROLRIG_API FRigUnit_MathRBFInterpolateQuatWorkData
 {
 	GENERATED_BODY()
 
@@ -61,7 +61,7 @@ struct FRigUnit_MathRBFInterpolateQuatWorkData
 };
 
 USTRUCT()
-struct FRigUnit_MathRBFInterpolateVectorWorkData
+struct CONTROLRIG_API FRigUnit_MathRBFInterpolateVectorWorkData
 {
 	GENERATED_BODY()
 
@@ -73,14 +73,14 @@ struct FRigUnit_MathRBFInterpolateVectorWorkData
 
 
 USTRUCT(meta = (Abstract, Category = "Math|RBF Interpolation"))
-struct FRigUnit_MathRBFInterpolateBase : 
+struct CONTROLRIG_API FRigUnit_MathRBFInterpolateBase : 
 	public FRigUnit_MathBase
 {
 	GENERATED_BODY()
 };
 
 USTRUCT(meta = (Abstract, Keywords = "RBF,Interpolate,Quaternion"))
-struct FRigUnit_MathRBFInterpolateQuatBase :
+struct CONTROLRIG_API FRigUnit_MathRBFInterpolateQuatBase :
 	public FRigUnit_MathRBFInterpolateBase
 {
 	GENERATED_BODY()
@@ -125,7 +125,7 @@ protected:
 };
 
 USTRUCT(meta = (Abstract, Keywords = "RBF,Interpolate,Vector"))
-struct FRigUnit_MathRBFInterpolateVectorBase :
+struct CONTROLRIG_API FRigUnit_MathRBFInterpolateVectorBase :
 	public FRigUnit_MathRBFInterpolateBase
 {
 	GENERATED_BODY()
@@ -182,7 +182,7 @@ struct FMathRBFInterpolateQuatFloat_Target
 };
 
 USTRUCT(meta = (Abstract, DisplayName="Quaternion to Float", PrototypeName="RBFInterpolateQuatToFloat"))
-struct FRigUnit_MathRBFInterpolateQuatFloat : 
+struct CONTROLRIG_API FRigUnit_MathRBFInterpolateQuatFloat : 
 	public FRigUnit_MathRBFInterpolateQuatBase
 {
 	GENERATED_BODY()
@@ -211,7 +211,7 @@ struct FMathRBFInterpolateQuatVector_Target
 };
 
 USTRUCT(meta = (Abstract, DisplayName = "Quaternion to Vector", PrototypeName = "RBFInterpolateQuatToVector"))
-struct FRigUnit_MathRBFInterpolateQuatVector :
+struct CONTROLRIG_API FRigUnit_MathRBFInterpolateQuatVector :
 	public FRigUnit_MathRBFInterpolateQuatBase
 {
 	GENERATED_BODY()
@@ -240,7 +240,7 @@ struct FMathRBFInterpolateQuatColor_Target
 };
 
 USTRUCT(meta = (Abstract, DisplayName = "Quaternion to Color", PrototypeName = "RBFInterpolateQuatToColor"))
-struct FRigUnit_MathRBFInterpolateQuatColor :
+struct CONTROLRIG_API FRigUnit_MathRBFInterpolateQuatColor :
 	public FRigUnit_MathRBFInterpolateQuatBase
 {
 	GENERATED_BODY()
@@ -269,7 +269,7 @@ struct FMathRBFInterpolateQuatQuat_Target
 };
 
 USTRUCT(meta = (Abstract, DisplayName = "Quaternion to Quaternion", PrototypeName = "RBFInterpolateQuatToQuat"))
-struct FRigUnit_MathRBFInterpolateQuatQuat :
+struct CONTROLRIG_API FRigUnit_MathRBFInterpolateQuatQuat :
 	public FRigUnit_MathRBFInterpolateQuatBase
 {
 	GENERATED_BODY()
@@ -298,7 +298,7 @@ struct FMathRBFInterpolateQuatXform_Target
 };
 
 USTRUCT(meta = (Abstract, DisplayName = "Quaternion to Transform", PrototypeName = "RBFInterpolateQuatToXform"))
-struct FRigUnit_MathRBFInterpolateQuatXform :
+struct CONTROLRIG_API FRigUnit_MathRBFInterpolateQuatXform :
 	public FRigUnit_MathRBFInterpolateQuatBase
 {
 	GENERATED_BODY()
@@ -329,7 +329,7 @@ struct FMathRBFInterpolateVectorFloat_Target
 };
 
 USTRUCT(meta = (Abstract, DisplayName = "Vector to Float", PrototypeName = "RBFInterpolateVectorToFloat", Keywords = "RBF,Interpolate,Vector"))
-struct FRigUnit_MathRBFInterpolateVectorFloat :
+struct CONTROLRIG_API FRigUnit_MathRBFInterpolateVectorFloat :
 	public FRigUnit_MathRBFInterpolateVectorBase
 {
 	GENERATED_BODY()
@@ -358,7 +358,7 @@ struct FMathRBFInterpolateVectorVector_Target
 };
 
 USTRUCT(meta = (Abstract, DisplayName = "Vector to Vector", PrototypeName = "RBFInterpolateVectorToVector", Keywords = "RBF,Interpolate,Vector"))
-struct FRigUnit_MathRBFInterpolateVectorVector :
+struct CONTROLRIG_API FRigUnit_MathRBFInterpolateVectorVector :
 	public FRigUnit_MathRBFInterpolateVectorBase
 {
 	GENERATED_BODY()
@@ -387,7 +387,7 @@ struct FMathRBFInterpolateVectorColor_Target
 };
 
 USTRUCT(meta = (Abstract, DisplayName = "Vector to Color", PrototypeName = "RBFInterpolateVectorToColor", Keywords = "RBF,Interpolate,Vector"))
-struct FRigUnit_MathRBFInterpolateVectorColor :
+struct CONTROLRIG_API FRigUnit_MathRBFInterpolateVectorColor :
 	public FRigUnit_MathRBFInterpolateVectorBase
 {
 	GENERATED_BODY()
@@ -417,7 +417,7 @@ struct FMathRBFInterpolateVectorQuat_Target
 };
 
 USTRUCT(meta = (Abstract, DisplayName = "Vector to Quat", PrototypeName = "RBFInterpolateVectorToQuat", Keywords = "RBF,Interpolate,Vector"))
-struct FRigUnit_MathRBFInterpolateVectorQuat :
+struct CONTROLRIG_API FRigUnit_MathRBFInterpolateVectorQuat :
 	public FRigUnit_MathRBFInterpolateVectorBase
 {
 	GENERATED_BODY()
@@ -445,7 +445,7 @@ struct FMathRBFInterpolateVectorXform_Target
 };
 
 USTRUCT(meta = (Abstract, DisplayName = "Vector to Transform", PrototypeName = "RBFInterpolateVectorToXform", Keywords = "RBF,Interpolate,Vector"))
-struct FRigUnit_MathRBFInterpolateVectorXform :
+struct CONTROLRIG_API FRigUnit_MathRBFInterpolateVectorXform :
 	public FRigUnit_MathRBFInterpolateVectorBase
 {
 	GENERATED_BODY()
diff --git a/Engine/Plugins/Experimental/ControlRig/Source/ControlRig/Private/Units/Math/RigUnit_MathTransform.h b/Engine/Plugins/Experimental/ControlRig/Source/ControlRig/Private/Units/Math/RigUnit_MathTransform.h
--- a/Engine/Plugins/Experimental/ControlRig/Source/ControlRig/Private/Units/Math/RigUnit_MathTransform.h
+++ b/Engine/Plugins/Experimental/ControlRig/Source/ControlRig/Private/Units/Math/RigUnit_MathTransform.h
@@ -8,13 +8,13 @@
 #include "RigUnit_MathTransform.generated.h"
 
 USTRUCT(meta=(Abstract, Category="Math|Transform", MenuDescSuffix="(Transform)"))
-struct FRigUnit_MathTransformBase : public FRigUnit_MathBase
+struct CONTROLRIG_API FRigUnit_MathTransformBase : public FRigUnit_MathBase
 {
 	GENERATED_BODY()
 };
 
 USTRUCT(meta=(Abstract))
-struct FRigUnit_MathTransformUnaryOp : public FRigUnit_MathTransformBase
+struct CONTROLRIG_API FRigUnit_MathTransformUnaryOp : public FRigUnit_MathTransformBase
 {
 	GENERATED_BODY()
 
@@ -31,7 +31,7 @@ struct FRigUnit_MathTransformUnaryOp : public FRigUnit_MathTransformBase
 };
 
 USTRUCT(meta=(Abstract))
-struct FRigUnit_MathTransformBinaryOp : public FRigUnit_MathTransformBase
+struct CONTROLRIG_API FRigUnit_MathTransformBinaryOp : public FRigUnit_MathTransformBase
 {
 	GENERATED_BODY()
 
@@ -54,7 +54,7 @@ struct FRigUnit_MathTransformBinaryOp : public FRigUnit_MathTransformBase
  * Makes a quaternion based transform from a euler based transform
  */
 USTRUCT(meta=(DisplayName="From Euler Transform", PrototypeName="FromEulerTransform", Keywords="Make,Construct"))
-struct FRigUnit_MathTransformFromEulerTransform : public FRigUnit_MathTransformBase
+struct CONTROLRIG_API FRigUnit_MathTransformFromEulerTransform : public FRigUnit_MathTransformBase
 {
 	GENERATED_BODY()
 
@@ -78,7 +78,7 @@ struct FRigUnit_MathTransformFromEulerTransform : public FRigUnit_MathTransformB
  * Retrieves the euler transform
  */
 USTRUCT(meta=(DisplayName="To Euler Transform", PrototypeName="ToEulerTransform", Keywords="Make,Construct"))
-struct FRigUnit_MathTransformToEulerTransform : public FRigUnit_MathTransformBase
+struct CONTROLRIG_API FRigUnit_MathTransformToEulerTransform : public FRigUnit_MathTransformBase
 {
 	GENERATED_BODY()
 
@@ -102,7 +102,7 @@ struct FRigUnit_MathTransformToEulerTransform : public FRigUnit_MathTransformBas
  * Returns the product of the two values
  */
 USTRUCT(meta=(DisplayName="Multiply", PrototypeName="Multiply", Keywords="Product,*,Global"))
-struct FRigUnit_MathTransformMul : public FRigUnit_MathTransformBinaryOp
+struct CONTROLRIG_API FRigUnit_MathTransformMul : public FRigUnit_MathTransformBinaryOp
 {
 	GENERATED_BODY()
 	RIGVM_METHOD()
@@ -113,7 +113,7 @@ struct FRigUnit_MathTransformMul : public FRigUnit_MathTransformBinaryOp
  * Returns the relative local transform within a parent's transform
  */
 USTRUCT(meta=(DisplayName="Make Relative", PrototypeName="MakeRelative", Keywords="Local,Global,Absolute"))
-struct FRigUnit_MathTransformMakeRelative : public FRigUnit_MathTransformBase
+struct CONTROLRIG_API FRigUnit_MathTransformMakeRelative : public FRigUnit_MathTransformBase
 {
 	GENERATED_BODY()
 
@@ -139,7 +139,7 @@ struct FRigUnit_MathTransformMakeRelative : public FRigUnit_MathTransformBase
  * Returns the absolute global transform within a parent's transform
  */
 USTRUCT(meta = (DisplayName = "Make Absolute", PrototypeName = "MakeAbsolute", Keywords = "Local,Global,Relative"))
-struct FRigUnit_MathTransformMakeAbsolute : public FRigUnit_MathTransformBase
+struct CONTROLRIG_API FRigUnit_MathTransformMakeAbsolute : public FRigUnit_MathTransformBase
 {
 	GENERATED_BODY()
 
@@ -165,7 +165,7 @@ struct FRigUnit_MathTransformMakeAbsolute : public FRigUnit_MathTransformBase
  * Returns the negative value
  */
 USTRUCT(meta=(DisplayName="Inverse", PrototypeName="Inverse"))
-struct FRigUnit_MathTransformInverse : public FRigUnit_MathTransformUnaryOp
+struct CONTROLRIG_API FRigUnit_MathTransformInverse : public FRigUnit_MathTransformUnaryOp
 {
 	GENERATED_BODY()
 	RIGVM_METHOD()
@@ -176,7 +176,7 @@ struct FRigUnit_MathTransformInverse : public FRigUnit_MathTransformUnaryOp
  * Linearly interpolates between A and B using the ratio T
  */
 USTRUCT(meta=(DisplayName="Interpolate", PrototypeName="Interpolate", Keywords="Lerp,Mix,Blend"))
-struct FRigUnit_MathTransformLerp : public FRigUnit_MathTransformBase
+struct CONTROLRIG_API FRigUnit_MathTransformLerp : public FRigUnit_MathTransformBase
 {
 	GENERATED_BODY()
 
@@ -206,7 +206,7 @@ struct FRigUnit_MathTransformLerp : public FRigUnit_MathTransformBase
  * Return one of the two values based on the condition
  */
 USTRUCT(meta=(DisplayName="Select", PrototypeName="Select", Keywords="Pick,If", Deprecated = "4.26.0"))
-struct FRigUnit_MathTransformSelectBool : public FRigUnit_MathTransformBase
+struct CONTROLRIG_API FRigUnit_MathTransformSelectBool : public FRigUnit_MathTransformBase
 {
 	GENERATED_BODY()
 
@@ -236,7 +236,7 @@ struct FRigUnit_MathTransformSelectBool : public FRigUnit_MathTransformBase
  * Rotates a given vector (direction) by the transform
  */
 USTRUCT(meta=(DisplayName="Transform Direction", PrototypeName="Rotate", Keywords="Transform,Direction"))
-struct FRigUnit_MathTransformRotateVector : public FRigUnit_MathTransformBase
+struct CONTROLRIG_API FRigUnit_MathTransformRotateVector : public FRigUnit_MathTransformBase
 {
 	GENERATED_BODY()
 
@@ -263,7 +263,7 @@ struct FRigUnit_MathTransformRotateVector : public FRigUnit_MathTransformBase
  * Rotates a given vector (location) by the transform
  */
 USTRUCT(meta=(DisplayName="Transform Location", PrototypeName="Multiply"))
-struct FRigUnit_MathTransformTransformVector : public FRigUnit_MathTransformBase
+struct CONTROLRIG_API FRigUnit_MathTransformTransformVector : public FRigUnit_MathTransformBase
 {
 	GENERATED_BODY()
 
@@ -290,7 +290,7 @@ struct FRigUnit_MathTransformTransformVector : public FRigUnit_MathTransformBase
  * Composes a Transform (and Euler Transform) from its components.
  */
 USTRUCT(meta=(DisplayName="Transform from SRT", Keywords ="EulerTransform,Scale,Rotation,Orientation,Translation,Location"))
-struct FRigUnit_MathTransformFromSRT : public FRigUnit_MathTransformBase
+struct CONTROLRIG_API FRigUnit_MathTransformFromSRT : public FRigUnit_MathTransformBase
 {
 	GENERATED_BODY()
 
@@ -330,7 +330,7 @@ struct FRigUnit_MathTransformFromSRT : public FRigUnit_MathTransformBase
  * Clamps a position using a plane collision, cylindric collision or spherical collision.
  */
 USTRUCT(meta = (DisplayName = "Clamp Spatially", PrototypeName = "ClampSpatially", Keywords = "Collide,Collision"))
-struct FRigUnit_MathTransformClampSpatially : public FRigUnit_MathTransformBase
+struct CONTROLRIG_API FRigUnit_MathTransformClampSpatially : public FRigUnit_MathTransformBase
 {
 	GENERATED_BODY()
 
diff --git a/Engine/Plugins/Experimental/ControlRig/Source/ControlRig/Private/Units/Math/RigUnit_MathVector.h b/Engine/Plugins/Experimental/ControlRig/Source/ControlRig/Private/Units/Math/RigUnit_MathVector.h
--- a/Engine/Plugins/Experimental/ControlRig/Source/ControlRig/Private/Units/Math/RigUnit_MathVector.h
+++ b/Engine/Plugins/Experimental/ControlRig/Source/ControlRig/Private/Units/Math/RigUnit_MathVector.h
@@ -7,13 +7,13 @@
 #include "RigUnit_MathVector.generated.h"
 
 USTRUCT(meta=(Abstract, Category="Math|Vector", MenuDescSuffix="(Vector)"))
-struct FRigUnit_MathVectorBase : public FRigUnit_MathBase
+struct CONTROLRIG_API FRigUnit_MathVectorBase : public FRigUnit_MathBase
 {
 	GENERATED_BODY()
 };
 
 USTRUCT(meta=(Abstract))
-struct FRigUnit_MathVectorUnaryOp : public FRigUnit_MathVectorBase
+struct CONTROLRIG_API FRigUnit_MathVectorUnaryOp : public FRigUnit_MathVectorBase
 {
 	GENERATED_BODY()
 	
@@ -30,7 +30,7 @@ struct FRigUnit_MathVectorUnaryOp : public FRigUnit_MathVectorBase
 };
 
 USTRUCT(meta=(Abstract))
-struct FRigUnit_MathVectorBinaryOp : public FRigUnit_MathVectorBase
+struct CONTROLRIG_API FRigUnit_MathVectorBinaryOp : public FRigUnit_MathVectorBase
 {
 	GENERATED_BODY()
 
@@ -53,7 +53,7 @@ struct FRigUnit_MathVectorBinaryOp : public FRigUnit_MathVectorBase
  * Makes a vector from a single float
  */
 USTRUCT(meta=(DisplayName="From Float", PrototypeName="FromFloat", Keywords="Make,Construct"))
-struct FRigUnit_MathVectorFromFloat : public FRigUnit_MathVectorBase
+struct CONTROLRIG_API FRigUnit_MathVectorFromFloat : public FRigUnit_MathVectorBase
 {
 	GENERATED_BODY()
 	RIGVM_METHOD()
@@ -76,7 +76,7 @@ struct FRigUnit_MathVectorFromFloat : public FRigUnit_MathVectorBase
  * Returns the sum of the two values
  */
 USTRUCT(meta=(DisplayName="Add", PrototypeName="Add", Keywords="Sum,+"))
-struct FRigUnit_MathVectorAdd : public FRigUnit_MathVectorBinaryOp
+struct CONTROLRIG_API FRigUnit_MathVectorAdd : public FRigUnit_MathVectorBinaryOp
 {
 	GENERATED_BODY()
 
@@ -88,7 +88,7 @@ struct FRigUnit_MathVectorAdd : public FRigUnit_MathVectorBinaryOp
  * Returns the difference of the two values
  */
 USTRUCT(meta=(DisplayName="Subtract", PrototypeName="Subtract", Keywords="-"))
-struct FRigUnit_MathVectorSub : public FRigUnit_MathVectorBinaryOp
+struct CONTROLRIG_API FRigUnit_MathVectorSub : public FRigUnit_MathVectorBinaryOp
 {
 	GENERATED_BODY()
 
@@ -100,7 +100,7 @@ struct FRigUnit_MathVectorSub : public FRigUnit_MathVectorBinaryOp
  * Returns the product of the two values
  */
 USTRUCT(meta=(DisplayName="Multiply", PrototypeName="Multiply", Keywords="Product,*"))
-struct FRigUnit_MathVectorMul : public FRigUnit_MathVectorBinaryOp
+struct CONTROLRIG_API FRigUnit_MathVectorMul : public FRigUnit_MathVectorBinaryOp
 {
 	GENERATED_BODY()
 
@@ -117,7 +117,7 @@ struct FRigUnit_MathVectorMul : public FRigUnit_MathVectorBinaryOp
  * Returns the product of the the vector and the float value
  */
 USTRUCT(meta = (DisplayName = "Scale", PrototypeName = "Scale", Keywords = "Multiply,Product,*,ByScalar,ByFloat"))
-struct FRigUnit_MathVectorScale : public FRigUnit_MathVectorBase
+struct CONTROLRIG_API FRigUnit_MathVectorScale : public FRigUnit_MathVectorBase
 {
 	GENERATED_BODY()
 
@@ -144,7 +144,7 @@ struct FRigUnit_MathVectorScale : public FRigUnit_MathVectorBase
  * Returns the division of the two values
  */
 USTRUCT(meta=(DisplayName="Divide", PrototypeName="Divide", Keywords="Division,Divisor,/"))
-struct FRigUnit_MathVectorDiv : public FRigUnit_MathVectorBinaryOp
+struct CONTROLRIG_API FRigUnit_MathVectorDiv : public FRigUnit_MathVectorBinaryOp
 {
 	GENERATED_BODY()
 
@@ -161,7 +161,7 @@ struct FRigUnit_MathVectorDiv : public FRigUnit_MathVectorBinaryOp
  * Returns the modulo of the two values
  */
 USTRUCT(meta=(DisplayName="Modulo", PrototypeName="Modulo", Keywords="%,fmod"))
-struct FRigUnit_MathVectorMod : public FRigUnit_MathVectorBinaryOp
+struct CONTROLRIG_API FRigUnit_MathVectorMod : public FRigUnit_MathVectorBinaryOp
 {
 	GENERATED_BODY()
 
@@ -179,7 +179,7 @@ struct FRigUnit_MathVectorMod : public FRigUnit_MathVectorBinaryOp
  * Returns the smaller of the two values for each component
  */
 USTRUCT(meta=(DisplayName="Minimum", PrototypeName="Minimum"))
-struct FRigUnit_MathVectorMin : public FRigUnit_MathVectorBinaryOp
+struct CONTROLRIG_API FRigUnit_MathVectorMin : public FRigUnit_MathVectorBinaryOp
 {
 	GENERATED_BODY()
 
@@ -191,7 +191,7 @@ struct FRigUnit_MathVectorMin : public FRigUnit_MathVectorBinaryOp
  * Returns the larger of the two values each component
  */
 USTRUCT(meta=(DisplayName="Maximum", PrototypeName="Maximum"))
-struct FRigUnit_MathVectorMax : public FRigUnit_MathVectorBinaryOp
+struct CONTROLRIG_API FRigUnit_MathVectorMax : public FRigUnit_MathVectorBinaryOp
 {
 	GENERATED_BODY()
 
@@ -203,7 +203,7 @@ struct FRigUnit_MathVectorMax : public FRigUnit_MathVectorBinaryOp
  * Returns the negative value
  */
 USTRUCT(meta=(DisplayName="Negate", PrototypeName="Negate", Keywords="-,Abs"))
-struct FRigUnit_MathVectorNegate : public FRigUnit_MathVectorUnaryOp
+struct CONTROLRIG_API FRigUnit_MathVectorNegate : public FRigUnit_MathVectorUnaryOp
 {
 	GENERATED_BODY()
 
@@ -215,7 +215,7 @@ struct FRigUnit_MathVectorNegate : public FRigUnit_MathVectorUnaryOp
  * Returns the absolute (positive) value
  */
 USTRUCT(meta=(DisplayName="Absolute", PrototypeName="Absolute", Keywords="Abs,Neg"))
-struct FRigUnit_MathVectorAbs : public FRigUnit_MathVectorUnaryOp
+struct CONTROLRIG_API FRigUnit_MathVectorAbs : public FRigUnit_MathVectorUnaryOp
 {
 	GENERATED_BODY()
 
@@ -227,7 +227,7 @@ struct FRigUnit_MathVectorAbs : public FRigUnit_MathVectorUnaryOp
  * Returns the closest lower full number (integer) of the value for each component
  */
 USTRUCT(meta=(DisplayName="Floor", PrototypeName="Floor", Keywords="Round"))
-struct FRigUnit_MathVectorFloor : public FRigUnit_MathVectorUnaryOp
+struct CONTROLRIG_API FRigUnit_MathVectorFloor : public FRigUnit_MathVectorUnaryOp
 {
 	GENERATED_BODY()
 
@@ -239,7 +239,7 @@ struct FRigUnit_MathVectorFloor : public FRigUnit_MathVectorUnaryOp
  * Returns the closest higher full number (integer) of the value for each component
  */
 USTRUCT(meta=(DisplayName="Ceiling", PrototypeName="Ceiling", Keywords="Round"))
-struct FRigUnit_MathVectorCeil : public FRigUnit_MathVectorUnaryOp
+struct CONTROLRIG_API FRigUnit_MathVectorCeil : public FRigUnit_MathVectorUnaryOp
 {
 	GENERATED_BODY()
 
@@ -251,7 +251,7 @@ struct FRigUnit_MathVectorCeil : public FRigUnit_MathVectorUnaryOp
  * Returns the closest higher full number (integer) of the value for each component
  */
 USTRUCT(meta=(DisplayName="Round", PrototypeName="Round"))
-struct FRigUnit_MathVectorRound : public FRigUnit_MathVectorUnaryOp
+struct CONTROLRIG_API FRigUnit_MathVectorRound : public FRigUnit_MathVectorUnaryOp
 {
 	GENERATED_BODY()
 
@@ -263,7 +263,7 @@ struct FRigUnit_MathVectorRound : public FRigUnit_MathVectorUnaryOp
  * Returns the sign of the value (+1 for >= FVector(0.f, 0.f, 0.f), -1 for < 0.f) for each component
  */
 USTRUCT(meta=(DisplayName="Sign", PrototypeName="Sign"))
-struct FRigUnit_MathVectorSign : public FRigUnit_MathVectorUnaryOp
+struct CONTROLRIG_API FRigUnit_MathVectorSign : public FRigUnit_MathVectorUnaryOp
 {
 	GENERATED_BODY()
 
@@ -275,7 +275,7 @@ struct FRigUnit_MathVectorSign : public FRigUnit_MathVectorUnaryOp
  * Clamps the given value within the range provided by minimum and maximum for each component
  */
 USTRUCT(meta=(DisplayName="Clamp", PrototypeName="Clamp", Keywords="Range,Remap"))
-struct FRigUnit_MathVectorClamp : public FRigUnit_MathVectorBase
+struct CONTROLRIG_API FRigUnit_MathVectorClamp : public FRigUnit_MathVectorBase
 {
 	GENERATED_BODY()
 
@@ -305,7 +305,7 @@ struct FRigUnit_MathVectorClamp : public FRigUnit_MathVectorBase
  * Linearly interpolates between A and B using the ratio T
  */
 USTRUCT(meta=(DisplayName="Interpolate", PrototypeName="Interpolate", Keywords="Lerp,Mix,Blend"))
-struct FRigUnit_MathVectorLerp : public FRigUnit_MathVectorBase
+struct CONTROLRIG_API FRigUnit_MathVectorLerp : public FRigUnit_MathVectorBase
 {
 	GENERATED_BODY()
 
@@ -336,7 +336,7 @@ struct FRigUnit_MathVectorLerp : public FRigUnit_MathVectorBase
  * Remaps the given value from a source range to a target range for each component
  */
 USTRUCT(meta=(DisplayName="Remap", PrototypeName="Remap", Keywords="Rescale,Scale"))
-struct FRigUnit_MathVectorRemap : public FRigUnit_MathVectorBase
+struct CONTROLRIG_API FRigUnit_MathVectorRemap : public FRigUnit_MathVectorBase
 {
 	GENERATED_BODY()
 
@@ -377,7 +377,7 @@ struct FRigUnit_MathVectorRemap : public FRigUnit_MathVectorBase
  * Returns true if the value A equals B
  */
 USTRUCT(meta=(DisplayName="Equals", PrototypeName="Equals", Keywords="Same,=="))
-struct FRigUnit_MathVectorEquals : public FRigUnit_MathVectorBase
+struct CONTROLRIG_API FRigUnit_MathVectorEquals : public FRigUnit_MathVectorBase
 {
 	GENERATED_BODY()
 
@@ -404,7 +404,7 @@ struct FRigUnit_MathVectorEquals : public FRigUnit_MathVectorBase
  * Returns true if the value A does not equal B
  */
 USTRUCT(meta=(DisplayName="Not Equals", PrototypeName="NotEquals", Keywords="Different,!="))
-struct FRigUnit_MathVectorNotEquals : public FRigUnit_MathVectorBase
+struct CONTROLRIG_API FRigUnit_MathVectorNotEquals : public FRigUnit_MathVectorBase
 {
 	GENERATED_BODY()
 
@@ -431,7 +431,7 @@ struct FRigUnit_MathVectorNotEquals : public FRigUnit_MathVectorBase
  * Returns true if the value is nearly zero
  */
 USTRUCT(meta=(DisplayName="Is Nearly Zero", PrototypeName="IsNearlyZero", Keywords="AlmostZero,0"))
-struct FRigUnit_MathVectorIsNearlyZero : public FRigUnit_MathVectorBase
+struct CONTROLRIG_API FRigUnit_MathVectorIsNearlyZero : public FRigUnit_MathVectorBase
 {
 	GENERATED_BODY()
 
@@ -459,7 +459,7 @@ struct FRigUnit_MathVectorIsNearlyZero : public FRigUnit_MathVectorBase
  * Returns true if the value A is almost equal to B
  */
 USTRUCT(meta=(DisplayName="Is Nearly Equal", PrototypeName="IsNearlyEqual", Keywords="AlmostEqual,=="))
-struct FRigUnit_MathVectorIsNearlyEqual : public FRigUnit_MathVectorBase
+struct CONTROLRIG_API FRigUnit_MathVectorIsNearlyEqual : public FRigUnit_MathVectorBase
 {
 	GENERATED_BODY()
 
@@ -490,7 +490,7 @@ struct FRigUnit_MathVectorIsNearlyEqual : public FRigUnit_MathVectorBase
  * Return one of the two values based on the condition
  */
 USTRUCT(meta=(DisplayName="Select", PrototypeName="Select", Keywords="Pick,If", Deprecated = "4.26.0"))
-struct FRigUnit_MathVectorSelectBool : public FRigUnit_MathVectorBase
+struct CONTROLRIG_API FRigUnit_MathVectorSelectBool : public FRigUnit_MathVectorBase
 {
 	GENERATED_BODY()
 
@@ -520,7 +520,7 @@ struct FRigUnit_MathVectorSelectBool : public FRigUnit_MathVectorBase
  * Returns the degrees of a given value in radians
  */
 USTRUCT(meta=(DisplayName="Degrees", PrototypeName="Degrees"))
-struct FRigUnit_MathVectorDeg : public FRigUnit_MathVectorUnaryOp
+struct CONTROLRIG_API FRigUnit_MathVectorDeg : public FRigUnit_MathVectorUnaryOp
 {
 	GENERATED_BODY()
 
@@ -532,7 +532,7 @@ struct FRigUnit_MathVectorDeg : public FRigUnit_MathVectorUnaryOp
  * Returns the radians of a given value in degrees
  */
 USTRUCT(meta=(DisplayName="Radians", PrototypeName="Radians"))
-struct FRigUnit_MathVectorRad : public FRigUnit_MathVectorUnaryOp
+struct CONTROLRIG_API FRigUnit_MathVectorRad : public FRigUnit_MathVectorUnaryOp
 {
 	GENERATED_BODY()
 
@@ -544,7 +544,7 @@ struct FRigUnit_MathVectorRad : public FRigUnit_MathVectorUnaryOp
  * Returns the squared length of the vector
  */
 USTRUCT(meta=(DisplayName="Length Squared", PrototypeName="LengthSquared", Keywords="Length,Size,Magnitude"))
-struct FRigUnit_MathVectorLengthSquared : public FRigUnit_MathVectorBase
+struct CONTROLRIG_API FRigUnit_MathVectorLengthSquared : public FRigUnit_MathVectorBase
 {
 	GENERATED_BODY()
 
@@ -568,7 +568,7 @@ struct FRigUnit_MathVectorLengthSquared : public FRigUnit_MathVectorBase
  * Returns the length of the vector
  */
 USTRUCT(meta=(DisplayName="Length", PrototypeName="Length", Keywords="Size,Magnitude"))
-struct FRigUnit_MathVectorLength : public FRigUnit_MathVectorBase
+struct CONTROLRIG_API FRigUnit_MathVectorLength : public FRigUnit_MathVectorBase
 {
 	GENERATED_BODY()
 
@@ -592,7 +592,7 @@ struct FRigUnit_MathVectorLength : public FRigUnit_MathVectorBase
  * Returns the distance from A to B
  */
 USTRUCT(meta=(DisplayName="Distance Between", PrototypeName="Distance"))
-struct FRigUnit_MathVectorDistance : public FRigUnit_MathVectorBase
+struct CONTROLRIG_API FRigUnit_MathVectorDistance : public FRigUnit_MathVectorBase
 {
 	GENERATED_BODY()
 
@@ -619,7 +619,7 @@ struct FRigUnit_MathVectorDistance : public FRigUnit_MathVectorBase
  * Returns the cross product between two vectors
  */
 USTRUCT(meta=(DisplayName="Cross", PrototypeName="Cross", Keywords="^"))
-struct FRigUnit_MathVectorCross : public FRigUnit_MathVectorBinaryOp
+struct CONTROLRIG_API FRigUnit_MathVectorCross : public FRigUnit_MathVectorBinaryOp
 {
 	GENERATED_BODY()
 	RIGVM_METHOD()
@@ -630,7 +630,7 @@ struct FRigUnit_MathVectorCross : public FRigUnit_MathVectorBinaryOp
  * Returns the dot product between two vectors
  */
 USTRUCT(meta=(DisplayName="Dot", PrototypeName="Dot,|"))
-struct FRigUnit_MathVectorDot : public FRigUnit_MathVectorBase
+struct CONTROLRIG_API FRigUnit_MathVectorDot : public FRigUnit_MathVectorBase
 {
 	GENERATED_BODY()
 
@@ -657,7 +657,7 @@ struct FRigUnit_MathVectorDot : public FRigUnit_MathVectorBase
  * Returns the normalized value
  */
 USTRUCT(meta=(DisplayName="Unit", PrototypeName="Unit", Keywords="Normalize"))
-struct FRigUnit_MathVectorUnit : public FRigUnit_MathVectorUnaryOp
+struct CONTROLRIG_API FRigUnit_MathVectorUnit : public FRigUnit_MathVectorUnaryOp
 {
 	GENERATED_BODY()
 
@@ -669,7 +669,7 @@ struct FRigUnit_MathVectorUnit : public FRigUnit_MathVectorUnaryOp
  * Sets the length of a given vector
  */
 USTRUCT(meta = (DisplayName = "SetLength", PrototypeName = "SetLength", Keywords = "Unit,Normalize,Scale"))
-struct FRigUnit_MathVectorSetLength: public FRigUnit_MathVectorBase
+struct CONTROLRIG_API FRigUnit_MathVectorSetLength: public FRigUnit_MathVectorBase
 {
 	GENERATED_BODY()
 
@@ -696,7 +696,7 @@ struct FRigUnit_MathVectorSetLength: public FRigUnit_MathVectorBase
  * Clamps the length of a given vector between a minimum and maximum
  */
 USTRUCT(meta = (DisplayName = "ClampLength", PrototypeName = "ClampLength", Keywords = "Unit,Normalize,Scale"))
-struct FRigUnit_MathVectorClampLength: public FRigUnit_MathVectorBase
+struct CONTROLRIG_API FRigUnit_MathVectorClampLength: public FRigUnit_MathVectorBase
 {
 	GENERATED_BODY()
 
@@ -728,7 +728,7 @@ struct FRigUnit_MathVectorClampLength: public FRigUnit_MathVectorBase
  * Mirror a vector about a normal vector.
  */
 USTRUCT(meta=(DisplayName="Mirror", PrototypeName="Mirror"))
-struct FRigUnit_MathVectorMirror : public FRigUnit_MathVectorBase
+struct CONTROLRIG_API FRigUnit_MathVectorMirror : public FRigUnit_MathVectorBase
 {
 	GENERATED_BODY()
 
@@ -755,7 +755,7 @@ struct FRigUnit_MathVectorMirror : public FRigUnit_MathVectorBase
  * Returns the angle between two vectors in radians
  */
 USTRUCT(meta=(DisplayName="Angle Between", PrototypeName="AngleBetween"))
-struct FRigUnit_MathVectorAngle : public FRigUnit_MathVectorBase
+struct CONTROLRIG_API FRigUnit_MathVectorAngle : public FRigUnit_MathVectorBase
 {
 	GENERATED_BODY()
 
@@ -782,7 +782,7 @@ struct FRigUnit_MathVectorAngle : public FRigUnit_MathVectorBase
  * Returns true if the two vectors are parallel
  */
 USTRUCT(meta=(DisplayName="Parallel", PrototypeName="Parallel"))
-struct FRigUnit_MathVectorParallel : public FRigUnit_MathVectorBase
+struct CONTROLRIG_API FRigUnit_MathVectorParallel : public FRigUnit_MathVectorBase
 {
 	GENERATED_BODY()
 
@@ -809,7 +809,7 @@ struct FRigUnit_MathVectorParallel : public FRigUnit_MathVectorBase
  * Returns true if the two vectors are orthogonal
  */
 USTRUCT(meta=(DisplayName="Orthogonal", PrototypeName="Orthogonal"))
-struct FRigUnit_MathVectorOrthogonal : public FRigUnit_MathVectorBase
+struct CONTROLRIG_API FRigUnit_MathVectorOrthogonal : public FRigUnit_MathVectorBase
 {
 	GENERATED_BODY()
 
@@ -837,7 +837,7 @@ struct FRigUnit_MathVectorOrthogonal : public FRigUnit_MathVectorBase
  * Returns the 4 point bezier interpolation
  */
 USTRUCT(meta=(DisplayName="Bezier Four Point"))
-struct FRigUnit_MathVectorBezierFourPoint : public FRigUnit_MathVectorBase
+struct CONTROLRIG_API FRigUnit_MathVectorBezierFourPoint : public FRigUnit_MathVectorBase
 {
 	GENERATED_BODY()
 
@@ -868,7 +868,7 @@ struct FRigUnit_MathVectorBezierFourPoint : public FRigUnit_MathVectorBase
  * Creates a bezier four point
  */
 USTRUCT(meta = (DisplayName = "Make Bezier Four Point", Constant))
-struct FRigUnit_MathVectorMakeBezierFourPoint : public FRigUnit_MathVectorBase
+struct CONTROLRIG_API FRigUnit_MathVectorMakeBezierFourPoint : public FRigUnit_MathVectorBase
 {
 	GENERATED_BODY()
 
@@ -888,7 +888,7 @@ struct FRigUnit_MathVectorMakeBezierFourPoint : public FRigUnit_MathVectorBase
  * Clamps a position using a plane collision, cylindric collision or spherical collision.
  */
 USTRUCT(meta = (DisplayName = "Clamp Spatially", PrototypeName = "ClampSpatially", Keywords="Collide,Collision"))
-struct FRigUnit_MathVectorClampSpatially: public FRigUnit_MathVectorBase
+struct CONTROLRIG_API FRigUnit_MathVectorClampSpatially: public FRigUnit_MathVectorBase
 {
 	GENERATED_BODY()
 
@@ -945,7 +945,7 @@ struct FRigUnit_MathVectorClampSpatially: public FRigUnit_MathVectorBase
  * Intersects a plane with a vector given a start and direction
  */
 USTRUCT(meta = (DisplayName = "Intersect Plane", Keywords = "Collide,Intersect,Raycast"))
-struct FRigUnit_MathIntersectPlane : public FRigUnit_MathVectorBase
+struct CONTROLRIG_API FRigUnit_MathIntersectPlane : public FRigUnit_MathVectorBase
 {
 	GENERATED_BODY()
 
diff --git a/Engine/Plugins/Experimental/ControlRig/Source/ControlRig/Private/Units/Math/RigUnit_Noise.h b/Engine/Plugins/Experimental/ControlRig/Source/ControlRig/Private/Units/Math/RigUnit_Noise.h
--- a/Engine/Plugins/Experimental/ControlRig/Source/ControlRig/Private/Units/Math/RigUnit_Noise.h
+++ b/Engine/Plugins/Experimental/ControlRig/Source/ControlRig/Private/Units/Math/RigUnit_Noise.h
@@ -9,7 +9,7 @@
  * Generates a float through a noise fluctuation process between a min and a max through speed
  */
 USTRUCT(meta=(DisplayName="Noise (Float)", Category="Math|Noise", PrototypeName="Noise"))
-struct FRigUnit_NoiseFloat : public FRigUnit_MathBase
+struct CONTROLRIG_API FRigUnit_NoiseFloat : public FRigUnit_MathBase
 {
 	GENERATED_BODY()
 
@@ -49,7 +49,7 @@ struct FRigUnit_NoiseFloat : public FRigUnit_MathBase
  * Generates a vector through a noise fluctuation process between a min and a max through speed
  */
 USTRUCT(meta = (DisplayName = "Noise (Vector)", Category = "Math|Noise", PrototypeName = "Noise"))
-struct FRigUnit_NoiseVector : public FRigUnit_MathBase
+struct CONTROLRIG_API FRigUnit_NoiseVector : public FRigUnit_MathBase
 {
 	GENERATED_BODY()
 
diff --git a/Engine/Plugins/Experimental/ControlRig/Source/ControlRig/Private/Units/Math/RigUnit_Random.h b/Engine/Plugins/Experimental/ControlRig/Source/ControlRig/Private/Units/Math/RigUnit_Random.h
--- a/Engine/Plugins/Experimental/ControlRig/Source/ControlRig/Private/Units/Math/RigUnit_Random.h
+++ b/Engine/Plugins/Experimental/ControlRig/Source/ControlRig/Private/Units/Math/RigUnit_Random.h
@@ -9,7 +9,7 @@
  * Generates a random float between a min and a max
  */
 USTRUCT(meta=(DisplayName="Random (Float)", Category="Math|Random", PrototypeName="Random"))
-struct FRigUnit_RandomFloat : public FRigUnit_MathBase
+struct CONTROLRIG_API FRigUnit_RandomFloat : public FRigUnit_MathBase
 {
 	GENERATED_BODY()
 
@@ -55,7 +55,7 @@ struct FRigUnit_RandomFloat : public FRigUnit_MathBase
  * Generates a random vector between a min and a max
  */
 USTRUCT(meta = (DisplayName = "Random (Vector)", Category = "Math|Random", PrototypeName = "Random"))
-struct FRigUnit_RandomVector: public FRigUnit_MathBase
+struct CONTROLRIG_API FRigUnit_RandomVector: public FRigUnit_MathBase
 {
 	GENERATED_BODY()
 
diff --git a/Engine/Plugins/Experimental/ControlRig/Source/ControlRig/Private/Units/Simulation/RigUnit_Accumulate.h b/Engine/Plugins/Experimental/ControlRig/Source/ControlRig/Private/Units/Simulation/RigUnit_Accumulate.h
--- a/Engine/Plugins/Experimental/ControlRig/Source/ControlRig/Private/Units/Simulation/RigUnit_Accumulate.h
+++ b/Engine/Plugins/Experimental/ControlRig/Source/ControlRig/Private/Units/Simulation/RigUnit_Accumulate.h
@@ -9,7 +9,7 @@
  * Adds a value over time over and over again
  */
 USTRUCT(meta=(DisplayName="Accumulate Add (Float)", PrototypeName="AccumulateAdd", Keywords="Simulate,++"))
-struct FRigUnit_AccumulateFloatAdd : public FRigUnit_SimBase
+struct CONTROLRIG_API FRigUnit_AccumulateFloatAdd : public FRigUnit_SimBase
 {
 	GENERATED_BODY()
 	
@@ -42,7 +42,7 @@ struct FRigUnit_AccumulateFloatAdd : public FRigUnit_SimBase
  * Adds a vector over time over and over again
  */
 USTRUCT(meta = (DisplayName="Accumulate Add (Vector)", PrototypeName="AccumulateAdd", Keywords="Simulate,++"))
-struct FRigUnit_AccumulateVectorAdd : public FRigUnit_SimBase
+struct CONTROLRIG_API FRigUnit_AccumulateVectorAdd : public FRigUnit_SimBase
 {
 	GENERATED_BODY()
 
@@ -75,7 +75,7 @@ struct FRigUnit_AccumulateVectorAdd : public FRigUnit_SimBase
  * Multiplies a value over time over and over again
  */
 USTRUCT(meta=(DisplayName="Accumulate Mul (Float)", PrototypeName="AccumulateMul", Keywords="Simulate,**"))
-struct FRigUnit_AccumulateFloatMul : public FRigUnit_SimBase
+struct CONTROLRIG_API FRigUnit_AccumulateFloatMul : public FRigUnit_SimBase
 {
 	GENERATED_BODY()
 	
@@ -108,7 +108,7 @@ struct FRigUnit_AccumulateFloatMul : public FRigUnit_SimBase
  * Multiplies a vector over time over and over again
  */
 USTRUCT(meta = (DisplayName="Accumulate Mul (Vector)", PrototypeName="AccumulateMul", Keywords="Simulate,**"))
-struct FRigUnit_AccumulateVectorMul : public FRigUnit_SimBase
+struct CONTROLRIG_API FRigUnit_AccumulateVectorMul : public FRigUnit_SimBase
 {
 	GENERATED_BODY()
 
@@ -141,7 +141,7 @@ struct FRigUnit_AccumulateVectorMul : public FRigUnit_SimBase
  * Multiplies a quaternion over time over and over again
  */
 USTRUCT(meta=(DisplayName="Accumulate Mul (Quaternion)", PrototypeName="AccumulateMul", Keywords="Simulate,**"))
-struct FRigUnit_AccumulateQuatMul : public FRigUnit_SimBase
+struct CONTROLRIG_API FRigUnit_AccumulateQuatMul : public FRigUnit_SimBase
 {
 	GENERATED_BODY()
 	
@@ -177,7 +177,7 @@ struct FRigUnit_AccumulateQuatMul : public FRigUnit_SimBase
  * Multiplies a transform over time over and over again
  */
 USTRUCT(meta = (DisplayName="Accumulate Mul (Transform)", PrototypeName="AccumulateMul", Keywords="Simulate,**"))
-struct FRigUnit_AccumulateTransformMul : public FRigUnit_SimBase
+struct CONTROLRIG_API FRigUnit_AccumulateTransformMul : public FRigUnit_SimBase
 {
 	GENERATED_BODY()
 
@@ -213,7 +213,7 @@ struct FRigUnit_AccumulateTransformMul : public FRigUnit_SimBase
  * Interpolates two values over time over and over again
  */
 USTRUCT(meta=(DisplayName="Accumulate Lerp (Float)", PrototypeName="AccumulateLerp", Keywords="Simulate,Ramp"))
-struct FRigUnit_AccumulateFloatLerp : public FRigUnit_SimBase
+struct CONTROLRIG_API FRigUnit_AccumulateFloatLerp : public FRigUnit_SimBase
 {
 	GENERATED_BODY()
 	
@@ -249,7 +249,7 @@ struct FRigUnit_AccumulateFloatLerp : public FRigUnit_SimBase
  * Interpolates two vectors over time over and over again
  */
 USTRUCT(meta = (DisplayName="Accumulate Lerp (Vector)", PrototypeName="AccumulateLerp", Keywords="Simulate,Ramp"))
-struct FRigUnit_AccumulateVectorLerp : public FRigUnit_SimBase
+struct CONTROLRIG_API FRigUnit_AccumulateVectorLerp : public FRigUnit_SimBase
 {
 	GENERATED_BODY()
 
@@ -286,7 +286,7 @@ struct FRigUnit_AccumulateVectorLerp : public FRigUnit_SimBase
  * Interpolates two quaternions over time over and over again
  */
 USTRUCT(meta=(DisplayName="Accumulate Lerp (Quaternion)", PrototypeName="AccumulateLerp", Keywords="Simulate,Ramp"))
-struct FRigUnit_AccumulateQuatLerp : public FRigUnit_SimBase
+struct CONTROLRIG_API FRigUnit_AccumulateQuatLerp : public FRigUnit_SimBase
 {
 	GENERATED_BODY()
 	
@@ -323,7 +323,7 @@ struct FRigUnit_AccumulateQuatLerp : public FRigUnit_SimBase
  * Interpolates two transforms over time over and over again
  */
 USTRUCT(meta = (DisplayName="Accumulate Lerp (Transform)", PrototypeName="AccumulateLerp", Keywords="Simulate,Ramp"))
-struct FRigUnit_AccumulateTransformLerp : public FRigUnit_SimBase
+struct CONTROLRIG_API FRigUnit_AccumulateTransformLerp : public FRigUnit_SimBase
 {
 	GENERATED_BODY()
 
@@ -360,7 +360,7 @@ struct FRigUnit_AccumulateTransformLerp : public FRigUnit_SimBase
  * Accumulates the min and max values over time
  */
 USTRUCT(meta=(DisplayName="Accumulate Range (Float)", PrototypeName="AccumulateRange", Keywords="Range"))
-struct FRigUnit_AccumulateFloatRange : public FRigUnit_SimBase
+struct CONTROLRIG_API FRigUnit_AccumulateFloatRange : public FRigUnit_SimBase
 {
 	GENERATED_BODY()
 	
@@ -392,7 +392,7 @@ struct FRigUnit_AccumulateFloatRange : public FRigUnit_SimBase
  * Accumulates the min and max values over time
  */
 USTRUCT(meta = (DisplayName="Accumulate Range (Vector)", PrototypeName="AccumulateRange", Keywords="Range"))
-struct FRigUnit_AccumulateVectorRange : public FRigUnit_SimBase
+struct CONTROLRIG_API FRigUnit_AccumulateVectorRange : public FRigUnit_SimBase
 {
 	GENERATED_BODY()
 
diff --git a/Engine/Plugins/Experimental/ControlRig/Source/ControlRig/Private/Units/Simulation/RigUnit_AlphaInterp.h b/Engine/Plugins/Experimental/ControlRig/Source/ControlRig/Private/Units/Simulation/RigUnit_AlphaInterp.h
--- a/Engine/Plugins/Experimental/ControlRig/Source/ControlRig/Private/Units/Simulation/RigUnit_AlphaInterp.h
+++ b/Engine/Plugins/Experimental/ControlRig/Source/ControlRig/Private/Units/Simulation/RigUnit_AlphaInterp.h
@@ -10,7 +10,7 @@
  * Adds a float value over time over and over again
  */
 USTRUCT(meta=(DisplayName="Interpolate", Keywords="Alpha,Lerp,LinearInterpolate", PrototypeName = "AlphaInterp"))
-struct FRigUnit_AlphaInterp : public FRigUnit_SimBase
+struct CONTROLRIG_API FRigUnit_AlphaInterp : public FRigUnit_SimBase
 {
 	GENERATED_BODY()
 	
@@ -83,7 +83,7 @@ struct FRigUnit_AlphaInterp : public FRigUnit_SimBase
  * Adds a vector value over time over and over again
  */
 USTRUCT(meta=(DisplayName="Interpolate", Keywords="Alpha,Lerp,LinearInterpolate", PrototypeName = "AlphaInterp", MenuDescSuffix = "(Vector)"))
-struct FRigUnit_AlphaInterpVector : public FRigUnit_SimBase
+struct CONTROLRIG_API FRigUnit_AlphaInterpVector : public FRigUnit_SimBase
 {
 	GENERATED_BODY()
 	
diff --git a/Engine/Plugins/Experimental/ControlRig/Source/ControlRig/Private/Units/Simulation/RigUnit_DeltaFromPrevious.h b/Engine/Plugins/Experimental/ControlRig/Source/ControlRig/Private/Units/Simulation/RigUnit_DeltaFromPrevious.h
--- a/Engine/Plugins/Experimental/ControlRig/Source/ControlRig/Private/Units/Simulation/RigUnit_DeltaFromPrevious.h
+++ b/Engine/Plugins/Experimental/ControlRig/Source/ControlRig/Private/Units/Simulation/RigUnit_DeltaFromPrevious.h
@@ -9,7 +9,7 @@
  * Computes the difference from the previous value going through the node
  */
 USTRUCT(meta=(DisplayName="DeltaFromPrevious (Float)", PrototypeName="DeltaFromPrevious", Keywords="Difference,Velocity,Acceleration"))
-struct FRigUnit_DeltaFromPreviousFloat : public FRigUnit_SimBase
+struct CONTROLRIG_API FRigUnit_DeltaFromPreviousFloat : public FRigUnit_SimBase
 {
 	GENERATED_BODY()
 	
@@ -38,7 +38,7 @@ struct FRigUnit_DeltaFromPreviousFloat : public FRigUnit_SimBase
  * Computes the difference from the previous value going through the node
  */
 USTRUCT(meta=(DisplayName="DeltaFromPrevious (Vector)", PrototypeName="DeltaFromPrevious", Keywords="Difference,Velocity,Acceleration"))
-struct FRigUnit_DeltaFromPreviousVector : public FRigUnit_SimBase
+struct CONTROLRIG_API FRigUnit_DeltaFromPreviousVector : public FRigUnit_SimBase
 {
 	GENERATED_BODY()
 	
@@ -67,7 +67,7 @@ struct FRigUnit_DeltaFromPreviousVector : public FRigUnit_SimBase
  * Computes the difference from the previous value going through the node
  */
 USTRUCT(meta=(DisplayName="DeltaFromPrevious (Quaternion)", PrototypeName="DeltaFromPrevious", Keywords="Difference,Velocity,Acceleration"))
-struct FRigUnit_DeltaFromPreviousQuat : public FRigUnit_SimBase
+struct CONTROLRIG_API FRigUnit_DeltaFromPreviousQuat : public FRigUnit_SimBase
 {
 	GENERATED_BODY()
 	
@@ -96,7 +96,7 @@ struct FRigUnit_DeltaFromPreviousQuat : public FRigUnit_SimBase
  * Computes the difference from the previous value going through the node
  */
 USTRUCT(meta=(DisplayName="DeltaFromPrevious (Transform)", PrototypeName="DeltaFromPrevious", Keywords="Difference,Velocity,Acceleration"))
-struct FRigUnit_DeltaFromPreviousTransform : public FRigUnit_SimBase
+struct CONTROLRIG_API FRigUnit_DeltaFromPreviousTransform : public FRigUnit_SimBase
 {
 	GENERATED_BODY()
 	
diff --git a/Engine/Plugins/Experimental/ControlRig/Source/ControlRig/Private/Units/Simulation/RigUnit_Kalman.h b/Engine/Plugins/Experimental/ControlRig/Source/ControlRig/Private/Units/Simulation/RigUnit_Kalman.h
--- a/Engine/Plugins/Experimental/ControlRig/Source/ControlRig/Private/Units/Simulation/RigUnit_Kalman.h
+++ b/Engine/Plugins/Experimental/ControlRig/Source/ControlRig/Private/Units/Simulation/RigUnit_Kalman.h
@@ -10,7 +10,7 @@
  * This uses a Kalman Filter internally.
  */
 USTRUCT(meta=(DisplayName="Average Over Time (Float)", PrototypeName="KalmanFilter", Keywords="Average,Smooth,KalmanFilter"))
-struct FRigUnit_KalmanFloat : public FRigUnit_SimBase
+struct CONTROLRIG_API FRigUnit_KalmanFloat : public FRigUnit_SimBase
 {
 	GENERATED_BODY()
 	
@@ -45,7 +45,7 @@ struct FRigUnit_KalmanFloat : public FRigUnit_SimBase
  * This uses a Kalman Filter internally.
  */
 USTRUCT(meta=(DisplayName="Average Over Time (Vector)", PrototypeName="KalmanFilter", Keywords="Average,Smooth,KalmanFilter"))
-struct FRigUnit_KalmanVector : public FRigUnit_SimBase
+struct CONTROLRIG_API FRigUnit_KalmanVector : public FRigUnit_SimBase
 {
 	GENERATED_BODY()
 	
@@ -80,7 +80,7 @@ struct FRigUnit_KalmanVector : public FRigUnit_SimBase
  * This uses a Kalman Filter internally.
  */
 USTRUCT(meta=(DisplayName="Average Over Time (Transform)", PrototypeName="KalmanFilter", Keywords="Average,Smooth,KalmanFilter"))
-struct FRigUnit_KalmanTransform : public FRigUnit_SimBase
+struct CONTROLRIG_API FRigUnit_KalmanTransform : public FRigUnit_SimBase
 {
 	GENERATED_BODY()
 	
diff --git a/Engine/Plugins/Experimental/ControlRig/Source/ControlRig/Private/Units/Simulation/RigUnit_PointSimulation.h b/Engine/Plugins/Experimental/ControlRig/Source/ControlRig/Private/Units/Simulation/RigUnit_PointSimulation.h
--- a/Engine/Plugins/Experimental/ControlRig/Source/ControlRig/Private/Units/Simulation/RigUnit_PointSimulation.h
+++ b/Engine/Plugins/Experimental/ControlRig/Source/ControlRig/Private/Units/Simulation/RigUnit_PointSimulation.h
@@ -12,7 +12,7 @@
 #include "RigUnit_PointSimulation.generated.h"
 
 USTRUCT()
-struct FRigUnit_PointSimulation_DebugSettings
+struct CONTROLRIG_API FRigUnit_PointSimulation_DebugSettings
 {
 	GENERATED_BODY()
 
@@ -64,7 +64,7 @@ struct FRigUnit_PointSimulation_DebugSettings
 };
 
 USTRUCT()
-struct FRigUnit_PointSimulation_BoneTarget
+struct CONTROLRIG_API FRigUnit_PointSimulation_BoneTarget
 {
 	GENERATED_BODY()
 
@@ -102,7 +102,7 @@ struct FRigUnit_PointSimulation_BoneTarget
 };
 
 USTRUCT()
-struct FRigUnit_PointSimulation_WorkData
+struct CONTROLRIG_API FRigUnit_PointSimulation_WorkData
 {
 	GENERATED_BODY()
 
@@ -118,7 +118,7 @@ struct FRigUnit_PointSimulation_WorkData
  * Note: Disabled for now.
  */
 USTRUCT(meta=(DisplayName="Point Simulation", Keywords="Simulate,Verlet,Springs", Deprecated="4.25"))
-struct FRigUnit_PointSimulation : public FRigUnit_SimBaseMutable
+struct CONTROLRIG_API FRigUnit_PointSimulation : public FRigUnit_SimBaseMutable
 {
 	GENERATED_BODY()
 	
diff --git a/Engine/Plugins/Experimental/ControlRig/Source/ControlRig/Private/Units/Simulation/RigUnit_SimBase.h b/Engine/Plugins/Experimental/ControlRig/Source/ControlRig/Private/Units/Simulation/RigUnit_SimBase.h
--- a/Engine/Plugins/Experimental/ControlRig/Source/ControlRig/Private/Units/Simulation/RigUnit_SimBase.h
+++ b/Engine/Plugins/Experimental/ControlRig/Source/ControlRig/Private/Units/Simulation/RigUnit_SimBase.h
@@ -6,13 +6,13 @@
 #include "RigUnit_SimBase.generated.h"
 
 USTRUCT(meta=(Abstract, Category = "Simulation", NodeColor = "0.25 0.05 0.05"))
-struct FRigUnit_SimBase : public FRigUnit
+struct CONTROLRIG_API FRigUnit_SimBase : public FRigUnit
 {
 	GENERATED_BODY()
 };
 
 USTRUCT(meta=(Abstract, Category = "Simulation", NodeColor = "0.25 0.05 0.05"))
-struct FRigUnit_SimBaseMutable : public FRigUnitMutable
+struct CONTROLRIG_API FRigUnit_SimBaseMutable : public FRigUnitMutable
 {
 	GENERATED_BODY()
 };
diff --git a/Engine/Plugins/Experimental/ControlRig/Source/ControlRig/Private/Units/Simulation/RigUnit_TimeOffset.h b/Engine/Plugins/Experimental/ControlRig/Source/ControlRig/Private/Units/Simulation/RigUnit_TimeOffset.h
--- a/Engine/Plugins/Experimental/ControlRig/Source/ControlRig/Private/Units/Simulation/RigUnit_TimeOffset.h
+++ b/Engine/Plugins/Experimental/ControlRig/Source/ControlRig/Private/Units/Simulation/RigUnit_TimeOffset.h
@@ -9,7 +9,7 @@
  * Records a value over time and can access the value from the past
  */
 USTRUCT(meta=(DisplayName="Value Over Time (Float)", PrototypeName="TimeOffset", Keywords="Buffer,Delta,History,Previous,TimeOffset,Delay"))
-struct FRigUnit_TimeOffsetFloat : public FRigUnit_SimBase
+struct CONTROLRIG_API FRigUnit_TimeOffsetFloat : public FRigUnit_SimBase
 {
 	GENERATED_BODY()
 	
@@ -62,7 +62,7 @@ struct FRigUnit_TimeOffsetFloat : public FRigUnit_SimBase
  * Records a value over time and can access the value from the past
  */
 USTRUCT(meta=(DisplayName="Value Over Time (Vector)", PrototypeName="TimeOffset", Keywords="Buffer,Delta,History,Previous,TimeOffset,Delay"))
-struct FRigUnit_TimeOffsetVector : public FRigUnit_SimBase
+struct CONTROLRIG_API FRigUnit_TimeOffsetVector : public FRigUnit_SimBase
 {
 	GENERATED_BODY()
 	
@@ -115,7 +115,7 @@ struct FRigUnit_TimeOffsetVector : public FRigUnit_SimBase
  * Records a value over time and can access the value from the past
  */
 USTRUCT(meta=(DisplayName="Value Over Time (Transform)", PrototypeName="TimeOffset", Keywords="Buffer,Delta,History,Previous,TimeOffset,Delay"))
-struct FRigUnit_TimeOffsetTransform : public FRigUnit_SimBase
+struct CONTROLRIG_API FRigUnit_TimeOffsetTransform : public FRigUnit_SimBase
 {
 	GENERATED_BODY()
 	
diff --git a/Engine/Plugins/Experimental/ControlRig/Source/ControlRig/Private/Units/Simulation/RigUnit_Verlet.h b/Engine/Plugins/Experimental/ControlRig/Source/ControlRig/Private/Units/Simulation/RigUnit_Verlet.h
--- a/Engine/Plugins/Experimental/ControlRig/Source/ControlRig/Private/Units/Simulation/RigUnit_Verlet.h
+++ b/Engine/Plugins/Experimental/ControlRig/Source/ControlRig/Private/Units/Simulation/RigUnit_Verlet.h
@@ -10,7 +10,7 @@
  * Simulates a single position over time using verlet integration
  */
 USTRUCT(meta=(DisplayName="Verlet (Vector)", PrototypeName="Verlet", Keywords="Simulate,Integrate"))
-struct FRigUnit_VerletIntegrateVector : public FRigUnit_SimBase
+struct CONTROLRIG_API FRigUnit_VerletIntegrateVector : public FRigUnit_SimBase
 {
 	GENERATED_BODY()
 	
