Control Rig: Fix static analysis issue with always true expression in RigHierarchyContainer

diff --git a/Engine/Plugins/Experimental/ControlRig/Source/ControlRig/Private/Rigs/RigHierarchyContainer.cpp b/Engine/Plugins/Experimental/ControlRig/Source/ControlRig/Private/Rigs/RigHierarchyContainer.cpp
--- a/Engine/Plugins/Experimental/ControlRig/Source/ControlRig/Private/Rigs/RigHierarchyContainer.cpp
+++ b/Engine/Plugins/Experimental/ControlRig/Source/ControlRig/Private/Rigs/RigHierarchyContainer.cpp
@@ -1136,7 +1136,7 @@ TArray<FRigElementKey> FRigHierarchyContainer::ImportFromText(const FString& InC
 					Data.LocalTransforms[Index].NormalizeRotation();
 					SetLocalTransform(Selection[Index], Data.LocalTransforms[Index]);
 				}
-				else if (InImportMode == ERigHierarchyImportMode::ReplaceCurrentGlobalTransform)
+				else // (InImportMode == ERigHierarchyImportMode::ReplaceCurrentGlobalTransform)
 				{
 					Data.GlobalTransforms[Index].NormalizeRotation();
 					SetGlobalTransform(Selection[Index], Data.GlobalTransforms[Index]);
