Control Rig: Add editor library to allow accessing blueprint functions in BP

diff --git a/Engine/Plugins/Experimental/ControlRig/Source/ControlRigEditor/Private/ControlRigBlueprintEditorLibrary.cpp b/Engine/Plugins/Experimental/ControlRig/Source/ControlRigEditor/Private/ControlRigBlueprintEditorLibrary.cpp
new file mode 100644
--- /dev/null
+++ b/Engine/Plugins/Experimental/ControlRig/Source/ControlRigEditor/Private/ControlRigBlueprintEditorLibrary.cpp
@@ -0,0 +1,94 @@
+// Copyright Epic Games, Inc. All Rights Reserved.
+
+#include "ControlRigBlueprintEditorLibrary.h"
+
+void UControlRigBlueprintEditorLibrary::SetPreviewMesh(UControlRigBlueprint* InRigBlueprint, USkeletalMesh* PreviewMesh, bool bMarkAsDirty)
+{
+	if(InRigBlueprint == nullptr)
+	{
+		return;
+	}
+	InRigBlueprint->SetPreviewMesh(PreviewMesh, bMarkAsDirty);
+}
+
+USkeletalMesh* UControlRigBlueprintEditorLibrary::GetPreviewMesh(UControlRigBlueprint* InRigBlueprint)
+{
+	if(InRigBlueprint == nullptr)
+	{
+		return nullptr;
+	}
+	return InRigBlueprint->GetPreviewMesh();
+}
+
+void UControlRigBlueprintEditorLibrary::RecompileVM(UControlRigBlueprint* InRigBlueprint)
+{
+	if(InRigBlueprint == nullptr)
+	{
+		return;
+	}
+	InRigBlueprint->RecompileVM();
+}
+
+void UControlRigBlueprintEditorLibrary::RecompileVMIfRequired(UControlRigBlueprint* InRigBlueprint)
+{
+	if(InRigBlueprint == nullptr)
+	{
+		return;
+	}
+	InRigBlueprint->RecompileVMIfRequired();
+}
+
+void UControlRigBlueprintEditorLibrary::RequestAutoVMRecompilation(UControlRigBlueprint* InRigBlueprint)
+{
+	if(InRigBlueprint == nullptr)
+	{
+		return;
+	}
+	InRigBlueprint->RequestAutoVMRecompilation();
+}
+
+void UControlRigBlueprintEditorLibrary::RequestControlRigInit(UControlRigBlueprint* InRigBlueprint)
+{
+	if(InRigBlueprint == nullptr)
+	{
+		return;
+	}
+	InRigBlueprint->RequestControlRigInit();
+}
+
+URigVMGraph* UControlRigBlueprintEditorLibrary::GetModel(UControlRigBlueprint* InRigBlueprint)
+{
+	if(InRigBlueprint == nullptr)
+	{
+		return nullptr;
+	}
+	return InRigBlueprint->Model;
+}
+
+URigVMController* UControlRigBlueprintEditorLibrary::GetController(UControlRigBlueprint* InRigBlueprint)
+{
+	if(InRigBlueprint == nullptr)
+	{
+		return nullptr;
+	}
+	return InRigBlueprint->Controller;
+}
+
+TArray<UControlRigBlueprint*> UControlRigBlueprintEditorLibrary::GetCurrentlyOpenRigBlueprints()
+{
+	return UControlRigBlueprint::GetCurrentlyOpenRigBlueprints();
+}
+
+TArray<UStruct*> UControlRigBlueprintEditorLibrary::GetAvailableRigUnits()
+{
+	return UControlRigBlueprint::GetAvailableRigUnits();
+}
+
+UControlRigHierarchyModifier* UControlRigBlueprintEditorLibrary::GetHierarchyModifier(UControlRigBlueprint* InRigBlueprint)
+{
+	if(InRigBlueprint == nullptr)
+	{
+		return nullptr;
+	}
+	return InRigBlueprint->GetHierarchyModifier();
+}
diff --git a/Engine/Plugins/Experimental/ControlRig/Source/ControlRigEditor/Public/ControlRigBlueprintEditorLibrary.h b/Engine/Plugins/Experimental/ControlRig/Source/ControlRigEditor/Public/ControlRigBlueprintEditorLibrary.h
new file mode 100644
--- /dev/null
+++ b/Engine/Plugins/Experimental/ControlRig/Source/ControlRigEditor/Public/ControlRigBlueprintEditorLibrary.h
@@ -0,0 +1,55 @@
+// Copyright Epic Games, Inc. All Rights Reserved.
+
+/**
+ *
+ * This thumbnail renderer displays a given Control Rig
+ */
+
+#pragma once
+
+#include "CoreMinimal.h"
+#include "UObject/ObjectMacros.h"
+#include "ControlRigBlueprint.h"
+#include "ControlRigBlueprintEditorLibrary.generated.h"
+
+UCLASS(meta=(ScriptName="ControlRigBlueprintLibrary"))
+class CONTROLRIGEDITOR_API UControlRigBlueprintEditorLibrary : public UBlueprintFunctionLibrary
+{
+	GENERATED_BODY()
+
+public:
+
+	UFUNCTION(BlueprintCallable, Category = "Control Rig Blueprint")
+	static void SetPreviewMesh(UControlRigBlueprint* InRigBlueprint, USkeletalMesh* PreviewMesh, bool bMarkAsDirty = true);
+	
+	UFUNCTION(BlueprintCallable, BlueprintPure, Category = "Control Rig Blueprint")
+	static USkeletalMesh* GetPreviewMesh(UControlRigBlueprint* InRigBlueprint);
+
+	UFUNCTION(BlueprintCallable, Category = "Control Rig Blueprint")
+	static void RecompileVM(UControlRigBlueprint* InRigBlueprint);
+
+	UFUNCTION(BlueprintCallable, Category = "Control Rig Blueprint")
+	static void RecompileVMIfRequired(UControlRigBlueprint* InRigBlueprint);
+	
+	UFUNCTION(BlueprintCallable, Category = "Control Rig Blueprint")
+	static void RequestAutoVMRecompilation(UControlRigBlueprint* InRigBlueprint);
+
+	UFUNCTION(BlueprintCallable, Category = "Control Rig Blueprint")
+	static void RequestControlRigInit(UControlRigBlueprint* InRigBlueprint);
+
+	UFUNCTION(BlueprintCallable, BlueprintPure, Category = "VM")
+	static URigVMGraph* GetModel(UControlRigBlueprint* InRigBlueprint);
+
+	UFUNCTION(BlueprintCallable, BlueprintPure, Category = "VM")
+	static URigVMController* GetController(UControlRigBlueprint* InRigBlueprint);
+
+	UFUNCTION(BlueprintCallable, BlueprintPure, Category = "VM")
+	static TArray<UControlRigBlueprint*> GetCurrentlyOpenRigBlueprints();
+
+	UFUNCTION(BlueprintCallable, BlueprintPure, Category = "VM")
+	static TArray<UStruct*> GetAvailableRigUnits();
+
+	UFUNCTION(BlueprintCallable, BlueprintPure, Category = "Hierarchy")
+	static UControlRigHierarchyModifier* GetHierarchyModifier(UControlRigBlueprint* InRigBlueprint);
+};
+
