Control Rig: Comment Details don't always appear in Details pane

diff --git a/Engine/Plugins/Experimental/ControlRig/Source/ControlRigEditor/Private/Editor/ControlRigEditor.cpp b/Engine/Plugins/Experimental/ControlRig/Source/ControlRigEditor/Private/Editor/ControlRigEditor.cpp
--- a/Engine/Plugins/Experimental/ControlRig/Source/ControlRigEditor/Private/Editor/ControlRigEditor.cpp
+++ b/Engine/Plugins/Experimental/ControlRig/Source/ControlRigEditor/Private/Editor/ControlRigEditor.cpp
@@ -982,7 +982,7 @@ void FControlRigEditor::SetDetailObjects(const TArray<UObject*>& InObjects)
 	StructToDisplay.Reset();
 	SKismetInspector::FShowDetailsOptions Options;
 	Options.bForceRefresh = true;
-	Inspector->ShowDetailsForObjects(InObjects);
+	Inspector->ShowDetailsForObjects(InObjects, Options);
 }
 
 void FControlRigEditor::SetDetailObject(UObject* Obj)
@@ -1068,11 +1068,18 @@ void FControlRigEditor::SetDetailObject(UObject* Obj)
 		}
 		return;
 	}
-
+	
 	TArray<UObject*> Objects;
 	if (Obj)
 	{
-		Objects.Add(Obj);
+		if (URigVMNode* ModelNode = Cast<URigVMNode>(Obj))
+		{
+			UControlRigGraph* RigGraph = Cast<UControlRigGraph>(GetFocusedGraph());
+			if (UEdGraphNode* EdNode = RigGraph->FindNodeForModelNodeName(ModelNode->GetFName()))
+			{
+				Objects.Add(EdNode);
+			}
+		}
 	}
 	SetDetailObjects(Objects);
 }
