Control Rig: Support for multiple SRigHierarchy widgets

diff --git a/Engine/Plugins/Experimental/ControlRig/Source/ControlRig/Private/ControlRig.cpp b/Engine/Plugins/Experimental/ControlRig/Source/ControlRig/Private/ControlRig.cpp
--- a/Engine/Plugins/Experimental/ControlRig/Source/ControlRig/Private/ControlRig.cpp
+++ b/Engine/Plugins/Experimental/ControlRig/Source/ControlRig/Private/ControlRig.cpp
@@ -1433,7 +1433,7 @@ FName UControlRig::AddTransientControl(const FRigElementKey& InElement)
 
 	if (InElement.Type == ERigElementType::Bone)
 	{
-		if (PreviewInstance)
+		if (PreviewInstance && Hierarchy.BoneHierarchy[ElementIndex].Type == ERigBoneType::Imported)
 		{
 			PreviewInstance->ModifyBone(InElement.Name);
 		}
diff --git a/Engine/Plugins/Experimental/ControlRig/Source/ControlRig/Private/Rigs/RigBoneHierarchy.cpp b/Engine/Plugins/Experimental/ControlRig/Source/ControlRig/Private/Rigs/RigBoneHierarchy.cpp
--- a/Engine/Plugins/Experimental/ControlRig/Source/ControlRig/Private/Rigs/RigBoneHierarchy.cpp
+++ b/Engine/Plugins/Experimental/ControlRig/Source/ControlRig/Private/Rigs/RigBoneHierarchy.cpp
@@ -660,6 +660,21 @@ void FRigBoneHierarchy::PropagateTransform(int32 InIndex)
 	}
 }
 
+void FRigBoneHierarchy::CopyGlobalToInitial(int32 InIndex, bool bPropagate)
+{
+	FTransform InitialGlobalTransform = GetGlobalTransform(InIndex);
+	SetInitialGlobalTransform(InIndex, InitialGlobalTransform);
+
+	if (bPropagate)
+	{
+		const TArray<int32> Dependents = Bones[InIndex].Dependents;
+		for (int32 DependentIndex = 0; DependentIndex < Dependents.Num(); ++DependentIndex)
+		{
+			CopyGlobalToInitial(Dependents[DependentIndex], bPropagate);
+		}
+	}
+}
+
 bool FRigBoneHierarchy::Select(const FName& InName, bool bSelect)
 {
 	if(GetIndex(InName) == INDEX_NONE)
diff --git a/Engine/Plugins/Experimental/ControlRig/Source/ControlRig/Public/Rigs/RigBoneHierarchy.h b/Engine/Plugins/Experimental/ControlRig/Source/ControlRig/Public/Rigs/RigBoneHierarchy.h
--- a/Engine/Plugins/Experimental/ControlRig/Source/ControlRig/Public/Rigs/RigBoneHierarchy.h
+++ b/Engine/Plugins/Experimental/ControlRig/Source/ControlRig/Public/Rigs/RigBoneHierarchy.h
@@ -194,6 +194,9 @@ struct CONTROLRIG_API FRigBoneHierarchy
 	// propagates the transform change for a single bone
 	void PropagateTransform(int32 InIndex);
 
+	// copies the global transform to the initial transform
+	void CopyGlobalToInitial(int32 InIndex, bool bPropagate = true);
+
 	// import skeleton
 	TArray<FRigElementKey> ImportSkeleton(const FReferenceSkeleton& InSkeleton, const FName& InNameSpace, bool bReplaceExistingBones, bool bRemoveObsoleteBones, bool bSelectBones, bool bNotify);
 
diff --git a/Engine/Plugins/Experimental/ControlRig/Source/ControlRigDeveloper/Public/ControlRigBlueprint.h b/Engine/Plugins/Experimental/ControlRig/Source/ControlRigDeveloper/Public/ControlRigBlueprint.h
--- a/Engine/Plugins/Experimental/ControlRig/Source/ControlRigDeveloper/Public/ControlRigBlueprint.h
+++ b/Engine/Plugins/Experimental/ControlRig/Source/ControlRigDeveloper/Public/ControlRigBlueprint.h
@@ -230,6 +230,7 @@ private:
 
 public:
 
+	void BroadcastRefreshEditor() { return RefreshEditorEvent.Broadcast(this); }
 	FOnRefreshEditorEvent& OnRefreshEditor() { return RefreshEditorEvent; }
 	FOnVariableDroppedEvent& OnVariableDropped() { return VariableDroppedEvent; }
 
diff --git a/Engine/Plugins/Experimental/ControlRig/Source/ControlRigEditor/Private/Editor/ControlRigEditor.cpp b/Engine/Plugins/Experimental/ControlRig/Source/ControlRigEditor/Private/Editor/ControlRigEditor.cpp
--- a/Engine/Plugins/Experimental/ControlRig/Source/ControlRigEditor/Private/Editor/ControlRigEditor.cpp
+++ b/Engine/Plugins/Experimental/ControlRig/Source/ControlRigEditor/Private/Editor/ControlRigEditor.cpp
@@ -3957,8 +3957,7 @@ void FControlRigEditor::HandleOnControlModified(UControlRig* Subject, const FRig
 
 				if (bSetupModeEnabled)
 				{
-					FTransform InitialGlobalTransform = Blueprint->HierarchyContainer.BoneHierarchy.GetGlobalTransform(BoneName);
-					Blueprint->HierarchyContainer.BoneHierarchy.SetInitialGlobalTransform(BoneName, InitialGlobalTransform);
+					Blueprint->HierarchyContainer.BoneHierarchy.CopyGlobalToInitial(Blueprint->HierarchyContainer.BoneHierarchy.GetIndex(BoneName));
 				}
 
 				Blueprint->PropagateHierarchyFromBPToInstances(false, false);
diff --git a/Engine/Plugins/Experimental/ControlRig/Source/ControlRigEditor/Private/Editor/RigHierarchyTabSummoner.cpp b/Engine/Plugins/Experimental/ControlRig/Source/ControlRigEditor/Private/Editor/RigHierarchyTabSummoner.cpp
--- a/Engine/Plugins/Experimental/ControlRig/Source/ControlRigEditor/Private/Editor/RigHierarchyTabSummoner.cpp
+++ b/Engine/Plugins/Experimental/ControlRig/Source/ControlRigEditor/Private/Editor/RigHierarchyTabSummoner.cpp
@@ -16,12 +16,23 @@ FRigHierarchyTabSummoner::FRigHierarchyTabSummoner(const TSharedRef<FControlRigE
 	TabLabel = LOCTEXT("RigHierarchyTabLabel", "Rig Hierarchy");
 	TabIcon = FSlateIcon(FControlRigEditorStyle::Get().GetStyleSetName(), "ControlRig.TabIcon");
 
-	bIsSingleton = true;
-
 	ViewMenuDescription = LOCTEXT("RigHierarchy_ViewMenu_Desc", "Rig Hierarchy");
 	ViewMenuTooltip = LOCTEXT("RigHierarchy_ViewMenu_ToolTip", "Show the Rig Hierarchy tab");
 }
 
+FTabSpawnerEntry& FRigHierarchyTabSummoner::RegisterTabSpawner(TSharedRef<FTabManager> InTabManager, const FApplicationMode* CurrentApplicationMode) const
+{
+	FTabSpawnerEntry& SpawnerEntry = FWorkflowTabFactory::RegisterTabSpawner(InTabManager, CurrentApplicationMode);
+
+	SpawnerEntry.SetReuseTabMethod(FOnFindTabToReuse::CreateLambda([](const FTabId& InTabId) ->TSharedPtr<SDockTab> {
+	
+		return TSharedPtr<SDockTab>();
+
+	}));
+
+	return SpawnerEntry;
+}
+
 TSharedRef<SWidget> FRigHierarchyTabSummoner::CreateTabBody(const FWorkflowTabSpawnInfo& Info) const
 {
 	return SNew(SRigHierarchy, ControlRigEditor.Pin().ToSharedRef());
diff --git a/Engine/Plugins/Experimental/ControlRig/Source/ControlRigEditor/Private/Editor/RigHierarchyTabSummoner.h b/Engine/Plugins/Experimental/ControlRig/Source/ControlRigEditor/Private/Editor/RigHierarchyTabSummoner.h
--- a/Engine/Plugins/Experimental/ControlRig/Source/ControlRigEditor/Private/Editor/RigHierarchyTabSummoner.h
+++ b/Engine/Plugins/Experimental/ControlRig/Source/ControlRigEditor/Private/Editor/RigHierarchyTabSummoner.h
@@ -15,8 +15,9 @@ public:
 public:
 	FRigHierarchyTabSummoner(const TSharedRef<FControlRigEditor>& InControlRigEditor);
 	
+	virtual FTabSpawnerEntry& RegisterTabSpawner(TSharedRef<FTabManager> TabManager, const FApplicationMode* CurrentApplicationMode) const;
 	virtual TSharedRef<SWidget> CreateTabBody(const FWorkflowTabSpawnInfo& Info) const override;
-	
+
 protected:
 	TWeakPtr<FControlRigEditor> ControlRigEditor;
 };
diff --git a/Engine/Plugins/Experimental/ControlRig/Source/ControlRigEditor/Private/Editor/SRigHierarchy.cpp b/Engine/Plugins/Experimental/ControlRig/Source/ControlRigEditor/Private/Editor/SRigHierarchy.cpp
--- a/Engine/Plugins/Experimental/ControlRig/Source/ControlRigEditor/Private/Editor/SRigHierarchy.cpp
+++ b/Engine/Plugins/Experimental/ControlRig/Source/ControlRigEditor/Private/Editor/SRigHierarchy.cpp
@@ -1126,6 +1126,10 @@ void SRigHierarchy::OnRigElementSelected(FRigHierarchyContainer* Container, cons
 
 void SRigHierarchy::HandleRefreshEditorFromBlueprint(UControlRigBlueprint* InBlueprint)
 {
+	if (bIsChangingRigHierarchy)
+	{
+		return;
+	}
 	RefreshTreeView();
 }
 
@@ -1339,6 +1343,12 @@ void SRigHierarchy::CreateRefreshMenu(FMenuBuilder& MenuBuilder)
 
 void SRigHierarchy::RefreshHierarchy(const FAssetData& InAssetData)
 {
+	if (bIsChangingRigHierarchy)
+	{
+		return;
+	}
+	TGuardValue<bool> GuardRigHierarchyChanges(bIsChangingRigHierarchy, true);
+
 	FRigHierarchyContainer* Hierarchy = GetHierarchyContainer();
 	USkeletalMesh* Mesh = Cast<USkeletalMesh>(InAssetData.GetAsset());
 	if (Mesh && Hierarchy)
@@ -1351,6 +1361,8 @@ void SRigHierarchy::RefreshHierarchy(const FAssetData& InAssetData)
 	}
 
 	FSlateApplication::Get().DismissAllMenus();
+	ControlRigBlueprint->BroadcastRefreshEditor();
+	ControlRigEditor.Pin()->OnHierarchyChanged();
 	RefreshTreeView();
 }
 
@@ -1383,6 +1395,12 @@ void SRigHierarchy::CreateImportMenu(FMenuBuilder& MenuBuilder)
 
 void SRigHierarchy::ImportHierarchy(const FAssetData& InAssetData)
 {
+	if (bIsChangingRigHierarchy)
+	{
+		return;
+	}
+	TGuardValue<bool> GuardRigHierarchyChanges(bIsChangingRigHierarchy, true);
+
 	FRigHierarchyContainer* Hierarchy = GetHierarchyContainer();
 	USkeletalMesh* Mesh = Cast<USkeletalMesh> (InAssetData.GetAsset());
 	if (Mesh && Hierarchy)
@@ -1399,6 +1417,7 @@ void SRigHierarchy::ImportHierarchy(const FAssetData& InAssetData)
 	}
 
 	FSlateApplication::Get().DismissAllMenus();
+	ControlRigBlueprint->BroadcastRefreshEditor();
 	RefreshTreeView();
 
 	if (ControlRigBlueprint->GetPreviewMesh() == nullptr &&
@@ -1711,6 +1730,10 @@ void SRigHierarchy::HandleDuplicateItem()
 
 	FSlateApplication::Get().DismissAllMenus();
 	ControlRigEditor.Pin()->OnHierarchyChanged();
+	{
+		TGuardValue<bool> GuardRigHierarchyChanges(bIsChangingRigHierarchy, true);
+		ControlRigBlueprint->BroadcastRefreshEditor();
+	}
 	RefreshTreeView();
 }
 
@@ -1841,6 +1864,10 @@ void SRigHierarchy::HandlePasteItems()
 
 	//ControlRigBlueprint->PropagateHierarchyFromBPToInstances(true);
 	ControlRigEditor.Pin()->OnHierarchyChanged();
+	{
+		TGuardValue<bool> GuardRigHierarchyChanges(bIsChangingRigHierarchy, true);
+		ControlRigBlueprint->BroadcastRefreshEditor();
+	}
 	RefreshTreeView();
 }
 
@@ -2050,6 +2077,9 @@ TOptional<EItemDropZone> SRigHierarchy::OnCanAcceptDrop(const FDragDropEvent& Dr
 
 FReply SRigHierarchy::OnAcceptDrop(const FDragDropEvent& DragDropEvent, EItemDropZone DropZone, TSharedPtr<FRigTreeElement> TargetItem)
 {
+	bool bMatchTransforms = DragDropEvent.GetModifierKeys().IsAltDown();
+	bool bReparentItems = !bMatchTransforms;
+
 	TSharedPtr<FRigElementHierarchyDragDropOp> RigDragDropOp = DragDropEvent.GetOperationAs<FRigElementHierarchyDragDropOp>();
 	if (RigDragDropOp.IsValid())
 	{
@@ -2063,6 +2093,8 @@ FReply SRigHierarchy::OnAcceptDrop(const FDragDropEvent& DragDropEvent, EItemDro
 			FScopedTransaction Transaction(LOCTEXT("HierarchyDragAndDrop", "Drag & Drop"));
 			ControlRigBlueprint->Modify();
 
+			FTransform TargetGlobalTransform = DebuggedContainer->GetGlobalTransform(TargetItem->Key);
+
 			for (const FRigElementKey& DraggedKey : RigDragDropOp->GetElements())
 			{
 				if (DraggedKey == TargetItem->Key)
@@ -2070,13 +2102,16 @@ FReply SRigHierarchy::OnAcceptDrop(const FDragDropEvent& DragDropEvent, EItemDro
 					return FReply::Unhandled();
 				}
 
-				if (Container->IsParentedTo(
-					TargetItem->Key.Type,
-					Container->GetIndex(TargetItem->Key),
-					DraggedKey.Type,
-					Container->GetIndex(DraggedKey)))
+				if (bReparentItems)
 				{
-					return FReply::Unhandled();
+					if (Container->IsParentedTo(
+						TargetItem->Key.Type,
+						Container->GetIndex(TargetItem->Key),
+						DraggedKey.Type,
+						Container->GetIndex(DraggedKey)))
+					{
+						return FReply::Unhandled();
+					}
 				}
 
 				if (DraggedKey.Type == ERigElementType::Bone)
@@ -2087,14 +2122,20 @@ FReply SRigHierarchy::OnAcceptDrop(const FDragDropEvent& DragDropEvent, EItemDro
 						const FRigBone& Bone = Container->BoneHierarchy[BoneIndex];
 						if (Bone.Type == ERigBoneType::Imported && Bone.ParentIndex != INDEX_NONE)
 						{
-							FText ConfirmReparent = LOCTEXT("ConfirmReparentBoneHierarchy", "Reparenting imported(white) bones can cause issues with animation - are you sure ?");
+							FText ConfirmText = bMatchTransforms ?
+								LOCTEXT("ConfirmMatchTransform", "Matching transforms of imported(white) bones can cause issues with animation - are you sure ?") :
+								LOCTEXT("ConfirmReparentBoneHierarchy", "Reparenting imported(white) bones can cause issues with animation - are you sure ?");
+
+							FText TitleText = bMatchTransforms ?
+								LOCTEXT("MatchTransformImportedBone", "Match Transform on Imported Bone") :
+								LOCTEXT("ReparentImportedBone", "Reparent Imported Bone");
 
-							FSuppressableWarningDialog::FSetupInfo Info(ConfirmReparent, LOCTEXT("ReparentImportedBone", "Reparent Imported Bone"), "ReparentImportedBoneHierarchy_Warning");
-							Info.ConfirmText = LOCTEXT("ReparentImportedBoneHierarchy_Yes", "Yes");
-							Info.CancelText = LOCTEXT("ReparentImportedBoneHierarchy_No", "No");
+							FSuppressableWarningDialog::FSetupInfo Info(ConfirmText, TitleText, "SRigHierarchy_Warning");
+							Info.ConfirmText = LOCTEXT("SRigHierarchy_Warning_Yes", "Yes");
+							Info.CancelText = LOCTEXT("SRigHierarchy_Warning_No", "No");
 
-							FSuppressableWarningDialog ReparentImportedBonesInHierarchy(Info);
-							if (ReparentImportedBonesInHierarchy.ShowModal() == FSuppressableWarningDialog::Cancel)
+							FSuppressableWarningDialog ChangeImportedBonesInHierarchy(Info);
+							if (ChangeImportedBonesInHierarchy.ShowModal() == FSuppressableWarningDialog::Cancel)
 							{
 								return FReply::Unhandled();
 							}
@@ -2105,6 +2146,43 @@ FReply SRigHierarchy::OnAcceptDrop(const FDragDropEvent& DragDropEvent, EItemDro
 
 			for (const FRigElementKey& DraggedKey : RigDragDropOp->GetElements())
 			{
+				if (bMatchTransforms)
+				{
+					if (DraggedKey.Type == ERigElementType::Control)
+					{
+						int32 ControlIndex = DebuggedContainer->ControlHierarchy.GetIndex(DraggedKey.Name);
+						if (ControlIndex == INDEX_NONE)
+						{
+							continue;
+						}
+
+						FTransform ParentTransform = DebuggedContainer->ControlHierarchy.GetParentTransform(ControlIndex, false);
+						FTransform OffsetTransform = TargetGlobalTransform.GetRelativeTransform(ParentTransform);
+
+						Container->ControlHierarchy[ControlIndex].OffsetTransform = OffsetTransform;
+						Container->SetLocalTransform(DraggedKey, FTransform::Identity);
+						Container->SetInitialTransform(DraggedKey, FTransform::Identity);
+						DebuggedContainer->ControlHierarchy[ControlIndex].OffsetTransform = OffsetTransform;
+						DebuggedContainer->SetLocalTransform(DraggedKey, FTransform::Identity);
+						DebuggedContainer->SetInitialTransform(DraggedKey, FTransform::Identity);
+					}
+					else if (DraggedKey.Type == ERigElementType::Bone)
+					{
+						Container->SetGlobalTransform(DraggedKey, TargetGlobalTransform);
+						Container->BoneHierarchy.CopyGlobalToInitial(Container->BoneHierarchy.GetIndex(DraggedKey.Name));
+						DebuggedContainer->SetGlobalTransform(DraggedKey, TargetGlobalTransform);
+						DebuggedContainer->BoneHierarchy.CopyGlobalToInitial(DebuggedContainer->BoneHierarchy.GetIndex(DraggedKey.Name));
+					}
+					else
+					{
+						Container->SetInitialGlobalTransform(DraggedKey, TargetGlobalTransform);
+						Container->SetGlobalTransform(DraggedKey, TargetGlobalTransform);
+						DebuggedContainer->SetInitialGlobalTransform(DraggedKey, TargetGlobalTransform);
+						DebuggedContainer->SetGlobalTransform(DraggedKey, TargetGlobalTransform);
+					}
+					continue;
+				}
+
 				FRigElementKey ParentKey = TargetItem->Key;
 
 				FTransform InitialTransform = DebuggedContainer->GetInitialGlobalTransform(DraggedKey);
@@ -2237,7 +2315,13 @@ FReply SRigHierarchy::OnAcceptDrop(const FDragDropEvent& DragDropEvent, EItemDro
 		}
 
 		ControlRigBlueprint->PropagateHierarchyFromBPToInstances(true);
-		RefreshTreeView();
+
+		if(bReparentItems)
+		{
+			TGuardValue<bool> GuardRigHierarchyChanges(bIsChangingRigHierarchy, true);
+			ControlRigBlueprint->BroadcastRefreshEditor();
+			RefreshTreeView();
+		}
 
 		return FReply::Handled();
 	}
