Control Rig: Make sure the component calls UpdateLODStatus

diff --git a/Engine/Plugins/Experimental/ControlRig/Source/ControlRig/Private/ControlRigComponent.cpp b/Engine/Plugins/Experimental/ControlRig/Source/ControlRig/Private/ControlRigComponent.cpp
--- a/Engine/Plugins/Experimental/ControlRig/Source/ControlRig/Private/ControlRigComponent.cpp
+++ b/Engine/Plugins/Experimental/ControlRig/Source/ControlRig/Private/ControlRigComponent.cpp
@@ -1356,6 +1356,7 @@ void UControlRigComponent::TransferOutputs()
 		for (USkeletalMeshComponent* SkeletalMeshComponent : ComponentsToTick)
 		{
 			SkeletalMeshComponent->TickAnimation(0.f, false);
+			SkeletalMeshComponent->UpdateLODStatus();
 			SkeletalMeshComponent->RefreshBoneTransforms();
 			SkeletalMeshComponent->RefreshSlaveComponents();
 			SkeletalMeshComponent->UpdateComponentToWorld();
