Sequencer: Control Rig. Toolbar Start and Pose Library

ControlRigPose contains the Pose Data UStructs the Pose Asset UObject.
ControlRigPoseMirrorTable contains logic for doing mirroring of controls.
SControlRigBaseListWidget is the main Pose Dialog. It's made to handle all three types of assets (pose,animation,selection), but
now is just hosting poses. You'll see that I left in code (commented out mostly) if and when we get these other types in.
SControlRigControlViews holds the specific View for the selected Asset, again just Poses.
The Thumbnail renderer is for automatically creating thumbnails (done it's own scene like Sequencer).. We went back and forth between that and taking a viewport snapshot, settling on the viewport snapshot instead. So like some other stuff may get removed but easier to do later than delete and re-create.
The rest is mostly tooling to get this setup, needed a proper toolkit to host the toolbar, etc..

diff --git a/Engine/Plugins/Experimental/ControlRig/Source/ControlRig/Private/Tools/ControlRigPose.cpp b/Engine/Plugins/Experimental/ControlRig/Source/ControlRig/Private/Tools/ControlRigPose.cpp
new file mode 100644
--- /dev/null
+++ b/Engine/Plugins/Experimental/ControlRig/Source/ControlRig/Private/Tools/ControlRigPose.cpp
@@ -0,0 +1,63 @@
+// Copyright Epic Games, Inc. All Rights Reserved.
+#include "Tools/ControlRigPose.h"
+#include "Tools/ControlRigPoseProjectSettings.h"
+#include "IControlRigObjectBinding.h"
+#include "Framework/Notifications/NotificationManager.h"
+#include "Widgets/Notifications/SNotificationList.h"
+#if WITH_EDITOR
+#include "ScopedTransaction.h"
+#endif
+
+#define LOCTEXT_NAMESPACE "ControlRigPose"
+
+UControlRigPoseAsset::UControlRigPoseAsset(const FObjectInitializer& ObjectInitializer) : Super(ObjectInitializer)
+{
+}
+
+void UControlRigPoseAsset::SavePose(UControlRig* InControlRig, bool bUseAll)
+{
+	Pose.SavePose(InControlRig,bUseAll);
+	Controls = Pose.GetControlNames();
+}
+
+void UControlRigPoseAsset::PastePose(UControlRig* InControlRig, bool bDoKey, bool bDoMirror)
+{
+#if WITH_EDITOR
+	FScopedTransaction ScopedTransaction(LOCTEXT("PastePoseTransaction", "Paste Pose"));
+	InControlRig->Modify();
+#endif
+	Pose.PastePose(InControlRig,bDoKey, bDoMirror);
+}
+
+void UControlRigPoseAsset::SelectControls(UControlRig* InControlRig)
+{
+#if WITH_EDITOR
+	FScopedTransaction ScopedTransaction(LOCTEXT("SelectControlTransaction", "Select Control"));
+	InControlRig->Modify();
+#endif
+	InControlRig->ClearControlSelection();
+	Controls = Pose.GetControlNames();
+	for (const FName& Name : Controls)
+	{
+		InControlRig->SelectControl(Name, true);
+	}
+}
+
+TArray<FRigControlCopy> UControlRigPoseAsset::GetCurrentPose(UControlRig* InControlRig) 
+{
+	FControlRigControlPose TempPose;
+	FString Name;
+	TempPose.SavePose(InControlRig,true);
+	return TempPose.GetPoses();
+}
+
+void UControlRigPoseAsset::BlendWithInitialPoses(TArray<FRigControlCopy>& InitialControls, UControlRig* InControlRig, bool bDoKey, bool bDoMirror, float BlendValue)
+{
+	if (BlendValue > 0.0f)
+	{
+		Pose.BlendWithInitialPoses(InitialControls, InControlRig, bDoKey, bDoMirror, BlendValue);
+	}
+}
+
+#undef LOCTEXT_NAMESPACE
+
diff --git a/Engine/Plugins/Experimental/ControlRig/Source/ControlRig/Private/Tools/ControlRigPoseMirrorSettings.cpp b/Engine/Plugins/Experimental/ControlRig/Source/ControlRig/Private/Tools/ControlRigPoseMirrorSettings.cpp
new file mode 100644
--- /dev/null
+++ b/Engine/Plugins/Experimental/ControlRig/Source/ControlRig/Private/Tools/ControlRigPoseMirrorSettings.cpp
@@ -0,0 +1,11 @@
+// Copyright Epic Games, Inc. All Rights Reserved.
+#include "Tools/ControlRigPoseMirrorSettings.h"
+
+UControlRigPoseMirrorSettings::UControlRigPoseMirrorSettings()
+{
+	RightSide = TEXT("Right");
+	LeftSide = TEXT("Left");
+	XAxis = 1.0f;
+	YAxis = 0.0f;
+	ZAxis = 0.0f;
+}
diff --git a/Engine/Plugins/Experimental/ControlRig/Source/ControlRig/Private/Tools/ControlRigPoseMirrorTable.cpp b/Engine/Plugins/Experimental/ControlRig/Source/ControlRig/Private/Tools/ControlRigPoseMirrorTable.cpp
new file mode 100644
--- /dev/null
+++ b/Engine/Plugins/Experimental/ControlRig/Source/ControlRig/Private/Tools/ControlRigPoseMirrorTable.cpp
@@ -0,0 +1,118 @@
+// Copyright Epic Games, Inc. All Rights Reserved.
+
+#include "Tools/ControlRigPoseMirrorTable.h"
+#include "Tools/ControlRigPoseMirrorSettings.h"
+#include "ControlRig.h"
+#include "Tools/ControlRigPose.h"
+
+void FControlRigPoseMirrorTable::SetUpMirrorTable(const UControlRig* ControlRig)
+{
+	const UControlRigPoseMirrorSettings* Settings = GetDefault<UControlRigPoseMirrorSettings>();
+	MatchedControls.Reset();
+	if (Settings && ControlRig)
+	{
+		const TArray<FRigControl>& CurrentControls = ControlRig->AvailableControls();
+		for (const FRigControl& RigControl : CurrentControls)
+		{
+			FString CurrentString = RigControl.Name.ToString();
+			if (CurrentString.Contains(Settings->RightSide))
+			{
+				FString NewString = CurrentString.Replace(*Settings->RightSide, *Settings->LeftSide);
+				FName CurrentName(*CurrentString);
+				FName NewName(*NewString);
+				MatchedControls.Add(CurrentName, NewName);
+			}
+			else if (CurrentString.Contains(Settings->LeftSide))
+			{
+				FString NewString = CurrentString.Replace(*Settings->LeftSide, *Settings->RightSide);
+				FName CurrentName(*CurrentString);
+				FName NewName(*NewString);
+				MatchedControls.Add(CurrentName, NewName);
+			}
+		}
+	}
+}
+
+FRigControl* FControlRigPoseMirrorTable::GetControl(UControlRig* ControlRig, const FName& Name) const
+{
+
+	if (MatchedControls.Num() <= 0)
+	{
+		if (ControlRig->IsControlSelected(Name))
+		{
+			return ControlRig->FindControl(Name);
+		}
+	}
+	if (MatchedControls.Num() > 0) 
+	{
+		if (const FName* MatchedName = MatchedControls.Find(Name))
+		{
+			if (ControlRig->IsControlSelected(*MatchedName))
+			{
+				if (FRigControl* RigControl = ControlRig->FindControl(*MatchedName))
+				{
+					return RigControl;
+				}
+			}
+			else
+			{
+				return nullptr;
+			}
+		}
+	}
+	
+	if (ControlRig->IsControlSelected(Name))
+	{
+		return ControlRig->FindControl(Name);
+	}
+	return nullptr;
+}
+
+bool FControlRigPoseMirrorTable::IsMatched(const FName& Name) const
+{
+	if (MatchedControls.Num() > 0)
+	{
+		if (const FName* MatchedName = MatchedControls.Find(Name))
+		{
+			return true;
+		}
+	}
+	return false;
+}
+
+// Mirrors Translation as Global(Component)
+// Mirrors Rotation as Local , only if NOT matched, if matched we just use it
+void FControlRigPoseMirrorTable::GetMirrorTransform(const FRigControlCopy& ControlCopy, bool bIsMatched, FVector& OutGlobalTranslation, FQuat& OutLocalRotation) const
+{
+	const UControlRigPoseMirrorSettings* Settings = GetDefault<UControlRigPoseMirrorSettings>();
+	if (Settings)
+	{
+		FVector Axis(Settings->XAxis, Settings->YAxis, Settings->ZAxis);
+		Axis.Normalize();
+		FTransform GlobalTransform = ControlCopy.GlobalTransform;
+
+		FVector Translation = ControlCopy.GlobalTransform.GetTranslation();
+		FPlane Plane(FVector::ZeroVector,Axis);
+		OutGlobalTranslation = Translation.MirrorByPlane(Plane);
+
+		FTransform LocalTransform = ControlCopy.LocalTransform;
+		FQuat Rotation = LocalTransform.GetRotation();
+		//test FRotator OldRotator = Rotation.Rotator();
+		
+		if (bIsMatched)
+		{
+			OutLocalRotation = Rotation;
+		}
+		else
+		{
+			FQuat MirrorNormalQuat = FQuat(Axis.X, Axis.Y, Axis.Z, 0);
+			Rotation.EnforceShortestArcWith(MirrorNormalQuat);
+			OutLocalRotation = MirrorNormalQuat * Rotation * MirrorNormalQuat;
+		}
+		//test FRotator NewRotator = OutLocalRotation.Rotator();
+		
+		return;
+	}
+	OutGlobalTranslation = FVector::ZeroVector;
+	OutLocalRotation = FQuat::Identity;
+}
diff --git a/Engine/Plugins/Experimental/ControlRig/Source/ControlRig/Private/Tools/ControlRigPoseProjectSettings.cpp b/Engine/Plugins/Experimental/ControlRig/Source/ControlRig/Private/Tools/ControlRigPoseProjectSettings.cpp
new file mode 100644
--- /dev/null
+++ b/Engine/Plugins/Experimental/ControlRig/Source/ControlRig/Private/Tools/ControlRigPoseProjectSettings.cpp
@@ -0,0 +1,7 @@
+// Copyright Epic Games, Inc. All Rights Reserved.
+#include "Tools/ControlRigPoseProjectSettings.h"
+
+UControlRigPoseProjectSettings::UControlRigPoseProjectSettings()
+{
+ 	RootSaveDir.Path = TEXT("/Game/ControlRig/Pose");
+}
diff --git a/Engine/Plugins/Experimental/ControlRig/Source/ControlRig/Public/Tools/ControlRigPose.h b/Engine/Plugins/Experimental/ControlRig/Source/ControlRig/Public/Tools/ControlRigPose.h
new file mode 100644
--- /dev/null
+++ b/Engine/Plugins/Experimental/ControlRig/Source/ControlRig/Public/Tools/ControlRigPose.h
@@ -0,0 +1,345 @@
+// Copyright Epic Games, Inc. All Rights Reserved.
+#pragma once
+/**
+*  Data To Store and Apply the Control Rig Pose
+*/
+
+#include "CoreMinimal.h"
+#include "UObject/Object.h"
+#include "TransformNoScale.h"
+#include "ControlRig.h"
+#include "Engine/SkeletalMesh.h"
+#include "ControlRigToolAsset.h"
+#include "Tools/ControlRigPoseMirrorTable.h"
+#include "Rigs/RigControlHierarchy.h"
+#include "ControlRigPose.generated.h"
+
+/**
+* The Data Stored For Each Control in A Pose.
+*/
+USTRUCT()
+struct CONTROLRIG_API FRigControlCopy
+{
+	GENERATED_BODY()
+
+		FRigControlCopy()
+		: Name(NAME_None)
+		, ControlType(ERigControlType::Transform)
+		, Value()
+		, ParentName(NAME_None)
+		, SpaceName(NAME_None)
+		, OffsetTransform(FTransform::Identity)
+		, InitialValue()
+		, ParentTransform(FTransform::Identity)
+		, LocalTransform(FTransform::Identity)
+		, GlobalTransform(FTransform::Identity)
+
+	{
+	}
+
+	FRigControlCopy(const FRigControl& InControl, const FRigControlHierarchy& Hierarchy)
+	{
+		Name = InControl.Name;
+		ControlType = InControl.ControlType;
+		Value = InControl.Value;
+		ParentName = InControl.ParentName;
+		SpaceName = InControl.SpaceName;
+		OffsetTransform = InControl.OffsetTransform;
+		InitialValue = InControl.InitialValue;
+
+		int32 Index = Hierarchy.GetIndex(Name);
+		ParentTransform = Index != INDEX_NONE ? Hierarchy.GetParentTransform(Index) : FTransform::Identity;
+		LocalTransform = Hierarchy.GetLocalTransform(Name);
+		GlobalTransform = Hierarchy.GetGlobalTransform(Name);
+	}
+	virtual ~FRigControlCopy() {}
+
+	UPROPERTY()
+	FName Name;
+
+	UPROPERTY()
+	ERigControlType ControlType;
+
+	UPROPERTY()
+	FRigControlValue Value;
+
+	UPROPERTY()
+	FName ParentName;
+
+	UPROPERTY()
+	FName SpaceName;
+
+	UPROPERTY()
+	FTransform OffsetTransform;
+
+	UPROPERTY()
+	FRigControlValue InitialValue;
+
+	UPROPERTY()
+	FTransform ParentTransform;
+
+	UPROPERTY()
+	FTransform LocalTransform;
+
+	UPROPERTY()
+	FTransform GlobalTransform;
+
+};
+
+/**
+* The Data Stored For Each Pose and associated Functions to Store and Paste It
+*/
+USTRUCT()
+struct FControlRigControlPose
+{
+	GENERATED_USTRUCT_BODY()
+
+	FControlRigControlPose() {};
+	FControlRigControlPose(UControlRig* InControlRig, bool bUseAll)
+	{
+		SavePose(InControlRig, bUseAll);
+	}
+	~FControlRigControlPose() {};
+
+	void SavePose(UControlRig* ControlRig,  bool bUseAll)
+	{
+		const TArray<FRigControl>& CurrentControls = ControlRig->AvailableControls();
+		CopyOfControls.SetNum(0);
+		FRigControlHierarchy& Hierarchy = ControlRig->GetControlHierarchy();
+		for (const FRigControl& RigControl : CurrentControls)
+		{
+			if (RigControl.bAnimatable && (bUseAll || ControlRig->IsControlSelected(RigControl.Name)))
+			{
+				FRigControlCopy Copy(RigControl, Hierarchy);
+				CopyOfControls.Add(Copy);
+			}
+		}
+	}
+
+	void PastePose(UControlRig* ControlRig, bool bDoKey, bool bDoMirror)
+	{
+		PastePoseInternal(ControlRig,bDoKey, bDoMirror, CopyOfControls);
+		//do it twice do to issues we have seen with the ordering of certain controls,due to how spaces and offsets work.
+		PastePoseInternal(ControlRig, bDoKey, bDoMirror, CopyOfControls);
+	}
+
+	void SetControlMirrorTransform(UControlRig* ControlRig, const FName& Name, const FVector& GlobalTranslation, const FQuat& LocalRotation, bool bNotify, FRigControlModifiedContext Context)
+	{
+		int32 Index = ControlRig->GetControlHierarchy().GetIndex(Name);
+		FTransform ParentTransform = ControlRig->GetControlHierarchy().GetParentTransform(Index);
+		FTransform CurrentTransform = ControlRig->GetControlHierarchy().GetLocalTransform(Index);
+		FVector NewLocal = ParentTransform.InverseTransformPositionNoScale(GlobalTranslation);
+		FTransform NewCurrentTransform(LocalRotation, NewLocal);
+		ControlRig->SetControlLocalTransform(Name, NewCurrentTransform, bNotify, Context);
+	}
+
+	void PastePoseInternal(UControlRig* ControlRig, bool bDoKey,bool bDoMirror, const TArray<FRigControlCopy>& ControlsToPaste)
+	{
+		FRigControlModifiedContext Context;
+		Context.SetKey = bDoKey ? EControlRigSetKey::Always : EControlRigSetKey::DoNotCare;
+		FControlRigPoseMirrorTable MirrorTable;
+		if (bDoMirror)
+		{
+			MirrorTable.SetUpMirrorTable(ControlRig);
+		}
+	
+		for (const FRigControlCopy& CopyRigControl : ControlsToPaste)
+		{
+			FRigControl* RigControl = MirrorTable.GetControl(ControlRig,CopyRigControl.Name);
+			if (RigControl)
+			{
+				if (CopyRigControl.ControlType == RigControl->ControlType)
+				{
+					switch (RigControl->ControlType)
+					{
+					case ERigControlType::Transform:
+					case ERigControlType::TransformNoScale:
+					case ERigControlType::EulerTransform:
+					case ERigControlType::Position:
+					case ERigControlType::Scale:
+					case ERigControlType::Rotator:
+					{
+						if (bDoMirror == false)
+						{
+							ControlRig->SetControlGlobalTransform(RigControl->Name, CopyRigControl.GlobalTransform,Context);
+						}
+						else
+						{
+							FVector GlobalTranslation;
+							FQuat LocalRotation;
+							bool bIsMatched = MirrorTable.IsMatched(RigControl->Name);
+							MirrorTable.GetMirrorTransform(CopyRigControl, bIsMatched,GlobalTranslation, LocalRotation);
+							SetControlMirrorTransform(ControlRig,RigControl->Name, GlobalTranslation, LocalRotation,true, Context);
+						}
+						break;
+					}		
+					case ERigControlType::Float:
+					{
+						float Val = CopyRigControl.Value.Get<float>();
+						ControlRig->SetControlValue<float>(RigControl->Name, Val, true, Context);
+						break;
+					}
+					case ERigControlType::Bool:
+					{
+						bool Val = CopyRigControl.Value.Get<bool>();
+						ControlRig->SetControlValue<bool>(RigControl->Name, Val, true, Context);
+						break;
+					}
+					case ERigControlType::Integer:
+					{
+						int32 Val = CopyRigControl.Value.Get<int32>();
+						ControlRig->SetControlValue<int32>(RigControl->Name, Val, true, Context);
+						break;
+					}
+					case ERigControlType::Vector2D:
+					{
+						FVector2D Val = CopyRigControl.Value.Get<FVector2D>();
+						ControlRig->SetControlValue<FVector2D>(RigControl->Name, Val, true, Context);
+						break;
+					}	
+					default:
+						//TODO add log
+						break;
+					};
+				}
+			}
+		}
+	}
+
+	void BlendWithInitialPoses(TArray<FRigControlCopy>& InitialControls, UControlRig *ControlRig, bool bDoKey, bool bDoMirror, float BlendValue)
+	{
+		if (InitialControls.Num() == 0)
+		{
+			return;
+		}
+
+		//though can be n^2 should be okay, we search from current Index which in most cases will be the same
+		//not run often anyway
+		FRigControlModifiedContext Context;
+		Context.SetKey = bDoKey ? EControlRigSetKey::Always : EControlRigSetKey::DoNotCare;
+		FControlRigPoseMirrorTable MirrorTable;
+		if (bDoMirror)
+		{
+			MirrorTable.SetUpMirrorTable(ControlRig);
+		}
+
+		for (int32 Index = 0; Index <CopyOfControls.Num(); ++Index)
+		{
+			FRigControlCopy& CopyRigControl = CopyOfControls[Index];
+			FRigControl* RigControl = MirrorTable.GetControl(ControlRig, CopyRigControl.Name);
+			if (RigControl)
+			{
+				if (RigControl->ControlType == CopyRigControl.ControlType)
+				{
+					FRigControlCopy* InitialFound = nullptr;
+					for (int32 NewIncIndex = Index; NewIncIndex < InitialControls.Num(); ++NewIncIndex)
+					{
+						if (CopyOfControls[Index].Name == InitialControls[NewIncIndex].Name)
+						{
+							InitialFound = &InitialControls[NewIncIndex];
+							break;
+						}
+					}
+					if (!InitialFound && Index >= 1)
+					{
+						//now search backwards to zero
+						for (int32 NewDecIndex = Index - 1; NewDecIndex >= 0; --NewDecIndex)
+						{
+							if (CopyOfControls[Index].Name == InitialControls[NewDecIndex].Name)
+							{
+								InitialFound = &InitialControls[NewDecIndex];
+								break;
+							}
+						}
+					}
+					if (InitialFound && InitialFound->ControlType == CopyOfControls[Index].ControlType)
+					{
+						if ((CopyRigControl.ControlType == ERigControlType::Transform || CopyRigControl.ControlType == ERigControlType::EulerTransform ||
+							CopyRigControl.ControlType == ERigControlType::TransformNoScale || CopyRigControl.ControlType == ERigControlType::Position ||
+							CopyRigControl.ControlType == ERigControlType::Rotator || CopyRigControl.ControlType == ERigControlType::Scale
+							))
+						{
+							if (bDoMirror == false)
+							{
+								FTransform Val = CopyRigControl.GlobalTransform;
+								FTransform InitialVal = InitialFound->GlobalTransform;
+								FVector Translation, Scale;
+								FQuat Rotation;
+								Translation = FMath::Lerp(InitialVal.GetTranslation(), Val.GetTranslation(), BlendValue);
+								Rotation = FQuat::Slerp(InitialVal.GetRotation(), Val.GetRotation(), BlendValue); //doing slerp here not fast lerp, can be slow this is for content creation
+								Scale = FMath::Lerp(InitialVal.GetScale3D(), Val.GetScale3D(), BlendValue);
+								Val = FTransform(Rotation, Translation, Scale);
+								ControlRig->SetControlGlobalTransform(RigControl->Name, Val, Context);
+							}
+							else
+							{
+								FVector GlobalTranslation;
+								FQuat LocalRotation;
+								bool bIsMatched = MirrorTable.IsMatched(CopyRigControl.Name);
+								MirrorTable.GetMirrorTransform(CopyRigControl, bIsMatched, GlobalTranslation, LocalRotation);
+								FVector InitialTranslation = InitialFound->GlobalTransform.GetTranslation();
+								FQuat InitialRotation = InitialFound->LocalTransform.GetRotation();
+								GlobalTranslation = FMath::Lerp(InitialTranslation, GlobalTranslation, BlendValue);
+								LocalRotation = FQuat::Slerp(InitialRotation, LocalRotation, BlendValue); //doing slerp here not fast lerp, can be slow this is for content creation
+								SetControlMirrorTransform(ControlRig, RigControl->Name, GlobalTranslation, LocalRotation, false, Context);
+							}
+						}
+					}
+				}
+			}			
+		}
+	}
+
+	TArray<FName> GetControlNames() const
+	{
+		TArray<FName> Controls;
+		Controls.Reserve(CopyOfControls.Num());
+		for (const FRigControlCopy& Control : CopyOfControls)
+		{
+			Controls.Add(Control.Name);
+		}
+		return Controls;
+	}
+
+	TArray<FRigControlCopy> GetPoses() const {return CopyOfControls;};
+
+	UPROPERTY()
+	TArray<FRigControlCopy> CopyOfControls;
+};
+
+
+/**
+* The Actual Pose Asset that holds the Pose.
+*/
+UCLASS(BlueprintType)
+class CONTROLRIG_API UControlRigPoseAsset : public UObject
+{
+	GENERATED_UCLASS_BODY()
+
+public:
+
+	UFUNCTION(BlueprintCallable, Category = "Pose")
+	void SavePose(UControlRig* InControlRig, bool bUseAll);
+
+	UFUNCTION(BlueprintCallable, Category = "Pose")
+	void PastePose(UControlRig* InControlRig, bool bDoKey = false, bool bDoMirror = false);
+	
+	UFUNCTION(BlueprintCallable, Category = "Pose")
+	void SelectControls(UControlRig* InControlRig);
+	
+	UFUNCTION(BlueprintCallable, Category = "Pose")
+	const TArray<FName>& GetControlNames() const { return Controls; }
+
+	TArray<FRigControlCopy> GetCurrentPose(UControlRig* InControlRig);
+	void BlendWithInitialPoses(TArray<FRigControlCopy>& InitialControls, UControlRig* ControlRig, bool bDoKey, bool bdoMirror, float BlendValue);
+
+private:
+
+	UPROPERTY()
+	FControlRigControlPose Pose;
+
+	/** Controls we are made of */
+	UPROPERTY(AssetRegistrySearchable)
+	TArray<FName> Controls;
+
+};
diff --git a/Engine/Plugins/Experimental/ControlRig/Source/ControlRig/Public/Tools/ControlRigPoseMirrorSettings.h b/Engine/Plugins/Experimental/ControlRig/Source/ControlRig/Public/Tools/ControlRigPoseMirrorSettings.h
new file mode 100644
--- /dev/null
+++ b/Engine/Plugins/Experimental/ControlRig/Source/ControlRig/Public/Tools/ControlRigPoseMirrorSettings.h
@@ -0,0 +1,37 @@
+// Copyright Epic Games, Inc. All Rights Reserved.
+#pragma once
+/**
+* Per Project Settings for the Mirroring, in particular what axis to use to mirror and the matching strings.
+*/
+#include "CoreMinimal.h"
+#include "UObject/ObjectMacros.h"
+#include "UObject/Object.h"
+#include "ControlRigPoseMirrorSettings.generated.h"
+
+
+UCLASS(config = EditorPerProjectUserSettings)
+class CONTROLRIG_API UControlRigPoseMirrorSettings : public UObject
+{
+
+public:
+	GENERATED_BODY()
+
+	UControlRigPoseMirrorSettings();
+
+	/** The String to denote Controls on the Right Side */
+	UPROPERTY(config, EditAnywhere, BlueprintReadWrite, Category = "Mirror Settings", meta = (ToolTip = "String To Denote Right Side"))
+	FString RightSide;
+
+	/** The String to denote Controls on the Left Side */
+	UPROPERTY(config, EditAnywhere, BlueprintReadWrite, Category = "Mirror Settings", meta = (ToolTip = "String To Denote Left Side"))
+	FString LeftSide;
+
+	UPROPERTY(config, EditAnywhere, BlueprintReadWrite, Category = "Mirror Plane", meta = (ToolTip = "X Axis Plane Normal"))
+	float XAxis;
+	UPROPERTY(config, EditAnywhere, BlueprintReadWrite, Category = "Mirror Plane", meta = (ToolTip = "Y Axis Plane Normal"))
+	float YAxis;
+	UPROPERTY(config, EditAnywhere, BlueprintReadWrite, Category = "Mirror Plane", meta = (ToolTip = "Z Axis Plane Normal"))
+	float ZAxis;
+
+};
+
diff --git a/Engine/Plugins/Experimental/ControlRig/Source/ControlRig/Public/Tools/ControlRigPoseMirrorTable.h b/Engine/Plugins/Experimental/ControlRig/Source/ControlRig/Public/Tools/ControlRigPoseMirrorTable.h
new file mode 100644
--- /dev/null
+++ b/Engine/Plugins/Experimental/ControlRig/Source/ControlRig/Public/Tools/ControlRigPoseMirrorTable.h
@@ -0,0 +1,39 @@
+// Copyright Epic Games, Inc. All Rights Reserved.
+#pragma once
+
+/**
+* Class to hold information on how a Pose May be Mirrored
+*
+*/
+
+#include "CoreMinimal.h"
+#include "UObject/ObjectMacros.h"
+#include "UObject/Object.h"
+
+class UControlRig;
+struct FRigControl;
+struct FRigControlCopy;
+
+struct FControlRigPoseMirrorTable
+{
+public:
+	FControlRigPoseMirrorTable() {};
+	~FControlRigPoseMirrorTable() {};
+
+	/*Set up the Mirror Table*/
+	void SetUpMirrorTable(const UControlRig* ControlRig);
+
+	/*Get the matched control with the given name*/
+	FRigControl* GetControl(UControlRig* ControlRig, const FName& ControlName) const;
+
+	/*Whether or not the Control with this name is matched*/
+	bool IsMatched(const FName& ControlName) const;
+
+	/*Return the Mirrored Global(In Control Rig Space) Translation and Mirrored Local Rotation*/
+	void GetMirrorTransform(const FRigControlCopy& ControlCopy,  bool bIsMatched,FVector& OutGlobalTranslation, FQuat& OutLocalRotation) const;
+
+private:
+
+	TMap<FName, FName>  MatchedControls;
+
+};
\ No newline at end of file
diff --git a/Engine/Plugins/Experimental/ControlRig/Source/ControlRig/Public/Tools/ControlRigPoseProjectSettings.h b/Engine/Plugins/Experimental/ControlRig/Source/ControlRig/Public/Tools/ControlRigPoseProjectSettings.h
new file mode 100644
--- /dev/null
+++ b/Engine/Plugins/Experimental/ControlRig/Source/ControlRig/Public/Tools/ControlRigPoseProjectSettings.h
@@ -0,0 +1,34 @@
+// Copyright Epic Games, Inc. All Rights Reserved.
+#pragma once
+/**
+* Per Project user settings for Control Rig Poses(and maybe animations etc).
+*/
+#include "CoreMinimal.h"
+#include "UObject/Object.h"
+#include "ControlRigPoseProjectSettings.generated.h"
+
+
+UCLASS(config = EditorPerProjectUserSettings)
+class CONTROLRIG_API UControlRigPoseProjectSettings : public UObject
+{
+
+public:
+	GENERATED_BODY()
+
+	UControlRigPoseProjectSettings();
+
+	/** The pose asset path  */
+	FString GetAssetPath() const { return RootSaveDir.Path; }
+
+	/** The root of the directory in which to save poses */
+	UPROPERTY(config, EditAnywhere, BlueprintReadWrite, Category = "Control Rig Poses", meta = (ContentDir))
+	FDirectoryPath RootSaveDir;
+
+	/** Not used but may put bad if we support other types.
+	bool bFilterPoses;
+
+	bool bFilterAnimations;
+
+	bool bFilterSelectionSets;
+	*/
+};
diff --git a/Engine/Plugins/Experimental/ControlRig/Source/ControlRig/Public/Tools/ControlRigToolAsset.h b/Engine/Plugins/Experimental/ControlRig/Source/ControlRig/Public/Tools/ControlRigToolAsset.h
new file mode 100644
--- /dev/null
+++ b/Engine/Plugins/Experimental/ControlRig/Source/ControlRig/Public/Tools/ControlRigToolAsset.h
@@ -0,0 +1,75 @@
+// Copyright Epic Games, Inc. All Rights Reserved.
+#pragma once
+
+#include "CoreMinimal.h"
+#include "UObject/Object.h"
+#include "ControlRig.h"
+#include "Tools/ControlRigPoseProjectSettings.h"
+#include "IAssetRegistry.h"
+#include "AssetRegistryModule.h"
+#include "Widgets/Notifications/SNotificationList.h"
+#include "Framework/Notifications/NotificationManager.h"
+#include "ControlRigObjectBinding.h"
+
+/**
+Create One of the Control Rig Assets, Currently Only support Pose Assets but may add Animation/Selection Sets.
+*/
+class FControlRigToolAsset
+{
+public:
+	template< typename Type> static  UObject* SaveAsset(UControlRig* InControlRig, const FString& CurrentPath, const FString& InString, bool bUseAllControls)
+	{
+		if (InControlRig)
+		{
+			FString InPackageName = CurrentPath;
+			InPackageName += FString("/") + InString;
+			if (!FPackageName::IsValidLongPackageName(InPackageName))
+			{
+				FText ErrorText = FText::Format(NSLOCTEXT("ControlRigToolAsset", "InvalidPathError", "{0} is not a valid asset path."), FText::FromString(InPackageName));
+				FNotificationInfo NotifyInfo(ErrorText);
+				NotifyInfo.ExpireDuration = 5.0f;
+				FSlateNotificationManager::Get().AddNotification(NotifyInfo)->SetCompletionState(SNotificationItem::CS_Fail);
+				return nullptr;
+			}
+
+			int32 UniqueIndex = 1;
+			int32 BasePackageLength = InPackageName.Len();
+
+			// Generate a unique control asset name for this take if there are already assets of the same name
+			IAssetRegistry& AssetRegistry = FModuleManager::LoadModuleChecked<FAssetRegistryModule>("AssetRegistry").Get();
+			while (AssetRegistry.GetAssetByObjectPath(*InPackageName).IsValid())
+			{
+				int32 TrimCount = InPackageName.Len() - BasePackageLength;
+				if (TrimCount > 0)
+				{
+					InPackageName.RemoveAt(BasePackageLength, TrimCount, false);
+				}
+
+				InPackageName += FString::Printf(TEXT("_%04d"), UniqueIndex++);
+			}
+
+			// Create the asset to record into
+			const FString NewAssetName = FPackageName::GetLongPackageAssetName(InPackageName);
+			UPackage* NewPackage = CreatePackage(*InPackageName);
+
+
+			// Create a new Pose From Scracth
+			Type* Asset = NewObject<Type>(NewPackage, *NewAssetName, RF_Public | RF_Standalone | RF_Transactional);
+			USkeletalMesh* SkelMesh = nullptr;
+			if (USkeletalMeshComponent* RigMeshComp = Cast<USkeletalMeshComponent>(InControlRig->GetObjectBinding()->GetBoundObject()))
+			{
+				SkelMesh = RigMeshComp->SkeletalMesh;
+			}
+
+			Asset->SavePose(InControlRig, bUseAllControls);
+			Asset->MarkPackageDirty();
+			FAssetRegistryModule::AssetCreated(Asset);
+			return Asset;
+
+		}
+		return nullptr;
+	}
+};
+
+
+
diff --git a/Engine/Plugins/Experimental/ControlRig/Source/ControlRigEditor/ControlRigEditor.Build.cs b/Engine/Plugins/Experimental/ControlRig/Source/ControlRigEditor/ControlRigEditor.Build.cs
--- a/Engine/Plugins/Experimental/ControlRig/Source/ControlRigEditor/ControlRigEditor.Build.cs
+++ b/Engine/Plugins/Experimental/ControlRig/Source/ControlRigEditor/ControlRigEditor.Build.cs
@@ -13,8 +13,9 @@ namespace UnrealBuildTool.Rules
                     "ControlRigEditor/Private/EditMode",
                     "ControlRigEditor/Private/Graph",
                     "ControlRigEditor/Private/Editor",
+                    "../../../../Source/Editor/UnrealEd/Private", 
 					"../../../../Source/Developer/MessageLog/Private", //compatibility for FBX importer
-					"../../../../Source/Editor/UnrealEd/Private" //compatibility for FBX importer
+					"ControlRigEditor/Private/Tools"
 				}
 			);
 
diff --git a/Engine/Plugins/Experimental/ControlRig/Source/ControlRigEditor/Private/ControlRigEditorModule.cpp b/Engine/Plugins/Experimental/ControlRig/Source/ControlRigEditor/Private/ControlRigEditorModule.cpp
--- a/Engine/Plugins/Experimental/ControlRig/Source/ControlRigEditor/Private/ControlRigEditorModule.cpp
+++ b/Engine/Plugins/Experimental/ControlRig/Source/ControlRigEditor/Private/ControlRigEditorModule.cpp
@@ -105,6 +105,8 @@
 #include "Rigs/FKControlRig.h"
 #include "SBakeToControlRigDialog.h"
 #include "Graph/SControlRigGraphPinVariableBinding.h"
+#include "AssetTypeActions_ControlRigPose.h"
+
 
 #define LOCTEXT_NAMESPACE "ControlRigEditorModule"
 
@@ -183,6 +185,7 @@ void FControlRigEditorModule::StartupModule()
 
 	RegisterAssetTypeAction(MakeShareable(new FControlRigBlueprintActions()));
 	RegisterAssetTypeAction(MakeShareable(new FControlRigGizmoLibraryActions()));
+	RegisterAssetTypeAction(MakeShareable(new FAssetTypeActions_ControlRigPose()));
 	
 	// Register sequencer track editor
 	ISequencerModule& SequencerModule = FModuleManager::Get().LoadModuleChecked<ISequencerModule>("Sequencer");
@@ -238,8 +241,10 @@ void FControlRigEditorModule::StartupModule()
 	);
 
 	UThumbnailManager::Get().RegisterCustomRenderer(UControlRigBlueprint::StaticClass(), UControlRigThumbnailRenderer::StaticClass());
+	//UThumbnailManager::Get().RegisterCustomRenderer(UControlRigPoseAsset::StaticClass(), UControlRigPoseThumbnailRenderer::StaticClass());
 
 	bFilterAssetBySkeleton = true;
+
 }
 
 void FControlRigEditorModule::ShutdownModule()
diff --git a/Engine/Plugins/Experimental/ControlRig/Source/ControlRigEditor/Private/EditMode/ControlRigEditMode.cpp b/Engine/Plugins/Experimental/ControlRig/Source/ControlRigEditor/Private/EditMode/ControlRigEditMode.cpp
--- a/Engine/Plugins/Experimental/ControlRig/Source/ControlRigEditor/Private/EditMode/ControlRigEditMode.cpp
+++ b/Engine/Plugins/Experimental/ControlRig/Source/ControlRigEditor/Private/EditMode/ControlRigEditMode.cpp
@@ -35,8 +35,6 @@
 #include "SEditorViewport.h"
 #include "ControlRigControlsProxy.h"
 #include "ScopedTransaction.h"
-#include "Framework/Notifications/NotificationManager.h"
-#include "Widgets/Notifications/SNotificationList.h"
 #include "RigVMModel/RigVMController.h"
 #include "Rigs/AdditiveControlRig.h"
 #include "Rigs/FKControlRig.h"
@@ -44,6 +42,7 @@
 #include "EngineUtils.h"
 #include "ControlRig/Private/Units/Execution/RigUnit_BeginExecution.h"
 
+
 void UControlRigEditModeDelegateHelper::OnPoseInitialized()
 {
 	if (EditMode)
@@ -90,6 +89,7 @@ void UControlRigEditModeDelegateHelper::RemoveDelegates()
 	}
 }
 
+
 FName FControlRigEditMode::ModeName("EditMode.ControlRig");
 
 #define LOCTEXT_NAMESPACE "ControlRigEditMode"
@@ -2368,5 +2368,17 @@ void FControlRigEditMode::PostPoseUpdate()
 	}
 
 }
+void FControlRigEditMode::SetOnlySelectRigControls(bool Val)
+{
+	Settings->bOnlySelectRigControls = Val;
+}
+
+bool FControlRigEditMode::GetOnlySelectRigControls()const
+{
+	return Settings->bOnlySelectRigControls;
+}
+
+
+
 
 #undef LOCTEXT_NAMESPACE
diff --git a/Engine/Plugins/Experimental/ControlRig/Source/ControlRigEditor/Private/EditMode/ControlRigEditMode.h b/Engine/Plugins/Experimental/ControlRig/Source/ControlRigEditor/Private/EditMode/ControlRigEditMode.h
--- a/Engine/Plugins/Experimental/ControlRig/Source/ControlRigEditor/Private/EditMode/ControlRigEditMode.h
+++ b/Engine/Plugins/Experimental/ControlRig/Source/ControlRigEditor/Private/EditMode/ControlRigEditMode.h
@@ -15,6 +15,7 @@
 #include "ControlRigEditMode.generated.h"
 
 
+
 class FEditorViewportClient;
 class FViewport;
 class UActorFactory;
@@ -288,6 +289,12 @@ private:
 	TArray<ECoordSystem> CoordSystemPerWidgetMode;
 	bool bIsChangingCoordSystem;
 
+
+public:
+	//Toolbar functions
+	void SetOnlySelectRigControls(bool val);
+	bool GetOnlySelectRigControls()const;
+
 private:
 
 	bool CreateGizmoActors(UWorld* World);
diff --git a/Engine/Plugins/Experimental/ControlRig/Source/ControlRigEditor/Private/EditMode/ControlRigEditModeToolkit.cpp b/Engine/Plugins/Experimental/ControlRig/Source/ControlRigEditor/Private/EditMode/ControlRigEditModeToolkit.cpp
new file mode 100644
--- /dev/null
+++ b/Engine/Plugins/Experimental/ControlRig/Source/ControlRigEditor/Private/EditMode/ControlRigEditModeToolkit.cpp
@@ -0,0 +1,94 @@
+// Copyright Epic Games, Inc. All Rights Reserved.
+
+/**
+* Control Rig Edit Mode Toolkit
+*/
+#include "ControlRigEditModeToolkit.h"
+#include "Widgets/DeclarativeSyntaxSupport.h"
+#include "Widgets/SWidget.h"
+#include "EditorModes.h"
+#include "Toolkits/BaseToolkit.h"
+#include "EditorModeManager.h"
+#include "SControlRigEditModeTools.h"
+#include "ControlRigEditMode.h"
+
+#define LOCTEXT_NAMESPACE "FControlRigEditModeToolkit"
+
+namespace 
+{
+	static const FName AnimationName(TEXT("Animation")); 
+	const TArray<FName> AnimationPaletteNames = { AnimationName };
+}
+
+void FControlRigEditModeToolkit::Init(const TSharedPtr<IToolkitHost>& InitToolkitHost)
+{
+	FPropertyEditorModule& PropertyEditorModule = FModuleManager::GetModuleChecked<FPropertyEditorModule>("PropertyEditor");
+
+	FDetailsViewArgs DetailsViewArgs(
+		/*bUpdateFromSelection=*/ false,
+		/*bLockable=*/ false,
+		/*bAllowSearch=*/ false,
+		FDetailsViewArgs::HideNameArea,
+		/*bHideSelectionTip=*/ true,
+		/*InNotifyHook=*/ nullptr,
+		/*InSearchInitialKeyFocus=*/ false,
+		/*InViewIdentifier=*/ NAME_None);
+	DetailsViewArgs.DefaultsOnlyVisibility = EEditDefaultsOnlyNodeVisibility::Automatic;
+	DetailsViewArgs.bShowOptions = false;
+	DetailsViewArgs.bAllowMultipleTopLevelObjects = true;
+
+	DetailsView = PropertyEditorModule.CreateDetailView(DetailsViewArgs);
+	FModeToolkit::Init(InitToolkitHost);
+}
+
+
+void FControlRigEditModeToolkit::GetToolPaletteNames(TArray<FName>& InPaletteName) const
+{
+	InPaletteName = AnimationPaletteNames;
+}
+
+FText FControlRigEditModeToolkit::GetToolPaletteDisplayName(FName PaletteName) const
+{
+	if (PaletteName == AnimationName)
+	{
+		FText::FromName(AnimationName);
+	}
+	return FText();
+}
+
+void FControlRigEditModeToolkit::BuildToolPalette(FName PaletteName, class FToolBarBuilder& ToolBarBuilder)
+{
+	if (PaletteName == AnimationName)
+	{
+		ModeTools->CustomizeToolBarPalette(ToolBarBuilder);
+	}
+}
+
+void FControlRigEditModeToolkit::OnToolPaletteChanged(FName PaletteName)
+{
+
+}
+
+FText FControlRigEditModeToolkit::GetActiveToolDisplayName() const
+{
+	return ModeTools->GetActiveToolName();
+}
+
+FText FControlRigEditModeToolkit::GetActiveToolMessage() const
+{
+
+	return ModeTools->GetActiveToolMessage();
+}
+
+//todo may need these later
+void FControlRigEditModeToolkit::RegisterTabSpawners(const TSharedRef<class FTabManager>& TabManager)
+{
+
+}
+
+void FControlRigEditModeToolkit::UnregisterTabSpawners(const TSharedRef<class FTabManager>& TabManager)
+{
+
+}
+
+#undef LOCTEXT_NAMESPACE
\ No newline at end of file
diff --git a/Engine/Plugins/Experimental/ControlRig/Source/ControlRigEditor/Private/EditMode/ControlRigEditModeToolkit.h b/Engine/Plugins/Experimental/ControlRig/Source/ControlRigEditor/Private/EditMode/ControlRigEditModeToolkit.h
--- a/Engine/Plugins/Experimental/ControlRig/Source/ControlRigEditor/Private/EditMode/ControlRigEditModeToolkit.h
+++ b/Engine/Plugins/Experimental/ControlRig/Source/ControlRigEditor/Private/EditMode/ControlRigEditModeToolkit.h
@@ -34,6 +34,20 @@ public:
 		}
 		return false;
 	}
+	virtual void Init(const TSharedPtr<IToolkitHost>& InitToolkitHost) override;
+
+	/** Mode Toolbar Palettes **/
+	virtual void GetToolPaletteNames(TArray<FName>& InPaletteName) const override;
+	virtual FText GetToolPaletteDisplayName(FName PaletteName) const override;
+	virtual void BuildToolPalette(FName PaletteName, class FToolBarBuilder& ToolbarBuilder) override;
+
+	/** Modes Panel Header Information **/
+	virtual FText GetActiveToolDisplayName() const override;
+	virtual FText GetActiveToolMessage() const override;
+	virtual void OnToolPaletteChanged(FName PaletteName) override;
+
+	virtual void RegisterTabSpawners(const TSharedRef<class FTabManager>& TabManager) override;
+	virtual void UnregisterTabSpawners(const TSharedRef<class FTabManager>& TabManager) override;
 
 private:
 	/** The edit mode we are bound to */
diff --git a/Engine/Plugins/Experimental/ControlRig/Source/ControlRigEditor/Private/EditMode/SControlRigBaseListWidget.cpp b/Engine/Plugins/Experimental/ControlRig/Source/ControlRigEditor/Private/EditMode/SControlRigBaseListWidget.cpp
new file mode 100644
--- /dev/null
+++ b/Engine/Plugins/Experimental/ControlRig/Source/ControlRigEditor/Private/EditMode/SControlRigBaseListWidget.cpp
@@ -0,0 +1,1124 @@
+// Copyright Epic Games, Inc. All Rights Reserved.
+
+
+#include "SControlRigBaseListWidget.h"
+#include "Framework/Notifications/NotificationManager.h"
+#include "Widgets/Notifications/SNotificationList.h"
+#include "Framework/MultiBox/MultiBoxBuilder.h"
+#include "AssetData.h"
+#include "EditorStyleSet.h"
+
+#include "Widgets/Input/SSearchBox.h"
+#include "Widgets/Text/SInlineEditableTextBlock.h"
+#include "ScopedTransaction.h"
+
+#include "ControlRig.h"
+#include "UnrealEdGlobals.h"
+#include "ControlRigEditMode.h"
+#include "Tools/ControlRigPose.h"
+#include "EditorModeManager.h"
+
+#include "IContentBrowserSingleton.h"
+#include "ContentBrowserModule.h"
+#include "Tools/ControlRigPoseProjectSettings.h"
+#include "Widgets/Layout/SScrollBox.h"
+#include "Framework/Commands/GenericCommands.h"
+#include "Framework/Commands/UICommandList.h"
+#include "ControlRigEditModeCommands.h"
+#include "Tools/CreateControlAssetRigSettings.h"
+#include "FileHelpers.h"
+#include "Tools/ControlRigPoseMirrorSettings.h"
+#include "ObjectTools.h"
+
+
+#define LOCTEXT_NAMESPACE "ControlRigBaseListWidget"
+
+enum class FControlRigAssetType {
+	ControlRigPose,
+	ControlRigAnimation,
+	ControlRigSelectionSet
+
+};
+
+/////////////////////////////////////////////////////
+// SControlRigPoseAnimSelectionToolbar
+
+DECLARE_DELEGATE_OneParam(FCreateControlAssetDelegate,
+FString);
+
+struct FCreateControlAssetRigDialog
+{
+	static void GetControlAssetParams(FControlRigAssetType Type, FCreateControlAssetDelegate& Delegate);
+
+};
+
+class SCreateControlAssetRigDialog : public SCompoundWidget
+{
+	SLATE_BEGIN_ARGS(SCreateControlAssetRigDialog) :
+		_AssetType(FControlRigAssetType::ControlRigPose) {}
+
+	SLATE_ARGUMENT(FControlRigAssetType, AssetType)
+
+		SLATE_END_ARGS()
+
+		FControlRigAssetType AssetType;
+
+	~SCreateControlAssetRigDialog()
+	{
+	}
+
+	void Construct(const FArguments& InArgs)
+	{
+		AssetType = InArgs._AssetType;
+
+		FPropertyEditorModule& PropertyEditor = FModuleManager::LoadModuleChecked<FPropertyEditorModule>("PropertyEditor");
+
+		FDetailsViewArgs DetailsViewArgs;
+		DetailsViewArgs.bShowOptions = false;
+		DetailsViewArgs.bAllowSearch = false;
+		DetailsViewArgs.bShowPropertyMatrixButton = false;
+		DetailsViewArgs.bUpdatesFromSelection = false;
+		DetailsViewArgs.bLockable = false;
+		DetailsViewArgs.NameAreaSettings = FDetailsViewArgs::HideNameArea;
+		DetailsViewArgs.ViewIdentifier = "Create Control Asset";
+
+		DetailView = PropertyEditor.CreateDetailView(DetailsViewArgs);
+
+		ChildSlot
+			[
+				SNew(SVerticalBox)
+
+				+ SVerticalBox::Slot()
+			[
+				DetailView.ToSharedRef()
+			]
+
+		+ SVerticalBox::Slot()
+			.AutoHeight()
+			.HAlign(HAlign_Right)
+			.Padding(5.f)
+			[
+				SNew(SButton)
+				.ContentPadding(FMargin(10, 5))
+			.Text(LOCTEXT("CreateControlAssetRig", "Create Asset"))
+			.OnClicked(this, &SCreateControlAssetRigDialog::OnCreateControlAssetRig)
+			]
+
+			];
+
+		switch (AssetType)
+		{
+		case FControlRigAssetType::ControlRigPose:
+		{
+			UCreateControlPoseAssetRigSettings* AssetSettings = GetMutableDefault<UCreateControlPoseAssetRigSettings>();
+			DetailView->SetObject(AssetSettings);
+			break;
+		}
+		/*
+		case FControlRigAssetType::ControlRigAnimation:
+		{
+			UCreateControlAnimationAssetRigSettings* AssetSettings = GetMutableDefault<UCreateControlAnimationAssetRigSettings>();
+			DetailView->SetObject(AssetSettings);
+			break;
+		}
+		case FControlRigAssetType::ControlRigSelectionSet:
+		{
+			UCreateControlSelectionSetAssetRigSettings* AssetSettings = GetMutableDefault<UCreateControlSelectionSetAssetRigSettings>();
+			DetailView->SetObject(AssetSettings);
+			break;
+		}
+		*/
+		};
+
+	}
+
+	void SetDelegate(FCreateControlAssetDelegate& InDelegate)
+	{
+		Delegate = InDelegate;
+	}
+
+private:
+
+	FReply OnCreateControlAssetRig()
+	{
+		FString AssetName("");
+		switch (AssetType)
+		{
+		case FControlRigAssetType::ControlRigPose:
+		{
+			UCreateControlPoseAssetRigSettings* AssetSettings = GetMutableDefault<UCreateControlPoseAssetRigSettings>();
+			AssetName = AssetSettings->AssetName;
+			break;
+		}
+		/*
+		case FControlRigAssetType::ControlRigAnimation:
+		{
+			UCreateControlAnimationAssetRigSettings* AssetSettings = GetMutableDefault<UCreateControlAnimationAssetRigSettings>();
+			AssetName = AssetSettings->AssetName;
+			break;
+		}
+		case FControlRigAssetType::ControlRigSelectionSet:
+		{
+			UCreateControlSelectionSetAssetRigSettings* AssetSettings = GetMutableDefault<UCreateControlSelectionSetAssetRigSettings>();
+			AssetName = AssetSettings->AssetName;
+			break;
+		}
+		*/
+		};
+
+		TSharedPtr<SWindow> Window = FSlateApplication::Get().FindWidgetWindow(AsShared());
+		if (Delegate.IsBound())
+		{
+			Delegate.Execute(AssetName);
+		}
+		if (Window.IsValid())
+		{
+			Window->RequestDestroyWindow();
+		}
+		return FReply::Handled();
+	}
+	TSharedPtr<IDetailsView> DetailView;
+	FCreateControlAssetDelegate  Delegate;
+
+};
+
+
+void FCreateControlAssetRigDialog::GetControlAssetParams(FControlRigAssetType Type, FCreateControlAssetDelegate& InDelegate)
+{
+	FText TitleText;
+	switch (Type)
+	{
+	case FControlRigAssetType::ControlRigPose:
+		TitleText = NSLOCTEXT("ControlRig", "CreateControlAssetRig", "Create Control Rig Pose");
+		break;
+	case FControlRigAssetType::ControlRigAnimation:
+		TitleText = NSLOCTEXT("ControlRig", "CreateControlAssetRig", "Create Control Rig Animation");
+		break;
+	case FControlRigAssetType::ControlRigSelectionSet:
+		TitleText = NSLOCTEXT("ControlRig", "CreateControlAssetRig", "Create Control Rig Selection Set");
+		break;
+	};
+
+
+	// Create the window to choose our options
+	TSharedRef<SWindow> Window = SNew(SWindow)
+		.Title(TitleText)
+		.HasCloseButton(true)
+		.SizingRule(ESizingRule::UserSized)
+		.ClientSize(FVector2D(400.0f, 200.0f))
+		.AutoCenter(EAutoCenter::PreferredWorkArea)
+		.SupportsMinimize(false);
+
+	TSharedRef<SCreateControlAssetRigDialog> DialogWidget = SNew(SCreateControlAssetRigDialog).AssetType(Type);
+	DialogWidget->SetDelegate(InDelegate);
+	Window->SetContent(DialogWidget);
+
+	FSlateApplication::Get().AddWindow(Window);
+
+}
+
+
+/////////////////////////////////////////////////////
+// SControlRigPoseAnimSelectionToolbar
+
+class SControlRigPoseAnimSelectionToolbar : public SCompoundWidget
+{
+
+public:
+	/** Default constructor. */
+	SControlRigPoseAnimSelectionToolbar();
+
+	/** Virtual destructor. */
+	virtual ~SControlRigPoseAnimSelectionToolbar();
+
+	SLATE_BEGIN_ARGS(SControlRigPoseAnimSelectionToolbar) {}
+	SLATE_ARGUMENT(SControlRigBaseListWidget*, OwningControlRigWidget)
+		SLATE_END_ARGS()
+
+		void Construct(const FArguments& InArgs);
+
+
+	void MakeControlRigAssetDialog(FControlRigAssetType Type, bool bSelectAll);
+	
+	//void ToggleFilter(FControlRigAssetType Type);
+	//bool IsOnToggleFilter(FControlRigAssetType Type) const;
+
+	//It's parent so will always be there..
+	SControlRigBaseListWidget* OwningControlRigWidget;
+
+};
+
+
+SControlRigPoseAnimSelectionToolbar::SControlRigPoseAnimSelectionToolbar() :OwningControlRigWidget(nullptr)
+{
+}
+
+SControlRigPoseAnimSelectionToolbar::~SControlRigPoseAnimSelectionToolbar()
+{
+}
+
+
+void SControlRigPoseAnimSelectionToolbar::Construct(const FArguments& InArgs)
+{
+
+	OwningControlRigWidget = InArgs._OwningControlRigWidget;
+
+	FToolBarBuilder ToolbarBuilder(TSharedPtr<const FUICommandList>(), FMultiBoxCustomization::None, TSharedPtr<FExtender>(), true);
+	FUIAction CreatePoseDialog(
+		FExecuteAction::CreateRaw(this, &SControlRigPoseAnimSelectionToolbar::MakeControlRigAssetDialog, FControlRigAssetType::ControlRigPose,false));
+	/*
+	FUIAction CreatePoseFromAllDialog(
+		FExecuteAction::CreateRaw(this, &SControlRigPoseAnimSelectionToolbar::MakeControlRigAssetDialog, FControlRigAssetType::ControlRigPose,true));
+	*/
+	/*
+	FUIAction CreateAnimationDialog(
+		FExecuteAction::CreateRaw(this, &SControlRigPoseAnimSelectionToolbar::MakeControlRigAssetDialog, FControlRigAssetType::ControlRigAnimation));
+	FUIAction CreateSelectionSetDialog(
+		FExecuteAction::CreateRaw(this, &SControlRigPoseAnimSelectionToolbar::MakeControlRigAssetDialog, FControlRigAssetType::ControlRigSelectionSet));
+	*/
+
+	ToolbarBuilder.BeginSection("Create");
+	{
+		ToolbarBuilder.AddToolBarButton(CreatePoseDialog,
+			NAME_None,
+			LOCTEXT("CreatePose", "Create Pose From Selection"),
+			LOCTEXT("CreatePoseTooltip", "Create Pose Asset From Selection"),
+			FSlateIcon(),
+			EUserInterfaceActionType::Button
+		);
+		/*
+		ToolbarBuilder.AddToolBarButton(CreatePoseFromAllDialog,
+			NAME_None,
+			LOCTEXT("CreatePoseAll", "Create Pose From All Controls"),
+			LOCTEXT("CreatePoseAllTooltip", "Create Pose Asset From All Controls Not Just Selected"),
+			FSlateIcon(),
+			EUserInterfaceActionType::Button
+			
+		);
+		*/
+		/** For now just making pose assets
+		ToolbarBuilder.AddToolBarButton(CreateAnimationDialog,
+			NAME_None,
+			LOCTEXT("CreateAnimation", "Create Animation"),
+			LOCTEXT("CreateAnimationTooltip", "Create Control Rig Animation Asset"),
+			FSlateIcon(),
+			EUserInterfaceActionType::Button
+		);
+		ToolbarBuilder.AddToolBarButton(CreateSelectionSetDialog,
+			NAME_None,
+			LOCTEXT("CreateSelectionSet", "Create Selection Set"),
+			LOCTEXT("CreateSelectionSetTooltip", "Create Selection Set Asset"),
+			FSlateIcon(),
+			EUserInterfaceActionType::Button
+		);
+		*/
+
+	}
+	/** TODO if we have multiple asset types we will put back the filters.
+	ToolbarBuilder.EndSection();
+	FUIAction ToggleFilterPose(
+		FExecuteAction::CreateRaw(this, &SControlRigPoseAnimSelectionToolbar::ToggleFilter, FControlRigAssetType::ControlRigPose),
+		FCanExecuteAction(),
+		FIsActionChecked::CreateRaw(this, &SControlRigPoseAnimSelectionToolbar::IsOnToggleFilter, FControlRigAssetType::ControlRigPose));
+	FUIAction ToggleFilterAnimation(
+		FExecuteAction::CreateRaw(this, &SControlRigPoseAnimSelectionToolbar::ToggleFilter, FControlRigAssetType::ControlRigAnimation),
+		FCanExecuteAction(),
+		FIsActionChecked::CreateRaw(this, &SControlRigPoseAnimSelectionToolbar::IsOnToggleFilter, FControlRigAssetType::ControlRigAnimation));
+	FUIAction ToggleFilterSelectionSet(
+		FExecuteAction::CreateRaw(this, &SControlRigPoseAnimSelectionToolbar::ToggleFilter, FControlRigAssetType::ControlRigSelectionSet),
+		FCanExecuteAction(),
+		FIsActionChecked::CreateRaw(this, &SControlRigPoseAnimSelectionToolbar::IsOnToggleFilter, FControlRigAssetType::ControlRigSelectionSet));
+
+	FToolBarBuilder RightToolbarBuilder(TSharedPtr<const FUICommandList>(), FMultiBoxCustomization::None, TSharedPtr<FExtender>(), true);
+	ToolbarBuilder.BeginSection("Filters");
+	{
+		RightToolbarBuilder.AddToolBarButton(ToggleFilterPose,
+			NAME_None,
+			LOCTEXT("FilterPose", "Filter Pose"),
+			LOCTEXT("FilterPoseTooltip", "Toggle to Show or Hide Pose Assets"),
+			FSlateIcon(),
+			EUserInterfaceActionType::ToggleButton
+		);
+		RightToolbarBuilder.AddToolBarButton(ToggleFilterAnimation,
+			NAME_None,
+			LOCTEXT("FilterAnimation", "Filter Animation"),
+			LOCTEXT("FilterAnimationTooltip", "Toggle to Show or Hide Animation Assets"),
+			FSlateIcon(),
+			EUserInterfaceActionType::ToggleButton
+		);
+		RightToolbarBuilder.AddToolBarButton(ToggleFilterSelectionSet,
+			NAME_None,
+			LOCTEXT("FilterSelectionSet", "Filter SelectionSet"),
+			LOCTEXT("FilterSelectionaTooltip", "Toggle to Show or Hide Selction Set Assets"),
+			FSlateIcon(),
+			EUserInterfaceActionType::ToggleButton
+		);
+
+	}
+	ToolbarBuilder.EndSection();
+	*/
+	// Create the tool bar!
+	ChildSlot
+		[
+			SNew(SHorizontalBox)
+
+			+ SHorizontalBox::Slot()
+		.HAlign(HAlign_Left)
+		.VAlign(VAlign_Center)
+		.FillWidth(1.0)
+		.Padding(0.0f)
+		[
+			SNew(SBorder)
+			.Padding(0)
+		.BorderImage(FEditorStyle::GetBrush("NoBorder"))
+		.IsEnabled(FSlateApplication::Get().GetNormalExecutionAttribute())
+		[
+			ToolbarBuilder.MakeWidget()
+		]
+		]
+	/*
+		+ SHorizontalBox::Slot()
+			.HAlign(HAlign_Right)
+			.VAlign(VAlign_Center)
+			.AutoWidth()
+			.Padding(0.0f)
+			[
+				SNew(SBorder)
+				.Padding(0)
+			.BorderImage(FEditorStyle::GetBrush("NoBorder"))
+			.IsEnabled(FSlateApplication::Get().GetNormalExecutionAttribute())
+			[
+				RightToolbarBuilder.MakeWidget()
+			]
+			]*/
+		];
+}
+
+void SControlRigPoseAnimSelectionToolbar::MakeControlRigAssetDialog(FControlRigAssetType Type, bool bSelectAll)
+{
+	FCreateControlAssetDelegate GetNameCallback = FCreateControlAssetDelegate::CreateLambda([this, Type, bSelectAll](FString AssetName)
+	{
+		if (OwningControlRigWidget)
+		{
+			FString Path = OwningControlRigWidget->GetCurrentlySelectedPath();
+
+			FControlRigEditMode* ControlRigEditMode = static_cast<FControlRigEditMode*>(GLevelEditorModeTools().GetActiveMode(FControlRigEditMode::ModeName));
+			if (ControlRigEditMode && ControlRigEditMode->GetControlRig(true))
+			{
+				UObject* NewAsset = nullptr;
+				switch (Type)
+				{
+				case FControlRigAssetType::ControlRigPose:
+					NewAsset = FControlRigToolAsset::SaveAsset<UControlRigPoseAsset>(ControlRigEditMode->GetControlRig(true), Path, AssetName, bSelectAll);
+
+					break;
+				case FControlRigAssetType::ControlRigAnimation:
+					break;
+				case FControlRigAssetType::ControlRigSelectionSet:
+					break;
+				default:
+					break;
+				};
+				if (NewAsset)
+				{
+					FControlRigView::CaptureThumbnail(NewAsset);
+				}
+				OwningControlRigWidget->SelectThisAsset(NewAsset);
+
+			}
+		}
+	});
+
+	FCreateControlAssetRigDialog::GetControlAssetParams(Type, GetNameCallback);
+}
+/*
+void SControlRigPoseAnimSelectionToolbar::ToggleFilter(FControlRigAssetType Type)
+{
+	UControlRigPoseProjectSettings* PoseSettings = GetMutableDefault<UControlRigPoseProjectSettings>();
+	if (PoseSettings)
+	{
+		switch (Type)
+		{
+		case FControlRigAssetType::ControlRigPose:
+			PoseSettings->bFilterPoses = PoseSettings->bFilterPoses ? false : true;
+			break;
+		case FControlRigAssetType::ControlRigAnimation:
+			PoseSettings->bFilterAnimations = PoseSettings->bFilterAnimations ? false : true;
+			break;
+		case FControlRigAssetType::ControlRigSelectionSet:
+			PoseSettings->bFilterSelectionSets = PoseSettings->bFilterSelectionSets ? false : true;
+			break;
+		default:
+			break;
+		};
+	}
+	if (OwningControlRigWidget)
+	{
+		OwningControlRigWidget->FilterChanged();
+	}
+}
+
+bool SControlRigPoseAnimSelectionToolbar::IsOnToggleFilter(FControlRigAssetType Type) const
+{
+	const UControlRigPoseProjectSettings* PoseSettings = GetDefault<UControlRigPoseProjectSettings>();
+	if (PoseSettings)
+	{
+		switch (Type)
+		{
+		case FControlRigAssetType::ControlRigPose:
+			return PoseSettings->bFilterPoses;
+		case FControlRigAssetType::ControlRigAnimation:
+			return PoseSettings->bFilterAnimations;
+		case FControlRigAssetType::ControlRigSelectionSet:
+			return PoseSettings->bFilterSelectionSets;
+		default:
+			break;
+		};
+	}
+
+	return false;
+}
+*/
+
+/////////////////////////////////////////////////////
+// SControlRigBaseListWidget, Main Dialog Window holding the path picker, asset view and pose view.
+
+void SControlRigBaseListWidget::Construct(const FArguments& InArgs)
+{
+
+	FControlRigEditMode* EditMode = GetEditMode();
+	BindCommands();
+
+
+	const UControlRigPoseProjectSettings* PoseSettings = GetDefault<UControlRigPoseProjectSettings>();
+
+	FString PosesDir = PoseSettings->GetAssetPath();
+	CurrentlySelectedPath = PosesDir;
+	FContentBrowserModule& ContentBrowserModule = FModuleManager::Get().LoadModuleChecked<FContentBrowserModule>(TEXT("ContentBrowser"));
+
+	// Configure filter for asset picker
+	FAssetPickerConfig AssetPickerConfig;
+
+	//AssetPickerConfig.Filter.bRecursiveClasses = true;
+	AssetPickerConfig.Filter.ClassNames.Add(UControlRigPoseAsset::StaticClass()->GetFName());
+
+	AssetPickerConfig.InitialAssetViewType = EAssetViewType::Tile;
+	AssetPickerConfig.bAllowDragging = false;
+	AssetPickerConfig.bCanShowFolders = true;
+	AssetPickerConfig.bCanShowRealTimeThumbnails = true;
+	AssetPickerConfig.ThumbnailLabel = EThumbnailLabel::AssetName;
+	AssetPickerConfig.bFocusSearchBoxWhenOpened = false;
+	AssetPickerConfig.Filter.PackagePaths.Add(FName(*PosesDir));
+	AssetPickerConfig.OnAssetSelected = FOnAssetSelected::CreateSP(this, &SControlRigBaseListWidget::OnAssetSelected);
+	AssetPickerConfig.OnAssetsActivated = FOnAssetsActivated::CreateSP(this, &SControlRigBaseListWidget::OnAssetsActivated);
+	AssetPickerConfig.SaveSettingsName = TEXT("ControlPoseDialog");
+	AssetPickerConfig.bCanShowDevelopersFolder = true;
+	AssetPickerConfig.OnFolderEntered = FOnPathSelected::CreateSP(this, &SControlRigBaseListWidget::HandleAssetViewFolderEntered);
+	AssetPickerConfig.OnGetAssetContextMenu = FOnGetAssetContextMenu::CreateSP(this, &SControlRigBaseListWidget::OnGetAssetContextMenu);	
+	AssetPickerConfig.SetFilterDelegates.Add(&SetFilterDelegate);
+	AssetPickerConfig.GetCurrentSelectionDelegates.Add(&GetCurrentSelectionDelegate);
+	AssetPickerConfig.SelectionMode = ESelectionMode::Single;
+	AssetPickerConfig.bAllowDragging = false;
+	AssetPickerConfig.AssetShowWarningText = LOCTEXT("NoPoses_Warning", "No Poses Found, Create One Using Button In Upper Left Corner");
+	AssetPickerConfig.OnIsAssetValidForCustomToolTip = FOnIsAssetValidForCustomToolTip::CreateLambda([](const FAssetData& AssetData) {return AssetData.IsAssetLoaded(); });
+	//AssetPickerConfig.OnGetCustomAssetToolTip = FOnGetCustomAssetToolTip::CreateSP(this, &SControlRigBaseListWidget::CreateCustomAssetToolTip);
+
+	FPathPickerConfig PathPickerConfig;
+	PathPickerConfig.bAddDefaultPath = true;
+	PathPickerConfig.DefaultPath = PosesDir;
+	PathPickerConfig.CustomFolderBlacklist = MakeShared<FBlacklistPaths>();
+	PathPickerConfig.CustomFolderBlacklist.Get()->AddWhitelistItem("PoseLibrary", PosesDir);
+
+	PathPickerConfig.bFocusSearchBoxWhenOpened = false;
+	PathPickerConfig.OnPathSelected = FOnPathSelected::CreateSP(this, &SControlRigBaseListWidget::HandlePathSelected);
+	PathPickerConfig.SetPathsDelegates.Add(&SetPathsDelegate);
+	PathPickerConfig.bAllowContextMenu = true;
+
+	// The root widget in this dialog.
+	TSharedRef<SVerticalBox> MainVerticalBox = SNew(SVerticalBox);
+
+	// Toolbar on Top
+	MainVerticalBox->AddSlot()
+		.AutoHeight()
+		[
+			SNew(SControlRigPoseAnimSelectionToolbar).OwningControlRigWidget(this)
+		];
+
+	// Path/Asset view
+	MainVerticalBox->AddSlot()
+		.HAlign(HAlign_Fill)
+		.FillHeight(0.7)
+		.Padding(0, 0, 0, 4)
+		[
+			SNew(SSplitter)
+
+			+ SSplitter::Slot()
+		.Value(0.33f)
+		[
+			SNew(SBorder)
+			.BorderImage(FEditorStyle::GetBrush("ToolPanel.GroupBorder"))
+		[
+			ContentBrowserModule.Get().CreatePathPicker(PathPickerConfig)
+		]
+		]
+
+	+ SSplitter::Slot()
+		.Value(0.66f)
+		[
+			SNew(SBorder)
+			.BorderImage(FEditorStyle::GetBrush("ToolPanel.GroupBorder"))
+		[
+			ContentBrowserModule.Get().CreateAssetPicker(AssetPickerConfig)
+		]
+		]
+		];
+
+
+	//Bottom Area View Container to hold specific View (Pose/Animation).
+	MainVerticalBox->AddSlot()
+		.HAlign(HAlign_Fill)
+		.FillHeight(0.3)
+		[
+			SAssignNew(ViewContainer, SBox)
+			.Padding(FMargin(5.0f, 0, 0, 0))
+		];
+
+	ChildSlot
+		[
+			MainVerticalBox
+		];
+
+	CurrentViewType = ESelectedControlAsset::Type::None;
+	CreateCurrentView(nullptr);
+}
+
+SControlRigBaseListWidget::~SControlRigBaseListWidget()
+{
+}
+
+void SControlRigBaseListWidget::NotifyUser(FNotificationInfo& NotificationInfo)
+{
+	TSharedPtr<SNotificationItem> Notification = FSlateNotificationManager::Get().AddNotification(NotificationInfo);
+	if (Notification.IsValid())
+	{
+		Notification->SetCompletionState(SNotificationItem::CS_Fail);
+	}
+}
+
+UControlRig* SControlRigBaseListWidget::GetControlRig()
+{
+	if (FControlRigEditMode* EditMode = GetEditMode())
+	{
+		return EditMode->GetControlRig(true);
+	}
+	return nullptr;
+}
+FControlRigEditMode* SControlRigBaseListWidget::GetEditMode()
+{
+	FControlRigEditMode* ControlRigEditMode = static_cast<FControlRigEditMode*>(GLevelEditorModeTools().GetActiveMode(FControlRigEditMode::ModeName));
+	return ControlRigEditMode;
+}
+
+FText SControlRigBaseListWidget::GetAssetNameText() const
+{
+	return FText::FromString(CurrentlyEnteredAssetName);
+}
+
+FText SControlRigBaseListWidget::GetPathNameText() const
+{
+	return FText::FromString(CurrentlySelectedPath);
+}
+
+void SControlRigBaseListWidget::SetCurrentlySelectedPath(const FString& NewPath)
+{
+	CurrentlySelectedPath = NewPath;
+	UpdateInputValidity();
+}
+
+FString SControlRigBaseListWidget::GetCurrentlySelectedPath() const
+{
+	return CurrentlySelectedPath;
+}
+void SControlRigBaseListWidget::SetCurrentlyEnteredAssetName(const FString& NewName)
+{
+	CurrentlyEnteredAssetName = NewName;
+	UpdateInputValidity();
+}
+
+//todo to be used with renaming when it comes in reuse. (next thing to do).
+void SControlRigBaseListWidget::UpdateInputValidity()
+{
+	bLastInputValidityCheckSuccessful = true;
+
+	if (CurrentlyEnteredAssetName.IsEmpty())
+	{
+		bLastInputValidityCheckSuccessful = false;
+	}
+
+	if (bLastInputValidityCheckSuccessful)
+	{
+		if (CurrentlySelectedPath.IsEmpty())
+		{
+			bLastInputValidityCheckSuccessful = false;
+		}
+	}
+
+	/**
+	const FString ObjectPath = GetObjectPathForSave();
+	FText ErrorMessage;
+	const bool bAllowExistingAsset = true;
+
+	FName AssetClassName = AssetClassNames.Num() == 1 ? AssetClassNames[0] : NAME_None;
+	UClass* AssetClass = AssetClassName != NAME_None ? FindObject<UClass>(ANY_PACKAGE, *AssetClassName.ToString(), true) : nullptr;
+
+	if (!ContentBrowserUtils::IsValidObjectPathForCreate(ObjectPath, AssetClass, ErrorMessage, bAllowExistingAsset))
+	{
+		LastInputValidityErrorText = ErrorMessage;
+		bLastInputValidityCheckSuccessful = false;
+	}
+	*/
+}
+
+
+FString SControlRigBaseListWidget::GetObjectPathForSave() const
+{
+	return CurrentlySelectedPath / CurrentlyEnteredAssetName + TEXT(".") + CurrentlyEnteredAssetName;
+}
+
+void SControlRigBaseListWidget::SelectThisAsset(UObject* Asset)
+{
+	if (Asset != nullptr)
+	{
+		if (UControlRigPoseAsset* PoseAsset = Cast<UControlRigPoseAsset>(Asset))
+		{
+			CurrentViewType = ESelectedControlAsset::Type::Pose;
+		}
+		else
+		{
+			CurrentViewType = ESelectedControlAsset::Type::None;
+		}
+		FString Path = FPaths::GetPath(Asset->GetOutermost()->GetPathName());
+		SetCurrentlySelectedPath(Path);
+		SetCurrentlyEnteredAssetName(Asset->GetName());
+	}
+	else
+	{
+		CurrentViewType = ESelectedControlAsset::Type::None;
+	}
+	CreateCurrentView(Asset);
+}
+
+void SControlRigBaseListWidget::OnAssetSelected(const FAssetData& AssetData)
+{
+	UObject* Asset = nullptr;
+	if (AssetData.IsValid())
+	{
+		Asset = AssetData.GetAsset();
+	}
+	SelectThisAsset(Asset);
+}
+
+void SControlRigBaseListWidget::OnAssetsActivated(const TArray<FAssetData>& SelectedAssets, EAssetTypeActivationMethod::Type ActivationType)
+{
+	if (SelectedAssets.Num() == 1 && (ActivationType == EAssetTypeActivationMethod::DoubleClicked))
+	{
+		UObject* Asset = nullptr;
+
+		if (SelectedAssets[0].IsValid())
+		{
+			Asset = SelectedAssets[0].GetAsset();
+			UControlRigPoseAsset* PoseAsset = Cast<UControlRigPoseAsset>(Asset);
+			if (PoseAsset)
+			{
+				ExecutePastePose(PoseAsset);
+			}
+		}
+		SelectThisAsset(Asset);
+	}
+}
+/*
+void SControlRigBaseListWidget::FilterChanged()
+{
+	FARFilter NewFilter;
+	NewFilter.PackagePaths.Add(FName(*CurrentlySelectedPath));
+
+	const UControlRigPoseProjectSettings* PoseSettings = GetDefault<UControlRigPoseProjectSettings>();
+	if (PoseSettings)
+	{
+		if (PoseSettings->bFilterPoses)
+		{
+			NewFilter.ClassNames.Add(UControlRigPoseAsset::StaticClass()->GetFName());
+		}
+		if (PoseSettings->bFilterAnimations)
+		{
+			//NewFilter.ClassNames.Add(UControlRigAnimationAsset::StaticClass()->GetFName());
+		}
+		if (PoseSettings->bFilterSelectionSets)
+		{
+			//NewFilter.ClassNames.Add(UControlRigSelectionSetAsset::StaticClass()->GetFName());
+		}
+	}
+	SetFilterDelegate.ExecuteIfBound(NewFilter);
+}
+*/
+
+void SControlRigBaseListWidget::HandlePathSelected(const FString& NewPath)
+{
+	SetCurrentlySelectedPath(NewPath);
+	//FilterChanged();
+
+}
+
+void SControlRigBaseListWidget::HandleAssetViewFolderEntered(const FString& NewPath)
+{
+	SetCurrentlySelectedPath(NewPath);
+
+	TArray<FString> NewPaths;
+	NewPaths.Add(NewPath);
+	SetPathsDelegate.Execute(NewPaths);
+}
+
+
+TSharedPtr<SWidget> SControlRigBaseListWidget::OnGetFolderContextMenu(const TArray<FString>& SelectedPaths, FContentBrowserMenuExtender_SelectedPaths InMenuExtender, FOnCreateNewFolder InOnCreateNewFolder)
+{
+
+	TSharedPtr<FExtender> Extender;
+	if (InMenuExtender.IsBound())
+	{
+		Extender = InMenuExtender.Execute(SelectedPaths);
+	}
+
+	FMenuBuilder MenuBuilder(true /*bInShouldCloseWindowAfterMenuSelection*/, Commands, Extender);
+	MenuBuilder.BeginSection("AssetDialogOptions", LOCTEXT("AssetDialogMenuHeading", "Options"));
+
+	if (SelectedPaths.Num() == 1)
+	{
+		FUIAction Action = FUIAction(FExecuteAction::CreateRaw((this), &SControlRigBaseListWidget::ExecuteAddFolder, SelectedPaths[0]));
+		const FText Label = LOCTEXT("AddFolder", "Add Folder");
+		const FText ToolTipText = LOCTEXT("AddFolder Tooltip", "Add Folder to the current Selected Folder");
+		MenuBuilder.AddMenuEntry(Label, ToolTipText, FSlateIcon(), Action);
+
+		FUIAction Action2 = FUIAction(FExecuteAction::CreateRaw((this), &SControlRigBaseListWidget::ExecuteRenameFolder, SelectedPaths[0]));
+		const FText Label2 = LOCTEXT("RenameFolder", "Rename Folder");
+		const FText ToolTipText2 = LOCTEXT("RenameFolderTooltip", "Rename Selected Folder.");
+		MenuBuilder.AddMenuEntry(Label2, ToolTipText2, FSlateIcon(), Action2);
+	}
+	else if (SelectedPaths.Num() > 0)
+	{
+		FUIAction Action = FUIAction(FExecuteAction::CreateRaw((this), &SControlRigBaseListWidget::ExecuteDeleteFolder, SelectedPaths));
+		const FText Label = LOCTEXT("DeleteFolder", "Delete Folder");
+		const FText ToolTipText = LOCTEXT("DeleteFolderTooltip", "Delete Selecte Folder(s), Note this will delete content.");
+		MenuBuilder.AddMenuEntry(Label, ToolTipText, FSlateIcon(), Action);
+
+	}
+	MenuBuilder.EndSection();
+
+	return MenuBuilder.MakeWidget();
+}
+
+TSharedPtr<SWidget> SControlRigBaseListWidget::OnGetAssetContextMenu(const TArray<FAssetData>& SelectedAssets)
+{
+	FMenuBuilder MenuBuilder(true /*bInShouldCloseWindowAfterMenuSelection*/, Commands);
+
+	if (SelectedAssets.Num() > 0)
+	{
+		MenuBuilder.BeginSection("PoseDialogOptions", LOCTEXT("Asset", "Asset"));
+		{
+			{
+				FUIAction Action = FUIAction(FExecuteAction::CreateRaw((this), &SControlRigBaseListWidget::ExecuteSaveAssets, SelectedAssets));
+				const FText Label = LOCTEXT("SaveAssetButton", "Save Asset");
+				const FText ToolTipText = LOCTEXT("SaveAssetButtonTooltip", "Save the Selected Assets.");
+				MenuBuilder.AddMenuEntry(Label, ToolTipText, FSlateIcon(), Action);
+			}
+			{
+				FUIAction Action = FUIAction(FExecuteAction::CreateRaw((this), &SControlRigBaseListWidget::ExecuteDeleteAssets, SelectedAssets));
+				const FText Label = LOCTEXT("DeleteAssetButton", "Delete Asset");
+				const FText ToolTipText = LOCTEXT("DeleteAssetButtonTooltip", "Delete the Selected Assets.");
+				MenuBuilder.AddMenuEntry(Label, ToolTipText, FSlateIcon(), Action);
+			}
+		}
+		MenuBuilder.EndSection();
+
+	}
+	if (SelectedAssets.Num() == 1)
+	{
+		UObject* SelectedAsset = SelectedAssets[0].GetAsset();
+		if (SelectedAsset == nullptr)
+		{
+			return nullptr;
+		}
+
+		UControlRigPoseAsset* PoseAsset = Cast<UControlRigPoseAsset>(SelectedAsset);
+		if (PoseAsset)
+		{
+
+			MenuBuilder.BeginSection("PoseDialogOptions", LOCTEXT("PoseDialogMenuHeading", "Paste"));
+			{
+
+				FUIAction Action = FUIAction(FExecuteAction::CreateRaw((this), &SControlRigBaseListWidget::ExecutePastePose, PoseAsset),
+					FCanExecuteAction::CreateRaw(this, &SControlRigBaseListWidget::CanExecutePastePose, PoseAsset));
+				const FText Label = LOCTEXT("PastePoseButton", "Paste Pose");
+				const FText ToolTipText = LOCTEXT("PastePoseButtonTooltip", "Paste the Selected Pose.");
+				MenuBuilder.AddMenuEntry(Label, ToolTipText, FSlateIcon(), Action);
+			}
+
+			{
+				FUIAction Action = FUIAction(FExecuteAction::CreateRaw((this), &SControlRigBaseListWidget::ExecutePasteMirrorPose, PoseAsset),
+					FCanExecuteAction::CreateRaw(this, &SControlRigBaseListWidget::CanExecutePasteMirrorPose, PoseAsset));
+				const FText Label = LOCTEXT("PasteMirrorPoseButton", "Paste Mirror Pose");
+				const FText ToolTipText = LOCTEXT("PastePoseButtonTooltip", "Paste the Mirror Pose.");
+				MenuBuilder.AddMenuEntry(Label, ToolTipText, FSlateIcon(), Action);
+			}
+
+			MenuBuilder.EndSection();
+
+			MenuBuilder.BeginSection("PoseDialogOptions", LOCTEXT("PoseDialogSelectHeading", "Selection"));
+			{
+				FUIAction Action = FUIAction(FExecuteAction::CreateRaw((this), &SControlRigBaseListWidget::ExecuteSelectControls, PoseAsset));
+				const FText Label = LOCTEXT("SelectControls", "Select Controls");
+				const FText ToolTipText = LOCTEXT("SelectControlsTooltip", "Select Controls in this Pose on Active Control Rig");
+				MenuBuilder.AddMenuEntry(Label, ToolTipText, FSlateIcon(), Action);
+			}
+			MenuBuilder.EndSection();
+
+			return MenuBuilder.MakeWidget();
+		}
+	}
+	return nullptr;
+}
+
+
+void SControlRigBaseListWidget::BindCommands()
+{
+	Commands = TSharedPtr< FUICommandList >(new FUICommandList);
+
+
+}
+void SControlRigBaseListWidget::ExecuteRenameFolder(const FString SelectedPath)
+{
+
+}
+
+void SControlRigBaseListWidget::ExecuteAddFolder(const FString SelectedPath)
+{
+
+}
+void SControlRigBaseListWidget::ExecuteSaveAssets(const TArray<FAssetData> SelectedAssets)
+{
+	TArray<UPackage*> PackagesToSave;
+	for (const FAssetData& AssetData : SelectedAssets)
+	{
+		if (AssetData.IsValid() && AssetData.GetPackage())
+		{
+			PackagesToSave.Add(AssetData.GetPackage());
+			FEditorFileUtils::PromptForCheckoutAndSave(PackagesToSave, false /*bCheckDirty*/, false /*bPromptToSave*/);
+		}
+	}
+
+}
+void SControlRigBaseListWidget::ExecuteDeleteAssets(const TArray<FAssetData> SelectedAssets)
+{
+	ObjectTools::DeleteAssets(SelectedAssets);
+	SelectThisAsset(nullptr);
+
+}
+void SControlRigBaseListWidget::ExecuteDeleteFolder(const TArray<FString> SelectedFolders)
+{
+
+	// Don't allow asset deletion during PIE
+	if (GIsEditor)
+	{
+		UEditorEngine* Editor = GEditor;
+		FWorldContext* PIEWorldContext = GEditor->GetPIEWorldContext();
+		if (PIEWorldContext)
+		{
+			FNotificationInfo Notification(LOCTEXT("CannotDeleteAssetInPIE", "Assets cannot be deleted while in PIE."));
+			Notification.ExpireDuration = 3.0f;
+			FSlateNotificationManager::Get().AddNotification(Notification);
+			return;
+		}
+	}
+	/*
+	const TArray<FContentBrowserItem> SelectedFiles = AssetPicker->GetAssetView()->GetSelectedFileItems();
+
+	// Batch these by their data sources
+	TMap<UContentBrowserDataSource*, TArray<FContentBrowserItemData>> SourcesAndItems;
+	for (const FContentBrowserItem& SelectedItem : SelectedFiles)
+	{
+		FContentBrowserItem::FItemDataArrayView ItemDataArray = SelectedItem.GetInternalItems();
+		for (const FContentBrowserItemData& ItemData : ItemDataArray)
+		{
+			if (UContentBrowserDataSource* ItemDataSource = ItemData.GetOwnerDataSource())
+			{
+				FText DeleteErrorMsg;
+				if (ItemDataSource->CanDeleteItem(ItemData, &DeleteErrorMsg))
+				{
+					TArray<FContentBrowserItemData>& ItemsForSource = SourcesAndItems.FindOrAdd(ItemDataSource);
+					ItemsForSource.Add(ItemData);
+				}
+				else
+				{
+					AssetViewUtils::ShowErrorNotifcation(DeleteErrorMsg);
+				}
+			}
+		}
+	}
+
+	// Execute the operation now
+	for (const auto& SourceAndItemsPair : SourcesAndItems)
+	{
+		SourceAndItemsPair.Key->BulkDeleteItems(SourceAndItemsPair.Value);
+	}
+	*/
+	// If we had any folders selected, ask the user whether they want to delete them 
+	// as it can be slow to build the deletion dialog on an accidental click
+	if (SelectedFolders.Num() > 0)
+	{
+		FText Prompt;
+		if (SelectedFolders.Num() == 1)
+		{
+			//			Prompt = FText::Format(LOCTEXT("FolderDeleteConfirm_Single", "Delete folder '{0}'?"), SelectedFolders[0]);
+		}
+		else
+		{
+			//		Prompt = FText::Format(LOCTEXT("FolderDeleteConfirm_Multiple", "Delete {0} folders?"), SelectedFolders.Num());
+		}
+
+		FContentBrowserModule& ContentBrowserModule = FModuleManager::Get().LoadModuleChecked<FContentBrowserModule>(TEXT("ContentBrowser"));
+		/*
+		// Spawn a confirmation dialog since this is potentially a highly destructive operation
+		ContentBrowserModule.Get().DisplayConfirmationPopup(
+			Prompt,
+			LOCTEXT("FolderDeleteConfirm_Yes", "Delete"),
+			LOCTEXT("FolderDeleteConfirm_No", "Cancel"),
+			ToSharedRef(),
+			FOnClicked::CreateSP(this, &SControlRigBaseListWidget::ExecuteDeleteFolderConfirmed)
+		);
+		*/
+	}
+
+}
+
+FReply SControlRigBaseListWidget::ExecuteDeleteFolderConfirmed()
+{
+	/*
+	TArray< FString > SelectedFolders = AssetPicker->GetAssetView()->GetSelectedFolders();
+
+	if (SelectedFolders.Num() > 0)
+	{
+		ContentBrowserUtils::DeleteFolders(SelectedFolders);
+	}
+	else
+	{
+		const TArray<FString>& SelectedPaths = PathPicker->GetPaths();
+
+		if (SelectedPaths.Num() > 0)
+		{
+			if (ContentBrowserUtils::DeleteFolders(SelectedPaths))
+			{
+				// Since the contents of the asset view have just been deleted, set the selected path to the default "/Game"
+				TArray<FString> DefaultSelectedPaths;
+				DefaultSelectedPaths.Add(TEXT("/Game"));
+				PathPicker->GetPathView()->SetSelectedPaths(DefaultSelectedPaths);
+
+				FSourcesData DefaultSourcesData(FName("/Game"));
+				AssetPicker->GetAssetView()->SetSourcesData(DefaultSourcesData);
+			}
+		}
+	}
+	*/
+	return FReply::Handled();
+}
+
+
+
+void SControlRigBaseListWidget::ExecutePastePose(UControlRigPoseAsset* PoseAsset)
+{
+	if (PoseAsset)
+	{
+		FControlRigEditMode* ControlRigEditMode = static_cast<FControlRigEditMode*>(GLevelEditorModeTools().GetActiveMode(FControlRigEditMode::ModeName));
+		if (ControlRigEditMode && ControlRigEditMode->GetControlRig(true))
+		{
+			PoseAsset->PastePose(ControlRigEditMode->GetControlRig(true));
+		}
+	}
+}
+bool SControlRigBaseListWidget::CanExecutePastePose(UControlRigPoseAsset* PoseAsset) const
+{
+	return PoseAsset != nullptr;
+}
+void SControlRigBaseListWidget::ExecuteSelectControls(UControlRigPoseAsset* PoseAsset)
+{
+	FControlRigEditMode* ControlRigEditMode = static_cast<FControlRigEditMode*>(GLevelEditorModeTools().GetActiveMode(FControlRigEditMode::ModeName));
+	if (ControlRigEditMode && ControlRigEditMode->GetControlRig(true))
+	{
+		PoseAsset->SelectControls(ControlRigEditMode->GetControlRig(true));
+	}
+}
+
+void SControlRigBaseListWidget::ExecutePasteMirrorPose(UControlRigPoseAsset* PoseAsset)
+{
+	if (PoseAsset)
+	{
+		FControlRigEditMode* ControlRigEditMode = static_cast<FControlRigEditMode*>(GLevelEditorModeTools().GetActiveMode(FControlRigEditMode::ModeName));
+		if (ControlRigEditMode && ControlRigEditMode->GetControlRig(true))
+		{
+			PoseAsset->PastePose(ControlRigEditMode->GetControlRig(true), false, true);
+		}
+	}
+}
+
+bool SControlRigBaseListWidget::CanExecutePasteMirrorPose(UControlRigPoseAsset* PoseAsset) const
+{
+	return PoseAsset != nullptr;
+}
+
+void SControlRigBaseListWidget::CreateCurrentView(UObject* Asset)
+{
+	PoseView.Reset();
+	AnimationView.Reset();
+	SelectionSetView.Reset();
+	EmptyBox.Reset();
+	TSharedRef<SWidget> NewView = SNullWidget::NullWidget;
+	switch (CurrentViewType)
+	{
+	case ESelectedControlAsset::Type::Pose:
+		PoseView = CreatePoseView(Asset);
+		ViewContainer->SetContent(PoseView.ToSharedRef());
+		break;
+	case ESelectedControlAsset::Type::Animation:
+		AnimationView = CreateAnimationView(Asset);
+
+		ViewContainer->SetContent(AnimationView.ToSharedRef());
+		break;
+	case ESelectedControlAsset::Type::SelectionSet:
+		SelectionSetView = CreateSelectionSetView(Asset);
+
+		ViewContainer->SetContent(SelectionSetView.ToSharedRef());
+		break;
+	case ESelectedControlAsset::Type::None:
+		EmptyBox = SNew(SBox);
+		ViewContainer->SetContent(EmptyBox.ToSharedRef());
+		break;
+	}
+
+}
+
+TSharedRef<class SControlRigPoseView> SControlRigBaseListWidget::CreatePoseView(UObject* InObject)
+{
+	UControlRigPoseAsset* PoseAsset = Cast<UControlRigPoseAsset>(InObject);
+	return SNew(SControlRigPoseView).
+		PoseAsset(PoseAsset);
+}
+
+TSharedRef<class SControlRigPoseView> SControlRigBaseListWidget::CreateAnimationView(UObject* InObject)
+{
+	return SNew(SControlRigPoseView);
+}
+
+TSharedRef<class SControlRigPoseView> SControlRigBaseListWidget::CreateSelectionSetView(UObject* InObject)
+{
+	return SNew(SControlRigPoseView);
+}
+
+#undef LOCTEXT_NAMESPACE
diff --git a/Engine/Plugins/Experimental/ControlRig/Source/ControlRigEditor/Private/EditMode/SControlRigBaseListWidget.h b/Engine/Plugins/Experimental/ControlRig/Source/ControlRigEditor/Private/EditMode/SControlRigBaseListWidget.h
new file mode 100644
--- /dev/null
+++ b/Engine/Plugins/Experimental/ControlRig/Source/ControlRigEditor/Private/EditMode/SControlRigBaseListWidget.h
@@ -0,0 +1,129 @@
+// Copyright Epic Games, Inc. All Rights Reserved.
+
+/**
+* Main Window and associated classes that holds the Path View, Asset Browswer
+* and the Control Views
+*/
+#pragma once
+
+#include "CoreMinimal.h"
+#include "Widgets/SWidget.h"
+#include "Widgets/DeclarativeSyntaxSupport.h"
+#include "Widgets/SCompoundWidget.h"
+#include "Toolkits/AssetEditorToolkit.h"
+#include "Widgets/Views/STableViewBase.h"
+#include "Widgets/Views/STableRow.h"
+#include "Widgets/Views/SListView.h"
+#include "EditorUndoClient.h"
+#include "AssetData.h"
+#include "IAssetTools.h"
+#include "ContentBrowserDelegates.h"
+#include "SControlRigControlViews.h"
+
+class UControlRig;
+class FControlRigEditMode;
+class UControlRigPoseAsset;
+
+namespace ESelectedControlAsset
+{
+	enum Type
+	{
+		None,
+		Pose,
+		Animation,
+		SelectionSet,
+
+		MAX
+	};
+}
+
+class SControlRigBaseListWidget : public SCompoundWidget
+{
+public:
+	SLATE_BEGIN_ARGS(SControlRigBaseListWidget)
+
+	{}
+
+	SLATE_END_ARGS()
+public:
+	void Construct(const FArguments& InArgs);
+
+
+	SControlRigBaseListWidget() : bLastInputValidityCheckSuccessful(false) {}
+	~SControlRigBaseListWidget();
+
+	FControlRigEditMode* GetEditMode();
+	UControlRig* GetControlRig(); 
+
+	/*Notification from Toolbar that the Filter Changed and so we should reset the filter.
+	Currently disabled
+	void FilterChanged();
+	*/
+	/*Current Path we use to save the asset.*/
+	FString GetCurrentlySelectedPath()const;
+
+	/* Select this Asset, make sure Views all sync up*/
+	void SelectThisAsset(UObject* Asset);
+private:
+
+	void BindCommands();
+
+	/** Asset Related Functions*/
+	FText GetAssetNameText()const;
+	FText GetPathNameText()const;
+	void SetCurrentlySelectedPath(const FString& NewPath);
+	void SetCurrentlyEnteredAssetName(const FString& NewName);
+	void UpdateInputValidity();
+	FString GetObjectPathForSave() const;
+
+	/* Asset Browser Functions*/
+	void OnAssetSelected(const FAssetData& AssetData);
+	void OnAssetsActivated(const TArray<FAssetData>& SelectedAssets, EAssetTypeActivationMethod::Type ActivationType);
+	void HandleAssetViewFolderEntered(const FString& NewPath);
+	TSharedPtr<SWidget> OnGetAssetContextMenu(const TArray<FAssetData>& SelectedAssets);
+	void HandlePathSelected(const FString& NewPath);
+	TSharedPtr<SWidget> OnGetFolderContextMenu(const TArray<FString>& SelectedPaths, FContentBrowserMenuExtender_SelectedPaths InMenuExtender, FOnCreateNewFolder InOnCreateNewFolder);
+	FSetARFilterDelegate SetFilterDelegate;
+	FGetCurrentSelectionDelegate GetCurrentSelectionDelegate;
+	FSetPathPickerPathsDelegate SetPathsDelegate;
+
+	/** Context Menu Functions*/
+	void ExecuteDeleteFolder(const TArray<FString> SelectedPaths);
+	void ExecuteRenameFolder(const FString SelectedPath);
+	void ExecuteAddFolder(const FString SelectedPath);
+	FReply ExecuteDeleteFolderConfirmed();
+	void ExecuteSaveAssets(const TArray<FAssetData> SelectedAssets);
+	void ExecuteDeleteAssets(const TArray<FAssetData> SelectedAssets);
+	void ExecutePastePose(UControlRigPoseAsset* PoseAsset);
+	bool CanExecutePastePose(UControlRigPoseAsset* PoseAsset) const;
+	void ExecuteSelectControls(UControlRigPoseAsset* PoseAsset);
+	void ExecutePasteMirrorPose(UControlRigPoseAsset* PoseAsset);
+	bool CanExecutePasteMirrorPose(UControlRigPoseAsset* PoseAsset) const;
+
+	/** Create the Views for the Selected Asset*/
+	void CreateCurrentView(UObject* Asset);
+	TSharedRef<SControlRigPoseView> CreatePoseView(UObject* Asset);
+	TSharedRef<SControlRigPoseView> CreateAnimationView(UObject* Asset);
+	TSharedRef<SControlRigPoseView> CreateSelectionSetView(UObject* Asset);
+
+	/** View Type Data*/
+	ESelectedControlAsset::Type CurrentViewType;
+	TSharedPtr<SControlRigPoseView> PoseView;
+	TSharedPtr<SControlRigPoseView> AnimationView; //todo
+	TSharedPtr<SControlRigPoseView> SelectionSetView; //todo
+	TSharedPtr<SBox> EmptyBox;
+	TSharedPtr<SBox> ViewContainer;
+
+
+
+	FString CurrentlySelectedPath;
+	FString CurrentlyEnteredAssetName;
+	bool bLastInputValidityCheckSuccessful;
+
+	/** Commands handled by this widget */
+	TSharedPtr< FUICommandList > Commands;
+
+	/** Utility function to display notifications to the user */
+	void NotifyUser(FNotificationInfo& NotificationInfo);
+
+};
\ No newline at end of file
diff --git a/Engine/Plugins/Experimental/ControlRig/Source/ControlRigEditor/Private/EditMode/SControlRigControlViews.cpp b/Engine/Plugins/Experimental/ControlRig/Source/ControlRigEditor/Private/EditMode/SControlRigControlViews.cpp
new file mode 100644
--- /dev/null
+++ b/Engine/Plugins/Experimental/ControlRig/Source/ControlRigEditor/Private/EditMode/SControlRigControlViews.cpp
@@ -0,0 +1,559 @@
+// Copyright Epic Games, Inc. All Rights Reserved.
+
+
+#include "SControlRigControlViews.h"
+#include "ControlRig.h"
+#include "Tools/ControlRigPose.h"
+#include "Widgets/Input/SEditableTextBox.h"
+#include "EditorStyleSet.h"
+#include "ScopedTransaction.h"
+#include "Editor/EditorEngine.h"
+#include "AssetViewUtils.h"
+#include "AssetRegistryModule.h"
+#include "IAssetRegistry.h"
+#include "AssetToolsModule.h"
+#include "AssetThumbnail.h"
+#include "FileHelpers.h"
+#include "Tools/ControlRigPoseMirrorSettings.h"
+#include "ControlRigEditMode.h"
+#include "EditorModeManager.h"
+
+#define LOCTEXT_NAMESPACE "ControlRigBaseListWidget"
+
+
+void FControlRigView::CaptureThumbnail(UObject* Asset)
+{
+	FViewport* Viewport = GEditor->GetActiveViewport();
+
+	if (GCurrentLevelEditingViewportClient && Viewport)
+	{
+		//have to re-render the requested viewport
+		FLevelEditorViewportClient* OldViewportClient = GCurrentLevelEditingViewportClient;
+		//remove selection box around client during render
+		GCurrentLevelEditingViewportClient = NULL;
+		Viewport->Draw();
+
+		TArray<FAssetData> SelectedAssets;
+		FAssetRegistryModule& AssetRegistryModule = FModuleManager::Get().LoadModuleChecked<FAssetRegistryModule>(TEXT("AssetRegistry"));
+		FName PathName = *Asset->GetPathName();
+		FAssetData AssetData = AssetRegistryModule.Get().GetAssetByObjectPath(PathName);
+		SelectedAssets.Emplace(AssetData);
+		AssetViewUtils::CaptureThumbnailFromViewport(Viewport, SelectedAssets);
+
+		//redraw viewport to have the yellow highlight again
+		GCurrentLevelEditingViewportClient = OldViewportClient;
+		Viewport->Draw();
+	}
+	
+}
+
+/** Widget wraps an editable text box for editing name of the asset */
+class SControlRigAssetEditableTextBox : public SCompoundWidget
+{
+public:
+	SLATE_BEGIN_ARGS(SControlRigAssetEditableTextBox) {}
+
+		SLATE_ARGUMENT(TWeakObjectPtr<UObject>, Asset)
+
+		SLATE_END_ARGS()
+
+		/**
+		 * Construct this widget
+		 *
+		 * @param	InArgs	The declaration data for this widget
+		 */
+		void Construct(const FArguments& InArgs);
+
+
+private:
+
+	/** Getter for the Text attribute of the editable text inside this widget */
+	FText GetNameText() const;
+
+	/** Getter for the ToolTipText attribute of the editable text inside this widget */
+	FText GetNameTooltipText() const;
+
+
+	/** Getter for the OnTextCommitted event of the editable text inside this widget */
+	void OnNameTextCommitted(const FText& NewText, ETextCommit::Type InTextCommit);
+
+	/** Callback to verify a text change */
+	void OnTextChanged(const FText& InLabel);
+
+	/** The list of objects whose names are edited by the widget */
+	TWeakObjectPtr<UObject> Asset;
+
+	/** The text box used to edit object names */
+	TSharedPtr< SEditableTextBox > TextBox;
+
+};
+
+void SControlRigAssetEditableTextBox::Construct(const FArguments& InArgs)
+{
+	Asset = InArgs._Asset;
+	ChildSlot
+		[
+			SAssignNew(TextBox, SEditableTextBox)
+			.Text(this, &SControlRigAssetEditableTextBox::GetNameText)
+			.ToolTipText(this, &SControlRigAssetEditableTextBox::GetNameTooltipText)
+			.OnTextCommitted(this, &SControlRigAssetEditableTextBox::OnNameTextCommitted)
+			.OnTextChanged(this, &SControlRigAssetEditableTextBox::OnTextChanged)
+			.RevertTextOnEscape(true)
+		];
+}
+
+FText SControlRigAssetEditableTextBox::GetNameText() const
+{
+	if (Asset.IsValid())
+	{
+		FString Result = Asset.Get()->GetName();
+		return FText::FromString(Result);
+	}
+	return FText();
+}
+
+FText SControlRigAssetEditableTextBox::GetNameTooltipText() const
+{
+	FText Result = FText::Format(LOCTEXT("AssetRenameTooltip", "Rename the selected {0}"), FText::FromString(Asset.Get()->GetClass()->GetName()));
+	
+	return Result;
+}
+
+
+void SControlRigAssetEditableTextBox::OnNameTextCommitted(const FText& NewText, ETextCommit::Type InTextCommit)
+{
+
+	if (InTextCommit != ETextCommit::OnCleared)
+	{
+		FText TrimmedText = FText::TrimPrecedingAndTrailing(NewText);
+
+		if (!TrimmedText.IsEmpty())
+		{
+
+			IAssetRegistry& AssetRegistry = FModuleManager::LoadModuleChecked<FAssetRegistryModule>("AssetRegistry").Get();
+			FAssetData AssetData = AssetRegistry.GetAssetByObjectPath(FName(*Asset->GetPathName()));
+			const FString PackagePath = FPackageName::GetLongPackagePath(Asset->GetOutermost()->GetName());
+
+			//Need to save asset before renaming else may lose snapshot
+			// save existing play list asset
+			TArray<UPackage*> PackagesToSave;
+			PackagesToSave.Add(Asset->GetPackage());
+			FEditorFileUtils::PromptForCheckoutAndSave(PackagesToSave, false /*bCheckDirty*/, false /*bPromptToSave*/);
+
+			FAssetToolsModule& AssetToolsModule = FModuleManager::LoadModuleChecked<FAssetToolsModule>("AssetTools");
+
+			TArray<FAssetRenameData> AssetsAndNames;
+			AssetsAndNames.Emplace(FAssetRenameData(Asset, PackagePath, TrimmedText.ToString()));
+			AssetToolsModule.Get().RenameAssetsWithDialog(AssetsAndNames);
+
+		}
+			
+		// Remove ourselves from the window focus so we don't get automatically reselected when scrolling around the context menu.
+		TSharedPtr< SWindow > ParentWindow = FSlateApplication::Get().FindWidgetWindow(SharedThis(this));
+		if (ParentWindow.IsValid())
+		{
+			ParentWindow->SetWidgetToFocusOnActivate(NULL);
+		}
+	}
+
+	// Clear Error 
+	TextBox->SetError(FText::GetEmpty());
+}
+
+void SControlRigAssetEditableTextBox::OnTextChanged(const FText& InLabel)
+{
+	const FString PackagePath = FPackageName::GetLongPackagePath(Asset->GetOutermost()->GetName());
+	const FString PackageName = PackagePath / InLabel.ToString();
+	const FString ObjectPath = FString::Printf(TEXT("%s.%s"), *PackageName, *(InLabel.ToString()));
+
+	FText OutErrorMessage;
+	if (!AssetViewUtils::IsValidObjectPathForCreate(ObjectPath, OutErrorMessage))
+	{
+		TextBox->SetError(FText::GetEmpty());
+	}
+	else
+	{
+		TextBox->SetError(FText::GetEmpty());
+	}
+}
+
+void SControlRigPoseView::Construct(const FArguments& InArgs)
+{
+	PoseAsset = InArgs._PoseAsset;
+
+	bIsKey = false;
+	bIsMirror = false;
+	PoseBlendValue = 0.0f;
+	bIsBlending = false;
+	bSliderStartedTransaction = false;
+	InitialControlValues.SetNum(0);
+
+	TSharedRef<SWidget> Thumbnail = GetThumbnailWidget();
+	TSharedRef <SControlRigAssetEditableTextBox> ObjectNameBox = SNew(SControlRigAssetEditableTextBox).Asset(PoseAsset);
+
+	//Not used currently CreateControlList();
+
+	//for mirror settings
+	UControlRigPoseMirrorSettings* MirrorSettings = GetMutableDefault<UControlRigPoseMirrorSettings>();
+	FPropertyEditorModule& PropertyEditor = FModuleManager::LoadModuleChecked<FPropertyEditorModule>("PropertyEditor");
+
+	FDetailsViewArgs DetailsViewArgs;
+	DetailsViewArgs.bShowOptions = false;
+	DetailsViewArgs.bAllowSearch = false;
+	DetailsViewArgs.bShowPropertyMatrixButton = false;
+	DetailsViewArgs.bUpdatesFromSelection = false;
+	DetailsViewArgs.bLockable = false;
+	DetailsViewArgs.NameAreaSettings = FDetailsViewArgs::HideNameArea;
+	DetailsViewArgs.ViewIdentifier = "Create Control Asset";
+
+	MirrorDetailsView = PropertyEditor.CreateDetailView(DetailsViewArgs);
+	MirrorDetailsView->SetObject(MirrorSettings);
+			
+	ChildSlot
+	[
+
+		SNew(SVerticalBox)
+		+SVerticalBox::Slot()
+		
+			.FillHeight(1)
+			.Padding(0, 0, 0, 4)
+			[
+			SNew(SSplitter)
+
+				+ SSplitter::Slot()
+				.Value(0.33f)
+				[
+					SNew(SBorder)
+					.BorderImage(FEditorStyle::GetBrush("ToolPanel.GroupBorder"))
+					[
+						SNew(SVerticalBox)
+						+ SVerticalBox::Slot()
+						.AutoHeight()
+						.HAlign(HAlign_Center)
+						.Padding(5.f)
+						[
+							SNew(SBox)
+							.VAlign(VAlign_Center)
+						[
+							ObjectNameBox
+						]
+						]
+
+					+ SVerticalBox::Slot()
+						.AutoHeight()
+						.HAlign(HAlign_Center)
+						.Padding(5.f)
+						[
+							SNew(SBox)
+							.VAlign(VAlign_Center)
+						[
+							Thumbnail
+						]
+						]
+					+ SVerticalBox::Slot()
+						.AutoHeight()
+						.HAlign(HAlign_Center)
+						.Padding(5.f)
+
+						[
+							SNew(SButton)
+							.ContentPadding(FMargin(10, 5))
+						.Text(LOCTEXT("CaptureThmbnail", "Capture Thumbnail"))
+						.ToolTipText(LOCTEXT("CaptureThmbnailTooltip", "Captures a thumbnail from the active viewport"))
+						.OnClicked(this, &SControlRigPoseView::OnCaptureThumbnail)
+						]
+					]
+				]
+
+			+ SSplitter::Slot()
+				.Value(0.33f)
+				[
+					SNew(SBorder)
+					.BorderImage(FEditorStyle::GetBrush("ToolPanel.GroupBorder"))
+					[
+						SNew(SVerticalBox)
+						+ SVerticalBox::Slot()
+							.AutoHeight()
+							.HAlign(HAlign_Center)
+							.Padding(5.f)
+
+							[
+							SNew(SButton)
+								.ContentPadding(FMargin(10, 5))
+							.Text(LOCTEXT("PastePose", "Paste Pose"))
+							.OnClicked(this, &SControlRigPoseView::OnPastePose)
+							]
+						+ SVerticalBox::Slot()
+							.AutoHeight()
+							.HAlign(HAlign_Center)
+							.Padding(5.f)
+							[
+								SNew(SHorizontalBox)
+								+ SHorizontalBox::Slot()
+							.AutoWidth()
+							.HAlign(HAlign_Center)
+							.Padding(5.f)
+							[
+								SNew(SCheckBox)
+								.IsChecked(this, &SControlRigPoseView::IsKeyPoseChecked)
+							.OnCheckStateChanged(this, &SControlRigPoseView::OnKeyPoseChecked)
+							.Padding(5.0f)
+
+							[
+								SNew(STextBlock).Text(LOCTEXT("Key", "Key"))
+
+							]
+							]
+
+						+ SHorizontalBox::Slot()
+							.AutoWidth()
+
+							.HAlign(HAlign_Center)
+
+							.Padding(5.f)
+							[
+								SNew(SCheckBox)
+								.IsChecked(this, &SControlRigPoseView::IsMirrorPoseChecked)
+							.OnCheckStateChanged(this, &SControlRigPoseView::OnMirrorPoseChecked)
+							.IsEnabled(this, &SControlRigPoseView::IsMirrorEnabled)
+							.Padding(1.0f)
+							[
+								SNew(STextBlock).Text(LOCTEXT("Mirror", "Mirror"))
+								.IsEnabled(this, &SControlRigPoseView::IsMirrorEnabled)
+							]
+							]
+							]
+						+ SVerticalBox::Slot()
+							.AutoHeight()
+							.HAlign(HAlign_Center)
+							.Padding(5.f)
+							[
+
+								SNew(SNumericEntryBox<float>)
+								// Only allow spinning if we have a single value
+								.Value(this, &SControlRigPoseView::OnGetPoseBlendValue)
+								.AllowSpin(true)
+								.MinValue(-1.0f)
+								.MaxValue(2.0f)
+								.MinSliderValue(-1.0f)
+								.MaxSliderValue(2.0f)
+								.SliderExponent(1)
+								.Delta(0.005f)
+								.OnValueChanged(this, &SControlRigPoseView::OnPoseBlendChanged)
+								.OnValueCommitted(this, &SControlRigPoseView::OnPoseBlendCommited)
+								.OnBeginSliderMovement(this,&SControlRigPoseView::OnBeginSliderMovement)
+								.OnEndSliderMovement(this,&SControlRigPoseView::OnEndSliderMovement)
+
+							]
+					]
+				]
+			+ SSplitter::Slot()
+				.Value(0.33f)
+				[
+					MirrorDetailsView.ToSharedRef()
+				]
+			/*  todo may want to put this back, it let's you see the controls...
+			+ SSplitter::Slot()
+				.Value(0.33f)
+				[
+					SNew(SBorder)
+					.BorderImage(FEditorStyle::GetBrush("ToolPanel.GroupBorder"))
+					[
+						SNew(SVerticalBox)
+						+ SVerticalBox::Slot()
+							.AutoHeight()
+							.HAlign(HAlign_Center)
+							.Padding(2.f)
+							[
+
+							SNew(SButton)
+								.ContentPadding(FMargin(5, 5))
+								.Text(LOCTEXT("SelectControls", "Select Controls"))
+								.ToolTipText(LOCTEXT("SelectControlsTooltip", "Select Controls From This Asset"))
+								.OnClicked(this, &SControlRigPoseView::OnSelectControls)
+				
+							]
+						+ SVerticalBox::Slot()
+							.VAlign(VAlign_Fill)
+							.HAlign(HAlign_Center)
+							.Padding(5.f)
+
+							[
+							SNew(SListView< TSharedPtr<FString> >)
+								.ItemHeight(24)
+								.ListItemsSource(&ControlList)
+								.SelectionMode(ESelectionMode::None)
+								.OnGenerateRow(this, &SControlRigPoseView::OnGenerateWidgetForList)
+							]
+					]
+				]
+				*/
+			]
+			
+
+	];
+
+}
+
+ECheckBoxState SControlRigPoseView::IsKeyPoseChecked() const
+{
+	if (bIsKey)
+	{
+		return ECheckBoxState::Checked;
+	}
+	return ECheckBoxState::Unchecked;
+}
+
+void SControlRigPoseView::OnKeyPoseChecked(ECheckBoxState NewState)
+{
+	if (NewState == ECheckBoxState::Checked)
+	{
+		bIsKey = true;
+	}
+	else
+	{
+		bIsKey = false;
+	}
+}
+ECheckBoxState SControlRigPoseView::IsMirrorPoseChecked() const
+{
+	if (bIsMirror)
+	{
+		return ECheckBoxState::Checked;
+	}
+	return ECheckBoxState::Unchecked;
+}
+
+void SControlRigPoseView::OnMirrorPoseChecked(ECheckBoxState NewState)
+{
+	if (NewState == ECheckBoxState::Checked)
+	{
+		bIsMirror = true;
+	}
+	else
+	{
+		bIsMirror = false;
+	}
+}
+
+bool SControlRigPoseView::IsMirrorEnabled() const
+{
+	return true;
+}
+
+
+FReply SControlRigPoseView::OnPastePose()
+{
+	if (GetControlRig() && PoseAsset.IsValid())
+	{
+		PoseAsset->PastePose(GetControlRig(), bIsKey,bIsMirror);
+	}
+	return FReply::Handled();
+}
+
+FReply SControlRigPoseView::OnSelectControls()
+{
+	if (GetControlRig() && PoseAsset.IsValid())
+	{
+		PoseAsset->SelectControls(GetControlRig());
+	}
+	return FReply::Handled();
+}
+
+void SControlRigPoseView::OnPoseBlendChanged(float ChangedVal)
+{
+	UControlRig* ControlRig = GetControlRig();
+	if (ControlRig && PoseAsset.IsValid())
+	{
+
+		PoseBlendValue = ChangedVal;
+		if (!bIsBlending)
+		{
+			bIsBlending = true;
+			InitialControlValues = PoseAsset->GetCurrentPose(ControlRig);
+		}
+
+		PoseAsset->BlendWithInitialPoses(InitialControlValues, ControlRig, false, bIsMirror, PoseBlendValue);
+	}
+}
+void SControlRigPoseView::OnBeginSliderMovement()
+{
+	if (bSliderStartedTransaction == false)
+	{
+		bSliderStartedTransaction = true;
+		GEditor->BeginTransaction(LOCTEXT("PastePoseTransation", "Paste Pose"));
+	}
+}
+void SControlRigPoseView::OnEndSliderMovement(float NewValue)
+{
+	if (bSliderStartedTransaction)
+	{
+		GEditor->EndTransaction();
+		bSliderStartedTransaction = false;
+
+	}
+}
+
+void SControlRigPoseView::OnPoseBlendCommited(float ChangedVal, ETextCommit::Type Type)
+{
+	UControlRig* ControlRig = GetControlRig();
+	if (ControlRig && PoseAsset.IsValid())
+	{
+		FScopedTransaction ScopedTransaction(LOCTEXT("PastePoseTransaction", "Paste Pose"));
+		PoseBlendValue = ChangedVal;
+		PoseAsset->BlendWithInitialPoses(InitialControlValues, ControlRig, bIsKey, bIsMirror, PoseBlendValue);
+		bIsBlending = false;
+		PoseBlendValue = 0.0f;
+	}
+}
+
+FReply SControlRigPoseView::OnCaptureThumbnail()
+{
+	FControlRigView::CaptureThumbnail(PoseAsset.Get());
+	return FReply::Handled();
+}
+
+TSharedRef<SWidget> SControlRigPoseView::GetThumbnailWidget()
+{
+	TSharedPtr<SWidget> ThumbnailWidget = nullptr;
+
+	ThumbnailPool = MakeShared<FAssetThumbnailPool>(1, false);
+
+	const int32 ThumbnailSize = 128;
+
+	TSharedRef<FAssetThumbnail> AssetThumbnail = MakeShared<FAssetThumbnail>(PoseAsset.Get(), ThumbnailSize, ThumbnailSize, ThumbnailPool);
+	ThumbnailWidget = AssetThumbnail->MakeThumbnailWidget();
+
+
+	return SNew(SBox)
+		.WidthOverride(ThumbnailSize + 5)
+		.HeightOverride(ThumbnailSize + 5)
+		[
+			ThumbnailWidget.IsValid() ? ThumbnailWidget.ToSharedRef() : SNullWidget::NullWidget
+		];
+}
+
+UControlRig* SControlRigPoseView::GetControlRig()
+{
+	if (FControlRigEditMode* EditMode = static_cast<FControlRigEditMode*>(GLevelEditorModeTools().GetActiveMode(FControlRigEditMode::ModeName)))
+	{
+		return EditMode->GetControlRig(true);
+	}
+	return nullptr;
+}
+
+/* We may want to list the Controls in it (design said no but animators said yes)
+void SControlRigPoseView::CreateControlList()
+{
+	if (PoseAsset.IsValid())
+	{
+		const TArray<FName>& Controls = PoseAsset.Get()->GetControlNames();
+		for (const FName& ControlName : Controls)
+		{
+			ControlList.Add(MakeShared<FString>(ControlName.ToString()));
+		}
+	}
+}
+*/
+
+#undef LOCTEXT_NAMESPACE
diff --git a/Engine/Plugins/Experimental/ControlRig/Source/ControlRigEditor/Private/EditMode/SControlRigControlViews.h b/Engine/Plugins/Experimental/ControlRig/Source/ControlRigEditor/Private/EditMode/SControlRigControlViews.h
new file mode 100644
--- /dev/null
+++ b/Engine/Plugins/Experimental/ControlRig/Source/ControlRigEditor/Private/EditMode/SControlRigControlViews.h
@@ -0,0 +1,96 @@
+// Copyright Epic Games, Inc. All Rights Reserved.
+/**
+* Hold the Slate Views for the different Control Rig Asset Views
+* This is shown in the Bottom of The SControlRigBaseListWidget
+*/
+#pragma once
+
+#include "CoreMinimal.h"
+#include "Widgets/SWidget.h"
+#include "Widgets/DeclarativeSyntaxSupport.h"
+#include "Widgets/SCompoundWidget.h"
+#include "Widgets/Input/SCheckBox.h"
+#include "Widgets/Input/SNumericEntryBox.h"
+#include "Widgets/Views/SListView.h"
+#include "Widgets/Text/STextBlock.h"
+#include "Widgets/Views/STableRow.h"
+#include "Tools/ControlRigPose.h"
+class UControlRig;
+
+//Class to Hold Statics that are shared and externally callable
+class FControlRigView
+{
+public:
+	static void CaptureThumbnail(UObject* Asset);
+};
+
+class SControlRigPoseView : public SCompoundWidget
+{
+	SLATE_BEGIN_ARGS(SControlRigPoseView) {}
+	SLATE_ARGUMENT(UControlRigPoseAsset*, PoseAsset)
+
+	SLATE_END_ARGS()
+	~SControlRigPoseView()
+	{
+	}
+
+	void Construct(const FArguments& InArgs);
+
+private:
+
+	UControlRig* GetControlRig();
+
+	/*
+	* Delegates and Helpers
+	*/
+	ECheckBoxState IsKeyPoseChecked() const;
+	void OnKeyPoseChecked(ECheckBoxState NewState);
+	ECheckBoxState IsMirrorPoseChecked() const;
+	void OnMirrorPoseChecked(ECheckBoxState NewState);
+	bool IsMirrorEnabled() const;
+	void OnPoseBlendChanged(float ChangedVal);
+	void OnPoseBlendCommited(float ChangedVal, ETextCommit::Type Type);
+	void OnBeginSliderMovement();
+	void OnEndSliderMovement(float NewValue);
+	TOptional<float> OnGetPoseBlendValue()const {return PoseBlendValue;}
+	FReply OnPastePose();
+	FReply OnSelectControls();
+	FReply OnCaptureThumbnail();
+
+	bool bIsKey;
+	bool bIsMirror;
+	float PoseBlendValue;
+	bool bIsBlending;
+	bool bSliderStartedTransaction;
+	TArray<FRigControlCopy> InitialControlValues;
+
+	TWeakObjectPtr<UControlRigPoseAsset> PoseAsset;
+
+	/* Thumbnail*/
+	TSharedRef<SWidget> GetThumbnailWidget();
+	TSharedPtr<FAssetThumbnailPool> ThumbnailPool;
+	
+	/* Mirroring*/
+	TSharedPtr<IDetailsView> MirrorDetailsView;
+
+	/*for list of controls stored in asset 
+	void CreateControlList();
+	TArray<TSharedPtr<FString>> ControlList;
+	*/
+
+	TSharedRef<ITableRow> OnGenerateWidgetForList(TSharedPtr<FString> InItem, const TSharedRef<STableViewBase>& OwnerTable)
+	{
+
+		return SNew(STableRow< TSharedPtr<FString> >, OwnerTable)
+			.Content()
+			[
+				SNew(SBox)
+				.Padding(2.0f)
+			[
+				SNew(STextBlock)
+				.Text(FText::FromString((*InItem)))
+			]
+			];
+	}
+};
+
diff --git a/Engine/Plugins/Experimental/ControlRig/Source/ControlRigEditor/Private/EditMode/SControlRigEditModeTools.cpp b/Engine/Plugins/Experimental/ControlRig/Source/ControlRigEditor/Private/EditMode/SControlRigEditModeTools.cpp
--- a/Engine/Plugins/Experimental/ControlRig/Source/ControlRigEditor/Private/EditMode/SControlRigEditModeTools.cpp
+++ b/Engine/Plugins/Experimental/ControlRig/Source/ControlRigEditor/Private/EditMode/SControlRigEditModeTools.cpp
@@ -23,6 +23,7 @@
 #include "DetailCategoryBuilder.h"
 #include "DetailWidgetRow.h"
 #include "Rigs/FKControlRig.h"
+#include "SControlRigBaseListWidget.h"
 
 #define LOCTEXT_NAMESPACE "ControlRigRootCustomization"
 
@@ -482,4 +483,137 @@ void SControlRigEditModeTools::OnRigOptionFinishedChange(const FPropertyChangedE
 	}
 }
 
+void SControlRigEditModeTools::CustomizeToolBarPalette(FToolBarBuilder& ToolBarBuilder)
+{
+	//TOGGLE SELECTED RIG CONTROLS
+	ToolBarBuilder.AddToolBarButton(
+		FUIAction(
+			FExecuteAction::CreateLambda([this] {
+		FControlRigEditMode* ControlRigEditMode = static_cast<FControlRigEditMode*>(GLevelEditorModeTools().GetActiveMode(FControlRigEditMode::ModeName));
+		if (ControlRigEditMode)
+		{
+			ControlRigEditMode->SetOnlySelectRigControls(!ControlRigEditMode->GetOnlySelectRigControls());
+		}
+	}),
+			FCanExecuteAction(),
+		FIsActionChecked::CreateLambda([] {
+		FControlRigEditMode* ControlRigEditMode = static_cast<FControlRigEditMode*>(GLevelEditorModeTools().GetActiveMode(FControlRigEditMode::ModeName));
+		if (ControlRigEditMode)
+		{
+			return ControlRigEditMode->GetOnlySelectRigControls();
+		}
+		return false;
+	})
+		),
+		NAME_None,
+		LOCTEXT("OnlySelectControls", "Only Select Controls"),
+		LOCTEXT("OnlySelectControlsTooltip", "Only Select Control Rig Controls"),
+		FSlateIcon(FEditorStyle::GetStyleSetName(), "FoliageEditMode.SetPaint"), //todo get icon
+		EUserInterfaceActionType::ToggleButton
+		);
+	ToolBarBuilder.AddSeparator();
+
+	//POSES
+
+	ToolBarBuilder.AddToolBarButton(
+		FExecuteAction::CreateSP(this, &SControlRigEditModeTools::MakePoseDialog),
+		NAME_None,
+		LOCTEXT("Poses", "Poses"),
+		LOCTEXT("PosesTooltip", "Show Poses"),
+		FSlateIcon(FEditorStyle::GetStyleSetName(), "FoliageEditMode.Settings"),
+		EUserInterfaceActionType::Button
+	);
+	ToolBarBuilder.AddSeparator();
+}
+
+
+void SControlRigEditModeTools::MakePoseDialog()
+{
+
+	FControlRigEditMode* ControlRigEditMode = static_cast<FControlRigEditMode*>(GLevelEditorModeTools().GetActiveMode(FControlRigEditMode::ModeName));
+	if (ControlRigEditMode)
+	{
+
+		TSharedPtr<SWindow> ExistingWindow = PoseWindow.Pin();
+		if (ExistingWindow.IsValid())
+		{
+			ExistingWindow->BringToFront();
+		}
+		else
+		{
+			ExistingWindow = SNew(SWindow)
+				.Title(LOCTEXT("Poses", "Poses"))
+				.HasCloseButton(true)
+				.SupportsMaximize(false)
+				.SupportsMinimize(false)
+				.ClientSize(FVector2D(850, 750));
+			TSharedPtr<SWindow> RootWindow = FGlobalTabmanager::Get()->GetRootWindow();
+			if (RootWindow.IsValid())
+			{
+				FSlateApplication::Get().AddWindowAsNativeChild(ExistingWindow.ToSharedRef(), RootWindow.ToSharedRef());
+			}
+			else
+			{
+				FSlateApplication::Get().AddWindow(ExistingWindow.ToSharedRef());
+			}
+
+		}
+
+		ExistingWindow->SetContent(
+			SNew(SControlRigBaseListWidget)
+		);
+		PoseWindow = ExistingWindow;
+	}
+	
+}
+
+/* MZ TODO
+void SControlRigEditModeTools::MakeSelectionSetDialog()
+{
+
+	FControlRigEditMode* ControlRigEditMode = static_cast<FControlRigEditMode*>(GLevelEditorModeTools().GetActiveMode(FControlRigEditMode::ModeName));
+	if (ControlRigEditMode)
+	{
+		TSharedPtr<SWindow> ExistingWindow = SelectionSetWindow.Pin();
+		if (ExistingWindow.IsValid())
+		{
+			ExistingWindow->BringToFront();
+		}
+		else
+		{
+			ExistingWindow = SNew(SWindow)
+				.Title(LOCTEXT("SelectionSets", "Selection Set"))
+				.HasCloseButton(true)
+				.SupportsMaximize(false)
+				.SupportsMinimize(false)
+				.ClientSize(FVector2D(165, 200));
+			TSharedPtr<SWindow> RootWindow = FGlobalTabmanager::Get()->GetRootWindow();
+			if (RootWindow.IsValid())
+			{
+				FSlateApplication::Get().AddWindowAsNativeChild(ExistingWindow.ToSharedRef(), RootWindow.ToSharedRef());
+			}
+			else
+			{
+				FSlateApplication::Get().AddWindow(ExistingWindow.ToSharedRef());
+			}
+
+		}
+
+		ExistingWindow->SetContent(
+			SNew(SControlRigBaseListWidget)
+		);
+		SelectionSetWindow = ExistingWindow;
+	}
+}
+*/
+FText SControlRigEditModeTools::GetActiveToolName() const
+{
+	return  FText();
+}
+
+FText SControlRigEditModeTools::GetActiveToolMessage() const
+{
+	return  FText();
+}
+
 #undef LOCTEXT_NAMESPACE
diff --git a/Engine/Plugins/Experimental/ControlRig/Source/ControlRigEditor/Private/EditMode/SControlRigEditModeTools.h b/Engine/Plugins/Experimental/ControlRig/Source/ControlRigEditor/Private/EditMode/SControlRigEditModeTools.h
--- a/Engine/Plugins/Experimental/ControlRig/Source/ControlRigEditor/Private/EditMode/SControlRigEditModeTools.h
+++ b/Engine/Plugins/Experimental/ControlRig/Source/ControlRigEditor/Private/EditMode/SControlRigEditModeTools.h
@@ -72,4 +72,26 @@ private:
 	EVisibility GetRigOptionExpanderVisibility() const;
 
 	void OnRigOptionFinishedChange(const FPropertyChangedEvent& PropertyChangedEvent);
+
+private:
+	/** Toolbar functions and windows*/
+
+	/** Make the Pose Dialg*/
+	void MakePoseDialog();
+
+	/** Pose Dialog Window*/
+	TWeakPtr<SWindow> PoseWindow;
+
+	//TODO may put back void MakeSelectionSetDialog();
+	//TWeakPtr<SWindow> SelectionSetWindow;
+public:
+	/** Modes Panel Header Information **/
+	void CustomizeToolBarPalette(FToolBarBuilder& ToolBarBuilder);
+	FText GetActiveToolName() const;
+	FText GetActiveToolMessage() const;
+
+
 };
+
+
+
diff --git a/Engine/Plugins/Experimental/ControlRig/Source/ControlRigEditor/Private/EditMode/SControlRigPosePicker.h b/Engine/Plugins/Experimental/ControlRig/Source/ControlRigEditor/Private/EditMode/SControlRigPosePicker.h
new file mode 100644
--- /dev/null
+++ b/Engine/Plugins/Experimental/ControlRig/Source/ControlRigEditor/Private/EditMode/SControlRigPosePicker.h
@@ -0,0 +1,15 @@
+// Copyright Epic Games, Inc. All Rights Reserved.
+
+#pragma once
+
+#include "CoreMinimal.h"
+#include "Widgets/SWidget.h"
+#include "Widgets/DeclarativeSyntaxSupport.h"
+#include "Widgets/SCompoundWidget.h"
+#include "Toolkits/AssetEditorToolkit.h"
+#include "Widgets/Views/STableViewBase.h"
+#include "Widgets/Views/STableRow.h"
+#include "Widgets/Views/SListView.h"
+#include "EditorUndoClient.h"
+
+
diff --git a/Engine/Plugins/Experimental/ControlRig/Source/ControlRigEditor/Private/Tools/AssetTypeActions_ControlRigPose.cpp b/Engine/Plugins/Experimental/ControlRig/Source/ControlRigEditor/Private/Tools/AssetTypeActions_ControlRigPose.cpp
new file mode 100644
--- /dev/null
+++ b/Engine/Plugins/Experimental/ControlRig/Source/ControlRigEditor/Private/Tools/AssetTypeActions_ControlRigPose.cpp
@@ -0,0 +1,66 @@
+// Copyright Epic Games, Inc. All Rights Reserved.
+
+#include "AssetTypeActions_ControlRigPose.h"
+#include "EditorFramework/AssetImportData.h"
+#include "ThumbnailRendering/SceneThumbnailInfo.h"
+#include "Tools/ControlRigPose.h"
+#include "ToolMenus.h"
+
+FText FAssetTypeActions_ControlRigPose::GetName() const
+{
+	return NSLOCTEXT("AssetTypeActions", "AssetTypeActions_ControlRigPose", "Control Pose");
+}
+
+FColor FAssetTypeActions_ControlRigPose::GetTypeColor() const
+{
+	return FColor(222, 128, 64);
+}
+
+UClass* FAssetTypeActions_ControlRigPose::GetSupportedClass() const
+{
+	return UControlRigPoseAsset::StaticClass();
+}
+
+bool FAssetTypeActions_ControlRigPose::HasActions(const TArray<UObject*>& InObjects) const
+{
+	return true;
+}
+
+void FAssetTypeActions_ControlRigPose::OpenAssetEditor(const TArray<UObject*>& InObjects, TSharedPtr<class IToolkitHost> EditWithinLevelEditor /*= TSharedPtr<IToolkitHost>()*/)
+{
+	FAssetTypeActions_Base::OpenAssetEditor(InObjects, EditWithinLevelEditor);
+}
+
+uint32 FAssetTypeActions_ControlRigPose::GetCategories()
+{
+	return EAssetTypeCategories::Animation;
+}
+
+class UThumbnailInfo* FAssetTypeActions_ControlRigPose::GetThumbnailInfo(UObject* Asset) const
+{
+	return nullptr; //mz todo it looks like this is not needed if we use the default. Will remove based on feedback.
+	/*
+	UControlRigPoseAsset* PoseAsset = CastChecked<UControlRigPoseAsset>(Asset);
+	UThumbnailInfo* ThumbnailInfo = PoseAsset->ThumbnailInfo;
+	if (ThumbnailInfo == NULL)
+	{
+		ThumbnailInfo = NewObject<USceneThumbnailInfo>(PoseAsset, NAME_None, RF_Transactional);
+		PoseAsset->ThumbnailInfo = ThumbnailInfo;
+	}
+	
+	return ThumbnailInfo;
+	*/
+}
+
+bool FAssetTypeActions_ControlRigPose::IsImportedAsset() const
+{
+	return false;
+}
+
+void FAssetTypeActions_ControlRigPose::GetActions(const TArray<UObject*>& InObjects, FToolMenuSection& Section)
+{
+	//TODO Add custom Actions based on feedback
+
+//	FAssetTypeActions_Base::GetActions(InObjects, Section);
+//	TArray<TWeakObjectPtr<UControlRigPoseAsset>> ControlRigPoseAsset = GetTypedWeakObjectPtrs<UControlRigPoseAsset>(InObjects);
+}
\ No newline at end of file
diff --git a/Engine/Plugins/Experimental/ControlRig/Source/ControlRigEditor/Private/Tools/AssetTypeActions_ControlRigPose.h b/Engine/Plugins/Experimental/ControlRig/Source/ControlRigEditor/Private/Tools/AssetTypeActions_ControlRigPose.h
new file mode 100644
--- /dev/null
+++ b/Engine/Plugins/Experimental/ControlRig/Source/ControlRigEditor/Private/Tools/AssetTypeActions_ControlRigPose.h
@@ -0,0 +1,28 @@
+// Copyright Epic Games, Inc. All Rights Reserved.
+#pragma once
+
+/** 
+*  Type-actions class for ControlRigPose assets
+*/
+
+
+#include "CoreMinimal.h"
+#include "Toolkits/IToolkitHost.h"
+#include "AssetTypeActions_Base.h"
+
+
+class FAssetTypeActions_ControlRigPose : public FAssetTypeActions_Base
+{
+public:
+	// IAssetTypeActions Implementation
+	virtual FText GetName() const override;
+	virtual FColor GetTypeColor() const override;
+	virtual UClass* GetSupportedClass() const override;
+	virtual bool HasActions(const TArray<UObject*>& InObjects) const override;
+	virtual void GetActions(const TArray<UObject*>& InObjects, struct FToolMenuSection& Section) override;
+	virtual void OpenAssetEditor(const TArray<UObject*>& InObjects, TSharedPtr<class IToolkitHost> EditWithinLevelEditor = TSharedPtr<IToolkitHost>()) override;
+	virtual uint32 GetCategories() override;
+	virtual class UThumbnailInfo* GetThumbnailInfo(UObject* Asset) const override;
+	virtual bool IsImportedAsset() const override;
+	// End IAssetTypeActions
+};
\ No newline at end of file
diff --git a/Engine/Plugins/Experimental/ControlRig/Source/ControlRigEditor/Private/Tools/ControlRigPoseThumbnailRenderer.cpp b/Engine/Plugins/Experimental/ControlRig/Source/ControlRigEditor/Private/Tools/ControlRigPoseThumbnailRenderer.cpp
new file mode 100644
--- /dev/null
+++ b/Engine/Plugins/Experimental/ControlRig/Source/ControlRigEditor/Private/Tools/ControlRigPoseThumbnailRenderer.cpp
@@ -0,0 +1,236 @@
+// Copyright Epic Games, Inc. All Rights Reserved.
+
+#include "Tools/ControlRigPoseThumbnailRenderer.h"
+#include "Tools/ControlRigPose.h"
+#include "ThumbnailHelpers.h"
+#include "ThumbnailRendering/SceneThumbnailInfo.h"
+#include "ControlRigObjectBinding.h"
+#include "AnimCustomInstanceHelper.h"
+#include "Sequencer/ControlRigLayerInstance.h"
+
+
+/*
+***************************************************************
+FControlRigPoseThumbnailScene
+***************************************************************
+*/
+
+class FControlRigPoseThumbnailScene : public FThumbnailPreviewScene
+{
+public:
+	/** Constructor */
+	FControlRigPoseThumbnailScene();
+
+	bool SetControlRigPoseAsset(UControlRigPoseAsset* PoseAsset);
+
+protected:
+	// FThumbnailPreviewScene implementation
+	virtual void GetViewMatrixParameters(const float InFOVDegrees, FVector& OutOrigin, float& OutOrbitPitch, float& OutOrbitYaw, float& OutOrbitZoom) const override;
+
+	//Clean up the children of this component
+	void CleanupComponentChildren(USceneComponent* Component);
+
+private:
+	/** The skeletal mesh actor*/
+	ASkeletalMeshActor* PreviewActor;
+
+	/** ControlRig Pose Asset*/
+	UControlRigPoseAsset* PoseAsset;
+
+};
+
+
+FControlRigPoseThumbnailScene::FControlRigPoseThumbnailScene()
+	: FThumbnailPreviewScene()
+{
+	bForceAllUsedMipsResident = false;
+
+	// Create preview actor
+	FActorSpawnParameters SpawnInfo;
+	SpawnInfo.SpawnCollisionHandlingOverride = ESpawnActorCollisionHandlingMethod::AlwaysSpawn;
+	SpawnInfo.bNoFail = true;
+	SpawnInfo.ObjectFlags = RF_Transient;
+	PreviewActor = GetWorld()->SpawnActor<ASkeletalMeshActor>(SpawnInfo);
+
+	PreviewActor->SetActorEnableCollision(false);
+}
+
+bool FControlRigPoseThumbnailScene::SetControlRigPoseAsset(UControlRigPoseAsset* InPoseAsset)
+{
+	PoseAsset = InPoseAsset;
+	bool bSetSucessfully = false;
+
+	if (PoseAsset)
+	{
+		USkeletalMesh* SkeletalMesh = nullptr; // PoseAsset->GetSkeletalMesh();
+		UControlRig* ControlRig = nullptr; // PoseAsset->GetControlRig();
+		PreviewActor->GetSkeletalMeshComponent()->OverrideMaterials.Empty();
+
+
+		if (SkeletalMesh && ControlRig)
+		{
+
+			PreviewActor->GetSkeletalMeshComponent()->SetSkeletalMesh(SkeletalMesh);
+			PreviewActor->GetSkeletalMeshComponent()->AnimScriptInstance = nullptr;
+
+
+			bSetSucessfully = true;
+			//now set up the control rig and the anim instance fo rit
+			UControlRig* TempControlRig = NewObject<UControlRig>(GetTransientPackage(), ControlRig->GetClass(),NAME_None, RF_Transient);
+			TempControlRig->SetObjectBinding(MakeShared<FControlRigObjectBinding>());
+			TempControlRig->GetObjectBinding()->BindToObject(SkeletalMesh);
+			TempControlRig->GetDataSourceRegistry()->RegisterDataSource(UControlRig::OwnerComponent, TempControlRig->GetObjectBinding()->GetBoundObject());
+
+
+			bool bWasCreated;
+			if (UControlRigLayerInstance* AnimInstance = FAnimCustomInstanceHelper::BindToSkeletalMeshComponent<UControlRigLayerInstance>(PreviewActor->GetSkeletalMeshComponent(), bWasCreated))
+			{
+
+				if (bWasCreated)
+				{
+					AnimInstance->RecalcRequiredBones();
+					AnimInstance->AddControlRigTrack(TempControlRig->GetUniqueID(), TempControlRig);
+					TempControlRig->CreateRigControlsForCurveContainer();
+				}
+				TempControlRig->Initialize();
+
+				float Weight = 1.0f;
+				FControlRigIOSettings InputSettings;
+				InputSettings.bUpdateCurves = false;
+				InputSettings.bUpdatePose = true;
+				AnimInstance->UpdateControlRigTrack(TempControlRig->GetUniqueID(), Weight, InputSettings, true);
+				PoseAsset->PastePose(TempControlRig);
+				TempControlRig->Evaluate_AnyThread();
+
+			}
+
+			PreviewActor->GetSkeletalMeshComponent()->TickAnimation(0.f, false);
+			PreviewActor->GetSkeletalMeshComponent()->RefreshBoneTransforms();
+			PreviewActor->GetSkeletalMeshComponent()->FinalizeBoneTransform();
+
+			FTransform MeshTransform = FTransform::Identity;
+
+			PreviewActor->SetActorLocation(FVector(0, 0, 0), false);
+			PreviewActor->GetSkeletalMeshComponent()->UpdateBounds();
+
+			// Center the mesh at the world origin then offset to put it on top of the plane
+			const float BoundsZOffset = GetBoundsZOffset(PreviewActor->GetSkeletalMeshComponent()->Bounds);
+			PreviewActor->SetActorLocation(-PreviewActor->GetSkeletalMeshComponent()->Bounds.Origin + FVector(0, 0, BoundsZOffset), false);
+			PreviewActor->GetSkeletalMeshComponent()->RecreateRenderState_Concurrent();
+			TempControlRig->MarkPendingKill();
+
+		}
+
+		if (!bSetSucessfully)
+		{
+			CleanupComponentChildren(PreviewActor->GetSkeletalMeshComponent());
+			PreviewActor->GetSkeletalMeshComponent()->SetSkeletalMesh(nullptr);
+			PreviewActor->GetSkeletalMeshComponent()->SetAnimInstanceClass(nullptr);
+			PreviewActor->GetSkeletalMeshComponent()->AnimScriptInstance = nullptr;
+		}
+
+	}
+	return bSetSucessfully;
+}
+
+void FControlRigPoseThumbnailScene::CleanupComponentChildren(USceneComponent* Component)
+{
+	if (Component)
+	{
+		for (int32 ComponentIdx = Component->GetAttachChildren().Num() - 1; ComponentIdx >= 0; --ComponentIdx)
+		{
+			CleanupComponentChildren(Component->GetAttachChildren()[ComponentIdx]);
+			Component->GetAttachChildren()[ComponentIdx]->DestroyComponent();
+		}
+		check(Component->GetAttachChildren().Num() == 0);
+	}
+}
+
+void FControlRigPoseThumbnailScene::GetViewMatrixParameters(const float InFOVDegrees, FVector& OutOrigin, float& OutOrbitPitch, float& OutOrbitYaw, float& OutOrbitZoom) const
+{
+	check(PreviewActor->GetSkeletalMeshComponent());
+	check(PreviewActor->GetSkeletalMeshComponent()->SkeletalMesh);
+
+	const float HalfFOVRadians = FMath::DegreesToRadians<float>(InFOVDegrees) * 0.5f;
+	// No need to add extra size to view slightly outside of the sphere to compensate for perspective since skeletal meshes already buffer bounds.
+	const float HalfMeshSize = PreviewActor->GetSkeletalMeshComponent()->Bounds.SphereRadius;
+	const float BoundsZOffset = GetBoundsZOffset(PreviewActor->GetSkeletalMeshComponent()->Bounds);
+	const float TargetDistance = HalfMeshSize / FMath::Tan(HalfFOVRadians);
+
+	USceneThumbnailInfo* ThumbnailInfo = nullptr;// Cast<USceneThumbnailInfo>(PoseAsset->ThumbnailInfo);
+	if (ThumbnailInfo)
+	{
+		if (TargetDistance + ThumbnailInfo->OrbitZoom < 0)
+		{
+			ThumbnailInfo->OrbitZoom = -TargetDistance;
+		}
+	}
+	else
+	{
+		ThumbnailInfo = USceneThumbnailInfo::StaticClass()->GetDefaultObject<USceneThumbnailInfo>();
+	}
+
+	OutOrigin = FVector(0, 0, -BoundsZOffset);
+	OutOrbitPitch = ThumbnailInfo->OrbitPitch;
+	OutOrbitYaw = ThumbnailInfo->OrbitYaw;
+	OutOrbitZoom = TargetDistance + ThumbnailInfo->OrbitZoom;
+}
+
+
+
+/*
+***************************************************************
+UControlRigPoseThumbnailRenderer
+***************************************************************
+*/
+
+UControlRigPoseThumbnailRenderer::UControlRigPoseThumbnailRenderer(const FObjectInitializer& ObjectInitializer)
+	: Super(ObjectInitializer)
+{
+	ControlRigPoseAsset = nullptr;
+	ThumbnailScene = nullptr;
+}
+
+bool UControlRigPoseThumbnailRenderer::CanVisualizeAsset(UObject* Object)
+{
+	return Cast<UControlRigPoseAsset>(Object) != nullptr;
+}
+
+void UControlRigPoseThumbnailRenderer::Draw(UObject* Object, int32 X, int32 Y, uint32 Width, uint32 Height, FRenderTarget* RenderTarget, FCanvas* Canvas, bool bAdditionalViewFamily)
+{
+	UControlRigPoseAsset* PoseAsset = Cast<UControlRigPoseAsset>(Object);
+	if (PoseAsset)
+	{
+		
+		if (ThumbnailScene == nullptr)
+		{
+			ThumbnailScene = new FControlRigPoseThumbnailScene();
+		}
+
+		ThumbnailScene->SetControlRigPoseAsset(PoseAsset);
+
+		FSceneViewFamilyContext ViewFamily(FSceneViewFamily::ConstructionValues(RenderTarget, ThumbnailScene->GetScene(), FEngineShowFlags(ESFIM_Game))
+			.SetWorldTimes(FApp::GetCurrentTime() - GStartTime, FApp::GetDeltaTime(), FApp::GetCurrentTime() - GStartTime)
+			.SetAdditionalViewFamily(bAdditionalViewFamily));
+
+		ViewFamily.EngineShowFlags.DisableAdvancedFeatures();
+		ViewFamily.EngineShowFlags.MotionBlur = 0;
+		ViewFamily.EngineShowFlags.LOD = 0;
+
+		ThumbnailScene->GetView(&ViewFamily, X, Y, Width, Height);
+		RenderViewFamily(Canvas, &ViewFamily);
+		ThumbnailScene->SetControlRigPoseAsset(nullptr);
+	}
+}
+
+void UControlRigPoseThumbnailRenderer::BeginDestroy()
+{
+	if (ThumbnailScene != nullptr)
+	{
+		delete ThumbnailScene;
+		ThumbnailScene = nullptr;
+	}
+
+	Super::BeginDestroy();
+}
+
diff --git a/Engine/Plugins/Experimental/ControlRig/Source/ControlRigEditor/Private/Tools/CreateControlAssetRigSettings.cpp b/Engine/Plugins/Experimental/ControlRig/Source/ControlRigEditor/Private/Tools/CreateControlAssetRigSettings.cpp
new file mode 100644
--- /dev/null
+++ b/Engine/Plugins/Experimental/ControlRig/Source/ControlRigEditor/Private/Tools/CreateControlAssetRigSettings.cpp
@@ -0,0 +1,12 @@
+// Copyright Epic Games, Inc. All Rights Reserved.
+
+#include "Tools/CreateControlAssetRigSettings.h"
+#include "CoreMinimal.h"
+#include "UObject/ObjectMacros.h"
+
+
+UCreateControlPoseAssetRigSettings::UCreateControlPoseAssetRigSettings(const FObjectInitializer& Initializer)
+	: Super(Initializer)
+{
+	AssetName = FString("ControlRigPose");
+}
diff --git a/Engine/Plugins/Experimental/ControlRig/Source/ControlRigEditor/Public/Tools/ControlRigPoseThumbnailRenderer.h b/Engine/Plugins/Experimental/ControlRig/Source/ControlRigEditor/Public/Tools/ControlRigPoseThumbnailRenderer.h
new file mode 100644
--- /dev/null
+++ b/Engine/Plugins/Experimental/ControlRig/Source/ControlRigEditor/Public/Tools/ControlRigPoseThumbnailRenderer.h
@@ -0,0 +1,37 @@
+// Copyright Epic Games, Inc. All Rights Reserved.
+
+/**
+ * This thumbnail renderer displays a given Control Rig Pose Asset
+ */
+
+#pragma once
+
+#include "CoreMinimal.h"
+#include "UObject/ObjectMacros.h"
+#include "ThumbnailRendering/SkeletalMeshThumbnailRenderer.h"
+#include "Tools/ControlRigPose.h"
+#include "ControlRigPoseThumbnailRenderer.generated.h"
+
+class FCanvas;
+class FRenderTarget;
+class FControlRigPoseThumbnailScene;
+
+UCLASS(config = Editor, MinimalAPI)
+class UControlRigPoseThumbnailRenderer : public UDefaultSizedThumbnailRenderer
+{
+	GENERATED_UCLASS_BODY()
+
+	// Begin UThumbnailRenderer Object
+	virtual bool CanVisualizeAsset(UObject* Object) override;
+	virtual void Draw(UObject* Object, int32 X, int32 Y, uint32 Width, uint32 Height, FRenderTarget*, FCanvas* Canvas, bool bAdditionalViewFamily) override;
+	// End UThumbnailRenderer Object
+
+	// UObject Implementation
+	virtual void BeginDestroy() override;
+	// End UObject Implementation
+
+private:
+	FControlRigPoseThumbnailScene* ThumbnailScene;
+protected:
+	UControlRigPoseAsset* ControlRigPoseAsset;
+};
diff --git a/Engine/Plugins/Experimental/ControlRig/Source/ControlRigEditor/Public/Tools/CreateControlAssetRigSettings.h b/Engine/Plugins/Experimental/ControlRig/Source/ControlRigEditor/Public/Tools/CreateControlAssetRigSettings.h
new file mode 100644
--- /dev/null
+++ b/Engine/Plugins/Experimental/ControlRig/Source/ControlRigEditor/Public/Tools/CreateControlAssetRigSettings.h
@@ -0,0 +1,26 @@
+// Copyright Epic Games, Inc. All Rights Reserved.
+
+
+/**
+ * Default Settings for the Control Rig Asset, including the Default Asset Name
+ */
+
+#pragma once
+
+#include "CoreMinimal.h"
+#include "UObject/ObjectMacros.h"
+#include "CreateControlAssetRigSettings.generated.h"
+
+UCLASS(BlueprintType, config = EditorSettings)
+class  UCreateControlPoseAssetRigSettings : public UObject
+{
+public:
+	UCreateControlPoseAssetRigSettings(const FObjectInitializer& Initializer);
+
+	GENERATED_BODY()
+
+	/** Asset Name*/
+	UPROPERTY(EditAnywhere, Category = "Control Pose")
+	FString AssetName;
+
+};
