CIS: Fix for Deprecrated FDetailViewArgs

diff --git a/Engine/Plugins/Experimental/ControlRig/Source/ControlRigEditor/Private/EditMode/ControlRigEditModeToolkit.cpp b/Engine/Plugins/Experimental/ControlRig/Source/ControlRigEditor/Private/EditMode/ControlRigEditModeToolkit.cpp
--- a/Engine/Plugins/Experimental/ControlRig/Source/ControlRigEditor/Private/EditMode/ControlRigEditModeToolkit.cpp
+++ b/Engine/Plugins/Experimental/ControlRig/Source/ControlRigEditor/Private/EditMode/ControlRigEditModeToolkit.cpp
@@ -24,15 +24,14 @@ void FControlRigEditModeToolkit::Init(const TSharedPtr<IToolkitHost>& InitToolki
 {
 	FPropertyEditorModule& PropertyEditorModule = FModuleManager::GetModuleChecked<FPropertyEditorModule>("PropertyEditor");
 
-	FDetailsViewArgs DetailsViewArgs(
-		/*bUpdateFromSelection=*/ false,
-		/*bLockable=*/ false,
-		/*bAllowSearch=*/ false,
-		FDetailsViewArgs::HideNameArea,
-		/*bHideSelectionTip=*/ true,
-		/*InNotifyHook=*/ nullptr,
-		/*InSearchInitialKeyFocus=*/ false,
-		/*InViewIdentifier=*/ NAME_None);
+	FDetailsViewArgs DetailsViewArgs;
+
+	DetailsViewArgs.bUpdatesFromSelection = false;
+	DetailsViewArgs.bLockable = false;
+	DetailsViewArgs.bAllowSearch = false;
+	DetailsViewArgs.NameAreaSettings = FDetailsViewArgs::HideNameArea;
+	DetailsViewArgs.bHideSelectionTip = true;
+	DetailsViewArgs.bSearchInitialKeyFocus = false;
 	DetailsViewArgs.DefaultsOnlyVisibility = EEditDefaultsOnlyNodeVisibility::Automatic;
 	DetailsViewArgs.bShowOptions = false;
 	DetailsViewArgs.bAllowMultipleTopLevelObjects = true;
