CIS Fixes for same stamement and copyright boilerplate
#trivial

diff --git a/Engine/Plugins/Experimental/ControlRig/Source/ControlRigEditor/Private/EditMode/SControlRigControlViews.cpp b/Engine/Plugins/Experimental/ControlRig/Source/ControlRigEditor/Private/EditMode/SControlRigControlViews.cpp
--- a/Engine/Plugins/Experimental/ControlRig/Source/ControlRigEditor/Private/EditMode/SControlRigControlViews.cpp
+++ b/Engine/Plugins/Experimental/ControlRig/Source/ControlRigEditor/Private/EditMode/SControlRigControlViews.cpp
@@ -169,7 +169,7 @@ void SControlRigAssetEditableTextBox::OnTextChanged(const FText& InLabel)
 	FText OutErrorMessage;
 	if (!AssetViewUtils::IsValidObjectPathForCreate(ObjectPath, OutErrorMessage))
 	{
-		TextBox->SetError(FText::GetEmpty());
+		TextBox->SetError(OutErrorMessage);
 	}
 	else
 	{
diff --git a/Engine/Plugins/Experimental/ControlRig/Source/ControlRigEditor/Private/EditMode/SControlRigPosePicker.cpp b/Engine/Plugins/Experimental/ControlRig/Source/ControlRigEditor/Private/EditMode/SControlRigPosePicker.cpp
deleted file mode 100644
diff --git a/Engine/Plugins/Experimental/ControlRig/Source/ControlRigEditor/Private/EditMode/SControlRigPosePicker.h b/Engine/Plugins/Experimental/ControlRig/Source/ControlRigEditor/Private/EditMode/SControlRigPosePicker.h
deleted file mode 100644
