Control Rig: Pose Library:  Allow for multi-selection but can onluy perfrom asset actions though since other ones require just a single pose.
#trivial

diff --git a/Engine/Plugins/Experimental/ControlRig/Source/ControlRigEditor/Private/EditMode/SControlRigBaseListWidget.cpp b/Engine/Plugins/Experimental/ControlRig/Source/ControlRigEditor/Private/EditMode/SControlRigBaseListWidget.cpp
--- a/Engine/Plugins/Experimental/ControlRig/Source/ControlRigEditor/Private/EditMode/SControlRigBaseListWidget.cpp
+++ b/Engine/Plugins/Experimental/ControlRig/Source/ControlRigEditor/Private/EditMode/SControlRigBaseListWidget.cpp
@@ -518,7 +518,7 @@ void SControlRigBaseListWidget::Construct(const FArguments& InArgs)
 	AssetPickerConfig.OnGetAssetContextMenu = FOnGetAssetContextMenu::CreateSP(this, &SControlRigBaseListWidget::OnGetAssetContextMenu);	
 	AssetPickerConfig.SetFilterDelegates.Add(&SetFilterDelegate);
 	AssetPickerConfig.GetCurrentSelectionDelegates.Add(&GetCurrentSelectionDelegate);
-	AssetPickerConfig.SelectionMode = ESelectionMode::Single;
+	AssetPickerConfig.SelectionMode = ESelectionMode::Multi;
 	AssetPickerConfig.bAllowDragging = false;
 	AssetPickerConfig.AssetShowWarningText = LOCTEXT("NoPoses_Warning", "No Poses Found, Create One Using Button In Upper Left Corner");
 	AssetPickerConfig.OnIsAssetValidForCustomToolTip = FOnIsAssetValidForCustomToolTip::CreateLambda([](const FAssetData& AssetData) {return AssetData.IsAssetLoaded(); });
@@ -820,7 +820,10 @@ TSharedPtr<SWidget> SControlRigBaseListWidget::OnGetFolderContextMenu(const TArr
 TSharedPtr<SWidget> SControlRigBaseListWidget::OnGetAssetContextMenu(const TArray<FAssetData>& SelectedAssets)
 {
 	FMenuBuilder MenuBuilder(true /*bInShouldCloseWindowAfterMenuSelection*/, Commands);
-
+	if (SelectedAssets.Num() == 0)
+	{
+		return nullptr;
+	}
 	if (SelectedAssets.Num() > 0)
 	{
 		MenuBuilder.BeginSection("PoseDialogOptions", LOCTEXT("Asset", "Asset"));
@@ -882,10 +885,10 @@ TSharedPtr<SWidget> SControlRigBaseListWidget::OnGetAssetContextMenu(const TArra
 			}
 			MenuBuilder.EndSection();
 
-			return MenuBuilder.MakeWidget();
 		}
 	}
-	return nullptr;
+	return MenuBuilder.MakeWidget();
+
 }
 
 
