CIS FIX
#trivial

diff --git a/Engine/Plugins/Experimental/ControlRig/Source/ControlRig/Public/Tools/ControlRigToolAsset.h b/Engine/Plugins/Experimental/ControlRig/Source/ControlRig/Public/Tools/ControlRigToolAsset.h
--- a/Engine/Plugins/Experimental/ControlRig/Source/ControlRig/Public/Tools/ControlRigToolAsset.h
+++ b/Engine/Plugins/Experimental/ControlRig/Source/ControlRig/Public/Tools/ControlRigToolAsset.h
@@ -10,6 +10,7 @@
 #include "Widgets/Notifications/SNotificationList.h"
 #include "Framework/Notifications/NotificationManager.h"
 #include "ControlRigObjectBinding.h"
+#include "Components/SkeletalMeshComponent.h"
 
 /**
 Create One of the Control Rig Assets, Currently Only support Pose Assets but may add Animation/Selection Sets.
