Make Pose in the Control Rig Pose asset Blueprintable Read and Write and fix control order when pasting.

diff --git a/Engine/Plugins/Experimental/ControlRig/Source/ControlRig/Private/ControlRig.cpp b/Engine/Plugins/Experimental/ControlRig/Source/ControlRig/Private/ControlRig.cpp
--- a/Engine/Plugins/Experimental/ControlRig/Source/ControlRig/Private/ControlRig.cpp
+++ b/Engine/Plugins/Experimental/ControlRig/Source/ControlRig/Private/ControlRig.cpp
@@ -1122,7 +1122,7 @@ FTransform UControlRig::GetControlGlobalTransform(const FName& InControlName) co
 	return Hierarchy.ControlHierarchy.GetGlobalTransform(InControlName);
 }
 
-bool UControlRig::SetControlGlobalTransform(const FName& InControlName, const FTransform& InGlobalTransform, const FRigControlModifiedContext& Context)
+bool UControlRig::SetControlGlobalTransform(const FName& InControlName, const FTransform& InGlobalTransform, bool bNotify, const FRigControlModifiedContext& Context)
 {
 	FTransform GlobalTransform = InGlobalTransform;
 	if (IsSetupModeEnabled())
@@ -1140,7 +1140,7 @@ bool UControlRig::SetControlGlobalTransform(const FName& InControlName, const FT
 		}
 	}
 
-	SetControlValue(InControlName, Value, true /* notify */, Context);
+	SetControlValue(InControlName, Value, bNotify, Context);
 	return true;
 }
 
diff --git a/Engine/Plugins/Experimental/ControlRig/Source/ControlRig/Private/Tools/ControlRigPose.cpp b/Engine/Plugins/Experimental/ControlRig/Source/ControlRig/Private/Tools/ControlRigPose.cpp
--- a/Engine/Plugins/Experimental/ControlRig/Source/ControlRig/Private/Tools/ControlRigPose.cpp
+++ b/Engine/Plugins/Experimental/ControlRig/Source/ControlRig/Private/Tools/ControlRigPose.cpp
@@ -14,10 +14,15 @@ UControlRigPoseAsset::UControlRigPoseAsset(const FObjectInitializer& ObjectIniti
 {
 }
 
+void UControlRigPoseAsset::PostLoad()
+{
+	Super::PostLoad();
+	Pose.SetUpControlMap();
+}
+
 void UControlRigPoseAsset::SavePose(UControlRig* InControlRig, bool bUseAll)
 {
 	Pose.SavePose(InControlRig,bUseAll);
-	Controls = Pose.GetControlNames();
 }
 
 void UControlRigPoseAsset::PastePose(UControlRig* InControlRig, bool bDoKey, bool bDoMirror)
@@ -36,26 +41,31 @@ void UControlRigPoseAsset::SelectControls(UControlRig* InControlRig)
 	InControlRig->Modify();
 #endif
 	InControlRig->ClearControlSelection();
-	Controls = Pose.GetControlNames();
+	TArray<FName> Controls = Pose.GetControlNames();
 	for (const FName& Name : Controls)
 	{
 		InControlRig->SelectControl(Name, true);
 	}
 }
 
+void UControlRigPoseAsset::GetCurrentPose(UControlRig* InControlRig, FControlRigControlPose& OutPose)
+{
+	OutPose.SavePose(InControlRig, true);
+}
+
+
 TArray<FRigControlCopy> UControlRigPoseAsset::GetCurrentPose(UControlRig* InControlRig) 
 {
 	FControlRigControlPose TempPose;
-	FString Name;
 	TempPose.SavePose(InControlRig,true);
 	return TempPose.GetPoses();
 }
 
-void UControlRigPoseAsset::BlendWithInitialPoses(TArray<FRigControlCopy>& InitialControls, UControlRig* InControlRig, bool bDoKey, bool bDoMirror, float BlendValue)
+void UControlRigPoseAsset::BlendWithInitialPoses(FControlRigControlPose& InitialPose, UControlRig* InControlRig, bool bDoKey, bool bDoMirror, float BlendValue)
 {
 	if (BlendValue > 0.0f)
 	{
-		Pose.BlendWithInitialPoses(InitialControls, InControlRig, bDoKey, bDoMirror, BlendValue);
+		Pose.BlendWithInitialPoses(InitialPose, InControlRig, bDoKey, bDoMirror, BlendValue);
 	}
 }
 
diff --git a/Engine/Plugins/Experimental/ControlRig/Source/ControlRig/Private/Tools/ControlRigPoseMirrorTable.cpp b/Engine/Plugins/Experimental/ControlRig/Source/ControlRig/Private/Tools/ControlRigPoseMirrorTable.cpp
--- a/Engine/Plugins/Experimental/ControlRig/Source/ControlRig/Private/Tools/ControlRigPoseMirrorTable.cpp
+++ b/Engine/Plugins/Experimental/ControlRig/Source/ControlRig/Private/Tools/ControlRigPoseMirrorTable.cpp
@@ -20,51 +20,50 @@ void FControlRigPoseMirrorTable::SetUpMirrorTable(const UControlRig* ControlRig)
 				FString NewString = CurrentString.Replace(*Settings->RightSide, *Settings->LeftSide);
 				FName CurrentName(*CurrentString);
 				FName NewName(*NewString);
-				MatchedControls.Add(CurrentName, NewName);
+				MatchedControls.Add(NewName, CurrentName);
 			}
 			else if (CurrentString.Contains(Settings->LeftSide))
 			{
 				FString NewString = CurrentString.Replace(*Settings->LeftSide, *Settings->RightSide);
 				FName CurrentName(*CurrentString);
 				FName NewName(*NewString);
-				MatchedControls.Add(CurrentName, NewName);
+				MatchedControls.Add(NewName, CurrentName);
 			}
 		}
 	}
 }
 
-FRigControl* FControlRigPoseMirrorTable::GetControl(UControlRig* ControlRig, const FName& Name) const
+
+FRigControlCopy* FControlRigPoseMirrorTable::GetControl(FControlRigControlPose& Pose, FName Name)
 {
 
+	TArray<FRigControlCopy> CopyOfControls;
+
+	//Cache of the Map, Used to make pasting faster.
+	TMap<FName, int32>  CopyOfControlsNameToIndex;
+
 	if (MatchedControls.Num() <= 0)
 	{
-		if (ControlRig->IsControlSelected(Name))
+		int32* Index = Pose.CopyOfControlsNameToIndex.Find(Name);
+		if (Index != nullptr && (*Index) >= 0 && (*Index) < Pose.CopyOfControls.Num())
 		{
-			return ControlRig->FindControl(Name);
+			return &(Pose.CopyOfControls[*Index]);
 		}
+	
 	}
 	if (MatchedControls.Num() > 0) 
 	{
 		if (const FName* MatchedName = MatchedControls.Find(Name))
 		{
-			if (ControlRig->IsControlSelected(*MatchedName))
-			{
-				if (FRigControl* RigControl = ControlRig->FindControl(*MatchedName))
-				{
-					return RigControl;
-				}
-			}
-			else
+
+			int32* Index = Pose.CopyOfControlsNameToIndex.Find(*MatchedName);
+			if (Index != nullptr && (*Index) >= 0 && (*Index) < Pose.CopyOfControls.Num())
 			{
-				return nullptr;
+				return &(Pose.CopyOfControls[*Index]);
 			}
 		}
 	}
-	
-	if (ControlRig->IsControlSelected(Name))
-	{
-		return ControlRig->FindControl(Name);
-	}
+
 	return nullptr;
 }
 
@@ -82,37 +81,46 @@ bool FControlRigPoseMirrorTable::IsMatched(const FName& Name) const
 
 // Mirrors Translation as Global(Component)
 // Mirrors Rotation as Local , only if NOT matched, if matched we just use it
-void FControlRigPoseMirrorTable::GetMirrorTransform(const FRigControlCopy& ControlCopy, bool bIsMatched, FVector& OutGlobalTranslation, FQuat& OutLocalRotation) const
+void FControlRigPoseMirrorTable::GetMirrorTransform(const FRigControlCopy& ControlCopy, bool bIsMatched, FVector& OutGlobalTranslation, FQuat& OutGlobalRotation) const
 {
 	const UControlRigPoseMirrorSettings* Settings = GetDefault<UControlRigPoseMirrorSettings>();
 	if (Settings)
 	{
+		FTransform GlobalTransform = ControlCopy.GlobalTransform;
+		FTransform LocalTransform = ControlCopy.LocalTransform;
+		//FQuat Rotation = LocalTransform.GetRotation();
+		FQuat Rotation = GlobalTransform.GetRotation();
+
 		FVector Axis(Settings->XAxis, Settings->YAxis, Settings->ZAxis);
+		//if axis nearly zero just return
+		if (Axis.IsNearlyZero())
+		{
+			OutGlobalTranslation = ControlCopy.GlobalTransform.GetTranslation();
+			OutGlobalRotation = Rotation;
+			return;
+		}
 		Axis.Normalize();
-		FTransform GlobalTransform = ControlCopy.GlobalTransform;
 
 		FVector Translation = ControlCopy.GlobalTransform.GetTranslation();
 		FPlane Plane(FVector::ZeroVector,Axis);
 		OutGlobalTranslation = Translation.MirrorByPlane(Plane);
 
-		FTransform LocalTransform = ControlCopy.LocalTransform;
-		FQuat Rotation = LocalTransform.GetRotation();
 		//test FRotator OldRotator = Rotation.Rotator();
 		
 		if (bIsMatched)
 		{
-			OutLocalRotation = Rotation;
+			OutGlobalRotation = Rotation;
 		}
 		else
 		{
 			FQuat MirrorNormalQuat = FQuat(Axis.X, Axis.Y, Axis.Z, 0);
 			Rotation.EnforceShortestArcWith(MirrorNormalQuat);
-			OutLocalRotation = MirrorNormalQuat * Rotation * MirrorNormalQuat;
+			OutGlobalRotation = MirrorNormalQuat * Rotation * MirrorNormalQuat;
 		}
-		//test FRotator NewRotator = OutLocalRotation.Rotator();
+		//test FRotator NewRotator = OutGlobalRotation.Rotator();
 		
 		return;
 	}
 	OutGlobalTranslation = FVector::ZeroVector;
-	OutLocalRotation = FQuat::Identity;
+	OutGlobalRotation = FQuat::Identity;
 }
diff --git a/Engine/Plugins/Experimental/ControlRig/Source/ControlRig/Public/ControlRig.h b/Engine/Plugins/Experimental/ControlRig/Source/ControlRig/Public/ControlRig.h
--- a/Engine/Plugins/Experimental/ControlRig/Source/ControlRig/Public/ControlRig.h
+++ b/Engine/Plugins/Experimental/ControlRig/Source/ControlRig/Public/ControlRig.h
@@ -309,7 +309,7 @@ public:
 		--InterRigSyncBracket;
 	}
 
-	bool SetControlGlobalTransform(const FName& InControlName, const FTransform& InGlobalTransform, const FRigControlModifiedContext& Context = FRigControlModifiedContext());
+	bool SetControlGlobalTransform(const FName& InControlName, const FTransform& InGlobalTransform, bool bNotify = true, const FRigControlModifiedContext& Context = FRigControlModifiedContext());
 
 	virtual FRigControlValue GetControlValueFromGlobalTransform(const FName& InControlName, const FTransform& InGlobalTransform);
 
diff --git a/Engine/Plugins/Experimental/ControlRig/Source/ControlRig/Public/Tools/ControlRigPose.h b/Engine/Plugins/Experimental/ControlRig/Source/ControlRig/Public/Tools/ControlRigPose.h
--- a/Engine/Plugins/Experimental/ControlRig/Source/ControlRig/Public/Tools/ControlRigPose.h
+++ b/Engine/Plugins/Experimental/ControlRig/Source/ControlRig/Public/Tools/ControlRigPose.h
@@ -17,7 +17,7 @@
 /**
 * The Data Stored For Each Control in A Pose.
 */
-USTRUCT()
+USTRUCT(BlueprintType)
 struct CONTROLRIG_API FRigControlCopy
 {
 	GENERATED_BODY()
@@ -29,7 +29,6 @@ struct CONTROLRIG_API FRigControlCopy
 		, ParentName(NAME_None)
 		, SpaceName(NAME_None)
 		, OffsetTransform(FTransform::Identity)
-		, InitialValue()
 		, ParentTransform(FTransform::Identity)
 		, LocalTransform(FTransform::Identity)
 		, GlobalTransform(FTransform::Identity)
@@ -45,7 +44,6 @@ struct CONTROLRIG_API FRigControlCopy
 		ParentName = InControl.ParentName;
 		SpaceName = InControl.SpaceName;
 		OffsetTransform = InControl.OffsetTransform;
-		InitialValue = InControl.InitialValue;
 
 		int32 Index = Hierarchy.GetIndex(Name);
 		ParentTransform = Index != INDEX_NONE ? Hierarchy.GetParentTransform(Index) : FTransform::Identity;
@@ -54,34 +52,31 @@ struct CONTROLRIG_API FRigControlCopy
 	}
 	virtual ~FRigControlCopy() {}
 
-	UPROPERTY()
+	UPROPERTY(BlueprintReadWrite, EditAnywhere, Category = "Names")
 	FName Name;
 
-	UPROPERTY()
+	UPROPERTY(BlueprintReadWrite, EditAnywhere, Category = "Type")
 	ERigControlType ControlType;
 
 	UPROPERTY()
 	FRigControlValue Value;
 
-	UPROPERTY()
+	UPROPERTY(BlueprintReadWrite, EditAnywhere, Category = "Names")
 	FName ParentName;
 
-	UPROPERTY()
+	UPROPERTY(BlueprintReadWrite, EditAnywhere, Category = "Names")
 	FName SpaceName;
 
-	UPROPERTY()
+	UPROPERTY(BlueprintReadWrite, EditAnywhere, Category = "Transforms")
 	FTransform OffsetTransform;
 
-	UPROPERTY()
-	FRigControlValue InitialValue;
-
-	UPROPERTY()
+	UPROPERTY(BlueprintReadWrite, EditAnywhere, Category = "Transforms")
 	FTransform ParentTransform;
 
-	UPROPERTY()
+	UPROPERTY(BlueprintReadWrite, EditAnywhere, Category = "Transforms")
 	FTransform LocalTransform;
 
-	UPROPERTY()
+	UPROPERTY(BlueprintReadWrite, EditAnywhere, Category = "Transforms")
 	FTransform GlobalTransform;
 
 };
@@ -89,7 +84,7 @@ struct CONTROLRIG_API FRigControlCopy
 /**
 * The Data Stored For Each Pose and associated Functions to Store and Paste It
 */
-USTRUCT()
+USTRUCT(BlueprintType)
 struct FControlRigControlPose
 {
 	GENERATED_USTRUCT_BODY()
@@ -103,7 +98,8 @@ struct FControlRigControlPose
 
 	void SavePose(UControlRig* ControlRig,  bool bUseAll)
 	{
-		const TArray<FRigControl>& CurrentControls = ControlRig->AvailableControls();
+		TArray<FRigControl> CurrentControls; 
+		ControlRig->GetControlsInOrder(CurrentControls);
 		CopyOfControls.SetNum(0);
 		FRigControlHierarchy& Hierarchy = ControlRig->GetControlHierarchy();
 		for (const FRigControl& RigControl : CurrentControls)
@@ -114,23 +110,19 @@ struct FControlRigControlPose
 				CopyOfControls.Add(Copy);
 			}
 		}
+		SetUpControlMap();
 	}
 
 	void PastePose(UControlRig* ControlRig, bool bDoKey, bool bDoMirror)
 	{
 		PastePoseInternal(ControlRig,bDoKey, bDoMirror, CopyOfControls);
-		//do it twice do to issues we have seen with the ordering of certain controls,due to how spaces and offsets work.
-		PastePoseInternal(ControlRig, bDoKey, bDoMirror, CopyOfControls);
+
 	}
 
-	void SetControlMirrorTransform(UControlRig* ControlRig, const FName& Name, const FVector& GlobalTranslation, const FQuat& LocalRotation, bool bNotify, FRigControlModifiedContext Context)
+	void SetControlMirrorTransform(UControlRig* ControlRig, const FName& Name, const FVector& GlobalTranslation, const FQuat& GlobalRotation, bool bNotify, FRigControlModifiedContext Context)
 	{
-		int32 Index = ControlRig->GetControlHierarchy().GetIndex(Name);
-		FTransform ParentTransform = ControlRig->GetControlHierarchy().GetParentTransform(Index);
-		FTransform CurrentTransform = ControlRig->GetControlHierarchy().GetLocalTransform(Index);
-		FVector NewLocal = ParentTransform.InverseTransformPositionNoScale(GlobalTranslation);
-		FTransform NewCurrentTransform(LocalRotation, NewLocal);
-		ControlRig->SetControlLocalTransform(Name, NewCurrentTransform, bNotify, Context);
+		FTransform NewGlobalTransform(GlobalRotation, GlobalTranslation);
+		ControlRig->SetControlGlobalTransform(Name, NewGlobalTransform, bNotify, Context);
 	}
 
 	void PastePoseInternal(UControlRig* ControlRig, bool bDoKey,bool bDoMirror, const TArray<FRigControlCopy>& ControlsToPaste)
@@ -143,14 +135,20 @@ struct FControlRigControlPose
 			MirrorTable.SetUpMirrorTable(ControlRig);
 		}
 	
-		for (const FRigControlCopy& CopyRigControl : ControlsToPaste)
-		{
-			FRigControl* RigControl = MirrorTable.GetControl(ControlRig,CopyRigControl.Name);
-			if (RigControl)
+		TArray<FRigControl> SortedControls;
+		ControlRig->GetControlsInOrder(SortedControls);
+		for(const FRigControl& RigControl : SortedControls)
+		{ 
+			if (!ControlRig->IsControlSelected(RigControl.Name))
 			{
-				if (CopyRigControl.ControlType == RigControl->ControlType)
+				continue;
+			}
+			FRigControlCopy* CopyRigControl = MirrorTable.GetControl(*this,RigControl.Name);
+			if (CopyRigControl)
+			{
+				if (CopyRigControl->ControlType == RigControl.ControlType)
 				{
-					switch (RigControl->ControlType)
+					switch (RigControl.ControlType)
 					{
 					case ERigControlType::Transform:
 					case ERigControlType::TransformNoScale:
@@ -161,40 +159,40 @@ struct FControlRigControlPose
 					{
 						if (bDoMirror == false)
 						{
-							ControlRig->SetControlGlobalTransform(RigControl->Name, CopyRigControl.GlobalTransform,Context);
+							ControlRig->SetControlGlobalTransform(RigControl.Name, CopyRigControl->GlobalTransform,true, Context);
 						}
 						else
 						{
 							FVector GlobalTranslation;
-							FQuat LocalRotation;
-							bool bIsMatched = MirrorTable.IsMatched(RigControl->Name);
-							MirrorTable.GetMirrorTransform(CopyRigControl, bIsMatched,GlobalTranslation, LocalRotation);
-							SetControlMirrorTransform(ControlRig,RigControl->Name, GlobalTranslation, LocalRotation,true, Context);
+							FQuat GlobalRotation;
+							bool bIsMatched = MirrorTable.IsMatched(CopyRigControl->Name);
+							MirrorTable.GetMirrorTransform(*CopyRigControl, bIsMatched,GlobalTranslation, GlobalRotation);
+							SetControlMirrorTransform(ControlRig,RigControl.Name, GlobalTranslation, GlobalRotation,true, Context);
 						}
 						break;
 					}		
 					case ERigControlType::Float:
 					{
-						float Val = CopyRigControl.Value.Get<float>();
-						ControlRig->SetControlValue<float>(RigControl->Name, Val, true, Context);
+						float Val = CopyRigControl->Value.Get<float>();
+						ControlRig->SetControlValue<float>(RigControl.Name, Val, true, Context);
 						break;
 					}
 					case ERigControlType::Bool:
 					{
-						bool Val = CopyRigControl.Value.Get<bool>();
-						ControlRig->SetControlValue<bool>(RigControl->Name, Val, true, Context);
+						bool Val = CopyRigControl->Value.Get<bool>();
+						ControlRig->SetControlValue<bool>(RigControl.Name, Val, true, Context);
 						break;
 					}
 					case ERigControlType::Integer:
 					{
-						int32 Val = CopyRigControl.Value.Get<int32>();
-						ControlRig->SetControlValue<int32>(RigControl->Name, Val, true, Context);
+						int32 Val = CopyRigControl->Value.Get<int32>();
+						ControlRig->SetControlValue<int32>(RigControl.Name, Val, true, Context);
 						break;
 					}
 					case ERigControlType::Vector2D:
 					{
-						FVector2D Val = CopyRigControl.Value.Get<FVector2D>();
-						ControlRig->SetControlValue<FVector2D>(RigControl->Name, Val, true, Context);
+						FVector2D Val = CopyRigControl->Value.Get<FVector2D>();
+						ControlRig->SetControlValue<FVector2D>(RigControl.Name, Val, true, Context);
 						break;
 					}	
 					default:
@@ -206,9 +204,9 @@ struct FControlRigControlPose
 		}
 	}
 
-	void BlendWithInitialPoses(TArray<FRigControlCopy>& InitialControls, UControlRig *ControlRig, bool bDoKey, bool bDoMirror, float BlendValue)
+	void BlendWithInitialPoses(FControlRigControlPose& InitialPose, UControlRig *ControlRig, bool bDoKey, bool bDoMirror, float BlendValue)
 	{
-		if (InitialControls.Num() == 0)
+		if (InitialPose.CopyOfControls.Num() == 0)
 		{
 			return;
 		}
@@ -223,45 +221,36 @@ struct FControlRigControlPose
 			MirrorTable.SetUpMirrorTable(ControlRig);
 		}
 
-		for (int32 Index = 0; Index <CopyOfControls.Num(); ++Index)
+		TArray<FRigControl> SortedControls;
+		ControlRig->GetControlsInOrder(SortedControls);
+		for (const FRigControl& RigControl : SortedControls)
 		{
-			FRigControlCopy& CopyRigControl = CopyOfControls[Index];
-			FRigControl* RigControl = MirrorTable.GetControl(ControlRig, CopyRigControl.Name);
-			if (RigControl)
+			if (!ControlRig->IsControlSelected(RigControl.Name))
+			{
+				continue;
+			}
+			FRigControlCopy* CopyRigControl = MirrorTable.GetControl(*this, RigControl.Name);
+			if (CopyRigControl)
 			{
-				if (RigControl->ControlType == CopyRigControl.ControlType)
+				if (CopyRigControl->ControlType == RigControl.ControlType)
 				{
+
 					FRigControlCopy* InitialFound = nullptr;
-					for (int32 NewIncIndex = Index; NewIncIndex < InitialControls.Num(); ++NewIncIndex)
+					int32* Index = InitialPose.CopyOfControlsNameToIndex.Find(CopyRigControl->Name);
+					if (Index)
 					{
-						if (CopyOfControls[Index].Name == InitialControls[NewIncIndex].Name)
-						{
-							InitialFound = &InitialControls[NewIncIndex];
-							break;
-						}
+						InitialFound = &(InitialPose.CopyOfControls[*Index]);
 					}
-					if (!InitialFound && Index >= 1)
+					if (InitialFound && InitialFound->ControlType == CopyRigControl->ControlType)
 					{
-						//now search backwards to zero
-						for (int32 NewDecIndex = Index - 1; NewDecIndex >= 0; --NewDecIndex)
-						{
-							if (CopyOfControls[Index].Name == InitialControls[NewDecIndex].Name)
-							{
-								InitialFound = &InitialControls[NewDecIndex];
-								break;
-							}
-						}
-					}
-					if (InitialFound && InitialFound->ControlType == CopyOfControls[Index].ControlType)
-					{
-						if ((CopyRigControl.ControlType == ERigControlType::Transform || CopyRigControl.ControlType == ERigControlType::EulerTransform ||
-							CopyRigControl.ControlType == ERigControlType::TransformNoScale || CopyRigControl.ControlType == ERigControlType::Position ||
-							CopyRigControl.ControlType == ERigControlType::Rotator || CopyRigControl.ControlType == ERigControlType::Scale
+						if ((CopyRigControl->ControlType == ERigControlType::Transform || CopyRigControl->ControlType == ERigControlType::EulerTransform ||
+							CopyRigControl->ControlType == ERigControlType::TransformNoScale || CopyRigControl->ControlType == ERigControlType::Position ||
+							CopyRigControl->ControlType == ERigControlType::Rotator || CopyRigControl->ControlType == ERigControlType::Scale
 							))
 						{
 							if (bDoMirror == false)
 							{
-								FTransform Val = CopyRigControl.GlobalTransform;
+								FTransform Val = CopyRigControl->GlobalTransform;
 								FTransform InitialVal = InitialFound->GlobalTransform;
 								FVector Translation, Scale;
 								FQuat Rotation;
@@ -269,19 +258,19 @@ struct FControlRigControlPose
 								Rotation = FQuat::Slerp(InitialVal.GetRotation(), Val.GetRotation(), BlendValue); //doing slerp here not fast lerp, can be slow this is for content creation
 								Scale = FMath::Lerp(InitialVal.GetScale3D(), Val.GetScale3D(), BlendValue);
 								Val = FTransform(Rotation, Translation, Scale);
-								ControlRig->SetControlGlobalTransform(RigControl->Name, Val, Context);
+								ControlRig->SetControlGlobalTransform(RigControl.Name, Val, false, Context);
 							}
 							else
 							{
 								FVector GlobalTranslation;
-								FQuat LocalRotation;
-								bool bIsMatched = MirrorTable.IsMatched(CopyRigControl.Name);
-								MirrorTable.GetMirrorTransform(CopyRigControl, bIsMatched, GlobalTranslation, LocalRotation);
+								FQuat GlobalRotation;
+								bool bIsMatched = MirrorTable.IsMatched(CopyRigControl->Name);
+								MirrorTable.GetMirrorTransform(*CopyRigControl, bIsMatched, GlobalTranslation, GlobalRotation);
 								FVector InitialTranslation = InitialFound->GlobalTransform.GetTranslation();
 								FQuat InitialRotation = InitialFound->LocalTransform.GetRotation();
 								GlobalTranslation = FMath::Lerp(InitialTranslation, GlobalTranslation, BlendValue);
-								LocalRotation = FQuat::Slerp(InitialRotation, LocalRotation, BlendValue); //doing slerp here not fast lerp, can be slow this is for content creation
-								SetControlMirrorTransform(ControlRig, RigControl->Name, GlobalTranslation, LocalRotation, false, Context);
+								GlobalRotation = FQuat::Slerp(InitialRotation, GlobalRotation, BlendValue); //doing slerp here not fast lerp, can be slow this is for content creation
+								SetControlMirrorTransform(ControlRig, RigControl.Name, GlobalTranslation, GlobalRotation, false, Context);
 							}
 						}
 					}
@@ -301,15 +290,29 @@ struct FControlRigControlPose
 		return Controls;
 	}
 
+	void SetUpControlMap()
+	{
+		CopyOfControlsNameToIndex.Reset();
+		
+		for (int32 Index = 0;  Index <  CopyOfControls.Num(); ++ Index)
+		{
+			const FRigControlCopy& Control =  CopyOfControls[Index];
+			CopyOfControlsNameToIndex.Add(Control.Name, Index);
+		}
+	}
 	TArray<FRigControlCopy> GetPoses() const {return CopyOfControls;};
 
-	UPROPERTY()
+public:
+	UPROPERTY(BlueprintReadWrite, EditAnywhere, Category = "Controls")
 	TArray<FRigControlCopy> CopyOfControls;
+
+	//Cache of the Map, Used to make pasting faster.
+	TMap<FName, int32>  CopyOfControlsNameToIndex;
 };
 
 
 /**
-* The Actual Pose Asset that holds the Pose.
+* An indivual Pose made of Control Rig Controls
 */
 UCLASS(BlueprintType)
 class CONTROLRIG_API UControlRigPoseAsset : public UObject
@@ -318,6 +321,9 @@ class CONTROLRIG_API UControlRigPoseAsset : public UObject
 
 public:
 
+	//UOBJECT
+	virtual void PostLoad() override;
+
 	UFUNCTION(BlueprintCallable, Category = "Pose")
 	void SavePose(UControlRig* InControlRig, bool bUseAll);
 
@@ -326,20 +332,15 @@ public:
 	
 	UFUNCTION(BlueprintCallable, Category = "Pose")
 	void SelectControls(UControlRig* InControlRig);
-	
-	UFUNCTION(BlueprintCallable, Category = "Pose")
-	const TArray<FName>& GetControlNames() const { return Controls; }
 
 	TArray<FRigControlCopy> GetCurrentPose(UControlRig* InControlRig);
-	void BlendWithInitialPoses(TArray<FRigControlCopy>& InitialControls, UControlRig* ControlRig, bool bDoKey, bool bdoMirror, float BlendValue);
 
-private:
+	UFUNCTION(BlueprintCallable, Category = "Pose")
+	void GetCurrentPose(UControlRig* InControlRig, FControlRigControlPose& OutPose);
 
-	UPROPERTY()
-	FControlRigControlPose Pose;
+	void BlendWithInitialPoses(FControlRigControlPose& InitialPose, UControlRig* ControlRig, bool bDoKey, bool bdoMirror, float BlendValue);
 
-	/** Controls we are made of */
-	UPROPERTY(AssetRegistrySearchable)
-	TArray<FName> Controls;
+	UPROPERTY(BlueprintReadWrite, EditAnywhere, Category = "Pose")
+	FControlRigControlPose Pose;
 
 };
diff --git a/Engine/Plugins/Experimental/ControlRig/Source/ControlRig/Public/Tools/ControlRigPoseMirrorTable.h b/Engine/Plugins/Experimental/ControlRig/Source/ControlRig/Public/Tools/ControlRigPoseMirrorTable.h
--- a/Engine/Plugins/Experimental/ControlRig/Source/ControlRig/Public/Tools/ControlRigPoseMirrorTable.h
+++ b/Engine/Plugins/Experimental/ControlRig/Source/ControlRig/Public/Tools/ControlRigPoseMirrorTable.h
@@ -13,6 +13,7 @@
 class UControlRig;
 struct FRigControl;
 struct FRigControlCopy;
+struct FControlRigControlPose;
 
 struct FControlRigPoseMirrorTable
 {
@@ -24,7 +25,7 @@ public:
 	void SetUpMirrorTable(const UControlRig* ControlRig);
 
 	/*Get the matched control with the given name*/
-	FRigControl* GetControl(UControlRig* ControlRig, const FName& ControlName) const;
+	FRigControlCopy* GetControl(FControlRigControlPose& Pose, FName ControlrigName);
 
 	/*Whether or not the Control with this name is matched*/
 	bool IsMatched(const FName& ControlName) const;
diff --git a/Engine/Plugins/Experimental/ControlRig/Source/ControlRigEditor/Private/EditMode/ControlRigEditMode.cpp b/Engine/Plugins/Experimental/ControlRig/Source/ControlRigEditor/Private/EditMode/ControlRigEditMode.cpp
--- a/Engine/Plugins/Experimental/ControlRig/Source/ControlRigEditor/Private/EditMode/ControlRigEditMode.cpp
+++ b/Engine/Plugins/Experimental/ControlRig/Source/ControlRigEditor/Private/EditMode/ControlRigEditMode.cpp
@@ -1983,7 +1983,7 @@ void FControlRigEditMode::MoveGizmo(AControlRigGizmoActor* GizmoActor, const boo
 				{
 					InOutLocal = InteractionControlRig->GetControlLocalTransform(GizmoActor->ControlName);
 				}
-				InteractionControlRig->SetControlGlobalTransform(GizmoActor->ControlName, NewTransform, Context);			// assumes it's attached to actor
+				InteractionControlRig->SetControlGlobalTransform(GizmoActor->ControlName, NewTransform, true, Context);			// assumes it's attached to actor
 				GizmoActor->SetGlobalTransform(CurrentTransform);
 				if (bCalcLocal)
 				{
diff --git a/Engine/Plugins/Experimental/ControlRig/Source/ControlRigEditor/Private/EditMode/SControlRigBaseListWidget.cpp b/Engine/Plugins/Experimental/ControlRig/Source/ControlRigEditor/Private/EditMode/SControlRigBaseListWidget.cpp
--- a/Engine/Plugins/Experimental/ControlRig/Source/ControlRigEditor/Private/EditMode/SControlRigBaseListWidget.cpp
+++ b/Engine/Plugins/Experimental/ControlRig/Source/ControlRigEditor/Private/EditMode/SControlRigBaseListWidget.cpp
@@ -885,6 +885,15 @@ TSharedPtr<SWidget> SControlRigBaseListWidget::OnGetAssetContextMenu(const TArra
 			}
 			MenuBuilder.EndSection();
 
+			MenuBuilder.BeginSection("PoseDialogOptions", LOCTEXT("PoseDialogUpdateHeading", "Update"));
+			{
+				FUIAction Action = FUIAction(FExecuteAction::CreateRaw((this), &SControlRigBaseListWidget::ExecuteUpdatePose, PoseAsset));
+				const FText Label = LOCTEXT("UpdatePose", "Update Pose");
+				const FText ToolTipText = LOCTEXT("Update Pose Tooltip", "Update The Pose Based Upon Current Control Rig Pose");
+				MenuBuilder.AddMenuEntry(Label, ToolTipText, FSlateIcon(), Action);
+			}
+			MenuBuilder.EndSection();
+
 		}
 	}
 	return MenuBuilder.MakeWidget();
@@ -1051,6 +1060,16 @@ bool SControlRigBaseListWidget::CanExecutePastePose(UControlRigPoseAsset* PoseAs
 {
 	return PoseAsset != nullptr;
 }
+
+void SControlRigBaseListWidget::ExecuteUpdatePose(UControlRigPoseAsset* PoseAsset)
+{
+	FControlRigEditMode* ControlRigEditMode = static_cast<FControlRigEditMode*>(GLevelEditorModeTools().GetActiveMode(FControlRigEditMode::ModeName));
+	if (ControlRigEditMode && ControlRigEditMode->GetControlRig(true))
+	{
+		PoseAsset->SavePose(ControlRigEditMode->GetControlRig(true),false);
+	}
+}
+
 void SControlRigBaseListWidget::ExecuteSelectControls(UControlRigPoseAsset* PoseAsset)
 {
 	FControlRigEditMode* ControlRigEditMode = static_cast<FControlRigEditMode*>(GLevelEditorModeTools().GetActiveMode(FControlRigEditMode::ModeName));
diff --git a/Engine/Plugins/Experimental/ControlRig/Source/ControlRigEditor/Private/EditMode/SControlRigBaseListWidget.h b/Engine/Plugins/Experimental/ControlRig/Source/ControlRigEditor/Private/EditMode/SControlRigBaseListWidget.h
--- a/Engine/Plugins/Experimental/ControlRig/Source/ControlRigEditor/Private/EditMode/SControlRigBaseListWidget.h
+++ b/Engine/Plugins/Experimental/ControlRig/Source/ControlRigEditor/Private/EditMode/SControlRigBaseListWidget.h
@@ -99,6 +99,7 @@ private:
 	void ExecuteSelectControls(UControlRigPoseAsset* PoseAsset);
 	void ExecutePasteMirrorPose(UControlRigPoseAsset* PoseAsset);
 	bool CanExecutePasteMirrorPose(UControlRigPoseAsset* PoseAsset) const;
+	void ExecuteUpdatePose(UControlRigPoseAsset* PoseAsset);
 
 	/** Create the Views for the Selected Asset*/
 	void CreateCurrentView(UObject* Asset);
diff --git a/Engine/Plugins/Experimental/ControlRig/Source/ControlRigEditor/Private/EditMode/SControlRigControlViews.cpp b/Engine/Plugins/Experimental/ControlRig/Source/ControlRigEditor/Private/EditMode/SControlRigControlViews.cpp
--- a/Engine/Plugins/Experimental/ControlRig/Source/ControlRigEditor/Private/EditMode/SControlRigControlViews.cpp
+++ b/Engine/Plugins/Experimental/ControlRig/Source/ControlRigEditor/Private/EditMode/SControlRigControlViews.cpp
@@ -178,16 +178,17 @@ void SControlRigAssetEditableTextBox::OnTextChanged(const FText& InLabel)
 	}
 }
 
+bool SControlRigPoseView::bIsKey = false;
+bool SControlRigPoseView::bIsMirror = false;
+
 void SControlRigPoseView::Construct(const FArguments& InArgs)
 {
 	PoseAsset = InArgs._PoseAsset;
 
-	bIsKey = false;
-	bIsMirror = false;
+
 	PoseBlendValue = 0.0f;
 	bIsBlending = false;
 	bSliderStartedTransaction = false;
-	InitialControlValues.SetNum(0);
 
 	TSharedRef<SWidget> Thumbnail = GetThumbnailWidget();
 	TSharedRef <SControlRigAssetEditableTextBox> ObjectNameBox = SNew(SControlRigAssetEditableTextBox).Asset(PoseAsset);
@@ -471,10 +472,10 @@ void SControlRigPoseView::OnPoseBlendChanged(float ChangedVal)
 		if (!bIsBlending)
 		{
 			bIsBlending = true;
-			InitialControlValues = PoseAsset->GetCurrentPose(ControlRig);
+			PoseAsset->GetCurrentPose(ControlRig,TempPose);
 		}
 
-		PoseAsset->BlendWithInitialPoses(InitialControlValues, ControlRig, false, bIsMirror, PoseBlendValue);
+		PoseAsset->BlendWithInitialPoses(TempPose, ControlRig, false, bIsMirror, PoseBlendValue);
 	}
 }
 void SControlRigPoseView::OnBeginSliderMovement()
@@ -502,7 +503,7 @@ void SControlRigPoseView::OnPoseBlendCommited(float ChangedVal, ETextCommit::Typ
 	{
 		FScopedTransaction ScopedTransaction(LOCTEXT("PastePoseTransaction", "Paste Pose"));
 		PoseBlendValue = ChangedVal;
-		PoseAsset->BlendWithInitialPoses(InitialControlValues, ControlRig, bIsKey, bIsMirror, PoseBlendValue);
+		PoseAsset->BlendWithInitialPoses(TempPose, ControlRig, bIsKey, bIsMirror, PoseBlendValue);
 		bIsBlending = false;
 		PoseBlendValue = 0.0f;
 	}
diff --git a/Engine/Plugins/Experimental/ControlRig/Source/ControlRigEditor/Private/EditMode/SControlRigControlViews.h b/Engine/Plugins/Experimental/ControlRig/Source/ControlRigEditor/Private/EditMode/SControlRigControlViews.h
--- a/Engine/Plugins/Experimental/ControlRig/Source/ControlRigEditor/Private/EditMode/SControlRigControlViews.h
+++ b/Engine/Plugins/Experimental/ControlRig/Source/ControlRigEditor/Private/EditMode/SControlRigControlViews.h
@@ -57,12 +57,12 @@ private:
 	FReply OnSelectControls();
 	FReply OnCaptureThumbnail();
 
-	bool bIsKey;
-	bool bIsMirror;
+	static bool bIsKey;
+	static bool bIsMirror;
 	float PoseBlendValue;
 	bool bIsBlending;
 	bool bSliderStartedTransaction;
-	TArray<FRigControlCopy> InitialControlValues;
+	FControlRigControlPose TempPose;
 
 	TWeakObjectPtr<UControlRigPoseAsset> PoseAsset;
 
