Control Rig: Adding Python API for unit tests

diff --git a/Engine/Plugins/Experimental/ControlRig/Source/ControlRig/Private/ControlRig.cpp b/Engine/Plugins/Experimental/ControlRig/Source/ControlRig/Private/ControlRig.cpp
--- a/Engine/Plugins/Experimental/ControlRig/Source/ControlRig/Private/ControlRig.cpp
+++ b/Engine/Plugins/Experimental/ControlRig/Source/ControlRig/Private/ControlRig.cpp
@@ -16,6 +16,7 @@
 #include "RigVMModel/RigVMNode.h"
 #include "Engine/Blueprint.h"
 #include "EdGraphSchema_K2.h"
+#include "Kismet2/BlueprintEditorUtils.h"
 #endif// WITH_EDITOR
 
 #define LOCTEXT_NAMESPACE "ControlRig"
@@ -279,6 +280,55 @@ FRigVMExternalVariable UControlRig::GetPublicVariableByName(const FName& InVaria
 	return FRigVMExternalVariable();
 }
 
+TArray<FName> UControlRig::GetScriptAccessibleVariables() const
+{
+	TArray<FRigVMExternalVariable> PublicVariables = GetPublicVariables();
+	TArray<FName> Names;
+	for (const FRigVMExternalVariable& PublicVariable : PublicVariables)
+	{
+		Names.Add(PublicVariable.Name);
+	}
+	return Names;
+}
+
+FName UControlRig::GetVariableType(const FName& InVariableName) const
+{
+	FRigVMExternalVariable PublicVariable = GetPublicVariableByName(InVariableName);
+	if (PublicVariable.IsValid(true /* allow nullptr */))
+	{
+		return PublicVariable.TypeName;
+	}
+	return NAME_None;
+}
+
+FString UControlRig::GetVariableAsString(const FName& InVariableName) const
+{
+#if WITH_EDITOR
+	if (const FProperty* Property = GetClass()->FindPropertyByName(InVariableName))
+	{
+		FString Result;
+		const uint8* Container = (const uint8*)this;
+		if (FBlueprintEditorUtils::PropertyValueToString(Property, Container, Result, nullptr))
+		{
+			return Result;
+		}
+	}
+#endif
+	return FString();
+}
+
+bool UControlRig::SetVariableFromString(const FName& InVariableName, const FString& InValue)
+{
+#if WITH_EDITOR
+	if (const FProperty* Property = GetClass()->FindPropertyByName(InVariableName))
+	{
+		uint8* Container = (uint8*)this;
+		return FBlueprintEditorUtils::PropertyValueFromString(Property, InValue, Container, nullptr);
+	}
+#endif
+	return false;
+}
+
 bool UControlRig::SupportsEvent(const FName& InEventName) const
 {
 	if (VM)
@@ -864,6 +914,16 @@ URigVM* UControlRig::GetVM()
 	return VM;
 }
 
+FTransform UControlRig::GetGlobalTransform(const FRigElementKey& InKey) const
+{
+	return Hierarchy.GetGlobalTransform(InKey);
+}
+
+void UControlRig::SetGlobalTransform(const FRigElementKey& InKey, const FTransform& InTransform, bool bPropagateTransform)
+{
+	Hierarchy.SetGlobalTransform(InKey, InTransform, bPropagateTransform);
+}
+
 FTransform UControlRig::GetGlobalTransform(const FName& BoneName) const
 {
 	int32 Index = Hierarchy.BoneHierarchy.GetIndex(BoneName);
diff --git a/Engine/Plugins/Experimental/ControlRig/Source/ControlRig/Private/Rigs/RigControlHierarchy.cpp b/Engine/Plugins/Experimental/ControlRig/Source/ControlRig/Private/Rigs/RigControlHierarchy.cpp
--- a/Engine/Plugins/Experimental/ControlRig/Source/ControlRig/Private/Rigs/RigControlHierarchy.cpp
+++ b/Engine/Plugins/Experimental/ControlRig/Source/ControlRig/Private/Rigs/RigControlHierarchy.cpp
@@ -461,6 +461,11 @@ void FRigControl::SetValueFromTransform(const FTransform& InTransform, ERigContr
 			break;
 		}
 	}
+
+	if (InValueType == ERigControlValueType::Current)
+	{
+		ApplyLimits(GetValue(InValueType));
+	}
 }
 
 ////////////////////////////////////////////////////////////////////////////////
diff --git a/Engine/Plugins/Experimental/ControlRig/Source/ControlRig/Public/ControlRig.h b/Engine/Plugins/Experimental/ControlRig/Source/ControlRig/Public/ControlRig.h
--- a/Engine/Plugins/Experimental/ControlRig/Source/ControlRig/Public/ControlRig.h
+++ b/Engine/Plugins/Experimental/ControlRig/Source/ControlRig/Public/ControlRig.h
@@ -78,18 +78,23 @@ public:
 #endif
 	
 	/** Set the current delta time */
+	UFUNCTION(BlueprintCallable, Category="Control Rig")
 	void SetDeltaTime(float InDeltaTime);
 
 	/** Set the current absolute time */
+	UFUNCTION(BlueprintCallable, Category = "Control Rig")
 	void SetAbsoluteTime(float InAbsoluteTime, bool InSetDeltaTimeZero = false);
 
 	/** Set the current absolute and delta times */
+	UFUNCTION(BlueprintCallable, Category = "Control Rig")
 	void SetAbsoluteAndDeltaTime(float InAbsoluteTime, float InDeltaTime);
 
 	/** Set the current fps */
+	UFUNCTION(BlueprintCallable, Category = "Control Rig")
 	void SetFramesPerSecond(float InFramesPerSecond);
 
 	/** Returns the current frames per second (this may change over time) */
+	UFUNCTION(BlueprintCallable, Category = "Control Rig")
 	float GetCurrentFramesPerSecond() const;
 
 #if WITH_EDITOR
@@ -118,6 +123,22 @@ public:
 	/** Returns a public variable given its name */
 	FRigVMExternalVariable GetPublicVariableByName(const FName& InVariableName) const;
 
+	/** Returns the names of variables accessible in scripting */
+	UFUNCTION(BlueprintCallable, Category = "Control Rig", meta=(DisplayName="Get Variables"))
+	TArray<FName> GetScriptAccessibleVariables() const;
+
+	/** Returns the type of a given variable */
+	UFUNCTION(BlueprintCallable, Category = "Control Rig")
+	FName GetVariableType(const FName& InVariableName) const;
+
+	/** Returns the value of a given variable as a string */
+	UFUNCTION(BlueprintCallable, Category = "Control Rig")
+	FString GetVariableAsString(const FName& InVariableName) const;
+
+	/** Returns the value of a given variable as a string */
+	UFUNCTION(BlueprintCallable, Category = "Control Rig")
+	bool SetVariableFromString(const FName& InVariableName, const FString& InValue);
+
 	template<class T>
 	FORCEINLINE T GetPublicVariableValue(const FName& InVariableName)
 	{
@@ -136,8 +157,10 @@ public:
 		return SupportsEvent(T::EventName);
 	}
 
+	UFUNCTION(BlueprintCallable, Category = "Control Rig")
 	bool SupportsEvent(const FName& InEventName) const;
 
+	UFUNCTION(BlueprintCallable, Category = "Control Rig")
 	TArray<FName> GetSupportedEvents() const;
 
 	/** Setup bindings to a runtime object (or clear by passing in nullptr). */
@@ -183,22 +206,32 @@ public:
 		return Hierarchy.CurveContainer;
 	}
 
+	/** Evaluate another animation ControlRig */
+	UFUNCTION(BlueprintCallable, Category = "Control Rig")
+	FTransform GetGlobalTransform(const FRigElementKey& InKey) const;
+
+	/** Evaluate another animation ControlRig */
+	UFUNCTION(BlueprintCallable, Category = "Control Rig")
+	void SetGlobalTransform(const FRigElementKey& InKey, const FTransform& InTransform, bool bPropagateTransform = true);
+
 	/** Evaluate another animation ControlRig */
 	FTransform GetGlobalTransform(const FName& BoneName) const;
 
 	/** Evaluate another animation ControlRig */
-	void SetGlobalTransform(const FName& BoneName, const FTransform& InTransform, bool bPropagateTransform = true) ;
+	void SetGlobalTransform(const FName& BoneName, const FTransform& InTransform, bool bPropagateTransform = true);
 
 	/** Evaluate another animation ControlRig */
 	FTransform GetGlobalTransform(const int32 BoneIndex) const;
 
 	/** Evaluate another animation ControlRig */
-	void SetGlobalTransform(const int32 BoneIndex, const FTransform& InTransform, bool bPropagateTransform = true) ;
+	void SetGlobalTransform(const int32 BoneIndex, const FTransform& InTransform, bool bPropagateTransform = true);
 
 	/** Evaluate another animation ControlRig */
+	UFUNCTION(BlueprintCallable, Category = "Control Rig")
 	float GetCurveValue(const FName& CurveName) const;
 
 	/** Evaluate another animation ControlRig */
+	UFUNCTION(BlueprintCallable, Category = "Control Rig")
 	void SetCurveValue(const FName& CurveName, const float CurveValue);
 
 	/** Evaluate another animation ControlRig */
@@ -223,15 +256,18 @@ public:
 	ERigExecutionType ExecutionType;
 
 	/** Execute */
+	UFUNCTION(BlueprintCallable, Category = "Control Rig")
 	void Execute(const EControlRigState State, const FName& InEventName);
 
 	/** ExecuteUnits */
 	virtual void ExecuteUnits(FRigUnitContext& InOutContext, const FName& InEventName);
 
 	/** Requests to perform an init during the next execution */
+	UFUNCTION(BlueprintCallable, Category = "Control Rig")
 	void RequestInit();
 
 	/** Requests to perform a setup during the next execution */
+	UFUNCTION(BlueprintCallable, Category = "Control Rig")
 	void RequestSetup();
 
 	/** Returns the queue of events to run */
diff --git a/Engine/Plugins/Experimental/ControlRig/Source/ControlRigDeveloper/Private/ControlRigBlueprint.cpp b/Engine/Plugins/Experimental/ControlRig/Source/ControlRigDeveloper/Private/ControlRigBlueprint.cpp
--- a/Engine/Plugins/Experimental/ControlRig/Source/ControlRigDeveloper/Private/ControlRigBlueprint.cpp
+++ b/Engine/Plugins/Experimental/ControlRig/Source/ControlRigDeveloper/Private/ControlRigBlueprint.cpp
@@ -514,6 +514,20 @@ TArray<UControlRigBlueprint*> UControlRigBlueprint::GetCurrentlyOpenRigBlueprint
 	return sCurrentlyOpenedRigBlueprints;
 }
 
+UClass* UControlRigBlueprint::GetControlRigClass()
+{
+	return GeneratedClass;
+}
+
+UControlRig* UControlRigBlueprint::CreateControlRig()
+{
+	RecompileVMIfRequired();
+
+	UControlRig* Rig = NewObject<UControlRig>(this, GetControlRigClass());
+	Rig->Initialize(true);
+	return Rig;
+}
+
 TArray<UStruct*> UControlRigBlueprint::GetAvailableRigUnits()
 {
 	const TArray<FRigVMFunction>& Functions = FRigVMRegistry::Get().GetFunctions();
@@ -547,6 +561,78 @@ UControlRigHierarchyModifier* UControlRigBlueprint::GetHierarchyModifier()
 
 #if WITH_EDITOR
 
+FName UControlRigBlueprint::AddMemberVariable(const FName& InName, const FString& InCPPType, bool bIsPublic, bool bIsReadOnly)
+{
+	FRigVMExternalVariable Variable;
+	Variable.Name = InName;
+	Variable.bIsPublic = bIsPublic;
+	Variable.bIsReadOnly = bIsReadOnly;
+
+	FString CPPType = InCPPType;
+	if (CPPType.StartsWith(TEXT("TMap<")))
+	{
+		UE_LOG(LogControlRigDeveloper, Warning, TEXT("TMap Variables are not supported."));
+		return NAME_None;
+	}
+
+	Variable.bIsArray = CPPType.StartsWith(TEXT("TArray<"));
+	if (Variable.bIsArray)
+	{
+		CPPType = CPPType.RightChop(7).LeftChop(1);
+	}
+
+	if (CPPType == TEXT("bool"))
+	{
+		Variable.TypeName = *CPPType;
+		Variable.Size = sizeof(bool);
+	}
+	else if (CPPType == TEXT("float"))
+	{
+		Variable.TypeName = *CPPType;
+		Variable.Size = sizeof(float);
+	}
+	else if (CPPType == TEXT("int32"))
+	{
+		Variable.TypeName = *CPPType;
+		Variable.Size = sizeof(int32);
+	}
+	else if (CPPType == TEXT("FString"))
+	{
+		Variable.TypeName = *CPPType;
+		Variable.Size = sizeof(FString);
+	}
+	else if (CPPType == TEXT("FName"))
+	{
+		Variable.TypeName = *CPPType;
+		Variable.Size = sizeof(FName);
+	}
+	else if (CPPType == TEXT("FName"))
+	{
+		Variable.TypeName = *CPPType;
+		Variable.Size = sizeof(FName);
+	}
+	else if(UScriptStruct* ScriptStruct = URigVMPin::FindObjectFromCPPTypeObjectPath<UScriptStruct>(CPPType))
+	{
+		Variable.TypeName = *ScriptStruct->GetStructCPPName();
+		Variable.TypeObject = ScriptStruct;
+		Variable.Size = ScriptStruct->GetStructureSize();
+	}
+	else if (UEnum* Enum= URigVMPin::FindObjectFromCPPTypeObjectPath<UEnum>(CPPType))
+	{
+		Variable.TypeName = *Enum->CppType;
+		Variable.TypeObject = Enum;
+		Variable.Size = Enum->GetResourceSizeBytes(EResourceSizeMode::EstimatedTotal);
+	}
+
+	FName Result = AddCRMemberVariableFromExternal(Variable);
+	if (!Result.IsNone())
+	{
+		FBPCompileRequest Request(this, EBlueprintCompileOptions::None, nullptr);
+		FBlueprintCompilationManager::CompileSynchronously(Request);
+	}
+	return Result;
+}
+
 FName UControlRigBlueprint::AddTransientControl(URigVMPin* InPin)
 {
 	TUniquePtr<FControlValueScope> ValueScope;
diff --git a/Engine/Plugins/Experimental/ControlRig/Source/ControlRigDeveloper/Public/ControlRigBlueprint.h b/Engine/Plugins/Experimental/ControlRig/Source/ControlRigDeveloper/Public/ControlRigBlueprint.h
--- a/Engine/Plugins/Experimental/ControlRig/Source/ControlRigDeveloper/Public/ControlRigBlueprint.h
+++ b/Engine/Plugins/Experimental/ControlRig/Source/ControlRigDeveloper/Public/ControlRigBlueprint.h
@@ -126,12 +126,23 @@ public:
 	UFUNCTION(BlueprintCallable, Category = "VM")
 	static TArray<UControlRigBlueprint*> GetCurrentlyOpenRigBlueprints();
 
+	UFUNCTION(BlueprintCallable, Category = "VM")
+	UClass* GetControlRigClass();
+
+	UFUNCTION(BlueprintCallable, Category = "VM")
+	UControlRig* CreateControlRig();
+
 	UFUNCTION(BlueprintCallable, Category = "VM")
 	static TArray<UStruct*> GetAvailableRigUnits();
 
 	UFUNCTION(BlueprintCallable, Category = "Hierarchy")
 	UControlRigHierarchyModifier* GetHierarchyModifier();
 
+#if WITH_EDITOR
+	UFUNCTION(BlueprintCallable, Category = "Variables")
+	FName AddMemberVariable(const FName& InName, const FString& InCPPType, bool bIsPublic = false, bool bIsReadOnly = false);
+#endif
+
 #if WITH_EDITORONLY_DATA
 	UPROPERTY(EditAnywhere, config, Category = DefaultGizmo)
 	TSoftObjectPtr<UControlRigGizmoLibrary> GizmoLibrary;
