Control Rig Anim Node: Added null check for QueueControlRigDrawInstructions()

diff --git a/Engine/Plugins/Experimental/ControlRig/Source/ControlRig/Private/AnimNode_ControlRigBase.cpp b/Engine/Plugins/Experimental/ControlRig/Source/ControlRig/Private/AnimNode_ControlRigBase.cpp
--- a/Engine/Plugins/Experimental/ControlRig/Source/ControlRig/Private/AnimNode_ControlRigBase.cpp
+++ b/Engine/Plugins/Experimental/ControlRig/Source/ControlRig/Private/AnimNode_ControlRigBase.cpp
@@ -436,51 +436,54 @@ void FAnimNode_ControlRigBase::QueueControlRigDrawInstructions(UControlRig* Cont
 	ensure(ControlRig);
 	ensure(Proxy);
 
-	for (const FControlRigDrawInstruction& Instruction : ControlRig->GetDrawInterface())
+	if (ControlRig && Proxy)
 	{
-		if (!Instruction.IsValid())
+		for (const FControlRigDrawInstruction& Instruction : ControlRig->GetDrawInterface())
 		{
-			continue;
-		}
+			if (!Instruction.IsValid())
+			{
+				continue;
+			}
 
-		FTransform InstructionTransform = Instruction.Transform * Proxy->GetComponentTransform();
-		switch (Instruction.PrimitiveType)
-		{
-			case EControlRigDrawSettings::Points:
+			FTransform InstructionTransform = Instruction.Transform * Proxy->GetComponentTransform();
+			switch (Instruction.PrimitiveType)
 			{
-				for (const FVector& Point : Instruction.Positions)
+				case EControlRigDrawSettings::Points:
 				{
-					Proxy->AnimDrawDebugPoint(InstructionTransform.TransformPosition(Point), Instruction.Thickness, Instruction.Color.ToFColor(true), false, -1.f, SDPG_Foreground);
+					for (const FVector& Point : Instruction.Positions)
+					{
+						Proxy->AnimDrawDebugPoint(InstructionTransform.TransformPosition(Point), Instruction.Thickness, Instruction.Color.ToFColor(true), false, -1.f, SDPG_Foreground);
+					}
+					break;
 				}
-				break;
-			}
-			case EControlRigDrawSettings::Lines:
-			{
-				const TArray<FVector>& Points = Instruction.Positions;
+				case EControlRigDrawSettings::Lines:
+				{
+					const TArray<FVector>& Points = Instruction.Positions;
 
-				for (int32 PointIndex = 0; PointIndex < Points.Num() - 1; PointIndex += 2)
+					for (int32 PointIndex = 0; PointIndex < Points.Num() - 1; PointIndex += 2)
+					{
+						Proxy->AnimDrawDebugLine(InstructionTransform.TransformPosition(Points[PointIndex]), InstructionTransform.TransformPosition(Points[PointIndex + 1]), Instruction.Color.ToFColor(true), false, -1.f, Instruction.Thickness, SDPG_Foreground);
+					}
+					break;
+				}
+				case EControlRigDrawSettings::LineStrip:
 				{
-					Proxy->AnimDrawDebugLine(InstructionTransform.TransformPosition(Points[PointIndex]), InstructionTransform.TransformPosition(Points[PointIndex + 1]), Instruction.Color.ToFColor(true), false, -1.f, Instruction.Thickness, SDPG_Foreground);
+					const TArray<FVector>& Points = Instruction.Positions;
+
+					for (int32 PointIndex = 0; PointIndex < Points.Num() - 1; PointIndex++)
+					{
+						Proxy->AnimDrawDebugLine(InstructionTransform.TransformPosition(Points[PointIndex]), InstructionTransform.TransformPosition(Points[PointIndex + 1]), Instruction.Color.ToFColor(true), false, -1.f, Instruction.Thickness, SDPG_Foreground);
+					}
+					break;
 				}
-				break;
-			}
-			case EControlRigDrawSettings::LineStrip:
-			{
-				const TArray<FVector>& Points = Instruction.Positions;
 
-				for (int32 PointIndex = 0; PointIndex < Points.Num() - 1; PointIndex++)
+				case EControlRigDrawSettings::DynamicMesh:
 				{
-					Proxy->AnimDrawDebugLine(InstructionTransform.TransformPosition(Points[PointIndex]), InstructionTransform.TransformPosition(Points[PointIndex + 1]), Instruction.Color.ToFColor(true), false, -1.f, Instruction.Thickness, SDPG_Foreground);
+					// TODO: Add support for this if anyone is actually using it. Currently it is only defined and referenced in an unused API, DrawCone in Control Rig.
+					break;
 				}
-				break;
 			}
-
-			case EControlRigDrawSettings::DynamicMesh:
-			{
-				// TODO: Add support for this if anyone is actually using it. Currently it is only defined and referenced in an unused API, DrawCone in Control Rig.
-				break;
-			} 
 		}
-	} 
+	}
 }
 
