ControlRigComponent: Fix debug drawing not working, now it can be toggled on/off through Control Rig Component Details Panel -> Rendering -> Show Debug Drawing

diff --git a/Engine/Plugins/Experimental/ControlRig/Source/ControlRig/Private/ControlRigComponent.cpp b/Engine/Plugins/Experimental/ControlRig/Source/ControlRig/Private/ControlRigComponent.cpp
--- a/Engine/Plugins/Experimental/ControlRig/Source/ControlRig/Private/ControlRigComponent.cpp
+++ b/Engine/Plugins/Experimental/ControlRig/Source/ControlRig/Private/ControlRigComponent.cpp
@@ -39,6 +39,7 @@ UControlRigComponent::UControlRigComponent(const FObjectInitializer& ObjectIniti
 	bUpdateRigOnTick = true;
 	bUpdateInEditor = true;
 	bDrawBones = true;
+	bShowDebugDrawing = false;
 }
 
 #if WITH_EDITOR
@@ -106,21 +107,28 @@ FBoxSphereBounds UControlRigComponent::CalcBounds(const FTransform& LocalToWorld
 {
 	FBox BBox(ForceInit);
 
-	for (int32 InstructionIndex = 0; InstructionIndex < DrawInterface.Num(); InstructionIndex++)
+	if (ControlRig)
 	{
-		const FControlRigDrawInstruction& Instruction = DrawInterface[InstructionIndex];
+		// Get bounding box for the debug drawings if they are drawn 
+		if (bShowDebugDrawing)
+		{ 
+			const FControlRigDrawInterface& DrawInterface = ControlRig->GetDrawInterface();
 
-		FTransform Transform = Instruction.Transform * GetComponentToWorld();
-		for (const FVector& Position : Instruction.Positions)
-		{
-			BBox += Transform.TransformPosition(Position);
+			for (int32 InstructionIndex = 0; InstructionIndex < DrawInterface.Num(); InstructionIndex++)
+			{
+				const FControlRigDrawInstruction& Instruction = DrawInterface[InstructionIndex];
+
+				FTransform Transform = Instruction.Transform * GetComponentToWorld();
+				for (const FVector& Position : Instruction.Positions)
+				{
+					BBox += Transform.TransformPosition(Position);
+				}
+			}
 		}
-	}
 
-	if (ControlRig)
-	{
 		FTransform Transform = GetComponentToWorld();
 
+		// Get bounding box for bones
 		const FRigBoneHierarchy& BoneHierarchy = ControlRig->GetBoneHierarchy();
 		for (const FRigBone& Bone : BoneHierarchy)
 		{
@@ -202,8 +210,6 @@ void UControlRigComponent::Initialize()
 
 void UControlRigComponent::Update(float DeltaTime)
 {
-	DrawInterface.Reset();
-
 	if(UControlRig* CR = SetupControlRigIfRequired())
 	{
 		if (CR->IsExecuting() || CR->IsInitializing())
@@ -236,15 +242,16 @@ void UControlRigComponent::Update(float DeltaTime)
 			}
 
 			CR->Evaluate_AnyThread();
+		} 
 
-			CR->DrawInterface.Reset();
+		if (bShowDebugDrawing)
+		{
+			if (CR->DrawInterface.Instructions.Num() > 0)
+			{ 
+				MarkRenderStateDirty();
+			}
 		}
-	}
-
-	if (DrawInterface.Instructions.Num() > 0)
-	{
-		MarkRenderStateDirty();
-	}
+	} 
 }
 
 TArray<FName> UControlRigComponent::GetElementNames(ERigElementType ElementType)
@@ -1621,7 +1628,6 @@ void UControlRigComponent::ReportError(const FString& InMessage)
 FControlRigSceneProxy::FControlRigSceneProxy(const UControlRigComponent* InComponent)
 : FPrimitiveSceneProxy(InComponent)
 , ControlRigComponent(InComponent)
-, DrawInterface(InComponent->DrawInterface)
 {
 	bWillEverBeLit = false;
 }
@@ -1681,44 +1687,49 @@ void FControlRigSceneProxy::GetDynamicMeshElements(const TArray<const FSceneView
 				}
 			}
 
-			for (int32 InstructionIndex = 0; InstructionIndex < DrawInterface.Num(); InstructionIndex++)
-			{
-				const FControlRigDrawInstruction& Instruction = DrawInterface[InstructionIndex];
-				if (Instruction.Positions.Num() == 0)
-				{
-					continue;
-				}
+			if (ControlRigComponent->bShowDebugDrawing)
+			{ 
+				const FControlRigDrawInterface& DrawInterface = ControlRigComponent->ControlRig->GetDrawInterface();
 
-				FTransform InstructionTransform = Instruction.Transform * ControlRigComponent->GetComponentToWorld();
-				switch (Instruction.PrimitiveType)
+				for (int32 InstructionIndex = 0; InstructionIndex < DrawInterface.Num(); InstructionIndex++)
 				{
-					case EControlRigDrawSettings::Points:
+					const FControlRigDrawInstruction& Instruction = DrawInterface[InstructionIndex];
+					if (Instruction.Positions.Num() == 0)
 					{
-						for (const FVector& Point : Instruction.Positions)
-						{
-							PDI->DrawPoint(InstructionTransform.TransformPosition(Point), Instruction.Color, Instruction.Thickness, SDPG_Foreground);
-						}
-						break;
+						continue;
 					}
-					case EControlRigDrawSettings::Lines:
+
+					FTransform InstructionTransform = Instruction.Transform * ControlRigComponent->GetComponentToWorld();
+					switch (Instruction.PrimitiveType)
 					{
-						const TArray<FVector>& Points = Instruction.Positions;
-						PDI->AddReserveLines(SDPG_Foreground, Points.Num() / 2, false, Instruction.Thickness > SMALL_NUMBER);
-						for (int32 PointIndex = 0; PointIndex < Points.Num() - 1; PointIndex += 2)
+						case EControlRigDrawSettings::Points:
 						{
-							PDI->DrawLine(InstructionTransform.TransformPosition(Points[PointIndex]), InstructionTransform.TransformPosition(Points[PointIndex + 1]), Instruction.Color, SDPG_Foreground, Instruction.Thickness);
+							for (const FVector& Point : Instruction.Positions)
+							{
+								PDI->DrawPoint(InstructionTransform.TransformPosition(Point), Instruction.Color, Instruction.Thickness, SDPG_Foreground);
+							}
+							break;
 						}
-						break;
-					}
-					case EControlRigDrawSettings::LineStrip:
-					{
-						const TArray<FVector>& Points = Instruction.Positions;
-						PDI->AddReserveLines(SDPG_Foreground, Points.Num() - 1, false, Instruction.Thickness > SMALL_NUMBER);
-						for (int32 PointIndex = 0; PointIndex < Points.Num() - 1; PointIndex++)
+						case EControlRigDrawSettings::Lines:
+						{
+							const TArray<FVector>& Points = Instruction.Positions;
+							PDI->AddReserveLines(SDPG_Foreground, Points.Num() / 2, false, Instruction.Thickness > SMALL_NUMBER);
+							for (int32 PointIndex = 0; PointIndex < Points.Num() - 1; PointIndex += 2)
+							{
+								PDI->DrawLine(InstructionTransform.TransformPosition(Points[PointIndex]), InstructionTransform.TransformPosition(Points[PointIndex + 1]), Instruction.Color, SDPG_Foreground, Instruction.Thickness);
+							}
+							break;
+						}
+						case EControlRigDrawSettings::LineStrip:
 						{
-							PDI->DrawLine(InstructionTransform.TransformPosition(Points[PointIndex]), InstructionTransform.TransformPosition(Points[PointIndex + 1]), Instruction.Color, SDPG_Foreground, Instruction.Thickness);
+							const TArray<FVector>& Points = Instruction.Positions;
+							PDI->AddReserveLines(SDPG_Foreground, Points.Num() - 1, false, Instruction.Thickness > SMALL_NUMBER);
+							for (int32 PointIndex = 0; PointIndex < Points.Num() - 1; PointIndex++)
+							{
+								PDI->DrawLine(InstructionTransform.TransformPosition(Points[PointIndex]), InstructionTransform.TransformPosition(Points[PointIndex + 1]), Instruction.Color, SDPG_Foreground, Instruction.Thickness);
+							}
+							break;
 						}
-						break;
 					}
 				}
 			}
@@ -1748,5 +1759,5 @@ uint32 FControlRigSceneProxy::GetMemoryFootprint(void) const
 
 uint32 FControlRigSceneProxy::GetAllocatedSize(void) const
 {
-	return FPrimitiveSceneProxy::GetAllocatedSize() + DrawInterface.GetAllocatedSize();
+	return FPrimitiveSceneProxy::GetAllocatedSize();
 }
diff --git a/Engine/Plugins/Experimental/ControlRig/Source/ControlRig/Public/ControlRigComponent.h b/Engine/Plugins/Experimental/ControlRig/Source/ControlRig/Public/ControlRigComponent.h
--- a/Engine/Plugins/Experimental/ControlRig/Source/ControlRig/Public/ControlRigComponent.h
+++ b/Engine/Plugins/Experimental/ControlRig/Source/ControlRig/Public/ControlRigComponent.h
@@ -295,6 +295,10 @@ public:
 	UPROPERTY(EditAnywhere, Category = "Rendering")
 	bool bDrawBones;
 
+	/** When checked the rig's debug drawing instructions are drawn in the viewport */
+	UPROPERTY(EditAnywhere, Category = "Rendering")
+	bool bShowDebugDrawing;
+
 	/**
 	 * Returns all of the names for a given element type (Bone, Control, etc)
 	 * @param ElementType The type of elements to return the names for
@@ -608,8 +612,6 @@ private:
 	UPROPERTY(transient)
 	UControlRig* ControlRig;
 
-	FControlRigDrawInterface DrawInterface;
-
 	TMap<USkeletalMeshComponent*, FCachedSkeletalMeshComponentSettings> CachedSkeletalMeshComponentSettings;
 
 	FControlRigComponentEvent ControlRigCreatedEvent;
@@ -638,5 +640,4 @@ public:
 private:
 
 	const UControlRigComponent* ControlRigComponent;
-	const FControlRigDrawInterface& DrawInterface;
 };
\ No newline at end of file
