Control Rig: Fixed Comment resize registered as two transaction, resulting in user having to undo twice

diff --git a/Engine/Plugins/Experimental/ControlRig/Source/ControlRigDeveloper/Private/Graph/ControlRigGraph.cpp b/Engine/Plugins/Experimental/ControlRig/Source/ControlRigDeveloper/Private/Graph/ControlRigGraph.cpp
--- a/Engine/Plugins/Experimental/ControlRig/Source/ControlRigDeveloper/Private/Graph/ControlRigGraph.cpp
+++ b/Engine/Plugins/Experimental/ControlRig/Source/ControlRigDeveloper/Private/Graph/ControlRigGraph.cpp
@@ -304,7 +304,8 @@ void UControlRigGraph::HandleModifiedEvent(ERigVMGraphNotifType InNotifType, URi
 			{
 				UEdGraphNode* EdNode = FindNodeForModelNodeName(ModelNode->GetFName());
 				if (EdNode)
-					{
+				{
+					// No need to call Node->Modify(), since control rig has its own undo/redo system see RigVMControllerActions.cpp
 					EdNode->NodePosX = (int32)ModelNode->GetPosition().X;
 					EdNode->NodePosY = (int32)ModelNode->GetPosition().Y;
 				}
@@ -318,6 +319,7 @@ void UControlRigGraph::HandleModifiedEvent(ERigVMGraphNotifType InNotifType, URi
 				UEdGraphNode_Comment* EdNode = Cast<UEdGraphNode_Comment>(FindNodeForModelNodeName(ModelNode->GetFName()));
 				if (EdNode)
 				{
+					// No need to call Node->Modify(), since control rig has its own undo/redo system see RigVMControllerActions.cpp
 					EdNode->NodeWidth = (int32)ModelNode->GetSize().X;
 					EdNode->NodeHeight = (int32)ModelNode->GetSize().Y;
 				}
diff --git a/Engine/Plugins/Experimental/ControlRig/Source/ControlRigEditor/Private/Graph/SControlRigGraphNodeComment.cpp b/Engine/Plugins/Experimental/ControlRig/Source/ControlRigEditor/Private/Graph/SControlRigGraphNodeComment.cpp
--- a/Engine/Plugins/Experimental/ControlRig/Source/ControlRigEditor/Private/Graph/SControlRigGraphNodeComment.cpp
+++ b/Engine/Plugins/Experimental/ControlRig/Source/ControlRigEditor/Private/Graph/SControlRigGraphNodeComment.cpp
@@ -23,9 +23,16 @@ SControlRigGraphNodeComment::SControlRigGraphNodeComment()
 
 FReply SControlRigGraphNodeComment::OnMouseButtonUp(const FGeometry& MyGeometry, const FPointerEvent& MouseEvent)
 {
-	FReply Reply = SGraphNodeComment::OnMouseButtonUp(MyGeometry, MouseEvent);
-	if (Reply.IsEventHandled())
-	{
+	if ((MouseEvent.GetEffectingButton() == EKeys::LeftMouseButton) && bUserIsDragging)
+	{ 
+		// bUserIsDragging is set to false in the base function
+
+		// Resize the node	
+		UserSize.X = FMath::RoundToFloat(UserSize.X);
+		UserSize.Y = FMath::RoundToFloat(UserSize.Y);
+
+		GetNodeObj()->ResizeNode(UserSize);
+
 		if (GraphNode)
 		{
 			UEdGraphNode_Comment* CommentNode = CastChecked<UEdGraphNode_Comment>(GraphNode);
@@ -41,9 +48,12 @@ FReply SControlRigGraphNodeComment::OnMouseButtonUp(const FGeometry& MyGeometry,
 					Blueprint->Controller->CloseUndoBracket();
 				}
 			}
-		}
-	}
-	return Reply;
+		} 
+	} 
+	
+	// Calling the base function at the end such that above actions are included in the undo system transaction scope.
+	// When undo is triggered, FBlueprintEditor::HandleUndoTransaction() will be called to make sure the undone changes are reflected in UI.
+	return SGraphNodeComment::OnMouseButtonUp(MyGeometry, MouseEvent);
 }
 
 void SControlRigGraphNodeComment::EndUserInteraction() const
