Control Rig: Offering Collapse and Expand contextual menu entries in Control Rig Editor

diff --git a/Engine/Plugins/Experimental/ControlRig/Source/ControlRigEditor/Private/ControlRigEditorModule.cpp b/Engine/Plugins/Experimental/ControlRig/Source/ControlRigEditor/Private/ControlRigEditorModule.cpp
--- a/Engine/Plugins/Experimental/ControlRig/Source/ControlRigEditor/Private/ControlRigEditorModule.cpp
+++ b/Engine/Plugins/Experimental/ControlRig/Source/ControlRigEditor/Private/ControlRigEditorModule.cpp
@@ -78,6 +78,7 @@
 #include "ActorFactories/ActorFactorySkeletalMesh.h"
 #include "ControlRigThumbnailRenderer.h"
 #include "RigVMModel/RigVMController.h"
+#include "RigVMModel/Nodes/RigVMLibraryNode.h"
 #include "ControlRigBlueprint.h"
 #include "ControlRig/Private/Units/Simulation/RigUnit_AlphaInterp.h"
 #include "ControlRig/Private/Units/Debug/RigUnit_VisualDebug.h"
@@ -1030,47 +1031,6 @@ FConnectionDrawingPolicy* FControlRigEditorModule::CreateConnectionDrawingPolicy
 	return new FControlRigConnectionDrawingPolicy(InBackLayerID, InFrontLayerID, InZoomFactor, InClippingRect, InDrawElements, InGraphObj);
 }
 
-void FControlRigEditorModule::GetNodeContextMenuActions(const UControlRigGraphNode* Node, UToolMenu* Menu, UGraphNodeContextMenuContext* Context) const
-{
-	{
-		if (Context->Pin != nullptr)
-		{
-			// Add array operations for array pins
-			if (Context->Pin->PinType.IsArray())
-			{
-				FToolMenuSection& Section = Menu->AddSection(TEXT("ArrayOperations"), LOCTEXT("ArrayOperations", "Array Operations"));
-
-				// Array operations
-				Section.AddMenuEntry(
-					"ClearArray",
-					LOCTEXT("ClearArray", "Clear"),
-					LOCTEXT("ClearArray_Tooltip", "Clear this array of all of its entries"),
-					FSlateIcon(),
-					FUIAction(FExecuteAction::CreateUObject(const_cast<UControlRigGraphNode*>(Node), &UControlRigGraphNode::HandleClearArray, Context->Pin->PinName.ToString())));
-			}
-			else if (Context->Pin->ParentPin != nullptr && Context->Pin->ParentPin->PinType.IsArray())
-			{
-				FToolMenuSection& Section = Menu->AddSection(TEXT("ArrayElementOperations"), LOCTEXT("ArrayElementOperations", "Array Element Operations"));
-
-				// Array element operations
-				Section.AddMenuEntry(
-					"RemoveArrayElement",
-					LOCTEXT("RemoveArrayElement", "Remove"),
-					LOCTEXT("RemoveArrayElement_Tooltip", "Remove this array element"),
-					FSlateIcon(),
-					FUIAction(FExecuteAction::CreateUObject(const_cast<UControlRigGraphNode*>(Node), &UControlRigGraphNode::HandleRemoveArrayElement, Context->Pin->PinName.ToString())));
-
-				Section.AddMenuEntry(
-					"InsertArrayElement",
-					LOCTEXT("InsertArrayElement", "Insert"),
-					LOCTEXT("InsertArrayElement_Tooltip", "Insert an array element after this one"),
-					FSlateIcon(),
-					FUIAction(FExecuteAction::CreateUObject(const_cast<UControlRigGraphNode*>(Node), &UControlRigGraphNode::HandleInsertArrayElement, Context->Pin->PinName.ToString())));
-			}
-		}
-	}
-}
-
 void FControlRigEditorModule::GetContextMenuActions(const UControlRigGraphSchema* Schema, UToolMenu* Menu, UGraphNodeContextMenuContext* Context) const
 {
 	if (Menu && Context)
@@ -1663,14 +1623,39 @@ void FControlRigEditorModule::GetContextMenuActions(const UControlRigGraphSchema
 						FToolMenuSection& VariablesSection = Menu->AddSection("EdGraphSchemaVariables", LOCTEXT("VariablesSettingsHeader", "Variables"));
 						VariablesSection.AddMenuEntry(
 							"MakePindingsFromVariableNode",
-							LOCTEXT("MakePindingsFromVariableNode", "Make Bindings From Node"),
-							LOCTEXT("MakePindingsFromVariableNode_Tooltip", "Turns the variable node into one ore more variable bindings on the pin(s)"),
+							LOCTEXT("MakeBindingsFromVariableNode", "Make Bindings From Node"),
+							LOCTEXT("MakeBindingsFromVariableNode_Tooltip", "Turns the variable node into one ore more variable bindings on the pin(s)"),
 							FSlateIcon(),
 							FUIAction(FExecuteAction::CreateLambda([RigBlueprint, VariableNode]() {
 								RigBlueprint->Controller->MakeBindingsFromVariableNode(VariableNode->GetFName());
 							})
 						));
 					}
+
+					FToolMenuSection& OrganizationSection = Menu->AddSection("EdGraphSchemaOrganization", LOCTEXT("OrganizationHeader", "Organization"));
+					OrganizationSection.AddMenuEntry(
+						"Collapse Nodes",
+						LOCTEXT("CollapseNodes", "Collapse Nodes"),
+						LOCTEXT("CollapseNodes_Tooltip", "Turns the selected nodes into a single Collapse node"),
+						FSlateIcon(),
+						FUIAction(FExecuteAction::CreateLambda([RigBlueprint]() {
+							TArray<FName> Nodes = RigBlueprint->Model->GetSelectNodes();
+							RigBlueprint->Controller->CollapseNodes(Nodes);
+						})
+					));
+
+					if (URigVMLibraryNode* LibraryNode = Cast<URigVMLibraryNode>(RigNode->GetModelNode()))
+					{
+						OrganizationSection.AddMenuEntry(
+							"Expand Node",
+							LOCTEXT("ExpandNode", "Expand Node"),
+							LOCTEXT("ExpandNode_Tooltip", "Expands the contents of the node into this graph"),
+							FSlateIcon(),
+							FUIAction(FExecuteAction::CreateLambda([RigBlueprint, LibraryNode]() {
+								RigBlueprint->Controller->ExpandLibraryNode(LibraryNode->GetFName());
+							})
+						));
+					}
 				}
 			}
 		}
diff --git a/Engine/Plugins/Experimental/ControlRig/Source/ControlRigEditor/Private/ControlRigEditorModule.h b/Engine/Plugins/Experimental/ControlRig/Source/ControlRigEditor/Private/ControlRigEditorModule.h
--- a/Engine/Plugins/Experimental/ControlRig/Source/ControlRigEditor/Private/ControlRigEditorModule.h
+++ b/Engine/Plugins/Experimental/ControlRig/Source/ControlRigEditor/Private/ControlRigEditorModule.h
@@ -54,7 +54,6 @@ public:
 	virtual void GetTypeActions(UControlRigBlueprint* CRB, FBlueprintActionDatabaseRegistrar& ActionRegistrar) override;
 	virtual void GetInstanceActions(UControlRigBlueprint* CRB, FBlueprintActionDatabaseRegistrar& ActionRegistrar) override;
 	virtual FConnectionDrawingPolicy* CreateConnectionDrawingPolicy(int32 InBackLayerID, int32 InFrontLayerID, float InZoomFactor, const FSlateRect& InClippingRect, class FSlateWindowElementList& InDrawElements, class UEdGraph* InGraphObj) override;
-	virtual void GetNodeContextMenuActions(const UControlRigGraphNode* Node, class UToolMenu* Menu, class UGraphNodeContextMenuContext* Context) const override;
 	virtual void GetContextMenuActions(const UControlRigGraphSchema* Schema, class UToolMenu* Menu, class UGraphNodeContextMenuContext* Context) const override;
 
 private:
diff --git a/Engine/Plugins/Experimental/ControlRig/Source/ControlRigEditor/Public/IControlRigEditorModule.h b/Engine/Plugins/Experimental/ControlRig/Source/ControlRigEditor/Public/IControlRigEditorModule.h
--- a/Engine/Plugins/Experimental/ControlRig/Source/ControlRigEditor/Public/IControlRigEditorModule.h
+++ b/Engine/Plugins/Experimental/ControlRig/Source/ControlRigEditor/Public/IControlRigEditorModule.h
@@ -43,6 +43,5 @@ public:
 	virtual void GetTypeActions(UControlRigBlueprint* CRB, FBlueprintActionDatabaseRegistrar& ActionRegistrar) = 0;
 	virtual void GetInstanceActions(UControlRigBlueprint* CRB, FBlueprintActionDatabaseRegistrar& ActionRegistrar) = 0;
 	virtual FConnectionDrawingPolicy* CreateConnectionDrawingPolicy(int32 InBackLayerID, int32 InFrontLayerID, float InZoomFactor, const FSlateRect& InClippingRect, class FSlateWindowElementList& InDrawElements, class UEdGraph* InGraphObj) = 0;
-	virtual void GetNodeContextMenuActions(const UControlRigGraphNode* Node, class UToolMenu* Menu, class UGraphNodeContextMenuContext* Context) const = 0;
 	virtual void GetContextMenuActions(const UControlRigGraphSchema* Schema, class UToolMenu* Menu, class UGraphNodeContextMenuContext* Context) const = 0;
 };
