Control Rig: Exposed a few more symbols: SetTransform, FCachedRigElement

diff --git a/Engine/Plugins/Experimental/ControlRig/Source/ControlRig/Public/Rigs/RigHierarchyCache.h b/Engine/Plugins/Experimental/ControlRig/Source/ControlRig/Public/Rigs/RigHierarchyCache.h
--- a/Engine/Plugins/Experimental/ControlRig/Source/ControlRig/Public/Rigs/RigHierarchyCache.h
+++ b/Engine/Plugins/Experimental/ControlRig/Source/ControlRig/Public/Rigs/RigHierarchyCache.h
@@ -13,7 +13,7 @@ struct FRigControlHierarchy;
 struct FRigCurveContainer;
 
 USTRUCT(BlueprintType)
-struct FCachedRigElement
+struct CONTROLRIG_API FCachedRigElement
 {
 	GENERATED_BODY()
 
diff --git a/Engine/Plugins/Experimental/ControlRig/Source/ControlRig/Public/Units/Hierarchy/RigUnit_SetTransform.h b/Engine/Plugins/Experimental/ControlRig/Source/ControlRig/Public/Units/Hierarchy/RigUnit_SetTransform.h
--- a/Engine/Plugins/Experimental/ControlRig/Source/ControlRig/Public/Units/Hierarchy/RigUnit_SetTransform.h
+++ b/Engine/Plugins/Experimental/ControlRig/Source/ControlRig/Public/Units/Hierarchy/RigUnit_SetTransform.h
@@ -82,7 +82,7 @@ struct CONTROLRIG_API FRigUnit_SetTransform : public FRigUnitMutable
  * SetTranslation is used to set a single translation on hierarchy.
  */
 USTRUCT(meta=(DisplayName="Set Translation", Category="Hierarchy", DocumentationPolicy = "Strict", Keywords="SetBoneTranslation,SetControlTranslation,SetInitialTranslation,SetSpaceTranslation,SetBoneLocation,SetControlLocation,SetInitialLocation,SetSpaceLocation,SetBonePosition,SetControlPosition,SetInitialPosition,SetSpacePosition,SetTranslation,SetLocation,SetPosition", Varying))
-struct FRigUnit_SetTranslation : public FRigUnitMutable
+struct CONTROLRIG_API FRigUnit_SetTranslation : public FRigUnitMutable
 {
 	GENERATED_BODY()
 
@@ -146,7 +146,7 @@ struct FRigUnit_SetTranslation : public FRigUnitMutable
  * SetRotation is used to set a single rotation on hierarchy.
  */
 USTRUCT(meta=(DisplayName="Set Rotation", Category="Hierarchy", DocumentationPolicy = "Strict", Keywords="SetBoneRotation,SetControlRotation,SetInitialRotation,SetSpaceRotation,SetBoneOrientation,SetControlOrientation,SetInitialOrientation,SetSpaceOrientation,SetRotation,SetOrientation", Varying))
-struct FRigUnit_SetRotation : public FRigUnitMutable
+struct CONTROLRIG_API FRigUnit_SetRotation : public FRigUnitMutable
 {
 	GENERATED_BODY()
 
