Control Rig: Refactor Blueprint Controller / Model accessor to getters

diff --git a/Engine/Plugins/Experimental/ControlRig/Source/ControlRig/Private/Tools/ControlRigPose.cpp b/Engine/Plugins/Experimental/ControlRig/Source/ControlRig/Private/Tools/ControlRigPose.cpp
--- a/Engine/Plugins/Experimental/ControlRig/Source/ControlRig/Private/Tools/ControlRigPose.cpp
+++ b/Engine/Plugins/Experimental/ControlRig/Source/ControlRig/Private/Tools/ControlRigPose.cpp
@@ -2,8 +2,6 @@
 #include "Tools/ControlRigPose.h"
 #include "Tools/ControlRigPoseProjectSettings.h"
 #include "IControlRigObjectBinding.h"
-#include "Framework/Notifications/NotificationManager.h"
-#include "Widgets/Notifications/SNotificationList.h"
 #if WITH_EDITOR
 #include "ScopedTransaction.h"
 #endif
diff --git a/Engine/Plugins/Experimental/ControlRig/Source/ControlRigDeveloper/Private/ControlRigBlueprint.cpp b/Engine/Plugins/Experimental/ControlRig/Source/ControlRigDeveloper/Private/ControlRigBlueprint.cpp
--- a/Engine/Plugins/Experimental/ControlRig/Source/ControlRigDeveloper/Private/ControlRigBlueprint.cpp
+++ b/Engine/Plugins/Experimental/ControlRig/Source/ControlRigDeveloper/Private/ControlRigBlueprint.cpp
@@ -52,7 +52,6 @@ UControlRigBlueprint::UControlRigBlueprint(const FObjectInitializer& ObjectIniti
 	VMRecompilationBracket = 0;
 
 	Model = ObjectInitializer.CreateDefaultSubobject<URigVMGraph>(this, TEXT("RigVMModel"));
-	Controller = nullptr;
 
 	Validator = ObjectInitializer.CreateDefaultSubobject<UControlRigValidator>(this, TEXT("ControlRigValidator"));
 
@@ -66,67 +65,9 @@ void UControlRigBlueprint::InitializeModelIfRequired(bool bRecompileVM)
 {
 	DECLARE_SCOPE_HIERARCHICAL_COUNTER_FUNC()
 
-	if (Controller == nullptr)
+	if (Controllers.Num() == 0)
 	{
-		Controller = NewObject<URigVMController>(this);
-		Controller->SetExecuteContextStruct(FControlRigExecuteContext::StaticStruct());
-		Controller->SetGraph(Model);
-		Controller->OnModified().AddUObject(this, &UControlRigBlueprint::HandleModifiedEvent);
-
-		Controller->UnfoldStructDelegate.BindLambda([](const UStruct* InStruct) -> bool {
-
-			if (InStruct == TBaseStructure<FQuat>::Get())
-			{
-				return false;
-			}
-			if (InStruct == FRuntimeFloatCurve::StaticStruct())
-			{
-				return false;
-			}
-			return true;
-		});
-
-		TWeakObjectPtr<UControlRigBlueprint> WeakThis(this);
-
-		// this delegate is used by the controller to determine variable validity
-		// during a bind process. the controller itself doesn't own the variables,
-		// so we need a delegate to request them from the owning blueprint
-		Controller->GetExternalVariablesDelegate.BindLambda([WeakThis]() -> TArray<FRigVMExternalVariable> {
-			
-			if (WeakThis.IsValid())
-			{
-				if (UControlRigBlueprintGeneratedClass* RigClass = WeakThis->GetControlRigBlueprintGeneratedClass())
-				{
-					if (UControlRig* CDO = Cast<UControlRig>(RigClass->GetDefaultObject(true /* create if needed */)))
-					{
-						return CDO->GetExternalVariablesImpl(true /* rely on variables within blueprint */);
-					}
-				}
-			}
-			return TArray<FRigVMExternalVariable>();
-
-		});
-
-#if WITH_EDITOR
-
-		// this sets up three delegates:
-		// a) get external variables (mapped to Controller->GetExternalVariables)
-		// b) bind pin to variable (mapped to Controller->BindPinToVariable)
-		// c) create external variable (mapped to the passed in tfunction)
-		// the last one is defined within the blueprint since the controller
-		// doesn't own the variables and can't create one itself.
-		Controller->SetupDefaultUnitNodeDelegates(TDelegate<FName(FRigVMExternalVariable)>::CreateLambda(
-			[WeakThis](FRigVMExternalVariable InVariableToCreate) -> FName {
-				if (WeakThis.IsValid())
-				{
-					return WeakThis->AddCRMemberVariableFromExternal(InVariableToCreate);
-				}
-				return NAME_None;
-			}
-		));
-#endif
-
-		Controller->RemoveStaleNodes();
+		GetOrCreateController(Model);
 
 		for (int32 i = 0; i < UbergraphPages.Num(); ++i)
 		{
@@ -279,16 +220,15 @@ void UControlRigBlueprint::PostLoad()
 
 	if (!IsInAsyncLoadingThread() || IsRunningCommandlet())
 	{
-		Controller->DetachLinksFromPinObjects();
+		GetOrCreateController()->DetachLinksFromPinObjects();
 		TArray<URigVMNode*> Nodes = Model->GetNodes();
 		for (URigVMNode* Node : Nodes)
 		{
-			Controller->RepopulatePinsOnNode(Node);
+			GetOrCreateController()->RepopulatePinsOnNode(Node);
 		}
 		SetupPinRedirectorsForBackwardsCompatibility();
 	}
-
-	Controller->ReattachLinksToPinObjects(true /* follow redirectors */);
+	GetOrCreateController()->ReattachLinksToPinObjects(true /* follow redirectors */);
 
 	RecompileVM();
 	RequestControlRigInit();
@@ -344,7 +284,7 @@ void UControlRigBlueprint::RecompileVM()
 
 		URigVMCompiler* Compiler = URigVMCompiler::StaticClass()->GetDefaultObject<URigVMCompiler>();
 		Compiler->Settings = VMCompileSettings;
-		Compiler->Compile(Model, Controller, CDO->VM, CDO->GetExternalVariablesImpl(false), UserData, &PinToOperandMap);
+		Compiler->Compile(Model, GetController(), CDO->VM, CDO->GetExternalVariablesImpl(false), UserData, &PinToOperandMap);
 
 		CDO->Execute(EControlRigState::Init, FRigUnit_BeginExecution::EventName); // need to clarify if we actually need this
 		Statistics = CDO->VM->GetStatistics();
@@ -421,6 +361,133 @@ void UControlRigBlueprint::RequestControlRigInit()
 	}
 }
 
+URigVMGraph* UControlRigBlueprint::GetModel(const UEdGraph* InEdGraph) const
+{
+	if (InEdGraph == nullptr)
+	{
+		return Model;
+	}
+
+	const UControlRigGraph* RigGraph = Cast< UControlRigGraph>(InEdGraph);
+	check(RigGraph);
+
+	if (RigGraph->GetOuter() == this)
+	{
+		return Model;
+	}
+
+	ensure(false);
+	return nullptr;
+}
+
+TArray<URigVMGraph*> UControlRigBlueprint::GetAllModels() const
+{
+	TArray<URigVMGraph*> Models;
+	Models.Add(GetModel());
+	Models.Append(GetModel()->GetContainedGraphs());
+	return Models;
+}
+
+URigVMController* UControlRigBlueprint::GetController(URigVMGraph* InGraph) const
+{
+	if (InGraph == nullptr)
+	{
+		InGraph = Model;
+	}
+
+	URigVMController* const* ControllerPtr = Controllers.Find(InGraph);
+	if (ControllerPtr)
+	{
+		return *ControllerPtr;
+	}
+	return nullptr;
+}
+
+URigVMController* UControlRigBlueprint::GetOrCreateController(URigVMGraph* InGraph)
+{
+	if (URigVMController* ExistingController = GetController(InGraph))
+	{
+		return ExistingController;
+	}
+
+	if (InGraph == nullptr)
+	{
+		InGraph = Model;
+	}
+
+	URigVMController* Controller = NewObject<URigVMController>(this);
+	Controller->SetExecuteContextStruct(FControlRigExecuteContext::StaticStruct());
+	Controller->SetGraph(InGraph);
+	Controller->OnModified().AddUObject(this, &UControlRigBlueprint::HandleModifiedEvent);
+
+	Controller->UnfoldStructDelegate.BindLambda([](const UStruct* InStruct) -> bool {
+
+		if (InStruct == TBaseStructure<FQuat>::Get())
+		{
+			return false;
+		}
+		if (InStruct == FRuntimeFloatCurve::StaticStruct())
+		{
+			return false;
+		}
+		return true;
+		});
+
+	TWeakObjectPtr<UControlRigBlueprint> WeakThis(this);
+
+	// this delegate is used by the controller to determine variable validity
+	// during a bind process. the controller itself doesn't own the variables,
+	// so we need a delegate to request them from the owning blueprint
+	Controller->GetExternalVariablesDelegate.BindLambda([WeakThis]() -> TArray<FRigVMExternalVariable> {
+
+		if (WeakThis.IsValid())
+		{
+			if (UControlRigBlueprintGeneratedClass* RigClass = WeakThis->GetControlRigBlueprintGeneratedClass())
+			{
+				if (UControlRig* CDO = Cast<UControlRig>(RigClass->GetDefaultObject(true /* create if needed */)))
+				{
+					return CDO->GetExternalVariablesImpl(true /* rely on variables within blueprint */);
+				}
+			}
+		}
+		return TArray<FRigVMExternalVariable>();
+
+	});
+
+#if WITH_EDITOR
+
+	// this sets up three delegates:
+	// a) get external variables (mapped to Controller->GetExternalVariables)
+	// b) bind pin to variable (mapped to Controller->BindPinToVariable)
+	// c) create external variable (mapped to the passed in tfunction)
+	// the last one is defined within the blueprint since the controller
+	// doesn't own the variables and can't create one itself.
+	Controller->SetupDefaultUnitNodeDelegates(TDelegate<FName(FRigVMExternalVariable)>::CreateLambda(
+		[WeakThis](FRigVMExternalVariable InVariableToCreate) -> FName {
+			if (WeakThis.IsValid())
+			{
+				return WeakThis->AddCRMemberVariableFromExternal(InVariableToCreate);
+			}
+			return NAME_None;
+		}
+	));
+#endif
+
+	Controller->RemoveStaleNodes();
+	Controllers.Add(InGraph, Controller);
+	return Controller;
+}
+
+URigVMController* UControlRigBlueprint::GetController(const UEdGraph* InEdGraph) const
+{
+	return GetController(GetModel(InEdGraph));
+}
+
+URigVMController* UControlRigBlueprint::GetOrCreateController(const UEdGraph* InEdGraph)
+{
+	return GetOrCreateController(GetModel(InEdGraph));
+}
+
 void UControlRigBlueprint::GetTypeActions(FBlueprintActionDatabaseRegistrar& ActionRegistrar) const
 {
 	DECLARE_SCOPE_HIERARCHICAL_COUNTER_FUNC()
@@ -912,7 +979,7 @@ void UControlRigBlueprint::PopulateModelFromGraphForBackwardsCompatibility(UCont
 				UScriptStruct* UnitStruct = URigVMPin::FindObjectFromCPPTypeObjectPath<UScriptStruct>(StructPath);
 				if (UnitStruct && UnitStruct->IsChildOf(FRigVMStruct::StaticStruct()))
 				{ 
-					ModelNode = Controller->AddUnitNode(UnitStruct, TEXT("Execute"), NodePosition, PropertyName.ToString(), false);
+					ModelNode = GetOrCreateController()->AddUnitNode(UnitStruct, TEXT("Execute"), NodePosition, PropertyName.ToString(), false);
 				}
 				else if (PropertyName != NAME_None) // check if this is a variable
 				{
@@ -985,7 +1052,7 @@ void UControlRigBlueprint::PopulateModelFromGraphForBackwardsCompatibility(UCont
 							bIsInput = false;
 						}
 
-						ModelNode = Controller->AddParameterNode(PropertyName, DataType.ToString(), DataTypeObject, bIsInput, FString(), NodePosition, PropertyName.ToString(), false);
+						ModelNode = GetOrCreateController()->AddParameterNode(PropertyName, DataType.ToString(), DataTypeObject, bIsInput, FString(), NodePosition, PropertyName.ToString(), false);
 					}
 				}
 				else
@@ -995,44 +1062,44 @@ void UControlRigBlueprint::PopulateModelFromGraphForBackwardsCompatibility(UCont
 
 				if (ModelNode)
 				{
-					bool bWasReportingEnabled = Controller->IsReportingEnabled();
-					Controller->EnableReporting(false);
+					bool bWasReportingEnabled = GetOrCreateController()->IsReportingEnabled();
+					GetOrCreateController()->EnableReporting(false);
 
-				for (UEdGraphPin* Pin : RigNode->Pins)
-				{
-						FString PinPath = LocalHelpers::FixUpPinPath(Pin->GetName());
+					for (UEdGraphPin* Pin : RigNode->Pins)
+					{
+							FString PinPath = LocalHelpers::FixUpPinPath(Pin->GetName());
 
-						// check the material + mesh pins for deprecated control nodes
-						if (URigVMUnitNode* ModelUnitNode = Cast<URigVMUnitNode>(ModelNode))
-						{
-							if (ModelUnitNode->GetScriptStruct()->IsChildOf(FRigUnit_Control::StaticStruct()))
+							// check the material + mesh pins for deprecated control nodes
+							if (URigVMUnitNode* ModelUnitNode = Cast<URigVMUnitNode>(ModelNode))
 							{
-								if (Pin->GetName().EndsWith(TEXT(".StaticMesh")) || Pin->GetName().EndsWith(TEXT(".Materials")))
+								if (ModelUnitNode->GetScriptStruct()->IsChildOf(FRigUnit_Control::StaticStruct()))
 								{
-									continue;
+									if (Pin->GetName().EndsWith(TEXT(".StaticMesh")) || Pin->GetName().EndsWith(TEXT(".Materials")))
+									{
+										continue;
+									}
 								}
 							}
-						}
 
-					if (Pin->Direction == EGPD_Input && Pin->PinType.ContainerType == EPinContainerType::Array)
-					{
-						int32 ArraySize = Pin->SubPins.Num();
-							Controller->SetArrayPinSize(PinPath, ArraySize, FString(), false);
-					}
+						if (Pin->Direction == EGPD_Input && Pin->PinType.ContainerType == EPinContainerType::Array)
+						{
+							int32 ArraySize = Pin->SubPins.Num();
+							GetOrCreateController()->SetArrayPinSize(PinPath, ArraySize, FString(), false);
+						}
 
 						if (RigNode->ExpandedPins_DEPRECATED.Find(Pin->GetName()) != INDEX_NONE)
-					{
-							Controller->SetPinExpansion(PinPath, true, false);
-					}
+						{
+								GetOrCreateController()->SetPinExpansion(PinPath, true, false);
+						}
 
 						if (Pin->SubPins.Num() == 0 && !Pin->DefaultValue.IsEmpty() && Pin->Direction == EGPD_Input)
-					{
-							Controller->SetPinDefaultValue(PinPath, Pin->DefaultValue, false, false, false);
+						{
+								GetOrCreateController()->SetPinDefaultValue(PinPath, Pin->DefaultValue, false, false, false);
+						}
 					}
-				}
 
-					Controller->EnableReporting(bWasReportingEnabled);
-			}
+					GetOrCreateController()->EnableReporting(bWasReportingEnabled);
+				}
 
 				const int32 VarIndex = FBlueprintEditorUtils::FindNewVariableIndex(this, PropertyName);
 				if (VarIndex != INDEX_NONE)
@@ -1045,7 +1112,7 @@ void UControlRigBlueprint::PopulateModelFromGraphForBackwardsCompatibility(UCont
 			{
 				FVector2D NodePosition = FVector2D((float)CommentNode->NodePosX, (float)CommentNode->NodePosY);
 				FVector2D NodeSize = FVector2D((float)CommentNode->NodeWidth, (float)CommentNode->NodeHeight);
-				Controller->AddCommentNode(CommentNode->NodeComment, NodePosition, NodeSize, CommentNode->CommentColor, CommentNode->GetName(), false);
+				GetOrCreateController()->AddCommentNode(CommentNode->NodeComment, NodePosition, NodeSize, CommentNode->CommentColor, CommentNode->GetName(), false);
 			}
 		}
 
@@ -1069,7 +1136,7 @@ void UControlRigBlueprint::PopulateModelFromGraphForBackwardsCompatibility(UCont
 						{
 							FString SourcePinPath = LocalHelpers::FixUpPinPath(Pin->GetName());
 							FString TargetPinPath = LocalHelpers::FixUpPinPath(LinkedPin->GetName());
-							Controller->AddLink(SourcePinPath, TargetPinPath, false);
+							GetOrCreateController()->AddLink(SourcePinPath, TargetPinPath, false);
 						}
 					}
 				}
@@ -1091,7 +1158,7 @@ void UControlRigBlueprint::SetupPinRedirectorsForBackwardsCompatibility()
 			{
 				URigVMPin* TransformPin = UnitNode->FindPin(TEXT("Transform"));
 				URigVMPin* ResultPin = UnitNode->FindPin(TEXT("Result"));
-				Controller->AddPinRedirector(false, true, TransformPin->GetPinPath(), ResultPin->GetPinPath());
+				GetOrCreateController()->AddPinRedirector(false, true, TransformPin->GetPinPath(), ResultPin->GetPinPath());
 			}
 		}
 	}
@@ -1102,7 +1169,7 @@ void UControlRigBlueprint::RebuildGraphFromModel()
 	DECLARE_SCOPE_HIERARCHICAL_COUNTER_FUNC()
 
 	TGuardValue<bool> SelfGuard(bSuspendModelNotificationsForSelf, true);
-	check(Controller);
+	check(GetOrCreateController());
 
 	for (UEdGraph* Graph : UbergraphPages)
 	{
@@ -1113,12 +1180,12 @@ void UControlRigBlueprint::RebuildGraphFromModel()
 		}
 	}
 
-	Controller->ResendAllNotifications();
+	GetController()->ResendAllNotifications();
 }
 
 void UControlRigBlueprint::Notify(ERigVMGraphNotifType InNotifType, UObject* InSubject)
 {
-	Controller->Notify(InNotifType, InSubject);
+	GetOrCreateController()->Notify(InNotifType, InSubject);
 }
 
 void UControlRigBlueprint::HandleModifiedEvent(ERigVMGraphNotifType InNotifType, URigVMGraph* InGraph, UObject* InSubject)
@@ -1560,7 +1627,7 @@ void UControlRigBlueprint::PatchVariableNodesOnLoad()
 	{
 		TGuardValue<bool> GuardNotifsSelf(bSuspendModelNotificationsForSelf, true);
 
-		Controller->ReattachLinksToPinObjects();
+		GetOrCreateController()->ReattachLinksToPinObjects();
 
 		check(Model);
 
@@ -1577,7 +1644,7 @@ void UControlRigBlueprint::PatchVariableNodesOnLoad()
 
 				int32 VariableIndex = AddedMemberVariableMap.FindChecked(Description.Name);
 				FName VarName = NewVariables[VariableIndex].VarName;
-				Controller->RefreshVariableNode(VariableNode->GetFName(), VarName, Description.CPPType, Description.CPPTypeObject, false);
+				GetOrCreateController()->RefreshVariableNode(VariableNode->GetFName(), VarName, Description.CPPType, Description.CPPTypeObject, false);
 				bDirtyDuringLoad = true;
 			}
 
@@ -1591,7 +1658,7 @@ void UControlRigBlueprint::PatchVariableNodesOnLoad()
 
 				int32 VariableIndex = AddedMemberVariableMap.FindChecked(Description.Name);
 				FName VarName = NewVariables[VariableIndex].VarName;
-				Controller->ReplaceParameterNodeWithVariable(ParameterNode->GetFName(), VarName, Description.CPPType, Description.CPPTypeObject, false);
+				GetOrCreateController()->ReplaceParameterNodeWithVariable(ParameterNode->GetFName(), VarName, Description.CPPType, Description.CPPTypeObject, false);
 				bDirtyDuringLoad = true;
 			}
 		}
@@ -2045,7 +2112,7 @@ void UControlRigBlueprint::OnVariableAdded(const FName& InVarName)
 
 void UControlRigBlueprint::OnVariableRemoved(const FName& InVarName)
 {
-	if (Controller)
+	if (URigVMController* Controller = GetController())
 	{
 		Controller->OnExternalVariableRemoved(InVarName, true);
 	}
@@ -2054,7 +2121,7 @@ void UControlRigBlueprint::OnVariableRemoved(const FName& InVarName)
 
 void UControlRigBlueprint::OnVariableRenamed(const FName& InOldVarName, const FName& InNewVarName)
 {
-	if (Controller)
+	if (URigVMController* Controller = GetController())
 	{
 		Controller->OnExternalVariableRenamed(InOldVarName, InNewVarName, true);
 	}
@@ -2063,7 +2130,7 @@ void UControlRigBlueprint::OnVariableRenamed(const FName& InOldVarName, const FN
 
 void UControlRigBlueprint::OnVariableTypeChanged(const FName& InVarName, FEdGraphPinType InOldPinType, FEdGraphPinType InNewPinType)
 {
-	if (Controller)
+	if (URigVMController* Controller = GetController())
 	{
 		FRigVMExternalVariable NewVariable = UControlRig::GetExternalVariableFromPinType(InVarName, InNewPinType);
 		if (NewVariable.IsValid(true)) // allow nullptr
diff --git a/Engine/Plugins/Experimental/ControlRig/Source/ControlRigDeveloper/Private/ControlRigBlueprintUtils.cpp b/Engine/Plugins/Experimental/ControlRig/Source/ControlRigDeveloper/Private/ControlRigBlueprintUtils.cpp
--- a/Engine/Plugins/Experimental/ControlRig/Source/ControlRigDeveloper/Private/ControlRigBlueprintUtils.cpp
+++ b/Engine/Plugins/Experimental/ControlRig/Source/ControlRigDeveloper/Private/ControlRigBlueprintUtils.cpp
@@ -98,7 +98,7 @@ void FControlRigBlueprintUtils::HandleRefreshAllNodes(UBlueprint* InBlueprint)
 
 	if(UControlRigBlueprint* RigBlueprint = Cast<UControlRigBlueprint>(InBlueprint))
 	{
-		if (RigBlueprint->Model == nullptr)
+		if (RigBlueprint->GetModel() == nullptr)
 		{
 			return;
 		}
diff --git a/Engine/Plugins/Experimental/ControlRig/Source/ControlRigDeveloper/Private/Graph/ControlRigGraph.cpp b/Engine/Plugins/Experimental/ControlRig/Source/ControlRigDeveloper/Private/Graph/ControlRigGraph.cpp
--- a/Engine/Plugins/Experimental/ControlRig/Source/ControlRigDeveloper/Private/Graph/ControlRigGraph.cpp
+++ b/Engine/Plugins/Experimental/ControlRig/Source/ControlRigDeveloper/Private/Graph/ControlRigGraph.cpp
@@ -145,70 +145,70 @@ void UControlRigGraph::HandleModifiedEvent(ERigVMGraphNotifType InNotifType, URi
 	if (bSuspendModelNotifications)
 	{
 		return;
-			}
+	}
 
 	switch (InNotifType)
-		{
+	{
 		case ERigVMGraphNotifType::GraphChanged:
-			{
+		{
 			for (URigVMNode* Node : InGraph->GetNodes())
-				{
+			{
 				UEdGraphNode* EdNode = FindNodeForModelNodeName(Node->GetFName());
 				if (EdNode != nullptr)
-					{
+				{
 					RemoveNode(EdNode);
 				}
 			}
 			break;
 		}
 		case ERigVMGraphNotifType::NodeSelectionChanged:
-						{
+		{
 			if (bIsSelecting)
-							{
+			{
 				return;
-							}
+			}
 			TGuardValue<bool> SelectionGuard(bIsSelecting, true);
 
 			TSet<const UEdGraphNode*> NodeSelection;
 			for (FName NodeName : InGraph->GetSelectNodes())
-							{
+			{
 				if (UEdGraphNode* EdNode = FindNodeForModelNodeName(NodeName))
-							{
+				{
 					NodeSelection.Add(EdNode);
-							}
-						}
+				}
+			}
 			SelectNodeSet(NodeSelection);
 			break;
-					}
+		}
 		case ERigVMGraphNotifType::NodeAdded:
-			{
+		{
 			if (URigVMNode* ModelNode = Cast<URigVMNode>(InSubject))
-				{
+			{
 				if (!ModelNode->IsVisibleInUI())
 				{
 					if (URigVMInjectionInfo* Injection = ModelNode->GetInjectionInfo())
-				{
+					{
 						if (URigVMPin* ModelPin = Injection->GetPin())
-				{
+						{
 							URigVMNode* ParentModelNode = ModelPin->GetNode();
 							if (ParentModelNode)
-					{
+							{
 								UEdGraphNode* EdNode = FindNodeForModelNodeName(ParentModelNode->GetFName());
 								if (EdNode)
-					{
+									{
 									if (UControlRigGraphNode* RigNode = Cast<UControlRigGraphNode>(EdNode))
-					{
+									{
 										RigNode->ReconstructNode_Internal(true);
+									}
+								}
+							}
+						}
 					}
-				}
-			}
-		}
-			}
 					break;
-		}
+				}
 
 				if (URigVMCommentNode* CommentModelNode = Cast<URigVMCommentNode>(ModelNode))
-		{
+				{
 					UEdGraphNode_Comment* NewNode = NewObject<UEdGraphNode_Comment>(this, CommentModelNode->GetFName());
 					AddNode(NewNode, false);
 
@@ -224,7 +224,7 @@ void UControlRigGraph::HandleModifiedEvent(ERigVMGraphNotifType InNotifType, URi
 					NewNode->NodeComment = CommentModelNode->GetCommentText();
 					NewNode->SetFlags(RF_Transactional);
 					NewNode->GetNodesUnderComment();
-	}
+				}
 				else if (URigVMRerouteNode* RerouteModelNode = Cast<URigVMRerouteNode>(ModelNode))
 				{
 					UControlRigGraphNode* NewNode = NewObject<UControlRigGraphNode>(this, ModelNode->GetFName());
@@ -242,11 +242,11 @@ void UControlRigGraph::HandleModifiedEvent(ERigVMGraphNotifType InNotifType, URi
 					NewNode->AllocateDefaultPins();
 
 					if (UEdGraphPin* ValuePin = NewNode->FindPin(ModelNode->FindPin("Value")->GetPinPath()))
-	{
+					{
 						NewNode->SetColorFromModel(GetSchema()->GetPinTypeColor(ValuePin->PinType));
 					}
-	}
-				else // struct, parameter + variable
+				}
+				else // struct, library, parameter + variable
 				{
 					UControlRigGraphNode* NewNode = NewObject<UControlRigGraphNode>(this, ModelNode->GetFName());
 					AddNode(NewNode, false);
@@ -261,6 +261,14 @@ void UControlRigGraph::HandleModifiedEvent(ERigVMGraphNotifType InNotifType, URi
 					NewNode->SetColorFromModel(ModelNode->GetNodeColor());
 					NewNode->SetFlags(RF_Transactional);
 					NewNode->AllocateDefaultPins();
+
+					/*
+					// if we need to setup a sub graph.. we should
+					if (URigVMCollapseNode* CollapseNode = Cast<URigVMCollapseNode>(ModelNode))
+					{
+
+					}
+					*/
 				}
 			}
 			break;
@@ -275,29 +283,29 @@ void UControlRigGraph::HandleModifiedEvent(ERigVMGraphNotifType InNotifType, URi
 					{
 						URigVMNode* ParentModelNode = ModelPin->GetNode();
 						if (ParentModelNode)
-					{
+						{
 							UEdGraphNode* EdNode = FindNodeForModelNodeName(ParentModelNode->GetFName());
 							if (EdNode)
-						{
-							if (UControlRigGraphNode* RigNode = Cast<UControlRigGraphNode>(EdNode))
 							{
+								if (UControlRigGraphNode* RigNode = Cast<UControlRigGraphNode>(EdNode))
+								{
 									RigNode->ReconstructNode_Internal(true);
 								}
 							}
-							}
 						}
-						break;
 					}
+					break;
+				}
 
 				UEdGraphNode* EdNode = FindNodeForModelNodeName(ModelNode->GetFName());
 				if (EdNode)
-					{
+				{
 					RemoveNode(EdNode, true);
 					NotifyGraphChanged();
 				}
 			}
-						break;
-					}
+			break;
+		}
 		case ERigVMGraphNotifType::NodePositionChanged:
 		{
 			if (URigVMNode* ModelNode = Cast<URigVMNode>(InSubject))
@@ -389,23 +397,23 @@ void UControlRigGraph::HandleModifiedEvent(ERigVMGraphNotifType InNotifType, URi
 					UControlRigGraphNode* SourceRigNode = Cast<UControlRigGraphNode>(FindNodeForModelNodeName(SourcePin->GetNode()->GetFName()));
 					UControlRigGraphNode* TargetRigNode = Cast<UControlRigGraphNode>(FindNodeForModelNodeName(TargetPin->GetNode()->GetFName()));
 
-				if (SourceRigNode != nullptr && TargetRigNode != nullptr)
-				{
+					if (SourceRigNode != nullptr && TargetRigNode != nullptr)
+					{
 						FString SourcePinPath = SourcePin->GetPinPath();
 						FString TargetPinPath = TargetPin->GetPinPath();
-					UEdGraphPin* SourceRigPin = SourceRigNode->FindPin(*SourcePinPath, EGPD_Output);
-					UEdGraphPin* TargetRigPin = TargetRigNode->FindPin(*TargetPinPath, EGPD_Input);
+						UEdGraphPin* SourceRigPin = SourceRigNode->FindPin(*SourcePinPath, EGPD_Output);
+						UEdGraphPin* TargetRigPin = TargetRigNode->FindPin(*TargetPinPath, EGPD_Input);
 
-					if (SourceRigPin != nullptr && TargetRigPin != nullptr)
-					{
-						if (AddLink)
-						{
-							SourceRigPin->MakeLinkTo(TargetRigPin);
-						}
-						else
+						if (SourceRigPin != nullptr && TargetRigPin != nullptr)
 						{
-							SourceRigPin->BreakLinkTo(TargetRigPin);
-						}
+							if (AddLink)
+							{
+								SourceRigPin->MakeLinkTo(TargetRigPin);
+							}
+							else
+							{
+								SourceRigPin->BreakLinkTo(TargetRigPin);
+							}
 
 							for (int32 LinkedPinIndex = 0; LinkedPinIndex < SourceRigPin->LinkedTo.Num();)
 							{
@@ -416,8 +424,8 @@ void UControlRigGraph::HandleModifiedEvent(ERigVMGraphNotifType InNotifType, URi
 								else
 								{
 									SourceRigPin->LinkedTo.RemoveAtSwap(LinkedPinIndex);
-					}
-				}
+								}
+							}
 
 							for (int32 LinkedPinIndex = 0; LinkedPinIndex < TargetRigPin->LinkedTo.Num();)
 							{
@@ -472,7 +480,7 @@ void UControlRigGraph::HandleModifiedEvent(ERigVMGraphNotifType InNotifType, URi
 					}
 				}
 				else if (URigVMInjectionInfo* Injection = ModelPin->GetNode()->GetInjectionInfo())
-						{
+				{
 					if (Injection->InputPin != ModelPin->GetRootPin())
 					{
 						if (URigVMPin* InjectionPin = Injection->GetPin())
@@ -488,13 +496,13 @@ void UControlRigGraph::HandleModifiedEvent(ERigVMGraphNotifType InNotifType, URi
 										HostRigNode->ReconstructNode_Internal(true);
 									}
 								}
-							}
-									}
-								}
 							}
 						}
-			break;
 					}
+				}
+			}
+			break;
+		}
 		case ERigVMGraphNotifType::PinArraySizeChanged:
 		case ERigVMGraphNotifType::PinDirectionChanged:
 		case ERigVMGraphNotifType::PinTypeChanged:
@@ -578,6 +586,32 @@ UEdGraphNode* UControlRigGraph::FindNodeForModelNodeName(const FName& InModelNod
 	return nullptr;
 }
 
+UControlRigBlueprint* UControlRigGraph::GetBlueprint() const
+{
+	if (UControlRigGraph* OuterGraph = Cast<UControlRigGraph>(GetOuter()))
+	{
+		return OuterGraph->GetBlueprint();
+	}
+	return Cast<UControlRigBlueprint>(GetOuter());
+}
+
+URigVMGraph* UControlRigGraph::GetModel() const
+{
+	if (UControlRigBlueprint* Blueprint = GetBlueprint())
+	{
+		return Blueprint->GetModel(this);
+	}
+	return nullptr;
+}
+
+URigVMController* UControlRigGraph::GetController() const
+{
+	if (UControlRigBlueprint* Blueprint = GetBlueprint())
+	{
+		return Blueprint->GetOrCreateController(this);
+	}
+	return nullptr;
+}
 
 URigVMController* UControlRigGraph::GetTemplateController()
 {
diff --git a/Engine/Plugins/Experimental/ControlRig/Source/ControlRigDeveloper/Private/Graph/ControlRigGraphNode.cpp b/Engine/Plugins/Experimental/ControlRig/Source/ControlRigDeveloper/Private/Graph/ControlRigGraphNode.cpp
--- a/Engine/Plugins/Experimental/ControlRig/Source/ControlRigDeveloper/Private/Graph/ControlRigGraphNode.cpp
+++ b/Engine/Plugins/Experimental/ControlRig/Source/ControlRigDeveloper/Private/Graph/ControlRigGraphNode.cpp
@@ -226,10 +226,9 @@ void UControlRigGraphNode::HandleClearArray(FString InPinPath)
 {
 	DECLARE_SCOPE_HIERARCHICAL_COUNTER_FUNC()
 
-	UControlRigBlueprint* Blueprint = Cast<UControlRigBlueprint>(GetOuter()->GetOuter());
-	if (Blueprint)
+	if(URigVMController* Controller = GetController())
 	{
-		Blueprint->Controller->ClearArrayPin(InPinPath);
+		Controller->ClearArrayPin(InPinPath);
 	}
 }
 
@@ -237,14 +236,13 @@ void UControlRigGraphNode::HandleAddArrayElement(FString InPinPath)
 {
 	DECLARE_SCOPE_HIERARCHICAL_COUNTER_FUNC()
 
-	UControlRigBlueprint* Blueprint = Cast<UControlRigBlueprint>(GetOuter()->GetOuter());
-	if (Blueprint)
+	if (URigVMController* Controller = GetController())
 	{
-		Blueprint->Controller->OpenUndoBracket(TEXT("Add Array Pin"));
-		FString PinPath = Blueprint->Controller->AddArrayPin(InPinPath);
-		Blueprint->Controller->SetPinExpansion(InPinPath, true);
-		Blueprint->Controller->SetPinExpansion(PinPath, true);
-		Blueprint->Controller->CloseUndoBracket();
+		Controller->OpenUndoBracket(TEXT("Add Array Pin"));
+		FString PinPath = Controller->AddArrayPin(InPinPath);
+		Controller->SetPinExpansion(InPinPath, true);
+		Controller->SetPinExpansion(PinPath, true);
+		Controller->CloseUndoBracket();
 	}
 }
 
@@ -252,10 +250,9 @@ void UControlRigGraphNode::HandleRemoveArrayElement(FString InPinPath)
 {
 	DECLARE_SCOPE_HIERARCHICAL_COUNTER_FUNC()
 
-	UControlRigBlueprint* Blueprint = Cast<UControlRigBlueprint>(GetOuter()->GetOuter());
-	if (Blueprint)
+	if (URigVMController* Controller = GetController())
 	{
-		Blueprint->Controller->RemoveArrayPin(InPinPath);
+		Controller->RemoveArrayPin(InPinPath);
 	}
 }
 
@@ -263,18 +260,17 @@ void UControlRigGraphNode::HandleInsertArrayElement(FString InPinPath)
 {
 	DECLARE_SCOPE_HIERARCHICAL_COUNTER_FUNC()
 
-	UControlRigBlueprint* Blueprint = Cast<UControlRigBlueprint>(GetOuter()->GetOuter());
-	if (Blueprint)
+	if (URigVMController* Controller = GetController())
 	{
 		if (URigVMPin* ArrayElementPin = GetModelPinFromPinPath(InPinPath))
 			{
 			if (URigVMPin* ArrayPin = ArrayElementPin->GetParentPin())
 				{
-				Blueprint->Controller->OpenUndoBracket(TEXT("Add Array Pin"));
-				FString PinPath = Blueprint->Controller->InsertArrayPin(InPinPath, ArrayElementPin->GetPinIndex() + 1, FString());
-				Blueprint->Controller->SetPinExpansion(InPinPath, true);
-				Blueprint->Controller->SetPinExpansion(PinPath, true);
-				Blueprint->Controller->CloseUndoBracket();
+				Controller->OpenUndoBracket(TEXT("Add Array Pin"));
+				FString PinPath = Controller->InsertArrayPin(InPinPath, ArrayElementPin->GetPinIndex() + 1, FString());
+				Controller->SetPinExpansion(InPinPath, true);
+				Controller->SetPinExpansion(PinPath, true);
+				Controller->CloseUndoBracket();
 			}
 		}
 	}
@@ -598,8 +594,11 @@ void UControlRigGraphNode::CopyPinDefaultsToModel(UEdGraphPin* Pin, bool bUndo)
 	}
 
 		if (ModelPin->GetDefaultValue() != DefaultValue)
-				{
-			GetBlueprint()->Controller->SetPinDefaultValue(ModelPin->GetPinPath(), DefaultValue, false, true, false);
+		{
+			if (URigVMController* Controller = GetController())
+			{
+				Controller->SetPinDefaultValue(ModelPin->GetPinPath(), DefaultValue, false, true, false);
+			}
 		}
 	}
 }
@@ -608,11 +607,34 @@ UControlRigBlueprint* UControlRigGraphNode::GetBlueprint() const
 {
 	if(UControlRigGraph* Graph = Cast<UControlRigGraph>(GetOuter()))
 	{
+		return Graph->GetBlueprint();
+		if (UControlRigGraph* OuterGraph = Cast<UControlRigGraph>(Graph->GetOuter()))
+		{
+			return OuterGraph->GetBlueprint();
+		}
 		return Cast<UControlRigBlueprint>(Graph->GetOuter());
 	}
 	return nullptr;
 }
 
+URigVMGraph* UControlRigGraphNode::GetModel() const
+{
+	if (UControlRigGraph* Graph = Cast<UControlRigGraph>(GetOuter()))
+	{
+		return Graph->GetModel();
+	}
+	return nullptr;
+}
+
+URigVMController* UControlRigGraphNode::GetController() const
+{
+	if (UControlRigGraph* Graph = Cast<UControlRigGraph>(GetOuter()))
+	{
+		return Graph->GetController();
+	}
+	return nullptr;
+}
+
 URigVMNode* UControlRigGraphNode::GetModelNode() const
 {
 	UControlRigGraphNode* MutableThis = (UControlRigGraphNode*)this;
@@ -639,12 +661,9 @@ URigVMNode* UControlRigGraphNode::GetModelNode() const
 
 #endif
 
-		if (UControlRigBlueprint* Blueprint = Cast<UControlRigBlueprint>(Graph->GetOuter()))
+		if (URigVMGraph* Model = GetModel())
 		{
-			if (URigVMGraph* Model = Blueprint->Model)
-			{
-				return MutableThis->CachedModelNode = Model->FindNode(ModelNodePath);
-			}
+			return MutableThis->CachedModelNode = Model->FindNode(ModelNodePath);
 		}
 	}
 
diff --git a/Engine/Plugins/Experimental/ControlRig/Source/ControlRigDeveloper/Private/Graph/ControlRigGraphSchema.cpp b/Engine/Plugins/Experimental/ControlRig/Source/ControlRigDeveloper/Private/Graph/ControlRigGraphSchema.cpp
--- a/Engine/Plugins/Experimental/ControlRig/Source/ControlRigDeveloper/Private/Graph/ControlRigGraphSchema.cpp
+++ b/Engine/Plugins/Experimental/ControlRig/Source/ControlRigDeveloper/Private/Graph/ControlRigGraphSchema.cpp
@@ -70,13 +70,16 @@ bool UControlRigGraphSchema::TryCreateConnection(UEdGraphPin* PinA, UEdGraphPin*
 	UControlRigBlueprint* RigBlueprint = Cast<UControlRigBlueprint>(Blueprint);
 	if (RigBlueprint != nullptr)
 	{
-		if (PinA->Direction == EGPD_Input)
+		if (URigVMController* Controller = RigBlueprint->GetOrCreateController(PinA->GetOwningNode()->GetGraph()))
 		{
-			UEdGraphPin* Temp = PinA;
-			PinA = PinB;
-			PinB = Temp;
+			if (PinA->Direction == EGPD_Input)
+			{
+				UEdGraphPin* Temp = PinA;
+				PinA = PinB;
+				PinB = Temp;
+			}
+			return Controller->AddLink(PinA->GetName(), PinB->GetName());
 		}
-		return RigBlueprint->Controller->AddLink(PinA->GetName(), PinB->GetName());
 	}
 	return false;
 }
@@ -119,7 +122,7 @@ const FPinConnectionResponse UControlRigGraphSchema::CanCreateConnection(const U
 			if (PinA)
 			{
 				PinA = PinA->GetPinForLink();
-				RigBlueprint->Model->PrepareCycleChecking(PinA, A->Direction == EGPD_Input);
+				RigNodeA->GetModel()->PrepareCycleChecking(PinA, A->Direction == EGPD_Input);
 			}
 
 			URigVMPin* PinB = RigNodeB->GetModelPinFromPinPath(B->GetName());
@@ -136,7 +139,7 @@ const FPinConnectionResponse UControlRigGraphSchema::CanCreateConnection(const U
 			}
 
 			FString FailureReason;
-			bool bResult = RigBlueprint->Model->CanLink(PinA, PinB, &FailureReason);
+			bool bResult = RigNodeA->GetModel()->CanLink(PinA, PinB, &FailureReason);
 			if (!bResult)
 			{
 				return FPinConnectionResponse(CONNECT_RESPONSE_DISALLOW, FText::FromString(FailureReason));
@@ -169,11 +172,9 @@ void UControlRigGraphSchema::BreakPinLinks(UEdGraphPin& TargetPin, bool bSendsNo
 	//const FScopedTransaction Transaction( LOCTEXT("GraphEd_BreakPinLinks", "Break Pin Links") );
 
 	// cache this here, as BreakPinLinks can trigger a node reconstruction invalidating the TargetPin referenceS
-	UBlueprint* Blueprint = FBlueprintEditorUtils::FindBlueprintForNodeChecked(TargetPin.GetOwningNode());
-	const UControlRigBlueprint* RigBlueprint = Cast<UControlRigBlueprint>(Blueprint);
-	if (RigBlueprint != nullptr)
+	if (UControlRigGraphNode* Node = Cast< UControlRigGraphNode>(TargetPin.GetOwningNode()))
 	{
-		RigBlueprint->Controller->BreakAllLinks(TargetPin.GetName(), TargetPin.Direction == EGPD_Input);
+		Node->GetController()->BreakAllLinks(TargetPin.GetName(), TargetPin.Direction == EGPD_Input);
 	}
 }
 
@@ -181,9 +182,7 @@ void UControlRigGraphSchema::BreakSinglePinLink(UEdGraphPin* SourcePin, UEdGraph
 {
 	//const FScopedTransaction Transaction(LOCTEXT("GraphEd_BreakSinglePinLink", "Break Pin Link") );
 
-	UBlueprint* Blueprint = FBlueprintEditorUtils::FindBlueprintForNodeChecked(TargetPin->GetOwningNode());
-	const UControlRigBlueprint* RigBlueprint = Cast<UControlRigBlueprint>(Blueprint);
-	if (RigBlueprint != nullptr)
+	if (UControlRigGraphNode* Node = Cast< UControlRigGraphNode>(TargetPin->GetOwningNode()))
 	{
 		if (SourcePin->Direction == EGPD_Input)
 		{
@@ -192,7 +191,7 @@ void UControlRigGraphSchema::BreakSinglePinLink(UEdGraphPin* SourcePin, UEdGraph
 			SourcePin = Temp;
 		}
 		
-		RigBlueprint->Controller->BreakLink(SourcePin->GetName(), TargetPin->GetName());
+		Node->GetController()->BreakLink(SourcePin->GetName(), TargetPin->GetName());
 	}
 }
 
@@ -214,11 +213,9 @@ bool UControlRigGraphSchema::ShouldHidePinDefaultValue(UEdGraphPin* Pin) const
 
 bool UControlRigGraphSchema::IsPinBeingWatched(UEdGraphPin const* Pin) const
 {
-	UBlueprint* Blueprint = FBlueprintEditorUtils::FindBlueprintForNodeChecked(Pin->GetOwningNode());
-	const UControlRigBlueprint* RigBlueprint = Cast<UControlRigBlueprint>(Blueprint);
-	if (RigBlueprint != nullptr)
+	if (UControlRigGraphNode* Node = Cast< UControlRigGraphNode>(Pin->GetOwningNode()))
 	{
-		if (URigVMPin* ModelPin = RigBlueprint->Model->FindPin(Pin->GetName()))
+		if (URigVMPin* ModelPin = Node->GetModel()->FindPin(Pin->GetName()))
 		{
 			return ModelPin->RequiresWatch();
 		}
@@ -228,23 +225,19 @@ bool UControlRigGraphSchema::IsPinBeingWatched(UEdGraphPin const* Pin) const
 
 void UControlRigGraphSchema::ClearPinWatch(UEdGraphPin const* Pin) const
 {
-	UBlueprint* Blueprint = FBlueprintEditorUtils::FindBlueprintForNodeChecked(Pin->GetOwningNode());
-	const UControlRigBlueprint* RigBlueprint = Cast<UControlRigBlueprint>(Blueprint);
-	if (RigBlueprint != nullptr)
+	if (UControlRigGraphNode* Node = Cast< UControlRigGraphNode>(Pin->GetOwningNode()))
 	{
-		RigBlueprint->Controller->SetPinIsWatched(Pin->GetName(), false);
+		Node->GetController()->SetPinIsWatched(Pin->GetName(), false);
 	}
 }
 
 void UControlRigGraphSchema::OnPinConnectionDoubleCicked(UEdGraphPin* PinA, UEdGraphPin* PinB, const FVector2D& GraphPosition) const
 {
-	UBlueprint* Blueprint = FBlueprintEditorUtils::FindBlueprintForNodeChecked(PinA->GetOwningNode());
-	const UControlRigBlueprint* RigBlueprint = Cast<UControlRigBlueprint>(Blueprint);
-	if (RigBlueprint != nullptr)
+	if (UControlRigGraphNode* Node = Cast< UControlRigGraphNode>(PinA->GetOwningNode()))
 	{
-		if (URigVMLink* Link = RigBlueprint->Model->FindLink(FString::Printf(TEXT("%s -> %s"), *PinA->GetName(), *PinB->GetName())))
+		if (URigVMLink* Link = Node->GetModel()->FindLink(FString::Printf(TEXT("%s -> %s"), *PinA->GetName(), *PinB->GetName())))
 		{
-			RigBlueprint->Controller->AddRerouteNodeOnLink(Link, false, GraphPosition);
+			Node->GetController()->AddRerouteNodeOnLink(Link, false, GraphPosition);
 		}
 	}
 }
@@ -271,13 +264,7 @@ void UControlRigGraphSchema::SetNodePosition(UEdGraphNode* Node, const FVector2D
 {
 	if (UControlRigGraphNode* RigNode = Cast<UControlRigGraphNode>(Node))
 	{
-		UBlueprint* Blueprint = FBlueprintEditorUtils::FindBlueprintForNodeChecked(RigNode);
-		UControlRigBlueprint* RigBlueprint = Cast<UControlRigBlueprint>(Blueprint);
-		if (RigBlueprint != nullptr)
-		{
-			RigBlueprint->Controller->SetNodePosition(RigNode->GetModelNode(), Position, true, false); 
-			return;
-		}
+		RigNode->GetController()->SetNodePosition(RigNode->GetModelNode(), Position, true, false); 
 	}
 }
 
@@ -435,12 +422,7 @@ bool UControlRigGraphSchema::SafeDeleteNodeFromGraph(UEdGraph* Graph, UEdGraphNo
 {
 	if (UControlRigGraphNode* RigNode = Cast<UControlRigGraphNode>(Node))
 	{
-		UBlueprint* Blueprint = FBlueprintEditorUtils::FindBlueprintForNodeChecked(RigNode);
-		UControlRigBlueprint* RigBlueprint = Cast<UControlRigBlueprint>(Blueprint);
-		if (RigBlueprint != nullptr)
-		{
-			return RigBlueprint->Controller->RemoveNode(RigNode->GetModelNode());
-		}
+		return RigNode->GetController()->RemoveNode(RigNode->GetModelNode());
 	}
 	return false;
 }
@@ -476,11 +458,9 @@ bool UControlRigGraphSchema::RequestVariableDropOnPin(UEdGraph* InGraph, FProper
 #if WITH_EDITOR
 	if (CanVariableBeDropped(InGraph, InVariableToDrop))
 	{
-		UBlueprint* Blueprint = FBlueprintEditorUtils::FindBlueprintForGraph(InGraph);
-		UControlRigBlueprint* RigBlueprint = Cast<UControlRigBlueprint>(Blueprint);
-		if (RigBlueprint != nullptr)
+		if(UControlRigGraph* Graph = Cast<UControlRigGraph>(InGraph))
 		{
-			if (URigVMPin* ModelPin = RigBlueprint->Model->FindPin(InPin->GetName()))
+			if (URigVMPin* ModelPin = Graph->GetModel()->FindPin(InPin->GetName()))
 			{
 				FRigVMExternalVariable ExternalVariable = FRigVMExternalVariable::Make(InVariableToDrop, nullptr);
 				if (ModelPin->CanBeBoundToVariable(ExternalVariable))
@@ -488,16 +468,16 @@ bool UControlRigGraphSchema::RequestVariableDropOnPin(UEdGraph* InGraph, FProper
 					FModifierKeysState KeyState = FSlateApplication::Get().GetModifierKeys();
 					if (KeyState.IsAltDown())
 					{
-						return RigBlueprint->Controller->BindPinToVariable(ModelPin->GetPinPath(), InVariableToDrop->GetName());
+						return Graph->GetController()->BindPinToVariable(ModelPin->GetPinPath(), InVariableToDrop->GetName());
 					}
 					else
 					{
-						RigBlueprint->Controller->OpenUndoBracket(TEXT("Bind Variable to Pin"));
-						if (URigVMVariableNode* VariableNode = RigBlueprint->Controller->AddVariableNode(ExternalVariable.Name, ExternalVariable.TypeName.ToString(), ExternalVariable.TypeObject, true, FString(), InDropPosition + FVector2D(0.f, -34.f)))
+						Graph->GetController()->OpenUndoBracket(TEXT("Bind Variable to Pin"));
+						if (URigVMVariableNode* VariableNode = Graph->GetController()->AddVariableNode(ExternalVariable.Name, ExternalVariable.TypeName.ToString(), ExternalVariable.TypeObject, true, FString(), InDropPosition + FVector2D(0.f, -34.f)))
 						{
-							RigBlueprint->Controller->AddLink(VariableNode->FindPin(TEXT("Value"))->GetPinPath(), ModelPin->GetPinPath(), true);
+							Graph->GetController()->AddLink(VariableNode->FindPin(TEXT("Value"))->GetPinPath(), ModelPin->GetPinPath(), true);
 						}
-						RigBlueprint->Controller->CloseUndoBracket();
+						Graph->GetController()->CloseUndoBracket();
 						return true;
 					}
 				}
@@ -519,14 +499,8 @@ void UControlRigGraphSchema::EndGraphNodeInteraction(UEdGraphNode* InNode) const
 		return;
 	}
 
-	UControlRigBlueprint* Blueprint = Cast<UControlRigBlueprint>(Graph->GetOuter());
-	if (Blueprint == nullptr)
-	{
-		return;
-	}
-
-	check(Blueprint->Controller);
-	check(Blueprint->Model);
+	check(Graph->GetController());
+	check(Graph->GetModel());
 
 	TArray<UEdGraphNode*> NodesToMove;
 	NodesToMove.Add(InNode);
@@ -558,15 +532,15 @@ void UControlRigGraphSchema::EndGraphNodeInteraction(UEdGraphNode* InNode) const
 
 	bool bMovedSomething = false;
 
-	Blueprint->Controller->OpenUndoBracket(TEXT("Move Nodes"));
+	Graph->GetController()->OpenUndoBracket(TEXT("Move Nodes"));
 
 	for (UEdGraphNode* NodeToMove : NodesToMove)
 	{
 		FName NodeName = NodeToMove->GetFName();
-		if (URigVMNode* ModelNode = Blueprint->Model->FindNodeByName(NodeName))
+		if (URigVMNode* ModelNode = Graph->GetModel()->FindNodeByName(NodeName))
 		{
 			FVector2D Position(NodeToMove->NodePosX, NodeToMove->NodePosY);
-			if (Blueprint->Controller->SetNodePositionByName(NodeName, Position, true, false))
+			if (Graph->GetController()->SetNodePositionByName(NodeName, Position, true, false))
 			{
 				bMovedSomething = true;
 			}
@@ -575,11 +549,11 @@ void UControlRigGraphSchema::EndGraphNodeInteraction(UEdGraphNode* InNode) const
 
 	if (bMovedSomething)
 	{
-		Blueprint->Controller->CloseUndoBracket();
+		Graph->GetController()->CloseUndoBracket();
 	}
 	else
 	{
-		Blueprint->Controller->CancelUndoBracket();
+		Graph->GetController()->CancelUndoBracket();
 	}
 
 #endif
diff --git a/Engine/Plugins/Experimental/ControlRig/Source/ControlRigDeveloper/Public/ControlRigBlueprint.h b/Engine/Plugins/Experimental/ControlRig/Source/ControlRigDeveloper/Public/ControlRigBlueprint.h
--- a/Engine/Plugins/Experimental/ControlRig/Source/ControlRigDeveloper/Public/ControlRigBlueprint.h
+++ b/Engine/Plugins/Experimental/ControlRig/Source/ControlRigDeveloper/Public/ControlRigBlueprint.h
@@ -100,14 +100,33 @@ public:
 	UFUNCTION(BlueprintCallable, Category = "Control Rig Blueprint")
 	void RequestControlRigInit();
 
+	UFUNCTION(BlueprintCallable, Category = "Control Rig Blueprint")
+	URigVMGraph* GetModel(const UEdGraph* InEdGraph = nullptr) const;
+
+	UFUNCTION(BlueprintCallable, Category = "Control Rig Blueprint")
+	TArray<URigVMGraph*> GetAllModels() const;
+
+	UFUNCTION(BlueprintCallable, Category = "Control Rig Blueprint")
+	URigVMController* GetController(URigVMGraph* InGraph = nullptr) const;
+
+	UFUNCTION(BlueprintCallable, Category = "Control Rig Blueprint")
+	URigVMController* GetOrCreateController(URigVMGraph* InGraph = nullptr);
+
+	URigVMController* GetController(const UEdGraph* InEdGraph) const;
+	URigVMController* GetOrCreateController(const UEdGraph* InGraph);
+
 	UPROPERTY(EditAnywhere, BlueprintReadWrite, Category = "VM")
 	FRigVMCompileSettings VMCompileSettings;
 
+protected:
+
 	UPROPERTY(BlueprintReadOnly, Category = "VM")
 	URigVMGraph* Model;
 
 	UPROPERTY(BlueprintReadOnly, transient, Category = "VM")
-	URigVMController* Controller;
+	TMap<URigVMGraph*, URigVMController*> Controllers;
+
+public:
 
 	UPROPERTY(transient)
 	TMap<FString, FRigVMOperand> PinToOperandMap;
diff --git a/Engine/Plugins/Experimental/ControlRig/Source/ControlRigDeveloper/Public/Graph/ControlRigGraph.h b/Engine/Plugins/Experimental/ControlRig/Source/ControlRigDeveloper/Public/Graph/ControlRigGraph.h
--- a/Engine/Plugins/Experimental/ControlRig/Source/ControlRigDeveloper/Public/Graph/ControlRigGraph.h
+++ b/Engine/Plugins/Experimental/ControlRig/Source/ControlRigDeveloper/Public/Graph/ControlRigGraph.h
@@ -53,6 +53,10 @@ public:
 
 	UEdGraphNode* FindNodeForModelNodeName(const FName& InModelNodeName);
 
+	UControlRigBlueprint* GetBlueprint() const;
+	URigVMGraph* GetModel() const;
+	URigVMController* GetController() const;
+
 private:
 
 	void HandleModifiedEvent(ERigVMGraphNotifType InNotifType, URigVMGraph* InGraph, UObject* InSubject);
diff --git a/Engine/Plugins/Experimental/ControlRig/Source/ControlRigDeveloper/Public/Graph/ControlRigGraphNode.h b/Engine/Plugins/Experimental/ControlRig/Source/ControlRigDeveloper/Public/Graph/ControlRigGraphNode.h
--- a/Engine/Plugins/Experimental/ControlRig/Source/ControlRigDeveloper/Public/Graph/ControlRigGraphNode.h
+++ b/Engine/Plugins/Experimental/ControlRig/Source/ControlRigDeveloper/Public/Graph/ControlRigGraphNode.h
@@ -112,6 +112,12 @@ public:
 	/** Get the blueprint that this node is contained within */
 	UControlRigBlueprint* GetBlueprint() const;
 
+	/** Get the VM model this node lives within */
+	URigVMGraph* GetModel() const;
+
+	/** Get the blueprint that this node is contained within */
+	URigVMController* GetController() const;
+
 	/** Get the VM node this is node is wrapping */
 	URigVMNode* GetModelNode() const;
 
diff --git a/Engine/Plugins/Experimental/ControlRig/Source/ControlRigEditor/Private/ControlRigBlueprintEditorLibrary.cpp b/Engine/Plugins/Experimental/ControlRig/Source/ControlRigEditor/Private/ControlRigBlueprintEditorLibrary.cpp
--- a/Engine/Plugins/Experimental/ControlRig/Source/ControlRigEditor/Private/ControlRigBlueprintEditorLibrary.cpp
+++ b/Engine/Plugins/Experimental/ControlRig/Source/ControlRigEditor/Private/ControlRigBlueprintEditorLibrary.cpp
@@ -62,7 +62,7 @@ URigVMGraph* UControlRigBlueprintEditorLibrary::GetModel(UControlRigBlueprint* I
 	{
 		return nullptr;
 	}
-	return InRigBlueprint->Model;
+	return InRigBlueprint->GetModel();
 }
 
 URigVMController* UControlRigBlueprintEditorLibrary::GetController(UControlRigBlueprint* InRigBlueprint)
@@ -71,7 +71,7 @@ URigVMController* UControlRigBlueprintEditorLibrary::GetController(UControlRigBl
 	{
 		return nullptr;
 	}
-	return InRigBlueprint->Controller;
+	return InRigBlueprint->GetController();
 }
 
 TArray<UControlRigBlueprint*> UControlRigBlueprintEditorLibrary::GetCurrentlyOpenRigBlueprints()
diff --git a/Engine/Plugins/Experimental/ControlRig/Source/ControlRigEditor/Private/ControlRigCompilerDetails.cpp b/Engine/Plugins/Experimental/ControlRig/Source/ControlRigEditor/Private/ControlRigCompilerDetails.cpp
--- a/Engine/Plugins/Experimental/ControlRig/Source/ControlRigEditor/Private/ControlRigCompilerDetails.cpp
+++ b/Engine/Plugins/Experimental/ControlRig/Source/ControlRigEditor/Private/ControlRigCompilerDetails.cpp
@@ -92,9 +92,9 @@ FReply FRigVMCompileSettingsDetails::OnCopyASTClicked()
 {
 	if (BlueprintBeingCustomized)
 	{
-		if (BlueprintBeingCustomized->Model)
+		if (BlueprintBeingCustomized->GetModel())
 		{
-			FString DotContent = BlueprintBeingCustomized->Model->GetRuntimeAST()->DumpDot();
+			FString DotContent = BlueprintBeingCustomized->GetModel()->GetRuntimeAST()->DumpDot();
 			FPlatformApplicationMisc::ClipboardCopy(*DotContent);
 		}
 	}
@@ -105,7 +105,7 @@ FReply FRigVMCompileSettingsDetails::OnCopyByteCodeClicked()
 {
 	if (BlueprintBeingCustomized)
 	{
-		if (BlueprintBeingCustomized->Model)
+		if (BlueprintBeingCustomized->GetModel())
 		{
 			if(UControlRig* ControlRig = Cast<UControlRig>(BlueprintBeingCustomized->GetObjectBeingDebugged()))
 			{
diff --git a/Engine/Plugins/Experimental/ControlRig/Source/ControlRigEditor/Private/ControlRigEditorModule.cpp b/Engine/Plugins/Experimental/ControlRig/Source/ControlRigEditor/Private/ControlRigEditorModule.cpp
--- a/Engine/Plugins/Experimental/ControlRig/Source/ControlRigEditor/Private/ControlRigEditorModule.cpp
+++ b/Engine/Plugins/Experimental/ControlRig/Source/ControlRigEditor/Private/ControlRigEditorModule.cpp
@@ -1039,6 +1039,8 @@ void FControlRigEditorModule::GetContextMenuActions(const UControlRigGraphSchema
 
 		if (UEdGraphPin* InGraphPin = (UEdGraphPin* )Context->Pin)
 		{
+			UEdGraph* Graph = InGraphPin->GetOwningNode()->GetGraph();
+
 			// Add the watch pin / unwatch pin menu items
 			{
 				FToolMenuSection& Section = Menu->AddSection("EdGraphSchemaWatches", LOCTEXT("WatchesHeader", "Watches"));
@@ -1058,8 +1060,9 @@ void FControlRigEditorModule::GetContextMenuActions(const UControlRigGraphSchema
 			// Add alphainterp menu entries
 			if(UControlRigBlueprint* RigBlueprint = Cast<UControlRigBlueprint>((UBlueprint*)Context->Blueprint))
 			{
-				if (URigVMPin* ModelPin = RigBlueprint->Model->FindPin(InGraphPin->GetName()))
+				if (URigVMPin* ModelPin = RigBlueprint->GetModel(Graph)->FindPin(InGraphPin->GetName()))
 				{
+					URigVMController* Controller = RigBlueprint->GetController(ModelPin->GetGraph());
 
 					if (ModelPin->IsArray())
 					{
@@ -1069,8 +1072,8 @@ void FControlRigEditorModule::GetContextMenuActions(const UControlRigGraphSchema
 							LOCTEXT("ClearPinArray", "Clear Array"),
 							LOCTEXT("ClearPinArray_Tooltip", "Removes all elements of the array."),
 							FSlateIcon(),
-							FUIAction(FExecuteAction::CreateLambda([RigBlueprint, ModelPin]() {
-								RigBlueprint->Controller->ClearArrayPin(ModelPin->GetPinPath());
+							FUIAction(FExecuteAction::CreateLambda([Controller, ModelPin]() {
+								Controller->ClearArrayPin(ModelPin->GetPinPath());
 							})
 						));
 		}
@@ -1082,8 +1085,8 @@ void FControlRigEditorModule::GetContextMenuActions(const UControlRigGraphSchema
 							LOCTEXT("RemoveArrayPin", "Remove Array Element"),
 							LOCTEXT("RemoveArrayPin_Tooltip", "Removes the selected element from the array"),
 							FSlateIcon(),
-							FUIAction(FExecuteAction::CreateLambda([RigBlueprint, ModelPin]() {
-								RigBlueprint->Controller->RemoveArrayPin(ModelPin->GetPinPath());
+							FUIAction(FExecuteAction::CreateLambda([Controller, ModelPin]() {
+								Controller->RemoveArrayPin(ModelPin->GetPinPath());
 							})
 						));
 						Section.AddMenuEntry(
@@ -1091,8 +1094,8 @@ void FControlRigEditorModule::GetContextMenuActions(const UControlRigGraphSchema
 							LOCTEXT("DuplicateArrayPin", "Duplicate Array Element"),
 							LOCTEXT("DuplicateArrayPin_Tooltip", "Duplicates the selected element"),
 							FSlateIcon(),
-							FUIAction(FExecuteAction::CreateLambda([RigBlueprint, ModelPin]() {
-								RigBlueprint->Controller->DuplicateArrayPin(ModelPin->GetPinPath());
+							FUIAction(FExecuteAction::CreateLambda([Controller, ModelPin]() {
+								Controller->DuplicateArrayPin(ModelPin->GetPinPath());
 							})
 						));
 					}
@@ -1110,8 +1113,8 @@ void FControlRigEditorModule::GetContextMenuActions(const UControlRigGraphSchema
 								LOCTEXT("MakeVariableNodeFromBinding", "Make Variable Node"),
 								LOCTEXT("MakeVariableNodeFromBinding_Tooltip", "Turns the variable binding on the pin to a variable node"),
 								FSlateIcon(),
-								FUIAction(FExecuteAction::CreateLambda([RigBlueprint, ModelPin, NodePosition]() {
-									RigBlueprint->Controller->MakeVariableNodeFromBinding(ModelPin->GetPinPath(), NodePosition);
+								FUIAction(FExecuteAction::CreateLambda([Controller, ModelPin, NodePosition]() {
+									Controller->MakeVariableNodeFromBinding(ModelPin->GetPinPath(), NodePosition);
 								})
 							));
 						}
@@ -1125,12 +1128,12 @@ void FControlRigEditorModule::GetContextMenuActions(const UControlRigGraphSchema
 								LOCTEXT("PromotePinToVariable", "Promote Pin To Variable"),
 								LOCTEXT("PromotePinToVariable_Tooltip", "Turns the variable into a variable"),
 								FSlateIcon(),
-								FUIAction(FExecuteAction::CreateLambda([RigBlueprint, ModelPin, NodePosition]() {
+								FUIAction(FExecuteAction::CreateLambda([Controller, ModelPin, NodePosition]() {
 
 									FModifierKeysState KeyState = FSlateApplication::Get().GetModifierKeys();
 									bool bCreateVariableNode = !KeyState.IsAltDown();
 
-									RigBlueprint->Controller->PromotePinToVariable(ModelPin->GetPinPath(), bCreateVariableNode, NodePosition);
+									Controller->PromotePinToVariable(ModelPin->GetPinPath(), bCreateVariableNode, NodePosition);
 								})
 							));
 						}
@@ -1159,8 +1162,8 @@ void FControlRigEditorModule::GetContextMenuActions(const UControlRigGraphSchema
 							LOCTEXT("ResetPinDefaultValue", "Reset Pin Value"),
 							LOCTEXT("ResetPinDefaultValue_Tooltip", "Resets the pin's value to its default."),
 							FSlateIcon(),
-							FUIAction(FExecuteAction::CreateLambda([RigBlueprint, ModelPin]() {
-								RigBlueprint->Controller->ResetPinDefaultValue(ModelPin->GetPinPath());
+							FUIAction(FExecuteAction::CreateLambda([Controller, ModelPin]() {
+								Controller->ResetPinDefaultValue(ModelPin->GetPinPath());
 							})
 						));
 					}
@@ -1195,8 +1198,8 @@ void FControlRigEditorModule::GetContextMenuActions(const UControlRigGraphSchema
 								LOCTEXT("EjectLastNode", "Eject Last Node"),
 								LOCTEXT("EjectLastNode_Tooltip", "Eject the last injected node"),
 								FSlateIcon(),
-								FUIAction(FExecuteAction::CreateLambda([RigBlueprint, ModelPin]() {
-									RigBlueprint->Controller->EjectNodeFromPin(ModelPin->GetPinPath());
+								FUIAction(FExecuteAction::CreateLambda([Controller, ModelPin]() {
+									Controller->EjectNodeFromPin(ModelPin->GetPinPath());
 								})
 							));
 						}
@@ -1241,13 +1244,13 @@ void FControlRigEditorModule::GetContextMenuActions(const UControlRigGraphSchema
 									LOCTEXT("AddAlphaInterp", "Add Interpolate"),
 									LOCTEXT("AddAlphaInterp_Tooltip", "Injects an interpolate node"),
 									FSlateIcon(),
-									FUIAction(FExecuteAction::CreateLambda([RigBlueprint, InGraphPin, ModelPin, ScriptStruct]() {
-										URigVMInjectionInfo* Injection = RigBlueprint->Controller->AddInjectedNode(ModelPin->GetPinPath(), ModelPin->GetDirection() != ERigVMPinDirection::Output, ScriptStruct, TEXT("Execute"), TEXT("Value"), TEXT("Result"));
+									FUIAction(FExecuteAction::CreateLambda([Controller, InGraphPin, ModelPin, ScriptStruct]() {
+										URigVMInjectionInfo* Injection = Controller->AddInjectedNode(ModelPin->GetPinPath(), ModelPin->GetDirection() != ERigVMPinDirection::Output, ScriptStruct, TEXT("Execute"), TEXT("Value"), TEXT("Result"));
 										if (Injection)
 										{
 											TArray<FName> NodeNames;
 											NodeNames.Add(Injection->UnitNode->GetFName());
-											RigBlueprint->Controller->SetNodeSelection(NodeNames);
+											Controller->SetNodeSelection(NodeNames);
 										}
 									})
 								));
@@ -1262,7 +1265,7 @@ void FControlRigEditorModule::GetContextMenuActions(const UControlRigGraphSchema
 									FUIAction(FExecuteAction::CreateLambda([RigBlueprint, InterpNode]() {
 									TArray<FName> NodeNames;
 									NodeNames.Add(InterpNode->GetFName());
-									RigBlueprint->Controller->SetNodeSelection(NodeNames);
+									RigBlueprint->GetController(InterpNode->GetGraph())->SetNodeSelection(NodeNames);
 								})
 									));
 								Section.AddMenuEntry(
@@ -1270,8 +1273,8 @@ void FControlRigEditorModule::GetContextMenuActions(const UControlRigGraphSchema
 									LOCTEXT("RemoveAlphaInterp", "Remove Interpolate"),
 									LOCTEXT("RemoveAlphaInterp_Tooltip", "Removes the interpolate node"),
 									FSlateIcon(),
-									FUIAction(FExecuteAction::CreateLambda([RigBlueprint, InGraphPin, ModelPin, InterpNode]() {
-										RigBlueprint->Controller->RemoveNodeByName(InterpNode->GetFName());
+									FUIAction(FExecuteAction::CreateLambda([Controller, InGraphPin, ModelPin, InterpNode]() {
+										Controller->RemoveNodeByName(InterpNode->GetFName());
 									})
 								));
 							}
@@ -1323,13 +1326,13 @@ void FControlRigEditorModule::GetContextMenuActions(const UControlRigGraphSchema
 									LOCTEXT("AddVisualDebug", "Add Visual Debug"),
 									LOCTEXT("AddVisualDebug_Tooltip", "Injects a visual debugging node"),
 									FSlateIcon(),
-									FUIAction(FExecuteAction::CreateLambda([RigBlueprint, InGraphPin, ModelPin, ScriptStruct]() {
-										URigVMInjectionInfo* Injection = RigBlueprint->Controller->AddInjectedNode(ModelPin->GetPinPath(), ModelPin->GetDirection() != ERigVMPinDirection::Output, ScriptStruct, TEXT("Execute"), TEXT("Value"), TEXT("Value"));
+									FUIAction(FExecuteAction::CreateLambda([RigBlueprint, Controller, InGraphPin, ModelPin, ScriptStruct]() {
+										URigVMInjectionInfo* Injection = Controller->AddInjectedNode(ModelPin->GetPinPath(), ModelPin->GetDirection() != ERigVMPinDirection::Output, ScriptStruct, TEXT("Execute"), TEXT("Value"), TEXT("Value"));
 										if (Injection)
 										{
 											TArray<FName> NodeNames;
 											NodeNames.Add(Injection->UnitNode->GetFName());
-											RigBlueprint->Controller->SetNodeSelection(NodeNames);
+											Controller->SetNodeSelection(NodeNames);
 
 											if (URigVMUnitNode* UnitNode = Cast<URigVMUnitNode>(ModelPin->GetNode()))
 											{
@@ -1353,11 +1356,11 @@ void FControlRigEditorModule::GetContextMenuActions(const UControlRigGraphSchema
 															if(URigVMPin* SpaceTypePin = SpacePin->FindSubPin(TEXT("Type")))
 															{
 																FString SpaceTypeStr = StaticEnum<ERigElementType>()->GetDisplayNameTextByValue((int64)SpaceKey.Type).ToString();
-																RigBlueprint->Controller->SetPinDefaultValue(SpaceTypePin->GetPinPath(), SpaceTypeStr);
+																Controller->SetPinDefaultValue(SpaceTypePin->GetPinPath(), SpaceTypeStr);
 															}
 															if(URigVMPin* SpaceNamePin = SpacePin->FindSubPin(TEXT("Name")))
 															{
-																RigBlueprint->Controller->SetPinDefaultValue(SpaceNamePin->GetPinPath(), SpaceKey.Name.ToString());
+																Controller->SetPinDefaultValue(SpaceNamePin->GetPinPath(), SpaceKey.Name.ToString());
 															}
 														}
 													}
@@ -1374,10 +1377,10 @@ void FControlRigEditorModule::GetContextMenuActions(const UControlRigGraphSchema
 									LOCTEXT("EditVisualDebug", "Edit Visual Debug"),
 									LOCTEXT("EditVisualDebug_Tooltip", "Edit the visual debugging node"),
 									FSlateIcon(),
-									FUIAction(FExecuteAction::CreateLambda([RigBlueprint, VisualDebugNode]() {
+									FUIAction(FExecuteAction::CreateLambda([Controller, VisualDebugNode]() {
 										TArray<FName> NodeNames;
 										NodeNames.Add(VisualDebugNode->GetFName());
-										RigBlueprint->Controller->SetNodeSelection(NodeNames);
+										Controller->SetNodeSelection(NodeNames);
 									})
 								));
 								Section.AddMenuEntry(
@@ -1385,10 +1388,10 @@ void FControlRigEditorModule::GetContextMenuActions(const UControlRigGraphSchema
 									LOCTEXT("ToggleVisualDebug", "Toggle Visual Debug"),
 									LOCTEXT("ToggleVisualDebug_Tooltip", "Toggle the visibility the visual debugging"),
 									FSlateIcon(),
-									FUIAction(FExecuteAction::CreateLambda([RigBlueprint, VisualDebugNode]() {
+									FUIAction(FExecuteAction::CreateLambda([Controller, VisualDebugNode]() {
 										URigVMPin* EnabledPin = VisualDebugNode->FindPin(TEXT("bEnabled"));
 										check(EnabledPin);
-										RigBlueprint->Controller->SetPinDefaultValue(EnabledPin->GetPinPath(), EnabledPin->GetDefaultValue() == TEXT("True") ? TEXT("False") : TEXT("True"), false);
+										Controller->SetPinDefaultValue(EnabledPin->GetPinPath(), EnabledPin->GetDefaultValue() == TEXT("True") ? TEXT("False") : TEXT("True"), false);
 									})
 								));
 								Section.AddMenuEntry(
@@ -1396,8 +1399,8 @@ void FControlRigEditorModule::GetContextMenuActions(const UControlRigGraphSchema
 									LOCTEXT("RemoveVisualDebug", "Remove Visual Debug"),
 									LOCTEXT("RemoveVisualDebug_Tooltip", "Removes the visual debugging node"),
 									FSlateIcon(),
-										FUIAction(FExecuteAction::CreateLambda([RigBlueprint, InGraphPin, ModelPin, VisualDebugNode]() {
-										RigBlueprint->Controller->RemoveNodeByName(VisualDebugNode->GetFName());
+										FUIAction(FExecuteAction::CreateLambda([Controller, InGraphPin, ModelPin, VisualDebugNode]() {
+										Controller->RemoveNodeByName(VisualDebugNode->GetFName());
 									})
 								));
 							}
@@ -1410,14 +1413,17 @@ void FControlRigEditorModule::GetContextMenuActions(const UControlRigGraphSchema
 		{
 			if (UControlRigBlueprint* RigBlueprint = Cast<UControlRigBlueprint>((UBlueprint*)Context->Blueprint))
 			{
+				URigVMGraph* Model = RigBlueprint->GetModel(Context->Node->GetGraph());
+				URigVMController* Controller = RigBlueprint->GetController(Model);
+
 				TArray<FRigElementKey> RigElementsToSelect;
 				TMap<const URigVMPin*, FRigElementKey> PinToKey;
-				TArray<FName> SelectedNodeNames = RigBlueprint->Model->GetSelectNodes();
+				TArray<FName> SelectedNodeNames = Model->GetSelectNodes();
 				SelectedNodeNames.AddUnique(Context->Node->GetFName());
 
 				for(const FName& SelectedNodeName : SelectedNodeNames)
 				{
-					if (URigVMNode* ModelNode = RigBlueprint->Model->FindNodeByName(SelectedNodeName))
+					if (URigVMNode* ModelNode = Model->FindNodeByName(SelectedNodeName))
 					{
 						TSharedPtr<FStructOnScope> StructOnScope;
 						FRigHierarchyContainer TemporaryHierarchy = RigBlueprint->HierarchyContainer;
@@ -1532,7 +1538,7 @@ void FControlRigEditorModule::GetContextMenuActions(const UControlRigGraphSchema
 						LOCTEXT("SearchAndReplaceNames", "Search & Replace / Mirror"),
 						LOCTEXT("SearchAndReplaceNames_Tooltip", "Searches within all names and replaces with a different text."),
 						FSlateIcon(),
-						FUIAction(FExecuteAction::CreateLambda([RigBlueprint, PinToKey]() {
+						FUIAction(FExecuteAction::CreateLambda([RigBlueprint, Controller, PinToKey]() {
 
 							FRigMirrorSettings Settings;
 							TSharedPtr<FStructOnScope> StructToDisplay = MakeShareable(new FStructOnScope(FRigMirrorSettings::StaticStruct(), (uint8*)&Settings));
@@ -1542,7 +1548,7 @@ void FControlRigEditorModule::GetContextMenuActions(const UControlRigGraphSchema
 
 							SGenericDialogWidget::OpenDialog(LOCTEXT("ControlRigHierarchyMirror", "Mirror Graph"), KismetInspector, SGenericDialogWidget::FArguments(), true);
 
-							RigBlueprint->Controller->OpenUndoBracket(TEXT("Mirroring Graph"));
+							Controller->OpenUndoBracket(TEXT("Mirroring Graph"));
 							int32 ReplacedNames = 0;
 
 							for (const TPair<const URigVMPin*, FRigElementKey>& Pair : PinToKey)
@@ -1562,7 +1568,7 @@ void FControlRigEditorModule::GetContextMenuActions(const UControlRigGraphSchema
 									Key.Name = *NewNameStr;
 									if(RigBlueprint->HierarchyContainer.GetIndex(Key) != INDEX_NONE)
 									{
-										RigBlueprint->Controller->SetPinDefaultValue(Pin->GetPinPath(), NewNameStr, false);
+										Controller->SetPinDefaultValue(Pin->GetPinPath(), NewNameStr, false);
 										ReplacedNames++;
 									}
 								}
@@ -1570,11 +1576,11 @@ void FControlRigEditorModule::GetContextMenuActions(const UControlRigGraphSchema
 
 							if (ReplacedNames > 0)
 							{
-								RigBlueprint->Controller->CloseUndoBracket();
+								Controller->CloseUndoBracket();
 							}
 							else
 							{
-								RigBlueprint->Controller->CancelUndoBracket();
+								Controller->CancelUndoBracket();
 							}
 						})
 					));
@@ -1626,8 +1632,8 @@ void FControlRigEditorModule::GetContextMenuActions(const UControlRigGraphSchema
 							LOCTEXT("MakeBindingsFromVariableNode", "Make Bindings From Node"),
 							LOCTEXT("MakeBindingsFromVariableNode_Tooltip", "Turns the variable node into one ore more variable bindings on the pin(s)"),
 							FSlateIcon(),
-							FUIAction(FExecuteAction::CreateLambda([RigBlueprint, VariableNode]() {
-								RigBlueprint->Controller->MakeBindingsFromVariableNode(VariableNode->GetFName());
+							FUIAction(FExecuteAction::CreateLambda([RigBlueprint, Controller, VariableNode]() {
+								Controller->MakeBindingsFromVariableNode(VariableNode->GetFName());
 							})
 						));
 					}
@@ -1638,9 +1644,9 @@ void FControlRigEditorModule::GetContextMenuActions(const UControlRigGraphSchema
 						LOCTEXT("CollapseNodes", "Collapse Nodes"),
 						LOCTEXT("CollapseNodes_Tooltip", "Turns the selected nodes into a single Collapse node"),
 						FSlateIcon(),
-						FUIAction(FExecuteAction::CreateLambda([RigBlueprint]() {
-							TArray<FName> Nodes = RigBlueprint->Model->GetSelectNodes();
-							RigBlueprint->Controller->CollapseNodes(Nodes);
+						FUIAction(FExecuteAction::CreateLambda([Model, Controller]() {
+							TArray<FName> Nodes = Model->GetSelectNodes();
+							Controller->CollapseNodes(Nodes);
 						})
 					));
 
@@ -1651,8 +1657,8 @@ void FControlRigEditorModule::GetContextMenuActions(const UControlRigGraphSchema
 							LOCTEXT("ExpandNode", "Expand Node"),
 							LOCTEXT("ExpandNode_Tooltip", "Expands the contents of the node into this graph"),
 							FSlateIcon(),
-							FUIAction(FExecuteAction::CreateLambda([RigBlueprint, LibraryNode]() {
-								RigBlueprint->Controller->ExpandLibraryNode(LibraryNode->GetFName());
+							FUIAction(FExecuteAction::CreateLambda([Controller, LibraryNode]() {
+								Controller->ExpandLibraryNode(LibraryNode->GetFName());
 							})
 						));
 					}
diff --git a/Engine/Plugins/Experimental/ControlRig/Source/ControlRigEditor/Private/Editor/ControlRigEditor.cpp b/Engine/Plugins/Experimental/ControlRig/Source/ControlRigEditor/Private/Editor/ControlRigEditor.cpp
--- a/Engine/Plugins/Experimental/ControlRig/Source/ControlRigEditor/Private/Editor/ControlRigEditor.cpp
+++ b/Engine/Plugins/Experimental/ControlRig/Source/ControlRigEditor/Private/Editor/ControlRigEditor.cpp
@@ -92,6 +92,7 @@ namespace ControlRigEditorTabs
 
 FControlRigEditor::FControlRigEditor()
 	: ControlRig(nullptr)
+	, ActiveController(nullptr)
 	, bControlRigEditorInitialized(false)
 	, bIsSettingObjectBeingDebugged(false)
 	, NodeDetailStruct(nullptr)
@@ -286,16 +287,16 @@ void FControlRigEditor::InitControlRigEditor(const EToolkitMode::Type Mode, cons
 	}
 
 	{
-		if (URigVMGraph* Model = InControlRigBlueprint->Model)
+		if (URigVMGraph* Model = InControlRigBlueprint->GetModel())
 		{
 			if (Model->GetNodes().Num() == 0)
 			{
-				URigVMNode* Node = InControlRigBlueprint->Controller->AddUnitNode(FRigUnit_BeginExecution::StaticStruct(), TEXT("Execute"), FVector2D::ZeroVector, FString(), false);
+				URigVMNode* Node = InControlRigBlueprint->GetController()->AddUnitNode(FRigUnit_BeginExecution::StaticStruct(), TEXT("Execute"), FVector2D::ZeroVector, FString(), false);
 				if (Node)
 				{
 					TArray<FName> NodeNames;
 					NodeNames.Add(Node->GetFName());
-					InControlRigBlueprint->Controller->SetNodeSelection(NodeNames, false);
+					InControlRigBlueprint->GetController()->SetNodeSelection(NodeNames, false);
 				}
 			}
 			else
@@ -303,7 +304,7 @@ void FControlRigEditor::InitControlRigEditor(const EToolkitMode::Type Mode, cons
 				InControlRigBlueprint->RebuildGraphFromModel();
 
 				// selection state does not need to be persistent, even though it is saved in the RigVM.
-				InControlRigBlueprint->Controller->ClearNodeSelection(false);
+				InControlRigBlueprint->GetController()->ClearNodeSelection(false);
 
 				if (UPackage* Package = InControlRigBlueprint->GetOutermost())
 				{
@@ -1388,7 +1389,7 @@ bool FControlRigEditor::TransactionObjectAffectsBlueprint(UObject* InTransactedO
 		return false;
 	}
 
-	if (InTransactedObject->GetOuter() == RigBlueprint->Controller)
+	if (InTransactedObject->GetOuter() == GetFocusedController())
 	{
 		return false;
 	}
@@ -1409,7 +1410,7 @@ void FControlRigEditor::DeleteSelectedNodes()
 	SetUISelectionState(NAME_None);
 
 	bool DeletedAnything = false;
-	RigBlueprint->Controller->OpenUndoBracket(TEXT("Delete selected nodes"));
+	GetFocusedController()->OpenUndoBracket(TEXT("Delete selected nodes"));
 
 	for (FGraphPanelSelectionSet::TConstIterator NodeIt(SelectedNodes); NodeIt; ++NodeIt)
 	{
@@ -1420,14 +1421,14 @@ void FControlRigEditor::DeleteSelectedNodes()
 				AnalyticsTrackNodeEvent(GetBlueprintObj(), Node, true);
 				if (UControlRigGraphNode* RigNode = Cast<UControlRigGraphNode>(Node))
 				{
-					if(RigBlueprint->Controller->RemoveNodeByName(*RigNode->ModelNodePath))
+					if(GetFocusedController()->RemoveNodeByName(*RigNode->ModelNodePath))
 					{
 						DeletedAnything = true;
 					}
 				}
 				else if (UEdGraphNode_Comment* CommentNode = Cast<UEdGraphNode_Comment>(Node))
 				{
-					if(RigBlueprint->Controller->RemoveNodeByName(CommentNode->GetFName()))
+					if(GetFocusedController()->RemoveNodeByName(CommentNode->GetFName()))
 					{
 						DeletedAnything = true;
 					}
@@ -1442,11 +1443,11 @@ void FControlRigEditor::DeleteSelectedNodes()
 
 	if(DeletedAnything)
 	{
-		RigBlueprint->Controller->CloseUndoBracket();
+		GetFocusedController()->CloseUndoBracket();
 	}
 	else
 	{
-		RigBlueprint->Controller->CancelUndoBracket();
+		GetFocusedController()->CancelUndoBracket();
 	}
 }
 
@@ -1457,50 +1458,27 @@ bool FControlRigEditor::CanDeleteNodes() const
 
 void FControlRigEditor::CopySelectedNodes()
 {
-	UControlRigBlueprint* RigBlueprint = Cast<UControlRigBlueprint>(GetBlueprintObj());
-	if (RigBlueprint == nullptr)
-	{
-		return;
-	}
-
-	FString ExportedText = RigBlueprint->Controller->ExportSelectedNodesToText();
+	FString ExportedText = GetFocusedController()->ExportSelectedNodesToText();
 	FPlatformApplicationMisc::ClipboardCopy(*ExportedText);
 }
 
 bool FControlRigEditor::CanCopyNodes() const
 {
-	UControlRigBlueprint* RigBlueprint = Cast<UControlRigBlueprint>(GetBlueprintObj());
-	if (RigBlueprint == nullptr)
-	{
-		return false;
-	}
-	return RigBlueprint->Model->GetSelectNodes().Num() > 0;
+	return GetFocusedModel()->GetSelectNodes().Num() > 0;
 }
 
 bool FControlRigEditor::CanPasteNodes() const
 {
-	UControlRigBlueprint* RigBlueprint = Cast<UControlRigBlueprint>(GetBlueprintObj());
-	if (RigBlueprint == nullptr)
-	{
-		return false;
-	}
-
 	FString TextToImport;
 	FPlatformApplicationMisc::ClipboardPaste(TextToImport);
-	return RigBlueprint->Controller->CanImportNodesFromText(TextToImport);
+	return GetFocusedController()->CanImportNodesFromText(TextToImport);
 }
 
 void FControlRigEditor::PasteNodes()
 {
 	DECLARE_SCOPE_HIERARCHICAL_COUNTER_FUNC()
 
-	UControlRigBlueprint* RigBlueprint = Cast<UControlRigBlueprint>(GetBlueprintObj());
-	if (RigBlueprint == nullptr)
-	{
-		return;
-	}
-
-	RigBlueprint->Controller->OpenUndoBracket(TEXT("Pasted Nodes."));
+	GetFocusedController()->OpenUndoBracket(TEXT("Pasted Nodes."));
 
 	FVector2D PasteLocation = FSlateApplication::Get().GetCursorPos();
 
@@ -1517,7 +1495,7 @@ void FControlRigEditor::PasteNodes()
 
 	FString TextToImport;
 	FPlatformApplicationMisc::ClipboardPaste(TextToImport);
-	TArray<FName> NodeNames = RigBlueprint->Controller->ImportNodesFromText(TextToImport);
+	TArray<FName> NodeNames = GetFocusedController()->ImportNodesFromText(TextToImport);
 
 	if (NodeNames.Num() > 0)
 	{
@@ -1526,7 +1504,7 @@ void FControlRigEditor::PasteNodes()
 
 		for (const FName& NodeName : NodeNames)
 		{
-			const URigVMNode* Node = RigBlueprint->Model->FindNodeByName(NodeName);
+			const URigVMNode* Node = GetFocusedModel()->FindNodeByName(NodeName);
 			check(Node);
 
 			FVector2D Position = Node->GetPosition();
@@ -1543,19 +1521,19 @@ void FControlRigEditor::PasteNodes()
 
 		for (const FName& NodeName : NodeNames)
 		{
-			const URigVMNode* Node = RigBlueprint->Model->FindNodeByName(NodeName);
+			const URigVMNode* Node = GetFocusedModel()->FindNodeByName(NodeName);
 			check(Node);
 
 			FVector2D Position = Node->GetPosition();
-			RigBlueprint->Controller->SetNodePositionByName(NodeName, PasteLocation + Position - Bounds.GetCenter());
+			GetFocusedController()->SetNodePositionByName(NodeName, PasteLocation + Position - Bounds.GetCenter());
 		}
 
-		RigBlueprint->Controller->SetNodeSelection(NodeNames);
-		RigBlueprint->Controller->CloseUndoBracket();
+		GetFocusedController()->SetNodeSelection(NodeNames);
+		GetFocusedController()->CloseUndoBracket();
 	}
 	else
 	{
-		RigBlueprint->Controller->CancelUndoBracket();
+		GetFocusedController()->CancelUndoBracket();
 	}
 }
 
@@ -1617,25 +1595,19 @@ void FControlRigEditor::EnsureValidRigElementInDetailPanel()
 
 void FControlRigEditor::OnStartWatchingPin()
 {
-	if (UControlRigBlueprint* ControlRigBlueprint = CastChecked<UControlRigBlueprint>(GetBlueprintObj()))
+	if (UEdGraphPin* Pin = GetCurrentlySelectedPin())
 	{
-		if (UEdGraphPin* Pin = GetCurrentlySelectedPin())
-		{
-			ControlRigBlueprint->Controller->SetPinIsWatched(Pin->GetName(), true);
-		}
+		GetFocusedController()->SetPinIsWatched(Pin->GetName(), true);
 	}
 }
 
 bool FControlRigEditor::CanStartWatchingPin() const
 {
-	if (UControlRigBlueprint* ControlRigBlueprint = CastChecked<UControlRigBlueprint>(GetBlueprintObj()))
+	if (UEdGraphPin* Pin = GetCurrentlySelectedPin())
 	{
-		if (UEdGraphPin* Pin = GetCurrentlySelectedPin())
+		if (URigVMPin* ModelPin = GetFocusedModel()->FindPin(Pin->GetName()))
 		{
-			if (URigVMPin* ModelPin = ControlRigBlueprint->Model->FindPin(Pin->GetName()))
-			{
-				return ModelPin->GetParentPin() == nullptr && !ModelPin->RequiresWatch();
-			}
+			return ModelPin->GetParentPin() == nullptr && !ModelPin->RequiresWatch();
 		}
 	}
 	return false;
@@ -1643,25 +1615,19 @@ bool FControlRigEditor::CanStartWatchingPin() const
 
 void FControlRigEditor::OnStopWatchingPin()
 {
-	if (UControlRigBlueprint* ControlRigBlueprint = CastChecked<UControlRigBlueprint>(GetBlueprintObj()))
+	if (UEdGraphPin* Pin = GetCurrentlySelectedPin())
 	{
-		if (UEdGraphPin* Pin = GetCurrentlySelectedPin())
-		{
-			ControlRigBlueprint->Controller->SetPinIsWatched(Pin->GetName(), false);
-		}
+		GetFocusedController()->SetPinIsWatched(Pin->GetName(), false);
 	}
 }
 
 bool FControlRigEditor::CanStopWatchingPin() const
 {
-	if (UControlRigBlueprint* ControlRigBlueprint = CastChecked<UControlRigBlueprint>(GetBlueprintObj()))
+	if (UEdGraphPin* Pin = GetCurrentlySelectedPin())
 	{
-		if (UEdGraphPin* Pin = GetCurrentlySelectedPin())
+		if (URigVMPin* ModelPin = GetFocusedModel()->FindPin(Pin->GetName()))
 		{
-			if (URigVMPin* ModelPin = ControlRigBlueprint->Model->FindPin(Pin->GetName()))
-			{
-				return ModelPin->RequiresWatch();
-			}
+			return ModelPin->RequiresWatch();
 		}
 	}
 	return false;
@@ -1964,20 +1930,20 @@ void FControlRigEditor::OnSelectedNodesChangedImpl(const TSet<class UObject*>& N
 			{
 				if (UEdGraphNode_Comment* CommentNode = Cast<UEdGraphNode_Comment>(Node))
 				{
-					URigVMNode* ModelNode = ControlRigBlueprint->Model->FindNodeByName(Node->GetFName());
+					URigVMNode* ModelNode = GetFocusedModel()->FindNodeByName(Node->GetFName());
 					if (ModelNode == nullptr)
 					{
 						TGuardValue<bool> BlueprintNotifGuard(ControlRigBlueprint->bSuspendModelNotificationsForOthers, true);
 						FVector2D NodePos(CommentNode->NodePosX, CommentNode->NodePosY);
 						FVector2D NodeSize(CommentNode->NodeWidth, CommentNode->NodeHeight);
 						FLinearColor NodeColor = CommentNode->CommentColor;
-						ControlRigBlueprint->Controller->AddCommentNode(CommentNode->NodeComment, NodePos, NodeSize, NodeColor, CommentNode->GetName(), true);
+						GetFocusedController()->AddCommentNode(CommentNode->NodeComment, NodePos, NodeSize, NodeColor, CommentNode->GetName(), true);
 					}
 				}
 				NodeNamesToSelect.Add(Node->GetFName());
 			}
 		}
-		ControlRigBlueprint->Controller->SetNodeSelection(NodeNamesToSelect, true);
+		GetFocusedController()->SetNodeSelection(NodeNamesToSelect, true);
 	}
 }
 
@@ -2427,14 +2393,14 @@ void FControlRigEditor::SetPinControlNameListText(const FText& NewTypeInValue, E
 			// find out if the controlled pin is part of a visual debug node
 			if (UControlRigBlueprint* ControlRigBlueprint = CastChecked<UControlRigBlueprint>(GetBlueprintObj()))
 			{
-				if (URigVMPin* ControlledPin = ControlRigBlueprint->Model->FindPin(Control.Name.ToString()))
+				if (URigVMPin* ControlledPin = GetFocusedModel()->FindPin(Control.Name.ToString()))
 				{
 					URigVMNode* ControlledNode = ControlledPin->GetPinForLink()->GetNode();
 					if (URigVMPin* BoneSpacePin = ControlledNode->FindPin(TEXT("BoneSpace")))
 					{
 						if (BoneSpacePin->GetCPPType() == TEXT("FName") && BoneSpacePin->GetCustomWidgetName() == TEXT("BoneName"))
 						{
-							ControlRigBlueprint->Controller->SetPinDefaultValue(BoneSpacePin->GetPinPath(), Control.ParentName.ToString(), false, false, false);
+							GetFocusedController()->SetPinDefaultValue(BoneSpacePin->GetPinPath(), Control.ParentName.ToString(), false, false, false);
 						}
 					}
 				}
@@ -2671,22 +2637,26 @@ void FControlRigEditor::UpdateStaleWatchedPins()
 
 	ControlRigBP->WatchedPins.Empty();
 
-	for (URigVMNode* ModelNode : ControlRigBP->Model->GetNodes())
+	TArray<URigVMGraph*> Models = ControlRigBP->GetAllModels();
+	for (URigVMGraph* Model : Models)
 	{
-		TArray<URigVMPin*> ModelPins = ModelNode->GetAllPinsRecursively();
-		for (URigVMPin* ModelPin : ModelPins)
+		for (URigVMNode* ModelNode : Model->GetNodes())
 		{
-			if (ModelPin->RequiresWatch())
+			TArray<URigVMPin*> ModelPins = ModelNode->GetAllPinsRecursively();
+			for (URigVMPin* ModelPin : ModelPins)
 			{
-				ControlRigBP->Controller->SetPinIsWatched(ModelPin->GetPinPath(), false, false);
+				if (ModelPin->RequiresWatch())
+				{
+					ControlRigBP->GetController(Model)->SetPinIsWatched(ModelPin->GetPinPath(), false, false);
+				}
 			}
 		}
 	}
-
 	for (UEdGraphPin* Pin : AllPins)
 	{
 		ControlRigBP->WatchedPins.Add(Pin);
-		ControlRigBP->Controller->SetPinIsWatched(Pin->GetName(), true, false);
+		UEdGraph* EdGraph = Pin->GetOwningNode()->GetGraph();
+		ControlRigBP->GetController(EdGraph)->SetPinIsWatched(Pin->GetName(), true, false);
 	}
 }
 
@@ -2713,7 +2683,7 @@ void FControlRigEditor::OnNodeTitleCommitted(const FText& NewText, ETextCommit::
 	{
 		if (UControlRigBlueprint* ControlRigBP = GetControlRigBlueprint())
 		{
-			ControlRigBP->Controller->SetCommentTextByName(CommentBeingChanged->GetFName(), NewText.ToString(), true);
+			GetFocusedController()->SetCommentTextByName(CommentBeingChanged->GetFName(), NewText.ToString(), true);
 		}
 	}
 }
@@ -2877,7 +2847,7 @@ void FControlRigEditor::OnFinishedChangingProperties(const FPropertyChangedEvent
 		if (!DefaultValue.IsEmpty())
 		{
 			FString PinPath = FString::Printf(TEXT("%s.%s"), *NodeDetailName.ToString(), *PropertyChangedEvent.MemberProperty->GetName());
-			ControlRigBP->Controller->SetPinDefaultValue(PinPath, DefaultValue, true, true);
+			GetFocusedController()->SetPinDefaultValue(PinPath, DefaultValue, true, true);
 		}
 	}
 
@@ -3132,7 +3102,7 @@ void FControlRigEditor::OnRigElementRemoved(FRigHierarchyContainer* Container, c
 						{
 							if (ModelPin->GetDefaultValue() == RemovedElementName)
 							{
-								Blueprint->Controller->SetPinDefaultValue(ModelPin->GetPinPath(), FName(NAME_None).ToString());
+								GetFocusedController()->SetPinDefaultValue(ModelPin->GetPinPath(), FName(NAME_None).ToString());
 							}
 						}
 						else if (ModelPin->GetCPPTypeObject() == FRigElementKey::StaticStruct())
@@ -3149,7 +3119,7 @@ void FControlRigEditor::OnRigElementRemoved(FRigHierarchyContainer* Container, c
 										FString NameStr = NamePin->GetDefaultValue();
 										if (NameStr == RemovedElementName)
 										{
-											Blueprint->Controller->SetPinDefaultValue(NamePin->GetPinPath(), NoneStr);
+											GetFocusedController()->SetPinDefaultValue(NamePin->GetPinPath(), NoneStr);
 										}
 									}
 								}
@@ -3210,7 +3180,7 @@ void FControlRigEditor::OnRigElementRenamed(FRigHierarchyContainer* Container, E
 						{
 							if (ModelPin->GetDefaultValue() == InOldName.ToString())
 							{
-								Blueprint->Controller->SetPinDefaultValue(ModelPin->GetPinPath(), InNewName.ToString(), false);
+								GetFocusedController()->SetPinDefaultValue(ModelPin->GetPinPath(), InNewName.ToString(), false);
 							}
 						}
 						else if (ModelPin->GetCPPTypeObject() == FRigElementKey::StaticStruct())
@@ -3226,7 +3196,7 @@ void FControlRigEditor::OnRigElementRenamed(FRigHierarchyContainer* Container, E
 										FString NameStr = NamePin->GetDefaultValue();
 										if (NameStr == OldNameStr)
 										{
-											Blueprint->Controller->SetPinDefaultValue(NamePin->GetPinPath(), NewNameStr);
+											GetFocusedController()->SetPinDefaultValue(NamePin->GetPinPath(), NewNameStr);
 										}
 									}
 								}
@@ -3491,9 +3461,9 @@ void FControlRigEditor::OnGraphNodeDropToPerform(TSharedPtr<FGraphNodeDragDropOp
 					LOCTEXT("CreateCollectionFromKeysTooltip", "Creates a collection from the selected elements in the hierarchy"),
 					FSlateIcon(),
 					FUIAction(
-						FExecuteAction::CreateLambda([Blueprint, DraggedKeys, NodePosition]()
+						FExecuteAction::CreateLambda([this, Blueprint, DraggedKeys, NodePosition]()
 							{
-								if (URigVMController* Controller = Blueprint->Controller)
+								if (URigVMController* Controller = GetFocusedController())
 								{
 									Controller->OpenUndoBracket(TEXT("Create Collection from Items"));
 
@@ -3554,12 +3524,12 @@ void FControlRigEditor::HandleMakeElementGetterSetter(ERigElementGetterSetterTyp
 	{
 		return;
 	}
-	if (Blueprint->Controller == nullptr)
+	if (GetFocusedController() == nullptr)
 	{
 		return;
 	}
 
-	Blueprint->Controller->OpenUndoBracket(TEXT("Adding Nodes from Hierarchy"));
+	GetFocusedController()->OpenUndoBracket(TEXT("Adding Nodes from Hierarchy"));
 
 	struct FNewNodeData
 	{
@@ -3804,7 +3774,7 @@ void FControlRigEditor::HandleMakeElementGetterSetter(ERigElementGetterSetterTyp
 		}
 
 		FName Name = FControlRigBlueprintUtils::ValidateName(Blueprint, StructTemplate->GetName());
-		if (URigVMUnitNode* ModelNode = Blueprint->Controller->AddUnitNode(StructTemplate, TEXT("Execute"), NodePosition))
+		if (URigVMUnitNode* ModelNode = GetFocusedController()->AddUnitNode(StructTemplate, TEXT("Execute"), NodePosition))
 		{
 			FString ItemTypeStr = StaticEnum<ERigElementType>()->GetDisplayNameTextByValue((int64)Key.Type).ToString();
 			NewNode.Name = ModelNode->GetFName();
@@ -3812,13 +3782,13 @@ void FControlRigEditor::HandleMakeElementGetterSetter(ERigElementGetterSetterTyp
 			
 			for (const FName& ItemPin : ItemPins)
 			{
-				Blueprint->Controller->SetPinDefaultValue(FString::Printf(TEXT("%s.%s.Name"), *ModelNode->GetName(), *ItemPin.ToString()), Key.Name.ToString());
-				Blueprint->Controller->SetPinDefaultValue(FString::Printf(TEXT("%s.%s.Type"), *ModelNode->GetName(), *ItemPin.ToString()), ItemTypeStr);			
+				GetFocusedController()->SetPinDefaultValue(FString::Printf(TEXT("%s.%s.Name"), *ModelNode->GetName(), *ItemPin.ToString()), Key.Name.ToString());
+				GetFocusedController()->SetPinDefaultValue(FString::Printf(TEXT("%s.%s.Type"), *ModelNode->GetName(), *ItemPin.ToString()), ItemTypeStr);
 			}
 
 			for (const FName& NamePin : NamePins)
 			{
-				Blueprint->Controller->SetPinDefaultValue(FString::Printf(TEXT("%s.%s"), *ModelNode->GetName(), *NamePin.ToString()), Key.Name.ToString());
+				GetFocusedController()->SetPinDefaultValue(FString::Printf(TEXT("%s.%s"), *ModelNode->GetName(), *NamePin.ToString()), Key.Name.ToString());
 			}
 
 			if (!NewNode.ValuePinName.IsNone())
@@ -3850,7 +3820,7 @@ void FControlRigEditor::HandleMakeElementGetterSetter(ERigElementGetterSetterTyp
 				}
 				if (!DefaultValue.IsEmpty())
 				{
-					Blueprint->Controller->SetPinDefaultValue(FString::Printf(TEXT("%s.%s"), *ModelNode->GetName(), *NewNode.ValuePinName.ToString()), DefaultValue);
+					GetFocusedController()->SetPinDefaultValue(FString::Printf(TEXT("%s.%s"), *ModelNode->GetName(), *NewNode.ValuePinName.ToString()), DefaultValue);
 				}
 			}
 
@@ -3867,12 +3837,12 @@ void FControlRigEditor::HandleMakeElementGetterSetter(ERigElementGetterSetterTyp
 		{
 			NewNodeNames.Add(NewNode.Name);
 		}
-		Blueprint->Controller->SetNodeSelection(NewNodeNames);
-		Blueprint->Controller->CloseUndoBracket();
+		GetFocusedController()->SetNodeSelection(NewNodeNames);
+		GetFocusedController()->CloseUndoBracket();
 	}
 	else
 	{
-		Blueprint->Controller->CancelUndoBracket();
+		GetFocusedController()->CancelUndoBracket();
 	}
 }
 
@@ -3932,7 +3902,7 @@ void FControlRigEditor::HandleOnControlModified(UControlRig* Subject, const FRig
 
 			if (!NewDefaultValue.IsEmpty())
 			{
-				Blueprint->Controller->SetPinDefaultValue(Pin->GetPinPath(), NewDefaultValue, true, true, true);
+				GetFocusedController()->SetPinDefaultValue(Pin->GetPinPath(), NewDefaultValue, true, true, true);
 			}
 		}
 		else
@@ -4005,7 +3975,7 @@ void FControlRigEditor::HandleVariableDroppedFromBlueprint(UObject* InSubject, F
 		return;
 	}
 
-	URigVMController* Controller = Blueprint->Controller;
+	URigVMController* Controller = GetFocusedController();
 	check(Controller);
 
 	FRigVMExternalVariable ExternalVariable = FRigVMExternalVariable::Make(InVariableToDrop, nullptr);
@@ -4324,4 +4294,26 @@ TSharedPtr<FUICommandList> FControlRigEditor::HandleOnViewportContextMenuCommand
 	return TSharedPtr<FUICommandList>();
 }
 
+URigVMGraph* FControlRigEditor::GetFocusedModel() const
+{
+	UControlRigBlueprint* Blueprint = Cast<UControlRigBlueprint>(GetBlueprintObj());
+	if (Blueprint == nullptr)
+	{
+		return nullptr;
+	}
+
+	UControlRigGraph* EdGraph = Cast<UControlRigGraph>(GetFocusedGraph());
+	return Blueprint->GetModel(EdGraph);
+}
+
+URigVMController* FControlRigEditor::GetFocusedController() const
+{
+	UControlRigBlueprint* Blueprint = Cast<UControlRigBlueprint>(GetBlueprintObj());
+	if (Blueprint == nullptr)
+	{
+		return nullptr;
+	}
+	return Blueprint->GetController(GetFocusedModel());
+}
+
 #undef LOCTEXT_NAMESPACE
diff --git a/Engine/Plugins/Experimental/ControlRig/Source/ControlRigEditor/Private/Editor/ControlRigEditor.h b/Engine/Plugins/Experimental/ControlRig/Source/ControlRigEditor/Private/Editor/ControlRigEditor.h
--- a/Engine/Plugins/Experimental/ControlRig/Source/ControlRigEditor/Private/Editor/ControlRigEditor.h
+++ b/Engine/Plugins/Experimental/ControlRig/Source/ControlRigEditor/Private/Editor/ControlRigEditor.h
@@ -350,6 +350,10 @@ protected:
 	/** delegate for changing property */
 	virtual void OnFinishedChangingProperties(const FPropertyChangedEvent& PropertyChangedEvent) override;
 
+	URigVMGraph* GetFocusedModel() const;
+	URigVMController* GetFocusedController() const;
+
+	URigVMController* ActiveController;
 	bool bControlRigEditorInitialized;
 	bool bIsSettingObjectBeingDebugged;
 	FRigElementKey RigElementInDetailPanel;
diff --git a/Engine/Plugins/Experimental/ControlRig/Source/ControlRigEditor/Private/Editor/SControlRigStackView.cpp b/Engine/Plugins/Experimental/ControlRig/Source/ControlRigEditor/Private/Editor/SControlRigStackView.cpp
--- a/Engine/Plugins/Experimental/ControlRig/Source/ControlRigEditor/Private/Editor/SControlRigStackView.cpp
+++ b/Engine/Plugins/Experimental/ControlRig/Source/ControlRigEditor/Private/Editor/SControlRigStackView.cpp
@@ -259,7 +259,7 @@ void SControlRigStackView::OnSelectionChanged(TSharedPtr<FRigStackEntry> Selecti
 		}
 
 		TMap<int32, FName> InstructionIndexToNodeName;
-		for (URigVMNode* ModelNode : ControlRigBlueprint->Model->GetNodes())
+		for (URigVMNode* ModelNode : ControlRigEditor.Pin()->GetFocusedModel()->GetNodes())
 		{
 			if (ModelNode->GetInstructionIndex() != INDEX_NONE)
 			{
@@ -277,7 +277,7 @@ void SControlRigStackView::OnSelectionChanged(TSharedPtr<FRigStackEntry> Selecti
 			}
 		}
 
-		ControlRigBlueprint->Controller->SetNodeSelection(SelectedNodes);
+		ControlRigEditor.Pin()->GetFocusedController()->SetNodeSelection(SelectedNodes);
 	}
 }
 
@@ -303,7 +303,7 @@ void SControlRigStackView::PopulateStackView(URigVM* InVM)
 	if (InVM)
 	{
 		FRigVMInstructionArray Instructions = InVM->GetInstructions();
-		const TArray<URigVMNode*>& Nodes = ControlRigBlueprint->Model->GetNodes();
+		const TArray<URigVMNode*>& Nodes = ControlRigBlueprint->GetModel()->GetNodes();
 		
 		// 1. cache information about instructions/nodes, which will be used later 
 		TArray<int32> InstructionIndexToNodeIndex;
@@ -541,11 +541,15 @@ void SControlRigStackView::HandleControlRigInitializedEvent(UControlRig* InContr
 
 		if (UControlRigBlueprint* RigBlueprint = ControlRigEditor.Pin()->GetControlRigBlueprint())
 		{
-			for (URigVMNode* ModelNode : RigBlueprint->Model->GetNodes())
+			TArray<URigVMGraph*> Models = RigBlueprint->GetAllModels();
+			for (URigVMGraph* Model : Models)
 			{
-				if (ModelNode->IsSelected())
+				for (URigVMNode* ModelNode : Model->GetNodes())
 				{
-					HandleModifiedEvent(ERigVMGraphNotifType::NodeSelected, RigBlueprint->Model, ModelNode);
+					if (ModelNode->IsSelected())
+					{
+						HandleModifiedEvent(ERigVMGraphNotifType::NodeSelected, Model, ModelNode);
+					}
 				}
 			}
 		}
diff --git a/Engine/Plugins/Experimental/ControlRig/Source/ControlRigEditor/Private/Graph/ControlRigConnectionDrawingPolicy.cpp b/Engine/Plugins/Experimental/ControlRig/Source/ControlRigEditor/Private/Graph/ControlRigConnectionDrawingPolicy.cpp
--- a/Engine/Plugins/Experimental/ControlRig/Source/ControlRigEditor/Private/Graph/ControlRigConnectionDrawingPolicy.cpp
+++ b/Engine/Plugins/Experimental/ControlRig/Source/ControlRigEditor/Private/Graph/ControlRigConnectionDrawingPolicy.cpp
@@ -38,7 +38,7 @@ void FControlRigConnectionDrawingPolicy::ResetIncompatiblePinDrawState(const TSe
 			UControlRigBlueprint* RigBlueprint = Cast<UControlRigBlueprint>(Blueprint);
 			if (RigBlueprint != nullptr)
 			{
-				RigBlueprint->Model->PrepareCycleChecking(nullptr, true);
+				RigBlueprint->GetModel(Pin->GetOwningNode()->GetGraph())->PrepareCycleChecking(nullptr, true);
 			}
 		}
 	}
diff --git a/Engine/Plugins/Experimental/ControlRig/Source/ControlRigEditor/Private/Graph/NodeSpawners/ControlRigBranchNodeSpawner.cpp b/Engine/Plugins/Experimental/ControlRig/Source/ControlRigEditor/Private/Graph/NodeSpawners/ControlRigBranchNodeSpawner.cpp
--- a/Engine/Plugins/Experimental/ControlRig/Source/ControlRigEditor/Private/Graph/NodeSpawners/ControlRigBranchNodeSpawner.cpp
+++ b/Engine/Plugins/Experimental/ControlRig/Source/ControlRigEditor/Private/Graph/NodeSpawners/ControlRigBranchNodeSpawner.cpp
@@ -68,7 +68,7 @@ UEdGraphNode* UControlRigBranchNodeSpawner::Invoke(UEdGraph* ParentGraph, FBindi
 	}
 #endif
 
-	URigVMController* Controller = bIsTemplateNode ? RigGraph->GetTemplateController() : RigBlueprint->Controller;
+	URigVMController* Controller = bIsTemplateNode ? RigGraph->GetTemplateController() : RigBlueprint->GetController(ParentGraph);
 
 	FName Name = *URigVMBranchNode::BranchName;
 
diff --git a/Engine/Plugins/Experimental/ControlRig/Source/ControlRigEditor/Private/Graph/NodeSpawners/ControlRigEnumNodeSpawner.cpp b/Engine/Plugins/Experimental/ControlRig/Source/ControlRigEditor/Private/Graph/NodeSpawners/ControlRigEnumNodeSpawner.cpp
--- a/Engine/Plugins/Experimental/ControlRig/Source/ControlRigEditor/Private/Graph/NodeSpawners/ControlRigEnumNodeSpawner.cpp
+++ b/Engine/Plugins/Experimental/ControlRig/Source/ControlRigEditor/Private/Graph/NodeSpawners/ControlRigEnumNodeSpawner.cpp
@@ -88,7 +88,7 @@ UEdGraphNode* UControlRigEnumNodeSpawner::Invoke(UEdGraph* ParentGraph, FBinding
 	}
 #endif
 
-	URigVMController* Controller = bIsTemplateNode ? RigGraph->GetTemplateController() : RigBlueprint->Controller;
+	URigVMController* Controller = bIsTemplateNode ? RigGraph->GetTemplateController() : RigBlueprint->GetController(ParentGraph);
 
 	FName Name = *URigVMEnumNode::EnumName;
 
diff --git a/Engine/Plugins/Experimental/ControlRig/Source/ControlRigEditor/Private/Graph/NodeSpawners/ControlRigIfNodeSpawner.cpp b/Engine/Plugins/Experimental/ControlRig/Source/ControlRigEditor/Private/Graph/NodeSpawners/ControlRigIfNodeSpawner.cpp
--- a/Engine/Plugins/Experimental/ControlRig/Source/ControlRigEditor/Private/Graph/NodeSpawners/ControlRigIfNodeSpawner.cpp
+++ b/Engine/Plugins/Experimental/ControlRig/Source/ControlRigEditor/Private/Graph/NodeSpawners/ControlRigIfNodeSpawner.cpp
@@ -91,7 +91,7 @@ UEdGraphNode* UControlRigIfNodeSpawner::Invoke(UEdGraph* ParentGraph, FBindingSe
 	{
 		if (const UEdGraphPin* LastPin = RigSchema->LastPinForCompatibleCheck)
 		{
-			if (URigVMPin* ModelPin = RigBlueprint->Model->FindPin(LastPin->GetName()))
+			if (URigVMPin* ModelPin = RigBlueprint->GetModel(ParentGraph)->FindPin(LastPin->GetName()))
 			{
 				CPPType = ModelPin->GetCPPType();
 				if (ModelPin->GetCPPTypeObject())
@@ -111,7 +111,7 @@ UEdGraphNode* UControlRigIfNodeSpawner::Invoke(UEdGraph* ParentGraph, FBindingSe
 	}
 #endif
 
-	URigVMController* Controller = bIsTemplateNode ? RigGraph->GetTemplateController() : RigBlueprint->Controller;
+	URigVMController* Controller = bIsTemplateNode ? RigGraph->GetTemplateController() : RigBlueprint->GetController(ParentGraph);
 
 	FName Name = *URigVMIfNode::IfName;
 
@@ -159,7 +159,8 @@ bool UControlRigIfNodeSpawner::IsTemplateNodeFilteredOut(FBlueprintActionFilter
 			}
 
 			FString PinPath = Filter.Context.Pins[0]->GetName();
-			if (URigVMPin* ModelPin = RigBlueprint->Model->FindPin(PinPath))
+			UEdGraph* EdGraph = Filter.Context.Pins[0]->GetOwningNode()->GetGraph();
+			if (URigVMPin* ModelPin = RigBlueprint->GetModel(EdGraph)->FindPin(PinPath))
 			{
 				return ModelPin->IsExecuteContext();
 			}
diff --git a/Engine/Plugins/Experimental/ControlRig/Source/ControlRigEditor/Private/Graph/NodeSpawners/ControlRigPrototypeNodeSpawner.cpp b/Engine/Plugins/Experimental/ControlRig/Source/ControlRigEditor/Private/Graph/NodeSpawners/ControlRigPrototypeNodeSpawner.cpp
--- a/Engine/Plugins/Experimental/ControlRig/Source/ControlRigEditor/Private/Graph/NodeSpawners/ControlRigPrototypeNodeSpawner.cpp
+++ b/Engine/Plugins/Experimental/ControlRig/Source/ControlRigEditor/Private/Graph/NodeSpawners/ControlRigPrototypeNodeSpawner.cpp
@@ -117,7 +117,7 @@ UControlRigGraphNode* UControlRigPrototypeNodeSpawner::SpawnNode(UEdGraph* Paren
 		}
 
 		FName Name = bIsTemplateNode ? Prototype->GetName() : FControlRigBlueprintUtils::ValidateName(RigBlueprint, Prototype->GetName().ToString());
-		URigVMController* Controller = bIsTemplateNode ? RigGraph->GetTemplateController() : RigBlueprint->Controller;
+		URigVMController* Controller = bIsTemplateNode ? RigGraph->GetTemplateController() : RigBlueprint->GetController(ParentGraph);
 
 		if (!bIsTemplateNode)
 		{
diff --git a/Engine/Plugins/Experimental/ControlRig/Source/ControlRigEditor/Private/Graph/NodeSpawners/ControlRigRerouteNodeSpawner.cpp b/Engine/Plugins/Experimental/ControlRig/Source/ControlRigEditor/Private/Graph/NodeSpawners/ControlRigRerouteNodeSpawner.cpp
--- a/Engine/Plugins/Experimental/ControlRig/Source/ControlRigEditor/Private/Graph/NodeSpawners/ControlRigRerouteNodeSpawner.cpp
+++ b/Engine/Plugins/Experimental/ControlRig/Source/ControlRigEditor/Private/Graph/NodeSpawners/ControlRigRerouteNodeSpawner.cpp
@@ -102,7 +102,7 @@ UEdGraphNode* UControlRigRerouteNodeSpawner::Invoke(UEdGraph* ParentGraph, FBind
 	}
 	else
 	{
-		URigVMController* Controller = RigBlueprint->Controller;
+		URigVMController* Controller = RigBlueprint->GetController(ParentGraph);
 		Controller->OpenUndoBracket(TEXT("Added Reroute Node."));
 
 		FString PinPath;
@@ -114,7 +114,7 @@ UEdGraphNode* UControlRigRerouteNodeSpawner::Invoke(UEdGraph* ParentGraph, FBind
 
 			if (const UEdGraphPin* LastPin = RigSchema->LastPinForCompatibleCheck)
 			{
-				if (URigVMPin* ModelPin = RigBlueprint->Model->FindPin(LastPin->GetName()))
+				if (URigVMPin* ModelPin = RigBlueprint->GetModel(ParentGraph)->FindPin(LastPin->GetName()))
 				{
 					PinPath = ModelPin->GetPinPath();
 				}
diff --git a/Engine/Plugins/Experimental/ControlRig/Source/ControlRigEditor/Private/Graph/NodeSpawners/ControlRigSelectNodeSpawner.cpp b/Engine/Plugins/Experimental/ControlRig/Source/ControlRigEditor/Private/Graph/NodeSpawners/ControlRigSelectNodeSpawner.cpp
--- a/Engine/Plugins/Experimental/ControlRig/Source/ControlRigEditor/Private/Graph/NodeSpawners/ControlRigSelectNodeSpawner.cpp
+++ b/Engine/Plugins/Experimental/ControlRig/Source/ControlRigEditor/Private/Graph/NodeSpawners/ControlRigSelectNodeSpawner.cpp
@@ -92,7 +92,7 @@ UEdGraphNode* UControlRigSelectNodeSpawner::Invoke(UEdGraph* ParentGraph, FBindi
 	{
 		if (const UEdGraphPin* LastPin = RigSchema->LastPinForCompatibleCheck)
 		{
-			if (URigVMPin* ModelPin = RigBlueprint->Model->FindPin(LastPin->GetName()))
+			if (URigVMPin* ModelPin = RigBlueprint->GetModel(ParentGraph)->FindPin(LastPin->GetName()))
 			{
 				CPPType = ModelPin->GetCPPType();
 				if (ModelPin->GetCPPTypeObject())
@@ -112,7 +112,7 @@ UEdGraphNode* UControlRigSelectNodeSpawner::Invoke(UEdGraph* ParentGraph, FBindi
 	}
 #endif
 
-	URigVMController* Controller = bIsTemplateNode ? RigGraph->GetTemplateController() : RigBlueprint->Controller;
+	URigVMController* Controller = bIsTemplateNode ? RigGraph->GetTemplateController() : RigBlueprint->GetController(ParentGraph);
 
 	FName Name = *URigVMSelectNode::SelectName;
 
@@ -160,7 +160,8 @@ bool UControlRigSelectNodeSpawner::IsTemplateNodeFilteredOut(FBlueprintActionFil
 			}
 
 			FString PinPath = Filter.Context.Pins[0]->GetName();
-			if (URigVMPin* ModelPin = RigBlueprint->Model->FindPin(PinPath))
+			UEdGraph* EdGraph = Filter.Context.Pins[0]->GetOwningNode()->GetGraph();
+			if (URigVMPin* ModelPin = RigBlueprint->GetModel(EdGraph)->FindPin(PinPath))
 			{
 				return ModelPin->IsExecuteContext();
 			}
diff --git a/Engine/Plugins/Experimental/ControlRig/Source/ControlRigEditor/Private/Graph/NodeSpawners/ControlRigUnitNodeSpawner.cpp b/Engine/Plugins/Experimental/ControlRig/Source/ControlRigEditor/Private/Graph/NodeSpawners/ControlRigUnitNodeSpawner.cpp
--- a/Engine/Plugins/Experimental/ControlRig/Source/ControlRigEditor/Private/Graph/NodeSpawners/ControlRigUnitNodeSpawner.cpp
+++ b/Engine/Plugins/Experimental/ControlRig/Source/ControlRigEditor/Private/Graph/NodeSpawners/ControlRigUnitNodeSpawner.cpp
@@ -134,7 +134,7 @@ UControlRigGraphNode* UControlRigUnitNodeSpawner::SpawnNode(UEdGraph* ParentGrap
 		bool const bUndo = !bIsTemplateNode;
 
 		FName Name = bIsTemplateNode ? *StructTemplate->GetDisplayNameText().ToString() : FControlRigBlueprintUtils::ValidateName(RigBlueprint, StructTemplate->GetFName().ToString());
-		URigVMController* Controller = bIsTemplateNode ? RigGraph->GetTemplateController() : RigBlueprint->Controller;
+		URigVMController* Controller = bIsTemplateNode ? RigGraph->GetTemplateController() : RigBlueprint->GetController(ParentGraph);
 
 		if (!bIsTemplateNode)
 		{
@@ -230,7 +230,7 @@ UControlRigGraphNode* UControlRigUnitNodeSpawner::SpawnNode(UEdGraph* ParentGrap
 
 void UControlRigUnitNodeSpawner::HookupMutableNode(URigVMNode* InModelNode, UControlRigBlueprint* InRigBlueprint)
 {
-	URigVMController* Controller = InRigBlueprint->Controller;
+	URigVMController* Controller = InRigBlueprint->GetController(InModelNode->GetGraph());
 
 	Controller->ClearNodeSelection(true);
 	Controller->SelectNode(InModelNode, true, true);
diff --git a/Engine/Plugins/Experimental/ControlRig/Source/ControlRigEditor/Private/Graph/NodeSpawners/ControlRigVariableNodeSpawner.cpp b/Engine/Plugins/Experimental/ControlRig/Source/ControlRigEditor/Private/Graph/NodeSpawners/ControlRigVariableNodeSpawner.cpp
--- a/Engine/Plugins/Experimental/ControlRig/Source/ControlRigEditor/Private/Graph/NodeSpawners/ControlRigVariableNodeSpawner.cpp
+++ b/Engine/Plugins/Experimental/ControlRig/Source/ControlRigEditor/Private/Graph/NodeSpawners/ControlRigVariableNodeSpawner.cpp
@@ -96,7 +96,7 @@ UEdGraphNode* UControlRigVariableNodeSpawner::Invoke(UEdGraph* ParentGraph, FBin
 	}
 #endif
 
-	URigVMController* Controller = bIsTemplateNode ? RigGraph->GetTemplateController() : RigBlueprint->Controller;
+	URigVMController* Controller = bIsTemplateNode ? RigGraph->GetTemplateController() : RigBlueprint->GetController(ParentGraph);
 
 	if (!bIsTemplateNode)
 	{
diff --git a/Engine/Plugins/Experimental/ControlRig/Source/ControlRigEditor/Private/Graph/SControlRigGraphNode.cpp b/Engine/Plugins/Experimental/ControlRig/Source/ControlRigEditor/Private/Graph/SControlRigGraphNode.cpp
--- a/Engine/Plugins/Experimental/ControlRig/Source/ControlRigEditor/Private/Graph/SControlRigGraphNode.cpp
+++ b/Engine/Plugins/Experimental/ControlRig/Source/ControlRigEditor/Private/Graph/SControlRigGraphNode.cpp
@@ -773,13 +773,9 @@ void SControlRigGraphNode::HandleGetChildrenForTree(URigVMPin* InItem, TArray<UR
 
 void SControlRigGraphNode::HandleExpansionChanged(URigVMPin* InItem, bool bExpanded)
 {
-	if (GraphNode)
+	if (UControlRigGraphNode* RigGraphNode = Cast< UControlRigGraphNode>(GraphNode))
 	{
-		UControlRigBlueprint* ControlRigBlueprint = Cast<UControlRigBlueprint>(GraphNode->GetGraph()->GetOuter());
-		if (ControlRigBlueprint)
-		{
-			ControlRigBlueprint->Controller->SetPinExpansion(InItem->GetPinPath(), bExpanded, true);
-		}
+		RigGraphNode->GetController()->SetPinExpansion(InItem->GetPinPath(), bExpanded, true);
 	}
 }
 
@@ -787,38 +783,34 @@ void SControlRigGraphNode::HandleExpandRecursively(URigVMPin* InItem, bool bExpa
 {
 	TSharedPtr<STreeView<URigVMPin*>>& TreeWidget = *TreeWidgetPtr;
 
-	if (GraphNode)
+	if (UControlRigGraphNode* RigGraphNode = Cast< UControlRigGraphNode>(GraphNode))
 	{
-		UControlRigBlueprint* ControlRigBlueprint = Cast<UControlRigBlueprint>(GraphNode->GetGraph()->GetOuter());
-		if (ControlRigBlueprint)
-		{
-			ControlRigBlueprint->Controller->OpenUndoBracket(TEXT("Expand pin recursively"));
+		RigGraphNode->GetController()->OpenUndoBracket(TEXT("Expand pin recursively"));
 
-			TArray<URigVMPin*> ModelPins;
-			ModelPins.Add(InItem);
+		TArray<URigVMPin*> ModelPins;
+		ModelPins.Add(InItem);
 
-			for (int32 PinIndex = 0; PinIndex < ModelPins.Num(); PinIndex++)
-			{
-				URigVMPin* ModelPin = ModelPins[PinIndex];
-				ModelPins.Append(ModelPin->GetSubPins());
-			}
+		for (int32 PinIndex = 0; PinIndex < ModelPins.Num(); PinIndex++)
+		{
+			URigVMPin* ModelPin = ModelPins[PinIndex];
+			ModelPins.Append(ModelPin->GetSubPins());
+		}
 
-			if (!bExpanded)
-			{
-				Algo::Reverse(ModelPins);
-			}
+		if (!bExpanded)
+		{
+			Algo::Reverse(ModelPins);
+		}
 
-			for (int32 PinIndex = 0; PinIndex < ModelPins.Num(); PinIndex++)
+		for (int32 PinIndex = 0; PinIndex < ModelPins.Num(); PinIndex++)
+		{
+			URigVMPin* ModelPin = ModelPins[PinIndex];
+			if (ModelPin->IsExpanded() != bExpanded && ModelPin->GetSubPins().Num() > 0)
 			{
-				URigVMPin* ModelPin = ModelPins[PinIndex];
-				if (ModelPin->IsExpanded() != bExpanded && ModelPin->GetSubPins().Num() > 0)
-				{
-					TreeWidget->SetItemExpansion(ModelPin, bExpanded);
-				}
+				TreeWidget->SetItemExpansion(ModelPin, bExpanded);
 			}
-
-			ControlRigBlueprint->Controller->CloseUndoBracket();
 		}
+
+		RigGraphNode->GetController()->CloseUndoBracket();
 	}
 }
 
@@ -873,10 +865,11 @@ void SControlRigGraphNode::GetNodeInfoPopups(FNodeInfoContext* Context, TArray<F
 	const FLinearColor PinnedWatchColor(0.35f, 0.25f, 0.25f);
 
 	UControlRig* ActiveObject = Cast<UControlRig>(K2Context->ActiveObjectBeingDebugged);
+	UControlRigGraphNode* RigGraphNode = Cast<UControlRigGraphNode>(GraphNode);
 	UControlRigBlueprint* RigBlueprint = Cast<UControlRigBlueprint>(K2Context->SourceBlueprint);
 
 	// Display any pending latent actions
-	if (ActiveObject && RigBlueprint)
+	if (ActiveObject && RigBlueprint && RigGraphNode)
 	{
 		// Display pinned watches
 		if (K2Context->WatchedNodeSet.Contains(GraphNode))
@@ -890,7 +883,7 @@ void SControlRigGraphNode::GetNodeInfoPopups(FNodeInfoContext* Context, TArray<F
 				UEdGraphPin* WatchPin = GraphNode->Pins[PinIndex];
 				if (K2Context->WatchedPinSet.Contains(WatchPin))
 				{
-					if (URigVMPin* ModelPin = RigBlueprint->Model->FindPin(WatchPin->GetName()))
+					if (URigVMPin* ModelPin = RigGraphNode->GetModel()->FindPin(WatchPin->GetName()))
 					{
 						if (ValidWatchCount > 0)
 						{
diff --git a/Engine/Plugins/Experimental/ControlRig/Source/ControlRigEditor/Private/Graph/SControlRigGraphNodeComment.cpp b/Engine/Plugins/Experimental/ControlRig/Source/ControlRigEditor/Private/Graph/SControlRigGraphNodeComment.cpp
--- a/Engine/Plugins/Experimental/ControlRig/Source/ControlRigEditor/Private/Graph/SControlRigGraphNodeComment.cpp
+++ b/Engine/Plugins/Experimental/ControlRig/Source/ControlRigEditor/Private/Graph/SControlRigGraphNodeComment.cpp
@@ -42,10 +42,13 @@ FReply SControlRigGraphNodeComment::OnMouseButtonUp(const FGeometry& MyGeometry,
 				{
 					FVector2D Position(CommentNode->NodePosX, CommentNode->NodePosY);
 					FVector2D Size(CommentNode->NodeWidth, CommentNode->NodeHeight);
-					Blueprint->Controller->OpenUndoBracket(TEXT("Resize Comment Box"));
-					Blueprint->Controller->SetNodePositionByName(CommentNode->GetFName(), Position, true);
-					Blueprint->Controller->SetNodeSizeByName(CommentNode->GetFName(), Size, true);
-					Blueprint->Controller->CloseUndoBracket();
+					if (URigVMController* Controller = Blueprint->GetController(Graph))
+					{
+						Controller->OpenUndoBracket(TEXT("Resize Comment Box"));
+						Controller->SetNodePositionByName(CommentNode->GetFName(), Position, true);
+						Controller->SetNodeSizeByName(CommentNode->GetFName(), Size, true);
+						Controller->CloseUndoBracket();
+					}
 				}
 			}
 		} 
@@ -90,7 +93,10 @@ void SControlRigGraphNodeComment::Tick(const FGeometry& AllottedGeometry, const
 			{
 				if (UControlRigBlueprint* Blueprint = Cast<UControlRigBlueprint>(Graph->GetOuter()))
 				{
-					Blueprint->Controller->SetCommentTextByName(CommentNode->GetFName(), CurrentCommentTitle, true);
+					if (URigVMController* Controller = Blueprint->GetController(Graph))
+					{
+						Controller->SetCommentTextByName(CommentNode->GetFName(), CurrentCommentTitle, true);
+					}
 				}
 			}
 		}
@@ -108,9 +114,12 @@ void SControlRigGraphNodeComment::Tick(const FGeometry& AllottedGeometry, const
 				{
 					if (UControlRigBlueprint* Blueprint = Cast<UControlRigBlueprint>(Graph->GetOuter()))
 					{
-						// for now we won't use our undo for this kind of change
-						Blueprint->Controller->SetNodeColorByName(GraphNode->GetFName(), CurrentNodeCommentColor, false, true);
-						CachedNodeCommentColor = CurrentNodeCommentColor;
+						if (URigVMController* Controller = Blueprint->GetController(Graph))
+						{
+							// for now we won't use our undo for this kind of change
+							Controller->SetNodeColorByName(GraphNode->GetFName(), CurrentNodeCommentColor, false, true);
+							CachedNodeCommentColor = CurrentNodeCommentColor;
+						}
 					}
 				}
 			}
diff --git a/Engine/Plugins/Experimental/ControlRig/Source/ControlRigEditor/Private/Graph/SControlRigGraphPinCurveFloat.cpp b/Engine/Plugins/Experimental/ControlRig/Source/ControlRigEditor/Private/Graph/SControlRigGraphPinCurveFloat.cpp
--- a/Engine/Plugins/Experimental/ControlRig/Source/ControlRigEditor/Private/Graph/SControlRigGraphPinCurveFloat.cpp
+++ b/Engine/Plugins/Experimental/ControlRig/Source/ControlRigEditor/Private/Graph/SControlRigGraphPinCurveFloat.cpp
@@ -130,13 +130,10 @@ void SControlRigGraphPinCurveFloat::OnCurveChanged(const TArray<FRichCurveEditIn
 	{
 		if (UControlRigGraphNode* Node = Cast<UControlRigGraphNode>(Pin->GetOwningNode()))
 		{
-			if (UControlRigBlueprint* RigBlueprint = Cast<UControlRigBlueprint>(Node->GetGraph()->GetOuter()))
-			{
-				FString ExportedText;
-				FRuntimeFloatCurve DefaultCurve;
-				FRuntimeFloatCurve::StaticStruct()->ExportText(ExportedText, &Curve, &DefaultCurve, nullptr, EPropertyPortFlags::PPF_None, nullptr, true);
-				RigBlueprint->Controller->SetPinDefaultValue(Pin->GetName(), ExportedText, true, true, true);
-			}
+			FString ExportedText;
+			FRuntimeFloatCurve DefaultCurve;
+			FRuntimeFloatCurve::StaticStruct()->ExportText(ExportedText, &Curve, &DefaultCurve, nullptr, EPropertyPortFlags::PPF_None, nullptr, true);
+			Node->GetController()->SetPinDefaultValue(Pin->GetName(), ExportedText, true, true, true);
 		}
 	}
 	ModifyOwner();
diff --git a/Engine/Plugins/Experimental/ControlRig/Source/ControlRigEditor/Private/Graph/SControlRigGraphPinVariableBinding.cpp b/Engine/Plugins/Experimental/ControlRig/Source/ControlRigEditor/Private/Graph/SControlRigGraphPinVariableBinding.cpp
--- a/Engine/Plugins/Experimental/ControlRig/Source/ControlRigEditor/Private/Graph/SControlRigGraphPinVariableBinding.cpp
+++ b/Engine/Plugins/Experimental/ControlRig/Source/ControlRigEditor/Private/Graph/SControlRigGraphPinVariableBinding.cpp
@@ -109,7 +109,7 @@ void SControlRigVariableBinding::OnAddBinding(FName InPropertyName, const TArray
 			ensure(ChainElement.Field);
 			Parts.Add(ChainElement.Field.GetName());
 		}
-		Blueprint->Controller->BindPinToVariable(ModelPin->GetPinPath(), FString::Join(Parts, TEXT(".")), true /* undo */);
+		Blueprint->GetController(ModelPin->GetGraph())->BindPinToVariable(ModelPin->GetPinPath(), FString::Join(Parts, TEXT(".")), true /* undo */);
 	}
 }
 
@@ -122,7 +122,7 @@ void SControlRigVariableBinding::OnRemoveBinding(FName InPropertyName)
 {
 	if (Blueprint && ModelPin)
 	{
-		Blueprint->Controller->UnbindPinFromVariable(ModelPin->GetPinPath(), true /* undo */);
+		Blueprint->GetController(ModelPin->GetGraph())->UnbindPinFromVariable(ModelPin->GetPinPath(), true /* undo */);
 	}
 }
 
diff --git a/Engine/Source/Developer/RigVMDeveloper/Private/RigVMModel/RigVMGraph.cpp b/Engine/Source/Developer/RigVMDeveloper/Private/RigVMModel/RigVMGraph.cpp
--- a/Engine/Source/Developer/RigVMDeveloper/Private/RigVMModel/RigVMGraph.cpp
+++ b/Engine/Source/Developer/RigVMDeveloper/Private/RigVMModel/RigVMGraph.cpp
@@ -2,6 +2,7 @@
 
 #include "RigVMModel/RigVMGraph.h"
 #include "RigVMModel/RigVMController.h"
+#include "RigVMModel/Nodes/RigVMLibraryNode.h"
 #include "UObject/Package.h"
 
 URigVMGraph::URigVMGraph()
@@ -20,6 +21,23 @@ const TArray<URigVMLink*>& URigVMGraph::GetLinks() const
 	return Links;
 }
 
+TArray<URigVMGraph*> URigVMGraph::GetContainedGraphs(bool bRecursive) const
+{
+	TArray<URigVMGraph*> Graphs;
+	for (URigVMNode* Node : GetNodes())
+	{
+		if (URigVMLibraryNode* LibraryNode = Cast<URigVMLibraryNode>(Node))
+		{
+			Graphs.Add(LibraryNode->GetContainedGraph());
+			if (bRecursive)
+			{
+				Graphs.Append(LibraryNode->GetContainedGraph()->GetContainedGraphs(true));
+			}
+		}
+	}
+	return Graphs;
+}
+
 TArray<FRigVMGraphVariableDescription> URigVMGraph::GetVariableDescriptions() const
 {
 	TArray<FRigVMGraphVariableDescription> Variables;
diff --git a/Engine/Source/Developer/RigVMDeveloper/Public/RigVMModel/RigVMGraph.h b/Engine/Source/Developer/RigVMDeveloper/Public/RigVMModel/RigVMGraph.h
--- a/Engine/Source/Developer/RigVMDeveloper/Public/RigVMModel/RigVMGraph.h
+++ b/Engine/Source/Developer/RigVMDeveloper/Public/RigVMModel/RigVMGraph.h
@@ -36,6 +36,10 @@ public:
 	UFUNCTION(BlueprintCallable, Category = RigVMGraph)
 	const TArray<URigVMLink*>& GetLinks() const;
 
+	// Returns all of the contained graphs
+	UFUNCTION(BlueprintCallable, Category = RigVMGraph)
+	TArray<URigVMGraph*> GetContainedGraphs(bool bRecursive = false) const;
+
 	// Returns a list of unique Variable descriptions within this Graph.
 	// Multiple Variable Nodes can share the same description.
 	UFUNCTION(BlueprintCallable, Category = RigVMGraph)
