Control Rig: Added missing reset transform in FControlRigEditor::Tick(). When there is no skeletal mesh, control rig is evaluated each frame at FControlRigEditor::Tick() instead of at FAnimNode_ControlRigBase::ExecuteControlRig(). BoneHierarchy.ResetTranform was called only in FAnimNode_ControlRigBase::ExecuteControlRig() but not in FControlRigEditor::Tick()

diff --git a/Engine/Plugins/Experimental/ControlRig/Source/ControlRigEditor/Private/Editor/ControlRigEditor.cpp b/Engine/Plugins/Experimental/ControlRig/Source/ControlRigEditor/Private/Editor/ControlRigEditor.cpp
--- a/Engine/Plugins/Experimental/ControlRig/Source/ControlRigEditor/Private/Editor/ControlRigEditor.cpp
+++ b/Engine/Plugins/Experimental/ControlRig/Source/ControlRigEditor/Private/Editor/ControlRigEditor.cpp
@@ -1858,6 +1858,9 @@ void FControlRigEditor::Tick(float DeltaTime)
 			ControlRig != nullptr && 
 			bExecutionControlRig)
 		{
+			// reset transforms here to prevent additive transforms from accumulating to INF
+			ControlRig->GetBoneHierarchy().ResetTransforms();
+			
 			ControlRig->SetDeltaTime(DeltaTime);
 			ControlRig->Evaluate_AnyThread();
 			bDrawHierarchyBones = true;
