Fixed Mirroring. Now use similar setup to Control Rig with just support for each axis insread of an aribtrary plane.

Added ability to Rename Controls in the Pose. Needed in case the rig changes name on it's controls, needed a fast way to select a whole bunch of poses and rename them.

Added set of Asset Actions for the right click context menu in the Asset Browser.

Now Support multiple folders to filter on with new Add Folder To View context menu item. This will also add the folder to the project settings so it shows up on in the next session also.

Added ability to Update a Pose, which basically does a save over the new pose. Will popup a dialog so a user can confirm this destructive behavior.

Fixed issue with the thumbnail not showing up in the Pose View, main issue that it seems we had to use the level editor thumbnail pool.

Made Pose Name in Pose View text box and not editable for now.

diff --git a/Engine/Plugins/Experimental/ControlRig/Source/ControlRig/Private/Tools/ControlRigPose.cpp b/Engine/Plugins/Experimental/ControlRig/Source/ControlRig/Private/Tools/ControlRigPose.cpp
--- a/Engine/Plugins/Experimental/ControlRig/Source/ControlRig/Private/Tools/ControlRigPose.cpp
+++ b/Engine/Plugins/Experimental/ControlRig/Source/ControlRig/Private/Tools/ControlRigPose.cpp
@@ -67,5 +67,15 @@ void UControlRigPoseAsset::BlendWithInitialPoses(FControlRigControlPose& Initial
 	}
 }
 
+TArray<FName> UControlRigPoseAsset::GetControlNames() const
+{
+	return Pose.GetControlNames();
+}
+
+void UControlRigPoseAsset::ReplaceControlName(const FName& CurrentName, const FName& NewName)
+{
+	Pose.ReplaceControlName(CurrentName, NewName);
+}
+
 #undef LOCTEXT_NAMESPACE
 
diff --git a/Engine/Plugins/Experimental/ControlRig/Source/ControlRig/Private/Tools/ControlRigPoseMirrorSettings.cpp b/Engine/Plugins/Experimental/ControlRig/Source/ControlRig/Private/Tools/ControlRigPoseMirrorSettings.cpp
--- a/Engine/Plugins/Experimental/ControlRig/Source/ControlRig/Private/Tools/ControlRigPoseMirrorSettings.cpp
+++ b/Engine/Plugins/Experimental/ControlRig/Source/ControlRig/Private/Tools/ControlRigPoseMirrorSettings.cpp
@@ -5,7 +5,6 @@ UControlRigPoseMirrorSettings::UControlRigPoseMirrorSettings()
 {
 	RightSide = TEXT("Right");
 	LeftSide = TEXT("Left");
-	XAxis = 1.0f;
-	YAxis = 0.0f;
-	ZAxis = 0.0f;
+	MirrorAxis = EAxis::X;
+	AxisToFlip = EAxis::Z;
 }
diff --git a/Engine/Plugins/Experimental/ControlRig/Source/ControlRig/Private/Tools/ControlRigPoseMirrorTable.cpp b/Engine/Plugins/Experimental/ControlRig/Source/ControlRig/Private/Tools/ControlRigPoseMirrorTable.cpp
--- a/Engine/Plugins/Experimental/ControlRig/Source/ControlRig/Private/Tools/ControlRigPoseMirrorTable.cpp
+++ b/Engine/Plugins/Experimental/ControlRig/Source/ControlRig/Private/Tools/ControlRigPoseMirrorTable.cpp
@@ -33,24 +33,10 @@ void FControlRigPoseMirrorTable::SetUpMirrorTable(const UControlRig* ControlRig)
 	}
 }
 
-
 FRigControlCopy* FControlRigPoseMirrorTable::GetControl(FControlRigControlPose& Pose, FName Name)
 {
-
 	TArray<FRigControlCopy> CopyOfControls;
-
-	//Cache of the Map, Used to make pasting faster.
-	TMap<FName, int32>  CopyOfControlsNameToIndex;
-
-	if (MatchedControls.Num() <= 0)
-	{
-		int32* Index = Pose.CopyOfControlsNameToIndex.Find(Name);
-		if (Index != nullptr && (*Index) >= 0 && (*Index) < Pose.CopyOfControls.Num())
-		{
-			return &(Pose.CopyOfControls[*Index]);
-		}
 	
-	}
 	if (MatchedControls.Num() > 0) 
 	{
 		if (const FName* MatchedName = MatchedControls.Find(Name))
@@ -63,6 +49,12 @@ FRigControlCopy* FControlRigPoseMirrorTable::GetControl(FControlRigControlPose&
 			}
 		}
 	}
+	//okay not matched so just find it.
+	int32* Index = Pose.CopyOfControlsNameToIndex.Find(Name);
+	if (Index != nullptr && (*Index) >= 0 && (*Index) < Pose.CopyOfControls.Num())
+	{
+		return &(Pose.CopyOfControls[*Index]);
+	}
 
 	return nullptr;
 }
@@ -79,48 +71,29 @@ bool FControlRigPoseMirrorTable::IsMatched(const FName& Name) const
 	return false;
 }
 
-// Mirrors Translation as Global(Component)
-// Mirrors Rotation as Local , only if NOT matched, if matched we just use it
-void FControlRigPoseMirrorTable::GetMirrorTransform(const FRigControlCopy& ControlCopy, bool bIsMatched, FVector& OutGlobalTranslation, FQuat& OutGlobalRotation) const
+//Now returns mirrored global and local unmirrored
+void FControlRigPoseMirrorTable::GetMirrorTransform(const FRigControlCopy& ControlCopy, bool bIsMatched, FVector& OutGlobalTranslation,
+	FQuat& OutGlobalRotation, FQuat& OutLocalRotation) const
 {
 	const UControlRigPoseMirrorSettings* Settings = GetDefault<UControlRigPoseMirrorSettings>();
-	if (Settings)
-	{
-		FTransform GlobalTransform = ControlCopy.GlobalTransform;
-		FTransform LocalTransform = ControlCopy.LocalTransform;
-		//FQuat Rotation = LocalTransform.GetRotation();
-		FQuat Rotation = GlobalTransform.GetRotation();
+	FTransform GlobalTransform = ControlCopy.GlobalTransform;
+	FQuat GlobalRotation = GlobalTransform.GetRotation();
 
-		FVector Axis(Settings->XAxis, Settings->YAxis, Settings->ZAxis);
-		//if axis nearly zero just return
-		if (Axis.IsNearlyZero())
-		{
-			OutGlobalTranslation = ControlCopy.GlobalTransform.GetTranslation();
-			OutGlobalRotation = Rotation;
-			return;
-		}
-		Axis.Normalize();
+	FTransform LocalTransform = ControlCopy.LocalTransform;
+	OutLocalRotation = LocalTransform.GetRotation();
+	OutGlobalTranslation = GlobalTransform.GetTranslation();
+	OutGlobalRotation = GlobalTransform.GetRotation();
+	
+	if (Settings &&  (bIsMatched || (OutGlobalRotation.IsIdentity() == false && OutLocalRotation.IsIdentity() == false)))
+	{
+		FRigMirrorSettings MirrorSettings;
 
-		FVector Translation = ControlCopy.GlobalTransform.GetTranslation();
-		FPlane Plane(FVector::ZeroVector,Axis);
-		OutGlobalTranslation = Translation.MirrorByPlane(Plane);
+		MirrorSettings.MirrorAxis = Settings->MirrorAxis;
+		MirrorSettings.AxisToFlip = Settings->AxisToFlip;
 
-		//test FRotator OldRotator = Rotation.Rotator();
-		
-		if (bIsMatched)
-		{
-			OutGlobalRotation = Rotation;
-		}
-		else
-		{
-			FQuat MirrorNormalQuat = FQuat(Axis.X, Axis.Y, Axis.Z, 0);
-			Rotation.EnforceShortestArcWith(MirrorNormalQuat);
-			OutGlobalRotation = MirrorNormalQuat * Rotation * MirrorNormalQuat;
-		}
-		//test FRotator NewRotator = OutGlobalRotation.Rotator();
-		
+		FTransform NewTransform = MirrorSettings.MirrorTransform(GlobalTransform);
+		OutGlobalTranslation = NewTransform.GetTranslation();
+		OutGlobalRotation = NewTransform.GetRotation();
 		return;
 	}
-	OutGlobalTranslation = FVector::ZeroVector;
-	OutGlobalRotation = FQuat::Identity;
 }
diff --git a/Engine/Plugins/Experimental/ControlRig/Source/ControlRig/Private/Tools/ControlRigPoseProjectSettings.cpp b/Engine/Plugins/Experimental/ControlRig/Source/ControlRig/Private/Tools/ControlRigPoseProjectSettings.cpp
--- a/Engine/Plugins/Experimental/ControlRig/Source/ControlRig/Private/Tools/ControlRigPoseProjectSettings.cpp
+++ b/Engine/Plugins/Experimental/ControlRig/Source/ControlRig/Private/Tools/ControlRigPoseProjectSettings.cpp
@@ -3,5 +3,17 @@
 
 UControlRigPoseProjectSettings::UControlRigPoseProjectSettings()
 {
+	FDirectoryPath RootSaveDir;
  	RootSaveDir.Path = TEXT("/Game/ControlRig/Pose");
+	RootSaveDirs.Add(RootSaveDir);
+}
+
+TArray<FString> UControlRigPoseProjectSettings::GetAssetPaths() const
+{
+	TArray<FString> Paths;
+	for (const FDirectoryPath& Path : RootSaveDirs)
+	{
+		Paths.Add(Path.Path);
+	}
+	return Paths;
 }
diff --git a/Engine/Plugins/Experimental/ControlRig/Source/ControlRig/Public/Tools/ControlRigPose.h b/Engine/Plugins/Experimental/ControlRig/Source/ControlRig/Public/Tools/ControlRigPose.h
--- a/Engine/Plugins/Experimental/ControlRig/Source/ControlRig/Public/Tools/ControlRigPose.h
+++ b/Engine/Plugins/Experimental/ControlRig/Source/ControlRig/Public/Tools/ControlRigPose.h
@@ -119,10 +119,22 @@ struct FControlRigControlPose
 
 	}
 
-	void SetControlMirrorTransform(UControlRig* ControlRig, const FName& Name, const FVector& GlobalTranslation, const FQuat& GlobalRotation, bool bNotify, FRigControlModifiedContext Context)
+	void SetControlMirrorTransform(UControlRig* ControlRig, const FName& Name, bool bIsMatched, const FVector& GlobalTranslation, const FQuat& GlobalRotation, const FQuat&LocalRotation, bool bNotify,const  FRigControlModifiedContext& Context)
 	{
-		FTransform NewGlobalTransform(GlobalRotation, GlobalTranslation);
-		ControlRig->SetControlGlobalTransform(Name, NewGlobalTransform, bNotify, Context);
+		if (bIsMatched)
+		{
+			int32 Index = ControlRig->GetControlHierarchy().GetIndex(Name);
+			FTransform ParentTransform = ControlRig->GetControlHierarchy().GetParentTransform(Index);
+			FTransform CurrentTransform = ControlRig->GetControlHierarchy().GetLocalTransform(Index);
+			FVector NewLocal = ParentTransform.InverseTransformPositionNoScale(GlobalTranslation);
+			FTransform NewCurrentTransform(LocalRotation, NewLocal);
+			ControlRig->SetControlLocalTransform(Name, NewCurrentTransform, bNotify, Context);
+		}
+		else
+		{
+			FTransform NewGlobalTransform(GlobalRotation, GlobalTranslation);
+			ControlRig->SetControlGlobalTransform(Name, NewGlobalTransform, bNotify, Context);
+		}
 	}
 
 	void PastePoseInternal(UControlRig* ControlRig, bool bDoKey,bool bDoMirror, const TArray<FRigControlCopy>& ControlsToPaste)
@@ -165,9 +177,10 @@ struct FControlRigControlPose
 						{
 							FVector GlobalTranslation;
 							FQuat GlobalRotation;
+							FQuat LocalRotation;
 							bool bIsMatched = MirrorTable.IsMatched(CopyRigControl->Name);
-							MirrorTable.GetMirrorTransform(*CopyRigControl, bIsMatched,GlobalTranslation, GlobalRotation);
-							SetControlMirrorTransform(ControlRig,RigControl.Name, GlobalTranslation, GlobalRotation,true, Context);
+							MirrorTable.GetMirrorTransform(*CopyRigControl, bIsMatched,GlobalTranslation, GlobalRotation, LocalRotation);
+							SetControlMirrorTransform(ControlRig,RigControl.Name,bIsMatched, GlobalTranslation, GlobalRotation, LocalRotation,true, Context);
 						}
 						break;
 					}		
@@ -264,13 +277,18 @@ struct FControlRigControlPose
 							{
 								FVector GlobalTranslation;
 								FQuat GlobalRotation;
+								FQuat LocalRotation;
 								bool bIsMatched = MirrorTable.IsMatched(CopyRigControl->Name);
-								MirrorTable.GetMirrorTransform(*CopyRigControl, bIsMatched, GlobalTranslation, GlobalRotation);
+								MirrorTable.GetMirrorTransform(*CopyRigControl, bIsMatched, GlobalTranslation, GlobalRotation, LocalRotation);
 								FVector InitialTranslation = InitialFound->GlobalTransform.GetTranslation();
-								FQuat InitialRotation = InitialFound->LocalTransform.GetRotation();
+								FQuat InitialGlobalRotation = InitialFound->GlobalTransform.GetRotation();
+								FQuat InitialLocationRotation = InitialFound->LocalTransform.GetRotation();
 								GlobalTranslation = FMath::Lerp(InitialTranslation, GlobalTranslation, BlendValue);
-								GlobalRotation = FQuat::Slerp(InitialRotation, GlobalRotation, BlendValue); //doing slerp here not fast lerp, can be slow this is for content creation
-								SetControlMirrorTransform(ControlRig, RigControl.Name, GlobalTranslation, GlobalRotation, false, Context);
+								GlobalRotation = FQuat::Slerp(InitialGlobalRotation, GlobalRotation, BlendValue); //doing slerp here not fast lerp, can be slow this is for content creation
+								LocalRotation = FQuat::Slerp(InitialLocationRotation, LocalRotation, BlendValue); //doing slerp here not fast lerp, can be slow this is for content creation
+
+								SetControlMirrorTransform(ControlRig, RigControl.Name, bIsMatched, GlobalTranslation, GlobalRotation, LocalRotation, false, Context);
+														
 							}
 						}
 					}
@@ -278,6 +296,23 @@ struct FControlRigControlPose
 			}			
 		}
 	}
+	bool ContainsName(const FName& Name) const
+	{
+		const int32* Index = CopyOfControlsNameToIndex.Find(Name);
+		return (Index && *Index >= 0);
+	}
+
+	void ReplaceControlName(const FName& Name, const FName& NewName)
+	{
+		int32* Index = CopyOfControlsNameToIndex.Find(Name);
+		if (Index && *Index >= 0)
+		{
+			FRigControlCopy& Control = CopyOfControls[*Index];
+			Control.Name = NewName;
+			CopyOfControlsNameToIndex.Remove(Name);
+			CopyOfControlsNameToIndex.Add(Control.Name, *Index);
+		}
+	}
 
 	TArray<FName> GetControlNames() const
 	{
@@ -302,6 +337,7 @@ struct FControlRigControlPose
 	}
 	TArray<FRigControlCopy> GetPoses() const {return CopyOfControls;};
 
+
 public:
 	UPROPERTY(BlueprintReadWrite, EditAnywhere, Category = "Controls")
 	TArray<FRigControlCopy> CopyOfControls;
@@ -338,6 +374,12 @@ public:
 	UFUNCTION(BlueprintCallable, Category = "Pose")
 	void GetCurrentPose(UControlRig* InControlRig, FControlRigControlPose& OutPose);
 
+	UFUNCTION(BlueprintCallable, Category = "Pose")
+	TArray<FName> GetControlNames() const;
+
+	UFUNCTION(BlueprintCallable, Category = "Pose")
+	void ReplaceControlName(const FName& CurrentName, const FName& NewName);
+
 	void BlendWithInitialPoses(FControlRigControlPose& InitialPose, UControlRig* ControlRig, bool bDoKey, bool bdoMirror, float BlendValue);
 
 	UPROPERTY(BlueprintReadWrite, EditAnywhere, Category = "Pose")
diff --git a/Engine/Plugins/Experimental/ControlRig/Source/ControlRig/Public/Tools/ControlRigPoseMirrorSettings.h b/Engine/Plugins/Experimental/ControlRig/Source/ControlRig/Public/Tools/ControlRigPoseMirrorSettings.h
--- a/Engine/Plugins/Experimental/ControlRig/Source/ControlRig/Public/Tools/ControlRigPoseMirrorSettings.h
+++ b/Engine/Plugins/Experimental/ControlRig/Source/ControlRig/Public/Tools/ControlRigPoseMirrorSettings.h
@@ -6,6 +6,7 @@
 #include "CoreMinimal.h"
 #include "UObject/ObjectMacros.h"
 #include "UObject/Object.h"
+#include "ControlRig.h"
 #include "ControlRigPoseMirrorSettings.generated.h"
 
 
@@ -26,12 +27,14 @@ public:
 	UPROPERTY(config, EditAnywhere, BlueprintReadWrite, Category = "Mirror Settings", meta = (ToolTip = "String To Denote Left Side"))
 	FString LeftSide;
 
-	UPROPERTY(config, EditAnywhere, BlueprintReadWrite, Category = "Mirror Plane", meta = (ToolTip = "X Axis Plane Normal"))
-	float XAxis;
-	UPROPERTY(config, EditAnywhere, BlueprintReadWrite, Category = "Mirror Plane", meta = (ToolTip = "Y Axis Plane Normal"))
-	float YAxis;
-	UPROPERTY(config, EditAnywhere, BlueprintReadWrite, Category = "Mirror Plane", meta = (ToolTip = "Z Axis Plane Normal"))
-	float ZAxis;
+	// the axis to mirror translations against
+	UPROPERTY(config, EditAnywhere, BlueprintReadWrite, Category = "Mirror Settings", meta = (ToolTip = "Axis to Mirror Translations"))
+	TEnumAsByte<EAxis::Type> MirrorAxis;
+
+	// the axis to flip for rotations
+	UPROPERTY(config, EditAnywhere, BlueprintReadWrite, Category = "Mirror Settings", meta = (ToolTip = "Axis to Flip for Rotations"))
+	TEnumAsByte<EAxis::Type> AxisToFlip;
 
 };
 
+
diff --git a/Engine/Plugins/Experimental/ControlRig/Source/ControlRig/Public/Tools/ControlRigPoseMirrorTable.h b/Engine/Plugins/Experimental/ControlRig/Source/ControlRig/Public/Tools/ControlRigPoseMirrorTable.h
--- a/Engine/Plugins/Experimental/ControlRig/Source/ControlRig/Public/Tools/ControlRigPoseMirrorTable.h
+++ b/Engine/Plugins/Experimental/ControlRig/Source/ControlRig/Public/Tools/ControlRigPoseMirrorTable.h
@@ -31,7 +31,7 @@ public:
 	bool IsMatched(const FName& ControlName) const;
 
 	/*Return the Mirrored Global(In Control Rig Space) Translation and Mirrored Local Rotation*/
-	void GetMirrorTransform(const FRigControlCopy& ControlCopy,  bool bIsMatched,FVector& OutGlobalTranslation, FQuat& OutLocalRotation) const;
+	void GetMirrorTransform(const FRigControlCopy& ControlCopy,  bool bIsMatched,FVector& OutGlobalTranslation, FQuat& OutGlobalRotation, FQuat& OutLocalRotation) const;
 
 private:
 
diff --git a/Engine/Plugins/Experimental/ControlRig/Source/ControlRig/Public/Tools/ControlRigPoseProjectSettings.h b/Engine/Plugins/Experimental/ControlRig/Source/ControlRig/Public/Tools/ControlRigPoseProjectSettings.h
--- a/Engine/Plugins/Experimental/ControlRig/Source/ControlRig/Public/Tools/ControlRigPoseProjectSettings.h
+++ b/Engine/Plugins/Experimental/ControlRig/Source/ControlRig/Public/Tools/ControlRigPoseProjectSettings.h
@@ -19,11 +19,11 @@ public:
 	UControlRigPoseProjectSettings();
 
 	/** The pose asset path  */
-	FString GetAssetPath() const { return RootSaveDir.Path; }
+	TArray<FString> GetAssetPaths() const;
 
 	/** The root of the directory in which to save poses */
 	UPROPERTY(config, EditAnywhere, BlueprintReadWrite, Category = "Control Rig Poses", meta = (ContentDir))
-	FDirectoryPath RootSaveDir;
+	TArray<FDirectoryPath> RootSaveDirs;
 
 	/** Not used but may put bad if we support other types.
 	bool bFilterPoses;
diff --git a/Engine/Plugins/Experimental/ControlRig/Source/ControlRigEditor/Private/EditMode/SControlRigBaseListWidget.cpp b/Engine/Plugins/Experimental/ControlRig/Source/ControlRigEditor/Private/EditMode/SControlRigBaseListWidget.cpp
--- a/Engine/Plugins/Experimental/ControlRig/Source/ControlRigEditor/Private/EditMode/SControlRigBaseListWidget.cpp
+++ b/Engine/Plugins/Experimental/ControlRig/Source/ControlRigEditor/Private/EditMode/SControlRigBaseListWidget.cpp
@@ -11,6 +11,10 @@
 #include "Widgets/Input/SSearchBox.h"
 #include "Widgets/Input/SButton.h"
 #include "Widgets/Text/SInlineEditableTextBlock.h"
+#include "Widgets/Layout/SUniformGridPanel.h"
+#include "EditorStyleSet.h"
+#include "Styling/CoreStyle.h"
+
 #include "ScopedTransaction.h"
 
 #include "ControlRig.h"
@@ -30,6 +34,8 @@
 #include "FileHelpers.h"
 #include "Tools/ControlRigPoseMirrorSettings.h"
 #include "ObjectTools.h"
+#include "SControlRigUpdatePose.h"
+#include "SControlRigRenamePoseControls.h"
 
 
 #define LOCTEXT_NAMESPACE "ControlRigBaseListWidget"
@@ -283,7 +289,7 @@ void SControlRigPoseAnimSelectionToolbar::Construct(const FArguments& InArgs)
 		ToolbarBuilder.AddToolBarButton(CreatePoseDialog,
 			NAME_None,
 			LOCTEXT("CreatePose", "Create Pose From Selection"),
-			LOCTEXT("CreatePoseTooltip", "Create Pose Asset From Selection"),
+			LOCTEXT("CreatePoseTooltip", "Create pose asset from selection."),
 			FSlateIcon(),
 			EUserInterfaceActionType::Button
 		);
@@ -291,7 +297,7 @@ void SControlRigPoseAnimSelectionToolbar::Construct(const FArguments& InArgs)
 		ToolbarBuilder.AddToolBarButton(CreatePoseFromAllDialog,
 			NAME_None,
 			LOCTEXT("CreatePoseAll", "Create Pose From All Controls"),
-			LOCTEXT("CreatePoseAllTooltip", "Create Pose Asset From All Controls Not Just Selected"),
+			LOCTEXT("CreatePoseAllTooltip", "Create Pose Asset from all Controls not just selected"),
 			FSlateIcon(),
 			EUserInterfaceActionType::Button
 			
@@ -480,6 +486,143 @@ bool SControlRigPoseAnimSelectionToolbar::IsOnToggleFilter(FControlRigAssetType
 	return false;
 }
 */
+/////////////////////////////////////////////////////////////////
+//  Helper to Select Path 
+
+class SWhiteListPathDialog : public SWindow
+{
+public:
+	SLATE_BEGIN_ARGS(SWhiteListPathDialog)
+	{
+	}
+
+	SLATE_ARGUMENT(FText, DefaultAssetPath)
+		SLATE_END_ARGS()
+
+		SWhiteListPathDialog()
+		: UserResponse(EAppReturnType::Cancel)
+	{
+	}
+
+	void Construct(const FArguments& InArgs);
+
+public:
+	/** Displays the dialog in a blocking fashion */
+	EAppReturnType::Type ShowModal();
+
+	/** Gets the resulting asset path */
+	FString GetAssetPath();
+
+protected:
+	void OnPathChange(const FString& NewPath);
+	FReply OnButtonClick(EAppReturnType::Type ButtonID);
+	bool IsOkButtonEnabled() const;
+	EAppReturnType::Type UserResponse;
+	FText AssetPath;
+
+};
+
+void SWhiteListPathDialog::Construct(const FArguments& InArgs)
+{
+	AssetPath = FText::FromString(FPackageName::GetLongPackagePath(InArgs._DefaultAssetPath.ToString()));
+
+	FPathPickerConfig PathPickerConfig;
+	PathPickerConfig.OnPathSelected = FOnPathSelected::CreateSP(this, &SWhiteListPathDialog::OnPathChange);
+
+	FContentBrowserModule& ContentBrowserModule = FModuleManager::LoadModuleChecked<FContentBrowserModule>("ContentBrowser");
+
+	SWindow::Construct(SWindow::FArguments()
+		.Title(LOCTEXT("SWhiteListPathDialog_Title", "Select Folder To Contain Poses"))
+		.SupportsMinimize(false)
+		.SupportsMaximize(false)
+		//.SizingRule( ESizingRule::Autosized )
+		.ClientSize(FVector2D(450, 450))
+		[
+			SNew(SVerticalBox)
+
+			+ SVerticalBox::Slot() // Add user input block
+			.Padding(2)
+			[
+				SNew(SBorder)
+				.BorderImage(FEditorStyle::GetBrush("ToolPanel.GroupBorder"))
+				[
+					SNew(SVerticalBox)
+
+					+SVerticalBox::Slot()
+					.AutoHeight()
+					[
+						SNew(STextBlock)
+						.Text(LOCTEXT("SelectPath", "Select Path"))
+						.Font(FCoreStyle::GetDefaultFontStyle("Regular", 14))
+					]
+
+					+SVerticalBox::Slot()
+					.FillHeight(1)
+					.Padding(3)
+					[
+						ContentBrowserModule.Get().CreatePathPicker(PathPickerConfig)
+					]
+				]
+			]
+
+			+SVerticalBox::Slot()
+			.AutoHeight()
+			.HAlign(HAlign_Right)
+			.Padding(5)
+			[
+				SNew(SUniformGridPanel)
+				.SlotPadding(FEditorStyle::GetMargin("StandardDialog.SlotPadding"))
+				.MinDesiredSlotWidth(FEditorStyle::GetFloat("StandardDialog.MinDesiredSlotWidth"))
+				.MinDesiredSlotHeight(FEditorStyle::GetFloat("StandardDialog.MinDesiredSlotHeight"))
+				+SUniformGridPanel::Slot(0, 0)
+				[
+					SNew(SButton)
+					.HAlign(HAlign_Center)
+					.ContentPadding(FEditorStyle::GetMargin("StandardDialog.ContentPadding"))
+					.Text(LOCTEXT("OK", "OK"))
+					.OnClicked(this, &SWhiteListPathDialog::OnButtonClick, EAppReturnType::Ok)
+					.IsEnabled(this, &SWhiteListPathDialog::IsOkButtonEnabled)
+				]
+				+SUniformGridPanel::Slot(1, 0)
+				[
+					SNew(SButton)
+					.HAlign(HAlign_Center)
+					.ContentPadding(FEditorStyle::GetMargin("StandardDialog.ContentPadding"))
+					.Text(LOCTEXT("Cancel", "Cancel"))
+					.OnClicked(this, &SWhiteListPathDialog::OnButtonClick, EAppReturnType::Cancel)
+				]
+			]
+		]);
+}
+
+void SWhiteListPathDialog::OnPathChange(const FString& NewPath)
+{
+	AssetPath = FText::FromString(NewPath);
+}
+
+FReply SWhiteListPathDialog::OnButtonClick(EAppReturnType::Type ButtonID)
+{
+	UserResponse = ButtonID;
+	RequestDestroyWindow();
+
+	return FReply::Handled();
+}
+
+bool SWhiteListPathDialog::IsOkButtonEnabled() const
+{
+	return !AssetPath.IsEmptyOrWhitespace();
+}
+
+EAppReturnType::Type SWhiteListPathDialog::ShowModal()
+{
+	GEditor->EditorAddModalWindow(SharedThis(this));
+	return UserResponse;
+}
+
+FString SWhiteListPathDialog::GetAssetPath()
+{
+	return AssetPath.ToString();
+}
 
 /////////////////////////////////////////////////////
 // SControlRigBaseListWidget, Main Dialog Window holding the path picker, asset view and pose view.
@@ -493,8 +636,11 @@ void SControlRigBaseListWidget::Construct(const FArguments& InArgs)
 
 	const UControlRigPoseProjectSettings* PoseSettings = GetDefault<UControlRigPoseProjectSettings>();
 
-	FString PosesDir = PoseSettings->GetAssetPath();
-	CurrentlySelectedPath = PosesDir;
+	TArray<FString> PosesDirectories = PoseSettings->GetAssetPaths();
+	if (PosesDirectories.Num() > 0)
+	{
+		CurrentlySelectedPath = PosesDirectories[0];
+	}
 	FContentBrowserModule& ContentBrowserModule = FModuleManager::Get().LoadModuleChecked<FContentBrowserModule>(TEXT("ContentBrowser"));
 
 	// Configure filter for asset picker
@@ -509,7 +655,7 @@ void SControlRigBaseListWidget::Construct(const FArguments& InArgs)
 	AssetPickerConfig.bCanShowRealTimeThumbnails = true;
 	AssetPickerConfig.ThumbnailLabel = EThumbnailLabel::AssetName;
 	AssetPickerConfig.bFocusSearchBoxWhenOpened = false;
-	AssetPickerConfig.Filter.PackagePaths.Add(FName(*PosesDir));
+	AssetPickerConfig.Filter.PackagePaths.Add(FName(*CurrentlySelectedPath));
 	AssetPickerConfig.OnAssetSelected = FOnAssetSelected::CreateSP(this, &SControlRigBaseListWidget::OnAssetSelected);
 	AssetPickerConfig.OnAssetsActivated = FOnAssetsActivated::CreateSP(this, &SControlRigBaseListWidget::OnAssetsActivated);
 	AssetPickerConfig.SaveSettingsName = TEXT("ControlPoseDialog");
@@ -522,14 +668,17 @@ void SControlRigBaseListWidget::Construct(const FArguments& InArgs)
 	AssetPickerConfig.bAllowDragging = false;
 	AssetPickerConfig.AssetShowWarningText = LOCTEXT("NoPoses_Warning", "No Poses Found, Create One Using Button In Upper Left Corner");
 	AssetPickerConfig.OnIsAssetValidForCustomToolTip = FOnIsAssetValidForCustomToolTip::CreateLambda([](const FAssetData& AssetData) {return AssetData.IsAssetLoaded(); });
-	//AssetPickerConfig.OnGetCustomAssetToolTip = FOnGetCustomAssetToolTip::CreateSP(this, &SControlRigBaseListWidget::CreateCustomAssetToolTip);
 
 	FPathPickerConfig PathPickerConfig;
 	PathPickerConfig.bAddDefaultPath = true;
-	PathPickerConfig.DefaultPath = PosesDir;
-	PathPickerConfig.CustomFolderBlacklist = MakeShared<FBlacklistPaths>();
-	PathPickerConfig.CustomFolderBlacklist.Get()->AddWhitelistItem("PoseLibrary", PosesDir);
+	PathPickerConfig.DefaultPath = CurrentlySelectedPath;
+	CustomFolderBlacklist = PathPickerConfig.CustomFolderBlacklist = MakeShared<FBlacklistPaths>();
+	for (const FString& Path : PosesDirectories)
+	{
+		PathPickerConfig.CustomFolderBlacklist.Get()->AddWhitelistItem("PoseLibrary", Path);
+	}
 
+	PathPickerConfig.OnGetFolderContextMenu = FOnGetFolderContextMenu::CreateSP(this, &SControlRigBaseListWidget::OnGetFolderContextMenu);
 	PathPickerConfig.bFocusSearchBoxWhenOpened = false;
 	PathPickerConfig.OnPathSelected = FOnPathSelected::CreateSP(this, &SControlRigBaseListWidget::HandlePathSelected);
 	PathPickerConfig.SetPathsDelegates.Add(&SetPathsDelegate);
@@ -538,7 +687,12 @@ void SControlRigBaseListWidget::Construct(const FArguments& InArgs)
 	// The root widget in this dialog.
 	TSharedRef<SVerticalBox> MainVerticalBox = SNew(SVerticalBox);
 
-	// Toolbar on Top
+	AssetPicker = ContentBrowserModule.Get().CreateAssetPicker(AssetPickerConfig);
+	PathPicker = ContentBrowserModule.Get().CreatePathPicker(PathPickerConfig);
+
+	FContentBrowserConfig Config;
+	Config.bCanSetAsPrimaryBrowser = false;
+		
 	MainVerticalBox->AddSlot()
 		.AutoHeight()
 		[
@@ -559,7 +713,7 @@ void SControlRigBaseListWidget::Construct(const FArguments& InArgs)
 			SNew(SBorder)
 			.BorderImage(FEditorStyle::GetBrush("ToolPanel.GroupBorder"))
 		[
-			ContentBrowserModule.Get().CreatePathPicker(PathPickerConfig)
+			PathPicker.ToSharedRef()
 		]
 		]
 
@@ -569,11 +723,11 @@ void SControlRigBaseListWidget::Construct(const FArguments& InArgs)
 			SNew(SBorder)
 			.BorderImage(FEditorStyle::GetBrush("ToolPanel.GroupBorder"))
 		[
-			ContentBrowserModule.Get().CreateAssetPicker(AssetPickerConfig)
+			AssetPicker.ToSharedRef()
 		]
 		]
 		];
-
+		
 
 	//Bottom Area View Container to hold specific View (Pose/Animation).
 	MainVerticalBox->AddSlot()
@@ -737,37 +891,19 @@ void SControlRigBaseListWidget::OnAssetsActivated(const TArray<FAssetData>& Sele
 		SelectThisAsset(Asset);
 	}
 }
-/*
+
 void SControlRigBaseListWidget::FilterChanged()
 {
 	FARFilter NewFilter;
 	NewFilter.PackagePaths.Add(FName(*CurrentlySelectedPath));
-
-	const UControlRigPoseProjectSettings* PoseSettings = GetDefault<UControlRigPoseProjectSettings>();
-	if (PoseSettings)
-	{
-		if (PoseSettings->bFilterPoses)
-		{
-			NewFilter.ClassNames.Add(UControlRigPoseAsset::StaticClass()->GetFName());
-		}
-		if (PoseSettings->bFilterAnimations)
-		{
-			//NewFilter.ClassNames.Add(UControlRigAnimationAsset::StaticClass()->GetFName());
-		}
-		if (PoseSettings->bFilterSelectionSets)
-		{
-			//NewFilter.ClassNames.Add(UControlRigSelectionSetAsset::StaticClass()->GetFName());
-		}
-	}
+	NewFilter.ClassNames.Add(UControlRigPoseAsset::StaticClass()->GetFName());
 	SetFilterDelegate.ExecuteIfBound(NewFilter);
 }
-*/
 
 void SControlRigBaseListWidget::HandlePathSelected(const FString& NewPath)
 {
 	SetCurrentlySelectedPath(NewPath);
-	//FilterChanged();
-
+	FilterChanged();
 }
 
 void SControlRigBaseListWidget::HandleAssetViewFolderEntered(const FString& NewPath)
@@ -779,7 +915,6 @@ void SControlRigBaseListWidget::HandleAssetViewFolderEntered(const FString& NewP
 	SetPathsDelegate.Execute(NewPaths);
 }
 
-
 TSharedPtr<SWidget> SControlRigBaseListWidget::OnGetFolderContextMenu(const TArray<FString>& SelectedPaths, FContentBrowserMenuExtender_SelectedPaths InMenuExtender, FOnCreateNewFolder InOnCreateNewFolder)
 {
 
@@ -794,23 +929,44 @@ TSharedPtr<SWidget> SControlRigBaseListWidget::OnGetFolderContextMenu(const TArr
 
 	if (SelectedPaths.Num() == 1)
 	{
-		FUIAction Action = FUIAction(FExecuteAction::CreateRaw((this), &SControlRigBaseListWidget::ExecuteAddFolder, SelectedPaths[0]));
-		const FText Label = LOCTEXT("AddFolder", "Add Folder");
-		const FText ToolTipText = LOCTEXT("AddFolder Tooltip", "Add Folder to the current Selected Folder");
-		MenuBuilder.AddMenuEntry(Label, ToolTipText, FSlateIcon(), Action);
-
-		FUIAction Action2 = FUIAction(FExecuteAction::CreateRaw((this), &SControlRigBaseListWidget::ExecuteRenameFolder, SelectedPaths[0]));
-		const FText Label2 = LOCTEXT("RenameFolder", "Rename Folder");
-		const FText ToolTipText2 = LOCTEXT("RenameFolderTooltip", "Rename Selected Folder.");
-		MenuBuilder.AddMenuEntry(Label2, ToolTipText2, FSlateIcon(), Action2);
+		{
+			FUIAction Action = FUIAction(FExecuteAction::CreateRaw((this), &SControlRigBaseListWidget::ExecuteAddFolder, SelectedPaths[0]));
+			const FText Label = LOCTEXT("AddFolder", "Add Folder");
+			const FText ToolTipText = LOCTEXT("AddFolder Tooltip", "Add Folder to the current selected folder");
+			MenuBuilder.AddMenuEntry(Label, ToolTipText, FSlateIcon(), Action);
+		}
+		{
+			FUIAction Action = FUIAction(FExecuteAction::CreateRaw((this), &SControlRigBaseListWidget::ExecuteRenameFolder, SelectedPaths[0]));
+			const FText Label = LOCTEXT("RenameFolder", "Rename Folder");
+			const FText ToolTipText = LOCTEXT("RenameFolderTooltip", "Rename selected folder");
+			MenuBuilder.AddMenuEntry(Label, ToolTipText, FSlateIcon(), Action);
+		}
+		MenuBuilder.AddSeparator();
+		/*
+		{
+			FUIAction Action = FUIAction(FExecuteAction::CreateRaw((this), &SControlRigBaseListWidget::ExecuteDeleteFolder, SelectedPaths));
+			const FText Label = LOCTEXT("DeleteFolder", "Delete Folder");
+			const FText ToolTipText = LOCTEXT("DeleteFolderTooltip", "Delete selected folder(s), Note this will delete content.");
+			MenuBuilder.AddMenuEntry(Label, ToolTipText, FSlateIcon(), Action);
+		}
+		*/
 	}
 	else if (SelectedPaths.Num() > 0)
 	{
-		FUIAction Action = FUIAction(FExecuteAction::CreateRaw((this), &SControlRigBaseListWidget::ExecuteDeleteFolder, SelectedPaths));
-		const FText Label = LOCTEXT("DeleteFolder", "Delete Folder");
-		const FText ToolTipText = LOCTEXT("DeleteFolderTooltip", "Delete Selecte Folder(s), Note this will delete content.");
+		/*
+		{
+			FUIAction Action = FUIAction(FExecuteAction::CreateRaw((this), &SControlRigBaseListWidget::ExecuteDeleteFolder, SelectedPaths));
+			const FText Label = LOCTEXT("DeleteFolder", "Delete Folder");
+			const FText ToolTipText = LOCTEXT("DeleteFolderTooltip", "Delete selected folder(s), Note this will delete content.");
+			MenuBuilder.AddMenuEntry(Label, ToolTipText, FSlateIcon(), Action);
+		}
+		*/
+	}
+	{
+		FUIAction Action = FUIAction(FExecuteAction::CreateRaw((this), &SControlRigBaseListWidget::ExecuteAddFolderToView));
+		const FText Label = LOCTEXT("AddExistingFolderToView", "Add Existing Folder To View");
+		const FText ToolTipText = LOCTEXT("AddExistingFolderToViewTooltip", "Add an existing folder to this view.");
 		MenuBuilder.AddMenuEntry(Label, ToolTipText, FSlateIcon(), Action);
-
 	}
 	MenuBuilder.EndSection();
 
@@ -824,45 +980,46 @@ TSharedPtr<SWidget> SControlRigBaseListWidget::OnGetAssetContextMenu(const TArra
 	{
 		return nullptr;
 	}
-	if (SelectedAssets.Num() > 0)
+	MenuBuilder.BeginSection("PoseDialogOptions", LOCTEXT("Asset", "Asset"));
 	{
-		MenuBuilder.BeginSection("PoseDialogOptions", LOCTEXT("Asset", "Asset"));
+		if (SelectedAssets.Num() == 1)
 		{
-			{
-				FUIAction Action = FUIAction(FExecuteAction::CreateRaw((this), &SControlRigBaseListWidget::ExecuteSaveAssets, SelectedAssets));
-				const FText Label = LOCTEXT("SaveAssetButton", "Save Asset");
-				const FText ToolTipText = LOCTEXT("SaveAssetButtonTooltip", "Save the Selected Assets.");
-				MenuBuilder.AddMenuEntry(Label, ToolTipText, FSlateIcon(), Action);
-			}
-			{
-				FUIAction Action = FUIAction(FExecuteAction::CreateRaw((this), &SControlRigBaseListWidget::ExecuteDeleteAssets, SelectedAssets));
-				const FText Label = LOCTEXT("DeleteAssetButton", "Delete Asset");
-				const FText ToolTipText = LOCTEXT("DeleteAssetButtonTooltip", "Delete the Selected Assets.");
-				MenuBuilder.AddMenuEntry(Label, ToolTipText, FSlateIcon(), Action);
-			}
+			FUIAction Action = FUIAction(FExecuteAction::CreateRaw((this), &SControlRigBaseListWidget::ExecuteRenameAssets, SelectedAssets),
+			FCanExecuteAction::CreateRaw(this, &SControlRigBaseListWidget::CanExecuteRenameAssets, SelectedAssets));
+			const FText Label = LOCTEXT("RenameAssetButton", "Rename Asset");
+			const FText ToolTipText = LOCTEXT("RenameAssetButtonTooltip", "Rename the selected asset");
+			MenuBuilder.AddMenuEntry(Label, ToolTipText, FSlateIcon(), Action);
 		}
-		MenuBuilder.EndSection();
 
-	}
-	if (SelectedAssets.Num() == 1)
-	{
-		UObject* SelectedAsset = SelectedAssets[0].GetAsset();
-		if (SelectedAsset == nullptr)
 		{
-			return nullptr;
+			FUIAction Action = FUIAction(FExecuteAction::CreateRaw((this), &SControlRigBaseListWidget::ExecuteSaveAssets, SelectedAssets));
+			const FText Label = LOCTEXT("SaveAssetButton", "Save Asset");
+			const FText ToolTipText = LOCTEXT("SaveAssetButtonTooltip", "Save the selected assets");
+			MenuBuilder.AddMenuEntry(Label, ToolTipText, FSlateIcon(), Action);
 		}
+		{
+			FUIAction Action = FUIAction(FExecuteAction::CreateRaw((this), &SControlRigBaseListWidget::ExecuteDeleteAssets, SelectedAssets));
+			const FText Label = LOCTEXT("DeleteAssetButton", "Delete Asset");
+			const FText ToolTipText = LOCTEXT("DeleteAssetButtonTooltip", "Delete the selected assets.");
+			MenuBuilder.AddMenuEntry(Label, ToolTipText, FSlateIcon(), Action);
+		}
+	}
+	MenuBuilder.EndSection();
+
+	UObject* SelectedAsset = SelectedAssets[0].GetAsset();
+	if (SelectedAsset)
+	{
 
 		UControlRigPoseAsset* PoseAsset = Cast<UControlRigPoseAsset>(SelectedAsset);
 		if (PoseAsset)
 		{
-
 			MenuBuilder.BeginSection("PoseDialogOptions", LOCTEXT("PoseDialogMenuHeading", "Paste"));
 			{
 
 				FUIAction Action = FUIAction(FExecuteAction::CreateRaw((this), &SControlRigBaseListWidget::ExecutePastePose, PoseAsset),
 					FCanExecuteAction::CreateRaw(this, &SControlRigBaseListWidget::CanExecutePastePose, PoseAsset));
 				const FText Label = LOCTEXT("PastePoseButton", "Paste Pose");
-				const FText ToolTipText = LOCTEXT("PastePoseButtonTooltip", "Paste the Selected Pose.");
+				const FText ToolTipText = LOCTEXT("PastePoseButtonTooltip", "Paste the selected pose");
 				MenuBuilder.AddMenuEntry(Label, ToolTipText, FSlateIcon(), Action);
 			}
 
@@ -870,17 +1027,16 @@ TSharedPtr<SWidget> SControlRigBaseListWidget::OnGetAssetContextMenu(const TArra
 				FUIAction Action = FUIAction(FExecuteAction::CreateRaw((this), &SControlRigBaseListWidget::ExecutePasteMirrorPose, PoseAsset),
 					FCanExecuteAction::CreateRaw(this, &SControlRigBaseListWidget::CanExecutePasteMirrorPose, PoseAsset));
 				const FText Label = LOCTEXT("PasteMirrorPoseButton", "Paste Mirror Pose");
-				const FText ToolTipText = LOCTEXT("PastePoseButtonTooltip", "Paste the Mirror Pose.");
+				const FText ToolTipText = LOCTEXT("PastePoseButtonTooltip", "Paste the mirror of the pose");
 				MenuBuilder.AddMenuEntry(Label, ToolTipText, FSlateIcon(), Action);
 			}
-
 			MenuBuilder.EndSection();
 
 			MenuBuilder.BeginSection("PoseDialogOptions", LOCTEXT("PoseDialogSelectHeading", "Selection"));
 			{
 				FUIAction Action = FUIAction(FExecuteAction::CreateRaw((this), &SControlRigBaseListWidget::ExecuteSelectControls, PoseAsset));
 				const FText Label = LOCTEXT("SelectControls", "Select Controls");
-				const FText ToolTipText = LOCTEXT("SelectControlsTooltip", "Select Controls in this Pose on Active Control Rig");
+				const FText ToolTipText = LOCTEXT("SelectControlsTooltip", "Select controls in this pose on active control rig");
 				MenuBuilder.AddMenuEntry(Label, ToolTipText, FSlateIcon(), Action);
 			}
 			MenuBuilder.EndSection();
@@ -889,33 +1045,53 @@ TSharedPtr<SWidget> SControlRigBaseListWidget::OnGetAssetContextMenu(const TArra
 			{
 				FUIAction Action = FUIAction(FExecuteAction::CreateRaw((this), &SControlRigBaseListWidget::ExecuteUpdatePose, PoseAsset));
 				const FText Label = LOCTEXT("UpdatePose", "Update Pose");
-				const FText ToolTipText = LOCTEXT("Update Pose Tooltip", "Update The Pose Based Upon Current Control Rig Pose");
+				const FText ToolTipText = LOCTEXT("UpdatePoseTooltip", "Update the pose based upon current control rig pose and selected controls");
 				MenuBuilder.AddMenuEntry(Label, ToolTipText, FSlateIcon(), Action);
 			}
 			MenuBuilder.EndSection();
-
 		}
 	}
-	return MenuBuilder.MakeWidget();
 
-}
+	MenuBuilder.BeginSection("PoseDialogOptions", LOCTEXT("PoseDialogRenameControlsHeading", "Rename Controls"));
+	{
+		FUIAction Action = FUIAction(FExecuteAction::CreateRaw((this), &SControlRigBaseListWidget::ExecuteRenamePoseControls, SelectedAssets));
+		const FText Label = LOCTEXT("RenameControls", "Rename Controls");
+		const FText ToolTipText = LOCTEXT("RenameControlsTooltip", "Rename controls in pose asset");
+		MenuBuilder.AddMenuEntry(Label, ToolTipText, FSlateIcon(), Action);
+	}
+	MenuBuilder.EndSection();
 
+	return MenuBuilder.MakeWidget();
+}
 
 void SControlRigBaseListWidget::BindCommands()
 {
 	Commands = TSharedPtr< FUICommandList >(new FUICommandList);
-
-
 }
+
 void SControlRigBaseListWidget::ExecuteRenameFolder(const FString SelectedPath)
 {
-
+	FContentBrowserModule& ContentBrowserModule = FModuleManager::Get().LoadModuleChecked<FContentBrowserModule>(TEXT("ContentBrowser"));
+	ContentBrowserModule.Get().ExecuteRename(PathPicker);
 }
 
 void SControlRigBaseListWidget::ExecuteAddFolder(const FString SelectedPath)
 {
+	FContentBrowserModule& ContentBrowserModule = FModuleManager::Get().LoadModuleChecked<FContentBrowserModule>(TEXT("ContentBrowser"));
+	ContentBrowserModule.Get().ExecuteAddFolder(PathPicker);
+}
+
+void SControlRigBaseListWidget::ExecuteRenameAssets(const TArray<FAssetData> SelectedAssets)
+{
+	FContentBrowserModule& ContentBrowserModule = FModuleManager::Get().LoadModuleChecked<FContentBrowserModule>(TEXT("ContentBrowser"));
+	ContentBrowserModule.Get().ExecuteRename(AssetPicker);
+}
 
+bool SControlRigBaseListWidget::CanExecuteRenameAssets(const TArray<FAssetData> SelectedAssets) const
+{
+	return true;
 }
+
 void SControlRigBaseListWidget::ExecuteSaveAssets(const TArray<FAssetData> SelectedAssets)
 {
 	TArray<UPackage*> PackagesToSave;
@@ -927,17 +1103,38 @@ void SControlRigBaseListWidget::ExecuteSaveAssets(const TArray<FAssetData> Selec
 			FEditorFileUtils::PromptForCheckoutAndSave(PackagesToSave, false /*bCheckDirty*/, false /*bPromptToSave*/);
 		}
 	}
-
 }
+
 void SControlRigBaseListWidget::ExecuteDeleteAssets(const TArray<FAssetData> SelectedAssets)
 {
 	ObjectTools::DeleteAssets(SelectedAssets);
 	SelectThisAsset(nullptr);
+}
 
+void SControlRigBaseListWidget::ExecuteAddFolderToView()
+{
+	TSharedRef<SWhiteListPathDialog> NewAnimDlg =
+		SNew(SWhiteListPathDialog);
+
+	if (NewAnimDlg->ShowModal() != EAppReturnType::Cancel)
+	{
+		FString AssetPath = NewAnimDlg->GetAssetPath();
+		CustomFolderBlacklist.Get()->AddWhitelistItem("PoseLibrary", AssetPath);
+	    UControlRigPoseProjectSettings* PoseSettings = GetMutableDefault<UControlRigPoseProjectSettings>();
+		FDirectoryPath Path;
+		Path.Path = AssetPath;
+		PoseSettings->RootSaveDirs.Add(Path);
+		PoseSettings->SaveConfig();
+
+		//Need to refresh since the blacklist changed
+		FContentBrowserModule& ContentBrowserModule = FModuleManager::Get().LoadModuleChecked<FContentBrowserModule>(TEXT("ContentBrowser"));
+		ContentBrowserModule.Get().RefreshPathView(PathPicker);
+	}
 }
+
 void SControlRigBaseListWidget::ExecuteDeleteFolder(const TArray<FString> SelectedFolders)
 {
-
+	/*
 	// Don't allow asset deletion during PIE
 	if (GIsEditor)
 	{
@@ -951,7 +1148,7 @@ void SControlRigBaseListWidget::ExecuteDeleteFolder(const TArray<FString> Select
 			return;
 		}
 	}
-	/*
+	
 	const TArray<FContentBrowserItem> SelectedFiles = AssetPicker->GetAssetView()->GetSelectedFileItems();
 
 	// Batch these by their data sources
@@ -982,7 +1179,7 @@ void SControlRigBaseListWidget::ExecuteDeleteFolder(const TArray<FString> Select
 	{
 		SourceAndItemsPair.Key->BulkDeleteItems(SourceAndItemsPair.Value);
 	}
-	*/
+	
 	// If we had any folders selected, ask the user whether they want to delete them 
 	// as it can be slow to build the deletion dialog on an accidental click
 	if (SelectedFolders.Num() > 0)
@@ -998,7 +1195,7 @@ void SControlRigBaseListWidget::ExecuteDeleteFolder(const TArray<FString> Select
 		}
 
 		FContentBrowserModule& ContentBrowserModule = FModuleManager::Get().LoadModuleChecked<FContentBrowserModule>(TEXT("ContentBrowser"));
-		/*
+	
 		// Spawn a confirmation dialog since this is potentially a highly destructive operation
 		ContentBrowserModule.Get().DisplayConfirmationPopup(
 			Prompt,
@@ -1007,8 +1204,9 @@ void SControlRigBaseListWidget::ExecuteDeleteFolder(const TArray<FString> Select
 			ToSharedRef(),
 			FOnClicked::CreateSP(this, &SControlRigBaseListWidget::ExecuteDeleteFolderConfirmed)
 		);
-		*/
+		
 	}
+	*/
 
 }
 
@@ -1044,7 +1242,6 @@ FReply SControlRigBaseListWidget::ExecuteDeleteFolderConfirmed()
 }
 
 
-
 void SControlRigBaseListWidget::ExecutePastePose(UControlRigPoseAsset* PoseAsset)
 {
 	if (PoseAsset)
@@ -1052,10 +1249,13 @@ void SControlRigBaseListWidget::ExecutePastePose(UControlRigPoseAsset* PoseAsset
 		FControlRigEditMode* ControlRigEditMode = static_cast<FControlRigEditMode*>(GLevelEditorModeTools().GetActiveMode(FControlRigEditMode::ModeName));
 		if (ControlRigEditMode && ControlRigEditMode->GetControlRig(true))
 		{
+			const FScopedTransaction Transaction(LOCTEXT("PastePose", "Paste Pose"));
+			ControlRigEditMode->GetControlRig(true)->Modify();
 			PoseAsset->PastePose(ControlRigEditMode->GetControlRig(true));
 		}
 	}
 }
+
 bool SControlRigBaseListWidget::CanExecutePastePose(UControlRigPoseAsset* PoseAsset) const
 {
 	return PoseAsset != nullptr;
@@ -1063,11 +1263,7 @@ bool SControlRigBaseListWidget::CanExecutePastePose(UControlRigPoseAsset* PoseAs
 
 void SControlRigBaseListWidget::ExecuteUpdatePose(UControlRigPoseAsset* PoseAsset)
 {
-	FControlRigEditMode* ControlRigEditMode = static_cast<FControlRigEditMode*>(GLevelEditorModeTools().GetActiveMode(FControlRigEditMode::ModeName));
-	if (ControlRigEditMode && ControlRigEditMode->GetControlRig(true))
-	{
-		PoseAsset->SavePose(ControlRigEditMode->GetControlRig(true),false);
-	}
+	FControlRigUpdatePoseDialog::UpdatePose(PoseAsset);
 }
 
 void SControlRigBaseListWidget::ExecuteSelectControls(UControlRigPoseAsset* PoseAsset)
@@ -1075,6 +1271,8 @@ void SControlRigBaseListWidget::ExecuteSelectControls(UControlRigPoseAsset* Pose
 	FControlRigEditMode* ControlRigEditMode = static_cast<FControlRigEditMode*>(GLevelEditorModeTools().GetActiveMode(FControlRigEditMode::ModeName));
 	if (ControlRigEditMode && ControlRigEditMode->GetControlRig(true))
 	{
+		const FScopedTransaction Transaction(LOCTEXT("SelectControls", "Select Controls"));
+		ControlRigEditMode->GetControlRig(true)->Modify();
 		PoseAsset->SelectControls(ControlRigEditMode->GetControlRig(true));
 	}
 }
@@ -1086,6 +1284,8 @@ void SControlRigBaseListWidget::ExecutePasteMirrorPose(UControlRigPoseAsset* Pos
 		FControlRigEditMode* ControlRigEditMode = static_cast<FControlRigEditMode*>(GLevelEditorModeTools().GetActiveMode(FControlRigEditMode::ModeName));
 		if (ControlRigEditMode && ControlRigEditMode->GetControlRig(true))
 		{
+			const FScopedTransaction Transaction(LOCTEXT("PasteMirrorPose", "Paste Mirror Pose"));
+			ControlRigEditMode->GetControlRig(true)->Modify();
 			PoseAsset->PastePose(ControlRigEditMode->GetControlRig(true), false, true);
 		}
 	}
@@ -1096,6 +1296,21 @@ bool SControlRigBaseListWidget::CanExecutePasteMirrorPose(UControlRigPoseAsset*
 	return PoseAsset != nullptr;
 }
 
+void SControlRigBaseListWidget::ExecuteRenamePoseControls(const TArray<FAssetData> SelectedAssets)
+{
+	FPropertyEditorModule& PropertyEditorModule = FModuleManager::LoadModuleChecked<FPropertyEditorModule>("PropertyEditor");
+	TArray<UControlRigPoseAsset*> SelectedObjects;
+	for (const FAssetData& AssetData : SelectedAssets)
+	{
+		if (AssetData.IsValid() && AssetData.GetAsset() && AssetData.GetAsset()->IsA<UControlRigPoseAsset>())
+		{
+			UControlRigPoseAsset* PoseAsset = Cast<UControlRigPoseAsset>(AssetData.GetAsset());
+			SelectedObjects.Add(PoseAsset);
+		}
+	}
+	FControlRigRenameControlsDialog::RenameControls(SelectedObjects);
+}
+
 void SControlRigBaseListWidget::CreateCurrentView(UObject* Asset)
 {
 	PoseView.Reset();
@@ -1124,7 +1339,6 @@ void SControlRigBaseListWidget::CreateCurrentView(UObject* Asset)
 		ViewContainer->SetContent(EmptyBox.ToSharedRef());
 		break;
 	}
-
 }
 
 TSharedRef<class SControlRigPoseView> SControlRigBaseListWidget::CreatePoseView(UObject* InObject)
diff --git a/Engine/Plugins/Experimental/ControlRig/Source/ControlRigEditor/Private/EditMode/SControlRigBaseListWidget.h b/Engine/Plugins/Experimental/ControlRig/Source/ControlRigEditor/Private/EditMode/SControlRigBaseListWidget.h
--- a/Engine/Plugins/Experimental/ControlRig/Source/ControlRigEditor/Private/EditMode/SControlRigBaseListWidget.h
+++ b/Engine/Plugins/Experimental/ControlRig/Source/ControlRigEditor/Private/EditMode/SControlRigBaseListWidget.h
@@ -23,6 +23,7 @@
 class UControlRig;
 class FControlRigEditMode;
 class UControlRigPoseAsset;
+class FAssetFileContextMenu;
 
 namespace ESelectedControlAsset
 {
@@ -55,10 +56,9 @@ public:
 	FControlRigEditMode* GetEditMode();
 	UControlRig* GetControlRig(); 
 
-	/*Notification from Toolbar that the Filter Changed and so we should reset the filter.
-	Currently disabled
+	/*Notification that filter has changed, this includes new folder selected*/
 	void FilterChanged();
-	*/
+	
 	/*Current Path we use to save the asset.*/
 	FString GetCurrentlySelectedPath()const;
 
@@ -92,6 +92,8 @@ private:
 	void ExecuteRenameFolder(const FString SelectedPath);
 	void ExecuteAddFolder(const FString SelectedPath);
 	FReply ExecuteDeleteFolderConfirmed();
+	void ExecuteRenameAssets(const TArray<FAssetData> RenameAssets);
+	bool CanExecuteRenameAssets(const TArray<FAssetData> RenameAssets) const;
 	void ExecuteSaveAssets(const TArray<FAssetData> SelectedAssets);
 	void ExecuteDeleteAssets(const TArray<FAssetData> SelectedAssets);
 	void ExecutePastePose(UControlRigPoseAsset* PoseAsset);
@@ -100,6 +102,8 @@ private:
 	void ExecutePasteMirrorPose(UControlRigPoseAsset* PoseAsset);
 	bool CanExecutePasteMirrorPose(UControlRigPoseAsset* PoseAsset) const;
 	void ExecuteUpdatePose(UControlRigPoseAsset* PoseAsset);
+	void ExecuteRenamePoseControls(const TArray<FAssetData> SelectedAssets);
+	void ExecuteAddFolderToView();
 
 	/** Create the Views for the Selected Asset*/
 	void CreateCurrentView(UObject* Asset);
@@ -115,7 +119,10 @@ private:
 	TSharedPtr<SBox> EmptyBox;
 	TSharedPtr<SBox> ViewContainer;
 
+	TSharedPtr<SWidget> PathPicker;
+	TSharedPtr<SWidget> AssetPicker;
 
+	TSharedPtr<FAssetFileContextMenu> AssetFileContextMenu;
 
 	FString CurrentlySelectedPath;
 	FString CurrentlyEnteredAssetName;
@@ -124,6 +131,9 @@ private:
 	/** Commands handled by this widget */
 	TSharedPtr< FUICommandList > Commands;
 
+	/** Blacklist held by the Path Picker*/
+	TSharedPtr<FBlacklistPaths> CustomFolderBlacklist;
+
 	/** Utility function to display notifications to the user */
 	void NotifyUser(FNotificationInfo& NotificationInfo);
 
diff --git a/Engine/Plugins/Experimental/ControlRig/Source/ControlRigEditor/Private/EditMode/SControlRigControlViews.cpp b/Engine/Plugins/Experimental/ControlRig/Source/ControlRigEditor/Private/EditMode/SControlRigControlViews.cpp
--- a/Engine/Plugins/Experimental/ControlRig/Source/ControlRigEditor/Private/EditMode/SControlRigControlViews.cpp
+++ b/Engine/Plugins/Experimental/ControlRig/Source/ControlRigEditor/Private/EditMode/SControlRigControlViews.cpp
@@ -18,6 +18,7 @@
 #include "ControlRigEditMode.h"
 #include "EditorModeManager.h"
 #include "Widgets/Input/SButton.h"
+#include "LevelEditor.h"
 
 #define LOCTEXT_NAMESPACE "ControlRigBaseListWidget"
 
@@ -85,7 +86,7 @@ private:
 	TWeakObjectPtr<UObject> Asset;
 
 	/** The text box used to edit object names */
-	TSharedPtr< SEditableTextBox > TextBox;
+	TSharedPtr< STextBlock > TextBox;
 
 };
 
@@ -94,12 +95,13 @@ void SControlRigAssetEditableTextBox::Construct(const FArguments& InArgs)
 	Asset = InArgs._Asset;
 	ChildSlot
 		[
-			SAssignNew(TextBox, SEditableTextBox)
+			SAssignNew(TextBox, STextBlock)
 			.Text(this, &SControlRigAssetEditableTextBox::GetNameText)
-			.ToolTipText(this, &SControlRigAssetEditableTextBox::GetNameTooltipText)
-			.OnTextCommitted(this, &SControlRigAssetEditableTextBox::OnNameTextCommitted)
-			.OnTextChanged(this, &SControlRigAssetEditableTextBox::OnTextChanged)
-			.RevertTextOnEscape(true)
+			// Current Thinking is to not have this be editable here, so removing it, but leaving in case we change our minds again.
+			//.ToolTipText(this, &SControlRigAssetEditableTextBox::GetNameTooltipText)
+			//.OnTextCommitted(this, &SControlRigAssetEditableTextBox::OnNameTextCommitted)
+			//.OnTextChanged(this, &SControlRigAssetEditableTextBox::OnTextChanged)
+			//.RevertTextOnEscape(true)
 		];
 }
 
@@ -158,7 +160,7 @@ void SControlRigAssetEditableTextBox::OnNameTextCommitted(const FText& NewText,
 	}
 
 	// Clear Error 
-	TextBox->SetError(FText::GetEmpty());
+	//TextBox->SetError(FText::GetEmpty());
 }
 
 void SControlRigAssetEditableTextBox::OnTextChanged(const FText& InLabel)
@@ -170,11 +172,11 @@ void SControlRigAssetEditableTextBox::OnTextChanged(const FText& InLabel)
 	FText OutErrorMessage;
 	if (!AssetViewUtils::IsValidObjectPathForCreate(ObjectPath, OutErrorMessage))
 	{
-		TextBox->SetError(OutErrorMessage);
+		//TextBox->SetError(OutErrorMessage);
 	}
 	else
 	{
-		TextBox->SetError(FText::GetEmpty());
+		//TextBox->SetError(FText::GetEmpty());
 	}
 }
 
@@ -190,7 +192,7 @@ void SControlRigPoseView::Construct(const FArguments& InArgs)
 	bIsBlending = false;
 	bSliderStartedTransaction = false;
 
-	TSharedRef<SWidget> Thumbnail = GetThumbnailWidget();
+	TSharedRef<SWidget> ThumbnailWidget = GetThumbnailWidget();
 	TSharedRef <SControlRigAssetEditableTextBox> ObjectNameBox = SNew(SControlRigAssetEditableTextBox).Asset(PoseAsset);
 
 	//Not used currently CreateControlList();
@@ -249,7 +251,7 @@ void SControlRigPoseView::Construct(const FArguments& InArgs)
 							SNew(SBox)
 							.VAlign(VAlign_Center)
 						[
-							Thumbnail
+							ThumbnailWidget
 						]
 						]
 					+ SVerticalBox::Slot()
@@ -371,7 +373,7 @@ void SControlRigPoseView::Construct(const FArguments& InArgs)
 							SNew(SButton)
 								.ContentPadding(FMargin(5, 5))
 								.Text(LOCTEXT("SelectControls", "Select Controls"))
-								.ToolTipText(LOCTEXT("SelectControlsTooltip", "Select Controls From This Asset"))
+								.ToolTipText(LOCTEXT("SelectControlsTooltip", "Select controls from this asset"))
 								.OnClicked(this, &SControlRigPoseView::OnSelectControls)
 				
 							]
@@ -517,21 +519,31 @@ FReply SControlRigPoseView::OnCaptureThumbnail()
 
 TSharedRef<SWidget> SControlRigPoseView::GetThumbnailWidget()
 {
-	TSharedPtr<SWidget> ThumbnailWidget = nullptr;
+	FLevelEditorModule& LevelEditorModule = FModuleManager::LoadModuleChecked<FLevelEditorModule>("LevelEditor");
+	TSharedPtr<FAssetThumbnailPool> ThumbnailPool = LevelEditorModule.GetFirstLevelEditor()->GetThumbnailPool();
 
-	ThumbnailPool = MakeShared<FAssetThumbnailPool>(1, false);
 
 	const int32 ThumbnailSize = 128;
-
-	TSharedRef<FAssetThumbnail> AssetThumbnail = MakeShared<FAssetThumbnail>(PoseAsset.Get(), ThumbnailSize, ThumbnailSize, ThumbnailPool);
-	ThumbnailWidget = AssetThumbnail->MakeThumbnailWidget();
-
+	Thumbnail = MakeShareable(new FAssetThumbnail(PoseAsset.Get(), ThumbnailSize, ThumbnailSize, ThumbnailPool));
+	FAssetThumbnailConfig ThumbnailConfig;
+	ThumbnailConfig.bAllowFadeIn = false;
+	ThumbnailConfig.bAllowHintText = false;
+	ThumbnailConfig.bAllowRealTimeOnHovered = false; // we use our own OnMouseEnter/Leave for logical asset item
+	ThumbnailConfig.bForceGenericThumbnail = false;
+
+	TSharedRef<SOverlay> ItemContentsOverlay = SNew(SOverlay);
+	ItemContentsOverlay->AddSlot()
+		[
+			Thumbnail->MakeThumbnailWidget(ThumbnailConfig)
+		];
 
 	return SNew(SBox)
-		.WidthOverride(ThumbnailSize + 5)
-		.HeightOverride(ThumbnailSize + 5)
+		.Padding(0)
+		.WidthOverride(ThumbnailSize)
+		.HeightOverride(ThumbnailSize)
 		[
-			ThumbnailWidget.IsValid() ? ThumbnailWidget.ToSharedRef() : SNullWidget::NullWidget
+			ItemContentsOverlay
+
 		];
 }
 
diff --git a/Engine/Plugins/Experimental/ControlRig/Source/ControlRigEditor/Private/EditMode/SControlRigControlViews.h b/Engine/Plugins/Experimental/ControlRig/Source/ControlRigEditor/Private/EditMode/SControlRigControlViews.h
--- a/Engine/Plugins/Experimental/ControlRig/Source/ControlRigEditor/Private/EditMode/SControlRigControlViews.h
+++ b/Engine/Plugins/Experimental/ControlRig/Source/ControlRigEditor/Private/EditMode/SControlRigControlViews.h
@@ -63,21 +63,14 @@ private:
 	bool bIsBlending;
 	bool bSliderStartedTransaction;
 	FControlRigControlPose TempPose;
+	TSharedPtr< FAssetThumbnail > Thumbnail;
+	TSharedRef<SWidget> GetThumbnailWidget();
 
 	TWeakObjectPtr<UControlRigPoseAsset> PoseAsset;
-
-	/* Thumbnail*/
-	TSharedRef<SWidget> GetThumbnailWidget();
-	TSharedPtr<FAssetThumbnailPool> ThumbnailPool;
 	
 	/* Mirroring*/
 	TSharedPtr<IDetailsView> MirrorDetailsView;
 
-	/*for list of controls stored in asset 
-	void CreateControlList();
-	TArray<TSharedPtr<FString>> ControlList;
-	*/
-
 	TSharedRef<ITableRow> OnGenerateWidgetForList(TSharedPtr<FString> InItem, const TSharedRef<STableViewBase>& OwnerTable)
 	{
 
diff --git a/Engine/Plugins/Experimental/ControlRig/Source/ControlRigEditor/Private/EditMode/SControlRigRenamePoseControls.cpp b/Engine/Plugins/Experimental/ControlRig/Source/ControlRigEditor/Private/EditMode/SControlRigRenamePoseControls.cpp
new file mode 100644
--- /dev/null
+++ b/Engine/Plugins/Experimental/ControlRig/Source/ControlRigEditor/Private/EditMode/SControlRigRenamePoseControls.cpp
@@ -0,0 +1,297 @@
+// Copyright Epic Games, Inc. All Rights Reserved.
+
+#include "SControlRigRenamePoseControls.h"
+#include "CoreMinimal.h"
+#include "Widgets/SWidget.h"
+#include "Widgets/SCompoundWidget.h"
+#include "Editor.h"
+#include "Modules/ModuleManager.h"
+#include "Framework/Application/SlateApplication.h"
+#include "Widgets/SWindow.h"
+#include "Widgets/Input/SButton.h"
+#include "ControlRigEditMode.h"
+#include "Tools/ControlRigPose.h"
+#include "EditorModeManager.h"
+#include "ScopedTransaction.h"
+#include "Widgets/Layout/SBox.h"
+#include "Widgets/Text/STextBlock.h"
+#include "Widgets/Input/SButton.h"
+#include "Widgets/Input/SEditableTextBox.h"
+#include "Widgets/Input/SCheckBox.h"
+#include "Widgets/Layout/SGridPanel.h"
+#include "Framework/Application/SlateApplication.h"
+#include "Widgets/Input/SEditableTextBox.h"
+#include "Widgets/Views/SListView.h"
+
+#define LOCTEXT_NAMESPACE "ControlRigRenamePoseControls"
+
+///// SControlRigRenamePoseControls /////
+
+struct SControlRigRenamePoseControls : public  SCompoundWidget
+{
+public:
+	struct FListItem
+	{
+		FString OriginalName;
+		FString NewName;
+	};
+
+	class SControlRenameListRow : public SMultiColumnTableRow<TSharedPtr<FListItem>>
+	{
+	public:
+		SLATE_BEGIN_ARGS(SControlRenameListRow) {}
+		SLATE_ARGUMENT(TSharedPtr<FListItem>, Item)
+			SLATE_END_ARGS()
+
+		static const FName NAME_Original;
+		static const FName NAME_New;
+
+		TSharedPtr<FListItem> Item;
+
+		void Construct(const FArguments& InArgs, const TSharedRef<STableViewBase>& OwnerTable)
+		{
+			Item = InArgs._Item;
+
+			FSuperRowType::FArguments SuperArgs = FSuperRowType::FArguments();
+			FSuperRowType::Construct(SuperArgs, OwnerTable);
+		}
+
+		virtual TSharedRef<SWidget> GenerateWidgetForColumn(const FName& ColumnName) override
+		{
+			
+			if (ColumnName == NAME_Original)
+			{
+				return SNew(SBox)
+					.Padding(2.0f)
+					.HAlign(HAlign_Left)
+					.VAlign(VAlign_Center)
+					[
+						SNew(STextBlock)
+						.Text(FText::FromString(Item->OriginalName))
+					];
+			}
+			else
+			{
+				return SNew(SBox)
+					.Padding(2.0f)
+						.HAlign(HAlign_Left)
+						.VAlign(VAlign_Center)
+						[
+							SNew(SEditableTextBox)
+							.Text(FText::FromString(Item->NewName))
+							.OnTextCommitted_Lambda([this](const FText& Val, ETextCommit::Type TextCommitType)
+								{
+									Item->NewName = Val.ToString();
+								})
+							.ToolTipText(LOCTEXT("NewControlName", "New control name"))
+						];
+					}
+		}
+	};
+
+	SLATE_BEGIN_ARGS(SControlRigRenamePoseControls) { }
+	SLATE_END_ARGS()
+
+	void Construct(const FArguments& Args, const TArray<UControlRigPoseAsset*>& InPoseAssets)
+	{
+		bResult = false;
+		SetPoseAssets(InPoseAssets);
+
+		for (const FName& Name: CommonControlNames)
+		{
+			TSharedRef<SControlRigRenamePoseControls::FListItem> Item = MakeShared<SControlRigRenamePoseControls::FListItem>();
+			Item->OriginalName = Name.ToString();
+			Item->NewName = Name.ToString();
+			Items.Add(Item);
+		}
+
+		ChildSlot
+			[
+				SNew(SBox)
+				.HAlign(HAlign_Fill)
+			.VAlign(VAlign_Fill)
+			.Padding(20.0f)
+			[
+				SNew(SVerticalBox)
+				+ SVerticalBox::Slot()
+				.AutoHeight()
+				.HAlign(HAlign_Fill)
+				.VAlign(VAlign_Center)
+				.Padding(0.0f, 0.0f, 0.0f, 8.0f)
+				[
+					SNew(STextBlock)
+					.Text(LOCTEXT("RenamePrompt", "Enter a new name for the Controls you wish to rename."))
+				]
+			+ SVerticalBox::Slot()
+				.HAlign(HAlign_Fill)
+				.VAlign(VAlign_Fill)
+				.Padding(0.0f, 0.0f, 0.0f, 8.0f)
+				.FillHeight(1.0f)
+				[
+					SNew(SBorder)
+					[
+						SAssignNew(RenameList, SListView<TSharedPtr<FListItem>>)
+						.ListItemsSource(&Items)
+				.OnGenerateRow(this, &SControlRigRenamePoseControls::OnGenerateRow)
+				.HeaderRow
+				(
+					SNew(SHeaderRow)
+
+					+ SHeaderRow::Column(SControlRenameListRow::NAME_Original)
+					.FillWidth(0.4)
+					.DefaultLabel(LOCTEXT("Original", "Original"))
+
+					+ SHeaderRow::Column(SControlRenameListRow::NAME_New)
+					.FillWidth(0.6)
+					.DefaultLabel(LOCTEXT("New", "New"))
+				)
+					]
+				]
+			+ SVerticalBox::Slot()
+				.AutoHeight()
+				.HAlign(HAlign_Fill)
+				.VAlign(VAlign_Fill)
+				[
+					SNew(SHorizontalBox)
+					+ SHorizontalBox::Slot()
+					.FillWidth(1.0f)
+					.Padding(5.0f, 0.0f)
+					[
+						SNullWidget::NullWidget
+					]
+				+ SHorizontalBox::Slot()
+					.Padding(5.0f, 0.0f)
+					.AutoWidth()
+					[
+						SNew(SButton)
+						.ContentPadding(FMargin(10.0f, 2.0f, 10.0f, 2.0f))
+					.OnClicked(this, &SControlRigRenamePoseControls::OnButtonClick_Continue)
+					[
+						SNew(STextBlock)
+						.Text(LOCTEXT("Rename", "Rename"))
+					]
+				]
+				+ SHorizontalBox::Slot()
+					.AutoWidth()
+					[
+						SNew(SButton)
+						.ContentPadding(FMargin(10.0f, 2.0f, 10.0f, 2.0f))
+						.OnClicked(this, &SControlRigRenamePoseControls::OnButtonClick_Cancel)
+						[
+							SNew(STextBlock)
+							.Text(LOCTEXT("Cancel", "Cancel"))
+						]
+					]
+				]
+			]
+		];	
+	}
+
+	TSharedRef<ITableRow> OnGenerateRow(TSharedPtr<FListItem> Item, const TSharedRef<STableViewBase>& OwnerTable)
+	{
+		return SNew(SControlRenameListRow, OwnerTable)
+			.Item(Item);
+	}
+
+	FReply OnButtonClick_Continue()
+	{
+		bResult = true;
+		TSharedPtr<SWindow> Window = FSlateApplication::Get().FindWidgetWindow(AsShared());
+
+		for (TWeakObjectPtr<UControlRigPoseAsset>& PoseAsset : PoseAssets)
+		{
+			if (PoseAsset.IsValid() )
+			{
+				const FScopedTransaction Transaction(NSLOCTEXT("ControlRig", "RenameControls", "Rename Controls"));
+				PoseAsset->Modify();
+				for (TSharedPtr<FListItem>&Item : Items)
+				{
+					PoseAsset->ReplaceControlName(FName(*Item->OriginalName),FName(*Item->NewName));
+				}
+			}
+			if (Window.IsValid())
+			{
+				Window->RequestDestroyWindow();
+			}
+		}
+		return FReply::Handled();
+	}
+
+	FReply OnButtonClick_Cancel()
+	{
+		bResult = false;
+		FSlateApplication::Get().FindWidgetWindow(AsShared())->RequestDestroyWindow();
+		return FReply::Handled();
+	}
+
+	void SetPoseAssets(const TArray<UControlRigPoseAsset*>& InPoseAssets)
+	{
+		CommonControlNames.Reset();
+		PoseAssets.Reset();
+
+		for (UControlRigPoseAsset* PoseAsset : InPoseAssets)
+		{
+			if (PoseAsset)
+			{
+				PoseAssets.Add(PoseAsset);
+			}
+		}
+		if (PoseAssets.Num() == 1)
+		{
+			CommonControlNames = PoseAssets[0].Get()->Pose.GetControlNames();
+		}
+		else
+		{
+			TArray<FName> PossibleControlNames = PoseAssets[0].Get()->Pose.GetControlNames();
+
+			for (FName& ControlName : PossibleControlNames)
+			{
+				int32 Index = 1;
+				for (; Index < PoseAssets.Num(); ++Index)
+				{
+					if (PoseAssets[Index].Get()->Pose.ContainsName(ControlName) == false)
+					{
+						break;
+					}
+				}
+				if (Index == PoseAssets.Num())
+				{
+					CommonControlNames.Add(ControlName);
+				}
+			}
+		}
+	}
+
+	bool bResult;
+
+	TArray<TSharedPtr<FListItem>> Items;
+	TSharedPtr<SListView<TSharedPtr<FListItem>>> RenameList;
+
+	TArray<TWeakObjectPtr<UControlRigPoseAsset>> PoseAssets;
+	TArray<FName> CommonControlNames;
+};
+
+const FName SControlRigRenamePoseControls::SControlRenameListRow::NAME_Original = FName(TEXT("Original"));
+const FName SControlRigRenamePoseControls::SControlRenameListRow::NAME_New = FName(TEXT("New"));
+
+
+void FControlRigRenameControlsDialog::RenameControls(const TArray<UControlRigPoseAsset*>& PoseAssets)
+{
+	const FText TitleText = NSLOCTEXT("ControlRig", "RenameControlsOnSelectedAssets", "Rename Controls On Selected Assets");
+
+	// Create the window to choose our options
+	TSharedRef<SWindow> Window = SNew(SWindow)
+		.Title(TitleText)
+		.HasCloseButton(true)
+		.SizingRule(ESizingRule::UserSized)
+		.ClientSize(FVector2D(500.0f, 600.0f))
+		.AutoCenter(EAutoCenter::PreferredWorkArea)
+		.SupportsMinimize(false);
+
+	TSharedRef<SControlRigRenamePoseControls> DialogWidget = SNew(SControlRigRenamePoseControls, PoseAssets);
+	Window->SetContent(DialogWidget);
+
+	FSlateApplication::Get().AddWindow(Window);
+}
+
+#undef LOCTEXT_NAMESPACE
diff --git a/Engine/Plugins/Experimental/ControlRig/Source/ControlRigEditor/Private/EditMode/SControlRigRenamePoseControls.h b/Engine/Plugins/Experimental/ControlRig/Source/ControlRigEditor/Private/EditMode/SControlRigRenamePoseControls.h
new file mode 100644
--- /dev/null
+++ b/Engine/Plugins/Experimental/ControlRig/Source/ControlRigEditor/Private/EditMode/SControlRigRenamePoseControls.h
@@ -0,0 +1,11 @@
+// Copyright Epic Games, Inc. All Rights Reserved.
+#pragma once
+
+#include "CoreMinimal.h"
+
+class UControlRigPoseAsset;
+
+struct FControlRigRenameControlsDialog
+{
+	static void RenameControls(const TArray<UControlRigPoseAsset*>& PoseAssets);
+};
\ No newline at end of file
diff --git a/Engine/Plugins/Experimental/ControlRig/Source/ControlRigEditor/Private/EditMode/SControlRigUpdatePose.cpp b/Engine/Plugins/Experimental/ControlRig/Source/ControlRigEditor/Private/EditMode/SControlRigUpdatePose.cpp
new file mode 100644
--- /dev/null
+++ b/Engine/Plugins/Experimental/ControlRig/Source/ControlRigEditor/Private/EditMode/SControlRigUpdatePose.cpp
@@ -0,0 +1,129 @@
+// Copyright Epic Games, Inc. All Rights Reserved.
+
+#include "SControlRigUpdatePose.h"
+#include "CoreMinimal.h"
+#include "Widgets/SWidget.h"
+#include "Widgets/SCompoundWidget.h"
+#include "Editor.h"
+#include "Modules/ModuleManager.h"
+#include "Framework/Application/SlateApplication.h"
+#include "Widgets/SWindow.h"
+#include "Widgets/Input/SButton.h"
+#include "ControlRigEditMode.h"
+#include "Tools/ControlRigPose.h"
+#include "EditorModeManager.h"
+#include "ScopedTransaction.h"
+
+class SControlRigUpdatePoseDialog : public SCompoundWidget
+{
+	SLATE_BEGIN_ARGS(SControlRigUpdatePoseDialog) {}
+	SLATE_END_ARGS()
+	~SControlRigUpdatePoseDialog()
+	{
+	}
+	
+	void Construct(const FArguments& InArgs)
+	{
+
+		ChildSlot
+			[
+				SNew(SVerticalBox)
+
+				+ SVerticalBox::Slot()
+				.AutoHeight()
+				.HAlign(HAlign_Center)
+				.Padding(5.f)
+				[
+				SNew(STextBlock)
+				.Text(NSLOCTEXT("ControlRig", "UpdatePoseWithSelectedControls", "Update Pose With Selected Controls?"))
+				]
+				+ SVerticalBox::Slot()
+				.AutoHeight()
+				.HAlign(HAlign_Center)
+				.Padding(5.f)
+				[
+					SNew(SHorizontalBox)
+					+SHorizontalBox::Slot()
+					.AutoWidth()
+					.HAlign(HAlign_Left)
+					.Padding(10)
+					[
+						SNew(SButton)
+						.ContentPadding(FMargin(10, 5))
+						.Text(NSLOCTEXT("ControlRig", "Yes", "Yes"))
+						.OnClicked(this, &SControlRigUpdatePoseDialog::UpdatePose)
+					]
+					+SHorizontalBox::Slot()
+					.AutoWidth()
+					.HAlign(HAlign_Right)
+					.Padding(10)
+					[
+						SNew(SButton)
+						.ContentPadding(FMargin(10, 5))
+						.Text(NSLOCTEXT("ControlRig", "No", "No"))
+						.OnClicked(this, &SControlRigUpdatePoseDialog::CancelUpdatePose)
+					]
+				]
+			];
+	}
+
+	void SetPose(UControlRigPoseAsset* InPoseAsset)
+	{
+		PoseAsset = InPoseAsset;
+	}
+
+private:
+
+	FReply UpdatePose()
+	{
+		TSharedPtr<SWindow> Window = FSlateApplication::Get().FindWidgetWindow(AsShared());
+
+		FControlRigEditMode* ControlRigEditMode = static_cast<FControlRigEditMode*>(GLevelEditorModeTools().GetActiveMode(FControlRigEditMode::ModeName));
+		if (PoseAsset.IsValid() && ControlRigEditMode && ControlRigEditMode->GetControlRig(true))
+		{
+			const FScopedTransaction Transaction(NSLOCTEXT("ControlRig", "UpdatePose", "Update Pose"));
+			PoseAsset->Modify();
+			PoseAsset->SavePose(ControlRigEditMode->GetControlRig(true), false);
+		}
+		if (Window.IsValid())
+		{
+			Window->RequestDestroyWindow();
+		}
+		return FReply::Handled();
+	}
+
+	FReply CancelUpdatePose()
+	{
+		TSharedPtr<SWindow> Window = FSlateApplication::Get().FindWidgetWindow(AsShared());
+		if (Window.IsValid())
+		{
+			Window->RequestDestroyWindow();
+		}
+		return FReply::Handled();
+	}
+
+	TSharedPtr<IDetailsView> DetailView;
+	TWeakObjectPtr<UControlRigPoseAsset> PoseAsset;
+};
+
+
+void FControlRigUpdatePoseDialog::UpdatePose(UControlRigPoseAsset* PoseAsset)
+{
+	const FText TitleText = NSLOCTEXT("ControlRig", "UpdateControlRigPose", "Update Control Rig Pose");
+
+	// Create the window to choose our options
+	TSharedRef<SWindow> Window = SNew(SWindow)
+		.Title(TitleText)
+		.HasCloseButton(true)
+		.SizingRule(ESizingRule::UserSized)
+		.ClientSize(FVector2D(400.0f, 100.0f))
+		.AutoCenter(EAutoCenter::PreferredWorkArea)
+		.SupportsMinimize(false);
+
+	TSharedRef<SControlRigUpdatePoseDialog> DialogWidget = SNew(SControlRigUpdatePoseDialog);
+	DialogWidget->SetPose(PoseAsset);
+	Window->SetContent(DialogWidget);
+
+	FSlateApplication::Get().AddWindow(Window);
+}
+
diff --git a/Engine/Plugins/Experimental/ControlRig/Source/ControlRigEditor/Private/EditMode/SControlRigUpdatePose.h b/Engine/Plugins/Experimental/ControlRig/Source/ControlRigEditor/Private/EditMode/SControlRigUpdatePose.h
new file mode 100644
--- /dev/null
+++ b/Engine/Plugins/Experimental/ControlRig/Source/ControlRigEditor/Private/EditMode/SControlRigUpdatePose.h
@@ -0,0 +1,11 @@
+// Copyright Epic Games, Inc. All Rights Reserved.
+#pragma once
+
+#include "CoreMinimal.h"
+
+class UControlRigPoseAsset;
+
+struct FControlRigUpdatePoseDialog
+{
+	static void UpdatePose(UControlRigPoseAsset* PoseAsset);
+};
\ No newline at end of file
diff --git a/Engine/Plugins/Experimental/ControlRig/Source/ControlRigEditor/Private/Tools/AssetTypeActions_ControlRigPose.cpp b/Engine/Plugins/Experimental/ControlRig/Source/ControlRigEditor/Private/Tools/AssetTypeActions_ControlRigPose.cpp
--- a/Engine/Plugins/Experimental/ControlRig/Source/ControlRigEditor/Private/Tools/AssetTypeActions_ControlRigPose.cpp
+++ b/Engine/Plugins/Experimental/ControlRig/Source/ControlRigEditor/Private/Tools/AssetTypeActions_ControlRigPose.cpp
@@ -3,12 +3,20 @@
 #include "AssetTypeActions_ControlRigPose.h"
 #include "EditorFramework/AssetImportData.h"
 #include "ThumbnailRendering/SceneThumbnailInfo.h"
-#include "Tools/ControlRigPose.h"
 #include "ToolMenus.h"
 
+#include "ControlRig.h"
+#include "UnrealEdGlobals.h"
+#include "ControlRigEditMode.h"
+#include "Tools/ControlRigPose.h"
+#include "EditorModeManager.h"
+#include "SControlRigRenamePoseControls.h"
+
+#define LOCTEXT_NAMESPACE "AssetTypeActions_ControlRigPose"
+
 FText FAssetTypeActions_ControlRigPose::GetName() const
 {
-	return NSLOCTEXT("AssetTypeActions", "AssetTypeActions_ControlRigPose", "Control Pose");
+	return LOCTEXT("AssetTypeActions_ControlRigPose", "Control Rig Pose");
 }
 
 FColor FAssetTypeActions_ControlRigPose::GetTypeColor() const
@@ -59,8 +67,117 @@ bool FAssetTypeActions_ControlRigPose::IsImportedAsset() const
 
 void FAssetTypeActions_ControlRigPose::GetActions(const TArray<UObject*>& InObjects, FToolMenuSection& Section)
 {
-	//TODO Add custom Actions based on feedback
+	if (InObjects.Num() == 1)
+	{
+		UObject* SelectedAsset = InObjects[0];
+		if (SelectedAsset == nullptr)
+		{
+			return;
+		}
+		FControlRigEditMode* ControlRigEditMode = static_cast<FControlRigEditMode*>(GLevelEditorModeTools().GetActiveMode(FControlRigEditMode::ModeName));
+		if (ControlRigEditMode && ControlRigEditMode->GetControlRig(true))
+		{
+			UControlRig* ControlRig = ControlRigEditMode->GetControlRig(true);
+			UControlRigPoseAsset* PoseAsset = Cast<UControlRigPoseAsset>(SelectedAsset);
+			if (PoseAsset)
+			{
+				Section.AddDynamicEntry("Control Rig Pose Actions", FNewToolMenuSectionDelegate::CreateLambda([ControlRig,PoseAsset](FToolMenuSection& InSection)
+					{
+						{
+							const FText Label = LOCTEXT("PastePoseButton", "Paste Pose");
+							const FText ToolTipText = LOCTEXT("PastePoseButtonTooltip", "Paste the selected pose");
+							InSection.AddMenuEntry(
+								"PastePose",
+								Label,
+								ToolTipText,
+								FSlateIcon(),
+								FUIAction(
+									FExecuteAction::CreateLambda([ControlRig,PoseAsset]()
+										{
+											PoseAsset->PastePose(ControlRig, false, false);
+										})));
+						}
+						{
+							const FText Label = LOCTEXT("SelectControls", "Select Controls");
+							const FText ToolTipText = LOCTEXT("SelectControlsTooltip", "Select controls in this pose on active control rig");
+							InSection.AddMenuEntry(
+								"SelectControls",
+								Label,
+								ToolTipText,
+								FSlateIcon(),
+								FUIAction(
+									FExecuteAction::CreateLambda([ControlRig, PoseAsset]()
+										{
+											PoseAsset->SelectControls(ControlRig);
+										})));
+						}
+						{
+							const FText Label = LOCTEXT("UpdatePose", "Update Pose");
+							const FText ToolTipText = LOCTEXT("UpdatePoseTooltip", "Update the pose based upon current control rig pose and selected controls");
+							InSection.AddMenuEntry(
+								"UpdatePose",
+								Label,
+								ToolTipText,
+								FSlateIcon(),
+								FUIAction(
+									FExecuteAction::CreateLambda([ControlRig,PoseAsset]()
+										{
+											PoseAsset->SavePose(ControlRig,false);
+										})));
+						}
+						{
+
+							const FText Label = LOCTEXT("RenameControls", "Rename Controls");
+							const FText ToolTipText = LOCTEXT("RenameControlsTooltip", "Rename controls on selected poses");
+							InSection.AddMenuEntry(
+								"RenameControls",
+								Label,
+								ToolTipText,
+								FSlateIcon(),
+								FUIAction(
+									FExecuteAction::CreateLambda([PoseAsset]()
+										{
+											TArray<UControlRigPoseAsset*> PoseAssets;
+											PoseAssets.Add(PoseAsset);
+											FControlRigRenameControlsDialog::RenameControls(PoseAssets);
+										})));
+						}
+
+					}));
+			}
+		}
+	}
+	else if (InObjects.Num() > 1)
+	{
+		TArray<UControlRigPoseAsset*> PoseAssets;
+		for (UObject* Object : InObjects)
+		{
+			if (Object && Object->IsA<UControlRigPoseAsset>())
+			{
+				UControlRigPoseAsset* PoseAsset = Cast<UControlRigPoseAsset>(Object);
+				PoseAssets.Add(PoseAsset);
+			}
+		}
+		Section.AddDynamicEntry("Control Rig Pose Actions", FNewToolMenuSectionDelegate::CreateLambda([PoseAssets](FToolMenuSection& InSection)
+			{
+				{
+
+					const FText Label = LOCTEXT("RenameControls", "Rename Controls");
+					const FText ToolTipText = LOCTEXT("RenameControlsTooltip", "Rename controls on selected poses");
+					InSection.AddMenuEntry(
+						"RenameControls",
+						Label,
+						ToolTipText,
+						FSlateIcon(),
+						FUIAction(
+							FExecuteAction::CreateLambda([PoseAssets]()
+								{
+									FControlRigRenameControlsDialog::RenameControls(PoseAssets);
+								})));
+				}
+			}));
+
+	}
+}
 
-//	FAssetTypeActions_Base::GetActions(InObjects, Section);
-//	TArray<TWeakObjectPtr<UControlRigPoseAsset>> ControlRigPoseAsset = GetTypedWeakObjectPtrs<UControlRigPoseAsset>(InObjects);
-}
\ No newline at end of file
+#undef LOCTEXT_NAMESPACE
