Sequencer: Pose Library, Missing from last submit, rename action available if control rig doesn't exist.
#trivial

diff --git a/Engine/Plugins/Experimental/ControlRig/Source/ControlRigEditor/Private/Tools/AssetTypeActions_ControlRigPose.cpp b/Engine/Plugins/Experimental/ControlRig/Source/ControlRigEditor/Private/Tools/AssetTypeActions_ControlRigPose.cpp
--- a/Engine/Plugins/Experimental/ControlRig/Source/ControlRigEditor/Private/Tools/AssetTypeActions_ControlRigPose.cpp
+++ b/Engine/Plugins/Experimental/ControlRig/Source/ControlRigEditor/Private/Tools/AssetTypeActions_ControlRigPose.cpp
@@ -74,13 +74,15 @@ void FAssetTypeActions_ControlRigPose::GetActions(const TArray<UObject*>& InObje
 		{
 			return;
 		}
-		FControlRigEditMode* ControlRigEditMode = static_cast<FControlRigEditMode*>(GLevelEditorModeTools().GetActiveMode(FControlRigEditMode::ModeName));
-		if (ControlRigEditMode && ControlRigEditMode->GetControlRig(true))
+
+		UControlRigPoseAsset* PoseAsset = Cast<UControlRigPoseAsset>(SelectedAsset);
+		if (PoseAsset)
 		{
-			UControlRig* ControlRig = ControlRigEditMode->GetControlRig(true);
-			UControlRigPoseAsset* PoseAsset = Cast<UControlRigPoseAsset>(SelectedAsset);
-			if (PoseAsset)
-			{
+			FControlRigEditMode* ControlRigEditMode = static_cast<FControlRigEditMode*>(GLevelEditorModeTools().GetActiveMode(FControlRigEditMode::ModeName));
+			if (ControlRigEditMode && ControlRigEditMode->GetControlRig(true))
+			{ 
+				UControlRig* ControlRig = ControlRigEditMode->GetControlRig(true);
+
 				Section.AddDynamicEntry("Control Rig Pose Actions", FNewToolMenuSectionDelegate::CreateLambda([ControlRig,PoseAsset](FToolMenuSection& InSection)
 					{
 						{
@@ -145,6 +147,29 @@ void FAssetTypeActions_ControlRigPose::GetActions(const TArray<UObject*>& InObje
 
 					}));
 			}
+			else
+			{
+				Section.AddDynamicEntry("Control Rig Pose Actions", FNewToolMenuSectionDelegate::CreateLambda([PoseAsset](FToolMenuSection& InSection)
+				{
+					{
+
+						const FText Label = LOCTEXT("RenameControls", "Rename Controls");
+						const FText ToolTipText = LOCTEXT("RenameControlsTooltip", "Rename controls on selected poses");
+						InSection.AddMenuEntry(
+							"RenameControls",
+							Label,
+							ToolTipText,
+							FSlateIcon(),
+							FUIAction(
+								FExecuteAction::CreateLambda([PoseAsset]()
+									{
+										TArray<UControlRigPoseAsset*> PoseAssets;
+										PoseAssets.Add(PoseAsset);
+										FControlRigRenameControlsDialog::RenameControls(PoseAssets);
+									})));
+					}
+				}));
+			}
 		}
 	}
 	else if (InObjects.Num() > 1)
