Control Rig: Add ability to drag and drop a control rig into Sequencer, which creates a skeletal mesh actor from the given preview mesh and a control rig parameter track for the control rig.

diff --git a/Engine/Plugins/Experimental/ControlRig/Source/ControlRigEditor/Private/Sequencer/ControlRigParameterTrackEditor.cpp b/Engine/Plugins/Experimental/ControlRig/Source/ControlRigEditor/Private/Sequencer/ControlRigParameterTrackEditor.cpp
--- a/Engine/Plugins/Experimental/ControlRig/Source/ControlRigEditor/Private/Sequencer/ControlRigParameterTrackEditor.cpp
+++ b/Engine/Plugins/Experimental/ControlRig/Source/ControlRigEditor/Private/Sequencer/ControlRigParameterTrackEditor.cpp
@@ -16,6 +16,7 @@
 #include "Widgets/Layout/SBox.h"
 #include "SequencerSectionPainter.h"
 #include "Components/SkeletalMeshComponent.h"
+#include "Editor.h"
 #include "Editor/UnrealEdEngine.h"
 #include "UnrealEdGlobals.h"
 #include "ClassViewerModule.h"
@@ -29,6 +30,7 @@
 #include "SequencerUtilities.h"
 #include "ISectionLayoutBuilder.h"
 #include "EditorStyleSet.h"
+#include "MovieSceneCommonHelpers.h"
 #include "MovieSceneTimeHelpers.h"
 #include "Fonts/FontMeasure.h"
 #include "AnimationEditorUtils.h"
@@ -66,6 +68,7 @@
 #include "SBakeToControlRigDialog.h"
 #include "ControlRigBlueprint.h"
 #include "ControlRigBlueprintGeneratedClass.h"
+#include "Animation/SkeletalMeshActor.h"
 
 #define LOCTEXT_NAMESPACE "FControlRigParameterTrackEditor"
 
@@ -2104,6 +2107,55 @@ void FControlRigParameterTrackEditor::BuildTrackContextMenu(FMenuBuilder& MenuBu
 
 }
 
+bool FControlRigParameterTrackEditor::HandleAssetAdded(UObject* Asset, const FGuid& TargetObjectGuid)
+{
+	if (!Asset->IsA<UControlRigBlueprint>())
+	{
+		return false;
+	}
+	
+	UControlRigBlueprint* ControlRigBlueprint = Cast<UControlRigBlueprint>(Asset);
+	UControlRigBlueprintGeneratedClass* RigClass = ControlRigBlueprint->GetControlRigBlueprintGeneratedClass();
+	if (!RigClass)
+	{
+		return false;
+	}
+
+	USkeletalMesh* SkeletalMesh = ControlRigBlueprint->GetPreviewMesh();
+	if (!SkeletalMesh)
+	{
+		FNotificationInfo Info(LOCTEXT("NoPreviewMesh", "Control rig has no preview mesh to create a spawnable skeletal mesh actor from"));
+		Info.ExpireDuration = 5.0f;
+		FSlateNotificationManager::Get().AddNotification(Info)->SetCompletionState(SNotificationItem::CS_Fail);
+		return false;
+	}
+
+	const FScopedTransaction Transaction(LOCTEXT("AddControlRig", "Add Control Rig"));
+
+	// Spawn a skeletal mesh actor with the preview mesh
+	ASkeletalMeshActor* SkeletalMeshActor = GCurrentLevelEditingViewportClient->GetWorld()->SpawnActor<ASkeletalMeshActor>();
+	SkeletalMeshActor->GetSkeletalMeshComponent()->SetSkeletalMesh(SkeletalMesh);
+
+	FGuid NewGuid = GetSequencer()->MakeNewSpawnable(*SkeletalMeshActor);
+	UObject* SpawnedSkeletalMeshActor = GetSequencer()->FindSpawnedObjectOrTemplate(NewGuid);
+	if (SkeletalMeshActor)
+	{
+		GCurrentLevelEditingViewportClient->GetWorld()->EditorDestroyActor(SkeletalMeshActor, true);
+		SkeletalMeshActor = Cast<ASkeletalMeshActor>(SpawnedSkeletalMeshActor);
+	}
+
+	FString NewName = MovieSceneHelpers::MakeUniqueSpawnableName(GetSequencer()->GetFocusedMovieSceneSequence()->GetMovieScene(), FName::NameToDisplayString(SkeletalMesh->GetName(), false));
+	SkeletalMeshActor->SetActorLabel(NewName, false);
+
+	UControlRig* CDO = Cast<UControlRig>(RigClass->GetDefaultObject(true /* create if needed */));
+	check(CDO);
+
+	AddControlRig(CDO->GetClass(), SkeletalMeshActor->GetSkeletalMeshComponent(), NewGuid);
+
+	return true;
+}
+
+
 void FControlRigParameterTrackEditor::ToggleFKControlRig(UMovieSceneControlRigParameterTrack* Track, UFKControlRig* FKControlRig)
 {
 	FScopedTransaction Transaction(LOCTEXT("ToggleFKControlRig", "Toggle FK Control Rig"));
diff --git a/Engine/Plugins/Experimental/ControlRig/Source/ControlRigEditor/Private/Sequencer/ControlRigParameterTrackEditor.h b/Engine/Plugins/Experimental/ControlRig/Source/ControlRigEditor/Private/Sequencer/ControlRigParameterTrackEditor.h
--- a/Engine/Plugins/Experimental/ControlRig/Source/ControlRigEditor/Private/Sequencer/ControlRigParameterTrackEditor.h
+++ b/Engine/Plugins/Experimental/ControlRig/Source/ControlRigEditor/Private/Sequencer/ControlRigParameterTrackEditor.h
@@ -60,6 +60,7 @@ public:
 	virtual bool HasTransformKeyOverridePriority() const override;
 	virtual void ObjectImplicitlyAdded(UObject* InObject)  override;
 	virtual void BuildTrackContextMenu(FMenuBuilder& MenuBuilder, UMovieSceneTrack* InTrack) override;
+	virtual bool HandleAssetAdded(UObject* Asset, const FGuid& TargetObjectGuid) override;
 
 private:
 
