Control Rig: Don't add for editor preview actors, which ultimately get swapped out after dragging and dropping into the viewport. This fixes a bug where if you drag a control rig into the viewport, it will create a binding and track for the preview actor and a binding and track for the resulting dropped actor as well.

diff --git a/Engine/Plugins/Experimental/ControlRig/Source/ControlRigEditor/Private/ControlRigBlueprintActions.cpp b/Engine/Plugins/Experimental/ControlRig/Source/ControlRigEditor/Private/ControlRigBlueprintActions.cpp
--- a/Engine/Plugins/Experimental/ControlRig/Source/ControlRigEditor/Private/ControlRigBlueprintActions.cpp
+++ b/Engine/Plugins/Experimental/ControlRig/Source/ControlRigEditor/Private/ControlRigBlueprintActions.cpp
@@ -199,7 +199,7 @@ USkeletalMesh* FControlRigBlueprintActions::GetSkeletalMeshFromControlRigBluepri
 
 void FControlRigBlueprintActions::PostSpawningSkeletalMeshActor(AActor* InSpawnedActor, UObject* InAsset)
 {
-	if (InSpawnedActor->HasAnyFlags(RF_Transient))
+	if (InSpawnedActor->HasAnyFlags(RF_Transient) || InSpawnedActor->bIsEditorPreviewActor)
 	{
 		return;
 	}
