Control Rig: Fix CIS static analysis warning

diff --git a/Engine/Plugins/Experimental/ControlRig/Source/ControlRigDeveloper/Private/ControlRigBlueprint.cpp b/Engine/Plugins/Experimental/ControlRig/Source/ControlRigDeveloper/Private/ControlRigBlueprint.cpp
--- a/Engine/Plugins/Experimental/ControlRig/Source/ControlRigDeveloper/Private/ControlRigBlueprint.cpp
+++ b/Engine/Plugins/Experimental/ControlRig/Source/ControlRigDeveloper/Private/ControlRigBlueprint.cpp
@@ -787,11 +787,6 @@ FName UControlRigBlueprint::AddMemberVariable(const FName& InName, const FString
 		Variable.TypeName = *CPPType;
 		Variable.Size = sizeof(FName);
 	}
-	else if (CPPType == TEXT("FName"))
-	{
-		Variable.TypeName = *CPPType;
-		Variable.Size = sizeof(FName);
-	}
 	else if(UScriptStruct* ScriptStruct = URigVMPin::FindObjectFromCPPTypeObjectPath<UScriptStruct>(CPPType))
 	{
 		Variable.TypeName = *ScriptStruct->GetStructCPPName();
