Control Rig: Disallow duplication of graphs

diff --git a/Engine/Plugins/Experimental/ControlRig/Source/ControlRigDeveloper/Public/Graph/ControlRigGraphSchema.h b/Engine/Plugins/Experimental/ControlRig/Source/ControlRigDeveloper/Public/Graph/ControlRigGraphSchema.h
--- a/Engine/Plugins/Experimental/ControlRig/Source/ControlRigDeveloper/Public/Graph/ControlRigGraphSchema.h
+++ b/Engine/Plugins/Experimental/ControlRig/Source/ControlRigDeveloper/Public/Graph/ControlRigGraphSchema.h
@@ -80,6 +80,7 @@ public:
 	virtual void GetGraphDisplayInformation(const UEdGraph& Graph, /*out*/ FGraphDisplayInfo& DisplayInfo) const override;
 	virtual bool TryDeleteGraph(UEdGraph* GraphToDelete) const override;
 	virtual bool TryRenameGraph(UEdGraph* GraphToRename, const FName& InNewName) const override;
+	virtual bool CanDuplicateGraph(UEdGraph* InSourceGraph) const { return false; }
 
 	/** Create a graph node for a rig */
 	UControlRigGraphNode* CreateGraphNode(UControlRigGraph* InGraph, const FName& InPropertyName) const;
