Control Rig: Fix static analysis

diff --git a/Engine/Plugins/Experimental/ControlRig/Source/ControlRigEditor/Private/Sequencer/ControlRigParameterTrackEditor.cpp b/Engine/Plugins/Experimental/ControlRig/Source/ControlRigEditor/Private/Sequencer/ControlRigParameterTrackEditor.cpp
--- a/Engine/Plugins/Experimental/ControlRig/Source/ControlRigEditor/Private/Sequencer/ControlRigParameterTrackEditor.cpp
+++ b/Engine/Plugins/Experimental/ControlRig/Source/ControlRigEditor/Private/Sequencer/ControlRigParameterTrackEditor.cpp
@@ -2200,6 +2200,11 @@ bool FControlRigParameterTrackEditor::HandleAssetAdded(UObject* Asset, const FGu
 
 	// Spawn a skeletal mesh actor with the preview mesh
 	ASkeletalMeshActor* SkeletalMeshActor = GCurrentLevelEditingViewportClient->GetWorld()->SpawnActor<ASkeletalMeshActor>();
+	if (!SkeletalMeshActor)
+	{
+		return false;
+	}
+
 	SkeletalMeshActor->GetSkeletalMeshComponent()->SetSkeletalMesh(SkeletalMesh);
 
 	FGuid NewGuid = GetSequencer()->MakeNewSpawnable(*SkeletalMeshActor);
