Fix editor exit crash.

diff --git a/Engine/Plugins/Experimental/ControlRig/Source/ControlRigEditor/Private/ControlRigEditorModule.cpp b/Engine/Plugins/Experimental/ControlRig/Source/ControlRigEditor/Private/ControlRigEditorModule.cpp
--- a/Engine/Plugins/Experimental/ControlRig/Source/ControlRigEditor/Private/ControlRigEditorModule.cpp
+++ b/Engine/Plugins/Experimental/ControlRig/Source/ControlRigEditor/Private/ControlRigEditorModule.cpp
@@ -289,11 +289,14 @@ void FControlRigEditorModule::ShutdownModule()
 
 	FKismetEditorUtilities::UnregisterAutoBlueprintNodeCreation(this);
 
-	FBlueprintEditorModule* BlueprintEditorModule = FModuleManager::GetModulePtr<FBlueprintEditorModule>("Kismet");
-	if (BlueprintEditorModule)
+	if (!IsEngineExitRequested())
 	{
-		BlueprintEditorModule->UnregisterVariableCustomization(FProperty::StaticClass());
-		BlueprintEditorModule->UnregisterGraphCustomization(GetDefault<UControlRigGraphSchema>());
+		FBlueprintEditorModule* BlueprintEditorModule = FModuleManager::GetModulePtr<FBlueprintEditorModule>("Kismet");
+		if (BlueprintEditorModule)
+		{
+			BlueprintEditorModule->UnregisterVariableCustomization(FProperty::StaticClass());
+			BlueprintEditorModule->UnregisterGraphCustomization(GetDefault<UControlRigGraphSchema>());
+		}
 	}
 
 	FPropertyEditorModule* PropertyEditorModule = FModuleManager::GetModulePtr<FPropertyEditorModule>("PropertyEditor");
