UE-103383 - Control Rig Data Source RegistryMember to use NewObject

diff --git a/Engine/Plugins/Experimental/ControlRig/Source/ControlRig/Private/ControlRig.cpp b/Engine/Plugins/Experimental/ControlRig/Source/ControlRig/Private/ControlRig.cpp
--- a/Engine/Plugins/Experimental/ControlRig/Source/ControlRig/Private/ControlRig.cpp
+++ b/Engine/Plugins/Experimental/ControlRig/Source/ControlRig/Private/ControlRig.cpp
@@ -65,8 +65,6 @@ UControlRig::UControlRig(const FObjectInitializer& ObjectInitializer)
 	, InterRigSyncBracket(0)
 {
 	VM = ObjectInitializer.CreateDefaultSubobject<URigVM>(this, TEXT("VM"));
-	// create default source registry
-	DataSourceRegistry = CreateDefaultSubobject<UAnimationDataSourceRegistry>(TEXT("DataSourceRegistry"));
 
 	EventQueue.Add(FRigUnit_BeginExecution::EventName);
 }
@@ -466,7 +464,7 @@ void UControlRig::Execute(const EControlRigState InState, const FName& InEventNa
 
 	if (InState == EControlRigState::Init)
 	{
-		Context.DataSourceRegistry = DataSourceRegistry;
+		Context.DataSourceRegistry = GetDataSourceRegistry();
 		AbsoluteTime = DeltaTime = 0.f;
 	}
 
@@ -484,7 +482,7 @@ void UControlRig::Execute(const EControlRigState InState, const FName& InEventNa
 
 	if (!OuterSceneComponent.IsValid())
 	{
-		USceneComponent* SceneComponentFromRegistry = DataSourceRegistry->RequestSource<USceneComponent>(UControlRig::OwnerComponent);
+		USceneComponent* SceneComponentFromRegistry = Context.DataSourceRegistry->RequestSource<USceneComponent>(UControlRig::OwnerComponent);
 		if (SceneComponentFromRegistry)
 		{
 			OuterSceneComponent = SceneComponentFromRegistry;
@@ -993,6 +991,15 @@ void UControlRig::GetMappableNodeData(TArray<FName>& OutNames, TArray<FNodeItem>
 	}
 }
 
+UAnimationDataSourceRegistry* UControlRig::GetDataSourceRegistry()
+{
+	if (DataSourceRegistry == nullptr)
+	{
+		DataSourceRegistry = NewObject<UAnimationDataSourceRegistry>(this, TEXT("DataSourceRegistry"));
+	}
+	return DataSourceRegistry;
+}
+
 #if WITH_EDITORONLY_DATA
 
 void UControlRig::PostReinstanceCallback(const UControlRig* Old)
diff --git a/Engine/Plugins/Experimental/ControlRig/Source/ControlRig/Public/ControlRig.h b/Engine/Plugins/Experimental/ControlRig/Source/ControlRig/Public/ControlRig.h
--- a/Engine/Plugins/Experimental/ControlRig/Source/ControlRig/Public/ControlRig.h
+++ b/Engine/Plugins/Experimental/ControlRig/Source/ControlRig/Public/ControlRig.h
@@ -282,7 +282,7 @@ public:
 	virtual void GetMappableNodeData(TArray<FName>& OutNames, TArray<FNodeItem>& OutNodeItems) const override;
 
 	/** Data Source Registry Getter */
-	UAnimationDataSourceRegistry* GetDataSourceRegistry() { return DataSourceRegistry; }
+	UAnimationDataSourceRegistry* GetDataSourceRegistry();
 
 	virtual const TArray<FRigSpace>& AvailableSpaces() const;
 	virtual FRigSpace* FindSpace(const FName& InSpaceName);
