Control Rig: Fixed "Deleting all imported bones is not shown correctly the rig hierarchy tree view and leads to crash when click".

This happened because Ctrl+A selects the last placeholder item in addition to other regular rig elements. In the old logic, we loop through all selected items and when it hits the placeholder item it triggers an early return. As a result, refresh tree view is not getting called if this placeholder item is included in the selection.

diff --git a/Engine/Plugins/Experimental/ControlRig/Source/ControlRigEditor/Private/Editor/SRigHierarchy.cpp b/Engine/Plugins/Experimental/ControlRig/Source/ControlRigEditor/Private/Editor/SRigHierarchy.cpp
--- a/Engine/Plugins/Experimental/ControlRig/Source/ControlRigEditor/Private/Editor/SRigHierarchy.cpp
+++ b/Engine/Plugins/Experimental/ControlRig/Source/ControlRigEditor/Private/Editor/SRigHierarchy.cpp
@@ -1570,6 +1570,7 @@ void SRigHierarchy::HandleDeleteItem()
 						}
 						Hierarchy->BoneHierarchy.Remove(SelectedKey.Name);
 					}
+					RemovedItems.Add(SelectedKey);
 					break;
 				}
 				case ERigElementType::Control:
@@ -1578,6 +1579,7 @@ void SRigHierarchy::HandleDeleteItem()
 					{
 						Hierarchy->ControlHierarchy.Remove(SelectedKey.Name);
 					}
+					RemovedItems.Add(SelectedKey);
 					break;
 				}
 				case ERigElementType::Space:
@@ -1586,15 +1588,14 @@ void SRigHierarchy::HandleDeleteItem()
 					{
 						Hierarchy->SpaceHierarchy.Remove(SelectedKey.Name);
 					}
+					RemovedItems.Add(SelectedKey);
 					break;
 				}
 				default:
 				{
-					return;
+					break;
 				}
-			}
-
-			RemovedItems.Add(SelectedKey);
+			} 
 		}
 
 		for (const FRigElementKey& RemovedItem : RemovedItems)
