Sequencer: Control Rig: Order Of Controls Different Between Details And Sequencer

Fix involves relying on "GetControlsInOrder" as well for the edit mode tool widget - so that things match up.

Merging

//UE4/Release-4.26/Engine/Plugins/Experimental/ControlRig/Source/ControlRigEditor/Private/EditMode/SControlHierarchy.cpp

to //UE5/Main/Engine/Plugins/Experimental/ControlRig/Source/ControlRigEditor/Private/EditMode/SControlHierarchy.cpp

diff --git a/Engine/Plugins/Experimental/ControlRig/Source/ControlRigEditor/Private/EditMode/SControlHierarchy.cpp b/Engine/Plugins/Experimental/ControlRig/Source/ControlRigEditor/Private/EditMode/SControlHierarchy.cpp
--- a/Engine/Plugins/Experimental/ControlRig/Source/ControlRigEditor/Private/EditMode/SControlHierarchy.cpp
+++ b/Engine/Plugins/Experimental/ControlRig/Source/ControlRigEditor/Private/EditMode/SControlHierarchy.cpp
@@ -192,8 +192,9 @@ void SControlHierarchy::RefreshTreeView()
 
 	if (ControlRig.IsValid())
 	{
-		FRigControlHierarchy& ControlHierarchy = ControlRig->GetControlHierarchy();
-		for (const FRigControl& Element : ControlHierarchy)
+		TArray<FRigControl> SortedControls;
+		ControlRig->GetControlsInOrder(SortedControls);
+		for (const FRigControl& Element : SortedControls)
 		{
 			if (!ControlRig->IsCurveControl(&Element))
 			{
@@ -224,6 +225,7 @@ void SControlHierarchy::RefreshTreeView()
 
 		TreeView->RequestTreeRefresh();
 
+		FRigControlHierarchy& ControlHierarchy = ControlRig->GetControlHierarchy();
 		TArray<FName> Selection = ControlHierarchy.CurrentSelection();
 		for (const FName& Name : Selection)
 		{
