Control Rig: Rely on central functionality for name validation

diff --git a/Engine/Plugins/Experimental/ControlRig/Source/ControlRig/Private/Rigs/RigBoneHierarchy.cpp b/Engine/Plugins/Experimental/ControlRig/Source/ControlRig/Private/Rigs/RigBoneHierarchy.cpp
--- a/Engine/Plugins/Experimental/ControlRig/Source/ControlRig/Private/Rigs/RigBoneHierarchy.cpp
+++ b/Engine/Plugins/Experimental/ControlRig/Source/ControlRig/Private/Rigs/RigBoneHierarchy.cpp
@@ -1,6 +1,7 @@
 // Copyright Epic Games, Inc. All Rights Reserved.
 
 #include "Rigs/RigBoneHierarchy.h"
+#include "Rigs/RigHierarchyContainer.h"
 #include "ControlRig.h"
 #include "HelperUtil.h"
 #include "AnimationRuntime.h"
@@ -46,21 +47,32 @@ FRigBoneHierarchy& FRigBoneHierarchy::operator= (const FRigBoneHierarchy &InOthe
 	return *this;
 }
 
-FName FRigBoneHierarchy::GetSafeNewName(const FName& InPotentialNewName) const
+bool FRigBoneHierarchy::IsNameAvailable(const FString& InPotentialNewName, FString* OutErrorMessage) const
 {
-	FName Name = InPotentialNewName;
-	int32 Suffix = 1;
-	while(!IsNameAvailable(Name))
+	if (Container)
 	{
-		Name = *FString::Printf(TEXT("%s_%d"), *InPotentialNewName.ToString(), ++Suffix);
+		return Container->IsNameAvailable(InPotentialNewName, ERigElementType::Bone, OutErrorMessage);
 	}
-	return Name;
+
+	// if we are not part of a container, fall back to index based validation
+	if (InPotentialNewName.Len() > FRigHierarchyContainer::GetMaxNameLength())
+	{
+		return false;
+	}
+	return GetIndex(*InPotentialNewName) == INDEX_NONE;
+}
+
+FName FRigBoneHierarchy::GetSafeNewName(const FString& InPotentialNewName) const
+{
+	return FRigHierarchyContainer::GetSafeNewName(InPotentialNewName, [this](const FString& InName) -> bool {
+		return IsNameAvailable(InName);
+	});
 }
 
 FRigBone& FRigBoneHierarchy::Add(const FName& InNewName, const FName& InParentName, ERigBoneType InType, const FTransform& InInitTransform)
 {
 	FRigBone NewBone;
-	NewBone.Name = GetSafeNewName(InNewName);
+	NewBone.Name = GetSafeNewName(InNewName.ToString());
 	NewBone.ParentIndex = GetIndex(InParentName);
 	NewBone.ParentName = NewBone.ParentIndex == INDEX_NONE ? NAME_None : InParentName;
 	NewBone.InitialTransform = InInitTransform;
@@ -418,8 +430,7 @@ FName FRigBoneHierarchy::Rename(const FName& InOldName, const FName& InNewName)
 		const int32 Found = GetIndex(InOldName);
 		if (Found != INDEX_NONE)
 		{
-			FName NewName = GetSafeNewName(InNewName);
-
+			FName NewName = GetSafeNewName(InNewName.ToString());
 			bool bWasSelected = IsSelected(InOldName);
 			if(bWasSelected)
 			{
diff --git a/Engine/Plugins/Experimental/ControlRig/Source/ControlRig/Private/Rigs/RigControlHierarchy.cpp b/Engine/Plugins/Experimental/ControlRig/Source/ControlRig/Private/Rigs/RigControlHierarchy.cpp
--- a/Engine/Plugins/Experimental/ControlRig/Source/ControlRig/Private/Rigs/RigControlHierarchy.cpp
+++ b/Engine/Plugins/Experimental/ControlRig/Source/ControlRig/Private/Rigs/RigControlHierarchy.cpp
@@ -1,6 +1,7 @@
 // Copyright Epic Games, Inc. All Rights Reserved.
 
 #include "Rigs/RigControlHierarchy.h"
+#include "Rigs/RigHierarchyContainer.h"
 #include "ControlRig.h"
 #include "HelperUtil.h"
 #include "Settings/ControlRigSettings.h"
@@ -502,15 +503,26 @@ FRigControlHierarchy& FRigControlHierarchy::operator= (const FRigControlHierarch
 	return *this;
 }
 
-FName FRigControlHierarchy::GetSafeNewName(const FName& InPotentialNewName) const
+bool FRigControlHierarchy::IsNameAvailable(const FString& InPotentialNewName, FString* OutErrorMessage) const
 {
-	FName Name = InPotentialNewName;
-	int32 Suffix = 1;
-	while(!IsNameAvailable(Name))
+	if (Container)
 	{
-		Name = *FString::Printf(TEXT("%s_%d"), *InPotentialNewName.ToString(), ++Suffix);
+		return Container->IsNameAvailable(InPotentialNewName, ERigElementType::Control, OutErrorMessage);
 	}
-	return Name;
+
+	// if we are not part of a container, fall back to index based validation
+	if (InPotentialNewName.Len() > FRigHierarchyContainer::GetMaxNameLength())
+	{
+		return false;
+	}
+	return GetIndex(*InPotentialNewName) == INDEX_NONE;
+}
+
+FName FRigControlHierarchy::GetSafeNewName(const FString& InPotentialNewName) const
+{
+	return FRigHierarchyContainer::GetSafeNewName(InPotentialNewName, [this](const FString& InName) -> bool {
+		return IsNameAvailable(InName);
+	});
 }
 
 FRigControl& FRigControlHierarchy::Add(
@@ -526,7 +538,7 @@ FRigControl& FRigControlHierarchy::Add(
 )
 {
 	FRigControl NewControl;
-	NewControl.Name = GetSafeNewName(InNewName);
+	NewControl.Name = GetSafeNewName(InNewName.ToString());
 	NewControl.ControlType = InControlType;
 	NewControl.ParentIndex = GetIndex(InParentName);
 	NewControl.ParentName = NewControl.ParentIndex == INDEX_NONE ? NAME_None : InParentName;
@@ -1074,7 +1086,7 @@ FName FRigControlHierarchy::Rename(const FName& InOldName, const FName& InNewNam
 		const int32 Found = GetIndex(InOldName);
 		if (Found != INDEX_NONE)
 		{
-			FName NewName = GetSafeNewName(InNewName);
+			FName NewName = GetSafeNewName(InNewName.ToString());
 
 			bool bWasSelected = IsSelected(InOldName);
 			if(bWasSelected)
diff --git a/Engine/Plugins/Experimental/ControlRig/Source/ControlRig/Private/Rigs/RigCurveContainer.cpp b/Engine/Plugins/Experimental/ControlRig/Source/ControlRig/Private/Rigs/RigCurveContainer.cpp
--- a/Engine/Plugins/Experimental/ControlRig/Source/ControlRig/Private/Rigs/RigCurveContainer.cpp
+++ b/Engine/Plugins/Experimental/ControlRig/Source/ControlRig/Private/Rigs/RigCurveContainer.cpp
@@ -1,6 +1,7 @@
 // Copyright Epic Games, Inc. All Rights Reserved.
 
 #include "Rigs/RigCurveContainer.h"
+#include "Rigs/RigHierarchyContainer.h"
 #include "ControlRig.h"
 #include "HelperUtil.h"
 #include "Animation/Skeleton.h"
@@ -46,15 +47,26 @@ FRigCurveContainer& FRigCurveContainer::operator= (const FRigCurveContainer &InO
 	return *this;
 }
 
-FName FRigCurveContainer::GetSafeNewName(const FName& InPotentialNewName) const
+bool FRigCurveContainer::IsNameAvailable(const FString& InPotentialNewName, FString* OutErrorMessage) const
 {
-	FName Name = InPotentialNewName;
-	int32 Suffix = 1;
-	while(!IsNameAvailable(Name))
+	if (Container)
 	{
-		Name = *FString::Printf(TEXT("%s_%d"), *InPotentialNewName.ToString(), ++Suffix);
+		return Container->IsNameAvailable(InPotentialNewName, ERigElementType::Curve, OutErrorMessage);
 	}
-	return Name;
+
+	// if we are not part of a container, fall back to index based validation
+	if (InPotentialNewName.Len() > FRigHierarchyContainer::GetMaxNameLength())
+	{
+		return false;
+	}
+	return GetIndex(*InPotentialNewName) == INDEX_NONE;
+}
+
+FName FRigCurveContainer::GetSafeNewName(const FString& InPotentialNewName) const
+{
+	return FRigHierarchyContainer::GetSafeNewName(InPotentialNewName, [this](const FString& InName) -> bool {
+		return IsNameAvailable(InName);
+	});
 }
 
 FRigCurve& FRigCurveContainer::Add(const FName& InNewName, float InValue)
@@ -62,7 +74,7 @@ FRigCurve& FRigCurveContainer::Add(const FName& InNewName, float InValue)
 	DECLARE_SCOPE_HIERARCHICAL_COUNTER_FUNC()
 
 	FRigCurve NewCurve;
-	NewCurve.Name = GetSafeNewName(InNewName);
+	NewCurve.Name = GetSafeNewName(InNewName.ToString());
 	NewCurve.Value = InValue;
 	FName NewCurveName = NewCurve.Name;
 	Curves.Add(NewCurve);
@@ -163,7 +175,7 @@ FName FRigCurveContainer::Rename(const FName& InOldName, const FName& InNewName)
 		const int32 Found = GetIndex(InOldName);
 		if (Found != INDEX_NONE)
 		{
-			FName NewName = GetSafeNewName(InNewName);
+			FName NewName = GetSafeNewName(InNewName.ToString());
 
 			bool bWasSelected = IsSelected(InOldName);
 			if(bWasSelected)
diff --git a/Engine/Plugins/Experimental/ControlRig/Source/ControlRig/Private/Rigs/RigHierarchyContainer.cpp b/Engine/Plugins/Experimental/ControlRig/Source/ControlRig/Private/Rigs/RigHierarchyContainer.cpp
--- a/Engine/Plugins/Experimental/ControlRig/Source/ControlRig/Private/Rigs/RigHierarchyContainer.cpp
+++ b/Engine/Plugins/Experimental/ControlRig/Source/ControlRig/Private/Rigs/RigHierarchyContainer.cpp
@@ -1439,3 +1439,96 @@ void FRigHierarchyContainer::AppendToPose(FRigPose& InOutPose) const
 	CurveContainer.AppendToPose(InOutPose);
 }
 
+void FRigHierarchyContainer::SanitizeName(FString& InOutName)
+{
+	// Sanitize the name
+	for (int32 i = 0; i < InOutName.Len(); ++i)
+	{
+		TCHAR& C = InOutName[i];
+
+		const bool bGoodChar =
+			((C >= 'A') && (C <= 'Z')) || ((C >= 'a') && (C <= 'z')) ||		// A-Z (upper and lowercase) anytime
+			(C == '_') ||													// _ anytime
+			((i > 0) && (C >= '0') && (C <= '9'));							// 0-9 after the first character
+
+		if (!bGoodChar)
+		{
+			C = '_';
+		}
+	}
+
+	if (InOutName.Len() > GetMaxNameLength())
+	{
+		InOutName.LeftChopInline(InOutName.Len() - GetMaxNameLength());
+	}
+}
+
+FName FRigHierarchyContainer::GetSanitizedName(const FString& InName)
+{
+	FString Name = InName;
+	SanitizeName(Name);
+
+	if (Name.IsEmpty())
+	{
+		return NAME_None;
+	}
+
+	return *Name;
+}
+
+FName FRigHierarchyContainer::GetSafeNewName(const FString& InName, TFunction<bool(const FString&)> IsNameAvailableFunction)
+{
+	check(IsNameAvailableFunction);
+
+	FString SanitizedName = InName;
+	SanitizeName(SanitizedName);
+	FString Name = SanitizedName;
+
+	int32 Suffix = 1;
+	while (!IsNameAvailableFunction(Name))
+	{
+		FString BaseString = SanitizedName;
+		if (BaseString.Len() > GetMaxNameLength() - 4)
+		{
+			BaseString.LeftChopInline(BaseString.Len() - (GetMaxNameLength() - 4));
+		}
+		Name = *FString::Printf(TEXT("%s_%d"), *BaseString, ++Suffix);
+	}
+	return *Name;
+}
+
+bool FRigHierarchyContainer::IsNameAvailable(const FString& InName, ERigElementType InElementType, FString* OutErrorMessage)
+{
+	FString UnsanitizedName = InName;
+	if (UnsanitizedName.Len() > GetMaxNameLength())
+	{
+		if (OutErrorMessage)
+		{
+			*OutErrorMessage = TEXT("Name too long.");
+		}
+		return false;
+	}
+
+	FString SanitizedName = UnsanitizedName;
+	SanitizeName(SanitizedName);
+
+	if (SanitizedName != UnsanitizedName)
+	{
+		if (OutErrorMessage)
+		{
+			*OutErrorMessage = TEXT("Name contains invalid characters.");
+		}
+		return false;
+	}
+
+	if (GetIndex(FRigElementKey(*InName, InElementType)) != INDEX_NONE)
+	{
+		if (OutErrorMessage)
+		{
+			*OutErrorMessage = TEXT("Name already used.");
+		}
+		return false;
+	}
+
+	return true;
+}
\ No newline at end of file
diff --git a/Engine/Plugins/Experimental/ControlRig/Source/ControlRig/Private/Rigs/RigSpaceHierarchy.cpp b/Engine/Plugins/Experimental/ControlRig/Source/ControlRig/Private/Rigs/RigSpaceHierarchy.cpp
--- a/Engine/Plugins/Experimental/ControlRig/Source/ControlRig/Private/Rigs/RigSpaceHierarchy.cpp
+++ b/Engine/Plugins/Experimental/ControlRig/Source/ControlRig/Private/Rigs/RigSpaceHierarchy.cpp
@@ -1,6 +1,7 @@
 // Copyright Epic Games, Inc. All Rights Reserved.
 
 #include "Rigs/RigSpaceHierarchy.h"
+#include "Rigs/RigHierarchyContainer.h"
 #include "ControlRig.h"
 #include "HelperUtil.h"
 
@@ -38,21 +39,32 @@ FRigSpaceHierarchy& FRigSpaceHierarchy::operator= (const FRigSpaceHierarchy &InO
 	return *this;
 }
 
-FName FRigSpaceHierarchy::GetSafeNewName(const FName& InPotentialNewName) const
+bool FRigSpaceHierarchy::IsNameAvailable(const FString& InPotentialNewName, FString* OutErrorMessage) const
 {
-	FName Name = InPotentialNewName;
-	int32 Suffix = 1;
-	while(!IsNameAvailable(Name))
+	if (Container)
 	{
-		Name = *FString::Printf(TEXT("%s_%d"), *InPotentialNewName.ToString(), ++Suffix);
+		return Container->IsNameAvailable(InPotentialNewName, ERigElementType::Space, OutErrorMessage);
 	}
-	return Name;
+
+	// if we are not part of a container, fall back to index based validation
+	if (InPotentialNewName.Len() > FRigHierarchyContainer::GetMaxNameLength())
+	{
+		return false;
+	}
+	return GetIndex(*InPotentialNewName) == INDEX_NONE;
+}
+
+FName FRigSpaceHierarchy::GetSafeNewName(const FString& InPotentialNewName) const
+{
+	return FRigHierarchyContainer::GetSafeNewName(InPotentialNewName, [this](const FString& InName) -> bool {
+		return IsNameAvailable(InName);
+	});
 }
 
 FRigSpace& FRigSpaceHierarchy::Add(const FName& InNewName, ERigSpaceType InSpaceType, const FName& InParentName, const FTransform& InTransform)
 {
 	FRigSpace NewSpace;
-	NewSpace.Name = GetSafeNewName(InNewName);
+	NewSpace.Name = GetSafeNewName(InNewName.ToString());
 	NewSpace.ParentIndex = GetParentIndex(InSpaceType, InParentName);
 	NewSpace.SpaceType = NewSpace.ParentIndex == INDEX_NONE ? ERigSpaceType::Global : InSpaceType;
 	NewSpace.ParentName = NewSpace.ParentIndex == INDEX_NONE ? NAME_None : InParentName;
@@ -415,7 +427,7 @@ FName FRigSpaceHierarchy::Rename(const FName& InOldName, const FName& InNewName)
 		const int32 Found = GetIndex(InOldName);
 		if (Found != INDEX_NONE)
 		{
-			FName NewName = GetSafeNewName(InNewName);
+			FName NewName = GetSafeNewName(InNewName.ToString());
 
 			bool bWasSelected = IsSelected(InOldName);
 			if(bWasSelected)
diff --git a/Engine/Plugins/Experimental/ControlRig/Source/ControlRig/Public/Rigs/RigBoneHierarchy.h b/Engine/Plugins/Experimental/ControlRig/Source/ControlRig/Public/Rigs/RigBoneHierarchy.h
--- a/Engine/Plugins/Experimental/ControlRig/Source/ControlRig/Public/Rigs/RigBoneHierarchy.h
+++ b/Engine/Plugins/Experimental/ControlRig/Source/ControlRig/Public/Rigs/RigBoneHierarchy.h
@@ -95,9 +95,9 @@ struct CONTROLRIG_API FRigBoneHierarchy
 	FORCEINLINE TArray<FRigBone>::RangedForIteratorType      end()         { return Bones.end();   }
 	FORCEINLINE TArray<FRigBone>::RangedForConstIteratorType end() const   { return Bones.end();   }
 
-	FORCEINLINE bool IsNameAvailable(const FName& InPotentialNewName) const { return GetIndex(InPotentialNewName) == INDEX_NONE; }
+	bool IsNameAvailable(const FString& InPotentialNewName, FString* OutErrorMessage = nullptr) const;
 
-	FName GetSafeNewName(const FName& InPotentialNewName) const;
+	FName GetSafeNewName(const FString& InPotentialNewName) const;
 
 	FRigBone& Add(const FName& InNewName, const FName& InParentName = NAME_None, ERigBoneType InType = ERigBoneType::Imported, const FTransform& InInitTransform = FTransform::Identity);
 
diff --git a/Engine/Plugins/Experimental/ControlRig/Source/ControlRig/Public/Rigs/RigControlHierarchy.h b/Engine/Plugins/Experimental/ControlRig/Source/ControlRig/Public/Rigs/RigControlHierarchy.h
--- a/Engine/Plugins/Experimental/ControlRig/Source/ControlRig/Public/Rigs/RigControlHierarchy.h
+++ b/Engine/Plugins/Experimental/ControlRig/Source/ControlRig/Public/Rigs/RigControlHierarchy.h
@@ -477,9 +477,9 @@ struct CONTROLRIG_API FRigControlHierarchy
 	FORCEINLINE_DEBUGGABLE TArray<FRigControl>::RangedForIteratorType      end()         { return Controls.end();   }
 	FORCEINLINE_DEBUGGABLE TArray<FRigControl>::RangedForConstIteratorType end() const   { return Controls.end();   }
 
-	FORCEINLINE_DEBUGGABLE bool IsNameAvailable(const FName& InPotentialNewName) const { return GetIndex(InPotentialNewName) == INDEX_NONE; }
+	bool IsNameAvailable(const FString& InPotentialNewName, FString* OutErrorMessage = nullptr) const;
 
-	FName GetSafeNewName(const FName& InPotentialNewName) const;
+	FName GetSafeNewName(const FString& InPotentialNewName) const;
 
 	FRigControl& Add(
 		const FName& InNewName,
diff --git a/Engine/Plugins/Experimental/ControlRig/Source/ControlRig/Public/Rigs/RigCurveContainer.h b/Engine/Plugins/Experimental/ControlRig/Source/ControlRig/Public/Rigs/RigCurveContainer.h
--- a/Engine/Plugins/Experimental/ControlRig/Source/ControlRig/Public/Rigs/RigCurveContainer.h
+++ b/Engine/Plugins/Experimental/ControlRig/Source/ControlRig/Public/Rigs/RigCurveContainer.h
@@ -55,9 +55,9 @@ public:
 	FORCEINLINE TArray<FRigCurve>::RangedForIteratorType      end()         { return Curves.end();   }
 	FORCEINLINE TArray<FRigCurve>::RangedForConstIteratorType end() const   { return Curves.end();   }
 
-	FORCEINLINE bool IsNameAvailable(const FName& InPotentialNewName) const { return GetIndex(InPotentialNewName) == INDEX_NONE; }
+	bool IsNameAvailable(const FString& InPotentialNewName, FString* OutErrorMessage = nullptr) const;
 
-	FName GetSafeNewName(const FName& InPotentialNewName) const;
+	FName GetSafeNewName(const FString& InPotentialNewName) const;
 
 	FRigCurve& Add(const FName& InNewName, float InValue = 0.f);
 
diff --git a/Engine/Plugins/Experimental/ControlRig/Source/ControlRig/Public/Rigs/RigHierarchyContainer.h b/Engine/Plugins/Experimental/ControlRig/Source/ControlRig/Public/Rigs/RigHierarchyContainer.h
--- a/Engine/Plugins/Experimental/ControlRig/Source/ControlRig/Public/Rigs/RigHierarchyContainer.h
+++ b/Engine/Plugins/Experimental/ControlRig/Source/ControlRig/Public/Rigs/RigHierarchyContainer.h
@@ -358,6 +358,12 @@ public:
 	void SendEvent(const FRigEventContext& InEvent, bool bAsyncronous = true);
 	FRigEventDelegate OnEventReceived;
 
+	static int32 GetMaxNameLength() { return 100; }
+	static void SanitizeName(FString& InOutName);
+	static FName GetSanitizedName(const FString& InName);
+	static FName GetSafeNewName(const FString& InName, TFunction<bool(const FString&)> IsNameAvailableFunction);
+	bool IsNameAvailable(const FString& InName, ERigElementType InElementType, FString* OutErrorMessage = nullptr);
+
 protected:
 
 	void HandleOnElementSelected(FRigHierarchyContainer* InContainer, const FRigElementKey& InKey, bool bSelected);
diff --git a/Engine/Plugins/Experimental/ControlRig/Source/ControlRig/Public/Rigs/RigSpaceHierarchy.h b/Engine/Plugins/Experimental/ControlRig/Source/ControlRig/Public/Rigs/RigSpaceHierarchy.h
--- a/Engine/Plugins/Experimental/ControlRig/Source/ControlRig/Public/Rigs/RigSpaceHierarchy.h
+++ b/Engine/Plugins/Experimental/ControlRig/Source/ControlRig/Public/Rigs/RigSpaceHierarchy.h
@@ -112,9 +112,9 @@ struct CONTROLRIG_API FRigSpaceHierarchy
 	FORCEINLINE TArray<FRigSpace>::RangedForIteratorType      end()         { return Spaces.end();   }
 	FORCEINLINE TArray<FRigSpace>::RangedForConstIteratorType end() const   { return Spaces.end();   }
 
-	FORCEINLINE bool IsNameAvailable(const FName& InPotentialNewName) const { return GetIndex(InPotentialNewName) == INDEX_NONE; }
+	bool IsNameAvailable(const FString& InPotentialNewName, FString* OutErrorMessage = nullptr) const;
 
-	FName GetSafeNewName(const FName& InPotentialNewName) const;
+	FName GetSafeNewName(const FString& InPotentialNewName) const;
 
 	FRigSpace& Add(const FName& InNewName, ERigSpaceType InSpaceType = ERigSpaceType::Global, const FName& InParentName = NAME_None, const FTransform& InTransform = FTransform::Identity);
 
diff --git a/Engine/Plugins/Experimental/ControlRig/Source/ControlRigDeveloper/Private/ControlRigHierarchyModifier.cpp b/Engine/Plugins/Experimental/ControlRig/Source/ControlRigDeveloper/Private/ControlRigHierarchyModifier.cpp
--- a/Engine/Plugins/Experimental/ControlRig/Source/ControlRigDeveloper/Private/ControlRigHierarchyModifier.cpp
+++ b/Engine/Plugins/Experimental/ControlRig/Source/ControlRigDeveloper/Private/ControlRigHierarchyModifier.cpp
@@ -20,11 +20,13 @@ TArray<FRigElementKey> UControlRigHierarchyModifier::GetElements() const
 	return Elements;
 }
 
-FRigElementKey UControlRigHierarchyModifier::AddBone(const FName& InNewName, const FName& InParentName, ERigBoneType InType)
+FRigElementKey UControlRigHierarchyModifier::AddBone(const FString& InNewName, const FString& InParentName, ERigBoneType InType)
 {
 	if(Container != nullptr)
 	{
-		return Container->BoneHierarchy.Add(InNewName, InParentName, InType).GetElementKey();
+		FName NewName = Container->BoneHierarchy.GetSafeNewName(InNewName);
+		FName ParentName = FRigHierarchyContainer::GetSanitizedName(InParentName);
+		return Container->BoneHierarchy.Add(NewName, ParentName, InType).GetElementKey();
 	}
 	return FRigElementKey();
 }
@@ -61,17 +63,21 @@ void UControlRigHierarchyModifier::SetBone(const FRigBone& InElement)
 }
 
 FRigElementKey UControlRigHierarchyModifier::AddControl(
-	const FName& InNewName,
+	const FString& InNewName,
 	ERigControlType InControlType,
-	const FName& InParentName,
-	const FName& InSpaceName,
-	const FName& InGizmoName,
+	const FString& InParentName,
+	const FString& InSpaceName,
+	const FString& InGizmoName,
 	const FLinearColor& InGizmoColor
 )
 {
 	if(Container != nullptr)
 	{
-		return Container->ControlHierarchy.Add(InNewName, InControlType, InParentName, InSpaceName, FTransform::Identity, FRigControlValue(), InGizmoName, FTransform::Identity, InGizmoColor).GetElementKey();
+		FName NewName = Container->ControlHierarchy.GetSafeNewName(InNewName);
+		FName ParentName = FRigHierarchyContainer::GetSanitizedName(InParentName);
+		FName SpaceName = FRigHierarchyContainer::GetSanitizedName(InSpaceName);
+		FName GizmoName = FRigHierarchyContainer::GetSanitizedName(InGizmoName);
+		return Container->ControlHierarchy.Add(NewName, InControlType, ParentName, SpaceName, FTransform::Identity, FRigControlValue(), GizmoName, FTransform::Identity, InGizmoColor).GetElementKey();
 	}
 	return FRigElementKey();
 }
@@ -310,14 +316,16 @@ void UControlRigHierarchyModifier::SetControlOffsetTransform(const FRigElementKe
 
 FRigElementKey UControlRigHierarchyModifier::AddSpace
 (
-	const FName& InNewName,
+	const FString& InNewName,
 	ERigSpaceType InSpaceType,
-	const FName& InParentName
+	const FString& InParentName
 )
 {
 	if(Container != nullptr)
 	{
-		return Container->SpaceHierarchy.Add(InNewName, InSpaceType, InParentName).GetElementKey();
+		FName NewName = Container->SpaceHierarchy.GetSafeNewName(InNewName);
+		FName ParentName = FRigHierarchyContainer::GetSanitizedName(InParentName);
+		return Container->SpaceHierarchy.Add(NewName, InSpaceType, ParentName).GetElementKey();
 	}
 	return FRigElementKey();
 }
@@ -353,11 +361,12 @@ void UControlRigHierarchyModifier::SetSpace(const FRigSpace& InElement)
 	}
 }
 
-FRigElementKey UControlRigHierarchyModifier::AddCurve(const FName& InNewName, float InValue)
+FRigElementKey UControlRigHierarchyModifier::AddCurve(const FString& InNewName, float InValue)
 {
 	if(Container != nullptr)
 	{
-		return Container->CurveContainer.Add(InNewName, InValue).GetElementKey();
+		FName NewName = Container->CurveContainer.GetSafeNewName(InNewName);
+		return Container->CurveContainer.Add(NewName, InValue).GetElementKey();
 	}
 	return FRigElementKey();
 }
@@ -436,7 +445,7 @@ bool UControlRigHierarchyModifier::RemoveElement(const FRigElementKey& InElement
 	return false;
 }
 
-FRigElementKey UControlRigHierarchyModifier::RenameElement(const FRigElementKey& InElement, const FName& InNewName)
+FRigElementKey UControlRigHierarchyModifier::RenameElement(const FRigElementKey& InElement, const FString& InNewName)
 {
 	if(Container == nullptr)
 	{
@@ -448,26 +457,28 @@ FRigElementKey UControlRigHierarchyModifier::RenameElement(const FRigElementKey&
 		return FRigElementKey();
 	}
 
+	FName SanitizedName = FRigHierarchyContainer::GetSanitizedName(InNewName);
+
 	switch (InElement.Type)
 	{
 		case ERigElementType::Bone:
 		{
-			FName NewName = Container->BoneHierarchy.Rename(InElement.Name, InNewName);
+			FName NewName = Container->BoneHierarchy.Rename(InElement.Name, SanitizedName);
 			return FRigElementKey(NewName, ERigElementType::Bone);
 		}
 		case ERigElementType::Control:
 		{
-			FName NewName = Container->ControlHierarchy.Rename(InElement.Name, InNewName);
+			FName NewName = Container->ControlHierarchy.Rename(InElement.Name, SanitizedName);
 			return FRigElementKey(NewName, ERigElementType::Control);
 		}
 		case ERigElementType::Space:
 		{
-			FName NewName = Container->SpaceHierarchy.Rename(InElement.Name, InNewName);
+			FName NewName = Container->SpaceHierarchy.Rename(InElement.Name, SanitizedName);
 			return FRigElementKey(NewName, ERigElementType::Space);
 		}
 		case ERigElementType::Curve:
 		{
-			FName NewName = Container->CurveContainer.Rename(InElement.Name, InNewName);
+			FName NewName = Container->CurveContainer.Rename(InElement.Name, SanitizedName);
 			return FRigElementKey(NewName, ERigElementType::Curve);
 		}
 		default:
diff --git a/Engine/Plugins/Experimental/ControlRig/Source/ControlRigDeveloper/Public/ControlRigHierarchyModifier.h b/Engine/Plugins/Experimental/ControlRig/Source/ControlRigDeveloper/Public/ControlRigHierarchyModifier.h
--- a/Engine/Plugins/Experimental/ControlRig/Source/ControlRigDeveloper/Public/ControlRigHierarchyModifier.h
+++ b/Engine/Plugins/Experimental/ControlRig/Source/ControlRigDeveloper/Public/ControlRigHierarchyModifier.h
@@ -26,8 +26,8 @@ public:
 	// Adds a new single bone
 	UFUNCTION(BlueprintCallable, Category = "Hierarchy")
 	FRigElementKey AddBone(
-		const FName& InNewName,
-		const FName& InParentName = NAME_None,
+		const FString& InNewName,
+		const FString& InParentName = TEXT(""),
 		ERigBoneType InType = ERigBoneType::User
 	);
 
@@ -42,11 +42,11 @@ public:
 	// Adds a new single control
 	UFUNCTION(BlueprintCallable, Category = "Hierarchy")
 	FRigElementKey AddControl(
-		const FName& InNewName,
+		const FString& InNewName,
 		ERigControlType InControlType = ERigControlType::Transform,
-		const FName& InParentName = NAME_None,
-		const FName& InSpaceName = NAME_None,
-		const FName& InGizmoName = TEXT("Gizmo"),
+		const FString& InParentName = TEXT(""),
+		const FString& InSpaceName = TEXT(""),
+		const FString& InGizmoName = TEXT("Gizmo"),
 		const FLinearColor& InGizmoColor = FLinearColor(1.0, 0.0, 0.0, 1.0)
 	);
 
@@ -125,9 +125,9 @@ public:
 	// Adds a new single space
 	UFUNCTION(BlueprintCallable, Category = "Hierarchy")
 	FRigElementKey AddSpace(
-		const FName& InNewName,
+		const FString& InNewName,
 		ERigSpaceType InSpaceType = ERigSpaceType::Global,
-		const FName& InParentName = NAME_None
+		const FString& InParentName = TEXT("")
 	);
 
 	// Returns a single space from provided key
@@ -140,7 +140,7 @@ public:
 
 	// Adds a new single curve
 	UFUNCTION(BlueprintCallable, Category = "Hierarchy")
-	FRigElementKey AddCurve(const FName& InNewName, float InValue = 0.f);
+	FRigElementKey AddCurve(const FString& InNewName, float InValue = 0.f);
 
 	// Returns a single curve from provided key
 	UFUNCTION(BlueprintCallable, Category = "Hierarchy")
@@ -156,7 +156,7 @@ public:
 
 	// Renames an existing element and returns the new element key
 	UFUNCTION(BlueprintCallable, Category = "Hierarchy")
-	FRigElementKey RenameElement(const FRigElementKey& InElement, const FName& InNewName);
+	FRigElementKey RenameElement(const FRigElementKey& InElement, const FString& InNewName);
 
 	// Reparents an element to another element, returns true if successful
 	UFUNCTION(BlueprintCallable, Category = "Hierarchy")
diff --git a/Engine/Plugins/Experimental/ControlRig/Source/ControlRigEditor/Private/Editor/SRigHierarchy.cpp b/Engine/Plugins/Experimental/ControlRig/Source/ControlRigEditor/Private/Editor/SRigHierarchy.cpp
--- a/Engine/Plugins/Experimental/ControlRig/Source/ControlRigEditor/Private/Editor/SRigHierarchy.cpp
+++ b/Engine/Plugins/Experimental/ControlRig/Source/ControlRigEditor/Private/Editor/SRigHierarchy.cpp
@@ -215,7 +215,7 @@ FText SRigHierarchyItem::GetName() const
 
 bool SRigHierarchyItem::OnVerifyNameChanged(const FText& InText, FText& OutErrorMessage)
 {
-	const FName NewName = FName(*InText.ToString());
+	FString NewName = InText.ToString();
 	if (OnVerifyElementNameChanged.IsBound())
 	{
 		return OnVerifyElementNameChanged.Execute(WeakRigTreeElement.Pin()->Key, NewName, OutErrorMessage);
@@ -231,15 +231,22 @@ void SRigHierarchyItem::OnNameCommitted(const FText& InText, ETextCommit::Type I
 	// because it is important to keep the unique names per pose
 	if (InCommitType == ETextCommit::OnEnter)
 	{
-		FName NewName = FName(*InText.ToString());
+		FString NewName = InText.ToString();
 		FRigElementKey OldKey = WeakRigTreeElement.Pin()->Key;
 
-		if (!OnRenameElement.IsBound() || OnRenameElement.Execute(OldKey, NewName))
+		if (OnRenameElement.IsBound())
 		{
-			if (WeakRigTreeElement.IsValid())
+			FName NewSanitizedName = OnRenameElement.Execute(OldKey, NewName);
+			if (NewSanitizedName.IsNone())
 			{
-				WeakRigTreeElement.Pin()->Key.Name = NewName;
+				return;
 			}
+			NewName = NewSanitizedName.ToString();
+		}
+
+		if (WeakRigTreeElement.IsValid())
+		{
+			WeakRigTreeElement.Pin()->Key.Name = *NewName;
 		}
 	}
 }
@@ -2355,14 +2362,18 @@ FReply SRigHierarchy::OnAcceptDrop(const FDragDropEvent& DragDropEvent, EItemDro
 	return FReply::Unhandled();
 }
 
-bool SRigHierarchy::RenameElement(const FRigElementKey& OldKey, const FName& NewName)
+FName SRigHierarchy::RenameElement(const FRigElementKey& OldKey, const FString& NewName)
 {
 	ClearDetailPanel();
 
-	if (OldKey.Name == NewName)
+	if (OldKey.Name.ToString() == NewName)
 	{
-		return true;
+		return NAME_None;
 	}
+	FString SanitizedNameStr = NewName;
+	FRigHierarchyContainer::SanitizeName(SanitizedNameStr);
+	FName SanitizedName = *SanitizedNameStr;
+	FName ResultingName = NAME_None;
 
 	// make sure there is no duplicate
 	if (ControlRigBlueprint.IsValid())
@@ -2375,35 +2386,35 @@ bool SRigHierarchy::RenameElement(const FRigElementKey& OldKey, const FName& New
 		{
 			case ERigElementType::Bone:
 			{
-				Container->BoneHierarchy.Rename(OldKey.Name, NewName);
+				ResultingName = Container->BoneHierarchy.Rename(OldKey.Name, SanitizedName);
 				break;
 			}
 			case ERigElementType::Control:
 			{
-				Container->ControlHierarchy.Rename(OldKey.Name, NewName);
+				ResultingName = Container->ControlHierarchy.Rename(OldKey.Name, SanitizedName);
 				break;
 			}
 			case ERigElementType::Space:
 			{
-				Container->SpaceHierarchy.Rename(OldKey.Name, NewName);
+				ResultingName = Container->SpaceHierarchy.Rename(OldKey.Name, SanitizedName);
 				break;
 			}
 			default:
 			{
-				return false;
+				return NAME_None;
 			}
 		}
 
 		ControlRigBlueprint->PropagateHierarchyFromBPToInstances(true, true);
-		return true;
+		return ResultingName;
 	}
 
-	return false;
+	return ResultingName;
 }
 
-bool SRigHierarchy::OnVerifyNameChanged(const FRigElementKey& OldKey, const FName& NewName, FText& OutErrorMessage)
+bool SRigHierarchy::OnVerifyNameChanged(const FRigElementKey& OldKey, const FString& NewName, FText& OutErrorMessage)
 {
-	if (OldKey.Name == NewName)
+	if (OldKey.Name.ToString() == NewName)
 	{
 		return true;
 	}
@@ -2412,27 +2423,15 @@ bool SRigHierarchy::OnVerifyNameChanged(const FRigElementKey& OldKey, const FNam
 	if (ControlRigBlueprint.IsValid())
 	{
 		FRigHierarchyContainer* Container = GetHierarchyContainer();
-		switch (OldKey.Type)
-		{
-		case ERigElementType::Bone:
-		{
-			return Container->BoneHierarchy.IsNameAvailable(NewName);
-		}
-		case ERigElementType::Control:
-		{
-			return Container->ControlHierarchy.IsNameAvailable(NewName);
-		}
-		case ERigElementType::Space:
-		{
-			return Container->SpaceHierarchy.IsNameAvailable(NewName);
-		}
-		default:
+
+		FString OutErrorString;
+		if (!Container->IsNameAvailable(NewName, OldKey.Type, &OutErrorString))
 		{
-			break;
-		}
+			OutErrorMessage = FText::FromString(OutErrorString);
+			return false;
 		}
 	}
-	return false;
+	return true;
 }
 
 void SRigHierarchy::HandleResetTransform(bool bSelectionOnly)
diff --git a/Engine/Plugins/Experimental/ControlRig/Source/ControlRigEditor/Private/Editor/SRigHierarchy.h b/Engine/Plugins/Experimental/ControlRig/Source/ControlRigEditor/Private/Editor/SRigHierarchy.h
--- a/Engine/Plugins/Experimental/ControlRig/Source/ControlRigEditor/Private/Editor/SRigHierarchy.h
+++ b/Engine/Plugins/Experimental/ControlRig/Source/ControlRigEditor/Private/Editor/SRigHierarchy.h
@@ -19,8 +19,8 @@ struct FAssetData;
 class FMenuBuilder;
 class SRigHierarchyItem;
 
-DECLARE_DELEGATE_RetVal_TwoParams(bool, FOnRenameElement, const FRigElementKey& /*OldKey*/, const FName& /*NewName*/);
-DECLARE_DELEGATE_RetVal_ThreeParams(bool, FOnVerifyElementNameChanged, const FRigElementKey& /*OldKey*/, const FName& /*NewName*/, FText& /*OutErrorMessage*/);
+DECLARE_DELEGATE_RetVal_TwoParams(FName, FOnRenameElement, const FRigElementKey& /*OldKey*/, const FString& /*NewName*/);
+DECLARE_DELEGATE_RetVal_ThreeParams(bool, FOnVerifyElementNameChanged, const FRigElementKey& /*OldKey*/, const FString& /*NewName*/, FText& /*OutErrorMessage*/);
 
 /** An item in the tree */
 class FRigTreeElement : public TSharedFromThis<FRigTreeElement>
@@ -327,8 +327,8 @@ private:
 	void ReparentElement(FRigElementKey InKey, FRigElementKey InParentKey);
 
 public:
-	bool RenameElement(const FRigElementKey& OldKey, const FName& NewName);
-	bool OnVerifyNameChanged(const FRigElementKey& OldKey, const FName& NewName, FText& OutErrorMessage);
+	FName RenameElement(const FRigElementKey& OldKey, const FString& NewName);
+	bool OnVerifyNameChanged(const FRigElementKey& OldKey, const FString& NewName, FText& OutErrorMessage);
 
 	friend class SRigHierarchyItem;
 };
\ No newline at end of file
