Control Rig: Added re-initialization for draw container at the beginning of each execute() such at curves imported from FBX can respond to rig graph change correctly

diff --git a/Engine/Plugins/Experimental/ControlRig/Source/ControlRig/Private/ControlRig.cpp b/Engine/Plugins/Experimental/ControlRig/Source/ControlRig/Private/ControlRig.cpp
--- a/Engine/Plugins/Experimental/ControlRig/Source/ControlRig/Private/ControlRig.cpp
+++ b/Engine/Plugins/Experimental/ControlRig/Source/ControlRig/Private/ControlRig.cpp
@@ -460,6 +460,15 @@ void UControlRig::Execute(const EControlRigState InState, const FName& InEventNa
 	FRigUnitContext Context;
 	DrawInterface.Reset();
 	Context.DrawInterface = &DrawInterface;
+
+	// draw container contains persistent draw instructions, 
+	// so we cannot call Reset(), which will clear them,
+	// instead, we re-initialize them from the CDO
+	if (UControlRig* CDO = GetClass()->GetDefaultObject<UControlRig>())
+	{
+		DrawContainer = CDO->DrawContainer;
+	}
+
 	Context.DrawContainer = &DrawContainer;
 	Context.DataSourceRegistry = GetDataSourceRegistry();
 
