Control Rig: Added unit test for the draw container to make sure they are re-initialized at the beginning of each execution

diff --git a/Engine/Plugins/Experimental/ControlRig/Source/ControlRig/Public/ControlRig.h b/Engine/Plugins/Experimental/ControlRig/Source/ControlRig/Public/ControlRig.h
--- a/Engine/Plugins/Experimental/ControlRig/Source/ControlRig/Public/ControlRig.h
+++ b/Engine/Plugins/Experimental/ControlRig/Source/ControlRig/Public/ControlRig.h
@@ -404,6 +404,9 @@ public:
 	const FControlRigDrawInterface& GetDrawInterface() const { return DrawInterface; };
 	FControlRigDrawInterface& GetDrawInterface() { return DrawInterface; };
 
+	const FControlRigDrawContainer& GetDrawContainer() const { return DrawContainer; };
+	FControlRigDrawContainer& GetDrawContainer() { return DrawContainer; };
+
 private:
 
 	UPROPERTY()
