Control Rig: Added python error reporting for ControlRigHierarchyModifer APIs

diff --git a/Engine/Plugins/Experimental/ControlRig/Source/ControlRigDeveloper/Private/ControlRigHierarchyModifier.cpp b/Engine/Plugins/Experimental/ControlRig/Source/ControlRigDeveloper/Private/ControlRigHierarchyModifier.cpp
--- a/Engine/Plugins/Experimental/ControlRig/Source/ControlRigDeveloper/Private/ControlRigHierarchyModifier.cpp
+++ b/Engine/Plugins/Experimental/ControlRig/Source/ControlRigDeveloper/Private/ControlRigHierarchyModifier.cpp
@@ -2,8 +2,10 @@
 
 #include "ControlRigHierarchyModifier.h"
 #include "ControlRigBlueprint.h"
+#include "Misc/CoreMisc.h"
 
 UControlRigHierarchyModifier::UControlRigHierarchyModifier()
+	:bReportWarningsAndErrors(true)
 {
 	Container = nullptr;
 }
@@ -22,7 +24,7 @@ TArray<FRigElementKey> UControlRigHierarchyModifier::GetElements() const
 
 FRigElementKey UControlRigHierarchyModifier::AddBone(const FString& InNewName, const FString& InParentName, ERigBoneType InType)
 {
-	if(Container != nullptr)
+	if (Container != nullptr)
 	{
 		FName NewName = Container->BoneHierarchy.GetSafeNewName(InNewName);
 		FName ParentName = FRigHierarchyContainer::GetSanitizedName(InParentName);
@@ -33,14 +35,19 @@ FRigElementKey UControlRigHierarchyModifier::AddBone(const FString& InNewName, c
 
 FRigBone UControlRigHierarchyModifier::GetBone(const FRigElementKey& InKey)
 {
-	if (Container != nullptr && InKey.Type == ERigElementType::Bone)
+	if (Container != nullptr)
 	{
-		int32 Index = Container->GetIndex(InKey);
-		if (Index != INDEX_NONE)
+		if (InKey.Type == ERigElementType::Bone)
 		{
-			return Container->BoneHierarchy.Bones[Index];
+			int32 Index = Container->GetIndex(InKey);
+			if (Index != INDEX_NONE)
+			{
+				return Container->BoneHierarchy.Bones[Index];
+			}
 		}
+		ReportError(TEXT("Invalid Rig Element Key"));
 	}
+
 	return InvalidBone;
 }
 
@@ -57,8 +64,10 @@ void UControlRigHierarchyModifier::SetBone(const FRigBone& InElement)
 			if (UControlRigBlueprint* Blueprint = Cast<UControlRigBlueprint>(GetOuter()))
 			{
 				Blueprint->PropagateHierarchyFromBPToInstances(true, true);
+				return;
 			}
 		}
+		ReportError(TEXT("Bone not found"));
 	}
 }
 
@@ -71,7 +80,7 @@ FRigElementKey UControlRigHierarchyModifier::AddControl(
 	const FLinearColor& InGizmoColor
 )
 {
-	if(Container != nullptr)
+	if (Container != nullptr)
 	{
 		FName NewName = Container->ControlHierarchy.GetSafeNewName(InNewName);
 		FName ParentName = FRigHierarchyContainer::GetSanitizedName(InParentName);
@@ -85,13 +94,17 @@ FRigElementKey UControlRigHierarchyModifier::AddControl(
 
 FRigControl UControlRigHierarchyModifier::GetControl(const FRigElementKey& InKey)
 {
-	if (Container != nullptr && InKey.Type == ERigElementType::Control)
+	if (Container != nullptr)
 	{
-		int32 Index = Container->GetIndex(InKey);
-		if (Index != INDEX_NONE)
+		if (InKey.Type == ERigElementType::Control)
 		{
-			return Container->ControlHierarchy.Controls[Index];
+			int32 Index = Container->GetIndex(InKey);
+			if (Index != INDEX_NONE)
+			{
+				return Container->ControlHierarchy.Controls[Index];
+			}
 		}
+		ReportError(TEXT("Invalid Rig Element Key"));
 	}
 	return InvalidControl;
 }
@@ -109,208 +122,282 @@ void UControlRigHierarchyModifier::SetControl(const FRigControl& InElement)
 			if (UControlRigBlueprint* Blueprint = Cast<UControlRigBlueprint>(GetOuter()))
 			{
 				Blueprint->PropagateHierarchyFromBPToInstances(true, true);
+				return;
 			}
 		}
+		ReportError(TEXT("Control not found"));
 	}
 }
 
 bool UControlRigHierarchyModifier::GetControlValueBool(const FRigElementKey& InKey, ERigControlValueType InValueType)
 {
-	if (Container != nullptr && InKey.Type == ERigElementType::Control)
+	if (Container != nullptr)
 	{
-		int32 Index = Container->GetIndex(InKey);
-		if (Index != INDEX_NONE)
+		if (InKey.Type == ERigElementType::Control)
 		{
-			return Container->ControlHierarchy.Controls[Index].GetValue(InValueType).Get<bool>();
+			int32 Index = Container->GetIndex(InKey);
+			if (Index != INDEX_NONE)
+			{
+				return Container->ControlHierarchy.Controls[Index].GetValue(InValueType).Get<bool>();
+			}
 		}
+		ReportError(TEXT("Invalid Rig Element Key"));
 	}
 	return false;
 }
 
 int32 UControlRigHierarchyModifier::GetControlValueInt(const FRigElementKey& InKey, ERigControlValueType InValueType)
 {
-	if (Container != nullptr && InKey.Type == ERigElementType::Control)
+	if (Container != nullptr)
 	{
-		int32 Index = Container->GetIndex(InKey);
-		if (Index != INDEX_NONE)
+		if (InKey.Type == ERigElementType::Control)
 		{
-			return Container->ControlHierarchy.Controls[Index].GetValue(InValueType).Get<int32>();
+			int32 Index = Container->GetIndex(InKey);
+			if (Index != INDEX_NONE)
+			{
+				return Container->ControlHierarchy.Controls[Index].GetValue(InValueType).Get<int32>();
+			}
 		}
+		ReportError(TEXT("Invalid Rig Element Key"));
 	}
 	return 0;
 }
 
 float UControlRigHierarchyModifier::GetControlValueFloat(const FRigElementKey& InKey, ERigControlValueType InValueType)
 {
-	if (Container != nullptr && InKey.Type == ERigElementType::Control)
+	if (Container != nullptr)
 	{
-		int32 Index = Container->GetIndex(InKey);
-		if (Index != INDEX_NONE)
+		if (InKey.Type == ERigElementType::Control)
 		{
-			return Container->ControlHierarchy.Controls[Index].GetValue(InValueType).Get<float>();
+			int32 Index = Container->GetIndex(InKey);
+			if (Index != INDEX_NONE)
+			{
+				return Container->ControlHierarchy.Controls[Index].GetValue(InValueType).Get<float>();
+			}
 		}
+		ReportError(TEXT("Invalid Rig Element Key"));
 	}
 	return 0.f;
 }
 
 FVector2D UControlRigHierarchyModifier::GetControlValueVector2D(const FRigElementKey& InKey, ERigControlValueType InValueType)
 {
-	if (Container != nullptr && InKey.Type == ERigElementType::Control)
+	if (Container != nullptr)
 	{
-		int32 Index = Container->GetIndex(InKey);
-		if (Index != INDEX_NONE)
+		if (InKey.Type == ERigElementType::Control)
 		{
-			return Container->ControlHierarchy.Controls[Index].GetValue(InValueType).Get<FVector2D>();
+			int32 Index = Container->GetIndex(InKey);
+			if (Index != INDEX_NONE)
+			{
+				return Container->ControlHierarchy.Controls[Index].GetValue(InValueType).Get<FVector2D>();
+			}
 		}
+		ReportError(TEXT("Invalid Rig Element Key"));
 	}
 	return FVector2D::ZeroVector;
 }
 
 FVector UControlRigHierarchyModifier::GetControlValueVector(const FRigElementKey& InKey, ERigControlValueType InValueType)
 {
-	if (Container != nullptr && InKey.Type == ERigElementType::Control)
+	if (Container != nullptr)
 	{
-		int32 Index = Container->GetIndex(InKey);
-		if (Index != INDEX_NONE)
+		if (InKey.Type == ERigElementType::Control)
 		{
-			return Container->ControlHierarchy.Controls[Index].GetValue(InValueType).Get<FVector>();
+			int32 Index = Container->GetIndex(InKey);
+			if (Index != INDEX_NONE)
+			{
+				return Container->ControlHierarchy.Controls[Index].GetValue(InValueType).Get<FVector>();
+			}
 		}
+		ReportError(TEXT("Invalid Rig Element Key"));
 	}
 	return FVector::ZeroVector;
 }
 
 FRotator UControlRigHierarchyModifier::GetControlValueRotator(const FRigElementKey& InKey, ERigControlValueType InValueType)
 {
-	if (Container != nullptr && InKey.Type == ERigElementType::Control)
+	if (Container != nullptr)
 	{
-		int32 Index = Container->GetIndex(InKey);
-		if (Index != INDEX_NONE)
+		if (InKey.Type == ERigElementType::Control)
 		{
-			return Container->ControlHierarchy.Controls[Index].GetValue(InValueType).Get<FRotator>();
+			int32 Index = Container->GetIndex(InKey);
+			if (Index != INDEX_NONE)
+			{
+				return Container->ControlHierarchy.Controls[Index].GetValue(InValueType).Get<FRotator>();
+			}
 		}
+		ReportError(TEXT("Invalid Rig Element Key"));
 	}
 	return FRotator::ZeroRotator;
 }
 
 FTransform UControlRigHierarchyModifier::GetControlValueTransform(const FRigElementKey& InKey, ERigControlValueType InValueType)
 {
-	if (Container != nullptr && InKey.Type == ERigElementType::Control)
+	if (Container != nullptr)
 	{
-		int32 Index = Container->GetIndex(InKey);
-		if (Index != INDEX_NONE)
+		if (InKey.Type == ERigElementType::Control)
 		{
-			return Container->ControlHierarchy.Controls[Index].GetTransformFromValue(InValueType);
+			int32 Index = Container->GetIndex(InKey);
+			if (Index != INDEX_NONE)
+			{
+				return Container->ControlHierarchy.Controls[Index].GetTransformFromValue(InValueType);
+			}
 		}
+		ReportError(TEXT("Invalid Rig Element Key"));
 	}
 	return FTransform::Identity;
 }
 
 FTransform UControlRigHierarchyModifier::GetControlOffsetTransform(const FRigElementKey& InKey)
 {
-	if (Container != nullptr && InKey.Type == ERigElementType::Control)
+	if (Container != nullptr)
 	{
-		int32 Index = Container->GetIndex(InKey);
-		if (Index != INDEX_NONE)
+		if (InKey.Type == ERigElementType::Control)
 		{
-			return Container->ControlHierarchy.Controls[Index].OffsetTransform;
+			int32 Index = Container->GetIndex(InKey);
+			if (Index != INDEX_NONE)
+			{
+				return Container->ControlHierarchy.Controls[Index].OffsetTransform;
+			}
 		}
+		ReportError(TEXT("Invalid Rig Element Key"));
 	}
 	return FTransform::Identity;
 }
 
 void UControlRigHierarchyModifier::SetControlValueBool(const FRigElementKey& InKey, bool InValue, ERigControlValueType InValueType)
 {
-	if (Container != nullptr && InKey.Type == ERigElementType::Control)
+	if (Container != nullptr)
 	{
-		int32 Index = Container->GetIndex(InKey);
-		if (Index != INDEX_NONE)
+		if (InKey.Type == ERigElementType::Control)
 		{
-			Container->ControlHierarchy.Controls[Index].GetValue(InValueType).Set<bool>(InValue);
+			int32 Index = Container->GetIndex(InKey);
+			if (Index != INDEX_NONE)
+			{
+				Container->ControlHierarchy.Controls[Index].GetValue(InValueType).Set<bool>(InValue);
+				return;
+			}
 		}
+		ReportError(TEXT("Invalid Rig Element Key"));
 	}
 }
 
 void UControlRigHierarchyModifier::SetControlValueInt(const FRigElementKey& InKey, int32 InValue, ERigControlValueType InValueType)
 {
-	if (Container != nullptr && InKey.Type == ERigElementType::Control)
+	if (Container != nullptr)
 	{
-		int32 Index = Container->GetIndex(InKey);
-		if (Index != INDEX_NONE)
+		if (InKey.Type == ERigElementType::Control)
 		{
-			Container->ControlHierarchy.Controls[Index].GetValue(InValueType).Set<int32>(InValue);
+			int32 Index = Container->GetIndex(InKey);
+			if (Index != INDEX_NONE)
+			{
+				Container->ControlHierarchy.Controls[Index].GetValue(InValueType).Set<int32>(InValue);
+				return;
+			}
 		}
+		ReportError(TEXT("Invalid Rig Element Key"));
 	}
 }
 
 void UControlRigHierarchyModifier::SetControlValueFloat(const FRigElementKey& InKey, float InValue, ERigControlValueType InValueType)
 {
-	if (Container != nullptr && InKey.Type == ERigElementType::Control)
+	if (Container != nullptr)
 	{
-		int32 Index = Container->GetIndex(InKey);
-		if (Index != INDEX_NONE)
+		if (InKey.Type == ERigElementType::Control)
 		{
-			Container->ControlHierarchy.Controls[Index].GetValue(InValueType).Set<float>(InValue);
+			int32 Index = Container->GetIndex(InKey);
+			if (Index != INDEX_NONE)
+			{
+				Container->ControlHierarchy.Controls[Index].GetValue(InValueType).Set<float>(InValue);
+				return;
+			}
 		}
+		ReportError(TEXT("Invalid Rig Element Key"));
 	}
 }
 
 void UControlRigHierarchyModifier::SetControlValueVector2D(const FRigElementKey& InKey, FVector2D InValue, ERigControlValueType InValueType)
 {
-	if (Container != nullptr && InKey.Type == ERigElementType::Control)
+	if (Container != nullptr)
 	{
-		int32 Index = Container->GetIndex(InKey);
-		if (Index != INDEX_NONE)
+		if (InKey.Type == ERigElementType::Control)
 		{
-			Container->ControlHierarchy.Controls[Index].GetValue(InValueType).Set<FVector2D>(InValue);
+			int32 Index = Container->GetIndex(InKey);
+			if (Index != INDEX_NONE)
+			{
+				Container->ControlHierarchy.Controls[Index].GetValue(InValueType).Set<FVector2D>(InValue);
+				return;
+			}
 		}
+		ReportError(TEXT("Invalid Rig Element Key"));
 	}
 }
 
 void UControlRigHierarchyModifier::SetControlValueVector(const FRigElementKey& InKey, FVector InValue, ERigControlValueType InValueType)
 {
-	if (Container != nullptr && InKey.Type == ERigElementType::Control)
+	if (Container != nullptr)
 	{
-		int32 Index = Container->GetIndex(InKey);
-		if (Index != INDEX_NONE)
+		if (InKey.Type == ERigElementType::Control)
 		{
-			Container->ControlHierarchy.Controls[Index].GetValue(InValueType).Set<FVector>(InValue);
+			int32 Index = Container->GetIndex(InKey);
+			if (Index != INDEX_NONE)
+			{
+				Container->ControlHierarchy.Controls[Index].GetValue(InValueType).Set<FVector>(InValue);
+				return;
+			}
 		}
+		ReportError(TEXT("Invalid Rig Element Key"));
 	}
 }
 
 void UControlRigHierarchyModifier::SetControlValueRotator(const FRigElementKey& InKey, FRotator InValue, ERigControlValueType InValueType)
 {
-	if (Container != nullptr && InKey.Type == ERigElementType::Control)
+	if (Container != nullptr)
 	{
-		int32 Index = Container->GetIndex(InKey);
-		if (Index != INDEX_NONE)
+		if (InKey.Type == ERigElementType::Control)
 		{
-			Container->ControlHierarchy.Controls[Index].GetValue(InValueType).Set<FRotator>(InValue);
+			int32 Index = Container->GetIndex(InKey);
+			if (Index != INDEX_NONE)
+			{
+				Container->ControlHierarchy.Controls[Index].GetValue(InValueType).Set<FRotator>(InValue);
+				return;
+			}
 		}
+		ReportError(TEXT("Invalid Rig Element Key"));
 	}
 }
 
 void UControlRigHierarchyModifier::SetControlValueTransform(const FRigElementKey& InKey, FTransform InValue, ERigControlValueType InValueType)
 {
-	if (Container != nullptr && InKey.Type == ERigElementType::Control)
+	if (Container != nullptr)
 	{
-		int32 Index = Container->GetIndex(InKey);
-		if (Index != INDEX_NONE)
+		if (InKey.Type == ERigElementType::Control)
 		{
-			Container->ControlHierarchy.Controls[Index].SetValueFromTransform(InValue, InValueType);
+			int32 Index = Container->GetIndex(InKey);
+			if (Index != INDEX_NONE)
+			{
+				Container->ControlHierarchy.Controls[Index].SetValueFromTransform(InValue, InValueType);
+				return;
+			}
 		}
+		ReportError(TEXT("Invalid Rig Element Key"));
 	}
 }
 
 void UControlRigHierarchyModifier::SetControlOffsetTransform(const FRigElementKey& InKey, FTransform InValue)
 {
-	if (Container != nullptr && InKey.Type == ERigElementType::Control)
+	if (Container != nullptr)
 	{
-		int32 Index = Container->GetIndex(InKey);
-		if (Index != INDEX_NONE)
+		if (InKey.Type == ERigElementType::Control)
 		{
-			Container->ControlHierarchy.Controls[Index].OffsetTransform = InValue;
+			int32 Index = Container->GetIndex(InKey);
+			if (Index != INDEX_NONE)
+			{
+				Container->ControlHierarchy.Controls[Index].OffsetTransform = InValue;
+				return;
+			}
 		}
+		ReportError(TEXT("Invalid Rig Element Key"));
 	}
 }
 
@@ -321,7 +408,7 @@ FRigElementKey UControlRigHierarchyModifier::AddSpace
 	const FString& InParentName
 )
 {
-	if(Container != nullptr)
+	if (Container != nullptr)
 	{
 		FName NewName = Container->SpaceHierarchy.GetSafeNewName(InNewName);
 		FName ParentName = FRigHierarchyContainer::GetSanitizedName(InParentName);
@@ -332,13 +419,17 @@ FRigElementKey UControlRigHierarchyModifier::AddSpace
 
 FRigSpace UControlRigHierarchyModifier::GetSpace(const FRigElementKey& InKey)
 {
-	if (Container != nullptr && InKey.Type == ERigElementType::Space)
+	if (Container != nullptr)
 	{
-		int32 Index = Container->GetIndex(InKey);
-		if (Index != INDEX_NONE)
+		if (InKey.Type == ERigElementType::Space)
 		{
-			return Container->SpaceHierarchy.Spaces[Index];
+			int32 Index = Container->GetIndex(InKey);
+			if (Index != INDEX_NONE)
+			{
+				return Container->SpaceHierarchy.Spaces[Index];
+			}
 		}
+		ReportError(TEXT("Invalid Rig Element Key"));
 	}
 	return InvalidSpace;
 }
@@ -356,14 +447,16 @@ void UControlRigHierarchyModifier::SetSpace(const FRigSpace& InElement)
 			if (UControlRigBlueprint* Blueprint = Cast<UControlRigBlueprint>(GetOuter()))
 			{
 				Blueprint->PropagateHierarchyFromBPToInstances(true, true);
+				return;
 			}
 		}
+		ReportError(TEXT("Space not found"));
 	}
 }
 
 FRigElementKey UControlRigHierarchyModifier::AddCurve(const FString& InNewName, float InValue)
 {
-	if(Container != nullptr)
+	if (Container != nullptr)
 	{
 		FName NewName = Container->CurveContainer.GetSafeNewName(InNewName);
 		return Container->CurveContainer.Add(NewName, InValue).GetElementKey();
@@ -373,13 +466,17 @@ FRigElementKey UControlRigHierarchyModifier::AddCurve(const FString& InNewName,
 
 FRigCurve UControlRigHierarchyModifier::GetCurve(const FRigElementKey& InKey)
 {
-	if (Container != nullptr && InKey.Type == ERigElementType::Curve)
-	{
-		int32 Index = Container->GetIndex(InKey);
-		if (Index != INDEX_NONE)
-		{
-			return Container->CurveContainer.Curves[Index];
+	if (Container != nullptr)
+	{ 
+		if (InKey.Type == ERigElementType::Curve) 
+		{ 
+			int32 Index = Container->GetIndex(InKey);
+			if (Index != INDEX_NONE)
+			{
+				return Container->CurveContainer.Curves[Index];
+			}
 		}
+		ReportError(TEXT("Invalid Rig Element Key"));
 	}
 	return InvalidCurve;
 }
@@ -397,20 +494,23 @@ void UControlRigHierarchyModifier::SetCurve(const FRigCurve& InElement)
 			if (UControlRigBlueprint* Blueprint = Cast<UControlRigBlueprint>(GetOuter()))
 			{
 				Blueprint->PropagateHierarchyFromBPToInstances(true, true);
+				return;
 			}
 		}
+		ReportError(TEXT("Curve not found"));
 	}
 }
 
 bool UControlRigHierarchyModifier::RemoveElement(const FRigElementKey& InElement)
 {
-	if(Container == nullptr)
+	if (Container == nullptr)
 	{
 		return false;
 	}
 
 	if (Container->GetIndex(InElement) == INDEX_NONE)
 	{
+		ReportError(TEXT("Element not found"));
 		return false;
 	}
 
@@ -447,13 +547,14 @@ bool UControlRigHierarchyModifier::RemoveElement(const FRigElementKey& InElement
 
 FRigElementKey UControlRigHierarchyModifier::RenameElement(const FRigElementKey& InElement, const FString& InNewName)
 {
-	if(Container == nullptr)
+	if (Container == nullptr)
 	{
 		return FRigElementKey();
 	}
 
 	if (Container->GetIndex(InElement) == INDEX_NONE)
 	{
+		ReportError(TEXT("Element not found"));
 		return FRigElementKey();
 	}
 
@@ -492,13 +593,14 @@ FRigElementKey UControlRigHierarchyModifier::RenameElement(const FRigElementKey&
 
 bool UControlRigHierarchyModifier::ReparentElement(const FRigElementKey& InElement, const FRigElementKey& InNewParent)
 {
-	if(Container == nullptr)
+	if (Container == nullptr)
 	{
 		return false;
 	}
 
 	if (Container->GetIndex(InElement) == INDEX_NONE)
 	{
+		ReportError(TEXT("Element not found"));
 		return false;
 	}
 
@@ -574,7 +676,7 @@ TArray<FRigElementKey> UControlRigHierarchyModifier::GetSelection() const
 
 bool UControlRigHierarchyModifier::Select(const FRigElementKey& InKey, bool bSelect)
 {
-	if(Container == nullptr)
+	if (Container == nullptr)
 	{
 		return false;
 	}
@@ -583,7 +685,7 @@ bool UControlRigHierarchyModifier::Select(const FRigElementKey& InKey, bool bSel
 
 bool UControlRigHierarchyModifier::ClearSelection()
 {
-	if(Container == nullptr)
+	if (Container == nullptr)
 	{
 		return false;
 	}
@@ -592,7 +694,7 @@ bool UControlRigHierarchyModifier::ClearSelection()
 
 bool UControlRigHierarchyModifier::IsSelected(const FRigElementKey& InKey) const
 {
-	if(Container == nullptr)
+	if (Container == nullptr)
 	{
 		return false;
 	}
@@ -601,7 +703,7 @@ bool UControlRigHierarchyModifier::IsSelected(const FRigElementKey& InKey) const
 
 void UControlRigHierarchyModifier::Initialize(bool bResetTransforms)
 {
-	if(Container == nullptr)
+	if (Container == nullptr)
 	{
 		return;
 	}
@@ -610,7 +712,7 @@ void UControlRigHierarchyModifier::Initialize(bool bResetTransforms)
 
 void UControlRigHierarchyModifier::Reset()
 {
-	if(Container == nullptr)
+	if (Container == nullptr)
 	{
 		return;
 	}
@@ -619,7 +721,7 @@ void UControlRigHierarchyModifier::Reset()
 
 void UControlRigHierarchyModifier::ResetTransforms()
 {
-	if(Container == nullptr)
+	if (Container == nullptr)
 	{
 		return;
 	}
@@ -628,7 +730,7 @@ void UControlRigHierarchyModifier::ResetTransforms()
 
 FTransform UControlRigHierarchyModifier::GetInitialTransform(const FRigElementKey& InKey) const
 {
-	if(Container == nullptr)
+	if (Container == nullptr)
 	{
 		return FTransform::Identity;
 	}
@@ -637,7 +739,7 @@ FTransform UControlRigHierarchyModifier::GetInitialTransform(const FRigElementKe
 
 void UControlRigHierarchyModifier::SetInitialTransform(const FRigElementKey& InKey, const FTransform& InTransform)
 {
-	if(Container == nullptr)
+	if (Container == nullptr)
 	{
 		return;
 	}
@@ -646,7 +748,7 @@ void UControlRigHierarchyModifier::SetInitialTransform(const FRigElementKey& InK
 
 FTransform UControlRigHierarchyModifier::GetInitialGlobalTransform(const FRigElementKey& InKey) const
 {
-	if(Container == nullptr)
+	if (Container == nullptr)
 	{
 		return FTransform::Identity;
 	}
@@ -656,7 +758,7 @@ FTransform UControlRigHierarchyModifier::GetInitialGlobalTransform(const FRigEle
 
 void UControlRigHierarchyModifier::SetInitialGlobalTransform(const FRigElementKey& InKey, const FTransform& InTransform)
 {
-	if(Container == nullptr)
+	if (Container == nullptr)
 	{
 		return;
 	}
@@ -665,7 +767,7 @@ void UControlRigHierarchyModifier::SetInitialGlobalTransform(const FRigElementKe
 
 FTransform UControlRigHierarchyModifier::GetLocalTransform(const FRigElementKey& InKey) const
 {
-	if(Container == nullptr)
+	if (Container == nullptr)
 	{
 		return FTransform::Identity;
 	}
@@ -675,7 +777,7 @@ FTransform UControlRigHierarchyModifier::GetLocalTransform(const FRigElementKey&
 
 void UControlRigHierarchyModifier::SetLocalTransform(const FRigElementKey& InKey, const FTransform& InTransform)
 {
-	if(Container == nullptr)
+	if (Container == nullptr)
 	{
 		return;
 	}
@@ -684,7 +786,7 @@ void UControlRigHierarchyModifier::SetLocalTransform(const FRigElementKey& InKey
 
 FTransform UControlRigHierarchyModifier::GetGlobalTransform(const FRigElementKey& InKey) const
 {
-	if(Container == nullptr)
+	if (Container == nullptr)
 	{
 		return FTransform::Identity;
 	}
@@ -694,7 +796,7 @@ FTransform UControlRigHierarchyModifier::GetGlobalTransform(const FRigElementKey
 
 void UControlRigHierarchyModifier::SetGlobalTransform(const FRigElementKey& InKey, const FTransform& InTransform)
 {
-	if(Container == nullptr)
+	if (Container == nullptr)
 	{
 		return;
 	}
@@ -703,7 +805,7 @@ void UControlRigHierarchyModifier::SetGlobalTransform(const FRigElementKey& InKe
 
 FString UControlRigHierarchyModifier::ExportToText(const TArray<FRigElementKey>& InElementsToExport) const
 {
-	if(Container == nullptr)
+	if (Container == nullptr)
 	{
 		return FString();
 	}
@@ -712,11 +814,27 @@ FString UControlRigHierarchyModifier::ExportToText(const TArray<FRigElementKey>&
 
 TArray<FRigElementKey> UControlRigHierarchyModifier::ImportFromText(const FString& InContent, ERigHierarchyImportMode InImportMode, bool bSelectNewElements)
 {
-	if(Container == nullptr)
+	if (Container == nullptr)
 	{
 		return TArray<FRigElementKey>();
 	}
 	return Container->ImportFromText(InContent, InImportMode, bSelectNewElements);
 }
 
+void UControlRigHierarchyModifier::ReportError(const FString& InMessage) const
+{
+	if (!bReportWarningsAndErrors)
+	{
+		return;
+	}
+
+	FString Message = InMessage;
+	if (UPackage* Package = Cast<UPackage>(this->GetOutermost()))
+	{
+		Message = FString::Printf(TEXT("%s : %s"), *Package->GetPathName(), *InMessage);
+	}
+
+	FScriptExceptionHandler::Get().HandleException(ELogVerbosity::Error, *Message, *FString());
+}
+
 #endif
\ No newline at end of file
diff --git a/Engine/Plugins/Experimental/ControlRig/Source/ControlRigDeveloper/Public/ControlRigHierarchyModifier.h b/Engine/Plugins/Experimental/ControlRig/Source/ControlRigDeveloper/Public/ControlRigHierarchyModifier.h
--- a/Engine/Plugins/Experimental/ControlRig/Source/ControlRigDeveloper/Public/ControlRigHierarchyModifier.h
+++ b/Engine/Plugins/Experimental/ControlRig/Source/ControlRigDeveloper/Public/ControlRigHierarchyModifier.h
@@ -233,6 +233,7 @@ public:
 #endif // WITH_EDITOR
 
 private:
+	void ReportError(const FString& InMessage) const;
 
 	FRigHierarchyContainer* Container;
 	FRigBone InvalidBone;
@@ -240,5 +241,7 @@ private:
 	FRigSpace InvalidSpace;
 	FRigCurve InvalidCurve;
 
+	bool bReportWarningsAndErrors;
+
 	friend class UControlRigBlueprint;
 };
