Control Rig: Changed popup text for "Set Initial Transform from Current/Closest Bone" to explain a corner case regarding controls

diff --git a/Engine/Plugins/Experimental/ControlRig/Source/ControlRigEditor/Private/ControlRigHierarchyCommands.cpp b/Engine/Plugins/Experimental/ControlRig/Source/ControlRigEditor/Private/ControlRigHierarchyCommands.cpp
--- a/Engine/Plugins/Experimental/ControlRig/Source/ControlRigEditor/Private/ControlRigHierarchyCommands.cpp
+++ b/Engine/Plugins/Experimental/ControlRig/Source/ControlRigEditor/Private/ControlRigHierarchyCommands.cpp
@@ -17,8 +17,8 @@ void FControlRigHierarchyCommands::RegisterCommands()
 	UI_COMMAND(PasteItems, "Paste", "Paste the selected items.", EUserInterfaceActionType::Button, FInputChord(EKeys::V, EModifierKey::Control));
 	UI_COMMAND(PasteLocalTransforms, "Paste Local Transform", "Paste the local transforms.", EUserInterfaceActionType::Button, FInputChord());
 	UI_COMMAND(PasteGlobalTransforms, "Paste Global Transform", "Paste the global transforms.", EUserInterfaceActionType::Button, FInputChord(EKeys::V, EModifierKey::Control | EModifierKey::Shift));
-	UI_COMMAND(SetInitialTransformFromClosestBone, "Set Initial Transform from Closest Bone", "Find the Closest Bone to Initial Transform of the Selected Bones", EUserInterfaceActionType::Button, FInputChord());
-	UI_COMMAND(SetInitialTransformFromCurrentTransform, "Set Initial Transform from Current", "Save the Current Transform To Initial Transform of the Selected Bones", EUserInterfaceActionType::Button, FInputChord());
+	UI_COMMAND(SetInitialTransformFromClosestBone, "Set Initial Transform from Closest Bone", "Find the Closest Bone to Initial Transform of the Selected Bones/Spaces or To Offset Transform of Selected Controls", EUserInterfaceActionType::Button, FInputChord());
+	UI_COMMAND(SetInitialTransformFromCurrentTransform, "Set Initial Transform from Current", "Save the Current Transform To Initial Transform of the Selected Bones/Spaces or To Offset Transform of Selected Controls", EUserInterfaceActionType::Button, FInputChord());
 	UI_COMMAND(ResetTransform, "Reset Transform", "Reset the Transform", EUserInterfaceActionType::Button, FInputChord(EKeys::G, EModifierKey::Control));
 	UI_COMMAND(ResetAllTransforms, "Reset All Transforms", "Resets all Transforms", EUserInterfaceActionType::Button, FInputChord(EKeys::G, EModifierKey::Control | EModifierKey::Shift));
 	UI_COMMAND(ResetSpace, "Reset Space", "Resets or injects a Space below the Control", EUserInterfaceActionType::Button, FInputChord());
