Control Rig: Support for Exposing pins by drag and drop them onto collapse, entry and return nodes.

diff --git a/Engine/Plugins/Experimental/ControlRig/Source/ControlRigDeveloper/Private/Graph/ControlRigGraphSchema.cpp b/Engine/Plugins/Experimental/ControlRig/Source/ControlRigDeveloper/Private/Graph/ControlRigGraphSchema.cpp
--- a/Engine/Plugins/Experimental/ControlRig/Source/ControlRigDeveloper/Private/Graph/ControlRigGraphSchema.cpp
+++ b/Engine/Plugins/Experimental/ControlRig/Source/ControlRigDeveloper/Private/Graph/ControlRigGraphSchema.cpp
@@ -18,6 +18,10 @@
 #include "Framework/Notifications/NotificationManager.h"
 #include "EulerTransform.h"
 #include "Curves/CurveFloat.h"
+#include "RigVMModel/Nodes/RigVMLibraryNode.h"
+#include "RigVMModel/Nodes/RigVMCollapseNode.h"
+#include "RigVMModel/Nodes/RigVMFunctionEntryNode.h"
+#include "RigVMModel/Nodes/RigVMFunctionReturnNode.h"
 
 #define LOCTEXT_NAMESPACE "ControlRigGraphSchema"
 
@@ -326,6 +330,110 @@ bool UControlRigGraphSchema::TryRenameGraph(UEdGraph* GraphToRename, const FName
 	return false;
 }
 
+UEdGraphPin* UControlRigGraphSchema::DropPinOnNode(UEdGraphNode* InTargetNode, const FName& InSourcePinName, const FEdGraphPinType& InSourcePinType, EEdGraphPinDirection InSourcePinDirection) const
+{
+	FString NewPinName;
+
+	if (UControlRigBlueprint* RigBlueprint = Cast<UControlRigBlueprint>(FBlueprintEditorUtils::FindBlueprintForNode(InTargetNode)))
+	{
+		if (UControlRigGraphNode* RigNode = Cast<UControlRigGraphNode>(InTargetNode))
+		{
+			if (URigVMNode* ModelNode = RigNode->GetModelNode())
+			{
+				URigVMGraph* Model = nullptr;
+				ERigVMPinDirection PinDirection = InSourcePinDirection == EGPD_Input ? ERigVMPinDirection::Input : ERigVMPinDirection::Output;
+
+				if (URigVMCollapseNode* CollapseNode = Cast<URigVMCollapseNode>(ModelNode))
+				{
+					Model = CollapseNode->GetContainedGraph();
+					PinDirection = PinDirection == ERigVMPinDirection::Output ? ERigVMPinDirection::Input : ERigVMPinDirection::Output;
+				}
+				else if (ModelNode->IsA<URigVMFunctionEntryNode>() ||
+					ModelNode->IsA<URigVMFunctionReturnNode>())
+				{
+					Model = ModelNode->GetGraph();
+				}
+
+				if (Model)
+				{
+					ensure(!Model->IsTopLevelGraph());
+
+					FRigVMExternalVariable ExternalVar = UControlRig::GetExternalVariableFromPinType(InSourcePinName, InSourcePinType);
+					if (ExternalVar.IsValid(true /* allow null memory */))
+					{
+						if (URigVMController* Controller = RigBlueprint->GetController(Model))
+						{
+							FString TypeName = ExternalVar.TypeName.ToString();
+							FName TypeObjectPathName = NAME_None;
+							if (ExternalVar.TypeObject)
+							{
+								TypeObjectPathName = *ExternalVar.TypeObject->GetPathName();
+							}
+
+							FName ExposedPinName = Controller->AddExposedPin(
+								InSourcePinName,
+								PinDirection,
+								TypeName,
+								TypeObjectPathName,
+								FString()
+							);
+							
+							if (!ExposedPinName.IsNone())
+							{
+								NewPinName = ExposedPinName.ToString();
+							}
+						}
+					}
+				}
+
+				if (!NewPinName.IsEmpty())
+				{
+					if (URigVMPin* NewModelPin = ModelNode->FindPin(NewPinName))
+					{
+						return RigNode->FindPin(*NewModelPin->GetPinPath());
+					}
+				}
+			}
+		}
+	}
+
+	return nullptr;
+}
+
+bool UControlRigGraphSchema::SupportsDropPinOnNode(UEdGraphNode* InTargetNode, const FEdGraphPinType& InSourcePinType, EEdGraphPinDirection InSourcePinDirection, FText& OutErrorMessage) const
+{
+	if (UControlRigGraphNode* RigNode = Cast<UControlRigGraphNode>(InTargetNode))
+	{
+		if(URigVMNode* ModelNode = RigNode->GetModelNode())
+		{
+			if (ModelNode->IsA<URigVMFunctionEntryNode>())
+			{
+				if (InSourcePinDirection == EGPD_Output)
+				{
+					OutErrorMessage = LOCTEXT("AddPinToReturnNode", "Add Pin to Return Node");
+					return false;
+				}
+				return true;
+			}
+			else if (ModelNode->IsA<URigVMFunctionReturnNode>())
+			{
+				if (InSourcePinDirection == EGPD_Input)
+				{
+					OutErrorMessage = LOCTEXT("AddPinToEntryNode", "Add Pin to Entry Node");
+					return false;
+				}
+				return true;
+			}
+			else if (ModelNode->IsA<URigVMCollapseNode>())
+			{
+				return true;
+			}
+		}
+	}
+
+	return false;
+}
+
 UControlRigGraphNode* UControlRigGraphSchema::CreateGraphNode(UControlRigGraph* InGraph, const FName& InPropertyName) const
 {
 	const bool bSelectNewNode = true;
diff --git a/Engine/Plugins/Experimental/ControlRig/Source/ControlRigDeveloper/Public/Graph/ControlRigGraphSchema.h b/Engine/Plugins/Experimental/ControlRig/Source/ControlRigDeveloper/Public/Graph/ControlRigGraphSchema.h
--- a/Engine/Plugins/Experimental/ControlRig/Source/ControlRigDeveloper/Public/Graph/ControlRigGraphSchema.h
+++ b/Engine/Plugins/Experimental/ControlRig/Source/ControlRigDeveloper/Public/Graph/ControlRigGraphSchema.h
@@ -81,6 +81,8 @@ public:
 	virtual bool TryDeleteGraph(UEdGraph* GraphToDelete) const override;
 	virtual bool TryRenameGraph(UEdGraph* GraphToRename, const FName& InNewName) const override;
 	virtual bool CanDuplicateGraph(UEdGraph* InSourceGraph) const { return false; }
+	virtual UEdGraphPin* DropPinOnNode(UEdGraphNode* InTargetNode, const FName& InSourcePinName, const FEdGraphPinType& InSourcePinType, EEdGraphPinDirection InSourcePinDirection) const override;
+	virtual bool SupportsDropPinOnNode(UEdGraphNode* InTargetNode, const FEdGraphPinType& InSourcePinType, EEdGraphPinDirection InSourcePinDirection, FText& OutErrorMessage) const override;
 
 	/** Create a graph node for a rig */
 	UControlRigGraphNode* CreateGraphNode(UControlRigGraph* InGraph, const FName& InPropertyName) const;
diff --git a/Engine/Source/Developer/RigVMDeveloper/Private/RigVMModel/RigVMController.cpp b/Engine/Source/Developer/RigVMDeveloper/Private/RigVMModel/RigVMController.cpp
--- a/Engine/Source/Developer/RigVMDeveloper/Private/RigVMModel/RigVMController.cpp
+++ b/Engine/Source/Developer/RigVMDeveloper/Private/RigVMModel/RigVMController.cpp
@@ -4928,7 +4928,7 @@ FName URigVMController::AddExposedPin(const FName& InPinName, ERigVMPinDirection
 		Notify(ERigVMGraphNotifType::PinAdded, Pin);
 	}
 
-	if (bSetupUndoRedo)
+	if (bSetupUndoRedo && !InDefaultValue.IsEmpty())
 	{
 		SetPinDefaultValue(Pin, InDefaultValue, true, bSetupUndoRedo, false);
 	}
@@ -4941,7 +4941,7 @@ FName URigVMController::AddExposedPin(const FName& InPinName, ERigVMPinDirection
 		ActionStack->EndAction(Action);
 	}
 	
-	return NAME_None;
+	return PinName;
 }
 
 bool URigVMController::RemoveExposedPin(const FName& InPinName, bool bSetupUndoRedo)
