Control Rig: Tweener and Pose Library is dockable.
Tweener added. Will need some feedback but currently just a slider.

Pose Library Opens at last position.
Provide prompt when saving a pose with no selection.

diff --git a/Engine/Plugins/Experimental/ControlRig/Source/ControlRig/Private/Tools/ControlRigPose.cpp b/Engine/Plugins/Experimental/ControlRig/Source/ControlRig/Private/Tools/ControlRigPose.cpp
--- a/Engine/Plugins/Experimental/ControlRig/Source/ControlRig/Private/Tools/ControlRigPose.cpp
+++ b/Engine/Plugins/Experimental/ControlRig/Source/ControlRig/Private/Tools/ControlRigPose.cpp
@@ -8,6 +8,250 @@
 
 #define LOCTEXT_NAMESPACE "ControlRigPose"
 
+void FControlRigControlPose::SavePose(UControlRig* ControlRig, bool bUseAll)
+{
+	TArray<FRigControl> CurrentControls;
+	ControlRig->GetControlsInOrder(CurrentControls);
+	CopyOfControls.SetNum(0);
+	FRigControlHierarchy& Hierarchy = ControlRig->GetControlHierarchy();
+	for (const FRigControl& RigControl : CurrentControls)
+	{
+		if (RigControl.bAnimatable && (bUseAll || ControlRig->IsControlSelected(RigControl.Name)))
+		{
+			FRigControlCopy Copy(RigControl, Hierarchy);
+			CopyOfControls.Add(Copy);
+		}
+	}
+	SetUpControlMap();
+}
+
+void FControlRigControlPose::PastePose(UControlRig* ControlRig, bool bDoKey, bool bDoMirror)
+{
+	PastePoseInternal(ControlRig, bDoKey, bDoMirror, CopyOfControls);
+	ControlRig->Evaluate_AnyThread();
+	PastePoseInternal(ControlRig, bDoKey, bDoMirror, CopyOfControls);
+
+}
+
+void FControlRigControlPose::SetControlMirrorTransform(UControlRig* ControlRig, const FName& Name, bool bIsMatched, const FVector& GlobalTranslation, const FQuat& GlobalRotation, const FQuat& LocalRotation, bool bNotify, const  FRigControlModifiedContext& Context)
+{
+	if (bIsMatched)
+	{
+		int32 Index = ControlRig->GetControlHierarchy().GetIndex(Name);
+		FTransform ParentTransform = ControlRig->GetControlHierarchy().GetParentTransform(Index);
+		FTransform CurrentTransform = ControlRig->GetControlHierarchy().GetLocalTransform(Index);
+		FVector NewLocal = ParentTransform.InverseTransformPositionNoScale(GlobalTranslation);
+		FTransform NewCurrentTransform(LocalRotation, NewLocal);
+		ControlRig->SetControlLocalTransform(Name, NewCurrentTransform, bNotify, Context);
+	}
+	else
+	{
+		FTransform NewGlobalTransform(GlobalRotation, GlobalTranslation);
+		ControlRig->SetControlGlobalTransform(Name, NewGlobalTransform, bNotify, Context);
+	}
+}
+
+void FControlRigControlPose::PastePoseInternal(UControlRig* ControlRig, bool bDoKey, bool bDoMirror, const TArray<FRigControlCopy>& ControlsToPaste)
+{
+	FRigControlModifiedContext Context;
+	Context.SetKey = bDoKey ? EControlRigSetKey::Always : EControlRigSetKey::DoNotCare;
+	FControlRigPoseMirrorTable MirrorTable;
+	if (bDoMirror)
+	{
+		MirrorTable.SetUpMirrorTable(ControlRig);
+	}
+
+	TArray<FRigControl> SortedControls;
+	ControlRig->GetControlsInOrder(SortedControls);
+	for (const FRigControl& RigControl : SortedControls)
+	{
+		if (!ControlRig->IsControlSelected(RigControl.Name))
+		{
+			continue;
+		}
+		FRigControlCopy* CopyRigControl = MirrorTable.GetControl(*this, RigControl.Name);
+		if (CopyRigControl)
+		{
+			if (CopyRigControl->ControlType == RigControl.ControlType)
+			{
+				switch (RigControl.ControlType)
+				{
+				case ERigControlType::Transform:
+				case ERigControlType::TransformNoScale:
+				case ERigControlType::EulerTransform:
+				case ERigControlType::Position:
+				case ERigControlType::Scale:
+				case ERigControlType::Rotator:
+				{
+					if (bDoMirror == false)
+					{
+						ControlRig->SetControlGlobalTransform(RigControl.Name, CopyRigControl->GlobalTransform, true, Context);
+					}
+					else
+					{
+						FVector GlobalTranslation;
+						FQuat GlobalRotation;
+						FQuat LocalRotation;
+						bool bIsMatched = MirrorTable.IsMatched(CopyRigControl->Name);
+						MirrorTable.GetMirrorTransform(*CopyRigControl, bIsMatched, GlobalTranslation, GlobalRotation, LocalRotation);
+						SetControlMirrorTransform(ControlRig, RigControl.Name, bIsMatched, GlobalTranslation, GlobalRotation, LocalRotation, true, Context);
+					}
+					break;
+				}
+				case ERigControlType::Float:
+				{
+					float Val = CopyRigControl->Value.Get<float>();
+					ControlRig->SetControlValue<float>(RigControl.Name, Val, true, Context);
+					break;
+				}
+				case ERigControlType::Bool:
+				{
+					bool Val = CopyRigControl->Value.Get<bool>();
+					ControlRig->SetControlValue<bool>(RigControl.Name, Val, true, Context);
+					break;
+				}
+				case ERigControlType::Integer:
+				{
+					int32 Val = CopyRigControl->Value.Get<int32>();
+					ControlRig->SetControlValue<int32>(RigControl.Name, Val, true, Context);
+					break;
+				}
+				case ERigControlType::Vector2D:
+				{
+					FVector2D Val = CopyRigControl->Value.Get<FVector2D>();
+					ControlRig->SetControlValue<FVector2D>(RigControl.Name, Val, true, Context);
+					break;
+				}
+				default:
+					//TODO add log
+					break;
+				};
+			}
+		}
+	}
+}
+
+void FControlRigControlPose::BlendWithInitialPoses(FControlRigControlPose& InitialPose, UControlRig* ControlRig, bool bDoKey, bool bDoMirror, float BlendValue)
+{
+	if (InitialPose.CopyOfControls.Num() == 0)
+	{
+		return;
+	}
+
+	//though can be n^2 should be okay, we search from current Index which in most cases will be the same
+	//not run often anyway
+	FRigControlModifiedContext Context;
+	Context.SetKey = bDoKey ? EControlRigSetKey::Always : EControlRigSetKey::DoNotCare;
+	FControlRigPoseMirrorTable MirrorTable;
+	if (bDoMirror)
+	{
+		MirrorTable.SetUpMirrorTable(ControlRig);
+	}
+
+	TArray<FRigControl> SortedControls;
+	ControlRig->GetControlsInOrder(SortedControls);
+	for (const FRigControl& RigControl : SortedControls)
+	{
+		if (!ControlRig->IsControlSelected(RigControl.Name))
+		{
+			continue;
+		}
+		FRigControlCopy* CopyRigControl = MirrorTable.GetControl(*this, RigControl.Name);
+		if (CopyRigControl)
+		{
+			if (CopyRigControl->ControlType == RigControl.ControlType)
+			{
+
+				FRigControlCopy* InitialFound = nullptr;
+				int32* Index = InitialPose.CopyOfControlsNameToIndex.Find(CopyRigControl->Name);
+				if (Index)
+				{
+					InitialFound = &(InitialPose.CopyOfControls[*Index]);
+				}
+				if (InitialFound && InitialFound->ControlType == CopyRigControl->ControlType)
+				{
+					if ((CopyRigControl->ControlType == ERigControlType::Transform || CopyRigControl->ControlType == ERigControlType::EulerTransform ||
+						CopyRigControl->ControlType == ERigControlType::TransformNoScale || CopyRigControl->ControlType == ERigControlType::Position ||
+						CopyRigControl->ControlType == ERigControlType::Rotator || CopyRigControl->ControlType == ERigControlType::Scale
+						))
+					{
+						if (bDoMirror == false)
+						{
+							FTransform Val = CopyRigControl->GlobalTransform;
+							FTransform InitialVal = InitialFound->GlobalTransform;
+							FVector Translation, Scale;
+							FQuat Rotation;
+							Translation = FMath::Lerp(InitialVal.GetTranslation(), Val.GetTranslation(), BlendValue);
+							Rotation = FQuat::Slerp(InitialVal.GetRotation(), Val.GetRotation(), BlendValue); //doing slerp here not fast lerp, can be slow this is for content creation
+							Scale = FMath::Lerp(InitialVal.GetScale3D(), Val.GetScale3D(), BlendValue);
+							Val = FTransform(Rotation, Translation, Scale);
+							ControlRig->SetControlGlobalTransform(RigControl.Name, Val, bDoKey, Context);
+						}
+						else
+						{
+							FVector GlobalTranslation;
+							FQuat GlobalRotation;
+							FQuat LocalRotation;
+							bool bIsMatched = MirrorTable.IsMatched(CopyRigControl->Name);
+							MirrorTable.GetMirrorTransform(*CopyRigControl, bIsMatched, GlobalTranslation, GlobalRotation, LocalRotation);
+							FVector InitialTranslation = InitialFound->GlobalTransform.GetTranslation();
+							FQuat InitialGlobalRotation = InitialFound->GlobalTransform.GetRotation();
+							FQuat InitialLocationRotation = InitialFound->LocalTransform.GetRotation();
+							GlobalTranslation = FMath::Lerp(InitialTranslation, GlobalTranslation, BlendValue);
+							GlobalRotation = FQuat::Slerp(InitialGlobalRotation, GlobalRotation, BlendValue); //doing slerp here not fast lerp, can be slow this is for content creation
+							LocalRotation = FQuat::Slerp(InitialLocationRotation, LocalRotation, BlendValue); //doing slerp here not fast lerp, can be slow this is for content creation
+
+							SetControlMirrorTransform(ControlRig, RigControl.Name, bIsMatched, GlobalTranslation, GlobalRotation, LocalRotation, bDoKey, Context);
+
+						}
+					}
+				}
+			}
+		}
+	}
+}
+
+bool FControlRigControlPose::ContainsName(const FName& Name) const
+{
+	const int32* Index = CopyOfControlsNameToIndex.Find(Name);
+	return (Index && *Index >= 0);
+}
+
+void FControlRigControlPose::ReplaceControlName(const FName& Name, const FName& NewName)
+{
+	int32* Index = CopyOfControlsNameToIndex.Find(Name);
+	if (Index && *Index >= 0)
+	{
+		FRigControlCopy& Control = CopyOfControls[*Index];
+		Control.Name = NewName;
+		CopyOfControlsNameToIndex.Remove(Name);
+		CopyOfControlsNameToIndex.Add(Control.Name, *Index);
+	}
+}
+
+TArray<FName> FControlRigControlPose::GetControlNames() const
+{
+	TArray<FName> Controls;
+	Controls.Reserve(CopyOfControls.Num());
+	for (const FRigControlCopy& Control : CopyOfControls)
+	{
+		Controls.Add(Control.Name);
+	}
+	return Controls;
+}
+
+void FControlRigControlPose::SetUpControlMap()
+{
+	CopyOfControlsNameToIndex.Reset();
+
+	for (int32 Index = 0; Index < CopyOfControls.Num(); ++Index)
+	{
+		const FRigControlCopy& Control = CopyOfControls[Index];
+		CopyOfControlsNameToIndex.Add(Control.Name, Index);
+	}
+}
+
+
 UControlRigPoseAsset::UControlRigPoseAsset(const FObjectInitializer& ObjectInitializer) : Super(ObjectInitializer)
 {
 }
@@ -77,5 +321,14 @@ void UControlRigPoseAsset::ReplaceControlName(const FName& CurrentName, const FN
 	Pose.ReplaceControlName(CurrentName, NewName);
 }
 
+bool UControlRigPoseAsset::DoesMirrorMatch(UControlRig* ControlRig, const FName& ControlName) const
+{
+	FControlRigPoseMirrorTable MirrorTable;
+	MirrorTable.SetUpMirrorTable(ControlRig);
+	return (MirrorTable.IsMatched(ControlName));
+}
+
+
+
 #undef LOCTEXT_NAMESPACE
 
diff --git a/Engine/Plugins/Experimental/ControlRig/Source/ControlRig/Public/Tools/ControlRigPose.h b/Engine/Plugins/Experimental/ControlRig/Source/ControlRig/Public/Tools/ControlRigPose.h
--- a/Engine/Plugins/Experimental/ControlRig/Source/ControlRig/Public/Tools/ControlRigPose.h
+++ b/Engine/Plugins/Experimental/ControlRig/Source/ControlRig/Public/Tools/ControlRigPose.h
@@ -85,7 +85,7 @@ struct CONTROLRIG_API FRigControlCopy
 * The Data Stored For Each Pose and associated Functions to Store and Paste It
 */
 USTRUCT(BlueprintType)
-struct FControlRigControlPose
+struct CONTROLRIG_API FControlRigControlPose
 {
 	GENERATED_USTRUCT_BODY()
 
@@ -96,248 +96,19 @@ struct FControlRigControlPose
 	}
 	~FControlRigControlPose() {};
 
-	void SavePose(UControlRig* ControlRig,  bool bUseAll)
-	{
-		TArray<FRigControl> CurrentControls; 
-		ControlRig->GetControlsInOrder(CurrentControls);
-		CopyOfControls.SetNum(0);
-		FRigControlHierarchy& Hierarchy = ControlRig->GetControlHierarchy();
-		for (const FRigControl& RigControl : CurrentControls)
-		{
-			if (RigControl.bAnimatable && (bUseAll || ControlRig->IsControlSelected(RigControl.Name)))
-			{
-				FRigControlCopy Copy(RigControl, Hierarchy);
-				CopyOfControls.Add(Copy);
-			}
-		}
-		SetUpControlMap();
-	}
-
-	void PastePose(UControlRig* ControlRig, bool bDoKey, bool bDoMirror)
-	{
-		PastePoseInternal(ControlRig,bDoKey, bDoMirror, CopyOfControls);
+	void SavePose(UControlRig* ControlRig, bool bUseAll);
+	void PastePose(UControlRig* ControlRig, bool bDoKey, bool bDoMirror);
+	void SetControlMirrorTransform(UControlRig* ControlRig, const FName& Name, bool bIsMatched, const FVector& GlobalTranslation, const FQuat& GlobalRotation, const FQuat& LocalRotation, bool bNotify, const  FRigControlModifiedContext& Context);
+	void PastePoseInternal(UControlRig* ControlRig, bool bDoKey, bool bDoMirror, const TArray<FRigControlCopy>& ControlsToPaste);
+	void BlendWithInitialPoses(FControlRigControlPose& InitialPose, UControlRig* ControlRig, bool bDoKey, bool bDoMirror, float BlendValue);
 
-	}
-
-	void SetControlMirrorTransform(UControlRig* ControlRig, const FName& Name, bool bIsMatched, const FVector& GlobalTranslation, const FQuat& GlobalRotation, const FQuat&LocalRotation, bool bNotify,const  FRigControlModifiedContext& Context)
-	{
-		if (bIsMatched)
-		{
-			int32 Index = ControlRig->GetControlHierarchy().GetIndex(Name);
-			FTransform ParentTransform = ControlRig->GetControlHierarchy().GetParentTransform(Index);
-			FTransform CurrentTransform = ControlRig->GetControlHierarchy().GetLocalTransform(Index);
-			FVector NewLocal = ParentTransform.InverseTransformPositionNoScale(GlobalTranslation);
-			FTransform NewCurrentTransform(LocalRotation, NewLocal);
-			ControlRig->SetControlLocalTransform(Name, NewCurrentTransform, bNotify, Context);
-		}
-		else
-		{
-			FTransform NewGlobalTransform(GlobalRotation, GlobalTranslation);
-			ControlRig->SetControlGlobalTransform(Name, NewGlobalTransform, bNotify, Context);
-		}
-	}
-
-	void PastePoseInternal(UControlRig* ControlRig, bool bDoKey,bool bDoMirror, const TArray<FRigControlCopy>& ControlsToPaste)
-	{
-		FRigControlModifiedContext Context;
-		Context.SetKey = bDoKey ? EControlRigSetKey::Always : EControlRigSetKey::DoNotCare;
-		FControlRigPoseMirrorTable MirrorTable;
-		if (bDoMirror)
-		{
-			MirrorTable.SetUpMirrorTable(ControlRig);
-		}
-	
-		TArray<FRigControl> SortedControls;
-		ControlRig->GetControlsInOrder(SortedControls);
-		for(const FRigControl& RigControl : SortedControls)
-		{ 
-			if (!ControlRig->IsControlSelected(RigControl.Name))
-			{
-				continue;
-			}
-			FRigControlCopy* CopyRigControl = MirrorTable.GetControl(*this,RigControl.Name);
-			if (CopyRigControl)
-			{
-				if (CopyRigControl->ControlType == RigControl.ControlType)
-				{
-					switch (RigControl.ControlType)
-					{
-					case ERigControlType::Transform:
-					case ERigControlType::TransformNoScale:
-					case ERigControlType::EulerTransform:
-					case ERigControlType::Position:
-					case ERigControlType::Scale:
-					case ERigControlType::Rotator:
-					{
-						if (bDoMirror == false)
-						{
-							ControlRig->SetControlGlobalTransform(RigControl.Name, CopyRigControl->GlobalTransform,true, Context);
-						}
-						else
-						{
-							FVector GlobalTranslation;
-							FQuat GlobalRotation;
-							FQuat LocalRotation;
-							bool bIsMatched = MirrorTable.IsMatched(CopyRigControl->Name);
-							MirrorTable.GetMirrorTransform(*CopyRigControl, bIsMatched,GlobalTranslation, GlobalRotation, LocalRotation);
-							SetControlMirrorTransform(ControlRig,RigControl.Name,bIsMatched, GlobalTranslation, GlobalRotation, LocalRotation,true, Context);
-						}
-						break;
-					}		
-					case ERigControlType::Float:
-					{
-						float Val = CopyRigControl->Value.Get<float>();
-						ControlRig->SetControlValue<float>(RigControl.Name, Val, true, Context);
-						break;
-					}
-					case ERigControlType::Bool:
-					{
-						bool Val = CopyRigControl->Value.Get<bool>();
-						ControlRig->SetControlValue<bool>(RigControl.Name, Val, true, Context);
-						break;
-					}
-					case ERigControlType::Integer:
-					{
-						int32 Val = CopyRigControl->Value.Get<int32>();
-						ControlRig->SetControlValue<int32>(RigControl.Name, Val, true, Context);
-						break;
-					}
-					case ERigControlType::Vector2D:
-					{
-						FVector2D Val = CopyRigControl->Value.Get<FVector2D>();
-						ControlRig->SetControlValue<FVector2D>(RigControl.Name, Val, true, Context);
-						break;
-					}	
-					default:
-						//TODO add log
-						break;
-					};
-				}
-			}
-		}
-	}
-
-	void BlendWithInitialPoses(FControlRigControlPose& InitialPose, UControlRig *ControlRig, bool bDoKey, bool bDoMirror, float BlendValue)
-	{
-		if (InitialPose.CopyOfControls.Num() == 0)
-		{
-			return;
-		}
-
-		//though can be n^2 should be okay, we search from current Index which in most cases will be the same
-		//not run often anyway
-		FRigControlModifiedContext Context;
-		Context.SetKey = bDoKey ? EControlRigSetKey::Always : EControlRigSetKey::DoNotCare;
-		FControlRigPoseMirrorTable MirrorTable;
-		if (bDoMirror)
-		{
-			MirrorTable.SetUpMirrorTable(ControlRig);
-		}
-
-		TArray<FRigControl> SortedControls;
-		ControlRig->GetControlsInOrder(SortedControls);
-		for (const FRigControl& RigControl : SortedControls)
-		{
-			if (!ControlRig->IsControlSelected(RigControl.Name))
-			{
-				continue;
-			}
-			FRigControlCopy* CopyRigControl = MirrorTable.GetControl(*this, RigControl.Name);
-			if (CopyRigControl)
-			{
-				if (CopyRigControl->ControlType == RigControl.ControlType)
-				{
-
-					FRigControlCopy* InitialFound = nullptr;
-					int32* Index = InitialPose.CopyOfControlsNameToIndex.Find(CopyRigControl->Name);
-					if (Index)
-					{
-						InitialFound = &(InitialPose.CopyOfControls[*Index]);
-					}
-					if (InitialFound && InitialFound->ControlType == CopyRigControl->ControlType)
-					{
-						if ((CopyRigControl->ControlType == ERigControlType::Transform || CopyRigControl->ControlType == ERigControlType::EulerTransform ||
-							CopyRigControl->ControlType == ERigControlType::TransformNoScale || CopyRigControl->ControlType == ERigControlType::Position ||
-							CopyRigControl->ControlType == ERigControlType::Rotator || CopyRigControl->ControlType == ERigControlType::Scale
-							))
-						{
-							if (bDoMirror == false)
-							{
-								FTransform Val = CopyRigControl->GlobalTransform;
-								FTransform InitialVal = InitialFound->GlobalTransform;
-								FVector Translation, Scale;
-								FQuat Rotation;
-								Translation = FMath::Lerp(InitialVal.GetTranslation(), Val.GetTranslation(), BlendValue);
-								Rotation = FQuat::Slerp(InitialVal.GetRotation(), Val.GetRotation(), BlendValue); //doing slerp here not fast lerp, can be slow this is for content creation
-								Scale = FMath::Lerp(InitialVal.GetScale3D(), Val.GetScale3D(), BlendValue);
-								Val = FTransform(Rotation, Translation, Scale);
-								ControlRig->SetControlGlobalTransform(RigControl.Name, Val, false, Context);
-							}
-							else
-							{
-								FVector GlobalTranslation;
-								FQuat GlobalRotation;
-								FQuat LocalRotation;
-								bool bIsMatched = MirrorTable.IsMatched(CopyRigControl->Name);
-								MirrorTable.GetMirrorTransform(*CopyRigControl, bIsMatched, GlobalTranslation, GlobalRotation, LocalRotation);
-								FVector InitialTranslation = InitialFound->GlobalTransform.GetTranslation();
-								FQuat InitialGlobalRotation = InitialFound->GlobalTransform.GetRotation();
-								FQuat InitialLocationRotation = InitialFound->LocalTransform.GetRotation();
-								GlobalTranslation = FMath::Lerp(InitialTranslation, GlobalTranslation, BlendValue);
-								GlobalRotation = FQuat::Slerp(InitialGlobalRotation, GlobalRotation, BlendValue); //doing slerp here not fast lerp, can be slow this is for content creation
-								LocalRotation = FQuat::Slerp(InitialLocationRotation, LocalRotation, BlendValue); //doing slerp here not fast lerp, can be slow this is for content creation
-
-								SetControlMirrorTransform(ControlRig, RigControl.Name, bIsMatched, GlobalTranslation, GlobalRotation, LocalRotation, false, Context);
-														
-							}
-						}
-					}
-				}
-			}			
-		}
-	}
-	bool ContainsName(const FName& Name) const
-	{
-		const int32* Index = CopyOfControlsNameToIndex.Find(Name);
-		return (Index && *Index >= 0);
-	}
-
-	void ReplaceControlName(const FName& Name, const FName& NewName)
-	{
-		int32* Index = CopyOfControlsNameToIndex.Find(Name);
-		if (Index && *Index >= 0)
-		{
-			FRigControlCopy& Control = CopyOfControls[*Index];
-			Control.Name = NewName;
-			CopyOfControlsNameToIndex.Remove(Name);
-			CopyOfControlsNameToIndex.Add(Control.Name, *Index);
-		}
-	}
-
-	TArray<FName> GetControlNames() const
-	{
-		TArray<FName> Controls;
-		Controls.Reserve(CopyOfControls.Num());
-		for (const FRigControlCopy& Control : CopyOfControls)
-		{
-			Controls.Add(Control.Name);
-		}
-		return Controls;
-	}
+	bool ContainsName(const FName& Name) const;
+	void ReplaceControlName(const FName& Name, const FName& NewName);
+	TArray<FName> GetControlNames() const;
 
-	void SetUpControlMap()
-	{
-		CopyOfControlsNameToIndex.Reset();
-		
-		for (int32 Index = 0;  Index <  CopyOfControls.Num(); ++ Index)
-		{
-			const FRigControlCopy& Control =  CopyOfControls[Index];
-			CopyOfControlsNameToIndex.Add(Control.Name, Index);
-		}
-	}
+	void SetUpControlMap();
 	TArray<FRigControlCopy> GetPoses() const {return CopyOfControls;};
 
-
 public:
 	UPROPERTY(BlueprintReadWrite, EditAnywhere, Category = "Controls")
 	TArray<FRigControlCopy> CopyOfControls;
@@ -348,7 +119,7 @@ public:
 
 
 /**
-* An indivual Pose made of Control Rig Controls
+* An individual Pose made of Control Rig Controls
 */
 UCLASS(BlueprintType)
 class CONTROLRIG_API UControlRigPoseAsset : public UObject
@@ -374,12 +145,15 @@ public:
 	UFUNCTION(BlueprintCallable, Category = "Pose")
 	void GetCurrentPose(UControlRig* InControlRig, FControlRigControlPose& OutPose);
 
-	UFUNCTION(BlueprintCallable, Category = "Pose")
+	UFUNCTION(BlueprintPure, Category = "Pose")
 	TArray<FName> GetControlNames() const;
 
 	UFUNCTION(BlueprintCallable, Category = "Pose")
 	void ReplaceControlName(const FName& CurrentName, const FName& NewName);
 
+	UFUNCTION(BlueprintPure, Category = "Pose")
+	bool DoesMirrorMatch(UControlRig* ControlRig, const FName& ControlName) const;
+
 	void BlendWithInitialPoses(FControlRigControlPose& InitialPose, UControlRig* ControlRig, bool bDoKey, bool bdoMirror, float BlendValue);
 
 	UPROPERTY(BlueprintReadWrite, EditAnywhere, Category = "Pose")
diff --git a/Engine/Plugins/Experimental/ControlRig/Source/ControlRigEditor/Private/ControlRigEditorModule.cpp b/Engine/Plugins/Experimental/ControlRig/Source/ControlRigEditor/Private/ControlRigEditorModule.cpp
--- a/Engine/Plugins/Experimental/ControlRig/Source/ControlRigEditor/Private/ControlRigEditorModule.cpp
+++ b/Engine/Plugins/Experimental/ControlRig/Source/ControlRigEditor/Private/ControlRigEditorModule.cpp
@@ -109,12 +109,16 @@
 #include "ControlRig/Private/Units/Execution/RigUnit_InverseExecution.h"
 #include "Graph/SControlRigGraphPinVariableBinding.h"
 #include "AssetTypeActions_ControlRigPose.h"
-
+#include "EditMode/SControlRigBaseListWidget.h"
+#include "EditMode/SControlRigTweenWidget.h"
 
 #define LOCTEXT_NAMESPACE "ControlRigEditorModule"
 
 DEFINE_LOG_CATEGORY(LogControlRigEditor);
 
+const FName IControlRigEditorModule::ControlRigPoseTab = FName("ControlRigPoseTab");
+const FName IControlRigEditorModule::ControlRigTweenTab = FName("ControlRigTweenTab");
+
 TSharedRef<SDockTab> SpawnRigProfiler( const FSpawnTabArgs& Args )
 {
 	return SNew(SDockTab)
@@ -124,6 +128,24 @@ TSharedRef<SDockTab> SpawnRigProfiler( const FSpawnTabArgs& Args )
 		];
 }
 
+TSharedRef<SDockTab> SpawnPoseTab(const FSpawnTabArgs& Args)
+{
+	return SNew(SDockTab)
+		.TabRole(ETabRole::NomadTab)
+		[
+			SNew(SControlRigBaseListWidget)
+		];
+}
+
+TSharedRef<SDockTab> SpawnTweenTab(const FSpawnTabArgs& Args)
+{
+	return SNew(SDockTab)
+		.TabRole(ETabRole::NomadTab)
+		[
+			SNew(SControlRigTweenWidget)
+		];
+}
+
 void FControlRigEditorModule::StartupModule()
 {
 	FControlRigEditModeCommands::Register();
@@ -232,6 +254,19 @@ void FControlRigEditorModule::StartupModule()
 			.SetTooltipText(NSLOCTEXT("UnrealEditor", "HierarchicalProfilerTooltip", "Open the Hierarchical Profiler tab."))
 			.SetGroup(WorkspaceMenu::GetMenuStructure().GetDeveloperToolsMiscCategory())
 			.SetIcon(FSlateIcon(TEXT("ControlRigEditorStyle"), TEXT("ControlRig.RigUnit")));
+
+		FGlobalTabmanager::Get()->RegisterNomadTabSpawner(IControlRigEditorModule::ControlRigPoseTab, FOnSpawnTab::CreateStatic(&SpawnPoseTab))
+			.SetDisplayName(NSLOCTEXT("UnrealEditor", "ControlRigPoseTab", "Control Rig Pose"))
+			.SetTooltipText(NSLOCTEXT("UnrealEditor", "ControlRigPoseTabTooltip", "Open the Control Rig Pose tab."))
+			.SetMenuType(ETabSpawnerMenuType::Hidden)
+			.SetIcon(FSlateIcon(FEditorStyle::GetStyleSetName(), "FoliageEditMode.Settings"));//MZ todo replace with correct icon
+		FGlobalTabmanager::Get()->RegisterDefaultTabWindowSize(IControlRigEditorModule::ControlRigPoseTab, FVector2D(850, 800));
+		FGlobalTabmanager::Get()->RegisterNomadTabSpawner(IControlRigEditorModule::ControlRigTweenTab, FOnSpawnTab::CreateStatic(&SpawnTweenTab))
+			.SetDisplayName(NSLOCTEXT("UnrealEditor", "ControlRigTweenTab", "Control Rig Tween"))
+			.SetTooltipText(NSLOCTEXT("UnrealEditor", "ControlRigTweenTooltip", "Open the Control Rig Tween tab."))
+			.SetMenuType(ETabSpawnerMenuType::Hidden)
+			.SetIcon(FSlateIcon(FEditorStyle::GetStyleSetName(), "SkeletonTree.BlendProfile")); //MZ todo replace with correct icon
+		FGlobalTabmanager::Get()->RegisterDefaultTabWindowSize(IControlRigEditorModule::ControlRigTweenTab, FVector2D(400, 100));
 	};
 #endif
 
@@ -256,6 +291,8 @@ void FControlRigEditorModule::ShutdownModule()
 	if (FSlateApplication::IsInitialized())
 	{
 		FGlobalTabmanager::Get()->UnregisterNomadTabSpawner("ControlRigProfiler");
+		FGlobalTabmanager::Get()->UnregisterNomadTabSpawner(IControlRigEditorModule::ControlRigPoseTab);
+		FGlobalTabmanager::Get()->UnregisterNomadTabSpawner(IControlRigEditorModule::ControlRigTweenTab);
 	}
 #endif
 
diff --git a/Engine/Plugins/Experimental/ControlRig/Source/ControlRigEditor/Private/EditMode/ControlRigEditMode.cpp b/Engine/Plugins/Experimental/ControlRig/Source/ControlRigEditor/Private/EditMode/ControlRigEditMode.cpp
--- a/Engine/Plugins/Experimental/ControlRig/Source/ControlRigEditor/Private/EditMode/ControlRigEditMode.cpp
+++ b/Engine/Plugins/Experimental/ControlRig/Source/ControlRigEditor/Private/EditMode/ControlRigEditMode.cpp
@@ -2178,6 +2178,7 @@ AControlRigGizmoActor* FControlRigEditMode::GetGizmoFromControlName(const FName&
 void FControlRigEditMode::AddControlRig(UControlRig* InControlRig)
 {
 	RuntimeControlRigs.AddUnique(InControlRig);
+	OnControlRigAddedOrRemovedDelegate.Broadcast(InControlRig, true);
 }
 
 UControlRig* FControlRigEditMode::GetControlRig(bool bInteractionRig, int32 InIndex) const
@@ -2203,6 +2204,7 @@ void FControlRigEditMode::RemoveControlRig(UControlRig* InControlRig)
 	int32 Index = RuntimeControlRigs.Find(InControlRig);
 	if (RuntimeControlRigs.IsValidIndex(Index))
 	{
+		OnControlRigAddedOrRemovedDelegate.Broadcast(InControlRig, false);
 		RuntimeControlRigs[Index]->ControlModified().RemoveAll(this);
 		RuntimeControlRigs.RemoveAt(Index);
 
diff --git a/Engine/Plugins/Experimental/ControlRig/Source/ControlRigEditor/Private/EditMode/ControlRigEditMode.h b/Engine/Plugins/Experimental/ControlRig/Source/ControlRigEditor/Private/EditMode/ControlRigEditMode.h
--- a/Engine/Plugins/Experimental/ControlRig/Source/ControlRigEditor/Private/EditMode/ControlRigEditMode.h
+++ b/Engine/Plugins/Experimental/ControlRig/Source/ControlRigEditor/Private/EditMode/ControlRigEditMode.h
@@ -42,6 +42,7 @@ enum class EControlRigSetKey : uint8;
 DECLARE_DELEGATE_RetVal_ThreeParams(FTransform, FOnGetRigElementTransform, const FRigElementKey& /*RigElementKey*/, bool /*bLocal*/, bool /*bOnDebugInstance*/);
 DECLARE_DELEGATE_ThreeParams(FOnSetRigElementTransform, const FRigElementKey& /*RigElementKey*/, const FTransform& /*Transform*/, bool /*bLocal*/);
 DECLARE_DELEGATE_RetVal(TSharedPtr<FUICommandList>, FNewMenuCommandsDelegate);
+DECLARE_MULTICAST_DELEGATE_TwoParams(FControlRigAddedOrRemoved, UControlRig*, bool /*true if added, false if removed*/);
 
 class FControlRigEditMode;
 
@@ -129,6 +130,9 @@ public:
 	FNewMenuCommandsDelegate& OnContextMenuCommands() { return OnContextMenuCommandsDelegate; }
 	FSimpleMulticastDelegate& OnAnimSystemInitialized() { return OnAnimSystemInitializedDelegate; }
 
+	/* Control Rig Changed Delegate*/
+	FControlRigAddedOrRemoved& OnControlRigAddedOrRemoved() { return OnControlRigAddedOrRemovedDelegate; }
+
 	// callback that gets called when rig element is selected in other view
 	void OnRigElementAdded(FRigHierarchyContainer* Container, const FRigElementKey& InKey);
 	void OnRigElementRemoved(FRigHierarchyContainer* Container, const FRigElementKey& InKey);
@@ -236,7 +240,8 @@ protected:
 	FNewMenuDelegate OnContextMenuDelegate;
 	FNewMenuCommandsDelegate OnContextMenuCommandsDelegate;
 	FSimpleMulticastDelegate OnAnimSystemInitializedDelegate;
-	
+	FControlRigAddedOrRemoved OnControlRigAddedOrRemovedDelegate;
+
 	TArray<FRigElementKey> SelectedRigElements;
 
 	/* Flag to recreate gizmos during tick */
diff --git a/Engine/Plugins/Experimental/ControlRig/Source/ControlRigEditor/Private/EditMode/SControlRigBaseListWidget.cpp b/Engine/Plugins/Experimental/ControlRig/Source/ControlRigEditor/Private/EditMode/SControlRigBaseListWidget.cpp
--- a/Engine/Plugins/Experimental/ControlRig/Source/ControlRigEditor/Private/EditMode/SControlRigBaseListWidget.cpp
+++ b/Engine/Plugins/Experimental/ControlRig/Source/ControlRigEditor/Private/EditMode/SControlRigBaseListWidget.cpp
@@ -36,10 +36,12 @@
 #include "ObjectTools.h"
 #include "SControlRigUpdatePose.h"
 #include "SControlRigRenamePoseControls.h"
-
+#include "Dialogs/Dialogs.h"
 
 #define LOCTEXT_NAMESPACE "ControlRigBaseListWidget"
 
+FString SControlRigBaseListWidget::CurrentlySelectedPath("");
+
 enum class FControlRigAssetType {
 	ControlRigPose,
 	ControlRigAnimation,
@@ -196,13 +198,13 @@ void FCreateControlAssetRigDialog::GetControlAssetParams(FControlRigAssetType Ty
 	switch (Type)
 	{
 	case FControlRigAssetType::ControlRigPose:
-		TitleText = NSLOCTEXT("ControlRig", "CreateControlAssetRig", "Create Control Rig Pose");
+		TitleText = LOCTEXT("CreateControlAssetRig", "Create Control Rig Pose");
 		break;
 	case FControlRigAssetType::ControlRigAnimation:
-		TitleText = NSLOCTEXT("ControlRig", "CreateControlAssetRig", "Create Control Rig Animation");
+		TitleText = LOCTEXT("CreateControlAssetRig", "Create Control Rig Animation");
 		break;
 	case FControlRigAssetType::ControlRigSelectionSet:
-		TitleText = NSLOCTEXT("ControlRig", "CreateControlAssetRig", "Create Control Rig Selection Set");
+		TitleText = LOCTEXT("CreateControlAssetRig", "Create Control Rig Selection Set");
 		break;
 	};
 
@@ -403,6 +405,29 @@ void SControlRigPoseAnimSelectionToolbar::Construct(const FArguments& InArgs)
 
 void SControlRigPoseAnimSelectionToolbar::MakeControlRigAssetDialog(FControlRigAssetType Type, bool bSelectAll)
 {
+	FControlRigEditMode* ControlRigEditMode = static_cast<FControlRigEditMode*>(GLevelEditorModeTools().GetActiveMode(FControlRigEditMode::ModeName));
+	if (ControlRigEditMode && ControlRigEditMode->GetControlRig(true))
+	{
+		UControlRig* ControlRig = ControlRigEditMode->GetControlRig(true);
+		if (ControlRig)
+		{
+			TArray<FName> SelectedControls = ControlRig->CurrentControlSelection();
+			if (SelectedControls.Num() <= 0)
+			{
+				FText ConfirmDelete = LOCTEXT("ConfirmNoSelectedControls", "You are saving a Pose with no selected Controls - are you sure?");
+
+				FSuppressableWarningDialog::FSetupInfo Info(ConfirmDelete, LOCTEXT("SavePose", "Save Pose"), "SavePose_Warning");
+				Info.ConfirmText = LOCTEXT("SavePose_Yes", "Yes");
+				Info.CancelText = LOCTEXT("SavePose_No", "No");
+
+				FSuppressableWarningDialog SavePose(Info);
+				if (SavePose.ShowModal() == FSuppressableWarningDialog::Cancel)
+				{
+					return;
+				}
+			}
+		}
+	}
 	FCreateControlAssetDelegate GetNameCallback = FCreateControlAssetDelegate::CreateLambda([this, Type, bSelectAll](FString AssetName)
 	{
 		if (OwningControlRigWidget)
@@ -411,7 +436,7 @@ void SControlRigPoseAnimSelectionToolbar::MakeControlRigAssetDialog(FControlRigA
 
 			FControlRigEditMode* ControlRigEditMode = static_cast<FControlRigEditMode*>(GLevelEditorModeTools().GetActiveMode(FControlRigEditMode::ModeName));
 			if (ControlRigEditMode && ControlRigEditMode->GetControlRig(true))
-			{
+			{	
 				UObject* NewAsset = nullptr;
 				switch (Type)
 				{
@@ -429,9 +454,8 @@ void SControlRigPoseAnimSelectionToolbar::MakeControlRigAssetDialog(FControlRigA
 				if (NewAsset)
 				{
 					FControlRigView::CaptureThumbnail(NewAsset);
+					OwningControlRigWidget->SelectThisAsset(NewAsset);
 				}
-				OwningControlRigWidget->SelectThisAsset(NewAsset);
-
 			}
 		}
 	});
@@ -639,7 +663,10 @@ void SControlRigBaseListWidget::Construct(const FArguments& InArgs)
 	TArray<FString> PosesDirectories = PoseSettings->GetAssetPaths();
 	if (PosesDirectories.Num() > 0)
 	{
-		CurrentlySelectedPath = PosesDirectories[0];
+		if (CurrentlySelectedPath.IsEmpty())
+		{
+			CurrentlySelectedPath = PosesDirectories[0];
+		}
 	}
 	FContentBrowserModule& ContentBrowserModule = FModuleManager::Get().LoadModuleChecked<FContentBrowserModule>(TEXT("ContentBrowser"));
 
@@ -701,42 +728,49 @@ void SControlRigBaseListWidget::Construct(const FArguments& InArgs)
 
 	// Path/Asset view
 	MainVerticalBox->AddSlot()
+
 		.HAlign(HAlign_Fill)
-		.FillHeight(0.7)
+		.FillHeight(0.6)
 		.Padding(0, 0, 0, 4)
 		[
 			SNew(SSplitter)
-
+			.Orientation(Orient_Vertical)
 			+ SSplitter::Slot()
-		.Value(0.33f)
-		[
-			SNew(SBorder)
-			.BorderImage(FEditorStyle::GetBrush("ToolPanel.GroupBorder"))
-		[
-			PathPicker.ToSharedRef()
-		]
-		]
+			.Value(0.66f)
+			[
+				SNew(SSplitter)
+
+				+ SSplitter::Slot()
+				.Value(0.33f)
+				[
+					SNew(SBorder)
+					.BorderImage(FEditorStyle::GetBrush("ToolPanel.GroupBorder"))
+				[
+					PathPicker.ToSharedRef()
+				]
+				]
 
-	+ SSplitter::Slot()
-		.Value(0.66f)
+			+ SSplitter::Slot()
+				.Value(0.66f)
+				[
+					SNew(SBorder)
+					.BorderImage(FEditorStyle::GetBrush("ToolPanel.GroupBorder"))
+					[
+						AssetPicker.ToSharedRef()
+					]
+				]
+			]
+		+ SSplitter::Slot()
+		.Value(0.4f)
 		[
 			SNew(SBorder)
 			.BorderImage(FEditorStyle::GetBrush("ToolPanel.GroupBorder"))
-		[
-			AssetPicker.ToSharedRef()
-		]
+			[
+				SAssignNew(ViewContainer, SBox)
+				.Padding(FMargin(5.0f, 0, 0, 0))
+			]
 		]
 		];
-		
-
-	//Bottom Area View Container to hold specific View (Pose/Animation).
-	MainVerticalBox->AddSlot()
-		.HAlign(HAlign_Fill)
-		.FillHeight(0.3)
-		[
-			SAssignNew(ViewContainer, SBox)
-			.Padding(FMargin(5.0f, 0, 0, 0))
-		];
 
 	ChildSlot
 		[
diff --git a/Engine/Plugins/Experimental/ControlRig/Source/ControlRigEditor/Private/EditMode/SControlRigBaseListWidget.h b/Engine/Plugins/Experimental/ControlRig/Source/ControlRigEditor/Private/EditMode/SControlRigBaseListWidget.h
--- a/Engine/Plugins/Experimental/ControlRig/Source/ControlRigEditor/Private/EditMode/SControlRigBaseListWidget.h
+++ b/Engine/Plugins/Experimental/ControlRig/Source/ControlRigEditor/Private/EditMode/SControlRigBaseListWidget.h
@@ -124,7 +124,7 @@ private:
 
 	TSharedPtr<FAssetFileContextMenu> AssetFileContextMenu;
 
-	FString CurrentlySelectedPath;
+	static FString CurrentlySelectedPath;
 	FString CurrentlyEnteredAssetName;
 	bool bLastInputValidityCheckSuccessful;
 
diff --git a/Engine/Plugins/Experimental/ControlRig/Source/ControlRigEditor/Private/EditMode/SControlRigControlViews.cpp b/Engine/Plugins/Experimental/ControlRig/Source/ControlRigEditor/Private/EditMode/SControlRigControlViews.cpp
--- a/Engine/Plugins/Experimental/ControlRig/Source/ControlRigEditor/Private/EditMode/SControlRigControlViews.cpp
+++ b/Engine/Plugins/Experimental/ControlRig/Source/ControlRigEditor/Private/EditMode/SControlRigControlViews.cpp
@@ -187,7 +187,6 @@ void SControlRigPoseView::Construct(const FArguments& InArgs)
 {
 	PoseAsset = InArgs._PoseAsset;
 
-
 	PoseBlendValue = 0.0f;
 	bIsBlending = false;
 	bSliderStartedTransaction = false;
@@ -208,6 +207,7 @@ void SControlRigPoseView::Construct(const FArguments& InArgs)
 	DetailsViewArgs.bUpdatesFromSelection = false;
 	DetailsViewArgs.bLockable = false;
 	DetailsViewArgs.NameAreaSettings = FDetailsViewArgs::HideNameArea;
+	DetailsViewArgs.NotifyHook = this;
 	DetailsViewArgs.ViewIdentifier = "Create Control Asset";
 
 	MirrorDetailsView = PropertyEditor.CreateDetailView(DetailsViewArgs);
@@ -290,18 +290,18 @@ void SControlRigPoseView::Construct(const FArguments& InArgs)
 						+ SVerticalBox::Slot()
 							.AutoHeight()
 							.HAlign(HAlign_Center)
-							.Padding(5.f)
+							.Padding(2.5f)
 							[
 								SNew(SHorizontalBox)
 								+ SHorizontalBox::Slot()
 							.AutoWidth()
 							.HAlign(HAlign_Center)
-							.Padding(5.f)
+							.Padding(2.5f)
 							[
 								SNew(SCheckBox)
 								.IsChecked(this, &SControlRigPoseView::IsKeyPoseChecked)
 							.OnCheckStateChanged(this, &SControlRigPoseView::OnKeyPoseChecked)
-							.Padding(5.0f)
+							.Padding(2.5f)
 
 							[
 								SNew(STextBlock).Text(LOCTEXT("Key", "Key"))
@@ -314,7 +314,7 @@ void SControlRigPoseView::Construct(const FArguments& InArgs)
 
 							.HAlign(HAlign_Center)
 
-							.Padding(5.f)
+							.Padding(2.5f)
 							[
 								SNew(SCheckBox)
 								.IsChecked(this, &SControlRigPoseView::IsMirrorPoseChecked)
@@ -330,25 +330,65 @@ void SControlRigPoseView::Construct(const FArguments& InArgs)
 						+ SVerticalBox::Slot()
 							.AutoHeight()
 							.HAlign(HAlign_Center)
-							.Padding(5.f)
+							.Padding(2.5f)
 							[
 
-								SNew(SNumericEntryBox<float>)
+								SNew(SSpinBox<float>)
 								// Only allow spinning if we have a single value
 								.Value(this, &SControlRigPoseView::OnGetPoseBlendValue)
-								.AllowSpin(true)
-								.MinValue(-1.0f)
+								.ToolTipText(LOCTEXT("BlendTooltip", "Blend between current pose and pose asset. Use Ctrl drag for under and over shoot."))
+								.MinValue(-2.0f)
 								.MaxValue(2.0f)
 								.MinSliderValue(-1.0f)
-								.MaxSliderValue(2.0f)
+								.MaxSliderValue(1.0f)
 								.SliderExponent(1)
 								.Delta(0.005f)
+								.MinDesiredWidth(100.0f)
+								.SupportDynamicSliderMinValue(true)
+								.SupportDynamicSliderMaxValue(true)
 								.OnValueChanged(this, &SControlRigPoseView::OnPoseBlendChanged)
 								.OnValueCommitted(this, &SControlRigPoseView::OnPoseBlendCommited)
 								.OnBeginSliderMovement(this,&SControlRigPoseView::OnBeginSliderMovement)
 								.OnEndSliderMovement(this,&SControlRigPoseView::OnEndSliderMovement)
 
 							]
+
+							+ SVerticalBox::Slot()
+							.AutoHeight()
+							.HAlign(HAlign_Center)
+							.Padding(15.f)
+							[
+								SNew(SButton)
+								.ContentPadding(FMargin(10, 5))
+								.Text(LOCTEXT("SelectControls", "Select Controls"))
+								.OnClicked(this, &SControlRigPoseView::OnSelectControls)
+							]
+							+ SVerticalBox::Slot()
+							.AutoHeight()
+							.HAlign(HAlign_Center)
+							.Padding(3.f)
+							[
+								SNew(SBorder)
+								.BorderImage(FEditorStyle::GetBrush("ToolPanel.GroupBorder"))
+								.Padding(FMargin(3.0f, 2.0f))
+								.Visibility(EVisibility::HitTestInvisible)
+								[
+									SAssignNew(TextStatusBlock1, STextBlock)
+								]
+							]
+							+ SVerticalBox::Slot()
+							.AutoHeight()
+							.HAlign(HAlign_Center)
+							.Padding(1.0f)
+								[
+									SNew(SBorder)
+									.BorderImage(FEditorStyle::GetBrush("ToolPanel.GroupBorder"))
+								.Padding(FMargin(3.0f, 0.0f))
+								.Visibility(EVisibility::HitTestInvisible)
+								[
+									SAssignNew(TextStatusBlock2, STextBlock)
+								]
+								]
 					]
 				]
 			+ SSplitter::Slot()
@@ -397,6 +437,28 @@ void SControlRigPoseView::Construct(const FArguments& InArgs)
 
 	];
 
+	if (FControlRigEditMode* EditMode = static_cast<FControlRigEditMode*>(GLevelEditorModeTools().GetActiveMode(FControlRigEditMode::ModeName)))
+	{
+		EditMode->OnControlRigAddedOrRemoved().AddRaw(this, &SControlRigPoseView::HandleControlAdded);
+		HandleControlAdded(GetControlRig(), true);
+	}
+}
+
+SControlRigPoseView::~SControlRigPoseView()
+{
+	if (FControlRigEditMode* EditMode = static_cast<FControlRigEditMode*>(GLevelEditorModeTools().GetActiveMode(FControlRigEditMode::ModeName)))
+	{
+		EditMode->OnControlRigAddedOrRemoved().RemoveAll(this);
+		if (EditMode->GetControlRig(true))
+		{
+			EditMode->GetControlRig(true)->ControlSelected().RemoveAll(this);
+		}
+	}
+}
+
+void SControlRigPoseView::NotifyPostChange(const FPropertyChangedEvent& PropertyChangedEvent, FProperty* PropertyThatChanged)
+{
+	UpdateStatusBlocks();
 }
 
 ECheckBoxState SControlRigPoseView::IsKeyPoseChecked() const
@@ -419,6 +481,7 @@ void SControlRigPoseView::OnKeyPoseChecked(ECheckBoxState NewState)
 		bIsKey = false;
 	}
 }
+
 ECheckBoxState SControlRigPoseView::IsMirrorPoseChecked() const
 {
 	if (bIsMirror)
@@ -438,6 +501,7 @@ void SControlRigPoseView::OnMirrorPoseChecked(ECheckBoxState NewState)
 	{
 		bIsMirror = false;
 	}
+	UpdateStatusBlocks();
 }
 
 bool SControlRigPoseView::IsMirrorEnabled() const
@@ -569,5 +633,88 @@ void SControlRigPoseView::CreateControlList()
 	}
 }
 */
+void SControlRigPoseView::HandleControlAdded(UControlRig* ControlRig, bool bIsAdded)
+{
+	if (ControlRig)
+	{
+		if (bIsAdded)
+		{
+			(ControlRig)->ControlSelected().RemoveAll(this);
+			(ControlRig)->ControlSelected().AddRaw(this, &SControlRigPoseView::HandleControlSelected);
+		}
+		else
+		{
+			(ControlRig)->ControlSelected().RemoveAll(this);
+		}
+	}
+	UpdateStatusBlocks();
+}
+
+void SControlRigPoseView::HandleControlSelected(UControlRig* Subject, const FRigControl& Control, bool bSelected)
+{
+	UpdateStatusBlocks();
+}
+
+void SControlRigPoseView::UpdateStatusBlocks()
+{
+	UControlRig* ControlRig = GetControlRig();
+	FText StatusText1;
+	FText StatusText2;
+
+	if (PoseAsset.IsValid() && ControlRig)
+	{
+		FFormatNamedArguments NamedArgs;
+		TArray<FName> ControlNames = PoseAsset->GetControlNames();
+		NamedArgs.Add("Total", ControlNames.Num());
+		TArray<FName> SelectedNames = ControlRig->CurrentControlSelection();
+		NamedArgs.Add("Selected", SelectedNames.Num());
+		uint32 Matching = 0;
+		uint32 MirrorMatching = 0;
+		
+		for (const FName& ControlName : ControlNames)
+		{
+			for (const FName& SelectedName : SelectedNames)
+			{
+				if (SelectedName == ControlName)
+				{
+					++Matching;
+					if (SControlRigPoseView::bIsMirror)
+					{
+						if (PoseAsset->DoesMirrorMatch(ControlRig,ControlName))
+						{
+							++MirrorMatching;
+						}
+					}
+				}
+			}
+		}
+		NamedArgs.Add("Matching", Matching);
+		NamedArgs.Add("MirrorMatching", MirrorMatching);
+
+		if(SControlRigPoseView::bIsMirror)
+		{
+			StatusText1 = FText::Format(LOCTEXT("NumberControlsAndMatch", "{Total} Controls Matching {Matching} of {Selected} Selected"), NamedArgs);
+			StatusText2 = FText::Format(LOCTEXT("NumberMirroredMatch", " {MirrorMatching} Mirror String Matches"), NamedArgs);
+		}
+		else
+		{
+			StatusText1 = FText::Format(LOCTEXT("NumberControlsAndMatch", "{Total} Controls Matching {Matching} of {Selected} Selected"), NamedArgs);
+			StatusText2 = FText::GetEmpty();
+		}
+	}
+	else
+	{
+		StatusText1 = FText::GetEmpty();
+		StatusText2 = FText::GetEmpty();
+	}
+	if (TextStatusBlock1.IsValid())
+	{
+		TextStatusBlock1->SetText(StatusText1);
+	}
+	if (TextStatusBlock2.IsValid())
+	{
+		TextStatusBlock2->SetText(StatusText2);
+	}
+}
 
 #undef LOCTEXT_NAMESPACE
diff --git a/Engine/Plugins/Experimental/ControlRig/Source/ControlRigEditor/Private/EditMode/SControlRigControlViews.h b/Engine/Plugins/Experimental/ControlRig/Source/ControlRigEditor/Private/EditMode/SControlRigControlViews.h
--- a/Engine/Plugins/Experimental/ControlRig/Source/ControlRigEditor/Private/EditMode/SControlRigControlViews.h
+++ b/Engine/Plugins/Experimental/ControlRig/Source/ControlRigEditor/Private/EditMode/SControlRigControlViews.h
@@ -15,6 +15,7 @@
 #include "Widgets/Text/STextBlock.h"
 #include "Widgets/Views/STableRow.h"
 #include "Tools/ControlRigPose.h"
+#include "Misc/NotifyHook.h"
 
 class FAssetThumbnail;
 class UControlRig;
@@ -26,19 +27,21 @@ public:
 	static void CaptureThumbnail(UObject* Asset);
 };
 
-class SControlRigPoseView : public SCompoundWidget
+class SControlRigPoseView : public SCompoundWidget, public FNotifyHook
 {
 	SLATE_BEGIN_ARGS(SControlRigPoseView) {}
 	SLATE_ARGUMENT(UControlRigPoseAsset*, PoseAsset)
 
 	SLATE_END_ARGS()
-	~SControlRigPoseView()
-	{
-	}
 
 	void Construct(const FArguments& InArgs);
+	~SControlRigPoseView();
+
+	//FNotifyHook
+	virtual void NotifyPostChange(const FPropertyChangedEvent& PropertyChangedEvent, FProperty* PropertyThatChanged) override;
 
 private:
+	void UpdateStatusBlocks();
 
 	UControlRig* GetControlRig();
 
@@ -54,10 +57,12 @@ private:
 	void OnPoseBlendCommited(float ChangedVal, ETextCommit::Type Type);
 	void OnBeginSliderMovement();
 	void OnEndSliderMovement(float NewValue);
-	TOptional<float> OnGetPoseBlendValue()const {return PoseBlendValue;}
+	float OnGetPoseBlendValue()const {return PoseBlendValue;}
 	FReply OnPastePose();
 	FReply OnSelectControls();
 	FReply OnCaptureThumbnail();
+	void HandleControlSelected(UControlRig* Subject, const FRigControl& Control, bool bSelected);
+	void HandleControlAdded(UControlRig* ControlRig, bool bIsAdded);
 
 	static bool bIsKey;
 	static bool bIsMirror;
@@ -70,6 +75,10 @@ private:
 
 	TWeakObjectPtr<UControlRigPoseAsset> PoseAsset;
 	
+	TSharedPtr<STextBlock> TextStatusBlock1;
+	TSharedPtr<STextBlock> TextStatusBlock2;
+
+
 	/* Mirroring*/
 	TSharedPtr<IDetailsView> MirrorDetailsView;
 
diff --git a/Engine/Plugins/Experimental/ControlRig/Source/ControlRigEditor/Private/EditMode/SControlRigEditModeTools.cpp b/Engine/Plugins/Experimental/ControlRig/Source/ControlRigEditor/Private/EditMode/SControlRigEditModeTools.cpp
--- a/Engine/Plugins/Experimental/ControlRig/Source/ControlRigEditor/Private/EditMode/SControlRigEditModeTools.cpp
+++ b/Engine/Plugins/Experimental/ControlRig/Source/ControlRigEditor/Private/EditMode/SControlRigEditModeTools.cpp
@@ -25,6 +25,9 @@
 #include "Rigs/FKControlRig.h"
 #include "SControlRigBaseListWidget.h"
 #include "Framework/MultiBox/MultiBoxBuilder.h"
+#include "SControlRigTweenWidget.h"
+#include "IControlRigEditorModule.h"
+#include "Framework/Docking/TabManager.h"
 
 #define LOCTEXT_NAMESPACE "ControlRigRootCustomization"
 
@@ -521,51 +524,39 @@ void SControlRigEditModeTools::CustomizeToolBarPalette(FToolBarBuilder& ToolBarB
 		NAME_None,
 		LOCTEXT("Poses", "Poses"),
 		LOCTEXT("PosesTooltip", "Show Poses"),
-		FSlateIcon(FEditorStyle::GetStyleSetName(), "FoliageEditMode.Settings"),
+		FSlateIcon(FEditorStyle::GetStyleSetName(), "FoliageEditMode.Settings"), //MZ todo replace with correct icon
 		EUserInterfaceActionType::Button
 	);
 	ToolBarBuilder.AddSeparator();
-}
 
+	// Tweens
+	ToolBarBuilder.AddToolBarButton(
+		FExecuteAction::CreateSP(this, &SControlRigEditModeTools::MakeTweenDialog),
+		NAME_None,
+		LOCTEXT("Tweens", "Tweens"),
+		LOCTEXT("TweensTooltip", "Create Tweens"),
+		FSlateIcon(FEditorStyle::GetStyleSetName(), "SkeletonTree.BlendProfile"), //MZ todo replace with correct icon
+		EUserInterfaceActionType::Button
+	);
+	ToolBarBuilder.AddSeparator();
+}
 
 void SControlRigEditModeTools::MakePoseDialog()
 {
-
 	FControlRigEditMode* ControlRigEditMode = static_cast<FControlRigEditMode*>(GLevelEditorModeTools().GetActiveMode(FControlRigEditMode::ModeName));
 	if (ControlRigEditMode)
 	{
+		FGlobalTabmanager::Get()->TryInvokeTab(IControlRigEditorModule::ControlRigPoseTab);
+	}
+}
 
-		TSharedPtr<SWindow> ExistingWindow = PoseWindow.Pin();
-		if (ExistingWindow.IsValid())
-		{
-			ExistingWindow->BringToFront();
-		}
-		else
-		{
-			ExistingWindow = SNew(SWindow)
-				.Title(LOCTEXT("Poses", "Poses"))
-				.HasCloseButton(true)
-				.SupportsMaximize(false)
-				.SupportsMinimize(false)
-				.ClientSize(FVector2D(850, 750));
-			TSharedPtr<SWindow> RootWindow = FGlobalTabmanager::Get()->GetRootWindow();
-			if (RootWindow.IsValid())
-			{
-				FSlateApplication::Get().AddWindowAsNativeChild(ExistingWindow.ToSharedRef(), RootWindow.ToSharedRef());
-			}
-			else
-			{
-				FSlateApplication::Get().AddWindow(ExistingWindow.ToSharedRef());
-			}
-
-		}
-
-		ExistingWindow->SetContent(
-			SNew(SControlRigBaseListWidget)
-		);
-		PoseWindow = ExistingWindow;
+void SControlRigEditModeTools::MakeTweenDialog()
+{
+	FControlRigEditMode* ControlRigEditMode = static_cast<FControlRigEditMode*>(GLevelEditorModeTools().GetActiveMode(FControlRigEditMode::ModeName));
+	if (ControlRigEditMode)
+	{
+		FGlobalTabmanager::Get()->TryInvokeTab(IControlRigEditorModule::ControlRigTweenTab);
 	}
-	
 }
 
 /* MZ TODO
diff --git a/Engine/Plugins/Experimental/ControlRig/Source/ControlRigEditor/Private/EditMode/SControlRigEditModeTools.h b/Engine/Plugins/Experimental/ControlRig/Source/ControlRigEditor/Private/EditMode/SControlRigEditModeTools.h
--- a/Engine/Plugins/Experimental/ControlRig/Source/ControlRigEditor/Private/EditMode/SControlRigEditModeTools.h
+++ b/Engine/Plugins/Experimental/ControlRig/Source/ControlRigEditor/Private/EditMode/SControlRigEditModeTools.h
@@ -80,8 +80,8 @@ private:
 	/** Make the Pose Dialg*/
 	void MakePoseDialog();
 
-	/** Pose Dialog Window*/
-	TWeakPtr<SWindow> PoseWindow;
+	/** Make the Pose Dialg*/
+	void MakeTweenDialog();
 
 	//TODO may put back void MakeSelectionSetDialog();
 	//TWeakPtr<SWindow> SelectionSetWindow;
@@ -90,8 +90,6 @@ public:
 	void CustomizeToolBarPalette(FToolBarBuilder& ToolBarBuilder);
 	FText GetActiveToolName() const;
 	FText GetActiveToolMessage() const;
-
-
 };
 
 
diff --git a/Engine/Plugins/Experimental/ControlRig/Source/ControlRigEditor/Private/EditMode/SControlRigTweenWidget.cpp b/Engine/Plugins/Experimental/ControlRig/Source/ControlRigEditor/Private/EditMode/SControlRigTweenWidget.cpp
new file mode 100644
--- /dev/null
+++ b/Engine/Plugins/Experimental/ControlRig/Source/ControlRigEditor/Private/EditMode/SControlRigTweenWidget.cpp
@@ -0,0 +1,135 @@
+// Copyright Epic Games, Inc. All Rights Reserved.
+
+#include "SControlRigTweenWidget.h"
+#include "Framework/Notifications/NotificationManager.h"
+#include "Widgets/Notifications/SNotificationList.h"
+#include "Framework/MultiBox/MultiBoxBuilder.h"
+#include "AssetData.h"
+#include "EditorStyleSet.h"
+#include "Widgets/Text/SInlineEditableTextBlock.h"
+#include "EditorStyleSet.h"
+#include "Styling/CoreStyle.h"
+#include "ScopedTransaction.h"
+#include "ControlRig.h"
+#include "UnrealEdGlobals.h"
+#include "ControlRigEditMode.h"
+#include "Tools/ControlRigPose.h"
+#include "EditorModeManager.h"
+#include "ISequencer.h"
+#include "LevelSequence.h"
+#include "LevelSequenceEditor/Private/LevelSequenceEditorBlueprintLibrary.h"
+#include "ILevelSequenceEditorToolkit.h"
+
+#define LOCTEXT_NAMESPACE "ControlRigTweenWidget"
+
+void SControlRigTweenWidget::Construct(const FArguments& InArgs)
+{
+	PoseBlendValue = 0.0f;
+	bIsBlending = false;
+	bSliderStartedTransaction = false;
+
+	ChildSlot
+	[
+		SNew(SVerticalBox)
+		 +SVerticalBox::Slot()
+			 .AutoHeight()
+			 .HAlign(HAlign_Center)
+			 .Padding(5.f)
+			 [
+				 SNew(SSpinBox<float>)
+				 .Value(this, &SControlRigTweenWidget::OnGetPoseBlendValueFloat)
+				  .ToolTipText(LOCTEXT("TweenTooltip", "Key at current frame between previous and next poses(keys). Use Ctrl drag for under and over shoot."))
+				 .MinValue(-2.0f)
+				 .MaxValue(2.0f)
+				 .MinSliderValue(-1.0f)
+				 .MaxSliderValue(1.0f)
+				 .SliderExponent(1)
+				 .Delta(0.005f)
+				 .MinDesiredWidth(100.0f)
+				 .SupportDynamicSliderMinValue(true)
+			     .SupportDynamicSliderMaxValue(true)
+				 .OnValueChanged(this, &SControlRigTweenWidget::OnPoseBlendChanged)
+				 .OnValueCommitted(this, &SControlRigTweenWidget::OnPoseBlendCommited)
+				 .OnBeginSliderMovement(this, &SControlRigTweenWidget::OnBeginSliderMovement)
+				 .OnEndSliderMovement(this, &SControlRigTweenWidget::OnEndSliderMovement)
+			 ]
+	];	
+}
+
+void SControlRigTweenWidget::OnPoseBlendChanged(float ChangedVal)
+{
+	UControlRig* ControlRig = GetControlRig();
+	if (ControlRig  && WeakSequencer.IsValid() && bIsBlending)
+	{
+		PoseBlendValue = ChangedVal;
+		ControlsToTween.Blend(WeakSequencer, ChangedVal);
+		WeakSequencer.Pin()->NotifyMovieSceneDataChanged(EMovieSceneDataChangeType::TrackValueChanged);
+	}
+}
+
+void SControlRigTweenWidget::OnBeginSliderMovement()
+{
+	if (bSliderStartedTransaction == false)
+	{
+		bIsBlending = true;
+		bSliderStartedTransaction = true;
+		GEditor->BeginTransaction(LOCTEXT("TweenTransaction", "Tween It"));
+		SetupControls();
+	}
+}
+
+void SControlRigTweenWidget::SetupControls()
+{
+	ULevelSequence* LevelSequence = ULevelSequenceEditorBlueprintLibrary::GetCurrentLevelSequence();
+	IAssetEditorInstance* AssetEditor = GEditor->GetEditorSubsystem<UAssetEditorSubsystem>()->FindEditorForAsset(LevelSequence, false);
+	ILevelSequenceEditorToolkit* LevelSequenceEditor = static_cast<ILevelSequenceEditorToolkit*>(AssetEditor);
+	WeakSequencer = LevelSequenceEditor ? LevelSequenceEditor->GetSequencer() : nullptr;
+	UControlRig* ControlRig = GetControlRig();
+	if (ControlRig && WeakSequencer.IsValid())
+	{
+		WeakSequencer.Pin()->GetFocusedMovieSceneSequence()->GetMovieScene()->Modify();
+		TArray<UControlRig*> SelectedControlRigs;
+		SelectedControlRigs.Add(ControlRig);
+		ControlsToTween.Setup(SelectedControlRigs, WeakSequencer);
+	}
+}
+
+void SControlRigTweenWidget::OnEndSliderMovement(float NewValue)
+{
+	if (bSliderStartedTransaction)
+	{
+		GEditor->EndTransaction();
+		bSliderStartedTransaction = false;
+
+	}
+	WeakSequencer = nullptr;
+}
+
+void SControlRigTweenWidget::OnPoseBlendCommited(float ChangedVal, ETextCommit::Type Type)
+{
+	UControlRig* ControlRig = GetControlRig();
+	if (ControlRig)
+	{
+		FScopedTransaction ScopedTransaction(LOCTEXT("TweenTransaction", "Tween"));
+		if (bIsBlending == false)
+		{
+			SetupControls();
+			bIsBlending = true;
+		}
+		PoseBlendValue = ChangedVal;
+		OnPoseBlendChanged(ChangedVal);
+		bIsBlending = false;
+		PoseBlendValue = 0.0f;
+	}
+}
+
+UControlRig* SControlRigTweenWidget::GetControlRig()
+{
+	if (FControlRigEditMode* EditMode = static_cast<FControlRigEditMode*>(GLevelEditorModeTools().GetActiveMode(FControlRigEditMode::ModeName)))
+	{
+		return EditMode->GetControlRig(true);
+	}
+	return nullptr;
+}
+
+#undef LOCTEXT_NAMESPACE
diff --git a/Engine/Plugins/Experimental/ControlRig/Source/ControlRigEditor/Private/EditMode/SControlRigTweenWidget.h b/Engine/Plugins/Experimental/ControlRig/Source/ControlRigEditor/Private/EditMode/SControlRigTweenWidget.h
new file mode 100644
--- /dev/null
+++ b/Engine/Plugins/Experimental/ControlRig/Source/ControlRigEditor/Private/EditMode/SControlRigTweenWidget.h
@@ -0,0 +1,56 @@
+// Copyright Epic Games, Inc. All Rights Reserved.
+/**
+* Hold the View for the Tween Widget
+*/
+#pragma once
+
+#include "CoreMinimal.h"
+#include "Widgets/SWidget.h"
+#include "Widgets/DeclarativeSyntaxSupport.h"
+#include "Widgets/SCompoundWidget.h"
+#include "Widgets/Input/SCheckBox.h"
+#include "Widgets/Input/SNumericEntryBox.h"
+#include "Widgets/Views/SListView.h"
+#include "Widgets/Text/STextBlock.h"
+#include "Widgets/Views/STableRow.h"
+#include "Tools/ControlRigPose.h"
+#include "Tools/ControlRigTweener.h"
+
+class UControlRig;
+class ISequencer;
+
+class SControlRigTweenWidget : public SCompoundWidget
+{
+	SLATE_BEGIN_ARGS(SControlRigTweenWidget) {}
+	SLATE_ARGUMENT(UControlRigPoseAsset*, PoseAsset)
+
+	SLATE_END_ARGS()
+	~SControlRigTweenWidget()
+	{
+	}
+
+	void Construct(const FArguments& InArgs);
+
+private:
+
+	UControlRig* GetControlRig();
+
+	/*
+	* Delegates and Helpers
+	*/
+	void OnPoseBlendChanged(float ChangedVal);
+	void OnPoseBlendCommited(float ChangedVal, ETextCommit::Type Type);
+	void OnBeginSliderMovement();
+	void OnEndSliderMovement(float NewValue);
+	float OnGetPoseBlendValueFloat() const { return PoseBlendValue; }
+
+	void SetupControls();
+	float PoseBlendValue;
+	bool bIsBlending;
+	bool bSliderStartedTransaction;
+
+	FControlsToTween  ControlsToTween;
+
+	TWeakPtr<ISequencer> WeakSequencer;
+};
+
diff --git a/Engine/Plugins/Experimental/ControlRig/Source/ControlRigEditor/Private/Tools/ControlRigTweener.cpp b/Engine/Plugins/Experimental/ControlRig/Source/ControlRigEditor/Private/Tools/ControlRigTweener.cpp
new file mode 100644
--- /dev/null
+++ b/Engine/Plugins/Experimental/ControlRig/Source/ControlRigEditor/Private/Tools/ControlRigTweener.cpp
@@ -0,0 +1,193 @@
+// Copyright Epic Games, Inc. All Rights Reserved.
+
+#include "Tools/ControlRigTweener.h"
+#include "Channels/MovieSceneFloatChannel.h"
+#include "ControlRig.h"
+#include "ISequencer.h"
+#include "MovieSceneSequence.h"
+#include "MovieScene.h"
+#include "Rigs/RigControlHierarchy.h"
+#include "Sequencer/MovieSceneControlRigParameterTrack.h"
+#include "Sequencer/MovieSceneControlRigParameterSection.h"
+#include "IKeyArea.h"
+#include "Channels/MovieSceneChannelProxy.h"
+
+//category not used so far, todo
+void FControlsToTween::Setup(const TArray<UControlRig*>& SelectedControlRigs, TWeakPtr<ISequencer>& InSequencer)
+{
+	if (InSequencer.IsValid() == false)
+	{
+		return;
+	}
+	ISequencer* Sequencer = InSequencer.Pin().Get();
+	ControlRigChannelsMap.Reset();
+
+	FQualifiedFrameTime CurrentTime = Sequencer->GetLocalTime();
+
+	UMovieScene* MovieScene = Sequencer->GetFocusedMovieSceneSequence()->GetMovieScene();
+
+	const TArray<FMovieSceneBinding>& Bindings = MovieScene->GetBindings();
+	for (const FMovieSceneBinding& Binding : Bindings)
+	{
+		UMovieSceneControlRigParameterTrack* Track = Cast<UMovieSceneControlRigParameterTrack>(MovieScene->FindTrack(UMovieSceneControlRigParameterTrack::StaticClass(), Binding.GetObjectGuid(), NAME_None));
+		if (Track && Track->GetControlRig() && SelectedControlRigs.Contains(Track->GetControlRig()))
+		{
+			for (UMovieSceneSection* MovieSection : Track->GetAllSections())
+			{
+				UMovieSceneControlRigParameterSection* Section = Cast<UMovieSceneControlRigParameterSection>(MovieSection);
+				if (Section && Section->IsActive() && Section->GetRange().Contains(CurrentTime.Time.GetFrame()))
+				{
+					Section->Modify();
+					UControlRig* ControlRig = Track->GetControlRig();
+					TArray<FRigControl> CurrentControls;
+					ControlRig->GetControlsInOrder(CurrentControls);
+					FRigControlHierarchy& Hierarchy = ControlRig->GetControlHierarchy();
+					FMovieSceneChannelProxy& Proxy = Section->GetChannelProxy();
+
+					const FName FloatChannelTypeName = FMovieSceneFloatChannel::StaticStruct()->GetFName();
+
+					FMovieSceneChannelProxy& ChannelProxy = Section->GetChannelProxy();
+					TArrayView<FMovieSceneFloatChannel*> Channels = ChannelProxy.GetChannels<FMovieSceneFloatChannel>();
+					//reuse these arrays
+					TArray<FFrameNumber> KeyTimes;
+					TArray<FKeyHandle> Handles;
+					for (const FRigControl& RigControl : CurrentControls)
+					{
+						if (RigControl.bAnimatable &&  ControlRig->IsControlSelected(RigControl.Name))
+						{
+							FControlRigChannels ControlRigChannels;
+							ControlRigChannels.NumChannels = 0;
+							ControlRigChannels.Section = Section;
+							KeyTimes.SetNum(0);
+							Handles.SetNum(0);
+							FChannelMapInfo* pChannelIndex = Section->ControlChannelMap.Find(RigControl.Name);
+							int NumChannels = 0;
+							switch (RigControl.ControlType)
+							{
+								case ERigControlType::Float:
+								{
+									NumChannels = 1;
+								}
+								case ERigControlType::Vector2D:
+								{
+									NumChannels = 2;
+									break;
+								}
+								case ERigControlType::Position:
+								case ERigControlType::Scale:
+								case ERigControlType::Rotator:
+								{
+									NumChannels = 3;
+									break;
+								}
+								case ERigControlType::Transform:
+								case ERigControlType::TransformNoScale:
+								case ERigControlType::EulerTransform:
+								{
+									NumChannels = 9;
+									break;
+								}
+							}
+							int32 BoundIndex = 0;
+							for (int ChannelIdx = pChannelIndex->ChannelIndex; ChannelIdx < (pChannelIndex->ChannelIndex + NumChannels); ++ChannelIdx)
+							{
+								FMovieSceneFloatChannel* Channel = Channels[ChannelIdx];
+								SetupControlRigChannel(CurrentTime.Time.GetFrame(), KeyTimes, Handles, Channel,
+									ControlRigChannels.KeyBounds[BoundIndex]);
+								if (ControlRigChannels.KeyBounds[BoundIndex].bValid)
+								{
+									++ControlRigChannels.NumChannels;
+								}
+								++BoundIndex;
+							}
+							if (ControlRigChannels.NumChannels > 0)
+							{
+								ControlRigChannelsMap.Add(RigControl.Name, ControlRigChannels);
+							}
+						}
+					}
+				}
+			}
+		}
+	}
+}
+
+void FControlsToTween::Blend(TWeakPtr<ISequencer>& InSequencer, float BlendValue)
+{
+	if (InSequencer.IsValid() == false)
+	{
+		return;
+	}
+	ISequencer* Sequencer = InSequencer.Pin().Get();
+	FFrameTime  FrameTime = Sequencer->GetLocalTime().Time;
+	float FrameAsFloat = (float)FrameTime.GetFrame().Value; 
+	Sequencer->GetFocusedMovieSceneSequence()->GetMovieScene()->Modify();
+	//BlendValue of 0.0 = half, 1.0f = Second, -1.0f = FirtValue;
+	//NormalizedBlendValue goes from 0 to 1.0f
+	for (const TPair<FName, FControlRigChannels>& Pair : ControlRigChannelsMap)
+	{
+		for (int Index = 0; Index < 9; ++Index)
+		{
+			float NormalizedBlendValue = (BlendValue + 1.0f) * 0.5f;
+
+			if (Pair.Value.KeyBounds[Index].bValid)
+			{
+				Pair.Value.Section->Modify();
+				if (Pair.Value.KeyBounds[Index].FirstFrame.Value != Pair.Value.KeyBounds[Index].SecondFrame.Value)
+				{
+					float FirstTime = (float)Pair.Value.KeyBounds[Index].FirstFrame.Value;
+					float SecondTime = (float)Pair.Value.KeyBounds[Index].SecondFrame.Value;
+					float TimeLocation = (FrameAsFloat - FirstTime) / (SecondTime - FirstTime);
+					//use TimeLocation to modify the NormalizedBlendValue 
+					if (NormalizedBlendValue > 0.5f)
+					{
+						NormalizedBlendValue = ((NormalizedBlendValue - .5f) / .5f) * (1.0f - TimeLocation) + TimeLocation;
+					}
+					else
+					{
+						NormalizedBlendValue = (NormalizedBlendValue / 0.5f) * TimeLocation;
+					}
+				}
+				float FirstValue = Pair.Value.KeyBounds[Index].FirstValue;
+				float SecondValue = Pair.Value.KeyBounds[Index].SecondValue;
+				float Value = FirstValue + (SecondValue - FirstValue) * (NormalizedBlendValue);
+				AddKeyToChannel(Pair.Value.KeyBounds[Index].Channel, FrameTime.GetFrame(), Value, Sequencer->GetKeyInterpolation());
+			}
+		}
+	}
+}
+
+void FControlsToTween::SetupControlRigChannel(FFrameNumber CurrentFrame, TArray<FFrameNumber>& KeyTimes, TArray<FKeyHandle>& Handles, FMovieSceneFloatChannel* FloatChannel,
+	FChannelKeyBounds& KeyBounds)
+{
+	KeyBounds.Channel = FloatChannel;
+	KeyBounds.FirstIndex = INDEX_NONE;
+	KeyBounds.SecondIndex = INDEX_NONE;
+	KeyTimes.SetNum(0);
+	Handles.SetNum(0);
+	FloatChannel->GetKeys(TRange<FFrameNumber>(), &KeyTimes, &Handles);
+	if (KeyTimes.Num() > 0)
+	{
+		TArrayView<const FMovieSceneFloatValue> Values = FloatChannel->GetValues();
+
+		for (int32 Index = 0; Index < KeyTimes.Num(); Index++)
+		{
+			FFrameNumber FrameNumber = KeyTimes[Index];
+			if (FrameNumber < CurrentFrame)
+			{
+				KeyBounds.FirstIndex = Index;
+				KeyBounds.FirstFrame = FrameNumber;
+				KeyBounds.FirstValue = Values[Index].Value;
+			}
+			else if (FrameNumber > CurrentFrame)
+			{
+				KeyBounds.SecondIndex = Index;
+				KeyBounds.SecondFrame = FrameNumber;
+				KeyBounds.SecondValue = Values[Index].Value;
+				break;
+			}
+		}
+	}
+	KeyBounds.bValid = (KeyBounds.FirstIndex != INDEX_NONE && KeyBounds.SecondIndex != INDEX_NONE
+		&& KeyBounds.FirstIndex != KeyBounds.SecondIndex) ? true : false;
+}
\ No newline at end of file
diff --git a/Engine/Plugins/Experimental/ControlRig/Source/ControlRigEditor/Public/IControlRigEditorModule.h b/Engine/Plugins/Experimental/ControlRig/Source/ControlRigEditor/Public/IControlRigEditorModule.h
--- a/Engine/Plugins/Experimental/ControlRig/Source/ControlRigEditor/Public/IControlRigEditorModule.h
+++ b/Engine/Plugins/Experimental/ControlRig/Source/ControlRigEditor/Public/IControlRigEditorModule.h
@@ -20,6 +20,7 @@ class FConnectionDrawingPolicy;
 class IControlRigEditorModule : public IModuleInterface, public IHasMenuExtensibility, public IHasToolBarExtensibility
 {
 public:
+
 	static FORCEINLINE IControlRigEditorModule& Get()
 	{
 		return FModuleManager::LoadModuleChecked< IControlRigEditorModule >(TEXT("ControlRigEditor"));
@@ -44,4 +45,7 @@ public:
 	virtual void GetInstanceActions(UControlRigBlueprint* CRB, FBlueprintActionDatabaseRegistrar& ActionRegistrar) = 0;
 	virtual FConnectionDrawingPolicy* CreateConnectionDrawingPolicy(int32 InBackLayerID, int32 InFrontLayerID, float InZoomFactor, const FSlateRect& InClippingRect, class FSlateWindowElementList& InDrawElements, class UEdGraph* InGraphObj) = 0;
 	virtual void GetContextMenuActions(const UControlRigGraphSchema* Schema, class UToolMenu* Menu, class UGraphNodeContextMenuContext* Context) const = 0;
+
+	static const FName ControlRigPoseTab;
+	static const FName ControlRigTweenTab;
 };
diff --git a/Engine/Plugins/Experimental/ControlRig/Source/ControlRigEditor/Public/Tools/ControlRigTweener.h b/Engine/Plugins/Experimental/ControlRig/Source/ControlRigEditor/Public/Tools/ControlRigTweener.h
new file mode 100644
--- /dev/null
+++ b/Engine/Plugins/Experimental/ControlRig/Source/ControlRigEditor/Public/Tools/ControlRigTweener.h
@@ -0,0 +1,64 @@
+// Copyright Epic Games, Inc. All Rights Reserved.
+
+#pragma once
+#include "CoreMinimal.h"
+#include "UObject/ObjectMacros.h"
+#include "Misc/FrameNumber.h"
+
+struct FMovieSceneFloatChannel;
+class UControlRig;
+class ISequencer;
+struct FRigControl;
+class UMovieSceneSection;
+
+
+//first and second key to blend between for each channel
+struct FChannelKeyBounds
+{
+	FChannelKeyBounds() : bValid(false), Channel(nullptr), FirstIndex(INDEX_NONE), SecondIndex(INDEX_NONE), FirstFrame(0), SecondFrame(0), FirstValue(0.0f), SecondValue(0.0f) {}
+	bool bValid;
+	FMovieSceneFloatChannel* Channel;
+	int32 FirstIndex;
+	int32 SecondIndex;
+	FFrameNumber FirstFrame;
+	FFrameNumber SecondFrame;
+	float FirstValue;
+	float SecondValue;
+};
+
+struct FControlRigChannels
+{
+	FControlRigChannels() : Section(nullptr), NumChannels(0) {};
+	FChannelKeyBounds  KeyBounds[9];
+	UMovieSceneSection* Section;
+	int NumChannels; //may not use all 9 channels.
+};
+
+struct FControlsToTween
+{
+	/**
+	 Setup up controls to tween
+	* @param InControlRig Control Rig to use
+	* @param InSequencer Sequencer to get keys from
+	*/
+	void Setup(const TArray<UControlRig*>& SelectedControlRigs, TWeakPtr<ISequencer>& InSequencer);
+
+	/**
+	* @param InSequencer Sequencer to blend at current time
+	* @param InSequencer BlendValue where in time to Blend in at Current Frame, 0.0 will be at current frame, -1.0 will value at previous key, 1.0 value at next key,
+	* other values interpolated between.
+	*/
+	void Blend(TWeakPtr<ISequencer>& InSequencer,  float BlendValue);
+
+private:
+	void SetupControlRigChannel(FFrameNumber CurrentFrame, TArray<FFrameNumber>& KeyTimes, TArray<FKeyHandle>& Handles, FMovieSceneFloatChannel* FloatChannel,
+		FChannelKeyBounds& KeyBounds);
+
+	TMap<FName, FControlRigChannels>  ControlRigChannelsMap;
+};
+
+
+
+
+
+
