Control Rig:  When getting control rig items we need to get control hierachy first so setting global transforms on multiple controls works correctly.

diff --git a/Engine/Plugins/Experimental/ControlRig/Source/ControlRig/Private/Rigs/RigHierarchyContainer.cpp b/Engine/Plugins/Experimental/ControlRig/Source/ControlRig/Private/Rigs/RigHierarchyContainer.cpp
--- a/Engine/Plugins/Experimental/ControlRig/Source/ControlRig/Private/Rigs/RigHierarchyContainer.cpp
+++ b/Engine/Plugins/Experimental/ControlRig/Source/ControlRig/Private/Rigs/RigHierarchyContainer.cpp
@@ -1289,14 +1289,16 @@ TArray<FRigElementKey> FRigHierarchyContainer::CurrentSelection() const
 TArray<FRigElementKey> FRigHierarchyContainer::GetAllItems(bool bSort) const
 {
 	TArray<FRigElementKey> Items;
-	for (const FRigBone& Element : BoneHierarchy)
+
+	for (const FRigControl& Element : ControlHierarchy)
 	{
 		Items.Add(Element.GetElementKey());
 	}
-	for (const FRigControl& Element : ControlHierarchy)
+	for (const FRigBone& Element : BoneHierarchy)
 	{
 		Items.Add(Element.GetElementKey());
 	}
+	
 	for (const FRigSpace& Element : SpaceHierarchy)
 	{
 		Items.Add(Element.GetElementKey());
@@ -1531,4 +1533,5 @@ bool FRigHierarchyContainer::IsNameAvailable(const FString& InName, ERigElementT
 	}
 
 	return true;
-}
\ No newline at end of file
+}
+
