Sequencer: Save settings for sequencer asset pickers

diff --git a/Engine/Plugins/Experimental/ControlRig/Source/ControlRigEditor/Private/Sequencer/ControlRigParameterTrackEditor.cpp b/Engine/Plugins/Experimental/ControlRig/Source/ControlRigEditor/Private/Sequencer/ControlRigParameterTrackEditor.cpp
--- a/Engine/Plugins/Experimental/ControlRig/Source/ControlRigEditor/Private/Sequencer/ControlRigParameterTrackEditor.cpp
+++ b/Engine/Plugins/Experimental/ControlRig/Source/ControlRigEditor/Private/Sequencer/ControlRigParameterTrackEditor.cpp
@@ -833,6 +833,7 @@ void FControlRigParameterTrackEditor::AddControlRigSubMenu(FMenuBuilder& MenuBui
 	AssetPickerConfig.InitialAssetViewType = EAssetViewType::List;
 	AssetPickerConfig.Filter.bRecursiveClasses = true;
 	AssetPickerConfig.Filter.ClassNames.Add(UControlRigSequence::StaticClass()->GetFName());
+	AssetPickerConfig.SaveSettingsName = TEXT("SequencerAssetPicker");
 	}
 
 	FContentBrowserModule& ContentBrowserModule = FModuleManager::Get().LoadModuleChecked<FContentBrowserModule>(TEXT("ContentBrowser"));
@@ -889,6 +890,8 @@ void FControlRigParameterTrackEditor::AddControlRigSubMenu(FMenuBuilder& MenuBui
 		AssetPickerConfig.Filter.bRecursiveClasses = true;
 		AssetPickerConfig.Filter.ClassNames.Add((UControlRig::StaticClass())->GetFName());
 		AssetPickerConfig.Filter.TagsAndValues.Add(TEXT("Skeleton"), FAssetData(Skeleton).GetExportTextName());
+		AssetPickerConfig.SaveSettingsName = TEXT("SequencerAssetPicker");
+
 		}
 
 		FContentBrowserModule& ContentBrowserModule = FModuleManager::Get().LoadModuleChecked<FContentBrowserModule>(TEXT("ContentBrowser"));
@@ -2812,6 +2815,7 @@ void FControlRigParameterSection::AddAnimationSubMenuForFK(FMenuBuilder& MenuBui
 		AssetPickerConfig.Filter.bRecursiveClasses = true;
 		AssetPickerConfig.Filter.ClassNames.Add(UAnimSequenceBase::StaticClass()->GetFName());
 		AssetPickerConfig.Filter.TagsAndValues.Add(TEXT("Skeleton"), FAssetData(Skeleton).GetExportTextName());
+		AssetPickerConfig.SaveSettingsName = TEXT("SequencerAssetPicker");
 	}
 
 	FContentBrowserModule& ContentBrowserModule = FModuleManager::Get().LoadModuleChecked<FContentBrowserModule>(TEXT("ContentBrowser"));
diff --git a/Engine/Plugins/Experimental/ImagePlate/Source/ImagePlateEditor/Private/ImagePlateTrackEditor.cpp b/Engine/Plugins/Experimental/ImagePlate/Source/ImagePlateEditor/Private/ImagePlateTrackEditor.cpp
--- a/Engine/Plugins/Experimental/ImagePlate/Source/ImagePlateEditor/Private/ImagePlateTrackEditor.cpp
+++ b/Engine/Plugins/Experimental/ImagePlate/Source/ImagePlateEditor/Private/ImagePlateTrackEditor.cpp
@@ -373,6 +373,7 @@ TSharedPtr<SWidget> FImagePlateTrackEditor::BuildOutlinerEditWidget(const FGuid&
 
 		AssetPickerConfig.Filter.bRecursiveClasses = true;
 		AssetPickerConfig.Filter.ClassNames.Add(UImagePlateFileSequence::StaticClass()->GetFName());
+		AssetPickerConfig.SaveSettingsName = TEXT("SequencerAssetPicker");
 
 		FContentBrowserModule& ContentBrowserModule = FModuleManager::Get().LoadModuleChecked<FContentBrowserModule>(TEXT("ContentBrowser"));
 
diff --git a/Engine/Plugins/Media/MediaCompositing/Source/MediaCompositingEditor/Private/Sequencer/MediaPlayerPropertyTrackEditor.cpp b/Engine/Plugins/Media/MediaCompositing/Source/MediaCompositingEditor/Private/Sequencer/MediaPlayerPropertyTrackEditor.cpp
--- a/Engine/Plugins/Media/MediaCompositing/Source/MediaCompositingEditor/Private/Sequencer/MediaPlayerPropertyTrackEditor.cpp
+++ b/Engine/Plugins/Media/MediaCompositing/Source/MediaCompositingEditor/Private/Sequencer/MediaPlayerPropertyTrackEditor.cpp
@@ -155,6 +155,7 @@ TSharedPtr<SWidget> FMediaPlayerPropertyTrackEditor::BuildOutlinerEditWidget(con
 			AssetPickerConfig.InitialAssetViewType = EAssetViewType::List;
 			AssetPickerConfig.Filter.bRecursiveClasses = true;
 			AssetPickerConfig.Filter.ClassNames.Add(UMediaSource::StaticClass()->GetFName());
+			AssetPickerConfig.SaveSettingsName = TEXT("SequencerAssetPicker");
 		}
 
 		FContentBrowserModule& ContentBrowserModule = FModuleManager::Get().LoadModuleChecked<FContentBrowserModule>(TEXT("ContentBrowser"));
diff --git a/Engine/Plugins/Media/MediaCompositing/Source/MediaCompositingEditor/Private/Sequencer/MediaTrackEditor.cpp b/Engine/Plugins/Media/MediaCompositing/Source/MediaCompositingEditor/Private/Sequencer/MediaTrackEditor.cpp
--- a/Engine/Plugins/Media/MediaCompositing/Source/MediaCompositingEditor/Private/Sequencer/MediaTrackEditor.cpp
+++ b/Engine/Plugins/Media/MediaCompositing/Source/MediaCompositingEditor/Private/Sequencer/MediaTrackEditor.cpp
@@ -96,6 +96,7 @@ TSharedPtr<SWidget> FMediaTrackEditor::BuildOutlinerEditWidget(const FGuid& Obje
 			AssetPickerConfig.InitialAssetViewType = EAssetViewType::List;
 			AssetPickerConfig.Filter.bRecursiveClasses = true;
 			AssetPickerConfig.Filter.ClassNames.Add(UMediaSource::StaticClass()->GetFName());
+			AssetPickerConfig.SaveSettingsName = TEXT("SequencerAssetPicker");
 		}
 
 		FContentBrowserModule& ContentBrowserModule = FModuleManager::Get().LoadModuleChecked<FContentBrowserModule>(TEXT("ContentBrowser"));
diff --git a/Engine/Plugins/MovieScene/TemplateSequence/Source/TemplateSequenceEditor/Private/TrackEditors/TemplateSequenceTrackEditor.cpp b/Engine/Plugins/MovieScene/TemplateSequence/Source/TemplateSequenceEditor/Private/TrackEditors/TemplateSequenceTrackEditor.cpp
--- a/Engine/Plugins/MovieScene/TemplateSequence/Source/TemplateSequenceEditor/Private/TrackEditors/TemplateSequenceTrackEditor.cpp
+++ b/Engine/Plugins/MovieScene/TemplateSequence/Source/TemplateSequenceEditor/Private/TrackEditors/TemplateSequenceTrackEditor.cpp
@@ -143,6 +143,7 @@ public:
 			AssetPickerConfig.InitialAssetViewType = EAssetViewType::List;
 			AssetPickerConfig.Filter.bRecursiveClasses = true;
 			AssetPickerConfig.Filter.ClassNames.Add(UTemplateSequence::StaticClass()->GetFName());
+			AssetPickerConfig.SaveSettingsName = TEXT("SequencerAssetPicker");
 			if (LegacyBaseClass != nullptr)
 			{
 				AssetPickerConfig.Filter.ClassNames.Add(LegacyBaseClass->GetFName());
diff --git a/Engine/Plugins/VirtualProduction/DMX/DMXEngine/Source/DMXEditor/Private/Sequencer/DMXLibraryTrackEditor.cpp b/Engine/Plugins/VirtualProduction/DMX/DMXEngine/Source/DMXEditor/Private/Sequencer/DMXLibraryTrackEditor.cpp
--- a/Engine/Plugins/VirtualProduction/DMX/DMXEngine/Source/DMXEditor/Private/Sequencer/DMXLibraryTrackEditor.cpp
+++ b/Engine/Plugins/VirtualProduction/DMX/DMXEngine/Source/DMXEditor/Private/Sequencer/DMXLibraryTrackEditor.cpp
@@ -47,6 +47,7 @@ TSharedRef<SWidget> CreateAssetPicker(FOnAssetSelected OnAssetSelected, FOnAsset
 		AssetPickerConfig.InitialAssetViewType = EAssetViewType::List;
 		AssetPickerConfig.Filter.bRecursiveClasses = true;
 		AssetPickerConfig.Filter.ClassNames.Add(UDMXLibrary::StaticClass()->GetFName());
+		AssetPickerConfig.SaveSettingsName = TEXT("SequencerAssetPicker");
 	}
 
 	FContentBrowserModule& ContentBrowserModule = FModuleManager::Get().LoadModuleChecked<FContentBrowserModule>(TEXT("ContentBrowser"));
diff --git a/Engine/Source/Editor/MovieSceneTools/Private/TrackEditors/AudioTrackEditor.cpp b/Engine/Source/Editor/MovieSceneTools/Private/TrackEditors/AudioTrackEditor.cpp
--- a/Engine/Source/Editor/MovieSceneTools/Private/TrackEditors/AudioTrackEditor.cpp
+++ b/Engine/Source/Editor/MovieSceneTools/Private/TrackEditors/AudioTrackEditor.cpp
@@ -1253,6 +1253,7 @@ TSharedRef<SWidget> FAudioTrackEditor::BuildAudioSubMenu(FOnAssetSelected OnAsse
 		{
 			AssetPickerConfig.Filter.ClassNames.Add(ClassName);
 		}
+		AssetPickerConfig.SaveSettingsName = TEXT("SequencerAssetPicker");
 	}
 
 	FContentBrowserModule& ContentBrowserModule = FModuleManager::Get().LoadModuleChecked<FContentBrowserModule>(TEXT("ContentBrowser"));
diff --git a/Engine/Source/Editor/MovieSceneTools/Private/TrackEditors/CameraAnimTrackEditor.cpp b/Engine/Source/Editor/MovieSceneTools/Private/TrackEditors/CameraAnimTrackEditor.cpp
--- a/Engine/Source/Editor/MovieSceneTools/Private/TrackEditors/CameraAnimTrackEditor.cpp
+++ b/Engine/Source/Editor/MovieSceneTools/Private/TrackEditors/CameraAnimTrackEditor.cpp
@@ -150,6 +150,7 @@ void FCameraAnimTrackEditor::AddCameraAnimSubMenu(FMenuBuilder& MenuBuilder, TAr
 		AssetPickerConfig.InitialAssetViewType = EAssetViewType::List;
 		AssetPickerConfig.Filter.ClassNames.Add(UCameraAnim::StaticClass()->GetFName());
 //		AssetPickerConfig.Filter.TagsAndValues.Add(TEXT("Skeleton"), FAssetData(Skeleton).GetExportTextName());
+		AssetPickerConfig.SaveSettingsName = TEXT("SequencerAssetPicker");
 	}
 
 	FContentBrowserModule& ContentBrowserModule = FModuleManager::Get().LoadModuleChecked<FContentBrowserModule>(TEXT("ContentBrowser"));
diff --git a/Engine/Source/Editor/MovieSceneTools/Private/TrackEditors/CameraShakeSourceShakeTrackEditor.cpp b/Engine/Source/Editor/MovieSceneTools/Private/TrackEditors/CameraShakeSourceShakeTrackEditor.cpp
--- a/Engine/Source/Editor/MovieSceneTools/Private/TrackEditors/CameraShakeSourceShakeTrackEditor.cpp
+++ b/Engine/Source/Editor/MovieSceneTools/Private/TrackEditors/CameraShakeSourceShakeTrackEditor.cpp
@@ -638,6 +638,7 @@ void FCameraShakeSourceShakeTrackEditor::AddOtherCameraShakeBrowserSubMenu(FMenu
 		AssetPickerConfig.bAllowNullSelection = false;
 		AssetPickerConfig.InitialAssetViewType = EAssetViewType::List;
 		AssetPickerConfig.Filter.ClassNames.Add(UBlueprint::StaticClass()->GetFName());
+		AssetPickerConfig.SaveSettingsName = TEXT("SequencerAssetPicker");
 
 		IAssetRegistry & AssetRegistry = FModuleManager::LoadModuleChecked<FAssetRegistryModule>("AssetRegistry").Get();
 		TArray<FName> ClassNames;
diff --git a/Engine/Source/Editor/MovieSceneTools/Private/TrackEditors/CameraShakeTrackEditor.cpp b/Engine/Source/Editor/MovieSceneTools/Private/TrackEditors/CameraShakeTrackEditor.cpp
--- a/Engine/Source/Editor/MovieSceneTools/Private/TrackEditors/CameraShakeTrackEditor.cpp
+++ b/Engine/Source/Editor/MovieSceneTools/Private/TrackEditors/CameraShakeTrackEditor.cpp
@@ -130,6 +130,7 @@ void FCameraShakeTrackEditor::AddCameraShakeSubMenu(FMenuBuilder& MenuBuilder, T
 		AssetPickerConfig.bAllowNullSelection = false;
 		AssetPickerConfig.InitialAssetViewType = EAssetViewType::List;
 		AssetPickerConfig.Filter.ClassNames.Add(UBlueprint::StaticClass()->GetFName());
+		AssetPickerConfig.SaveSettingsName = TEXT("SequencerAssetPicker");
 
 		IAssetRegistry & AssetRegistry = FModuleManager::LoadModuleChecked<FAssetRegistryModule>("AssetRegistry").Get();
 		TArray<FName> ClassNames;
diff --git a/Engine/Source/Editor/MovieSceneTools/Private/TrackEditors/CinematicShotTrackEditor.cpp b/Engine/Source/Editor/MovieSceneTools/Private/TrackEditors/CinematicShotTrackEditor.cpp
--- a/Engine/Source/Editor/MovieSceneTools/Private/TrackEditors/CinematicShotTrackEditor.cpp
+++ b/Engine/Source/Editor/MovieSceneTools/Private/TrackEditors/CinematicShotTrackEditor.cpp
@@ -634,6 +634,7 @@ TSharedRef<SWidget> FCinematicShotTrackEditor::HandleAddCinematicShotComboButton
 		AssetPickerConfig.bAllowNullSelection = false;
 		AssetPickerConfig.InitialAssetViewType = EAssetViewType::Tile;
 		AssetPickerConfig.Filter.ClassNames.Add(TEXT("LevelSequence"));
+		AssetPickerConfig.SaveSettingsName = TEXT("SequencerAssetPicker");
 	}
 
 	FContentBrowserModule& ContentBrowserModule = FModuleManager::Get().LoadModuleChecked<FContentBrowserModule>(TEXT("ContentBrowser"));
diff --git a/Engine/Source/Editor/MovieSceneTools/Private/TrackEditors/MaterialParameterCollectionTrackEditor.cpp b/Engine/Source/Editor/MovieSceneTools/Private/TrackEditors/MaterialParameterCollectionTrackEditor.cpp
--- a/Engine/Source/Editor/MovieSceneTools/Private/TrackEditors/MaterialParameterCollectionTrackEditor.cpp
+++ b/Engine/Source/Editor/MovieSceneTools/Private/TrackEditors/MaterialParameterCollectionTrackEditor.cpp
@@ -45,6 +45,7 @@ TSharedRef<SWidget> CreateAssetPicker(FOnAssetSelected OnAssetSelected, FOnAsset
 		AssetPickerConfig.InitialAssetViewType = EAssetViewType::List;
 		AssetPickerConfig.Filter.bRecursiveClasses = true;
 		AssetPickerConfig.Filter.ClassNames.Add(UMaterialParameterCollection::StaticClass()->GetFName());
+		AssetPickerConfig.SaveSettingsName = TEXT("SequencerAssetPicker");
 	}
 
 	FContentBrowserModule& ContentBrowserModule = FModuleManager::Get().LoadModuleChecked<FContentBrowserModule>(TEXT("ContentBrowser"));
diff --git a/Engine/Source/Editor/MovieSceneTools/Private/TrackEditors/SkeletalAnimationTrackEditor.cpp b/Engine/Source/Editor/MovieSceneTools/Private/TrackEditors/SkeletalAnimationTrackEditor.cpp
--- a/Engine/Source/Editor/MovieSceneTools/Private/TrackEditors/SkeletalAnimationTrackEditor.cpp
+++ b/Engine/Source/Editor/MovieSceneTools/Private/TrackEditors/SkeletalAnimationTrackEditor.cpp
@@ -1531,6 +1531,7 @@ void FSkeletalAnimationTrackEditor::AddAnimationSubMenu(FMenuBuilder& MenuBuilde
 		AssetPickerConfig.Filter.bRecursiveClasses = true;
 		AssetPickerConfig.Filter.ClassNames.Add(UAnimSequenceBase::StaticClass()->GetFName());
 		AssetPickerConfig.Filter.TagsAndValues.Add(TEXT("Skeleton"), FAssetData(Skeleton).GetExportTextName());
+		AssetPickerConfig.SaveSettingsName = TEXT("SequencerAssetPicker");
 	}
 
 	FContentBrowserModule& ContentBrowserModule = FModuleManager::Get().LoadModuleChecked<FContentBrowserModule>(TEXT("ContentBrowser"));
diff --git a/Engine/Source/Editor/MovieSceneTools/Private/TrackEditors/SubTrackEditor.cpp b/Engine/Source/Editor/MovieSceneTools/Private/TrackEditors/SubTrackEditor.cpp
--- a/Engine/Source/Editor/MovieSceneTools/Private/TrackEditors/SubTrackEditor.cpp
+++ b/Engine/Source/Editor/MovieSceneTools/Private/TrackEditors/SubTrackEditor.cpp
@@ -580,6 +580,7 @@ TSharedRef<SWidget> FSubTrackEditor::HandleAddSubSequenceComboButtonGetMenuConte
 			AssetPickerConfig.bAllowNullSelection = false;
 			AssetPickerConfig.InitialAssetViewType = EAssetViewType::Tile;
 			AssetPickerConfig.Filter.ClassNames.Add(TEXT("LevelSequence"));
+			AssetPickerConfig.SaveSettingsName = TEXT("SequencerAssetPicker");
 		}
 
 		FContentBrowserModule& ContentBrowserModule = FModuleManager::Get().LoadModuleChecked<FContentBrowserModule>(TEXT("ContentBrowser"));
diff --git a/Engine/Source/Editor/MovieSceneTools/Private/TrackEditors/TransformTrackEditor.cpp b/Engine/Source/Editor/MovieSceneTools/Private/TrackEditors/TransformTrackEditor.cpp
--- a/Engine/Source/Editor/MovieSceneTools/Private/TrackEditors/TransformTrackEditor.cpp
+++ b/Engine/Source/Editor/MovieSceneTools/Private/TrackEditors/TransformTrackEditor.cpp
@@ -163,6 +163,7 @@ void F3DTransformTrackEditor::BuildTrackContextMenu( FMenuBuilder& MenuBuilder,
 		AssetPickerConfig.Filter.ClassNames.Add(UAnimSequence::StaticClass()->GetFName());
 		AssetPickerConfig.OnAssetSelected = FOnAssetSelected::CreateStatic(&F3DTransformTrackEditor::ImportAnimSequenceTransforms, InSequencer, InTransformTrack);
 		AssetPickerConfig.OnAssetEnterPressed = FOnAssetEnterPressed::CreateStatic(&F3DTransformTrackEditor::ImportAnimSequenceTransformsEnterPressed, InSequencer, InTransformTrack);
+		AssetPickerConfig.SaveSettingsName = TEXT("SequencerAssetPicker");
 
 		FContentBrowserModule& ContentBrowserModule = FModuleManager::LoadModuleChecked<FContentBrowserModule>(TEXT("ContentBrowser"));
 
