ControlRig: fix acceleration computation in verlet integration node

diff --git a/Engine/Plugins/Experimental/ControlRig/Source/ControlRig/Private/Units/Simulation/RigUnit_Verlet.cpp b/Engine/Plugins/Experimental/ControlRig/Source/ControlRig/Private/Units/Simulation/RigUnit_Verlet.cpp
--- a/Engine/Plugins/Experimental/ControlRig/Source/ControlRig/Private/Units/Simulation/RigUnit_Verlet.cpp
+++ b/Engine/Plugins/Experimental/ControlRig/Source/ControlRig/Private/Units/Simulation/RigUnit_Verlet.cpp
@@ -22,11 +22,14 @@ FRigUnit_VerletIntegrateVector_Execute()
 	}
 
 	Point.LinearDamping = Damp;
-	float U = FMath::Clamp<float>(Blend * Context.DeltaTime, 0.f, 1.f);
-	FVector Force = (Target - Point.Position) * FMath::Max(Strength, 0.0001f);
-	FVector PreviousVelocity = Point.LinearVelocity;
-	Point = Point.IntegrateVerlet(Force, Blend, Context.DeltaTime);
-	Acceleration = Point.LinearVelocity - PreviousVelocity;
-	Position = Point.Position;
-	Velocity = Point.LinearVelocity;
+	if (Context.DeltaTime > SMALL_NUMBER)
+	{
+		float U = FMath::Clamp<float>(Blend * Context.DeltaTime, 0.f, 1.f);
+		FVector Force = (Target - Point.Position) * FMath::Max(Strength, 0.0001f);
+		FVector PreviousVelocity = Point.LinearVelocity;
+		Point = Point.IntegrateVerlet(Force, Blend, Context.DeltaTime);
+		Acceleration = (Point.LinearVelocity - PreviousVelocity) / Context.DeltaTime;
+		Position = Point.Position;
+		Velocity = Point.LinearVelocity;
+	}
 }
