CIS Non-Unity Fix
#trivial

diff --git a/Engine/Plugins/Experimental/ControlRig/Source/ControlRigEditor/Public/Tools/ControlRigTweener.h b/Engine/Plugins/Experimental/ControlRig/Source/ControlRigEditor/Public/Tools/ControlRigTweener.h
--- a/Engine/Plugins/Experimental/ControlRig/Source/ControlRigEditor/Public/Tools/ControlRigTweener.h
+++ b/Engine/Plugins/Experimental/ControlRig/Source/ControlRigEditor/Public/Tools/ControlRigTweener.h
@@ -4,6 +4,7 @@
 #include "CoreMinimal.h"
 #include "UObject/ObjectMacros.h"
 #include "Misc/FrameNumber.h"
+#include "Curves/KeyHandle.h"
 
 struct FMovieSceneFloatChannel;
 class UControlRig;
