Control Rig: Padding/tooltip fix for Tween tool.
#trivial

diff --git a/Engine/Plugins/Experimental/ControlRig/Source/ControlRigEditor/Private/EditMode/SControlRigTweenWidget.cpp b/Engine/Plugins/Experimental/ControlRig/Source/ControlRigEditor/Private/EditMode/SControlRigTweenWidget.cpp
--- a/Engine/Plugins/Experimental/ControlRig/Source/ControlRigEditor/Private/EditMode/SControlRigTweenWidget.cpp
+++ b/Engine/Plugins/Experimental/ControlRig/Source/ControlRigEditor/Private/EditMode/SControlRigTweenWidget.cpp
@@ -30,29 +30,34 @@ void SControlRigTweenWidget::Construct(const FArguments& InArgs)
 
 	ChildSlot
 	[
-		SNew(SVerticalBox)
-		 +SVerticalBox::Slot()
-			 .AutoHeight()
-			 .HAlign(HAlign_Center)
-			 .Padding(5.f)
-			 [
-				 SNew(SSpinBox<float>)
-				 .Value(this, &SControlRigTweenWidget::OnGetPoseBlendValueFloat)
-				  .ToolTipText(LOCTEXT("TweenTooltip", "Key at current frame between previous and next poses(keys). Use Ctrl drag for under and over shoot."))
-				 .MinValue(-2.0f)
-				 .MaxValue(2.0f)
-				 .MinSliderValue(-1.0f)
-				 .MaxSliderValue(1.0f)
-				 .SliderExponent(1)
-				 .Delta(0.005f)
-				 .MinDesiredWidth(100.0f)
-				 .SupportDynamicSliderMinValue(true)
-			     .SupportDynamicSliderMaxValue(true)
-				 .OnValueChanged(this, &SControlRigTweenWidget::OnPoseBlendChanged)
-				 .OnValueCommitted(this, &SControlRigTweenWidget::OnPoseBlendCommited)
-				 .OnBeginSliderMovement(this, &SControlRigTweenWidget::OnBeginSliderMovement)
-				 .OnEndSliderMovement(this, &SControlRigTweenWidget::OnEndSliderMovement)
-			 ]
+		SNew(SBorder)
+		.HAlign(HAlign_Fill)
+		.VAlign(VAlign_Fill)
+		.Padding(FMargin(20.0f))
+		[
+			SNew(SVerticalBox)
+			 +SVerticalBox::Slot()
+				 .AutoHeight()
+				 .HAlign(HAlign_Center)
+				 [
+					 SNew(SSpinBox<float>)
+					 .Value(this, &SControlRigTweenWidget::OnGetPoseBlendValueFloat)
+					  .ToolTipText(LOCTEXT("TweenTooltip", "Key at current frame between previous(-1.0) and next(1.0) poses. Use Ctrl drag for under and over shoot."))
+					 .MinValue(-2.0f)
+					 .MaxValue(2.0f)
+					 .MinSliderValue(-1.0f)
+					 .MaxSliderValue(1.0f)
+					 .SliderExponent(1)
+					 .Delta(0.005f)
+					 .MinDesiredWidth(100.0f)
+					 .SupportDynamicSliderMinValue(true)
+					 .SupportDynamicSliderMaxValue(true)
+					 .OnValueChanged(this, &SControlRigTweenWidget::OnPoseBlendChanged)
+					 .OnValueCommitted(this, &SControlRigTweenWidget::OnPoseBlendCommited)
+					 .OnBeginSliderMovement(this, &SControlRigTweenWidget::OnBeginSliderMovement)
+					 .OnEndSliderMovement(this, &SControlRigTweenWidget::OnEndSliderMovement)
+				 ]
+		]
 	];	
 }
 
