Control Rig: UI Work for Functions - next step for encapsulation

diff --git a/Engine/Plugins/Experimental/ControlRig/Source/ControlRigDeveloper/ControlRigDeveloper.Build.cs b/Engine/Plugins/Experimental/ControlRig/Source/ControlRigDeveloper/ControlRigDeveloper.Build.cs
--- a/Engine/Plugins/Experimental/ControlRig/Source/ControlRigDeveloper/ControlRigDeveloper.Build.cs
+++ b/Engine/Plugins/Experimental/ControlRig/Source/ControlRigDeveloper/ControlRigDeveloper.Build.cs
@@ -53,7 +53,8 @@ namespace UnrealBuildTool.Rules
                         "BlueprintGraph",
                         "PropertyEditor",
                         "RigVMDeveloper",
-                    }
+						"GraphEditor",
+					}
                 );
 
                 PrivateIncludePathModuleNames.Add("ControlRigEditor");
diff --git a/Engine/Plugins/Experimental/ControlRig/Source/ControlRigDeveloper/Private/ControlRigBlueprint.cpp b/Engine/Plugins/Experimental/ControlRig/Source/ControlRigDeveloper/Private/ControlRigBlueprint.cpp
--- a/Engine/Plugins/Experimental/ControlRig/Source/ControlRigDeveloper/Private/ControlRigBlueprint.cpp
+++ b/Engine/Plugins/Experimental/ControlRig/Source/ControlRigDeveloper/Private/ControlRigBlueprint.cpp
@@ -54,6 +54,15 @@ UControlRigBlueprint::UControlRigBlueprint(const FObjectInitializer& ObjectIniti
 
 	Model = ObjectInitializer.CreateDefaultSubobject<URigVMGraph>(this, TEXT("RigVMModel"));
 	FunctionLibrary = ObjectInitializer.CreateDefaultSubobject<URigVMFunctionLibrary>(this, TEXT("RigVMFunctionLibrary"));
+	FunctionLibraryEdGraph = ObjectInitializer.CreateDefaultSubobject<UControlRigGraph>(this, TEXT("RigVMFunctionLibraryEdGraph"));
+	FunctionLibraryEdGraph->Schema = UControlRigGraphSchema::StaticClass();
+	FunctionLibraryEdGraph->bAllowRenaming = 0;
+	FunctionLibraryEdGraph->bEditable = 0;
+	FunctionLibraryEdGraph->bAllowDeletion = 0;
+	FunctionLibraryEdGraph->bIsFunctionDefinition = false;
+	FunctionLibraryEdGraph->Initialize(this);
+
+	Model->SetDefaultFunctionLibrary(FunctionLibrary);
 
 	Validator = ObjectInitializer.CreateDefaultSubobject<UControlRigValidator>(this, TEXT("ControlRigValidator"));
 
@@ -70,6 +79,7 @@ void UControlRigBlueprint::InitializeModelIfRequired(bool bRecompileVM)
 	if (Controllers.Num() == 0)
 	{
 		GetOrCreateController(Model);
+		GetOrCreateController(FunctionLibrary);
 
 		for (int32 i = 0; i < UbergraphPages.Num(); ++i)
 		{
@@ -86,6 +96,8 @@ void UControlRigBlueprint::InitializeModelIfRequired(bool bRecompileVM)
 			}
 		}
 
+		FunctionLibraryEdGraph->Initialize(this);
+
 		HierarchyContainer.OnElementAdded.AddUObject(this, &UControlRigBlueprint::HandleOnElementAdded);
 		HierarchyContainer.OnElementRemoved.AddUObject(this, &UControlRigBlueprint::HandleOnElementRemoved);
 		HierarchyContainer.OnElementRenamed.AddUObject(this, &UControlRigBlueprint::HandleOnElementRenamed);
@@ -223,17 +235,28 @@ void UControlRigBlueprint::PostLoad()
 
 #if WITH_EDITOR
 
+	TArray<URigVMGraph*> GraphsToDetach;
+	GraphsToDetach.Add(GetModel());
+	GraphsToDetach.Add(GetLocalFunctionLibrary());
+
 	if (!IsInAsyncLoadingThread() || IsRunningCommandlet())
 	{
-		GetOrCreateController()->DetachLinksFromPinObjects();
-		TArray<URigVMNode*> Nodes = Model->GetNodes();
-		for (URigVMNode* Node : Nodes)
+		for (URigVMGraph* GraphToDetach : GraphsToDetach)
 		{
-			GetOrCreateController()->RepopulatePinsOnNode(Node);
+			GetOrCreateController(GraphToDetach)->DetachLinksFromPinObjects();
+			TArray<URigVMNode*> Nodes = GraphToDetach->GetNodes();
+			for (URigVMNode* Node : Nodes)
+			{
+				GetOrCreateController()->RepopulatePinsOnNode(Node);
+			}
 		}
 		SetupPinRedirectorsForBackwardsCompatibility();
 	}
-	GetOrCreateController()->ReattachLinksToPinObjects(true /* follow redirectors */);
+
+	for (URigVMGraph* GraphToDetach : GraphsToDetach)
+	{
+		GetOrCreateController(GraphToDetach)->ReattachLinksToPinObjects(true /* follow redirectors */);
+	}
 
 	RecompileVM();
 	RequestControlRigInit();
@@ -373,15 +396,31 @@ URigVMGraph* UControlRigBlueprint::GetModel(const UEdGraph* InEdGraph) const
 		return Model;
 	}
 
+#if WITH_EDITORONLY_DATA
+	if (InEdGraph == FunctionLibraryEdGraph)
+	{
+		return FunctionLibrary;
+	}
+#endif
+
 	const UControlRigGraph* RigGraph = Cast< UControlRigGraph>(InEdGraph);
 	check(RigGraph);
 
+	FString ModelNodePath = RigGraph->ModelNodePath;
+
+	if (RigGraph->bIsFunctionDefinition)
+	{
+		if (URigVMLibraryNode* LibraryNode = FunctionLibrary->FindFunction(*ModelNodePath))
+		{
+			return LibraryNode->GetContainedGraph();
+		}
+	}
+
 	if (RigGraph->GetOuter() == this)
 	{
 		return Model;
 	}
 
-	FString ModelNodePath = RigGraph->ModelNodePath;
 	ensure(!ModelNodePath.IsEmpty());
 
 	URigVMGraph* SubModel = Model;
@@ -429,6 +468,8 @@ TArray<URigVMGraph*> UControlRigBlueprint::GetAllModels() const
 	TArray<URigVMGraph*> Models;
 	Models.Add(GetModel());
 	Models.Append(GetModel()->GetContainedGraphs(true /* recursive */));
+	Models.Add(GetLocalFunctionLibrary());
+	Models.Append(GetLocalFunctionLibrary()->GetContainedGraphs(true /* recursive */));
 	return Models;
 }
 
@@ -570,13 +611,36 @@ URigVMController* UControlRigBlueprint::GetTemplateController()
 
 UEdGraph* UControlRigBlueprint::GetEdGraph(URigVMGraph* InModel) const
 {
+	if (InModel == nullptr)
+	{
+		return nullptr;
+	}
+
+#if WITH_EDITORONLY_DATA
+	if (InModel == FunctionLibrary)
+	{
+		return FunctionLibraryEdGraph;
+	}
+#endif
+
 	TArray<UEdGraph*> EdGraphs;
 	GetAllGraphs(EdGraphs);
 
+	bool bIsFunctionDefinition = false;
+	if (URigVMLibraryNode* LibraryNode = Cast<URigVMLibraryNode>(InModel->GetOuter()))
+	{
+		bIsFunctionDefinition = LibraryNode->GetGraph()->IsA<URigVMFunctionLibrary>();
+	}
+
 	for (UEdGraph* EdGraph : EdGraphs)
 	{
 		if (UControlRigGraph* RigGraph = Cast<UControlRigGraph>(EdGraph))
 		{
+			if (RigGraph->bIsFunctionDefinition != bIsFunctionDefinition)
+			{
+				continue;
+			}
+
 			if (RigGraph->ModelNodePath == InModel->GetNodePath())
 			{
 				return RigGraph;
@@ -1302,12 +1366,22 @@ void UControlRigBlueprint::RebuildGraphFromModel()
 		{
 			Graph->RemoveNode(Node);
 		}
+
+		if (UControlRigGraph* RigGraph = Cast<UControlRigGraph>(Graph))
+		{
+			if (RigGraph->bIsFunctionDefinition)
+			{
+				FunctionGraphs.Remove(RigGraph);
+			}
+		}
 	}
 
 	TArray<URigVMGraph*> RigGraphs;
 	RigGraphs.Add(GetModel());
+	RigGraphs.Add(GetLocalFunctionLibrary());
 
 	GetOrCreateController(RigGraphs[0])->ResendAllNotifications();
+	GetOrCreateController(RigGraphs[1])->ResendAllNotifications();
 
 	for (int32 RigGraphIndex = 0; RigGraphIndex < RigGraphs.Num(); RigGraphIndex++)
 	{
@@ -1322,6 +1396,18 @@ void UControlRigBlueprint::RebuildGraphFromModel()
 			}
 		}
 	}
+
+	EdGraphs.Reset();
+	GetAllGraphs(EdGraphs);
+
+	for (UEdGraph* Graph : EdGraphs)
+	{
+		if (UControlRigGraph* RigGraph = Cast<UControlRigGraph>(Graph))
+		{
+			RigGraph->CacheNameLists(&HierarchyContainer, &DrawContainer);
+		}
+	}
+
 }
 
 void UControlRigBlueprint::Notify(ERigVMGraphNotifType InNotifType, UObject* InSubject)
@@ -1548,12 +1634,12 @@ void UControlRigBlueprint::HandleModifiedEvent(ERigVMGraphNotifType InNotifType,
 			}
 			case ERigVMGraphNotifType::NodeRenamed:
 			{
-				if (URigVMLibraryNode* LibraryNode = Cast<URigVMLibraryNode>(InSubject))
+				if (URigVMCollapseNode* CollapseNode = Cast<URigVMCollapseNode>(InSubject))
 				{
-					FString NewNodePath = LibraryNode->GetNodePath(true /* recursive */);
+					FString NewNodePath = CollapseNode->GetNodePath(true /* recursive */);
 					FString Left, Right = NewNodePath;
 					URigVMNode::SplitNodePathAtEnd(NewNodePath, Left, Right);
-					FString OldNodePath = LibraryNode->GetPreviousFName().ToString();
+					FString OldNodePath = CollapseNode->GetPreviousFName().ToString();
 					if (!Left.IsEmpty())
 					{
 						OldNodePath = URigVMNode::JoinNodePath(Left, OldNodePath);
@@ -1580,6 +1666,11 @@ void UControlRigBlueprint::HandleModifiedEvent(ERigVMGraphNotifType InNotifType,
 						}
 					}
 
+					if (UEdGraph* ContainedEdGraph = GetEdGraph(CollapseNode->GetContainedGraph()))
+					{
+						ContainedEdGraph->Rename(*CollapseNode->GetEditorSubGraphName(), nullptr);
+					}
+
 					FBlueprintEditorUtils::MarkBlueprintAsStructurallyModified(this);
 				}
 				break;
@@ -2388,7 +2479,44 @@ void UControlRigBlueprint::CreateEdGraphForCollapseNodeIfNeeded(URigVMCollapseNo
 		RemoveEdGraphForCollapseNode(InNode, false);
 	}
 
-	if (UControlRigGraph* RigGraph = Cast<UControlRigGraph>(GetEdGraph(InNode->GetGraph())))
+	if (InNode->GetGraph()->IsA<URigVMFunctionLibrary>())
+	{
+		if (URigVMGraph* ContainedGraph = InNode->GetContainedGraph())
+		{
+			bool bFunctionGraphExists = false;
+			for (UEdGraph* FunctionGraph : FunctionGraphs)
+			{
+				if (UControlRigGraph* RigFunctionGraph = Cast<UControlRigGraph>(FunctionGraph))
+				{
+					if (RigFunctionGraph->ModelNodePath == ContainedGraph->GetNodePath())
+					{
+						bFunctionGraphExists = true;
+						break;
+					}
+				}
+			}
+
+			if (!bFunctionGraphExists)
+			{
+				// create a sub graph
+				UControlRigGraph* RigFunctionGraph = NewObject<UControlRigGraph>(this, *InNode->GetName(), RF_Transactional);
+				RigFunctionGraph->Schema = UControlRigGraphSchema::StaticClass();
+				RigFunctionGraph->bAllowRenaming = 1;
+				RigFunctionGraph->bEditable = 1;
+				RigFunctionGraph->bAllowDeletion = 1;
+				RigFunctionGraph->ModelNodePath = ContainedGraph->GetNodePath();
+				RigFunctionGraph->bIsFunctionDefinition = true;
+
+				FunctionGraphs.Add(RigFunctionGraph);
+
+				RigFunctionGraph->Initialize(this);
+
+				GetOrCreateController(ContainedGraph)->ResendAllNotifications();
+			}
+
+		}
+	}
+	else if (UControlRigGraph* RigGraph = Cast<UControlRigGraph>(GetEdGraph(InNode->GetGraph())))
 	{
 		if (URigVMGraph* ContainedGraph = InNode->GetContainedGraph())
 		{
@@ -2413,7 +2541,8 @@ void UControlRigBlueprint::CreateEdGraphForCollapseNodeIfNeeded(URigVMCollapseNo
 				SubRigGraph->bAllowRenaming = 1;
 				SubRigGraph->bEditable = 1;
 				SubRigGraph->bAllowDeletion = 1;
-				SubRigGraph->ModelNodePath = InNode->GetNodePath(true);
+				SubRigGraph->ModelNodePath = ContainedGraph->GetNodePath();
+				SubRigGraph->bIsFunctionDefinition = false;
 
 				RigGraph->SubGraphs.Add(SubRigGraph);
 
@@ -2423,14 +2552,40 @@ void UControlRigBlueprint::CreateEdGraphForCollapseNodeIfNeeded(URigVMCollapseNo
 			}
 		}
 	}
-
 }
 
 bool UControlRigBlueprint::RemoveEdGraphForCollapseNode(URigVMCollapseNode* InNode, bool bNotify)
 {
 	check(InNode);
 
-	if (UControlRigGraph* RigGraph = Cast<UControlRigGraph>(GetEdGraph(InNode->GetGraph())))
+	if (InNode->GetGraph()->IsA<URigVMFunctionLibrary>())
+	{
+		if (URigVMGraph* ContainedGraph = InNode->GetContainedGraph())
+		{
+			for (UEdGraph* FunctionGraph : FunctionGraphs)
+			{
+				if (UControlRigGraph* RigFunctionGraph = Cast<UControlRigGraph>(FunctionGraph))
+				{
+					if (RigFunctionGraph->ModelNodePath == ContainedGraph->GetNodePath())
+					{
+						if (URigVMController* SubController = GetController(ContainedGraph))
+						{
+							SubController->OnModified().RemoveAll(RigFunctionGraph);
+						}
+
+						if (ModifiedEvent.IsBound() && bNotify)
+						{
+							ModifiedEvent.Broadcast(ERigVMGraphNotifType::NodeRemoved, InNode->GetGraph(), InNode);
+						}
+
+						FunctionGraphs.Remove(RigFunctionGraph);
+						return bNotify;
+					}
+				}
+			}
+		}
+	}
+	else if (UControlRigGraph* RigGraph = Cast<UControlRigGraph>(GetEdGraph(InNode->GetGraph())))
 	{
 		if (URigVMGraph* ContainedGraph = InNode->GetContainedGraph())
 		{
diff --git a/Engine/Plugins/Experimental/ControlRig/Source/ControlRigDeveloper/Private/Graph/ControlRigGraph.cpp b/Engine/Plugins/Experimental/ControlRig/Source/ControlRigDeveloper/Private/Graph/ControlRigGraph.cpp
--- a/Engine/Plugins/Experimental/ControlRig/Source/ControlRigDeveloper/Private/Graph/ControlRigGraph.cpp
+++ b/Engine/Plugins/Experimental/ControlRig/Source/ControlRigDeveloper/Private/Graph/ControlRigGraph.cpp
@@ -26,6 +26,7 @@ UControlRigGraph::UControlRigGraph()
 	bSuspendModelNotifications = false;
 	bIsTemporaryGraphForCopyPaste = false;
 	bIsSelecting = false;
+	bIsFunctionDefinition = false;
 }
 
 void UControlRigGraph::Initialize(UControlRigBlueprint* InBlueprint)
@@ -59,6 +60,11 @@ void UControlRigGraph::Serialize(FArchive& Ar)
 
 void UControlRigGraph::CacheNameLists(const FRigHierarchyContainer* HierarchyContainer, const FControlRigDrawContainer* DrawContainer)
 {
+	if (UControlRigGraph* OuterGraph = Cast<UControlRigGraph>(GetOuter()))
+	{
+		return;
+	}
+
 	check(HierarchyContainer);
 	check(DrawContainer);
 
@@ -71,21 +77,37 @@ void UControlRigGraph::CacheNameLists(const FRigHierarchyContainer* HierarchyCon
 
 const TArray<TSharedPtr<FString>>& UControlRigGraph::GetBoneNameList(URigVMPin* InPin) const
 {
+	if (UControlRigGraph* OuterGraph = Cast<UControlRigGraph>(GetOuter()))
+	{
+		return OuterGraph->GetBoneNameList(InPin);
+	}
 	return BoneNameList;
 }
 
 const TArray<TSharedPtr<FString>>& UControlRigGraph::GetControlNameList(URigVMPin* InPin) const
 {
+	if (UControlRigGraph* OuterGraph = Cast<UControlRigGraph>(GetOuter()))
+	{
+		return OuterGraph->GetControlNameList(InPin);
+	}
 	return ControlNameList;
 }
 
 const TArray<TSharedPtr<FString>>& UControlRigGraph::GetSpaceNameList(URigVMPin* InPin) const
 {
+	if (UControlRigGraph* OuterGraph = Cast<UControlRigGraph>(GetOuter()))
+	{
+		return OuterGraph->GetSpaceNameList(InPin);
+	}
 	return SpaceNameList;
 }
 
 const TArray<TSharedPtr<FString>>& UControlRigGraph::GetCurveNameList(URigVMPin* InPin) const
 {
+	if (UControlRigGraph* OuterGraph = Cast<UControlRigGraph>(GetOuter()))
+	{
+		return OuterGraph->GetCurveNameList(InPin);
+	}
 	return CurveNameList;
 }
 
@@ -115,6 +137,11 @@ const TArray<TSharedPtr<FString>>& UControlRigGraph::GetElementNameList(URigVMPi
 
 const TArray<TSharedPtr<FString>>& UControlRigGraph::GetElementNameList(ERigElementType InElementType) const
 {
+	if (UControlRigGraph* OuterGraph = Cast<UControlRigGraph>(GetOuter()))
+	{
+		return OuterGraph->GetElementNameList(InElementType);
+	}
+
 	switch (InElementType)
 	{
 		case ERigElementType::Bone:
@@ -140,6 +167,10 @@ const TArray<TSharedPtr<FString>>& UControlRigGraph::GetElementNameList(ERigElem
 
 const TArray<TSharedPtr<FString>>& UControlRigGraph::GetDrawingNameList(URigVMPin* InPin) const
 {
+	if (UControlRigGraph* OuterGraph = Cast<UControlRigGraph>(GetOuter()))
+	{
+		return OuterGraph->GetDrawingNameList(InPin);
+	}
 	return DrawingNameList;
 }
 
diff --git a/Engine/Plugins/Experimental/ControlRig/Source/ControlRigDeveloper/Private/Graph/ControlRigGraphNode.cpp b/Engine/Plugins/Experimental/ControlRig/Source/ControlRigDeveloper/Private/Graph/ControlRigGraphNode.cpp
--- a/Engine/Plugins/Experimental/ControlRig/Source/ControlRigDeveloper/Private/Graph/ControlRigGraphNode.cpp
+++ b/Engine/Plugins/Experimental/ControlRig/Source/ControlRigDeveloper/Private/Graph/ControlRigGraphNode.cpp
@@ -23,6 +23,9 @@
 #include "RigVMCore/RigVMExecuteContext.h"
 #include "ControlRigDeveloper.h"
 #include "ControlRigObjectVersion.h"
+#include "RigVMModel/Nodes/RigVMFunctionReferenceNode.h"
+#include "RigVMModel/Nodes/RigVMFunctionEntryNode.h"
+#include "RigVMModel/Nodes/RigVMFunctionReturnNode.h"
 
 #if WITH_EDITOR
 #include "IControlRigEditorModule.h"
@@ -82,7 +85,7 @@ void UControlRigGraphNode::ReconstructNode_Internal(bool bForce)
 {
 	DECLARE_SCOPE_HIERARCHICAL_COUNTER_FUNC()
 
-	UControlRigGraph* RigGraph = Cast<UControlRigGraph>(GetGraph());
+	UControlRigGraph* RigGraph = Cast<UControlRigGraph>(GetOuter()); 
 	if (RigGraph && !bForce)
 	{
 		if (RigGraph->bIsTemporaryGraphForCopyPaste)
@@ -117,7 +120,10 @@ void UControlRigGraphNode::ReconstructNode_Internal(bool bForce)
 	// Let subclasses do any additional work
 	PostReconstructNode();
 
-	GetGraph()->NotifyGraphChanged();
+	if (RigGraph)
+	{
+		RigGraph->NotifyGraphChanged();
+	}
 }
 
 bool UControlRigGraphNode::IsDeprecated() const
@@ -521,11 +527,36 @@ FLinearColor UControlRigGraphNode::GetNodeBodyTintColor() const
 	return CachedNodeColor * GetNodeOpacityColor();
 }
 
+bool UControlRigGraphNode::ShowPaletteIconOnNode() const
+{
+	if (URigVMNode* ModelNode = GetModelNode())
+	{
+		return ModelNode->IsEvent() ||
+			ModelNode->IsA<URigVMFunctionEntryNode>() ||
+			ModelNode->IsA<URigVMFunctionReturnNode>() ||
+			ModelNode->IsA<URigVMFunctionReferenceNode>();
+	}
+	return false;
+}
+
 FSlateIcon UControlRigGraphNode::GetIconAndTint(FLinearColor& OutColor) const
 {
-	OutColor = GetNodeTitleColor();
-	static FSlateIcon Icon("EditorStyle", "Kismet.AllClasses.FunctionIcon");
-	return Icon;
+	OutColor = FLinearColor::White;
+
+	static FSlateIcon FunctionIcon("EditorStyle", "Kismet.AllClasses.FunctionIcon");
+	static FSlateIcon EventIcon("EditorStyle", "GraphEditor.Event_16x");
+
+	if (URigVMNode* ModelNode = GetModelNode())
+	{
+		if (ModelNode->IsEvent() || 
+			ModelNode->IsA<URigVMFunctionEntryNode>() || 
+			ModelNode->IsA<URigVMFunctionReturnNode>())
+		{
+			return EventIcon;
+		}
+	}
+
+	return FunctionIcon;
 }
 
 void UControlRigGraphNode::GetNodeContextMenuActions(class UToolMenu* Menu, class UGraphNodeContextMenuContext* Context) const
@@ -577,7 +608,7 @@ void UControlRigGraphNode::CopyPinDefaultsToModel(UEdGraphPin* Pin, bool bUndo)
 	if (Pin->Direction != EGPD_Input)
 	{
 		return;
-		}
+	}
 
 	const UEdGraphSchema_K2* K2Schema = GetDefault<UEdGraphSchema_K2>();
 	if (URigVMPin* ModelPin = GetModelPinFromPinPath(Pin->GetName()))
@@ -585,13 +616,13 @@ void UControlRigGraphNode::CopyPinDefaultsToModel(UEdGraphPin* Pin, bool bUndo)
 		if (ModelPin->GetSubPins().Num() > 0)
 		{
 			return;
-	}
+		}
 
 		FString DefaultValue = Pin->DefaultValue;
 		if (DefaultValue == FName(NAME_None).ToString() && Pin->PinType.PinSubCategory == UEdGraphSchema_K2::PC_Name)
 		{
 			DefaultValue = FString();
-	}
+		}
 
 		if (ModelPin->GetDefaultValue() != DefaultValue)
 		{
diff --git a/Engine/Plugins/Experimental/ControlRig/Source/ControlRigDeveloper/Private/Graph/ControlRigGraphSchema.cpp b/Engine/Plugins/Experimental/ControlRig/Source/ControlRigDeveloper/Private/Graph/ControlRigGraphSchema.cpp
--- a/Engine/Plugins/Experimental/ControlRig/Source/ControlRigDeveloper/Private/Graph/ControlRigGraphSchema.cpp
+++ b/Engine/Plugins/Experimental/ControlRig/Source/ControlRigDeveloper/Private/Graph/ControlRigGraphSchema.cpp
@@ -8,6 +8,7 @@
 #include "Units/RigUnit.h"
 #include "Kismet2/BlueprintEditorUtils.h"
 #include "EdGraphNode_Comment.h"
+#include "EdGraphSchema_K2_Actions.h"
 #include "ScopedTransaction.h"
 #include "Framework/MultiBox/MultiBoxBuilder.h"
 #include "GraphEditorActions.h"
@@ -27,6 +28,78 @@
 
 const FName UControlRigGraphSchema::GraphName_ControlRig(TEXT("Rig Graph"));
 
+FReply FControlRigFunctionDragDropAction::DroppedOnPanel(const TSharedRef< class SWidget >& Panel, FVector2D ScreenPosition, FVector2D GraphPosition, UEdGraph& Graph)
+{
+	if (UControlRigGraph* TargetRigGraph = Cast<UControlRigGraph>(&Graph))
+	{
+		if (UControlRigBlueprint* TargetRigBlueprint = Cast<UControlRigBlueprint>(FBlueprintEditorUtils::FindBlueprintForGraph(TargetRigGraph)))
+		{
+			if (URigVMGraph* FunctionDefinitionGraph = SourceRigBlueprint->GetModel(SourceRigGraph))
+			{
+				if (URigVMLibraryNode* FunctionDefinitionNode = Cast<URigVMLibraryNode>(FunctionDefinitionGraph->GetOuter()))
+				{
+					if(URigVMController* TargetController = TargetRigBlueprint->GetController(TargetRigGraph))
+					{
+						TargetController->AddFunctionReferenceNode(FunctionDefinitionNode, GraphPosition);
+					}
+				}
+			}
+		}
+	}
+	return FReply::Unhandled();
+}
+
+FReply FControlRigFunctionDragDropAction::DroppedOnPin(FVector2D ScreenPosition, FVector2D GraphPosition)
+{
+	return FReply::Unhandled();
+}
+
+FReply FControlRigFunctionDragDropAction::DroppedOnAction(TSharedRef<FEdGraphSchemaAction> Action)
+{
+	return FReply::Unhandled();
+}
+
+FReply FControlRigFunctionDragDropAction::DroppedOnCategory(FText Category)
+{
+	// todo
+	/*
+	if (SourceAction.IsValid())
+	{
+		SourceAction->MovePersistentItemToCategory(Category);
+	}
+	*/
+	return FReply::Unhandled();
+}
+
+void FControlRigFunctionDragDropAction::HoverTargetChanged()
+{
+	// todo - see FMyBlueprintItemDragDropAction
+	FGraphSchemaActionDragDropAction::HoverTargetChanged();
+
+	// check for category + graph, everything else we won't allow for now.
+
+	bDropTargetValid = true;
+}
+
+FControlRigFunctionDragDropAction::FControlRigFunctionDragDropAction()
+	: FGraphSchemaActionDragDropAction()
+	, SourceRigBlueprint(nullptr)
+	, SourceRigGraph(nullptr)
+	, bControlDrag(false)
+	, bAltDrag(false)
+{
+}
+
+TSharedRef<FControlRigFunctionDragDropAction> FControlRigFunctionDragDropAction::New(TSharedPtr<FEdGraphSchemaAction> InAction, UControlRigBlueprint* InRigBlueprint, UControlRigGraph* InRigGraph)
+{
+	TSharedRef<FControlRigFunctionDragDropAction> Action = MakeShareable(new FControlRigFunctionDragDropAction);
+	Action->SourceAction = InAction;
+	Action->SourceRigBlueprint = InRigBlueprint;
+	Action->SourceRigGraph = InRigGraph;
+	Action->Construct();
+	return Action;
+}
+
 UControlRigGraphSchema::UControlRigGraphSchema()
 {
 }
@@ -199,6 +272,46 @@ void UControlRigGraphSchema::BreakSinglePinLink(UEdGraphPin* SourcePin, UEdGraph
 	}
 }
 
+bool UControlRigGraphSchema::CanGraphBeDropped(TSharedPtr<FEdGraphSchemaAction> InAction) const
+{
+	if (!InAction.IsValid())
+	{
+		return false;
+	}
+
+	if (InAction->GetTypeId() == FEdGraphSchemaAction_K2Graph::StaticGetTypeId())
+	{
+		FEdGraphSchemaAction_K2Graph* FuncAction = (FEdGraphSchemaAction_K2Graph*)InAction.Get();
+		if (UControlRigGraph* RigGraph = Cast<UControlRigGraph>((UEdGraph*)FuncAction->EdGraph))
+		{
+			return true;
+		}
+	}
+	
+	return false;
+}
+
+FReply UControlRigGraphSchema::BeginGraphDragAction(TSharedPtr<FEdGraphSchemaAction> InAction) const
+{
+	if (!InAction.IsValid())
+	{
+		return FReply::Unhandled();
+	}
+
+	if (InAction->GetTypeId() == FEdGraphSchemaAction_K2Graph::StaticGetTypeId())
+	{
+		FEdGraphSchemaAction_K2Graph* FuncAction = (FEdGraphSchemaAction_K2Graph*)InAction.Get();
+		if (UControlRigGraph* RigGraph = Cast<UControlRigGraph>(FuncAction->EdGraph))
+		{
+			if (UControlRigBlueprint* RigBlueprint = Cast<UControlRigBlueprint>(FBlueprintEditorUtils::FindBlueprintForGraph(RigGraph)))
+			{
+				return FReply::Handled().BeginDragDrop(FControlRigFunctionDragDropAction::New(InAction, RigBlueprint, RigGraph));
+			}
+		}
+	}
+	return FReply::Unhandled();
+}
+
 FConnectionDrawingPolicy* UControlRigGraphSchema::CreateConnectionDrawingPolicy(int32 InBackLayerID, int32 InFrontLayerID, float InZoomFactor, const FSlateRect& InClippingRect, class FSlateWindowElementList& InDrawElements, class UEdGraph* InGraphObj) const
 {
 #if WITH_EDITOR
@@ -295,9 +408,15 @@ bool UControlRigGraphSchema::TryDeleteGraph(UEdGraph* GraphToDelete) const
 		{
 			if (URigVMGraph* Model = RigBlueprint->GetModel())
 			{
-				if (URigVMLibraryNode* LibraryNode = Cast<URigVMLibraryNode>(Model->FindNode(RigGraph->ModelNodePath)))
+				URigVMLibraryNode* LibraryNode = Cast<URigVMLibraryNode>(Model->FindNode(RigGraph->ModelNodePath));
+				if (LibraryNode == nullptr && RigBlueprint->GetLocalFunctionLibrary())
 				{
-					if (URigVMController* Controller = RigBlueprint->GetController(LibraryNode->GetGraph()))
+					LibraryNode = Cast<URigVMLibraryNode>(RigBlueprint->GetLocalFunctionLibrary()->FindFunction(*RigGraph->ModelNodePath));
+				}
+
+				if (LibraryNode)
+				{
+					if (URigVMController* Controller = RigBlueprint->GetOrCreateController(LibraryNode->GetGraph()))
 					{
 						return Controller->RemoveNode(LibraryNode);
 					}
@@ -316,9 +435,15 @@ bool UControlRigGraphSchema::TryRenameGraph(UEdGraph* GraphToRename, const FName
 		{
 			if (URigVMGraph* Model = RigBlueprint->GetModel())
 			{
-				if (URigVMLibraryNode* LibraryNode = Cast<URigVMLibraryNode>(Model->FindNode(RigGraph->ModelNodePath)))
+				URigVMLibraryNode* LibraryNode = Cast<URigVMLibraryNode>(Model->FindNode(RigGraph->ModelNodePath));
+				if (LibraryNode == nullptr && RigBlueprint->GetLocalFunctionLibrary())
+				{
+					LibraryNode = Cast<URigVMLibraryNode>(RigBlueprint->GetLocalFunctionLibrary()->FindFunction(*RigGraph->ModelNodePath));
+				}
+
+				if (LibraryNode)
 				{
-					if (URigVMController* Controller = RigBlueprint->GetController(LibraryNode->GetGraph()))
+					if (URigVMController* Controller = RigBlueprint->GetOrCreateController(LibraryNode->GetGraph()))
 					{
 						Controller->RenameNode(LibraryNode, InNewName);
 						return true;
@@ -370,12 +495,24 @@ UEdGraphPin* UControlRigGraphSchema::DropPinOnNode(UEdGraphNode* InTargetNode, c
 								TypeObjectPathName = *ExternalVar.TypeObject->GetPathName();
 							}
 
+							FString DefaultValue;
+							if (PinBeingDropped)
+							{
+								if (UControlRigGraphNode* SourceNode = Cast<UControlRigGraphNode>(PinBeingDropped->GetOwningNode()))
+								{
+									if (URigVMPin* SourcePin = SourceNode->GetModelPinFromPinPath(PinBeingDropped->GetName()))
+									{
+										DefaultValue = SourcePin->GetDefaultValue();
+									}
+								}
+							}
+
 							FName ExposedPinName = Controller->AddExposedPin(
 								InSourcePinName,
 								PinDirection,
 								TypeName,
 								TypeObjectPathName,
-								FString()
+								DefaultValue
 							);
 							
 							if (!ExposedPinName.IsNone())
diff --git a/Engine/Plugins/Experimental/ControlRig/Source/ControlRigDeveloper/Public/ControlRigBlueprint.h b/Engine/Plugins/Experimental/ControlRig/Source/ControlRigDeveloper/Public/ControlRigBlueprint.h
--- a/Engine/Plugins/Experimental/ControlRig/Source/ControlRigDeveloper/Public/ControlRigBlueprint.h
+++ b/Engine/Plugins/Experimental/ControlRig/Source/ControlRigDeveloper/Public/ControlRigBlueprint.h
@@ -69,7 +69,7 @@ public:
 
 	virtual bool SupportsGlobalVariables() const override { return true; }
 	virtual bool SupportsLocalVariables() const override { return false; }
-	virtual bool SupportsFunctions() const override { return false; }
+	virtual bool SupportsFunctions() const override { return true; }
 	virtual bool SupportsMacros() const override { return false; }
 	virtual bool SupportsDelegates() const override { return false; }
 	virtual bool SupportsEventGraphs() const override { return false; }
@@ -129,6 +129,11 @@ public:
 	UEdGraph* GetEdGraph(const FString& InNodePath) const;
 #endif
 
+#if WITH_EDITORONLY_DATA
+	UPROPERTY()
+	UControlRigGraph* FunctionLibraryEdGraph;
+#endif
+
 	UPROPERTY(EditAnywhere, BlueprintReadWrite, Category = "VM")
 	FRigVMCompileSettings VMCompileSettings;
 
diff --git a/Engine/Plugins/Experimental/ControlRig/Source/ControlRigDeveloper/Public/ControlRigBlueprintCompiler.h b/Engine/Plugins/Experimental/ControlRig/Source/ControlRigDeveloper/Public/ControlRigBlueprintCompiler.h
--- a/Engine/Plugins/Experimental/ControlRig/Source/ControlRigDeveloper/Public/ControlRigBlueprintCompiler.h
+++ b/Engine/Plugins/Experimental/ControlRig/Source/ControlRigDeveloper/Public/ControlRigBlueprintCompiler.h
@@ -44,6 +44,8 @@ public:
 	virtual void CleanAndSanitizeClass(UBlueprintGeneratedClass* ClassToClean, UObject*& InOldCDO) override;
 	virtual void MergeUbergraphPagesIn(UEdGraph* Ubergraph) override {};
 	virtual void PreCompileUpdateBlueprintOnLoad(UBlueprint* BP) override;
+	virtual void CreateFunctionList() override {}
+	virtual void ProcessOneFunctionGraph(UEdGraph* SourceGraph, bool bInternalFunction = false) override {}
 
 private:
 
diff --git a/Engine/Plugins/Experimental/ControlRig/Source/ControlRigDeveloper/Public/Graph/ControlRigGraph.h b/Engine/Plugins/Experimental/ControlRig/Source/ControlRigDeveloper/Public/Graph/ControlRigGraph.h
--- a/Engine/Plugins/Experimental/ControlRig/Source/ControlRigDeveloper/Public/Graph/ControlRigGraph.h
+++ b/Engine/Plugins/Experimental/ControlRig/Source/ControlRigDeveloper/Public/Graph/ControlRigGraph.h
@@ -62,6 +62,9 @@ public:
 	UPROPERTY()
 	FString ModelNodePath;
 
+	UPROPERTY()
+	bool bIsFunctionDefinition;
+
 private:
 
 	template<class T>
diff --git a/Engine/Plugins/Experimental/ControlRig/Source/ControlRigDeveloper/Public/Graph/ControlRigGraphNode.h b/Engine/Plugins/Experimental/ControlRig/Source/ControlRigDeveloper/Public/Graph/ControlRigGraphNode.h
--- a/Engine/Plugins/Experimental/ControlRig/Source/ControlRigDeveloper/Public/Graph/ControlRigGraphNode.h
+++ b/Engine/Plugins/Experimental/ControlRig/Source/ControlRigDeveloper/Public/Graph/ControlRigGraphNode.h
@@ -80,6 +80,7 @@ public:
 	virtual FText GetNodeTitle(ENodeTitleType::Type TitleType) const override;
 	virtual FLinearColor GetNodeTitleColor() const override;
 	virtual FLinearColor GetNodeBodyTintColor() const;
+	virtual bool ShowPaletteIconOnNode() const override;
 	virtual FSlateIcon GetIconAndTint(FLinearColor& OutColor) const override;
 	virtual void AllocateDefaultPins() override;
 	virtual void ReconstructNode() override;
@@ -201,4 +202,5 @@ private:
 
 	friend class SControlRigGraphNode;
 	friend class FControlRigArgumentLayout;
+	friend class FControlRigGraphDetails;
 };
diff --git a/Engine/Plugins/Experimental/ControlRig/Source/ControlRigDeveloper/Public/Graph/ControlRigGraphSchema.h b/Engine/Plugins/Experimental/ControlRig/Source/ControlRigDeveloper/Public/Graph/ControlRigGraphSchema.h
--- a/Engine/Plugins/Experimental/ControlRig/Source/ControlRigDeveloper/Public/Graph/ControlRigGraphSchema.h
+++ b/Engine/Plugins/Experimental/ControlRig/Source/ControlRigDeveloper/Public/Graph/ControlRigGraphSchema.h
@@ -3,6 +3,7 @@
 #pragma once
 
 #include "EdGraph/EdGraphSchema.h"
+#include "GraphEditorDragDropAction.h"
 #include "ControlRigGraphSchema.generated.h"
 
 class UControlRigBlueprint;
@@ -39,6 +40,44 @@ struct FControlRigPinConnectionResponse
 	ECanCreateConnectionResponse_Extended ExtendedResponse;
 };
 
+/** DragDropAction class for drag and dropping an item from the My Blueprints tree (e.g., variable or function) */
+class CONTROLRIGDEVELOPER_API FControlRigFunctionDragDropAction : public FGraphSchemaActionDragDropAction
+{
+public:
+
+	DRAG_DROP_OPERATOR_TYPE(FControlRigFunctionDragDropAction, FGraphSchemaActionDragDropAction)
+
+	// FGraphEditorDragDropAction interface
+	virtual FReply DroppedOnPanel(const TSharedRef< class SWidget >& Panel, FVector2D ScreenPosition, FVector2D GraphPosition, UEdGraph& Graph) override;
+	virtual FReply DroppedOnPin(FVector2D ScreenPosition, FVector2D GraphPosition) override;
+	virtual FReply DroppedOnAction(TSharedRef<FEdGraphSchemaAction> Action) override;
+	virtual FReply DroppedOnCategory(FText Category) override;
+	virtual void HoverTargetChanged() override;
+	// End of FGraphEditorDragDropAction
+
+	/** Set if operation is modified by alt */
+	void SetAltDrag(bool InIsAltDrag) { bAltDrag = InIsAltDrag; }
+
+	/** Set if operation is modified by the ctrl key */
+	void SetCtrlDrag(bool InIsCtrlDrag) { bControlDrag = InIsCtrlDrag; }
+
+protected:
+
+	/** Constructor */
+	FControlRigFunctionDragDropAction();
+
+	static TSharedRef<FControlRigFunctionDragDropAction> New(TSharedPtr<FEdGraphSchemaAction> InAction, UControlRigBlueprint* InRigBlueprint, UControlRigGraph* InRigGraph);
+
+protected:
+
+	UControlRigBlueprint* SourceRigBlueprint;
+	UControlRigGraph* SourceRigGraph;
+	bool bControlDrag;
+	bool bAltDrag;
+
+	friend class UControlRigGraphSchema;
+};
+
 UCLASS()
 class CONTROLRIGDEVELOPER_API UControlRigGraphSchema : public UEdGraphSchema
 {
@@ -59,6 +98,8 @@ public:
 	virtual FLinearColor GetPinTypeColor(const FEdGraphPinType& PinType) const override;
 	virtual void BreakPinLinks(UEdGraphPin& TargetPin, bool bSendsNodeNotifcation) const override;
 	virtual void BreakSinglePinLink(UEdGraphPin* SourcePin, UEdGraphPin* TargetPin) const override;
+	virtual bool CanGraphBeDropped(TSharedPtr<FEdGraphSchemaAction> InAction) const override;
+	virtual FReply BeginGraphDragAction(TSharedPtr<FEdGraphSchemaAction> InAction) const override;
 	virtual class FConnectionDrawingPolicy* CreateConnectionDrawingPolicy(int32 InBackLayerID, int32 InFrontLayerID, float InZoomFactor, const FSlateRect& InClippingRect, class FSlateWindowElementList& InDrawElements, class UEdGraph* InGraphObj) const override;
 	virtual bool ShouldHidePinDefaultValue(UEdGraphPin* Pin) const override;
 	virtual void TrySetDefaultValue(UEdGraphPin& InPin, const FString& InNewDefaultValue, bool bMarkAsModified = true) const override;
@@ -83,6 +124,7 @@ public:
 	virtual bool CanDuplicateGraph(UEdGraph* InSourceGraph) const { return false; }
 	virtual UEdGraphPin* DropPinOnNode(UEdGraphNode* InTargetNode, const FName& InSourcePinName, const FEdGraphPinType& InSourcePinType, EEdGraphPinDirection InSourcePinDirection) const override;
 	virtual bool SupportsDropPinOnNode(UEdGraphNode* InTargetNode, const FEdGraphPinType& InSourcePinType, EEdGraphPinDirection InSourcePinDirection, FText& OutErrorMessage) const override;
+	virtual void SetPinBeingDroppedOnNode(UEdGraphPin* InSourcePin) const override { PinBeingDropped = InSourcePin; }
 
 	/** Create a graph node for a rig */
 	UControlRigGraphNode* CreateGraphNode(UControlRigGraph* InGraph, const FName& InPropertyName) const;
@@ -105,6 +147,7 @@ private:
 
 	const UEdGraphPin* LastPinForCompatibleCheck = nullptr;
 	bool bLastPinWasInput;
+	mutable UEdGraphPin* PinBeingDropped = nullptr;
 
 	friend class UControlRigRerouteNodeSpawner;
 	friend class UControlRigIfNodeSpawner;
diff --git a/Engine/Plugins/Experimental/ControlRig/Source/ControlRigEditor/Private/ControlRigEditorModule.cpp b/Engine/Plugins/Experimental/ControlRig/Source/ControlRigEditor/Private/ControlRigEditorModule.cpp
--- a/Engine/Plugins/Experimental/ControlRig/Source/ControlRigEditor/Private/ControlRigEditorModule.cpp
+++ b/Engine/Plugins/Experimental/ControlRig/Source/ControlRigEditor/Private/ControlRigEditorModule.cpp
@@ -1733,6 +1733,20 @@ void FControlRigEditorModule::GetContextMenuActions(const UControlRigGraphSchema
 						})
 					));
 
+					if (URigVMCollapseNode* CollapseNode = Cast<URigVMCollapseNode>(RigNode->GetModelNode()))
+					{
+						OrganizationSection.AddMenuEntry(
+							"Promote To Function",
+							LOCTEXT("PromoteToFunction", "Promote To Function"),
+							LOCTEXT("PromoteToFunction_Tooltip", "Turns the Collapse Node into a Function"),
+							FSlateIcon(),
+							FUIAction(FExecuteAction::CreateLambda([Controller, CollapseNode]() {
+								// Controller->PromoteCollapseNodeToFunction(CollapseNode->GetFName());
+								// todo
+							})
+						));
+					}
+
 					if (URigVMLibraryNode* LibraryNode = Cast<URigVMLibraryNode>(RigNode->GetModelNode()))
 					{
 						OrganizationSection.AddMenuEntry(
diff --git a/Engine/Plugins/Experimental/ControlRig/Source/ControlRigEditor/Private/ControlRigGraphDetails.cpp b/Engine/Plugins/Experimental/ControlRig/Source/ControlRigEditor/Private/ControlRigGraphDetails.cpp
--- a/Engine/Plugins/Experimental/ControlRig/Source/ControlRigEditor/Private/ControlRigGraphDetails.cpp
+++ b/Engine/Plugins/Experimental/ControlRig/Source/ControlRigEditor/Private/ControlRigGraphDetails.cpp
@@ -19,7 +19,10 @@
 
 #define LOCTEXT_NAMESPACE "ControlRigGraphDetails"
 
-FControlRigArgumentGroupLayout::FControlRigArgumentGroupLayout(URigVMGraph* InGraph, UControlRigBlueprint* InBlueprint, bool bInputs)
+FControlRigArgumentGroupLayout::FControlRigArgumentGroupLayout(
+	URigVMGraph* InGraph, 
+	UControlRigBlueprint* InBlueprint, 
+	bool bInputs)
 	: GraphPtr(InGraph)
 	, ControlRigBlueprintPtr(InBlueprint)
 	, bIsInputGroup(bInputs)
@@ -54,7 +57,8 @@ void FControlRigArgumentGroupLayout::GenerateChildContent(IDetailChildrenBuilder
 					TSharedRef<class FControlRigArgumentLayout> ControlRigArgumentLayout = MakeShareable(new FControlRigArgumentLayout(
 						Pin,
 						Graph,
-						ControlRigBlueprintPtr.Get()));
+						ControlRigBlueprintPtr.Get()
+					));
 					ChildrenBuilder.AddCustomBuilder(ControlRigArgumentLayout);
 					WasContentAdded = true;
 				}
@@ -122,73 +126,73 @@ void FControlRigArgumentLayout::GenerateHeaderRowContent(FDetailWidgetRow& NodeR
 	TypeTreeFilter |= ETypeTreeFilter::AllowExec;
 
 	NodeRow
-	.NameContent()
-	[
-		SNew(SHorizontalBox)
-		+ SHorizontalBox::Slot()
+		.NameContent()
+		[
+			SNew(SHorizontalBox)
+			+ SHorizontalBox::Slot()
 		.FillWidth(1)
 		.VAlign(VAlign_Center)
 		[
 			SAssignNew(ArgumentNameWidget, SEditableTextBox)
 			.Text(this, &FControlRigArgumentLayout::OnGetArgNameText)
-			.OnTextChanged(this, &FControlRigArgumentLayout::OnArgNameChange)
-			.OnTextCommitted(this, &FControlRigArgumentLayout::OnArgNameTextCommitted)
-			.ToolTipText(this, &FControlRigArgumentLayout::OnGetArgToolTipText)
-			.Font(IDetailLayoutBuilder::GetDetailFont())
-			.IsEnabled(!ShouldPinBeReadOnly())
+		.OnTextChanged(this, &FControlRigArgumentLayout::OnArgNameChange)
+		.OnTextCommitted(this, &FControlRigArgumentLayout::OnArgNameTextCommitted)
+		.ToolTipText(this, &FControlRigArgumentLayout::OnGetArgToolTipText)
+		.Font(IDetailLayoutBuilder::GetDetailFont())
+		.IsEnabled(!ShouldPinBeReadOnly())
+		]
 		]
-	]
 	.ValueContent()
-	.MaxDesiredWidth(980.f)
-	[
-		SNew(SHorizontalBox)
-		+ SHorizontalBox::Slot()
+		.MaxDesiredWidth(980.f)
+		[
+			SNew(SHorizontalBox)
+			+ SHorizontalBox::Slot()
 		.VAlign(VAlign_Center)
 		.Padding(0.0f, 0.0f, 4.0f, 0.0f)
 		.AutoWidth()
 		[
 			SNew(SPinTypeSelector, FGetPinTypeTree::CreateUObject(K2Schema, &UEdGraphSchema_K2::GetVariableTypeTree))
 			.TargetPinType(this, &FControlRigArgumentLayout::OnGetPinInfo)
-			.OnPinTypePreChanged(this, &FControlRigArgumentLayout::OnPrePinInfoChange)
-			.OnPinTypeChanged(this, &FControlRigArgumentLayout::PinInfoChanged)
-			.Schema(K2Schema)
-			.TypeTreeFilter(TypeTreeFilter)
-			.bAllowArrays(!ShouldPinBeReadOnly())
-			.IsEnabled(!ShouldPinBeReadOnly(true))
-			.Font(IDetailLayoutBuilder::GetDetailFont())
+		.OnPinTypePreChanged(this, &FControlRigArgumentLayout::OnPrePinInfoChange)
+		.OnPinTypeChanged(this, &FControlRigArgumentLayout::PinInfoChanged)
+		.Schema(K2Schema)
+		.TypeTreeFilter(TypeTreeFilter)
+		.bAllowArrays(!ShouldPinBeReadOnly())
+		.IsEnabled(!ShouldPinBeReadOnly(true))
+		.Font(IDetailLayoutBuilder::GetDetailFont())
 		]
-		+ SHorizontalBox::Slot()
+	+ SHorizontalBox::Slot()
 		.AutoWidth()
 		[
 			SNew(SButton)
 			.ButtonStyle(FAppStyle::Get(), TEXT("SimpleButton"))
-			.ContentPadding(0)
-			.IsEnabled(!IsPinEditingReadOnly())
-			.OnClicked(this, &FControlRigArgumentLayout::OnArgMoveUp)
-			.ToolTipText(LOCTEXT("FunctionArgDetailsArgMoveUpTooltip", "Move this parameter up in the list."))
-			[
-				SNew(SImage)
-				.Image(FEditorStyle::GetBrush("Icons.ChevronUp"))
-			.ColorAndOpacity(FSlateColor::UseForeground())
-			]
+		.ContentPadding(0)
+		.IsEnabled(!IsPinEditingReadOnly())
+		.OnClicked(this, &FControlRigArgumentLayout::OnArgMoveUp)
+		.ToolTipText(LOCTEXT("FunctionArgDetailsArgMoveUpTooltip", "Move this parameter up in the list."))
+		[
+			SNew(SImage)
+			.Image(FEditorStyle::GetBrush("Icons.ChevronUp"))
+		.ColorAndOpacity(FSlateColor::UseForeground())
+		]
 		]
-		+ SHorizontalBox::Slot()
+	+ SHorizontalBox::Slot()
 		.AutoWidth()
 		.Padding(2, 0)
 		[
 			SNew(SButton)
 			.ButtonStyle(FAppStyle::Get(), TEXT("SimpleButton"))
-			.ContentPadding(0)
-			.IsEnabled(!IsPinEditingReadOnly())
-			.OnClicked(this, &FControlRigArgumentLayout::OnArgMoveDown)
-			.ToolTipText(LOCTEXT("FunctionArgDetailsArgMoveDownTooltip", "Move this parameter down in the list."))
-			[
-				SNew(SImage)
-				.Image(FEditorStyle::GetBrush("Icons.ChevronDown"))
-			.ColorAndOpacity(FSlateColor::UseForeground())
-			]
+		.ContentPadding(0)
+		.IsEnabled(!IsPinEditingReadOnly())
+		.OnClicked(this, &FControlRigArgumentLayout::OnArgMoveDown)
+		.ToolTipText(LOCTEXT("FunctionArgDetailsArgMoveDownTooltip", "Move this parameter down in the list."))
+		[
+			SNew(SImage)
+			.Image(FEditorStyle::GetBrush("Icons.ChevronDown"))
+		.ColorAndOpacity(FSlateColor::UseForeground())
 		]
-		+ SHorizontalBox::Slot()
+		]
+	+ SHorizontalBox::Slot()
 		.HAlign(HAlign_Right)
 		.VAlign(VAlign_Center)
 		.Padding(10, 0, 0, 0)
@@ -196,49 +200,12 @@ void FControlRigArgumentLayout::GenerateHeaderRowContent(FDetailWidgetRow& NodeR
 		[
 			PropertyCustomizationHelpers::MakeClearButton(FSimpleDelegate::CreateSP(this, &FControlRigArgumentLayout::OnRemoveClicked), LOCTEXT("FunctionArgDetailsClearTooltip", "Remove this parameter."), !IsPinEditingReadOnly())
 		]
-	];
+		];
 }
 
 void FControlRigArgumentLayout::GenerateChildContent(IDetailChildrenBuilder& ChildrenBuilder)
 {
-	if (UEdGraphPin* FoundPin = GetPin())
-	{
-		// Certain types are outlawed at the compiler level, or to keep consistency with variable rules for actors
-		const UClass* ClassObject = Cast<UClass>(FoundPin->PinType.PinSubCategoryObject.Get());
-		const bool bTypeWithNoDefaults = (FoundPin->PinType.PinCategory == UEdGraphSchema_K2::PC_Object) || (FoundPin->PinType.PinCategory == UEdGraphSchema_K2::PC_Class) || (FoundPin->PinType.PinCategory == UEdGraphSchema_K2::PC_Interface)
-			|| (FoundPin->PinType.PinCategory == UEdGraphSchema_K2::PC_SoftObject && ClassObject && ClassObject->IsChildOf(AActor::StaticClass()))
-			|| UEdGraphSchema_K2::IsExecPin(*FoundPin)
-			|| PinPtr.Get()->IsExecuteContext()
-			|| FoundPin->PinType.bIsReference;
-
-		if (!bTypeWithNoDefaults)
-		{
-			DefaultValuePinWidget = FNodeFactory::CreatePinWidget(FoundPin);
-			DefaultValuePinWidget->SetOnlyShowDefaultValue(true);
-			TSharedRef<SWidget> DefaultValueWidget = DefaultValuePinWidget->GetDefaultValueWidget();
-
-			if (DefaultValueWidget != SNullWidget::NullWidget)
-			{
-				ChildrenBuilder.AddCustomRow(LOCTEXT("FunctionArgDetailsDefaultValue", "Default Value"))
-				.NameContent()
-				[
-					SNew(STextBlock)
-					.Text(LOCTEXT("FunctionArgDetailsDefaultValue", "Default Value"))
-					.ToolTipText(LOCTEXT("FunctionArgDetailsDefaultValueParamTooltip", "The default value of the parameter."))
-					.Font(IDetailLayoutBuilder::GetDetailFont())
-				]
-				.ValueContent()
-				.MaxDesiredWidth(512)
-				[
-					DefaultValueWidget
-				];
-			}
-			else
-			{
-				DefaultValuePinWidget.Reset();
-			}
-		}
-	}
+	// we don't show defaults here - we rely on a SControlRigGraphNode widget in the top of the details
 }
 
 void FControlRigArgumentLayout::OnRemoveClicked()
@@ -311,19 +278,17 @@ FReply FControlRigArgumentLayout::OnArgMoveDown()
 
 bool FControlRigArgumentLayout::ShouldPinBeReadOnly(bool bIsEditingPinType/* = false*/) const
 {
-	if (PinPtr.IsValid())
-	{
-		return PinPtr.Get()->IsExecuteContext();
-	}
-	return false;
+	return IsPinEditingReadOnly(bIsEditingPinType);
 }
 
 bool FControlRigArgumentLayout::IsPinEditingReadOnly(bool bIsEditingPinType/* = false*/) const
 {
+	/*
 	if (PinPtr.IsValid())
 	{
 		return PinPtr.Get()->IsExecuteContext();
 	}
+	*/
 	return false;
 }
 
@@ -372,22 +337,6 @@ FEdGraphPinType FControlRigArgumentLayout::OnGetPinInfo() const
 	return FEdGraphPinType();
 }
 
-UEdGraphPin* FControlRigArgumentLayout::GetPin() const
-{
-	if (PinPtr.IsValid() && ControlRigBlueprintPtr.IsValid())
-	{
-		UControlRigBlueprint* Blueprint = ControlRigBlueprintPtr.Get();
-		if(UControlRigGraph* RigGraph = Cast<UControlRigGraph>(Blueprint->GetEdGraph(PinPtr.Get()->GetGraph())))
-		{
-			if (UControlRigGraphNode* RigNode = Cast<UControlRigGraphNode>(RigGraph->FindNodeForModelNodeName(PinPtr.Get()->GetNode()->GetFName())))
-			{
-				return RigNode->FindPin(PinPtr.Get()->GetPinPath());
-			}
-		}
-	}
-	return nullptr;
-}
-
 ECheckBoxState FControlRigArgumentLayout::IsRefChecked() const
 {
 	FEdGraphPinType PinType = OnGetPinInfo();
@@ -448,6 +397,132 @@ void FControlRigArgumentLayout::OnPrePinInfoChange(const FEdGraphPinType& PinTyp
 	// not needed for Control Rig
 }
 
+FControlRigArgumentDefaultNode::FControlRigArgumentDefaultNode(
+	URigVMGraph* InGraph,
+	UControlRigBlueprint* InBlueprint
+)
+	: GraphPtr(InGraph)
+	, ControlRigBlueprintPtr(InBlueprint)
+{
+	if (GraphPtr.IsValid() && ControlRigBlueprintPtr.IsValid())
+	{
+		ControlRigBlueprintPtr.Get()->OnModified().AddRaw(this, &FControlRigArgumentDefaultNode::HandleModifiedEvent);
+
+		if (URigVMLibraryNode* LibraryNode = Cast<URigVMLibraryNode>(GraphPtr->GetOuter()))
+		{
+			if (UControlRigGraph* RigGraph = Cast<UControlRigGraph>(ControlRigBlueprintPtr->GetEdGraph(LibraryNode->GetGraph())))
+			{
+				GraphChangedDelegateHandle = RigGraph->AddOnGraphChangedHandler(
+					FOnGraphChanged::FDelegate::CreateRaw(this, &FControlRigArgumentDefaultNode::OnGraphChanged)
+				);
+			}
+		}
+	}
+}
+
+FControlRigArgumentDefaultNode::~FControlRigArgumentDefaultNode()
+{
+	if (GraphPtr.IsValid() && ControlRigBlueprintPtr.IsValid())
+	{
+		ControlRigBlueprintPtr.Get()->OnModified().RemoveAll(this);
+
+		if (URigVMLibraryNode* LibraryNode = Cast<URigVMLibraryNode>(GraphPtr->GetOuter()))
+		{
+			if (UControlRigGraph* RigGraph = Cast<UControlRigGraph>(ControlRigBlueprintPtr->GetEdGraph(LibraryNode->GetGraph())))
+			{
+				if (GraphChangedDelegateHandle.IsValid())
+				{
+					RigGraph->RemoveOnGraphChangedHandler(GraphChangedDelegateHandle);
+				}
+			}
+		}
+	}
+}
+
+void FControlRigArgumentDefaultNode::GenerateChildContent(IDetailChildrenBuilder& ChildrenBuilder)
+{
+	if (!GraphPtr.IsValid() || !ControlRigBlueprintPtr.IsValid())
+	{
+		return;
+	}
+
+	UControlRigBlueprint* Blueprint = ControlRigBlueprintPtr.Get();
+	URigVMGraph* Graph = GraphPtr.Get();
+	UControlRigGraphNode* ControlRigGraphNode = nullptr;
+	if (URigVMLibraryNode* LibraryNode = Cast<URigVMLibraryNode>(Graph->GetOuter()))
+	{
+		if (UControlRigGraph* RigGraph = Cast<UControlRigGraph>(Blueprint->GetEdGraph(LibraryNode->GetGraph())))
+		{
+			ControlRigGraphNode = Cast<UControlRigGraphNode>(RigGraph->FindNodeForModelNodeName(LibraryNode->GetFName()));
+		}
+	}
+
+	if (ControlRigGraphNode == nullptr)
+	{
+		return;
+	}
+
+	ChildrenBuilder.AddCustomRow(FText::GetEmpty())
+	.WholeRowContent()
+	.MaxDesiredWidth(980.f)
+	[
+		SAssignNew(OwnedNodeWidget, SControlRigGraphNode).GraphNodeObj(ControlRigGraphNode)
+	];
+}
+
+void FControlRigArgumentDefaultNode::OnGraphChanged(const FEdGraphEditAction& InAction)
+{
+	OnRebuildChildren.ExecuteIfBound();
+}
+
+void FControlRigArgumentDefaultNode::HandleModifiedEvent(ERigVMGraphNotifType InNotifType, URigVMGraph* InGraph, UObject* InSubject)
+{
+	if (!GraphPtr.IsValid())
+	{
+		return;
+	}
+
+	URigVMGraph* Graph = GraphPtr.Get();
+	URigVMLibraryNode* LibraryNode = Cast<URigVMLibraryNode>(Graph->GetOuter());
+	if (LibraryNode == nullptr)
+	{
+		return;
+	}
+	if (LibraryNode->GetGraph() != InGraph)
+	{
+		return;
+	}
+
+	switch (InNotifType)
+	{
+		case ERigVMGraphNotifType::PinAdded:
+		case ERigVMGraphNotifType::PinRemoved:
+		case ERigVMGraphNotifType::PinTypeChanged:
+		case ERigVMGraphNotifType::PinRenamed:
+		{
+			URigVMPin* Pin = CastChecked<URigVMPin>(InSubject);
+			if (Pin->GetNode() == LibraryNode)
+			{
+				OnRebuildChildren.ExecuteIfBound();
+			}
+			break;
+		}
+		case ERigVMGraphNotifType::NodeRenamed:
+		{
+			URigVMNode* Node = CastChecked<URigVMNode>(InSubject);
+			if (Node == LibraryNode)
+			{
+				OnRebuildChildren.ExecuteIfBound();
+			}
+			break;
+		}
+		default:
+		{
+			break;
+		}
+	}
+}
+
 TSharedPtr<IDetailCustomization> FControlRigGraphDetails::MakeInstance(TSharedPtr<IBlueprintEditor> InBlueprintEditor)
 {
 	const TArray<UObject*>* Objects = (InBlueprintEditor.IsValid() ? InBlueprintEditor->GetObjectsCurrentlyBeingEdited() : nullptr);
@@ -503,11 +578,13 @@ void FControlRigGraphDetails::CustomizeDetails(IDetailLayoutBuilder& DetailLayou
 	}
 
 	IDetailCategoryBuilder& InputsCategory = DetailLayout.EditCategory("Inputs", LOCTEXT("FunctionDetailsInputs", "Inputs"));
-	TSharedRef<FControlRigArgumentGroupLayout> InputArgumentGroup = MakeShareable(new FControlRigArgumentGroupLayout(Model, Blueprint, true));
+	TSharedRef<FControlRigArgumentGroupLayout> InputArgumentGroup = MakeShareable(new FControlRigArgumentGroupLayout(
+		Model, 
+		Blueprint, 
+		true));
 	InputsCategory.AddCustomBuilder(InputArgumentGroup);
 
 	TSharedRef<SHorizontalBox> InputsHeaderContentWidget = SNew(SHorizontalBox);
-	TWeakPtr<SWidget> WeakInputsHeaderWidget = InputsHeaderContentWidget;
 
 	InputsHeaderContentWidget->AddSlot()
 	.HAlign(HAlign_Right)
@@ -531,11 +608,13 @@ void FControlRigGraphDetails::CustomizeDetails(IDetailLayoutBuilder& DetailLayou
 	InputsCategory.HeaderContent(InputsHeaderContentWidget);
 
 	IDetailCategoryBuilder& OutputsCategory = DetailLayout.EditCategory("Outputs", LOCTEXT("FunctionDetailsOutputs", "Outputs"));
-	TSharedRef<FControlRigArgumentGroupLayout> OutputArgumentGroup = MakeShareable(new FControlRigArgumentGroupLayout(Model, Blueprint, false));
+	TSharedRef<FControlRigArgumentGroupLayout> OutputArgumentGroup = MakeShareable(new FControlRigArgumentGroupLayout(
+		Model, 
+		Blueprint, 
+		false));
 	OutputsCategory.AddCustomBuilder(OutputArgumentGroup);
 
 	TSharedRef<SHorizontalBox> OutputsHeaderContentWidget = SNew(SHorizontalBox);
-	TWeakPtr<SWidget> WeakOutputsHeaderWidget = OutputsHeaderContentWidget;
 
 	OutputsHeaderContentWidget->AddSlot()
 	.HAlign(HAlign_Right)
@@ -557,6 +636,13 @@ void FControlRigGraphDetails::CustomizeDetails(IDetailLayoutBuilder& DetailLayou
 		]
 	];
 	OutputsCategory.HeaderContent(OutputsHeaderContentWidget);
+
+	IDetailCategoryBuilder& DefaultsCategory = DetailLayout.EditCategory("NodeDefaults", LOCTEXT("FunctionDetailsNodeDefaults", "Node Defaults"));
+	TSharedRef<FControlRigArgumentDefaultNode> DefaultsArgumentNode = MakeShareable(new FControlRigArgumentDefaultNode(
+		Model,
+		Blueprint));
+	DefaultsCategory.AddCustomBuilder(DefaultsArgumentNode);
+
 }
 
 bool FControlRigGraphDetails::IsAddNewInputOutputEnabled() const
diff --git a/Engine/Plugins/Experimental/ControlRig/Source/ControlRigEditor/Private/ControlRigGraphDetails.h b/Engine/Plugins/Experimental/ControlRig/Source/ControlRigEditor/Private/ControlRigGraphDetails.h
--- a/Engine/Plugins/Experimental/ControlRig/Source/ControlRigEditor/Private/ControlRigGraphDetails.h
+++ b/Engine/Plugins/Experimental/ControlRig/Source/ControlRigEditor/Private/ControlRigGraphDetails.h
@@ -10,6 +10,7 @@
 #include "ControlRigBlueprint.h"
 #include "ControlRigEditor.h"
 #include "SGraphPin.h"
+#include "Graph/SControlRigGraphNode.h"
 
 class FMobilityCustomization;
 class IDetailLayoutBuilder;
@@ -17,7 +18,10 @@ class IDetailLayoutBuilder;
 class FControlRigArgumentGroupLayout : public IDetailCustomNodeBuilder, public TSharedFromThis<FControlRigArgumentGroupLayout>
 {
 public:
-	FControlRigArgumentGroupLayout(URigVMGraph* InGraph, UControlRigBlueprint* InBlueprint, bool bInputs);
+	FControlRigArgumentGroupLayout(
+		URigVMGraph* InGraph, 
+		UControlRigBlueprint* InBlueprint, 
+		bool bInputs);
 	virtual ~FControlRigArgumentGroupLayout();
 
 private:
@@ -45,7 +49,10 @@ class FControlRigArgumentLayout : public IDetailCustomNodeBuilder, public TShare
 {
 public:
 
-	FControlRigArgumentLayout(URigVMPin* InPin, URigVMGraph* InGraph, UControlRigBlueprint* InBlueprint)
+	FControlRigArgumentLayout(
+		URigVMPin* InPin, 
+		URigVMGraph* InGraph, 
+		UControlRigBlueprint* InBlueprint)
 		: PinPtr(InPin)
 		, GraphPtr(InGraph)
 		, ControlRigBlueprintPtr(InBlueprint)
@@ -83,9 +90,6 @@ private:
 	void PinInfoChanged(const FEdGraphPinType& PinType);
 	void OnPrePinInfoChange(const FEdGraphPinType& PinType);
 
-	/** Returns the graph pin representing this variable */
-	UEdGraphPin* GetPin() const;
-
 	/** Returns whether the "Pass-by-Reference" checkbox is checked or not */
 	ECheckBoxState IsRefChecked() const;
 
@@ -105,11 +109,40 @@ private:
 
 	/** Holds a weak pointer to the argument name widget, used for error notifications */
 	TWeakPtr<SEditableTextBox> ArgumentNameWidget;
+};
+
+class FControlRigArgumentDefaultNode : public IDetailCustomNodeBuilder, public TSharedFromThis<FControlRigArgumentDefaultNode>
+{
+public:
+	FControlRigArgumentDefaultNode(
+		URigVMGraph* InGraph,
+		UControlRigBlueprint* InBlueprint
+	);
+	virtual ~FControlRigArgumentDefaultNode();
+
+private:
+	/** IDetailCustomNodeBuilder Interface*/
+	virtual void SetOnRebuildChildren(FSimpleDelegate InOnRegenerateChildren) override { OnRebuildChildren = InOnRegenerateChildren; }
+	virtual void GenerateHeaderRowContent(FDetailWidgetRow& NodeRow) override {}
+	virtual void GenerateChildContent(IDetailChildrenBuilder& ChildrenBuilder) override;
+	virtual void Tick(float DeltaTime) override {}
+	virtual bool RequiresTick() const override { return false; }
+	virtual FName GetName() const override { return NAME_None; }
+	virtual bool InitiallyCollapsed() const override { return false; }
 
-	/** The SGraphPin widget created to show/edit default value */
-	TSharedPtr<SGraphPin> DefaultValuePinWidget;
+private:
+
+	void OnGraphChanged(const FEdGraphEditAction& InAction);
+	void HandleModifiedEvent(ERigVMGraphNotifType InNotifType, URigVMGraph* InGraph, UObject* InSubject);
+
+	TWeakObjectPtr<URigVMGraph> GraphPtr;
+	TWeakObjectPtr<UControlRigBlueprint> ControlRigBlueprintPtr;
+	FSimpleDelegate OnRebuildChildren;
+	TSharedPtr<SControlRigGraphNode> OwnedNodeWidget;
+	FDelegateHandle GraphChangedDelegateHandle;
 };
 
+
 /** Customization for editing Control Rig graphs */
 class FControlRigGraphDetails : public IDetailCustomization
 {
@@ -135,9 +168,9 @@ private:
 	/** The Blueprint editor we are embedded in */
 	TWeakPtr<IControlRigEditor> ControlRigEditorPtr;
 
-	/** The blueprint we are editing */
-	TWeakObjectPtr<UControlRigBlueprint> ControlRigBlueprintPtr;
-
 	/** The graph we are editing */
 	TWeakObjectPtr<UControlRigGraph> GraphPtr;
+
+	/** The blueprint we are editing */
+	TWeakObjectPtr<UControlRigBlueprint> ControlRigBlueprintPtr;
 };
\ No newline at end of file
diff --git a/Engine/Plugins/Experimental/ControlRig/Source/ControlRigEditor/Private/Editor/ControlRigEditor.cpp b/Engine/Plugins/Experimental/ControlRig/Source/ControlRigEditor/Private/Editor/ControlRigEditor.cpp
--- a/Engine/Plugins/Experimental/ControlRig/Source/ControlRigEditor/Private/Editor/ControlRigEditor.cpp
+++ b/Engine/Plugins/Experimental/ControlRig/Source/ControlRigEditor/Private/Editor/ControlRigEditor.cpp
@@ -216,7 +216,10 @@ void FControlRigEditor::InitControlRigEditor(const EToolkitMode::Type Mode, cons
 
 	for (UBlueprint* Blueprint : ControlRigBlueprints)
 	{
-		for (UEdGraph* Graph : Blueprint->UbergraphPages)
+		TArray<UEdGraph*> EdGraphs;
+		Blueprint->GetAllGraphs(EdGraphs);
+
+		for (UEdGraph* Graph : EdGraphs)
 		{
 			UControlRigGraph* RigGraph = Cast<UControlRigGraph>(Graph);
 			if (RigGraph == nullptr)
@@ -1760,12 +1763,37 @@ void FControlRigEditor::CreateDefaultTabContents(const TArray<UBlueprint*>& InBl
 	FBlueprintEditor::CreateDefaultTabContents(InBlueprints);
 }
 
+void FControlRigEditor::NewDocument_OnClicked(ECreatedDocumentType GraphType)
+{
+	if (GraphType != FBlueprintEditor::CGT_NewFunctionGraph)
+	{
+		return;
+	}
+
+	if (UControlRigBlueprint* Blueprint = GetControlRigBlueprint())
+	{
+		if (URigVMController* Controller = Blueprint->GetOrCreateController(Blueprint->GetLocalFunctionLibrary()))
+		{
+			if (URigVMLibraryNode* FunctionNode = Controller->AddFunctionToLibrary(TEXT("New Function"), true))
+			{
+				if (UEdGraph* NewGraph = Blueprint->GetEdGraph(FunctionNode->GetContainedGraph()))
+				{
+					OpenDocument(NewGraph, FDocumentTracker::OpenNewDocument);
+					RenameNewlyAddedAction(FunctionNode->GetFName());
+				}
+
+			}
+		}
+	}
+}
+
 bool FControlRigEditor::IsSectionVisible(NodeSectionID::Type InSectionID) const
 {
 	switch (InSectionID)
 	{
 		case NodeSectionID::GRAPH:
 		case NodeSectionID::VARIABLE:
+		case NodeSectionID::FUNCTION:
 		{
 			return true;
 		}
@@ -2628,15 +2656,16 @@ void FControlRigEditor::CacheNameLists()
 
 	if (UControlRigBlueprint* ControlRigBP = GetControlRigBlueprint())
 	{
-		// make sure the bone name list is up 2 date for the editor graph
-		for (UEdGraph* Graph : ControlRigBP->UbergraphPages)
+		TArray<UEdGraph*> EdGraphs;
+		ControlRigBP->GetAllGraphs(EdGraphs);
+
+		for (UEdGraph* Graph : EdGraphs)
 		{
 			UControlRigGraph* RigGraph = Cast<UControlRigGraph>(Graph);
 			if (RigGraph == nullptr)
 			{
 				continue;
 			}
-
 			RigGraph->CacheNameLists(&ControlRigBP->HierarchyContainer, &ControlRigBP->DrawContainer);
 		}
 	}
@@ -3158,7 +3187,10 @@ void FControlRigEditor::OnRigElementRemoved(FRigHierarchyContainer* Container, c
 	FString RemovedElementName = InKey.Name.ToString();
 	ERigElementType RemovedElementType = InKey.Type;
 
-	for (UEdGraph* Graph : Blueprint->UbergraphPages)
+	TArray<UEdGraph*> EdGraphs;
+	Blueprint->GetAllGraphs(EdGraphs);
+
+	for (UEdGraph* Graph : EdGraphs)
 	{
 		UControlRigGraph* RigGraph = Cast<UControlRigGraph>(Graph);
 		if (RigGraph == nullptr)
@@ -3236,7 +3268,10 @@ void FControlRigEditor::OnRigElementRenamed(FRigHierarchyContainer* Container, E
 	FString OldNameStr = InOldName.ToString();
 	FString NewNameStr = InNewName.ToString();
 
-	for (UEdGraph* Graph : Blueprint->UbergraphPages)
+	TArray<UEdGraph*> EdGraphs;
+	Blueprint->GetAllGraphs(EdGraphs);
+
+	for (UEdGraph* Graph : EdGraphs)
 	{
 		UControlRigGraph* RigGraph = Cast<UControlRigGraph>(Graph);
 		if (RigGraph == nullptr)
@@ -4135,6 +4170,15 @@ void FControlRigEditor::OnNodeDoubleClicked(UControlRigBlueprint* InBlueprint, U
 	}
 }
 
+bool FControlRigEditor::OnActionMatchesName(FEdGraphSchemaAction* InAction, const FName& InName) const
+{
+	if (InAction->GetMenuDescription().ToString() == InName.ToString())
+	{
+		return true;
+	}
+	return false;
+}
+
 void FControlRigEditor::UpdateGraphCompilerErrors()
 {
 	DECLARE_SCOPE_HIERARCHICAL_COUNTER_FUNC()
@@ -4152,7 +4196,10 @@ void FControlRigEditor::UpdateGraphCompilerErrors()
 			return;
 		}
 
-		for (UEdGraph* Graph : Blueprint->UbergraphPages)
+		TArray<UEdGraph*> EdGraphs;
+		Blueprint->GetAllGraphs(EdGraphs);
+
+		for (UEdGraph* Graph : EdGraphs)
 		{
 			UControlRigGraph* RigGraph = Cast<UControlRigGraph>(Graph);
 			if (RigGraph == nullptr)
diff --git a/Engine/Plugins/Experimental/ControlRig/Source/ControlRigEditor/Private/Editor/ControlRigEditor.h b/Engine/Plugins/Experimental/ControlRig/Source/ControlRigEditor/Private/Editor/ControlRigEditor.h
--- a/Engine/Plugins/Experimental/ControlRig/Source/ControlRigEditor/Private/Editor/ControlRigEditor.h
+++ b/Engine/Plugins/Experimental/ControlRig/Source/ControlRigEditor/Private/Editor/ControlRigEditor.h
@@ -190,6 +190,7 @@ protected:
 	virtual void SaveAssetAs_Execute() override;
 	virtual bool IsInAScriptingMode() const override { return true; }
 	virtual void CreateDefaultTabContents(const TArray<UBlueprint*>& InBlueprints) override;
+	virtual void NewDocument_OnClicked(ECreatedDocumentType GraphType) override;
 	virtual bool IsSectionVisible(NodeSectionID::Type InSectionID) const override;
 	virtual FGraphAppearanceInfo GetGraphAppearance(class UEdGraph* InGraph) const override;
 	virtual bool IsEditable(UEdGraph* InGraph) const override;
@@ -306,6 +307,8 @@ private:
 
 	 void OnNodeDoubleClicked(UControlRigBlueprint* InBlueprint, URigVMNode* InNode);
 
+	 virtual bool OnActionMatchesName(FEdGraphSchemaAction* InAction, const FName& InName) const override;
+
 protected:
 
 	/** Toolbox hosting widget */
diff --git a/Engine/Source/Developer/RigVMDeveloper/Private/RigVMModel/RigVMController.cpp b/Engine/Source/Developer/RigVMDeveloper/Private/RigVMModel/RigVMController.cpp
--- a/Engine/Source/Developer/RigVMDeveloper/Private/RigVMModel/RigVMController.cpp
+++ b/Engine/Source/Developer/RigVMDeveloper/Private/RigVMModel/RigVMController.cpp
@@ -282,7 +282,7 @@ URigVMUnitNode* URigVMController::AddUnitNode(UScriptStruct* InScriptStruct, con
 
 			if (!StructMemory->GetEventName().IsNone())
 			{
-				ReportError(TEXT("Event nodes can only be added to top level graphs."));
+				ReportAndNotifyError(TEXT("Event nodes can only be added to top level graphs."));
 				return nullptr;
 			}
 		}
@@ -3002,6 +3002,54 @@ bool URigVMController::RemoveNode(URigVMNode* InNode, bool bSetupUndoRedo, bool
 
 	Notify(ERigVMGraphNotifType::NodeRemoved, InNode);
 
+	if (URigVMLibraryNode* LibraryNode = Cast<URigVMLibraryNode>(InNode))
+	{
+		if(URigVMFunctionReferenceNode* FunctionReferenceNode = Cast<URigVMFunctionReferenceNode>(LibraryNode))
+		{
+			if (URigVMFunctionLibrary* FunctionLibrary = Cast<URigVMFunctionLibrary>(FunctionReferenceNode->GetLibrary()))
+			{
+				FRigVMFunctionReferenceArray* References = FunctionLibrary->FunctionReferences.Find(FunctionReferenceNode->GetReferencedNode());
+				if (References)
+				{
+					References->FunctionReferences.RemoveAll(
+						[FunctionReferenceNode](TSoftObjectPtr<URigVMFunctionReferenceNode>& FunctionReferencePtr) {
+
+							/*
+							if (!FunctionReferencePtr.IsValid())
+							{
+								FunctionReferencePtr.LoadSynchronous();
+							}
+							*/
+
+							if (!FunctionReferencePtr.IsValid())
+							{
+								return true;
+							}
+							return FunctionReferencePtr.Get() == FunctionReferenceNode;
+					});
+				}
+			}
+		}
+		else if (URigVMFunctionLibrary* FunctionLibrary = Cast<URigVMFunctionLibrary>(LibraryNode->GetGraph()))
+		{
+			const FRigVMFunctionReferenceArray* FunctionReferencesPtr = FunctionLibrary->FunctionReferences.Find(LibraryNode);
+			if (FunctionReferencesPtr)
+			{
+				for (const TSoftObjectPtr<URigVMFunctionReferenceNode>& FunctionReferencePtr : FunctionReferencesPtr->FunctionReferences)
+				{
+					if (FunctionReferencePtr.IsValid())
+					{
+						FunctionReferencePtr->ReferencedNodePtr.Reset();
+
+						FRigVMControllerGraphGuard GraphGuard(this, FunctionReferencePtr->GetGraph(), false);
+						RepopulatePinsOnNode(FunctionReferencePtr.Get(), false, true);
+					}
+				}
+			}
+			FunctionLibrary->FunctionReferences.Remove(LibraryNode);
+		}
+	}
+
 	if (URigVMVariableNode* VariableNode = Cast<URigVMVariableNode>(InNode))
 	{
 		Notify(ERigVMGraphNotifType::VariableRemoved, VariableNode);
@@ -3083,6 +3131,26 @@ bool URigVMController::RenameNode(URigVMNode* InNode, const FName& InNewName, bo
 		Notify(ERigVMGraphNotifType::LinkAdded, Link);
 	}
 
+	if(URigVMLibraryNode* LibraryNode = Cast<URigVMLibraryNode>(InNode))
+	{
+		if (URigVMFunctionLibrary* FunctionLibrary = Cast<URigVMFunctionLibrary>(LibraryNode->GetGraph()))
+		{
+			FRigVMFunctionReferenceArray* ReferencesEntry = FunctionLibrary->FunctionReferences.Find(LibraryNode);
+			if (ReferencesEntry)
+			{
+				for (TSoftObjectPtr<URigVMFunctionReferenceNode> FunctionReferencePtr : ReferencesEntry->FunctionReferences)
+				{
+					// only update valid, living references
+					if (FunctionReferencePtr.IsValid())
+					{
+						FRigVMControllerGraphGuard GraphGuard(this, FunctionReferencePtr->GetGraph(), false);
+						RenameNode(FunctionReferencePtr.Get(), InNewName, false);
+					}
+				}
+			}
+		}
+	}
+
 	if (bSetupUndoRedo)
 	{
 		ActionStack->EndAction(Action);
@@ -5059,13 +5127,15 @@ FName URigVMController::AddExposedPin(const FName& InPinName, ERigVMPinDirection
 		Notify(ERigVMGraphNotifType::PinAdded, Pin);
 	}
 
-	if (bSetupUndoRedo && !InDefaultValue.IsEmpty())
+	if (!InDefaultValue.IsEmpty())
 	{
+		FRigVMControllerGraphGuard GraphGuard(this, Pin->GetGraph(), false);
 		SetPinDefaultValue(Pin, InDefaultValue, true, bSetupUndoRedo, false);
 	}
 
 	RefreshFunctionPins(Graph->GetEntryNode(), true);
 	RefreshFunctionPins(Graph->GetReturnNode(), true);
+	RefreshFunctionReferences(LibraryNode, false);
 
 	if (bSetupUndoRedo)
 	{
@@ -5120,6 +5190,7 @@ bool URigVMController::RemoveExposedPin(const FName& InPinName, bool bSetupUndoR
 
 	RefreshFunctionPins(Graph->GetEntryNode(), true);
 	RefreshFunctionPins(Graph->GetReturnNode(), true);
+	RefreshFunctionReferences(LibraryNode, false);
 
 	if (bSetupUndoRedo)
 	{
@@ -5244,6 +5315,26 @@ bool URigVMController::RenameExposedPin(const FName& InOldPinName, const FName&
 		}
 	}
 
+	if (URigVMFunctionLibrary* FunctionLibrary = Cast<URigVMFunctionLibrary>(LibraryNode->GetGraph()))
+	{
+		FRigVMFunctionReferenceArray* ReferencesEntry = FunctionLibrary->FunctionReferences.Find(LibraryNode);
+		if (ReferencesEntry)
+		{
+			for (TSoftObjectPtr<URigVMFunctionReferenceNode> FunctionReferencePtr : ReferencesEntry->FunctionReferences)
+			{
+				// only update valid, living references
+				if (FunctionReferencePtr.IsValid())
+				{
+					if (URigVMPin* EntryPin = FunctionReferencePtr->FindPin(InOldPinName.ToString()))
+					{
+						FRigVMControllerGraphGuard GraphGuard(this, FunctionReferencePtr->GetGraph(), false);
+						Local::RenamePin(this, EntryPin, PinName);
+					}
+				}
+			}
+		}
+	}
+
 	if (bSetupUndoRedo)
 	{
 		ActionStack->EndAction(Action);
@@ -5318,6 +5409,26 @@ bool URigVMController::ChangeExposedPinType(const FName& InPinName, const FStrin
 		}
 	}
 
+	if (URigVMFunctionLibrary* FunctionLibrary = Cast<URigVMFunctionLibrary>(LibraryNode->GetGraph()))
+	{
+		FRigVMFunctionReferenceArray* ReferencesEntry = FunctionLibrary->FunctionReferences.Find(LibraryNode);
+		if (ReferencesEntry)
+		{
+			for (TSoftObjectPtr<URigVMFunctionReferenceNode> FunctionReferencePtr : ReferencesEntry->FunctionReferences)
+			{
+				// only update valid, living references
+				if (FunctionReferencePtr.IsValid())
+				{
+					if (URigVMPin* ReferencedNodePin = FunctionReferencePtr->FindPin(Pin->GetName()))
+					{
+						FRigVMControllerGraphGuard GraphGuard(this, FunctionReferencePtr->GetGraph(), false);
+						ChangePinType(ReferencedNodePin, InCPPType, InCPPTypeObjectPath, bSetupUndoRedo);
+					}
+				}
+			}
+		}
+	}
+
 	if (bSetupUndoRedo)
 	{
 		ActionStack->EndAction(Action);
@@ -5348,6 +5459,23 @@ URigVMFunctionReferenceNode* URigVMController::AddFunctionReferenceNode(URigVMLi
 		return nullptr;
 	}
 
+	if (!InFunctionDefinition->GetGraph()->IsA<URigVMFunctionLibrary>())
+	{
+		ReportAndNotifyError(TEXT("Cannot use the function definition for a function reference node."));
+		return nullptr;
+	}
+
+	URigVMLibraryNode* ParentLibraryNode = Cast<URigVMLibraryNode>(Graph->GetOuter());
+	while (ParentLibraryNode)
+	{
+		if (ParentLibraryNode == InFunctionDefinition)
+		{
+			ReportAndNotifyError(TEXT("You cannot place functions inside of itself or an indirect recursion."));
+			return nullptr;
+		}
+		ParentLibraryNode = Cast<URigVMLibraryNode>(ParentLibraryNode->GetGraph()->GetOuter());
+	}
+
 	FString NodeName = GetValidNodeName(InNodeName.IsEmpty() ? InFunctionDefinition->GetName() : InNodeName);
 	URigVMFunctionReferenceNode* FunctionRefNode = NewObject<URigVMFunctionReferenceNode>(Graph, *NodeName);
 	FunctionRefNode->Position = InNodePosition;
@@ -5356,6 +5484,25 @@ URigVMFunctionReferenceNode* URigVMController::AddFunctionReferenceNode(URigVMLi
 
 	RepopulatePinsOnNode(FunctionRefNode, false, false);
 
+	Notify(ERigVMGraphNotifType::NodeAdded, FunctionRefNode);
+
+	if (URigVMFunctionLibrary* FunctionLibrary = InFunctionDefinition->GetLibrary())
+	{
+		FunctionLibrary->FunctionReferences.FindOrAdd(InFunctionDefinition).FunctionReferences.Add(FunctionRefNode);
+	}
+
+	for (URigVMPin* SourcePin : InFunctionDefinition->Pins)
+	{
+		if (URigVMPin* TargetPin = FunctionRefNode->FindPin(SourcePin->GetName()))
+		{
+			FString DefaultValue = SourcePin->GetDefaultValue();
+			if (!DefaultValue.IsEmpty())
+			{
+				SetPinDefaultValue(TargetPin, DefaultValue, true, false, false);
+			}
+		}
+	}
+
 	if (bSetupUndoRedo)
 	{
 		FRigVMInverseAction InverseAction;
@@ -5369,7 +5516,7 @@ URigVMFunctionReferenceNode* URigVMController::AddFunctionReferenceNode(URigVMLi
 	return FunctionRefNode;
 }
 
-URigVMLibraryNode* URigVMController::AddFunctionToLibrary(const FName& InFunctionName, const FVector2D& InNodePosition, bool bSetupUndoRedo)
+URigVMLibraryNode* URigVMController::AddFunctionToLibrary(const FName& InFunctionName, bool bMutable, const FVector2D& InNodePosition, bool bSetupUndoRedo)
 {
 	if (!IsValidGraph())
 	{
@@ -5391,6 +5538,18 @@ URigVMLibraryNode* URigVMController::AddFunctionToLibrary(const FName& InFunctio
 	CollapseNode->Position = InNodePosition;
 	Graph->Nodes.Add(CollapseNode);
 
+	if (bMutable)
+	{
+		URigVMPin* ExecutePin = NewObject<URigVMPin>(CollapseNode, FRigVMStruct::ExecuteContextName);
+		ExecutePin->CPPType = FString::Printf(TEXT("F%s"), *ExecuteContextStruct->GetName());
+		ExecutePin->CPPTypeObject = ExecuteContextStruct;
+		ExecutePin->CPPTypeObjectPath = *ExecutePin->CPPTypeObject->GetPathName();
+		ExecutePin->Direction = ERigVMPinDirection::IO;
+		CollapseNode->Pins.Add(ExecutePin);
+	}
+
+	Notify(ERigVMGraphNotifType::NodeAdded, CollapseNode);
+
 	{
 		FRigVMControllerGraphGuard GraphGuard(this, CollapseNode->GetContainedGraph(), false);
 
@@ -5405,6 +5564,11 @@ URigVMLibraryNode* URigVMController::AddFunctionToLibrary(const FName& InFunctio
 		ReturnNode->Position = FVector2D(250.f, 0.f);
 		RefreshFunctionPins(ReturnNode, false);
 		Notify(ERigVMGraphNotifType::NodeAdded, ReturnNode);
+
+		if (bMutable)
+		{
+			AddLink(EntryNode->FindPin(FRigVMStruct::ExecuteContextName.ToString()), ReturnNode->FindPin(FRigVMStruct::ExecuteContextName.ToString()), false);
+		}
 	}
 
 	if (bSetupUndoRedo)
@@ -7039,6 +7203,10 @@ void URigVMController::RepopulatePinsOnNode(URigVMNode* InNode, bool bFollowCore
 	{
 		if (URigVMLibraryNode* ReferencedNode = FunctionRefNode->GetReferencedNode())
 		{
+			// we want to make sure notify the graph of a potential name change
+			// when repopulating the function ref node
+			Notify(ERigVMGraphNotifType::NodeRenamed, FunctionRefNode);
+
 			TArray<URigVMPin*> Pins = InNode->GetPins();
 			for (URigVMPin* Pin : Pins)
 			{
@@ -7046,6 +7214,8 @@ void URigVMController::RepopulatePinsOnNode(URigVMNode* InNode, bool bFollowCore
 			}
 			InNode->Pins.Reset();
 
+			TMap<FString, FPinState> ReferencedPinStates = GetPinStates(ReferencedNode);
+
 			for (URigVMPin* ReferencedPin : ReferencedNode->Pins)
 			{
 				URigVMPin* NewPin = NewObject<URigVMPin>(InNode, ReferencedPin->GetFName());
@@ -7060,6 +7230,8 @@ void URigVMController::RepopulatePinsOnNode(URigVMNode* InNode, bool bFollowCore
 
 				Notify(ERigVMGraphNotifType::PinAdded, NewPin);
 			}
+
+			ApplyPinStates(InNode, ReferencedPinStates);
 		}
 	}
 
@@ -7665,3 +7837,29 @@ TArray<FRigVMExternalVariable> URigVMController::GetExternalVariables()
 	}
 	return TArray<FRigVMExternalVariable>();
 }
+
+void URigVMController::RefreshFunctionReferences(URigVMLibraryNode* InFunctionDefinition, bool bSetupUndoRedo)
+{
+	check(InFunctionDefinition);
+
+	if (URigVMFunctionLibrary* FunctionLibrary = Cast<URigVMFunctionLibrary>(InFunctionDefinition->GetGraph()))
+	{
+		FRigVMFunctionReferenceArray* ReferencesEntry = FunctionLibrary->FunctionReferences.Find(InFunctionDefinition);
+		if (ReferencesEntry)
+		{
+			for (TSoftObjectPtr<URigVMFunctionReferenceNode> FunctionReferencePtr : ReferencesEntry->FunctionReferences)
+			{
+				// only update valid, living references
+				if (FunctionReferencePtr.IsValid())
+				{
+					FRigVMControllerGraphGuard GraphGuard(this, FunctionReferencePtr->GetGraph(), bSetupUndoRedo);
+
+					TArray<URigVMLink*> Links = FunctionReferencePtr->GetLinks();
+					DetachLinksFromPinObjects(&Links, true);
+					RepopulatePinsOnNode(FunctionReferencePtr.Get(), false, true);
+					ReattachLinksToPinObjects(false, &Links, true);
+				}
+			}
+		}
+	}
+}
diff --git a/Engine/Source/Developer/RigVMDeveloper/Private/RigVMModel/RigVMFunctionLibrary.cpp b/Engine/Source/Developer/RigVMDeveloper/Private/RigVMModel/RigVMFunctionLibrary.cpp
--- a/Engine/Source/Developer/RigVMDeveloper/Private/RigVMModel/RigVMFunctionLibrary.cpp
+++ b/Engine/Source/Developer/RigVMDeveloper/Private/RigVMModel/RigVMFunctionLibrary.cpp
@@ -7,6 +7,11 @@ URigVMFunctionLibrary::URigVMFunctionLibrary()
 {
 }
 
+FString URigVMFunctionLibrary::GetNodePath() const
+{
+	return FString::Printf(TEXT("FunctionLibrary::%s"), *Super::GetNodePath());
+}
+
 TArray<URigVMLibraryNode*> URigVMFunctionLibrary::GetFunctions() const
 {
 	TArray<URigVMLibraryNode*> Functions;
@@ -23,6 +28,11 @@ TArray<URigVMLibraryNode*> URigVMFunctionLibrary::GetFunctions() const
 
 URigVMLibraryNode* URigVMFunctionLibrary::FindFunction(const FName& InFunctionName) const
 {
-	return Cast<URigVMLibraryNode>(FindNodeByName(InFunctionName));
+	FString FunctionNameStr = InFunctionName.ToString();
+	if (FunctionNameStr.StartsWith(TEXT("FunctionLibrary::|")))
+	{
+		FunctionNameStr.RightChopInline(18);
+	}
+	return Cast<URigVMLibraryNode>(FindNodeByName(*FunctionNameStr));
 }
 
diff --git a/Engine/Source/Developer/RigVMDeveloper/Private/RigVMModel/RigVMGraph.cpp b/Engine/Source/Developer/RigVMDeveloper/Private/RigVMModel/RigVMGraph.cpp
--- a/Engine/Source/Developer/RigVMDeveloper/Private/RigVMModel/RigVMGraph.cpp
+++ b/Engine/Source/Developer/RigVMDeveloper/Private/RigVMModel/RigVMGraph.cpp
@@ -29,12 +29,12 @@ TArray<URigVMGraph*> URigVMGraph::GetContainedGraphs(bool bRecursive) const
 	TArray<URigVMGraph*> Graphs;
 	for (URigVMNode* Node : GetNodes())
 	{
-		if (URigVMLibraryNode* LibraryNode = Cast<URigVMLibraryNode>(Node))
+		if (URigVMCollapseNode* CollapseNode = Cast<URigVMCollapseNode>(Node))
 		{
-			Graphs.Add(LibraryNode->GetContainedGraph());
+			Graphs.Add(CollapseNode->GetContainedGraph());
 			if (bRecursive)
 			{
-				Graphs.Append(LibraryNode->GetContainedGraph()->GetContainedGraphs(true));
+				Graphs.Append(CollapseNode->GetContainedGraph()->GetContainedGraphs(true));
 			}
 		}
 	}
diff --git a/Engine/Source/Developer/RigVMDeveloper/Public/RigVMModel/Nodes/RigVMFunctionReferenceNode.h b/Engine/Source/Developer/RigVMDeveloper/Public/RigVMModel/Nodes/RigVMFunctionReferenceNode.h
--- a/Engine/Source/Developer/RigVMDeveloper/Public/RigVMModel/Nodes/RigVMFunctionReferenceNode.h
+++ b/Engine/Source/Developer/RigVMDeveloper/Public/RigVMModel/Nodes/RigVMFunctionReferenceNode.h
@@ -3,10 +3,11 @@
 #pragma once
 
 #include "RigVMModel/RigVMGraph.h"
-#include "RigVMModel/RigVMFunctionLibrary.h"
 #include "RigVMModel/Nodes/RigVMLibraryNode.h"
 #include "RigVMFunctionReferenceNode.generated.h"
 
+class URigVMFunctionLibrary;
+
 /**
  * The Function Reference Node is a library node which references
  * a library node from a separate function library graph.
diff --git a/Engine/Source/Developer/RigVMDeveloper/Public/RigVMModel/RigVMController.h b/Engine/Source/Developer/RigVMDeveloper/Public/RigVMModel/RigVMController.h
--- a/Engine/Source/Developer/RigVMDeveloper/Public/RigVMModel/RigVMController.h
+++ b/Engine/Source/Developer/RigVMDeveloper/Public/RigVMModel/RigVMController.h
@@ -497,7 +497,7 @@ public:
 
 	// Adds a function definition to a function library graph
 	UFUNCTION(BlueprintCallable, Category = RigVMController)
-	URigVMLibraryNode* AddFunctionToLibrary(const FName& InFunctionName, const FVector2D& InNodePosition = FVector2D::ZeroVector, bool bSetupUndoRedo = true);
+	URigVMLibraryNode* AddFunctionToLibrary(const FName& InFunctionName, bool bMutable, const FVector2D& InNodePosition = FVector2D::ZeroVector, bool bSetupUndoRedo = true);
 
 	// Removes a function from a function library graph
 	UFUNCTION(BlueprintCallable, Category = RigVMController)
@@ -634,6 +634,8 @@ private:
 	FRigVMExternalVariable GetExternalVariableByName(const FName& InExternalVariableName);
 	TArray<FRigVMExternalVariable> GetExternalVariables();
 
+	void RefreshFunctionReferences(URigVMLibraryNode* InFunctionDefinition, bool bSetupUndoRedo);
+
 	UPROPERTY(transient)
 	TArray<URigVMGraph*> Graphs;
 
diff --git a/Engine/Source/Developer/RigVMDeveloper/Public/RigVMModel/RigVMFunctionLibrary.h b/Engine/Source/Developer/RigVMDeveloper/Public/RigVMModel/RigVMFunctionLibrary.h
--- a/Engine/Source/Developer/RigVMDeveloper/Public/RigVMModel/RigVMFunctionLibrary.h
+++ b/Engine/Source/Developer/RigVMDeveloper/Public/RigVMModel/RigVMFunctionLibrary.h
@@ -4,8 +4,30 @@
 
 #include "RigVMGraph.h"
 #include "RigVMModel/Nodes/RigVMLibraryNode.h"
+#include "RigVMModel/Nodes/RigVMFunctionReferenceNode.h"
 #include "RigVMFunctionLibrary.generated.h"
 
+USTRUCT(BlueprintType)
+struct RIGVMDEVELOPER_API FRigVMFunctionReferenceArray
+{
+	GENERATED_BODY()
+
+	// Resets the data structure and maintains all storage.
+	void Reset() { FunctionReferences.Reset();  }
+
+	// Returns true if a given function reference index is valid.
+	bool IsValidIndex(int32 InIndex) const { return FunctionReferences.IsValidIndex(InIndex); }
+
+	// Returns the number of reference functions
+	FORCEINLINE int32 Num() const { return FunctionReferences.Num(); }
+
+	// const accessor for an function reference given its index
+	FORCEINLINE const TSoftObjectPtr<URigVMFunctionReferenceNode>& operator[](int32 InIndex) const { return FunctionReferences[InIndex]; }
+
+	UPROPERTY()
+	TArray< TSoftObjectPtr<URigVMFunctionReferenceNode> > FunctionReferences;
+};
+
 /**
  * The Function Library is a graph used only to store
  * the sub graphs used for functions.
@@ -20,6 +42,10 @@ public:
 	// Default constructor
 	URigVMFunctionLibrary();
 
+	// URigVMGraph interface
+	virtual FString GetNodePath() const override;
+	// end URigVMGraph interface
+
 	// Returns all of the stored functions
 	UFUNCTION(BlueprintCallable, Category = RigVMGraph)
 	TArray<URigVMLibraryNode*> GetFunctions() const;
@@ -30,6 +56,9 @@ public:
 
 private:
 
+	UPROPERTY()
+	TMap< URigVMLibraryNode*, FRigVMFunctionReferenceArray > FunctionReferences;
+
 	friend class URigVMController;
 	friend class URigVMCompiler;
 };
diff --git a/Engine/Source/Developer/RigVMDeveloper/Public/RigVMModel/RigVMGraph.h b/Engine/Source/Developer/RigVMDeveloper/Public/RigVMModel/RigVMGraph.h
--- a/Engine/Source/Developer/RigVMDeveloper/Public/RigVMModel/RigVMGraph.h
+++ b/Engine/Source/Developer/RigVMDeveloper/Public/RigVMModel/RigVMGraph.h
@@ -62,7 +62,7 @@ public:
 
 	// Returns the path of this graph as defined by its invoking nodes
 	UFUNCTION(BlueprintCallable, Category = RigVMGraph)
-	FString GetNodePath() const;
+	virtual FString GetNodePath() const;
 
 	// Returns the name of this graph (as defined by the node path)
 	UFUNCTION(BlueprintCallable, Category = RigVMGraph)
diff --git a/Engine/Source/Editor/GraphEditor/Private/DragConnection.cpp b/Engine/Source/Editor/GraphEditor/Private/DragConnection.cpp
--- a/Engine/Source/Editor/GraphEditor/Private/DragConnection.cpp
+++ b/Engine/Source/Editor/GraphEditor/Private/DragConnection.cpp
@@ -261,6 +261,8 @@ FReply FDragConnection::DroppedOnNode(FVector2D ScreenPosition, FVector2D GraphP
 		{
 			for (UEdGraphPin* SourcePin : ValidSourcePins)
 			{
+				SourcePin->GetSchema()->SetPinBeingDroppedOnNode(SourcePin);
+
 				// Check for pin drop support
 				FText ResponseText;
 				if (SourcePin->GetOwningNode() != NodeOver && SourcePin->GetSchema()->SupportsDropPinOnNode(NodeOver, SourcePin->PinType, SourcePin->Direction, ResponseText))
@@ -288,6 +290,8 @@ FReply FDragConnection::DroppedOnNode(FVector2D ScreenPosition, FVector2D GraphP
 				{
 					bHandledPinDropOnNode = true;
 				}
+
+				SourcePin->GetSchema()->SetPinBeingDroppedOnNode(nullptr);
 			}
 		}
 	}
diff --git a/Engine/Source/Editor/GraphEditor/Private/SGraphNode.cpp b/Engine/Source/Editor/GraphEditor/Private/SGraphNode.cpp
--- a/Engine/Source/Editor/GraphEditor/Private/SGraphNode.cpp
+++ b/Engine/Source/Editor/GraphEditor/Private/SGraphNode.cpp
@@ -479,6 +479,10 @@ void SGraphNode::Tick( const FGeometry& AllottedGeometry, const double InCurrent
 bool SGraphNode::IsSelectedExclusively() const
 {
 	TSharedPtr<SGraphPanel> OwnerPanel = OwnerGraphPanelPtr.Pin();
+	if (!OwnerPanel.IsValid())
+	{
+		return false;
+	}
 
 	if (!OwnerPanel->HasKeyboardFocus() || OwnerPanel->SelectionManager.GetSelectedNodes().Num() > 1)
 	{
diff --git a/Engine/Source/Editor/GraphEditor/Private/SGraphPin.cpp b/Engine/Source/Editor/GraphEditor/Private/SGraphPin.cpp
--- a/Engine/Source/Editor/GraphEditor/Private/SGraphPin.cpp
+++ b/Engine/Source/Editor/GraphEditor/Private/SGraphPin.cpp
@@ -994,10 +994,12 @@ const FSlateBrush* SGraphPin::GetPinBorder() const
 {
 	bool bIsMarkedPin = false;
 	TSharedPtr<SGraphPanel> OwnerPanelPtr = OwnerNodePtr.Pin()->GetOwnerPanel();
-	check(OwnerPanelPtr.IsValid());
-	if (OwnerPanelPtr->MarkedPin.IsValid())
+	if (OwnerPanelPtr.IsValid())
 	{
-		bIsMarkedPin = (OwnerPanelPtr->MarkedPin.Pin() == SharedThis(this));
+		if (OwnerPanelPtr->MarkedPin.IsValid())
+		{
+			bIsMarkedPin = (OwnerPanelPtr->MarkedPin.Pin() == SharedThis(this));
+		}
 	}
 
 	return (IsHovered() || bIsMarkedPin || GraphPinObj->bIsDiffing || bOnlyShowDefaultValue) ? CachedImg_Pin_BackgroundHovered : CachedImg_Pin_Background;
diff --git a/Engine/Source/Editor/Kismet/Private/SMyBlueprint.cpp b/Engine/Source/Editor/Kismet/Private/SMyBlueprint.cpp
--- a/Engine/Source/Editor/Kismet/Private/SMyBlueprint.cpp
+++ b/Engine/Source/Editor/Kismet/Private/SMyBlueprint.cpp
@@ -494,6 +494,7 @@ void SMyBlueprint::Construct(const FArguments& InArgs, TWeakPtr<FBlueprintEditor
 		.OnCanRenameSelectedAction(this, &SMyBlueprint::CanRequestRenameOnActionNode)
 		.OnGetSectionTitle(this, &SMyBlueprint::OnGetSectionTitle)
 		.OnGetSectionWidget(this, &SMyBlueprint::OnGetSectionWidget)
+		.OnActionMatchesName(this, &SMyBlueprint::HandleActionMatchesName)
 		.AlphaSortItems(false)
 		.UseSectionStyling(true);
 
@@ -904,6 +905,15 @@ bool SMyBlueprint::CanAddNewElementToSection(int32 InSectionID) const
 	return false;
 }
 
+bool SMyBlueprint::HandleActionMatchesName(FEdGraphSchemaAction* InAction, const FName& InName) const
+{
+	if (BlueprintEditorPtr.IsValid())
+	{
+		return BlueprintEditorPtr.Pin()->OnActionMatchesName(InAction, InName);
+	}
+	return false;
+}
+
 EVisibility SMyBlueprint::OnGetSectionTextVisibility(TWeakPtr<SWidget> RowWidget, int32 InSectionID) const
 {
 	bool ShowText = RowWidget.Pin()->IsHovered();
@@ -1694,7 +1704,15 @@ FReply SMyBlueprint::OnActionDragged( const TArray< TSharedPtr<FEdGraphSchemaAct
 		if(InAction->GetTypeId() == FEdGraphSchemaAction_K2Graph::StaticGetTypeId())
 		{
 			FEdGraphSchemaAction_K2Graph* FuncAction = (FEdGraphSchemaAction_K2Graph*)InAction.Get();
-			
+
+			if (FuncAction->EdGraph)
+			{
+				if (FuncAction->EdGraph->GetSchema()->CanGraphBeDropped(InAction))
+				{
+					return FuncAction->EdGraph->GetSchema()->BeginGraphDragAction(InAction);
+				}
+			}
+
 			if (FuncAction->GraphType == EEdGraphSchemaAction_K2Graph::Function ||FuncAction->GraphType == EEdGraphSchemaAction_K2Graph::Interface)
 			{
 				// Callback function to report that the user cannot drop this function in the graph
diff --git a/Engine/Source/Editor/Kismet/Private/SMyBlueprint.h b/Engine/Source/Editor/Kismet/Private/SMyBlueprint.h
--- a/Engine/Source/Editor/Kismet/Private/SMyBlueprint.h
+++ b/Engine/Source/Editor/Kismet/Private/SMyBlueprint.h
@@ -183,6 +183,8 @@ private:
 	FReply OnAddButtonClickedOnSection(int32 InSectionID);
 	bool CanAddNewElementToSection(int32 InSectionID) const;
 
+	bool HandleActionMatchesName(FEdGraphSchemaAction* InAction, const FName& InName) const;
+
 	/** Support functions for checkbox to manage displaying user variables only */
 	bool IsShowingInheritedVariables() const;
 	void OnToggleShowInheritedVariables();
@@ -290,6 +292,7 @@ private:
 
 	/** Helper function indicating whehter we're in editing mode, and can modify the target blueprint */
 	bool IsEditingMode() const;
+
 private:
 	/** List of UI Commands for this scope */
 	TSharedPtr<FUICommandList> CommandList;
diff --git a/Engine/Source/Editor/Kismet/Public/BlueprintEditor.h b/Engine/Source/Editor/Kismet/Public/BlueprintEditor.h
--- a/Engine/Source/Editor/Kismet/Public/BlueprintEditor.h
+++ b/Engine/Source/Editor/Kismet/Public/BlueprintEditor.h
@@ -5,6 +5,7 @@
 #include "CoreMinimal.h"
 #include "Stats/Stats.h"
 #include "EdGraph/EdGraphPin.h"
+#include "EdGraph/EdGraphSchema.h"
 #include "UObject/GCObject.h"
 #include "Layout/Visibility.h"
 #include "Input/Reply.h"
@@ -514,7 +515,7 @@ public:
 	};
 
 	/** Called when New Function button is clicked */
-	void NewDocument_OnClicked(ECreatedDocumentType GraphType);
+	virtual void NewDocument_OnClicked(ECreatedDocumentType GraphType);
 	FReply NewDocument_OnClick(ECreatedDocumentType GraphType) { NewDocument_OnClicked(GraphType); return FReply::Handled(); }
 
 	/** Called when New Delegate button is clicked */
@@ -1368,6 +1369,11 @@ public:
 	DECLARE_EVENT ( FBlueprintEditor, FOnRefreshEvent );
 	FOnRefreshEvent& OnRefresh() { return RefreshEvent; }
 
+	/** function used by the SMyBlueprint to determine if an action matches a name.
+	  * This happens during selection of a custom action ( a graph, a variable etc )
+	  */
+	virtual bool OnActionMatchesName(FEdGraphSchemaAction* InAction, const FName& InName) const { return false; }
+
 private:
 	/** Notification used whenever the editor wants associated controls to refresh. */
 	FOnRefreshEvent RefreshEvent;
diff --git a/Engine/Source/Runtime/Engine/Classes/EdGraph/EdGraphSchema.h b/Engine/Source/Runtime/Engine/Classes/EdGraph/EdGraphSchema.h
--- a/Engine/Source/Runtime/Engine/Classes/EdGraph/EdGraphSchema.h
+++ b/Engine/Source/Runtime/Engine/Classes/EdGraph/EdGraphSchema.h
@@ -949,6 +949,16 @@ class ENGINE_API UEdGraphSchema : public UObject
 	 */
 	virtual bool CanEncapuslateNode(UEdGraphNode const& TestNode) const { return true; }
 
+	/*
+	 * Can the function graph be dropped into another graph
+	 */
+	virtual bool CanGraphBeDropped(TSharedPtr<FEdGraphSchemaAction> InAction) const { return false; }
+
+	/*
+	 * Begins a drag and drop action to drag a graph into another graph
+	 */
+	virtual FReply BeginGraphDragAction(TSharedPtr<FEdGraphSchemaAction> InAction) const { return FReply::Unhandled(); }
+
 	/**
 	 * Gets display information for a graph
 	 *
@@ -1034,6 +1044,13 @@ class ENGINE_API UEdGraphSchema : public UObject
 	 */
 	virtual bool SupportsDropPinOnNode(UEdGraphNode* InTargetNode, const FEdGraphPinType& InSourcePinType, EEdGraphPinDirection InSourcePinDirection, FText& OutErrorMessage) const { return false; }
 
+	/**
+	 * Let's the schema know about the next pin being dropped
+	 *
+	 * @param InSourcePin					The pin which is about to be dropped
+	 */
+	virtual void SetPinBeingDroppedOnNode(UEdGraphPin* InSourcePin) const {}
+
 	/**
 	 * Checks if a CacheRefreshID is out of date
 	 *
