Control Rig: Implement support for local function in contextual spawner menu

diff --git a/Engine/Plugins/Experimental/ControlRig/Source/ControlRigDeveloper/Public/Graph/ControlRigGraph.h b/Engine/Plugins/Experimental/ControlRig/Source/ControlRigDeveloper/Public/Graph/ControlRigGraph.h
--- a/Engine/Plugins/Experimental/ControlRig/Source/ControlRigDeveloper/Public/Graph/ControlRigGraph.h
+++ b/Engine/Plugins/Experimental/ControlRig/Source/ControlRigDeveloper/Public/Graph/ControlRigGraph.h
@@ -118,6 +118,7 @@ private:
 	friend class UControlRigSelectNodeSpawner;
 	friend class UControlRigPrototypeNodeSpawner;
 	friend class UControlRigEnumNodeSpawner;
+	friend class UControlRigFunctionRefNodeSpawner;
 #endif
 	friend class UControlRigGraphNode;
 	friend class FControlRigEditor;
diff --git a/Engine/Plugins/Experimental/ControlRig/Source/ControlRigEditor/Private/ControlRigEditorModule.cpp b/Engine/Plugins/Experimental/ControlRig/Source/ControlRigEditor/Private/ControlRigEditorModule.cpp
--- a/Engine/Plugins/Experimental/ControlRig/Source/ControlRigEditor/Private/ControlRigEditorModule.cpp
+++ b/Engine/Plugins/Experimental/ControlRig/Source/ControlRigEditor/Private/ControlRigEditorModule.cpp
@@ -53,6 +53,7 @@
 #include "Graph/NodeSpawners/ControlRigSelectNodeSpawner.h"
 #include "Graph/NodeSpawners/ControlRigPrototypeNodeSpawner.h"
 #include "Graph/NodeSpawners/ControlRigEnumNodeSpawner.h"
+#include "Graph/NodeSpawners/ControlRigFunctionRefNodeSpawner.h"
 #include "Kismet2/BlueprintEditorUtils.h"
 #include "Kismet2/KismetDebugUtilities.h"
 #include "Graph/ControlRigGraphNode.h"
@@ -1106,6 +1107,17 @@ void FControlRigEditorModule::GetInstanceActions(UControlRigBlueprint* CRB, FBlu
 				ActionRegistrar.AddBlueprintAction(GeneratedClass, UControlRigVariableNodeSpawner::CreateFromExternalVariable(CRB, ExternalVariable, false, MenuDesc, NodeCategory, ToolTip));
 			}
 		}
+
+		if (URigVMFunctionLibrary* LocalFunctionLibrary = CRB->GetLocalFunctionLibrary())
+		{
+			TArray<URigVMLibraryNode*> Functions = LocalFunctionLibrary->GetFunctions();
+			for (URigVMLibraryNode* Function : Functions)
+			{
+				UBlueprintNodeSpawner* NodeSpawner = UControlRigFunctionRefNodeSpawner::CreateFromFunction(Function);
+				check(NodeSpawner != nullptr);
+				ActionRegistrar.AddBlueprintAction(GeneratedClass, NodeSpawner);
+			}
+		}
 	}
 }
 
diff --git a/Engine/Plugins/Experimental/ControlRig/Source/ControlRigEditor/Private/Graph/NodeSpawners/ControlRigFunctionRefNodeSpawner.cpp b/Engine/Plugins/Experimental/ControlRig/Source/ControlRigEditor/Private/Graph/NodeSpawners/ControlRigFunctionRefNodeSpawner.cpp
new file mode 100644
--- /dev/null
+++ b/Engine/Plugins/Experimental/ControlRig/Source/ControlRigEditor/Private/Graph/NodeSpawners/ControlRigFunctionRefNodeSpawner.cpp
@@ -0,0 +1,172 @@
+// Copyright Epic Games, Inc. All Rights Reserved.
+
+#include "ControlRigFunctionRefNodeSpawner.h"
+#include "ControlRigUnitNodeSpawner.h"
+#include "Graph/ControlRigGraph.h"
+#include "Graph/ControlRigGraphNode.h"
+#include "Graph/ControlRigGraphSchema.h"
+#include "EdGraphSchema_K2.h"
+#include "Kismet2/BlueprintEditorUtils.h"
+#include "Classes/EditorStyleSettings.h"
+#include "Editor/EditorEngine.h"
+#include "ObjectEditorUtils.h"
+#include "EditorCategoryUtils.h"
+#include "K2Node_Variable.h"
+#include "BlueprintNodeTemplateCache.h"
+#include "ControlRigBlueprintUtils.h"
+#include "ScopedTransaction.h"
+#include "ControlRig/Private/Units/Execution/RigUnit_BeginExecution.h"
+#include "ControlRig.h"
+#include "Settings/ControlRigSettings.h"
+
+#if WITH_EDITOR
+#include "Editor.h"
+#include "SGraphActionMenu.h"
+#endif
+
+#define LOCTEXT_NAMESPACE "ControlRigFunctionRefNodeSpawner"
+
+UControlRigFunctionRefNodeSpawner* UControlRigFunctionRefNodeSpawner::CreateFromFunction(URigVMLibraryNode* InFunction)
+{
+	check(InFunction);
+
+	UControlRigFunctionRefNodeSpawner* NodeSpawner = NewObject<UControlRigFunctionRefNodeSpawner>(GetTransientPackage());
+	NodeSpawner->ReferencedFunctionPtr = InFunction;
+	NodeSpawner->NodeClass = UControlRigGraphNode::StaticClass();
+
+	FBlueprintActionUiSpec& MenuSignature = NodeSpawner->DefaultMenuSignature;
+
+	FString Category = InFunction->GetCategory();
+	FString CategoryPrefix;
+
+	if (UControlRigBlueprint* Blueprint = Cast<UControlRigBlueprint>(InFunction->GetLibrary()->GetOuter()))
+	{
+		/*if(Blueprint->IsFunctionLibrary())
+		{
+			// todo
+			ensure(false);
+		}
+		else
+		*/
+		{
+			CategoryPrefix = TEXT("Local Functions");
+		}
+	}
+
+	if (!Category.IsEmpty() && !CategoryPrefix.IsEmpty())
+	{
+		CategoryPrefix += TEXT("|");
+	}
+
+	MenuSignature.MenuName = FText::FromString(InFunction->GetName());
+	MenuSignature.Tooltip = FText::FromString(InFunction->GetTooltip());
+	MenuSignature.Category = FText::FromString(CategoryPrefix + Category);
+	MenuSignature.Keywords = FText::FromString(InFunction->GetKeywords());
+
+	// add at least one character, so that PrimeDefaultUiSpec() doesn't 
+	// attempt to query the template node
+	if (MenuSignature.Keywords.IsEmpty())
+	{
+		// want to set it to something so we won't end up back in this condition
+		MenuSignature.Keywords = FText::FromString(TEXT(" "));
+	}
+
+	MenuSignature.Icon = FSlateIcon("EditorStyle", "Kismet.AllClasses.FunctionIcon");
+
+	return NodeSpawner;
+}
+
+void UControlRigFunctionRefNodeSpawner::Prime()
+{
+	// we expect that you don't need a node template to construct menu entries
+	// from this, so we choose not to pre-cache one here
+}
+
+FBlueprintNodeSignature UControlRigFunctionRefNodeSpawner::GetSpawnerSignature() const
+{
+	if(ReferencedFunctionPtr.IsValid())
+	{
+		return FBlueprintNodeSignature(TEXT("RigFunction=") + ReferencedFunctionPtr->GetPathName());
+	}
+	return FBlueprintNodeSignature(TEXT("Invalid RigFunction"));
+}
+
+FBlueprintActionUiSpec UControlRigFunctionRefNodeSpawner::GetUiSpec(FBlueprintActionContext const& Context, FBindingSet const& Bindings) const
+{
+	UEdGraph* TargetGraph = (Context.Graphs.Num() > 0) ? Context.Graphs[0] : nullptr;
+	FBlueprintActionUiSpec MenuSignature = PrimeDefaultUiSpec(TargetGraph);
+
+	DynamicUiSignatureGetter.ExecuteIfBound(Context, Bindings, &MenuSignature);
+	return MenuSignature;
+}
+
+UEdGraphNode* UControlRigFunctionRefNodeSpawner::Invoke(UEdGraph* ParentGraph, FBindingSet const& Bindings, FVector2D const Location) const
+{
+	UControlRigGraphNode* NewNode = nullptr;
+
+	if(ReferencedFunctionPtr.IsValid())
+	{
+#if WITH_EDITOR
+		if (GEditor)
+		{
+			GEditor->CancelTransaction(0);
+		}
+#endif
+
+		UBlueprint* Blueprint = FBlueprintEditorUtils::FindBlueprintForGraph(ParentGraph);
+		NewNode = SpawnNode(ParentGraph, Blueprint, ReferencedFunctionPtr.Get(), Location);
+	}
+
+	return NewNode;
+}
+
+UControlRigGraphNode* UControlRigFunctionRefNodeSpawner::SpawnNode(UEdGraph* ParentGraph, UBlueprint* Blueprint, URigVMLibraryNode* InFunction, FVector2D const Location)
+{
+	UControlRigGraphNode* NewNode = nullptr;
+	UControlRigBlueprint* RigBlueprint = Cast<UControlRigBlueprint>(Blueprint);
+	UControlRigGraph* RigGraph = Cast<UControlRigGraph>(ParentGraph);
+
+	if (RigBlueprint != nullptr && RigGraph != nullptr)
+	{
+		bool const bIsTemplateNode = FBlueprintNodeTemplateCache::IsTemplateOuter(ParentGraph);
+		bool const bUndo = !bIsTemplateNode;
+
+		FName Name = bIsTemplateNode ? *InFunction->GetName() : FControlRigBlueprintUtils::ValidateName(RigBlueprint, InFunction->GetName());
+		URigVMController* Controller = bIsTemplateNode ? RigGraph->GetTemplateController() : RigBlueprint->GetController(ParentGraph);
+
+		if (!bIsTemplateNode)
+		{
+			Controller->OpenUndoBracket(FString::Printf(TEXT("Add '%s' Node"), *Name.ToString()));
+		}
+
+		if (URigVMFunctionReferenceNode* ModelNode = Controller->AddFunctionReferenceNode(InFunction, Location, Name.ToString(), bUndo))
+		{
+			NewNode = Cast<UControlRigGraphNode>(RigGraph->FindNodeForModelNodeName(ModelNode->GetFName()));
+			check(NewNode);
+
+			if (NewNode && bUndo)
+			{
+				Controller->ClearNodeSelection(true);
+				Controller->SelectNode(ModelNode, true, true);
+
+				UControlRigUnitNodeSpawner::HookupMutableNode(ModelNode, RigBlueprint);
+			}
+
+			if (bUndo)
+			{
+				Controller->CloseUndoBracket();
+			}
+		}
+		else
+		{
+			if (bUndo)
+			{
+				Controller->CancelUndoBracket();
+			}
+		}
+	}
+	return NewNode;
+}
+
+
+#undef LOCTEXT_NAMESPACE
diff --git a/Engine/Plugins/Experimental/ControlRig/Source/ControlRigEditor/Private/Graph/NodeSpawners/ControlRigFunctionRefNodeSpawner.h b/Engine/Plugins/Experimental/ControlRig/Source/ControlRigEditor/Private/Graph/NodeSpawners/ControlRigFunctionRefNodeSpawner.h
new file mode 100644
--- /dev/null
+++ b/Engine/Plugins/Experimental/ControlRig/Source/ControlRigEditor/Private/Graph/NodeSpawners/ControlRigFunctionRefNodeSpawner.h
@@ -0,0 +1,53 @@
+// Copyright Epic Games, Inc. All Rights Reserved.
+
+#pragma once
+
+#include "CoreMinimal.h"
+#include "UObject/ObjectMacros.h"
+#include "UObject/Object.h"
+#include "Templates/SubclassOf.h"
+#include "UObject/UnrealType.h"
+#include "EdGraph/EdGraphPin.h"
+#include "Engine/Blueprint.h"
+#include "BlueprintActionFilter.h"
+#include "BlueprintNodeSignature.h"
+#include "BlueprintFieldNodeSpawner.h"
+#include "ControlRigBlueprint.h"
+#include "Graph/ControlRigGraph.h"
+#include "Graph/ControlRigGraphNode.h"
+#include "RigVMModel/Nodes/RigVMUnitNode.h"
+#include "ControlRigFunctionRefNodeSpawner.generated.h"
+
+class UControlRigGraphNode;
+
+UCLASS(Transient)
+class CONTROLRIGEDITOR_API UControlRigFunctionRefNodeSpawner : public UBlueprintNodeSpawner
+{
+	GENERATED_BODY()
+
+public:
+	/**
+	 * Creates a new UControlRigVariableNodeSpawner, charged with spawning 
+	 * a new member-variable node
+	 * 
+	 * @return A newly allocated instance of this class.
+	 */
+	static UControlRigFunctionRefNodeSpawner* CreateFromFunction(URigVMLibraryNode* InFunction);
+
+	// UBlueprintNodeSpawner interface
+	virtual void Prime() override;
+	virtual FBlueprintNodeSignature GetSpawnerSignature() const override;
+	virtual FBlueprintActionUiSpec GetUiSpec(FBlueprintActionContext const& Context, FBindingSet const& Bindings) const override;
+	virtual UEdGraphNode* Invoke(UEdGraph* ParentGraph, FBindingSet const& Bindings, FVector2D const Location) const override;
+	// End UBlueprintNodeSpawner interface
+
+private:
+	/** The unit type we will spawn */
+	UPROPERTY(Transient)
+	TWeakObjectPtr<URigVMLibraryNode> ReferencedFunctionPtr;
+
+	static UControlRigGraphNode* SpawnNode(UEdGraph* ParentGraph, UBlueprint* Blueprint, URigVMLibraryNode* InFunction, FVector2D const Location);
+
+	friend class UEngineTestControlRig;
+	friend class FControlRigEditor;
+};
diff --git a/Engine/Source/Developer/RigVMDeveloper/Public/RigVMModel/Nodes/RigVMLibraryNode.h b/Engine/Source/Developer/RigVMDeveloper/Public/RigVMModel/Nodes/RigVMLibraryNode.h
--- a/Engine/Source/Developer/RigVMDeveloper/Public/RigVMModel/Nodes/RigVMLibraryNode.h
+++ b/Engine/Source/Developer/RigVMDeveloper/Public/RigVMModel/Nodes/RigVMLibraryNode.h
@@ -28,6 +28,9 @@ public:
 	virtual bool IsDefinedAsVarying() const override;
 
 	// Library node interface
+	virtual FString GetCategory() const { return Category; }
+	virtual FString GetTooltip() const { return Tooltip; }
+	virtual FString GetKeywords() const { return Keywords; }
 	virtual URigVMFunctionLibrary* GetLibrary() const { return nullptr; }
 	virtual URigVMGraph* GetContainedGraph() const { return nullptr; }
 	virtual const TArray<URigVMNode*>& GetContainedNodes() const;
@@ -37,6 +40,15 @@ public:
 
 protected:
 
+	UPROPERTY()
+	FString Category;
+
+	UPROPERTY()
+	FString Tooltip;
+
+	UPROPERTY()
+	FString Keywords;
+
 	const static TArray<URigVMNode*> EmptyNodes;
 	const static TArray<URigVMLink*> EmptyLinks;
 
