Control Rig: Support for error reporting from AST parser all the way up to the Blueprint / Editor

diff --git a/Engine/Plugins/Experimental/ControlRig/Source/ControlRigDeveloper/Private/ControlRigBlueprint.cpp b/Engine/Plugins/Experimental/ControlRig/Source/ControlRigDeveloper/Private/ControlRigBlueprint.cpp
--- a/Engine/Plugins/Experimental/ControlRig/Source/ControlRigDeveloper/Private/ControlRigBlueprint.cpp
+++ b/Engine/Plugins/Experimental/ControlRig/Source/ControlRigDeveloper/Private/ControlRigBlueprint.cpp
@@ -31,6 +31,8 @@
 #include "UnrealEdGlobals.h"
 #include "Editor/UnrealEdEngine.h"
 #include "CookOnTheSide/CookOnTheFlyServer.h"
+#include "Widgets/Notifications/SNotificationList.h"
+#include "Framework/Notifications/NotificationManager.h"
 #endif//WITH_EDITOR
 
 #define LOCTEXT_NAMESPACE "ControlRigBlueprint"
@@ -67,9 +69,12 @@ UControlRigBlueprint::UControlRigBlueprint(const FObjectInitializer& ObjectIniti
 	Validator = ObjectInitializer.CreateDefaultSubobject<UControlRigValidator>(this, TEXT("ControlRigValidator"));
 
 	bDirtyDuringLoad = false;
+	bErrorsDuringCompilation = false;
 
 	SupportedEventNames.Reset();
 	bExposesAnimatableControls = false;
+
+	VMCompileSettings.ASTSettings.ReportDelegate.BindUObject(this, &UControlRigBlueprint::HandleReportFromCompiler);
 }
 
 void UControlRigBlueprint::InitializeModelIfRequired(bool bRecompileVM)
@@ -276,6 +281,8 @@ void UControlRigBlueprint::PostLoad()
 
 void UControlRigBlueprint::RecompileVM()
 {
+	bErrorsDuringCompilation = false;
+
 	UControlRigBlueprintGeneratedClass* RigClass = GetControlRigBlueprintGeneratedClass();
 	UControlRig* CDO = Cast<UControlRig>(RigClass->GetDefaultObject(true /* create if needed */));
 	if (CDO->VM != nullptr)
@@ -314,6 +321,11 @@ void UControlRigBlueprint::RecompileVM()
 		Compiler->Settings = VMCompileSettings;
 		Compiler->Compile(Model, GetController(), CDO->VM, CDO->GetExternalVariablesImpl(false), UserData, &PinToOperandMap);
 
+		if (bErrorsDuringCompilation)
+		{
+			return;
+		}
+
 		CDO->Execute(EControlRigState::Init, FRigUnit_BeginExecution::EventName); // need to clarify if we actually need this
 		Statistics = CDO->VM->GetStatistics();
 
@@ -372,6 +384,53 @@ void UControlRigBlueprint::DecrementVMRecompileBracket()
 	}
 }
 
+void UControlRigBlueprint::HandleReportFromCompiler(EMessageSeverity::Type InSeverity, UObject* InSubject, const FString& InMessage)
+{
+	if (InSeverity == EMessageSeverity::Error)
+	{
+		Status = BS_Error;
+		MarkPackageDirty();
+
+#if WITH_EDITOR
+		FNotificationInfo Info(FText::FromString(InMessage));
+		Info.Image = FCoreStyle::Get().GetBrush(TEXT("MessageLog.Error"));
+		Info.bFireAndForget = true;
+		Info.FadeOutDuration = 10.0f;
+		Info.ExpireDuration = 0.0f;
+		TSharedPtr<SNotificationItem> NotificationPtr = FSlateNotificationManager::Get().AddNotification(Info);
+		NotificationPtr->SetCompletionState(SNotificationItem::CS_Success);
+#endif
+
+		FScriptExceptionHandler::Get().HandleException(ELogVerbosity::Error, *InMessage, *FString());
+		bErrorsDuringCompilation = true;
+	}
+	else if (InSeverity == EMessageSeverity::Warning)
+	{
+		FScriptExceptionHandler::Get().HandleException(ELogVerbosity::Warning, *InMessage, *FString());
+	}
+	else
+	{
+		UE_LOG(LogControlRigDeveloper, Display, TEXT("%s"), *InMessage);
+	}
+
+	if (URigVMNode* Node = Cast<URigVMNode>(InSubject))
+	{
+		if (URigVMGraph* Graph = Node->GetGraph())
+		{
+			if (UControlRigGraph* EdGraph = Cast<UControlRigGraph>(GetEdGraph(Graph)))
+			{
+				if (UControlRigGraphNode* EdGraphNode = Cast<UControlRigGraphNode>(EdGraph->FindNodeForModelNodeName(Node->GetFName())))
+				{
+					EdGraphNode->ErrorType = (int32)InSeverity;
+					EdGraphNode->ErrorMsg = InMessage;
+					EdGraphNode->bHasCompilerMessage = EdGraphNode->ErrorType <= int32(EMessageSeverity::Info);
+
+				}
+			}
+		}
+	}
+}
+
 void UControlRigBlueprint::RequestControlRigInit()
 {
 	UControlRigBlueprintGeneratedClass* RigClass = GetControlRigBlueprintGeneratedClass();
diff --git a/Engine/Plugins/Experimental/ControlRig/Source/ControlRigDeveloper/Public/ControlRigBlueprint.h b/Engine/Plugins/Experimental/ControlRig/Source/ControlRigDeveloper/Public/ControlRigBlueprint.h
--- a/Engine/Plugins/Experimental/ControlRig/Source/ControlRigDeveloper/Public/ControlRigBlueprint.h
+++ b/Engine/Plugins/Experimental/ControlRig/Source/ControlRigDeveloper/Public/ControlRigBlueprint.h
@@ -380,8 +380,10 @@ private:
 
 	void CreateEdGraphForCollapseNodeIfNeeded(URigVMCollapseNode* InNode, bool bForce = false);
 	bool RemoveEdGraphForCollapseNode(URigVMCollapseNode* InNode, bool bNotify = false);
+	void HandleReportFromCompiler(EMessageSeverity::Type InSeverity, UObject* InSubject, const FString& InMessage);
 
 	bool bDirtyDuringLoad;
+	bool bErrorsDuringCompilation;
 
 	friend class FControlRigBlueprintCompilerContext;
 	friend class SRigHierarchy;
diff --git a/Engine/Plugins/Experimental/ControlRig/Source/ControlRigEditor/Private/Editor/ControlRigEditor.cpp b/Engine/Plugins/Experimental/ControlRig/Source/ControlRigEditor/Private/Editor/ControlRigEditor.cpp
--- a/Engine/Plugins/Experimental/ControlRig/Source/ControlRigEditor/Private/Editor/ControlRigEditor.cpp
+++ b/Engine/Plugins/Experimental/ControlRig/Source/ControlRigEditor/Private/Editor/ControlRigEditor.cpp
@@ -2625,7 +2625,6 @@ void FControlRigEditor::UpdateControlRig()
 			GetBlueprintObj()->SetObjectBeingDebugged(ControlRig);
 
 			// initialize is moved post reinstance
-			FInputBlendPose Filter;
 			AnimInstance->ResetControlRigTracks();
 			AnimInstance->AddControlRigTrack(0, ControlRig);
 			AnimInstance->UpdateControlRigTrack(0, 1.0f, FControlRigIOSettings::MakeEnabled(), bExecutionControlRig);
@@ -2642,6 +2641,15 @@ void FControlRigEditor::UpdateControlRig()
 			Blueprint->RecompileVM();
 			Blueprint->ClearFlags(RF_Transient);
 
+			if (Blueprint->bErrorsDuringCompilation)
+			{
+				AnimInstance->ResetControlRigTracks();
+				AnimInstance->RecalcRequiredBones();
+				GetBlueprintObj()->SetObjectBeingDebugged(nullptr);
+				ControlRig = nullptr;
+				return;
+			}
+
 			ControlRig->OnInitialized_AnyThread().AddSP(this, &FControlRigEditor::HandleControlRigExecutedEvent);
 			ControlRig->OnExecuted_AnyThread().AddSP(this, &FControlRigEditor::HandleControlRigExecutedEvent);
 			ControlRig->RequestInit();
@@ -4163,9 +4171,12 @@ void FControlRigEditor::OnNodeDoubleClicked(UControlRigBlueprint* InBlueprint, U
 
 	if (URigVMLibraryNode* LibraryNode = Cast<URigVMLibraryNode>(InNode))
 	{
-		if (UEdGraph* EdGraph = InBlueprint->GetEdGraph(LibraryNode->GetContainedGraph()))
+		if(URigVMGraph* ContainedGraph = LibraryNode->GetContainedGraph())
 		{
-			OpenGraphAndBringToFront(EdGraph, true);
+			if (UEdGraph* EdGraph = InBlueprint->GetEdGraph(ContainedGraph))
+			{
+				OpenGraphAndBringToFront(EdGraph, true);
+			}
 		}
 	}
 }
diff --git a/Engine/Source/Developer/RigVMDeveloper/Private/RigVMCompiler/RigVMAST.cpp b/Engine/Source/Developer/RigVMDeveloper/Private/RigVMCompiler/RigVMAST.cpp
--- a/Engine/Source/Developer/RigVMDeveloper/Private/RigVMCompiler/RigVMAST.cpp
+++ b/Engine/Source/Developer/RigVMDeveloper/Private/RigVMCompiler/RigVMAST.cpp
@@ -820,12 +820,36 @@ const FRigVMVarExprAST* FRigVMCallExternExprAST::FindVarWithPinName(const FName&
 	return nullptr;
 }
 
+void FRigVMParserASTSettings::Report(EMessageSeverity::Type InSeverity, UObject* InSubject, const FString& InMessage) const
+{
+	if (ReportDelegate.IsBound())
+	{
+		ReportDelegate.Execute(InSeverity, InSubject, InMessage);
+	}
+	else
+	{
+		if (InSeverity == EMessageSeverity::Error)
+		{
+			FScriptExceptionHandler::Get().HandleException(ELogVerbosity::Error, *InMessage, *FString());
+		}
+		else if (InSeverity == EMessageSeverity::Warning)
+		{
+			FScriptExceptionHandler::Get().HandleException(ELogVerbosity::Warning, *InMessage, *FString());
+		}
+		else
+		{
+			UE_LOG(LogRigVMDeveloper, Display, TEXT("%s"), *InMessage);
+		}
+	}
+}
+
 const TArray<FRigVMASTProxy> FRigVMParserAST::EmptyProxyArray;
 
 FRigVMParserAST::FRigVMParserAST(URigVMGraph* InGraph, URigVMController* InController, const FRigVMParserASTSettings& InSettings, const TArray<FRigVMExternalVariable>& InExternalVariables, const TArray<FRigVMUserDataArray>& InRigVMUserData)
 {
 	DECLARE_SCOPE_HIERARCHICAL_COUNTER_FUNC()
 
+	Settings = InSettings;
 	ObsoleteBlock = nullptr;
 	LastCycleCheckExpr = nullptr;
 	LinksToSkip = InSettings.LinksToSkip;
@@ -2550,6 +2574,7 @@ void FRigVMParserAST::Inline(URigVMGraph* InGraph, const TArray<FRigVMASTProxy>&
 		TMap<FRigVMASTProxy, TArray<FRigVMASTProxy>>* TargetLinks;
 		TMap<FRigVMASTProxy, TArray<FRigVMASTProxy>>* SourceLinks;
 		TArray<FRigVMASTProxy> LibraryNodeCallstack;
+		FRigVMParserASTSettings* Settings;
 
 		static bool ShouldRecursePin(const FRigVMASTProxy& InPinProxy)
 		{
@@ -2790,6 +2815,30 @@ void FRigVMParserAST::Inline(URigVMGraph* InGraph, const TArray<FRigVMASTProxy>&
 
 			if (URigVMLibraryNode* LibraryNode = InNodeProxy.GetSubject<URigVMLibraryNode>())
 			{
+				if (LibraryNode->GetContainedGraph() == nullptr)
+				{
+					if (URigVMFunctionReferenceNode* FunctionRef = Cast<URigVMFunctionReferenceNode>(LibraryNode))
+					{
+						FString FunctionPath = FunctionRef->ReferencedNodePtr.ToString();
+
+						OutTraversalInfo.Settings->Reportf(
+							EMessageSeverity::Error, 
+							LibraryNode, 
+							TEXT("Function Reference '%s' references a missing function (%s)."), 
+							*LibraryNode->GetName(),
+							*FunctionPath);
+					}
+					else
+					{
+						OutTraversalInfo.Settings->Reportf(
+							EMessageSeverity::Error,
+							LibraryNode,
+							TEXT("Library Node '%s' doesn't contain a subgraph."),
+							*LibraryNode->GetName());
+					}
+					return;
+				}
+
 				OutTraversalInfo.LibraryNodeCallstack.Push(InNodeProxy);
 				TArray<URigVMNode*> ContainedNodes = LibraryNode->GetContainedNodes();
 				for (URigVMNode* ContainedNode : ContainedNodes)
@@ -2826,6 +2875,7 @@ void FRigVMParserAST::Inline(URigVMGraph* InGraph, const TArray<FRigVMASTProxy>&
 	TraversalInfo.PinOverrides = &PinOverrides;
 	TraversalInfo.TargetLinks = &TargetLinks;
 	TraversalInfo.SourceLinks = &SourceLinks;
+	TraversalInfo.Settings = &Settings;
 
 	for (const FRigVMASTProxy& NodeProxy : NodeProxies)
 	{
diff --git a/Engine/Source/Developer/RigVMDeveloper/Private/RigVMCompiler/RigVMCompiler.cpp b/Engine/Source/Developer/RigVMDeveloper/Private/RigVMCompiler/RigVMCompiler.cpp
--- a/Engine/Source/Developer/RigVMDeveloper/Private/RigVMCompiler/RigVMCompiler.cpp
+++ b/Engine/Source/Developer/RigVMDeveloper/Private/RigVMCompiler/RigVMCompiler.cpp
@@ -1719,18 +1719,19 @@ FRigVMOperand URigVMCompiler::FindOrAddRegister(const FRigVMVarExprAST* InVarExp
 
 void URigVMCompiler::ReportInfo(const FString& InMessage)
 {
-	if (!Settings.SurpressInfoMessages)
+	if (Settings.SurpressInfoMessages)
 	{
-		UE_LOG(LogRigVMDeveloper, Display, TEXT("%s"), *InMessage);
+		return;
 	}
+	Settings.Report(EMessageSeverity::Info, nullptr, InMessage);
 }
 
 void URigVMCompiler::ReportWarning(const FString& InMessage)
 {
-	FScriptExceptionHandler::Get().HandleException(ELogVerbosity::Warning, *InMessage, *FString());
+	Settings.Report(EMessageSeverity::Warning, nullptr, InMessage);
 }
 
 void URigVMCompiler::ReportError(const FString& InMessage)
 {
-	FScriptExceptionHandler::Get().HandleException(ELogVerbosity::Error, *InMessage, *FString());
+	Settings.Report(EMessageSeverity::Error, nullptr, InMessage);
 }
diff --git a/Engine/Source/Developer/RigVMDeveloper/Private/RigVMModel/RigVMController.cpp b/Engine/Source/Developer/RigVMDeveloper/Private/RigVMModel/RigVMController.cpp
--- a/Engine/Source/Developer/RigVMDeveloper/Private/RigVMModel/RigVMController.cpp
+++ b/Engine/Source/Developer/RigVMDeveloper/Private/RigVMModel/RigVMController.cpp
@@ -92,7 +92,7 @@ void URigVMController::PushGraph(URigVMGraph* InGraph, bool bSetupUndoRedo)
 
 	if (bSetupUndoRedo)
 	{
-		ActionStack->AddAction(FRigVMPushGraphAction(InGraph->GetGraphName()));
+		ActionStack->AddAction(FRigVMPushGraphAction(InGraph));
 	}
 }
 
@@ -104,7 +104,7 @@ URigVMGraph* URigVMController::PopGraph(bool bSetupUndoRedo)
 
 	if (bSetupUndoRedo)
 	{
-		ActionStack->AddAction(FRigVMPopGraphAction(LastGraph->GetGraphName()));
+		ActionStack->AddAction(FRigVMPopGraphAction(LastGraph));
 	}
 
 	return LastGraph;
@@ -1731,12 +1731,6 @@ TArray<FName> URigVMController::ImportNodesFromText(const FString& InText, bool
 	URigVMGraph* Graph = GetGraph();
 	check(Graph);
 
-	// can't import nodes to a function library
-	if (Graph->IsA<URigVMFunctionLibrary>())
-	{
-		return NodeNames;
-	}
-
 	FRigVMControllerObjectFactory Factory(this);
 	Factory.ProcessBuffer(Graph, RF_Transactional, InText);
 
@@ -3122,10 +3116,15 @@ bool URigVMController::RemoveNode(URigVMNode* InNode, bool bSetupUndoRedo, bool
 			BreakAllLinksRecursive(Pin, false, false, bSetupUndoRedo);
 		}
 
+		if (bSetupUndoRedo)
+		{
+			ActionStack->AddAction(FRigVMRemoveNodeAction(InNode, this));
+		}
+
 		if (URigVMCollapseNode* CollapseNode = Cast<URigVMCollapseNode>(InNode))
 		{
 			URigVMGraph* SubGraph = CollapseNode->GetContainedGraph();
-			FRigVMControllerGraphGuard GraphGuard(this, SubGraph, bSetupUndoRedo);
+			FRigVMControllerGraphGuard GraphGuard(this, SubGraph, false);
 
 			TArray<URigVMNode*> ContainedNodes = SubGraph->GetNodes();
 			for (URigVMNode* ContainedNode : ContainedNodes)
@@ -3135,16 +3134,11 @@ bool URigVMController::RemoveNode(URigVMNode* InNode, bool bSetupUndoRedo, bool
 				{
 					continue;
 				}
-				RemoveNode(ContainedNode, bSetupUndoRedo, bRecursive);
+				RemoveNode(ContainedNode, false, true);
 			}
 		}
 	}
 
-	if (bSetupUndoRedo)
-	{
-		ActionStack->AddAction(FRigVMRemoveNodeAction(InNode, this));
-	}
-
 	Graph->Nodes.Remove(InNode);
 	if (!bSuspendNotifications)
 	{
@@ -3190,10 +3184,15 @@ bool URigVMController::RemoveNode(URigVMNode* InNode, bool bSetupUndoRedo, bool
 				{
 					if (FunctionReferencePtr.IsValid())
 					{
-						FunctionReferencePtr->ReferencedNodePtr.Reset();
+						{
+							TGuardValue<TSoftObjectPtr<URigVMLibraryNode>> ClearReferencedNodePtr(
+								FunctionReferencePtr->ReferencedNodePtr,
+								TSoftObjectPtr<URigVMLibraryNode>());
 
-						FRigVMControllerGraphGuard GraphGuard(this, FunctionReferencePtr->GetGraph(), false);
-						RepopulatePinsOnNode(FunctionReferencePtr.Get(), false, true);
+							FRigVMControllerGraphGuard GraphGuard(this, FunctionReferencePtr->GetGraph(), false);
+							RepopulatePinsOnNode(FunctionReferencePtr.Get(), false, true);
+						}
+						FunctionReferencePtr->ReferencedNodePtr.ResetWeakPtr();
 					}
 				}
 			}
diff --git a/Engine/Source/Developer/RigVMDeveloper/Private/RigVMModel/RigVMControllerActions.cpp b/Engine/Source/Developer/RigVMDeveloper/Private/RigVMModel/RigVMControllerActions.cpp
--- a/Engine/Source/Developer/RigVMDeveloper/Private/RigVMModel/RigVMControllerActions.cpp
+++ b/Engine/Source/Developer/RigVMDeveloper/Private/RigVMModel/RigVMControllerActions.cpp
@@ -1433,9 +1433,9 @@ bool FRigVMRenameNodeAction::Redo(URigVMController* InController)
 	return FRigVMBaseAction::Redo(InController);
 }
 
-FRigVMPushGraphAction::FRigVMPushGraphAction(const FString& InGraphName)
-	: GraphName(InGraphName)
+FRigVMPushGraphAction::FRigVMPushGraphAction(UObject* InGraph)
 {
+	GraphPath = TSoftObjectPtr<URigVMGraph>(Cast<URigVMGraph>(InGraph)).GetUniqueID();
 }
 
 bool FRigVMPushGraphAction::Undo(URigVMController* InController)
@@ -1449,25 +1449,18 @@ bool FRigVMPushGraphAction::Undo(URigVMController* InController)
 
 bool FRigVMPushGraphAction::Redo(URigVMController* InController)
 {
-	if (GraphName.IsEmpty())
+	TSoftObjectPtr<URigVMGraph> GraphPtr(GraphPath);
+	if(URigVMGraph* Graph = GraphPtr.Get())
 	{
-		InController->PushGraph(InController->GetTopLevelGraph(), false);
+		InController->PushGraph(Graph, false);
 		return FRigVMBaseAction::Redo(InController);
 	}
-	else if (URigVMNode* Node = InController->GetGraph()->FindNode(GraphName))
-	{
-		if (URigVMCollapseNode* CollapseNode = Cast<URigVMCollapseNode>(Node))
-		{
-			InController->PushGraph(CollapseNode->GetContainedGraph(), false);
-			return FRigVMBaseAction::Redo(InController);
-		}
-	}
 	return false;
 }
 
-FRigVMPopGraphAction::FRigVMPopGraphAction(const FString& InGraphName)
-	: GraphName(InGraphName)
+FRigVMPopGraphAction::FRigVMPopGraphAction(UObject* InGraph)
 {
+	GraphPath = TSoftObjectPtr<URigVMGraph>(Cast<URigVMGraph>(InGraph)).GetUniqueID();
 }
 
 bool FRigVMPopGraphAction::Undo(URigVMController* InController)
@@ -1477,19 +1470,12 @@ bool FRigVMPopGraphAction::Undo(URigVMController* InController)
 		return false;
 	}
 
-	if (GraphName.IsEmpty())
+	TSoftObjectPtr<URigVMGraph> GraphPtr(GraphPath);
+	if (URigVMGraph* Graph = GraphPtr.Get())
 	{
-		InController->PushGraph(InController->GetTopLevelGraph(), false);
+		InController->PushGraph(Graph, false);
 		return true;
 	}
-	else if (URigVMNode* Node = InController->GetGraph()->FindNode(GraphName))
-	{
-		if (URigVMCollapseNode* CollapseNode = Cast<URigVMCollapseNode>(Node))
-		{
-			InController->PushGraph(CollapseNode->GetContainedGraph(), false);
-			return true;
-		}
-	}
 	return false;
 }
 
diff --git a/Engine/Source/Developer/RigVMDeveloper/Public/RigVMCompiler/RigVMAST.h b/Engine/Source/Developer/RigVMDeveloper/Public/RigVMCompiler/RigVMAST.h
--- a/Engine/Source/Developer/RigVMDeveloper/Public/RigVMCompiler/RigVMAST.h
+++ b/Engine/Source/Developer/RigVMDeveloper/Public/RigVMCompiler/RigVMAST.h
@@ -10,6 +10,8 @@
 #include "RigVMASTProxy.h"
 #include "RigVMAST.generated.h"
 
+DECLARE_DELEGATE_ThreeParams(FRigVMReportDelegate, EMessageSeverity::Type, UObject*, const FString&);
+
 class FRigVMParserAST;
 class FRigVMBlockExprAST;
 class FRigVMEntryExprAST;
@@ -1114,6 +1116,8 @@ struct RIGVMDEVELOPER_API FRigVMParserASTSettings
 	UPROPERTY()
 	TArray<URigVMLink*> LinksToSkip;
 
+	FRigVMReportDelegate ReportDelegate;
+
 	// static method to provide fast AST parse settings
 	static FRigVMParserASTSettings Fast()
 	{
@@ -1134,6 +1138,14 @@ struct RIGVMDEVELOPER_API FRigVMParserASTSettings
 		Settings.bFoldConstantBranches = true;
 		return Settings;
 	}
+
+	void Report(EMessageSeverity::Type InSeverity, UObject* InSubject, const FString& InMessage) const;
+
+	template <typename FmtType, typename... Types>
+	void Reportf(EMessageSeverity::Type InSeverity, UObject* InSubject, const FmtType& Fmt, Types... Args) const
+	{
+		Report(InSeverity, InSubject, FString::Printf(Fmt, Args...));
+	}
 };
 
 /*
@@ -1373,6 +1385,8 @@ private:
 	TArray<ETraverseRelationShip> CycleCheckFlags;
 	TArray<URigVMLink*> LinksToSkip;
 
+	FRigVMParserASTSettings Settings;
+
 	friend class FRigVMExprAST;
 	friend class URigVMCompiler;
 };
diff --git a/Engine/Source/Developer/RigVMDeveloper/Public/RigVMCompiler/RigVMCompiler.h b/Engine/Source/Developer/RigVMDeveloper/Public/RigVMCompiler/RigVMCompiler.h
--- a/Engine/Source/Developer/RigVMDeveloper/Public/RigVMCompiler/RigVMCompiler.h
+++ b/Engine/Source/Developer/RigVMDeveloper/Public/RigVMCompiler/RigVMCompiler.h
@@ -13,6 +13,7 @@
 #include "RigVMCore/RigVM.h"
 #include "RigVMCore/RigVMStruct.h"
 #include "RigVMCompiler/RigVMAST.h"
+#include "Logging/TokenizedMessage.h"
 
 #include "RigVMCompiler.generated.h"
 
@@ -63,6 +64,17 @@ public:
 		Settings.ASTSettings = FRigVMParserASTSettings::Optimized();
 		return Settings;
 	}
+
+	void Report(EMessageSeverity::Type InSeverity, UObject* InSubject, const FString& InMessage) const
+	{
+		ASTSettings.Report(InSeverity, InSubject, InMessage);
+	}
+
+	template <typename FmtType, typename... Types>
+	void Reportf(EMessageSeverity::Type InSeverity, UObject* InSubject, const FmtType& Fmt, Types... Args) const
+	{
+		Report(InSeverity, InSubject, FString::Printf(Fmt, Args...));
+	}
 };
 
 struct RIGVMDEVELOPER_API FRigVMCompilerWorkData
diff --git a/Engine/Source/Developer/RigVMDeveloper/Public/RigVMModel/Nodes/RigVMFunctionReferenceNode.h b/Engine/Source/Developer/RigVMDeveloper/Public/RigVMModel/Nodes/RigVMFunctionReferenceNode.h
--- a/Engine/Source/Developer/RigVMDeveloper/Public/RigVMModel/Nodes/RigVMFunctionReferenceNode.h
+++ b/Engine/Source/Developer/RigVMDeveloper/Public/RigVMModel/Nodes/RigVMFunctionReferenceNode.h
@@ -40,5 +40,6 @@ private:
 	mutable TSoftObjectPtr<URigVMLibraryNode> ReferencedNodePtr;
 
 	friend class URigVMController;
+	friend class FRigVMParserAST;
 };
 
diff --git a/Engine/Source/Developer/RigVMDeveloper/Public/RigVMModel/RigVMControllerActions.h b/Engine/Source/Developer/RigVMDeveloper/Public/RigVMModel/RigVMControllerActions.h
--- a/Engine/Source/Developer/RigVMDeveloper/Public/RigVMModel/RigVMControllerActions.h
+++ b/Engine/Source/Developer/RigVMDeveloper/Public/RigVMModel/RigVMControllerActions.h
@@ -1234,13 +1234,13 @@ struct FRigVMPushGraphAction : public FRigVMBaseAction
 public:
 
 	FRigVMPushGraphAction() {}
-	FRigVMPushGraphAction(const FString& InGraphName);
+	FRigVMPushGraphAction(UObject* InGraph);
 	virtual ~FRigVMPushGraphAction() {};
 	virtual bool Undo(URigVMController* InController) override;
 	virtual bool Redo(URigVMController* InController) override;
 
 	UPROPERTY()
-	FString GraphName;
+	FSoftObjectPath GraphPath;
 };
 
 /**
@@ -1254,13 +1254,13 @@ struct FRigVMPopGraphAction : public FRigVMBaseAction
 public:
 
 	FRigVMPopGraphAction() {}
-	FRigVMPopGraphAction(const FString& InGraphName);
+	FRigVMPopGraphAction(UObject* InGraph);
 	virtual ~FRigVMPopGraphAction() {};
 	virtual bool Undo(URigVMController* InController) override;
 	virtual bool Redo(URigVMController* InController) override;
 
 	UPROPERTY()
-	FString GraphName;
+	FSoftObjectPath GraphPath;
 };
 
 
