Control Rig: Support for collapse nodes in editor below function library nodes

diff --git a/Engine/Plugins/Experimental/ControlRig/Source/ControlRigDeveloper/Private/ControlRigBlueprint.cpp b/Engine/Plugins/Experimental/ControlRig/Source/ControlRigDeveloper/Private/ControlRigBlueprint.cpp
--- a/Engine/Plugins/Experimental/ControlRig/Source/ControlRigDeveloper/Private/ControlRigBlueprint.cpp
+++ b/Engine/Plugins/Experimental/ControlRig/Source/ControlRigDeveloper/Private/ControlRigBlueprint.cpp
@@ -483,6 +483,12 @@ URigVMGraph* UControlRigBlueprint::GetModel(const UEdGraph* InEdGraph) const
 	ensure(!ModelNodePath.IsEmpty());
 
 	URigVMGraph* SubModel = Model;
+	if (ModelNodePath.StartsWith(FunctionLibrary->GetNodePath()))
+	{
+		SubModel = FunctionLibrary;
+		ModelNodePath = ModelNodePath.Right(ModelNodePath.Len() - FunctionLibrary->GetNodePath().Len() - 1);
+	}
+
 	while (!ModelNodePath.IsEmpty())
 	{
 		FString NodeName = ModelNodePath;
