Control Rig: Added Plane Distance node

diff --git a/Engine/Plugins/Experimental/ControlRig/Source/ControlRig/Private/Units/Math/RigUnit_MathVector.cpp b/Engine/Plugins/Experimental/ControlRig/Source/ControlRig/Private/Units/Math/RigUnit_MathVector.cpp
--- a/Engine/Plugins/Experimental/ControlRig/Source/ControlRig/Private/Units/Math/RigUnit_MathVector.cpp
+++ b/Engine/Plugins/Experimental/ControlRig/Source/ControlRig/Private/Units/Math/RigUnit_MathVector.cpp
@@ -489,4 +489,25 @@ FRigUnit_MathIntersectPlane_Execute()
 
 	Result = FMath::RayPlaneIntersection(Start, Direction, Plane);
 	Distance = (Start - Result).Size();
-}
\ No newline at end of file
+}
+
+FRigUnit_MathDistanceToPlane_Execute()
+{
+	DECLARE_SCOPE_HIERARCHICAL_COUNTER_RIGUNIT()
+
+	FVector UnitVectorNormal = PlaneNormal.GetSafeNormal();
+
+	if (!UnitVectorNormal.IsZero())
+	{ 
+		FPlane Plane(PlanePoint, UnitVectorNormal);
+
+		ClosestPointOnPlane = FVector::PointPlaneProject(Point, Plane);
+		SignedDistance = FVector::PointPlaneDist(Point, PlanePoint, UnitVectorNormal);
+	}
+	else
+	{
+		ClosestPointOnPlane = FVector::ZeroVector;
+		SignedDistance = 0;
+	}
+}
+
diff --git a/Engine/Plugins/Experimental/ControlRig/Source/ControlRig/Private/Units/Math/RigUnit_MathVector.h b/Engine/Plugins/Experimental/ControlRig/Source/ControlRig/Private/Units/Math/RigUnit_MathVector.h
--- a/Engine/Plugins/Experimental/ControlRig/Source/ControlRig/Private/Units/Math/RigUnit_MathVector.h
+++ b/Engine/Plugins/Experimental/ControlRig/Source/ControlRig/Private/Units/Math/RigUnit_MathVector.h
@@ -976,4 +976,38 @@ struct CONTROLRIG_API FRigUnit_MathIntersectPlane : public FRigUnit_MathVectorBa
 
 	UPROPERTY(meta = (Output))
 	float Distance;
+};
+
+/**
+ * Find the point on the plane that is closest to the given point and the distance between them.
+ */
+USTRUCT(meta = (DisplayName = "Distance To Plane", Keywords = "Distance,Plane,Closest,Project"))
+struct CONTROLRIG_API FRigUnit_MathDistanceToPlane : public FRigUnit_MathVectorBase
+{
+	GENERATED_BODY()
+
+	FRigUnit_MathDistanceToPlane()
+	{
+		Point = PlanePoint = ClosestPointOnPlane = FVector::ZeroVector;
+		PlaneNormal = FVector(0.f, 0.f, 1.f);
+		SignedDistance = 0.f;
+	}
+
+	RIGVM_METHOD()
+	virtual void Execute(const FRigUnitContext& Context) override;
+
+	UPROPERTY(meta = (Input))
+	FVector Point;
+
+	UPROPERTY(meta = (Input))
+	FVector PlanePoint;
+
+	UPROPERTY(meta = (Input))
+	FVector PlaneNormal;
+
+	UPROPERTY(meta = (Output))
+	FVector ClosestPointOnPlane;
+
+	UPROPERTY(meta = (Output))
+	float SignedDistance;
 };
\ No newline at end of file
diff --git a/Engine/Plugins/Experimental/ControlRig/Source/ControlRig/Private/Units/Math/RigUnit_MathVectorTests.cpp b/Engine/Plugins/Experimental/ControlRig/Source/ControlRig/Private/Units/Math/RigUnit_MathVectorTests.cpp
--- a/Engine/Plugins/Experimental/ControlRig/Source/ControlRig/Private/Units/Math/RigUnit_MathVectorTests.cpp
+++ b/Engine/Plugins/Experimental/ControlRig/Source/ControlRig/Private/Units/Math/RigUnit_MathVectorTests.cpp
@@ -395,4 +395,22 @@ IMPLEMENT_RIGUNIT_AUTOMATION_TEST(FRigUnit_MathVectorOrthogonal)
 	return true;
 }
 
+IMPLEMENT_RIGUNIT_AUTOMATION_TEST(FRigUnit_MathDistanceToPlane)
+{
+	Unit.Point = FVector(5, -6, 3);
+	Unit.PlanePoint = FVector(4, 5, 0);
+	Unit.PlaneNormal = FVector(3, -2, 1);
+	InitAndExecute();
+	AddErrorIfFalse(FRigUnit_MathVectorTest_Utils::IsNearlyEqual(Unit.ClosestPointOnPlane, FVector(-1,-2,1), KINDA_SMALL_NUMBER), TEXT("unexpected result"));
+	AddErrorIfFalse(FMath::IsNearlyEqual(Unit.SignedDistance, 7.48331477f, KINDA_SMALL_NUMBER), TEXT("unexpected result"));
+	
+	Unit.Point = FVector(5, -6, 3);
+	Unit.PlanePoint = FVector(2, 2, 2);
+	Unit.PlaneNormal = FVector(0,0,0);
+	InitAndExecute();
+	AddErrorIfFalse(Unit.ClosestPointOnPlane.IsZero(), TEXT("unexpected result"));
+	AddErrorIfFalse(Unit.SignedDistance == 0, TEXT("unexpected result"));
+	return true;
+}
+
 #endif
