Sequencer: Control rig: Mirror Fixes for demo. WIP.

diff --git a/Engine/Plugins/Experimental/ControlRig/Source/ControlRig/Private/Tools/ControlRigPose.cpp b/Engine/Plugins/Experimental/ControlRig/Source/ControlRig/Private/Tools/ControlRigPose.cpp
--- a/Engine/Plugins/Experimental/ControlRig/Source/ControlRig/Private/Tools/ControlRigPose.cpp
+++ b/Engine/Plugins/Experimental/ControlRig/Source/ControlRig/Private/Tools/ControlRigPose.cpp
@@ -33,21 +33,26 @@ void FControlRigControlPose::PastePose(UControlRig* ControlRig, bool bDoKey, boo
 
 }
 
-void FControlRigControlPose::SetControlMirrorTransform(UControlRig* ControlRig, const FName& Name, bool bIsMatched, const FVector& GlobalTranslation, const FQuat& GlobalRotation, const FQuat& LocalRotation, bool bNotify, const  FRigControlModifiedContext& Context)
+void FControlRigControlPose::SetControlMirrorTransform(bool bDoLocal, UControlRig* ControlRig, const FName& Name, bool bIsMatched, const FVector& GlobalTranslation, const FQuat& GlobalRotation, const FVector& LocalTranslation, const FQuat& LocalRotation, bool bNotify, const  FRigControlModifiedContext& Context)
 {
-	if (bIsMatched)
+	if (bDoLocal)
 	{
-		int32 Index = ControlRig->GetControlHierarchy().GetIndex(Name);
-		FTransform ParentTransform = ControlRig->GetControlHierarchy().GetParentTransform(Index);
-		FTransform CurrentTransform = ControlRig->GetControlHierarchy().GetLocalTransform(Index);
-		FVector NewLocal = ParentTransform.InverseTransformPositionNoScale(GlobalTranslation);
-		FTransform NewCurrentTransform(LocalRotation, NewLocal);
-		ControlRig->SetControlLocalTransform(Name, NewCurrentTransform, bNotify, Context);
+		FTransform NewLocalTransform(LocalRotation, LocalTranslation);
+		ControlRig->SetControlLocalTransform(Name, NewLocalTransform, bNotify, Context);
 	}
 	else
 	{
-		FTransform NewGlobalTransform(GlobalRotation, GlobalTranslation);
-		ControlRig->SetControlGlobalTransform(Name, NewGlobalTransform, bNotify, Context);
+		if (bIsMatched)
+		{
+			FTransform NewLocalTransform(LocalRotation, LocalTranslation);
+			ControlRig->SetControlLocalTransform(Name, NewLocalTransform, bNotify, Context);
+
+		}
+		else
+		{
+			FTransform NewGlobalTransform(GlobalRotation, GlobalTranslation);
+			ControlRig->SetControlGlobalTransform(Name, NewGlobalTransform, bNotify, Context);
+		}
 	}
 }
 
@@ -63,6 +68,7 @@ void FControlRigControlPose::PastePoseInternal(UControlRig* ControlRig, bool bDo
 
 	TArray<FRigControl> SortedControls;
 	ControlRig->GetControlsInOrder(SortedControls);
+	bool bDoLocal = true;
 	for (const FRigControl& RigControl : SortedControls)
 	{
 		if (!ControlRig->IsControlSelected(RigControl.Name))
@@ -85,17 +91,25 @@ void FControlRigControlPose::PastePoseInternal(UControlRig* ControlRig, bool bDo
 				{
 					if (bDoMirror == false)
 					{
-						ControlRig->SetControlGlobalTransform(RigControl.Name, CopyRigControl->GlobalTransform, true, Context);
+						if (bDoLocal)
+						{
+							ControlRig->SetControlLocalTransform(RigControl.Name, CopyRigControl->LocalTransform, true, Context);
+						}
+						else
+						{
+							ControlRig->SetControlGlobalTransform(RigControl.Name, CopyRigControl->GlobalTransform, true, Context);
+						}
 					}
 					else
 					{
 						FVector GlobalTranslation;
 						FQuat GlobalRotation;
+						FVector LocalTranslation;
 						FQuat LocalRotation;
 						bool bIsMatched = MirrorTable.IsMatched(CopyRigControl->Name);
-						MirrorTable.GetMirrorTransform(*CopyRigControl, bIsMatched, GlobalTranslation, GlobalRotation, LocalRotation);
-						SetControlMirrorTransform(ControlRig, RigControl.Name, bIsMatched, GlobalTranslation, GlobalRotation, LocalRotation, true, Context);
-					}
+						MirrorTable.GetMirrorTransform(*CopyRigControl, bDoLocal,bIsMatched, GlobalTranslation, GlobalRotation, LocalTranslation, LocalRotation);
+						SetControlMirrorTransform(bDoLocal,ControlRig, RigControl.Name, bIsMatched, GlobalTranslation, GlobalRotation, LocalTranslation,LocalRotation, true, Context);
+					}				
 					break;
 				}
 				case ERigControlType::Float:
@@ -150,6 +164,7 @@ void FControlRigControlPose::BlendWithInitialPoses(FControlRigControlPose& Initi
 
 	TArray<FRigControl> SortedControls;
 	ControlRig->GetControlsInOrder(SortedControls);
+	bool bDoLocal = true;
 	for (const FRigControl& RigControl : SortedControls)
 	{
 		if (!ControlRig->IsControlSelected(RigControl.Name))
@@ -177,32 +192,48 @@ void FControlRigControlPose::BlendWithInitialPoses(FControlRigControlPose& Initi
 					{
 						if (bDoMirror == false)
 						{
-							FTransform Val = CopyRigControl->GlobalTransform;
-							FTransform InitialVal = InitialFound->GlobalTransform;
-							FVector Translation, Scale;
-							FQuat Rotation;
-							Translation = FMath::Lerp(InitialVal.GetTranslation(), Val.GetTranslation(), BlendValue);
-							Rotation = FQuat::Slerp(InitialVal.GetRotation(), Val.GetRotation(), BlendValue); //doing slerp here not fast lerp, can be slow this is for content creation
-							Scale = FMath::Lerp(InitialVal.GetScale3D(), Val.GetScale3D(), BlendValue);
-							Val = FTransform(Rotation, Translation, Scale);
-							ControlRig->SetControlGlobalTransform(RigControl.Name, Val, bDoKey, Context);
+							if (bDoLocal == true)
+							{
+								FTransform Val = CopyRigControl->LocalTransform;
+								FTransform InitialVal = InitialFound->LocalTransform;
+								FVector Translation, Scale;
+								FQuat Rotation;
+								Translation = FMath::Lerp(InitialVal.GetTranslation(), Val.GetTranslation(), BlendValue);
+								Rotation = FQuat::Slerp(InitialVal.GetRotation(), Val.GetRotation(), BlendValue); //doing slerp here not fast lerp, can be slow this is for content creation
+								Scale = FMath::Lerp(InitialVal.GetScale3D(), Val.GetScale3D(), BlendValue);
+								Val = FTransform(Rotation, Translation, Scale);
+								ControlRig->SetControlLocalTransform(RigControl.Name, Val, bDoKey, Context);
+							}
+							else
+							{
+								FTransform Val = CopyRigControl->GlobalTransform;
+								FTransform InitialVal = InitialFound->GlobalTransform;
+								FVector Translation, Scale;
+								FQuat Rotation;
+								Translation = FMath::Lerp(InitialVal.GetTranslation(), Val.GetTranslation(), BlendValue);
+								Rotation = FQuat::Slerp(InitialVal.GetRotation(), Val.GetRotation(), BlendValue); //doing slerp here not fast lerp, can be slow this is for content creation
+								Scale = FMath::Lerp(InitialVal.GetScale3D(), Val.GetScale3D(), BlendValue);
+								Val = FTransform(Rotation, Translation, Scale);
+								ControlRig->SetControlGlobalTransform(RigControl.Name, Val, bDoKey, Context);
+							}
 						}
 						else
 						{
 							FVector GlobalTranslation;
 							FQuat GlobalRotation;
+							FVector LocalTranslation;
 							FQuat LocalRotation;
 							bool bIsMatched = MirrorTable.IsMatched(CopyRigControl->Name);
-							MirrorTable.GetMirrorTransform(*CopyRigControl, bIsMatched, GlobalTranslation, GlobalRotation, LocalRotation);
+							MirrorTable.GetMirrorTransform(*CopyRigControl,bDoLocal, bIsMatched, GlobalTranslation, GlobalRotation, LocalTranslation, LocalRotation);
 							FVector InitialTranslation = InitialFound->GlobalTransform.GetTranslation();
 							FQuat InitialGlobalRotation = InitialFound->GlobalTransform.GetRotation();
+							FVector InitialLocalTranslation = InitialFound->LocalTransform.GetTranslation();
 							FQuat InitialLocationRotation = InitialFound->LocalTransform.GetRotation();
 							GlobalTranslation = FMath::Lerp(InitialTranslation, GlobalTranslation, BlendValue);
 							GlobalRotation = FQuat::Slerp(InitialGlobalRotation, GlobalRotation, BlendValue); //doing slerp here not fast lerp, can be slow this is for content creation
+							LocalTranslation = FMath::Lerp(InitialLocalTranslation, LocalTranslation, BlendValue);
 							LocalRotation = FQuat::Slerp(InitialLocationRotation, LocalRotation, BlendValue); //doing slerp here not fast lerp, can be slow this is for content creation
-
-							SetControlMirrorTransform(ControlRig, RigControl.Name, bIsMatched, GlobalTranslation, GlobalRotation, LocalRotation, bDoKey, Context);
-
+							SetControlMirrorTransform(bDoLocal,ControlRig, RigControl.Name, bIsMatched, GlobalTranslation, GlobalRotation,LocalTranslation, LocalRotation, bDoKey, Context);							
 						}
 					}
 				}
@@ -276,7 +307,7 @@ void UControlRigPoseAsset::PastePose(UControlRig* InControlRig, bool bDoKey, boo
 	Pose.PastePose(InControlRig,bDoKey, bDoMirror);
 }
 
-void UControlRigPoseAsset::SelectControls(UControlRig* InControlRig)
+void UControlRigPoseAsset::SelectControls(UControlRig* InControlRig, bool bDoMirror)
 {
 #if WITH_EDITOR
 	FScopedTransaction ScopedTransaction(LOCTEXT("SelectControlTransaction", "Select Control"));
@@ -284,9 +315,31 @@ void UControlRigPoseAsset::SelectControls(UControlRig* InControlRig)
 #endif
 	InControlRig->ClearControlSelection();
 	TArray<FName> Controls = Pose.GetControlNames();
+	FControlRigPoseMirrorTable MirrorTable;
+	FControlRigControlPose TempPose;
+	if (bDoMirror)
+	{
+		MirrorTable.SetUpMirrorTable(InControlRig);
+		TempPose.SavePose(InControlRig, true);
+	}
 	for (const FName& Name : Controls)
 	{
-		InControlRig->SelectControl(Name, true);
+		if (bDoMirror)
+		{
+			FRigControlCopy* CopyRigControl = MirrorTable.GetControl(TempPose, Name);
+			if (CopyRigControl)
+			{
+				InControlRig->SelectControl(CopyRigControl->Name, true);
+			}
+			else
+			{
+				InControlRig->SelectControl(Name, true);
+			}
+		}
+		else
+		{
+			InControlRig->SelectControl(Name, true);
+		}
 	}
 }
 
diff --git a/Engine/Plugins/Experimental/ControlRig/Source/ControlRig/Private/Tools/ControlRigPoseMirrorTable.cpp b/Engine/Plugins/Experimental/ControlRig/Source/ControlRig/Private/Tools/ControlRigPoseMirrorTable.cpp
--- a/Engine/Plugins/Experimental/ControlRig/Source/ControlRig/Private/Tools/ControlRigPoseMirrorTable.cpp
+++ b/Engine/Plugins/Experimental/ControlRig/Source/ControlRig/Private/Tools/ControlRigPoseMirrorTable.cpp
@@ -41,7 +41,6 @@ FRigControlCopy* FControlRigPoseMirrorTable::GetControl(FControlRigControlPose&
 	{
 		if (const FName* MatchedName = MatchedControls.Find(Name))
 		{
-
 			int32* Index = Pose.CopyOfControlsNameToIndex.Find(*MatchedName);
 			if (Index != nullptr && (*Index) >= 0 && (*Index) < Pose.CopyOfControls.Num())
 			{
@@ -72,8 +71,8 @@ bool FControlRigPoseMirrorTable::IsMatched(const FName& Name) const
 }
 
 //Now returns mirrored global and local unmirrored
-void FControlRigPoseMirrorTable::GetMirrorTransform(const FRigControlCopy& ControlCopy, bool bIsMatched, FVector& OutGlobalTranslation,
-	FQuat& OutGlobalRotation, FQuat& OutLocalRotation) const
+void FControlRigPoseMirrorTable::GetMirrorTransform(const FRigControlCopy& ControlCopy, bool bDoLocal, bool bIsMatched, FVector& OutGlobalTranslation,
+	FQuat& OutGlobalRotation, FVector& OutLocalTranslation, FQuat& OutLocalRotation) const
 {
 	const UControlRigPoseMirrorSettings* Settings = GetDefault<UControlRigPoseMirrorSettings>();
 	FTransform GlobalTransform = ControlCopy.GlobalTransform;
@@ -81,19 +80,20 @@ void FControlRigPoseMirrorTable::GetMirrorTransform(const FRigControlCopy& Contr
 
 	FTransform LocalTransform = ControlCopy.LocalTransform;
 	OutLocalRotation = LocalTransform.GetRotation();
+	OutLocalTranslation = LocalTransform.GetTranslation();
 	OutGlobalTranslation = GlobalTransform.GetTranslation();
 	OutGlobalRotation = GlobalTransform.GetRotation();
-	
-	if (Settings &&  (bIsMatched || (OutGlobalRotation.IsIdentity() == false && OutLocalRotation.IsIdentity() == false)))
+	if (Settings)
 	{
-		FRigMirrorSettings MirrorSettings;
-
-		MirrorSettings.MirrorAxis = Settings->MirrorAxis;
-		MirrorSettings.AxisToFlip = Settings->AxisToFlip;
-
-		FTransform NewTransform = MirrorSettings.MirrorTransform(GlobalTransform);
-		OutGlobalTranslation = NewTransform.GetTranslation();
-		OutGlobalRotation = NewTransform.GetRotation();
-		return;
+		if (bIsMatched  && (OutGlobalRotation.IsIdentity() == false && OutLocalRotation.IsIdentity() == false))
+		{
+			FRigMirrorSettings MirrorSettings;
+			MirrorSettings.MirrorAxis = Settings->MirrorAxis;
+			MirrorSettings.AxisToFlip = Settings->AxisToFlip;
+			FTransform NewTransform = MirrorSettings.MirrorTransform(GlobalTransform);
+			OutGlobalTranslation = NewTransform.GetTranslation();
+			OutGlobalRotation = NewTransform.GetRotation();
+			return;
+		}
 	}
 }
diff --git a/Engine/Plugins/Experimental/ControlRig/Source/ControlRig/Public/Tools/ControlRigPose.h b/Engine/Plugins/Experimental/ControlRig/Source/ControlRig/Public/Tools/ControlRigPose.h
--- a/Engine/Plugins/Experimental/ControlRig/Source/ControlRig/Public/Tools/ControlRigPose.h
+++ b/Engine/Plugins/Experimental/ControlRig/Source/ControlRig/Public/Tools/ControlRigPose.h
@@ -98,7 +98,7 @@ struct CONTROLRIG_API FControlRigControlPose
 
 	void SavePose(UControlRig* ControlRig, bool bUseAll);
 	void PastePose(UControlRig* ControlRig, bool bDoKey, bool bDoMirror);
-	void SetControlMirrorTransform(UControlRig* ControlRig, const FName& Name, bool bIsMatched, const FVector& GlobalTranslation, const FQuat& GlobalRotation, const FQuat& LocalRotation, bool bNotify, const  FRigControlModifiedContext& Context);
+	void SetControlMirrorTransform(bool bDoLocalSpace, UControlRig* ControlRig, const FName& Name, bool bIsMatched, const FVector& GlobalTranslation, const FQuat& GlobalRotation,const FVector& LocalTranslation, const FQuat& LocalRotation, bool bNotify, const  FRigControlModifiedContext& Context);
 	void PastePoseInternal(UControlRig* ControlRig, bool bDoKey, bool bDoMirror, const TArray<FRigControlCopy>& ControlsToPaste);
 	void BlendWithInitialPoses(FControlRigControlPose& InitialPose, UControlRig* ControlRig, bool bDoKey, bool bDoMirror, float BlendValue);
 
@@ -138,7 +138,7 @@ public:
 	void PastePose(UControlRig* InControlRig, bool bDoKey = false, bool bDoMirror = false);
 	
 	UFUNCTION(BlueprintCallable, Category = "Pose")
-	void SelectControls(UControlRig* InControlRig);
+	void SelectControls(UControlRig* InControlRig, bool bDoMirror = false);
 
 	TArray<FRigControlCopy> GetCurrentPose(UControlRig* InControlRig);
 
diff --git a/Engine/Plugins/Experimental/ControlRig/Source/ControlRig/Public/Tools/ControlRigPoseMirrorTable.h b/Engine/Plugins/Experimental/ControlRig/Source/ControlRig/Public/Tools/ControlRigPoseMirrorTable.h
--- a/Engine/Plugins/Experimental/ControlRig/Source/ControlRig/Public/Tools/ControlRigPoseMirrorTable.h
+++ b/Engine/Plugins/Experimental/ControlRig/Source/ControlRig/Public/Tools/ControlRigPoseMirrorTable.h
@@ -31,7 +31,7 @@ public:
 	bool IsMatched(const FName& ControlName) const;
 
 	/*Return the Mirrored Global(In Control Rig Space) Translation and Mirrored Local Rotation*/
-	void GetMirrorTransform(const FRigControlCopy& ControlCopy,  bool bIsMatched,FVector& OutGlobalTranslation, FQuat& OutGlobalRotation, FQuat& OutLocalRotation) const;
+	void GetMirrorTransform(const FRigControlCopy& ControlCopy, bool bDoLocal, bool bIsMatched,FVector& OutGlobalTranslation, FQuat& OutGlobalRotation, FVector& OutLocalTranslation,FQuat& OutLocalRotation) const;
 
 private:
 
diff --git a/Engine/Plugins/Experimental/ControlRig/Source/ControlRigEditor/Private/EditMode/SControlRigControlViews.cpp b/Engine/Plugins/Experimental/ControlRig/Source/ControlRigEditor/Private/EditMode/SControlRigControlViews.cpp
--- a/Engine/Plugins/Experimental/ControlRig/Source/ControlRigEditor/Private/EditMode/SControlRigControlViews.cpp
+++ b/Engine/Plugins/Experimental/ControlRig/Source/ControlRigEditor/Private/EditMode/SControlRigControlViews.cpp
@@ -523,7 +523,7 @@ FReply SControlRigPoseView::OnSelectControls()
 {
 	if (GetControlRig() && PoseAsset.IsValid())
 	{
-		PoseAsset->SelectControls(GetControlRig());
+		PoseAsset->SelectControls(GetControlRig(), bIsMirror);
 	}
 	return FReply::Handled();
 }
@@ -542,6 +542,8 @@ void SControlRigPoseView::OnPoseBlendChanged(float ChangedVal)
 		}
 
 		PoseAsset->BlendWithInitialPoses(TempPose, ControlRig, false, bIsMirror, PoseBlendValue);
+		PoseAsset->BlendWithInitialPoses(TempPose, ControlRig, false, bIsMirror, PoseBlendValue);
+
 	}
 }
 void SControlRigPoseView::OnBeginSliderMovement()
@@ -570,6 +572,7 @@ void SControlRigPoseView::OnPoseBlendCommited(float ChangedVal, ETextCommit::Typ
 		FScopedTransaction ScopedTransaction(LOCTEXT("PastePoseTransaction", "Paste Pose"));
 		PoseBlendValue = ChangedVal;
 		PoseAsset->BlendWithInitialPoses(TempPose, ControlRig, bIsKey, bIsMirror, PoseBlendValue);
+		PoseAsset->BlendWithInitialPoses(TempPose, ControlRig, bIsKey, bIsMirror, PoseBlendValue);
 		bIsBlending = false;
 		PoseBlendValue = 0.0f;
 	}
