Control Rig: Pose Library: Change default Mirror strings.
#trivial

diff --git a/Engine/Plugins/Experimental/ControlRig/Source/ControlRig/Private/Tools/ControlRigPoseMirrorSettings.cpp b/Engine/Plugins/Experimental/ControlRig/Source/ControlRig/Private/Tools/ControlRigPoseMirrorSettings.cpp
--- a/Engine/Plugins/Experimental/ControlRig/Source/ControlRig/Private/Tools/ControlRigPoseMirrorSettings.cpp
+++ b/Engine/Plugins/Experimental/ControlRig/Source/ControlRig/Private/Tools/ControlRigPoseMirrorSettings.cpp
@@ -3,8 +3,8 @@
 
 UControlRigPoseMirrorSettings::UControlRigPoseMirrorSettings()
 {
-	RightSide = TEXT("Right");
-	LeftSide = TEXT("Left");
+	RightSide = TEXT("_r_");
+	LeftSide = TEXT("_l_");
 	MirrorAxis = EAxis::X;
 	AxisToFlip = EAxis::Z;
 }
