Sequencer: Control Rig: For the animation editor toolbar extender, capture the editor interface instead of the sequence itself since the sequence may change without recreating the extender, if that happens we will open the wrong linked level sequence.

diff --git a/Engine/Plugins/Experimental/ControlRig/Source/ControlRigEditor/Private/ControlRigEditorModule.cpp b/Engine/Plugins/Experimental/ControlRig/Source/ControlRigEditor/Private/ControlRigEditorModule.cpp
--- a/Engine/Plugins/Experimental/ControlRig/Source/ControlRigEditor/Private/ControlRigEditorModule.cpp
+++ b/Engine/Plugins/Experimental/ControlRig/Source/ControlRigEditor/Private/ControlRigEditorModule.cpp
@@ -378,7 +378,7 @@ TSharedRef<FExtender> FControlRigEditorModule::GetAnimationEditorToolbarExtender
 				"Asset",
 				EExtensionHook::After,
 				CommandList,
-				FToolBarExtensionDelegate::CreateRaw(this, &FControlRigEditorModule::HandleAddControlRigExtenderToToolbar, AnimSequence, SkeletalMesh, Skeleton)
+				FToolBarExtensionDelegate::CreateRaw(this, &FControlRigEditorModule::HandleAddControlRigExtenderToToolbar, InAnimationEditor)
 			);
 		}
 	}
@@ -387,241 +387,251 @@ TSharedRef<FExtender> FControlRigEditorModule::GetAnimationEditorToolbarExtender
 }
 
 
-TSharedRef< SWidget > FControlRigEditorModule::GenerateAnimationMenu(UAnimSequence* AnimSequence,USkeletalMesh* SkeletalMesh, USkeleton* Skeleton) 
+TSharedRef< SWidget > FControlRigEditorModule::GenerateAnimationMenu(TSharedRef<IAnimationEditor> InAnimationEditor)
 {
+	USkeleton* Skeleton = InAnimationEditor->GetPersonaToolkit()->GetSkeleton();
+	USkeletalMesh* SkeletalMesh = InAnimationEditor->GetPersonaToolkit()->GetPreviewMesh();
+	if (!SkeletalMesh) //if no preview mesh just get normal mesh
+	{
+		SkeletalMesh = InAnimationEditor->GetPersonaToolkit()->GetMesh();
+	}
+	UAnimSequence* AnimSequence = Cast<UAnimSequence>(InAnimationEditor->GetPersonaToolkit()->GetAnimationAsset());
+
 	const bool bShouldCloseWindowAfterMenuSelection = true;
 	FMenuBuilder MenuBuilder(bShouldCloseWindowAfterMenuSelection, nullptr);
+	if (Skeleton && SkeletalMesh && AnimSequence)
+	{
+		FUIAction EditWithFKControlRig(
+			FExecuteAction::CreateRaw(this, &FControlRigEditorModule::EditWithFKControlRig, AnimSequence, SkeletalMesh, Skeleton));
 
-	FUIAction EditWithFKControlRig(
-		FExecuteAction::CreateRaw(this,&FControlRigEditorModule::EditWithFKControlRig,AnimSequence, SkeletalMesh, Skeleton));
-
-	FUIAction OpenIt(
-		FExecuteAction::CreateStatic(&FControlRigEditorModule::OpenLevelSequence, AnimSequence),
-		FCanExecuteAction::CreateLambda([AnimSequence]()
-			{
-				if (AnimSequence)
+		FUIAction OpenIt(
+			FExecuteAction::CreateStatic(&FControlRigEditorModule::OpenLevelSequence, AnimSequence),
+			FCanExecuteAction::CreateLambda([AnimSequence]()
 				{
-					if (IInterface_AssetUserData* AnimAssetUserData = Cast< IInterface_AssetUserData >(AnimSequence))
+					if (AnimSequence)
 					{
-						UAnimSequenceLevelSequenceLink* AnimLevelLink = AnimAssetUserData->GetAssetUserData< UAnimSequenceLevelSequenceLink >();
-						if (AnimLevelLink)
+						if (IInterface_AssetUserData* AnimAssetUserData = Cast< IInterface_AssetUserData >(AnimSequence))
 						{
-							ULevelSequence* LevelSequence = AnimLevelLink->ResolveLevelSequence();
-							if (LevelSequence)
+							UAnimSequenceLevelSequenceLink* AnimLevelLink = AnimAssetUserData->GetAssetUserData< UAnimSequenceLevelSequenceLink >();
+							if (AnimLevelLink)
 							{
-								return true;
+								ULevelSequence* LevelSequence = AnimLevelLink->ResolveLevelSequence();
+								if (LevelSequence)
+								{
+									return true;
+								}
 							}
 						}
 					}
+					return false;
 				}
-				return false;
-			}
-		)	
-	
-	);
+			)
 
-	FUIAction UnLinkIt(
-		FExecuteAction::CreateStatic(&FControlRigEditorModule::UnLinkLevelSequence, AnimSequence),
-		FCanExecuteAction::CreateLambda([AnimSequence]()
-			{
-				if (AnimSequence)
+		);
+
+		FUIAction UnLinkIt(
+			FExecuteAction::CreateStatic(&FControlRigEditorModule::UnLinkLevelSequence, AnimSequence),
+			FCanExecuteAction::CreateLambda([AnimSequence]()
 				{
-					if (IInterface_AssetUserData* AnimAssetUserData = Cast< IInterface_AssetUserData >(AnimSequence))
+					if (AnimSequence)
 					{
-						UAnimSequenceLevelSequenceLink* AnimLevelLink = AnimAssetUserData->GetAssetUserData< UAnimSequenceLevelSequenceLink >();
-						if (AnimLevelLink)
+						if (IInterface_AssetUserData* AnimAssetUserData = Cast< IInterface_AssetUserData >(AnimSequence))
 						{
-							ULevelSequence* LevelSequence = AnimLevelLink->ResolveLevelSequence();
-							if (LevelSequence)
+							UAnimSequenceLevelSequenceLink* AnimLevelLink = AnimAssetUserData->GetAssetUserData< UAnimSequenceLevelSequenceLink >();
+							if (AnimLevelLink)
 							{
-								return true;
+								ULevelSequence* LevelSequence = AnimLevelLink->ResolveLevelSequence();
+								if (LevelSequence)
+								{
+									return true;
+								}
 							}
 						}
 					}
+					return false;
 				}
-				return false;
-			}
-		)
+			)
 
-	);
+		);
 
-	FUIAction ToggleFilterAssetBySkeleton(
-		FExecuteAction::CreateLambda([this]()
-			{
-				bFilterAssetBySkeleton = bFilterAssetBySkeleton ? false : true;
-			}
-		),
-		FCanExecuteAction(),
-		FIsActionChecked::CreateLambda([this]()
-			{
-				return bFilterAssetBySkeleton;
-			}
-		)
-	);
-	if (Skeleton)
-	{
-		MenuBuilder.BeginSection("Control Rig", LOCTEXT("ControlRig", "Control Rig"));
+		FUIAction ToggleFilterAssetBySkeleton(
+			FExecuteAction::CreateLambda([this]()
+				{
+					bFilterAssetBySkeleton = bFilterAssetBySkeleton ? false : true;
+				}
+			),
+			FCanExecuteAction(),
+					FIsActionChecked::CreateLambda([this]()
+						{
+							return bFilterAssetBySkeleton;
+						}
+					)
+					);
+		if (Skeleton)
 		{
-			MenuBuilder.AddMenuEntry(LOCTEXT("EditWithFKControlRig", "Edit With FK Control Rig"),
-				FText(), FSlateIcon(), EditWithFKControlRig, NAME_None, EUserInterfaceActionType::Button);
+			MenuBuilder.BeginSection("Control Rig", LOCTEXT("ControlRig", "Control Rig"));
+			{
+				MenuBuilder.AddMenuEntry(LOCTEXT("EditWithFKControlRig", "Edit With FK Control Rig"),
+					FText(), FSlateIcon(), EditWithFKControlRig, NAME_None, EUserInterfaceActionType::Button);
 
 
-			MenuBuilder.AddMenuEntry(
-				LOCTEXT("FilterAssetBySkeleton", "Filter Asset By Skeleton"),
-				LOCTEXT("FilterAssetBySkeletonTooltip", "Filters Control Rig Assets To Match Current Skeleton"),
-				FSlateIcon(),
-				ToggleFilterAssetBySkeleton,
-				NAME_None,
-				EUserInterfaceActionType::ToggleButton);
+				MenuBuilder.AddMenuEntry(
+					LOCTEXT("FilterAssetBySkeleton", "Filter Asset By Skeleton"),
+					LOCTEXT("FilterAssetBySkeletonTooltip", "Filters Control Rig Assets To Match Current Skeleton"),
+					FSlateIcon(),
+					ToggleFilterAssetBySkeleton,
+					NAME_None,
+					EUserInterfaceActionType::ToggleButton);
 
-			MenuBuilder.AddSubMenu(
-				LOCTEXT("BakeToControlRig", "Bake To Control Rig"), NSLOCTEXT("AnimationModeToolkit", "BakeToControlRigTooltip", "This Control Rig will Drive This Animation."),
-				FNewMenuDelegate::CreateLambda([this, AnimSequence, SkeletalMesh, Skeleton](FMenuBuilder& InSubMenuBuilder)
-					{
-						//todo move to .h for ue5
-						class FControlRigClassFilter : public IClassViewerFilter
+				MenuBuilder.AddSubMenu(
+					LOCTEXT("BakeToControlRig", "Bake To Control Rig"), NSLOCTEXT("AnimationModeToolkit", "BakeToControlRigTooltip", "This Control Rig will Drive This Animation."),
+					FNewMenuDelegate::CreateLambda([this, AnimSequence, SkeletalMesh, Skeleton](FMenuBuilder& InSubMenuBuilder)
 						{
-						public:
-							FControlRigClassFilter(bool bInCheckSkeleton, bool bInCheckAnimatable, bool bInCheckInversion, USkeleton* InSkeleton) :
-								bFilterAssetBySkeleton(bInCheckSkeleton),
-								bFilterExposesAnimatableControls(bInCheckAnimatable),
-								bFilterInversion(bInCheckInversion),
-								AssetRegistry(FModuleManager::GetModuleChecked<FAssetRegistryModule>(TEXT("AssetRegistry")).Get())
+							//todo move to .h for ue5
+							class FControlRigClassFilter : public IClassViewerFilter
 							{
-								if (InSkeleton)
+							public:
+								FControlRigClassFilter(bool bInCheckSkeleton, bool bInCheckAnimatable, bool bInCheckInversion, USkeleton* InSkeleton) :
+									bFilterAssetBySkeleton(bInCheckSkeleton),
+									bFilterExposesAnimatableControls(bInCheckAnimatable),
+									bFilterInversion(bInCheckInversion),
+									AssetRegistry(FModuleManager::GetModuleChecked<FAssetRegistryModule>(TEXT("AssetRegistry")).Get())
 								{
-									SkeletonName = FAssetData(InSkeleton).GetExportTextName();
+									if (InSkeleton)
+									{
+										SkeletonName = FAssetData(InSkeleton).GetExportTextName();
+									}
 								}
-							}
-							bool bFilterAssetBySkeleton;
-							bool bFilterExposesAnimatableControls;
-							bool bFilterInversion;
+								bool bFilterAssetBySkeleton;
+								bool bFilterExposesAnimatableControls;
+								bool bFilterInversion;
 
-							FString SkeletonName;
-							const IAssetRegistry& AssetRegistry;
+								FString SkeletonName;
+								const IAssetRegistry& AssetRegistry;
 
-							bool MatchesFilter(const FAssetData& AssetData)
-							{
-								bool bExposesAnimatableControls = AssetData.GetTagValueRef<bool>(TEXT("bExposesAnimatableControls"));
-								if (bFilterExposesAnimatableControls == true && bExposesAnimatableControls == false)
-								{
-									return false;
-								}
-								if (bFilterInversion)
+								bool MatchesFilter(const FAssetData& AssetData)
 								{
-									bool bHasInversion = false;
-									FAssetDataTagMapSharedView::FFindTagResult Tag = AssetData.TagsAndValues.FindTag(TEXT("SupportedEventNames"));
-									if (Tag.IsSet())
+									bool bExposesAnimatableControls = AssetData.GetTagValueRef<bool>(TEXT("bExposesAnimatableControls"));
+									if (bFilterExposesAnimatableControls == true && bExposesAnimatableControls == false)
 									{
-										FString EventString = FRigUnit_InverseExecution::EventName.ToString();
-										TArray<FString> SupportedEventNames;
-										Tag.GetValue().ParseIntoArray(SupportedEventNames, TEXT(","), true);
-
-										for (const FString& Name : SupportedEventNames)
+										return false;
+									}
+									if (bFilterInversion)
+									{
+										bool bHasInversion = false;
+										FAssetDataTagMapSharedView::FFindTagResult Tag = AssetData.TagsAndValues.FindTag(TEXT("SupportedEventNames"));
+										if (Tag.IsSet())
 										{
-											if (Name.Contains(EventString))
+											FString EventString = FRigUnit_InverseExecution::EventName.ToString();
+											TArray<FString> SupportedEventNames;
+											Tag.GetValue().ParseIntoArray(SupportedEventNames, TEXT(","), true);
+
+											for (const FString& Name : SupportedEventNames)
 											{
-												bHasInversion = true;
-												break;
+												if (Name.Contains(EventString))
+												{
+													bHasInversion = true;
+													break;
+												}
+											}
+											if (bHasInversion == false)
+											{
+												return false;
 											}
-										}
-										if (bHasInversion == false)
-										{
-											return false;
 										}
 									}
-								}
-								if (bFilterAssetBySkeleton)
-								{
-									FString PreviewSkeletalMesh = AssetData.GetTagValueRef<FString>(TEXT("PreviewSkeletalMesh"));
-									if (PreviewSkeletalMesh.Len() > 0)
+									if (bFilterAssetBySkeleton)
 									{
-										FAssetData SkelMeshData = AssetRegistry.GetAssetByObjectPath(FName(*PreviewSkeletalMesh));
-										FString PreviewSkeleton = SkelMeshData.GetTagValueRef<FString>(TEXT("Skeleton"));
+										FString PreviewSkeletalMesh = AssetData.GetTagValueRef<FString>(TEXT("PreviewSkeletalMesh"));
+										if (PreviewSkeletalMesh.Len() > 0)
+										{
+											FAssetData SkelMeshData = AssetRegistry.GetAssetByObjectPath(FName(*PreviewSkeletalMesh));
+											FString PreviewSkeleton = SkelMeshData.GetTagValueRef<FString>(TEXT("Skeleton"));
+											if (PreviewSkeleton == SkeletonName)
+											{
+												return true;
+											}
+										}
+										FString PreviewSkeleton = AssetData.GetTagValueRef<FString>(TEXT("PreviewSkeleton"));
 										if (PreviewSkeleton == SkeletonName)
 										{
 											return true;
 										}
+										FString SourceHierarchyImport = AssetData.GetTagValueRef<FString>(TEXT("SourceHierarchyImport"));
+										if (SourceHierarchyImport == SkeletonName)
+										{
+											return true;
+										}
+										FString SourceCurveImport = AssetData.GetTagValueRef<FString>(TEXT("SourceCurveImport"));
+										if (SourceCurveImport == SkeletonName)
+										{
+											return true;
+										}
+										return false;
 									}
-									FString PreviewSkeleton = AssetData.GetTagValueRef<FString>(TEXT("PreviewSkeleton"));
-									if (PreviewSkeleton == SkeletonName)
-									{
-										return true;
-									}
-									FString SourceHierarchyImport = AssetData.GetTagValueRef<FString>(TEXT("SourceHierarchyImport"));
-									if (SourceHierarchyImport == SkeletonName)
-									{
-										return true;
-									}
-									FString SourceCurveImport = AssetData.GetTagValueRef<FString>(TEXT("SourceCurveImport"));
-									if (SourceCurveImport == SkeletonName)
+									return true;
+
+								}
+								bool IsClassAllowed(const FClassViewerInitializationOptions& InInitOptions, const UClass* InClass, TSharedRef< FClassViewerFilterFuncs > InFilterFuncs) override
+								{
+									const bool bChildOfObjectClass = InClass->IsChildOf(UControlRig::StaticClass());
+									const bool bMatchesFlags = !InClass->HasAnyClassFlags(CLASS_Hidden | CLASS_HideDropDown | CLASS_Deprecated | CLASS_Abstract);
+									const bool bNotNative = !InClass->IsNative();
+
+									if (bChildOfObjectClass && bMatchesFlags && bNotNative)
 									{
-										return true;
+										FAssetData AssetData(InClass);
+										return MatchesFilter(AssetData);
+
 									}
 									return false;
 								}
-								return true;
-
-							}
-							bool IsClassAllowed(const FClassViewerInitializationOptions& InInitOptions, const UClass* InClass, TSharedRef< FClassViewerFilterFuncs > InFilterFuncs) override
-							{
-								const bool bChildOfObjectClass = InClass->IsChildOf(UControlRig::StaticClass());
-								const bool bMatchesFlags = !InClass->HasAnyClassFlags(CLASS_Hidden | CLASS_HideDropDown | CLASS_Deprecated | CLASS_Abstract);
-								const bool bNotNative = !InClass->IsNative();
 
-								if (bChildOfObjectClass && bMatchesFlags && bNotNative)
+								virtual bool IsUnloadedClassAllowed(const FClassViewerInitializationOptions& InInitOptions, const TSharedRef< const IUnloadedBlueprintData > InUnloadedClassData, TSharedRef< FClassViewerFilterFuncs > InFilterFuncs) override
 								{
-									FAssetData AssetData(InClass);
-									return MatchesFilter(AssetData);
+									const bool bChildOfObjectClass = InUnloadedClassData->IsChildOf(UControlRig::StaticClass());
+									const bool bMatchesFlags = !InUnloadedClassData->HasAnyClassFlags(CLASS_Hidden | CLASS_HideDropDown | CLASS_Deprecated | CLASS_Abstract);
+									if (bChildOfObjectClass && bMatchesFlags)
+									{
+										FString GeneratedClassPathString = InUnloadedClassData->GetClassPath().ToString();
+										FName BlueprintPath = FName(*GeneratedClassPathString.LeftChop(2)); // Chop off _C
+										FAssetData AssetData = AssetRegistry.GetAssetByObjectPath(BlueprintPath);
+										return MatchesFilter(AssetData);
 
+									}
+									return false;
 								}
-								return false;
-							}
-
-							virtual bool IsUnloadedClassAllowed(const FClassViewerInitializationOptions& InInitOptions, const TSharedRef< const IUnloadedBlueprintData > InUnloadedClassData, TSharedRef< FClassViewerFilterFuncs > InFilterFuncs) override
-							{
-								const bool bChildOfObjectClass = InUnloadedClassData->IsChildOf(UControlRig::StaticClass());
-								const bool bMatchesFlags = !InUnloadedClassData->HasAnyClassFlags(CLASS_Hidden | CLASS_HideDropDown | CLASS_Deprecated | CLASS_Abstract);
-								if (bChildOfObjectClass && bMatchesFlags)
-								{
-									FString GeneratedClassPathString = InUnloadedClassData->GetClassPath().ToString();
-									FName BlueprintPath = FName(*GeneratedClassPathString.LeftChop(2)); // Chop off _C
-									FAssetData AssetData = AssetRegistry.GetAssetByObjectPath(BlueprintPath);
-									return MatchesFilter(AssetData);
 
-								}
-								return false;
-							}
+							};
 
-						};
+							FClassViewerInitializationOptions Options;
+							Options.bShowUnloadedBlueprints = true;
+							Options.NameTypeToDisplay = EClassViewerNameTypeToDisplay::DisplayName;
 
-						FClassViewerInitializationOptions Options;
-						Options.bShowUnloadedBlueprints = true;
-						Options.NameTypeToDisplay = EClassViewerNameTypeToDisplay::DisplayName;
+							TSharedPtr<FControlRigClassFilter> ClassFilter = MakeShareable(new FControlRigClassFilter(bFilterAssetBySkeleton, true, true, Skeleton));
+							Options.ClassFilter = ClassFilter;
+							Options.bShowNoneOption = false;
 
-						TSharedPtr<FControlRigClassFilter> ClassFilter = MakeShareable(new FControlRigClassFilter(bFilterAssetBySkeleton,true, true,Skeleton));
-						Options.ClassFilter = ClassFilter;
-						Options.bShowNoneOption = false;
+							FClassViewerModule& ClassViewerModule = FModuleManager::LoadModuleChecked<FClassViewerModule>("ClassViewer");
 
-						FClassViewerModule& ClassViewerModule = FModuleManager::LoadModuleChecked<FClassViewerModule>("ClassViewer");
+							TSharedRef<SWidget> ClassViewer = ClassViewerModule.CreateClassViewer(Options, FOnClassPicked::CreateRaw(this, &FControlRigEditorModule::BakeToControlRig, AnimSequence, SkeletalMesh, Skeleton));
+							InSubMenuBuilder.AddWidget(ClassViewer, FText::GetEmpty(), true);
+						})
+				);
+			}
+			MenuBuilder.EndSection();
 
-						TSharedRef<SWidget> ClassViewer = ClassViewerModule.CreateClassViewer(Options, FOnClassPicked::CreateRaw(this, &FControlRigEditorModule::BakeToControlRig, AnimSequence, SkeletalMesh, Skeleton));
-						InSubMenuBuilder.AddWidget(ClassViewer, FText::GetEmpty(), true);
-					})
-			);
 		}
-		MenuBuilder.EndSection();
-		
-	}
 
-	MenuBuilder.AddMenuEntry(LOCTEXT("OpenLevelSequence", "Open Level Sequence"),
-		FText(), FSlateIcon(), OpenIt, NAME_None, EUserInterfaceActionType::Button);
-
-	MenuBuilder.AddMenuEntry(LOCTEXT("UnlinkLevelSequence", "Unlink Level Sequence"),
-		FText(), FSlateIcon(), UnLinkIt, NAME_None, EUserInterfaceActionType::Button);
+		MenuBuilder.AddMenuEntry(LOCTEXT("OpenLevelSequence", "Open Level Sequence"),
+			FText(), FSlateIcon(), OpenIt, NAME_None, EUserInterfaceActionType::Button);
 
+		MenuBuilder.AddMenuEntry(LOCTEXT("UnlinkLevelSequence", "Unlink Level Sequence"),
+			FText(), FSlateIcon(), UnLinkIt, NAME_None, EUserInterfaceActionType::Button);
 
+	}
 	return MenuBuilder.MakeWidget();
+
 }
 
 
@@ -858,12 +868,12 @@ void FControlRigEditorModule::OpenLevelSequence(UAnimSequence* AnimSequence)
 	}
 }
 
-void FControlRigEditorModule::HandleAddControlRigExtenderToToolbar(FToolBarBuilder& ParentToolbarBuilder, UAnimSequence* AnimSequence, USkeletalMesh* SkeletalMesh,USkeleton* Skeleton)
+void FControlRigEditorModule::HandleAddControlRigExtenderToToolbar(FToolBarBuilder& ParentToolbarBuilder, TSharedRef<IAnimationEditor> InAnimationEditor)
 {
 
 	ParentToolbarBuilder.AddComboButton(
 		FUIAction(),
-		FOnGetContent::CreateRaw(this, &FControlRigEditorModule::GenerateAnimationMenu,AnimSequence,SkeletalMesh,Skeleton),
+		FOnGetContent::CreateRaw(this, &FControlRigEditorModule::GenerateAnimationMenu, InAnimationEditor),
 		LOCTEXT("EditInSequencer", "Edit in Sequencer"),
 		LOCTEXT("EditInSequencer_Tooltip", "Edit this Anim Sequence In Sequencer."),
 		FSlateIcon(FEditorStyle::GetStyleSetName(), "Persona.ExportToFBX")
diff --git a/Engine/Plugins/Experimental/ControlRig/Source/ControlRigEditor/Private/ControlRigEditorModule.h b/Engine/Plugins/Experimental/ControlRig/Source/ControlRigEditor/Private/ControlRigEditorModule.h
--- a/Engine/Plugins/Experimental/ControlRig/Source/ControlRigEditor/Private/ControlRigEditorModule.h
+++ b/Engine/Plugins/Experimental/ControlRig/Source/ControlRigEditor/Private/ControlRigEditorModule.h
@@ -41,8 +41,8 @@ public:
 	
 	/** Animation Toolbar Extender*/
 	TSharedRef<FExtender> GetAnimationEditorToolbarExtender(const TSharedRef<FUICommandList> CommandList, TSharedRef<IAnimationEditor> InAnimationEditor);
-	void HandleAddControlRigExtenderToToolbar(FToolBarBuilder& ParentToolbarBuilder, UAnimSequence* AnimSequence, USkeletalMesh* SkeletalMesh,USkeleton* Skeleton);
-	TSharedRef< SWidget > GenerateAnimationMenu(UAnimSequence* AnimSequence, USkeletalMesh* SkeletalMesh,USkeleton* Skeleton) ;
+	void HandleAddControlRigExtenderToToolbar(FToolBarBuilder& ParentToolbarBuilder, TSharedRef<IAnimationEditor> InAnimationEditor);
+	TSharedRef< SWidget > GenerateAnimationMenu(TSharedRef<IAnimationEditor> InAnimationEditor) ;
 	void ToggleIsDrivenByLevelSequence(UAnimSequence* AnimSequence)const;
 	bool IsDrivenByLevelSequence(UAnimSequence* AnimSequence)const;
 	void EditWithFKControlRig(UAnimSequence* AnimSequence, USkeletalMesh* SkelMesh, USkeleton* InSkeleton);
