Control Rig: Move up and down arguments in the details panel for collapse nodes / functions

diff --git a/Engine/Plugins/Experimental/ControlRig/Source/ControlRigDeveloper/Private/ControlRigBlueprint.cpp b/Engine/Plugins/Experimental/ControlRig/Source/ControlRigDeveloper/Private/ControlRigBlueprint.cpp
--- a/Engine/Plugins/Experimental/ControlRig/Source/ControlRigDeveloper/Private/ControlRigBlueprint.cpp
+++ b/Engine/Plugins/Experimental/ControlRig/Source/ControlRigDeveloper/Private/ControlRigBlueprint.cpp
@@ -1660,6 +1660,7 @@ void UControlRigBlueprint::HandleModifiedEvent(ERigVMGraphNotifType InNotifType,
 			}
 			case ERigVMGraphNotifType::PinWatchedChanged:
 			case ERigVMGraphNotifType::PinTypeChanged:
+			case ERigVMGraphNotifType::PinIndexChanged:
 			{
 				if (URigVMPin* ModelPin = Cast<URigVMPin>(InSubject))
 				{
diff --git a/Engine/Plugins/Experimental/ControlRig/Source/ControlRigDeveloper/Private/Graph/ControlRigGraph.cpp b/Engine/Plugins/Experimental/ControlRig/Source/ControlRigDeveloper/Private/Graph/ControlRigGraph.cpp
--- a/Engine/Plugins/Experimental/ControlRig/Source/ControlRigDeveloper/Private/Graph/ControlRigGraph.cpp
+++ b/Engine/Plugins/Experimental/ControlRig/Source/ControlRigDeveloper/Private/Graph/ControlRigGraph.cpp
@@ -547,6 +547,7 @@ void UControlRigGraph::HandleModifiedEvent(ERigVMGraphNotifType InNotifType, URi
 		case ERigVMGraphNotifType::PinArraySizeChanged:
 		case ERigVMGraphNotifType::PinDirectionChanged:
 		case ERigVMGraphNotifType::PinTypeChanged:
+		case ERigVMGraphNotifType::PinIndexChanged:
 		case ERigVMGraphNotifType::PinBoundVariableChanged:
 		case ERigVMGraphNotifType::PinAdded:
 		case ERigVMGraphNotifType::PinRemoved:
diff --git a/Engine/Plugins/Experimental/ControlRig/Source/ControlRigEditor/Private/ControlRigGraphDetails.cpp b/Engine/Plugins/Experimental/ControlRig/Source/ControlRigEditor/Private/ControlRigGraphDetails.cpp
--- a/Engine/Plugins/Experimental/ControlRig/Source/ControlRigEditor/Private/ControlRigGraphDetails.cpp
+++ b/Engine/Plugins/Experimental/ControlRig/Source/ControlRigEditor/Private/ControlRigGraphDetails.cpp
@@ -103,6 +103,7 @@ void FControlRigArgumentGroupLayout::HandleModifiedEvent(ERigVMGraphNotifType In
 	{
 		case ERigVMGraphNotifType::PinAdded:
 		case ERigVMGraphNotifType::PinRemoved:
+		case ERigVMGraphNotifType::PinIndexChanged:
 		{
 			URigVMPin* Pin = CastChecked<URigVMPin>(InSubject);
 			if (Pin->GetNode() == LibraryNode)
@@ -226,54 +227,38 @@ void FControlRigArgumentLayout::OnRemoveClicked()
 
 FReply FControlRigArgumentLayout::OnArgMoveUp()
 {
-	/* todo
-
-	const int32 ThisParamIndex = TargetNode->UserDefinedPins.Find(ParamItemPtr.Pin());
-	const int32 NewParamIndex = ThisParamIndex - 1;
-	if (ThisParamIndex != INDEX_NONE && NewParamIndex >= 0)
+	if (PinPtr.IsValid() && ControlRigBlueprintPtr.IsValid())
 	{
-		const FScopedTransaction Transaction(LOCTEXT("K2_MovePinUp", "Move Pin Up"));
-		TArray<UK2Node_EditablePinBase*> TargetNodes = GatherAllResultNodes(TargetNode);
-		for (UK2Node_EditablePinBase* Node : TargetNodes)
+		URigVMPin* Pin = PinPtr.Get();
+		UControlRigBlueprint* Blueprint = ControlRigBlueprintPtr.Get();
+		if (URigVMLibraryNode* LibraryNode = Cast<URigVMLibraryNode>(Pin->GetNode()))
 		{
-			Node->Modify();
-			Node->UserDefinedPins.Swap(ThisParamIndex, NewParamIndex);
-
-			TSharedPtr<FBaseBlueprintGraphActionDetails> GraphActionDetails = GraphActionDetailsPtr.Pin();
-			if (GraphActionDetails.IsValid())
+			if (URigVMController* Controller = Blueprint->GetController(LibraryNode->GetContainedGraph()))
 			{
-				GraphActionDetails->OnParamsChanged(Node, true);
+				Controller->SetExposedPinIndex(Pin->GetFName(), Pin->GetPinIndex() - 1);
+				return FReply::Handled();
 			}
 		}
 	}
-
-	*/
-	return FReply::Handled();
+	return FReply::Unhandled();
 }
 
 FReply FControlRigArgumentLayout::OnArgMoveDown()
 {
-	/* todo
-	const int32 ThisParamIndex = TargetNode->UserDefinedPins.Find(ParamItemPtr.Pin());
-	const int32 NewParamIndex = ThisParamIndex + 1;
-	if (ThisParamIndex != INDEX_NONE && NewParamIndex < TargetNode->UserDefinedPins.Num())
+	if (PinPtr.IsValid() && ControlRigBlueprintPtr.IsValid())
 	{
-		const FScopedTransaction Transaction(LOCTEXT("K2_MovePinDown", "Move Pin Down"));
-		TArray<UK2Node_EditablePinBase*> TargetNodes = GatherAllResultNodes(TargetNode);
-		for (UK2Node_EditablePinBase* Node : TargetNodes)
+		URigVMPin* Pin = PinPtr.Get();
+		UControlRigBlueprint* Blueprint = ControlRigBlueprintPtr.Get();
+		if (URigVMLibraryNode* LibraryNode = Cast<URigVMLibraryNode>(Pin->GetNode()))
 		{
-			Node->Modify();
-			Node->UserDefinedPins.Swap(ThisParamIndex, NewParamIndex);
-
-			TSharedPtr<FBaseBlueprintGraphActionDetails> GraphActionDetails = GraphActionDetailsPtr.Pin();
-			if (GraphActionDetails.IsValid())
+			if (URigVMController* Controller = Blueprint->GetController(LibraryNode->GetContainedGraph()))
 			{
-				GraphActionDetails->OnParamsChanged(Node, true);
+				Controller->SetExposedPinIndex(Pin->GetFName(), Pin->GetPinIndex() + 1);
+				return FReply::Handled();
 			}
 		}
 	}
-	*/
-	return FReply::Handled();
+	return FReply::Unhandled();
 }
 
 bool FControlRigArgumentLayout::ShouldPinBeReadOnly(bool bIsEditingPinType/* = false*/) const
@@ -501,6 +486,7 @@ void FControlRigArgumentDefaultNode::HandleModifiedEvent(ERigVMGraphNotifType In
 		case ERigVMGraphNotifType::PinAdded:
 		case ERigVMGraphNotifType::PinRemoved:
 		case ERigVMGraphNotifType::PinTypeChanged:
+		case ERigVMGraphNotifType::PinIndexChanged:
 		case ERigVMGraphNotifType::PinRenamed:
 		{
 			URigVMPin* Pin = CastChecked<URigVMPin>(InSubject);
diff --git a/Engine/Source/Developer/RigVMDeveloper/Private/RigVMModel/RigVMController.cpp b/Engine/Source/Developer/RigVMDeveloper/Private/RigVMModel/RigVMController.cpp
--- a/Engine/Source/Developer/RigVMDeveloper/Private/RigVMModel/RigVMController.cpp
+++ b/Engine/Source/Developer/RigVMDeveloper/Private/RigVMModel/RigVMController.cpp
@@ -5699,6 +5699,69 @@ bool URigVMController::ChangeExposedPinType(const FName& InPinName, const FStrin
 	return true;
 }
 
+bool URigVMController::SetExposedPinIndex(const FName& InPinName, int32 InNewIndex, bool bSetupUndoRedo)
+{
+	if (!IsValidGraph())
+	{
+		return false;
+	}
+
+	URigVMGraph* Graph = GetGraph();
+	check(Graph);
+
+	FString PinPath = InPinName.ToString();
+	if (PinPath.Contains(TEXT(".")))
+	{
+		ReportError(TEXT("Cannot change pin index for pins on nodes for now - only within collapse nodes."));
+		return false;
+	}
+
+	URigVMLibraryNode* LibraryNode = Cast<URigVMLibraryNode>(Graph->GetOuter());
+	if (LibraryNode == nullptr)
+	{
+		ReportError(TEXT("Graph is not under a Collapse Node"));
+		return false;
+	}
+
+	URigVMPin* Pin = LibraryNode->FindPin(PinPath);
+	if (Pin == nullptr)
+	{
+		ReportErrorf(TEXT("Cannot find exposed pin '%s'."), *PinPath);
+		return false;
+	}
+
+	if (Pin->GetPinIndex() == InNewIndex)
+	{
+		return false;
+	}
+
+	if (InNewIndex < 0 || InNewIndex >= LibraryNode->GetPins().Num())
+	{
+		ReportErrorf(TEXT("Invalid new pin index '%d'."), InNewIndex);
+		return false;
+	}
+
+	FRigVMSetPinIndexAction PinIndexAction(Pin, InNewIndex);
+	{
+		LibraryNode->Pins.Remove(Pin);
+		LibraryNode->Pins.Insert(Pin, InNewIndex);
+
+		FRigVMControllerGraphGuard GraphGuard(this, LibraryNode->GetGraph(), false);
+		Notify(ERigVMGraphNotifType::PinIndexChanged, Pin);
+	}
+
+	RefreshFunctionPins(LibraryNode->GetEntryNode(), true);
+	RefreshFunctionPins(LibraryNode->GetReturnNode(), true);
+	RefreshFunctionReferences(LibraryNode, false);
+	
+	if (bSetupUndoRedo)
+	{
+		ActionStack->AddAction(PinIndexAction);
+	}
+
+	return true;
+}
+
 URigVMFunctionReferenceNode* URigVMController::AddFunctionReferenceNode(URigVMLibraryNode* InFunctionDefinition, const FVector2D& InNodePosition, const FString& InNodeName, bool bSetupUndoRedo)
 {
 	if (!IsValidGraph())
diff --git a/Engine/Source/Developer/RigVMDeveloper/Private/RigVMModel/RigVMControllerActions.cpp b/Engine/Source/Developer/RigVMDeveloper/Private/RigVMModel/RigVMControllerActions.cpp
--- a/Engine/Source/Developer/RigVMDeveloper/Private/RigVMModel/RigVMControllerActions.cpp
+++ b/Engine/Source/Developer/RigVMDeveloper/Private/RigVMModel/RigVMControllerActions.cpp
@@ -1573,3 +1573,35 @@ bool FRigVMRenameExposedPinAction::Redo(URigVMController* InController)
 	}
 	return FRigVMBaseAction::Redo(InController);
 }
+
+FRigVMSetPinIndexAction::FRigVMSetPinIndexAction()
+	: PinPath()
+	, OldIndex(INDEX_NONE)
+	, NewIndex(INDEX_NONE)
+{
+}
+
+FRigVMSetPinIndexAction::FRigVMSetPinIndexAction(URigVMPin* InPin, int32 InNewIndex)
+	: PinPath(InPin->GetPinPath())
+	, OldIndex(InPin->GetPinIndex())
+	, NewIndex(InNewIndex)
+{
+}
+
+bool FRigVMSetPinIndexAction::Undo(URigVMController* InController)
+{
+	if (!FRigVMBaseAction::Undo(InController))
+	{
+		return false;
+	}
+	return InController->SetExposedPinIndex(*PinPath, OldIndex);
+}
+
+bool FRigVMSetPinIndexAction::Redo(URigVMController* InController)
+{
+	if (!InController->SetExposedPinIndex(*PinPath, NewIndex))
+	{
+		return false;
+	}
+	return FRigVMBaseAction::Redo(InController);
+}
diff --git a/Engine/Source/Developer/RigVMDeveloper/Public/RigVMModel/RigVMController.h b/Engine/Source/Developer/RigVMDeveloper/Public/RigVMModel/RigVMController.h
--- a/Engine/Source/Developer/RigVMDeveloper/Public/RigVMModel/RigVMController.h
+++ b/Engine/Source/Developer/RigVMDeveloper/Public/RigVMModel/RigVMController.h
@@ -499,6 +499,10 @@ public:
 	UFUNCTION(BlueprintCallable, Category = RigVMController)
 	bool ChangeExposedPinType(const FName& InPinName, const FString& InCPPType, const FName& InCPPTypeObjectPath, bool bSetupUndoRedo = true);
 
+	// Sets the index for an exposed pin. This can be used to move the pin up and down on the node.
+	UFUNCTION(BlueprintCallable, Category = RigVMController)
+	bool SetExposedPinIndex(const FName& InPinName, int32 InNewIndex, bool bSetupUndoRedo = true);
+
 	// Adds a function reference / invocation to the graph
 	UFUNCTION(BlueprintCallable, Category = RigVMController)
 	URigVMFunctionReferenceNode* AddFunctionReferenceNode(URigVMLibraryNode* InFunctionDefinition, const FVector2D& InNodePosition = FVector2D::ZeroVector, const FString& InNodeName = TEXT(""), bool bSetupUndoRedo = true);
diff --git a/Engine/Source/Developer/RigVMDeveloper/Public/RigVMModel/RigVMControllerActions.h b/Engine/Source/Developer/RigVMDeveloper/Public/RigVMModel/RigVMControllerActions.h
--- a/Engine/Source/Developer/RigVMDeveloper/Public/RigVMModel/RigVMControllerActions.h
+++ b/Engine/Source/Developer/RigVMDeveloper/Public/RigVMModel/RigVMControllerActions.h
@@ -1349,4 +1349,30 @@ public:
 
 	UPROPERTY()
 	FString NewPinName;
+};
+
+/**
+ * An action to reorder pins on a node
+ */
+USTRUCT()
+struct FRigVMSetPinIndexAction : public FRigVMBaseAction
+{
+	GENERATED_BODY()
+
+public:
+
+	FRigVMSetPinIndexAction();
+	FRigVMSetPinIndexAction(URigVMPin* InPin, int32 InNewIndex);
+	virtual ~FRigVMSetPinIndexAction() {};
+	virtual bool Undo(URigVMController* InController) override;
+	virtual bool Redo(URigVMController* InController) override;
+
+	UPROPERTY()
+	FString PinPath;
+
+	UPROPERTY()
+	int32 OldIndex;
+
+	UPROPERTY()
+	int32 NewIndex;
 };
\ No newline at end of file
diff --git a/Engine/Source/Developer/RigVMDeveloper/Public/RigVMModel/RigVMNotifications.h b/Engine/Source/Developer/RigVMDeveloper/Public/RigVMModel/RigVMNotifications.h
--- a/Engine/Source/Developer/RigVMDeveloper/Public/RigVMModel/RigVMNotifications.h
+++ b/Engine/Source/Developer/RigVMDeveloper/Public/RigVMModel/RigVMNotifications.h
@@ -33,6 +33,7 @@ enum class ERigVMGraphNotifType : uint8
 	PinDefaultValueChanged, // A pin's default value has changed (Subject == URigVMPin)
 	PinDirectionChanged, // A pin's direction has changed (Subject == URigVMPin)
 	PinTypeChanged, // A pin's data type has changed (Subject == URigVMPin)
+	PinIndexChanged, // A pin's index has changed (Subject == URigVMPin)
 	LinkAdded, // A link has been added (Subject == URigVMLink)
 	LinkRemoved, // A link has been removed (Subject == URigVMLink)
 	CommentTextChanged, // A comment node's text has changed (Subject == URigVMCommentNode)
