Control Rig: Use pin defaults when promoting to variables

diff --git a/Engine/Plugins/Experimental/ControlRig/Source/ControlRigDeveloper/Private/ControlRigBlueprint.cpp b/Engine/Plugins/Experimental/ControlRig/Source/ControlRigDeveloper/Private/ControlRigBlueprint.cpp
--- a/Engine/Plugins/Experimental/ControlRig/Source/ControlRigDeveloper/Private/ControlRigBlueprint.cpp
+++ b/Engine/Plugins/Experimental/ControlRig/Source/ControlRigDeveloper/Private/ControlRigBlueprint.cpp
@@ -617,11 +617,11 @@ URigVMController* UControlRigBlueprint::GetOrCreateController(URigVMGraph* InGra
 	// c) create external variable (mapped to the passed in tfunction)
 	// the last one is defined within the blueprint since the controller
 	// doesn't own the variables and can't create one itself.
-	Controller->SetupDefaultUnitNodeDelegates(TDelegate<FName(FRigVMExternalVariable)>::CreateLambda(
-		[WeakThis](FRigVMExternalVariable InVariableToCreate) -> FName {
+	Controller->SetupDefaultUnitNodeDelegates(TDelegate<FName(FRigVMExternalVariable, FString)>::CreateLambda(
+		[WeakThis](FRigVMExternalVariable InVariableToCreate, FString InDefaultValue) -> FName {
 			if (WeakThis.IsValid())
 			{
-				return WeakThis->AddCRMemberVariableFromExternal(InVariableToCreate);
+				return WeakThis->AddCRMemberVariableFromExternal(InVariableToCreate, InDefaultValue);
 			}
 			return NAME_None;
 		}
@@ -880,7 +880,7 @@ UControlRigHierarchyModifier* UControlRigBlueprint::GetHierarchyModifier()
 
 #if WITH_EDITOR
 
-FName UControlRigBlueprint::AddMemberVariable(const FName& InName, const FString& InCPPType, bool bIsPublic, bool bIsReadOnly)
+FName UControlRigBlueprint::AddMemberVariable(const FName& InName, const FString& InCPPType, bool bIsPublic, bool bIsReadOnly, FString InDefaultValue)
 {
 	FRigVMExternalVariable Variable;
 	Variable.Name = InName;
@@ -938,7 +938,7 @@ FName UControlRigBlueprint::AddMemberVariable(const FName& InName, const FString
 		Variable.Size = Enum->GetResourceSizeBytes(EResourceSizeMode::EstimatedTotal);
 	}
 
-	FName Result = AddCRMemberVariableFromExternal(Variable);
+	FName Result = AddCRMemberVariableFromExternal(Variable, InDefaultValue);
 	if (!Result.IsNone())
 	{
 		FBPCompileRequest Request(this, EBlueprintCompileOptions::None, nullptr);
@@ -1844,7 +1844,7 @@ void UControlRigBlueprint::CreateMemberVariablesOnLoad()
 				}
 
 				FName VarName = FindCRMemberVariableUniqueName(NameValidator, Description.Name.ToString());
-				int32 VariableIndex = AddCRMemberVariable(this, VarName, PinType, false, false);
+				int32 VariableIndex = AddCRMemberVariable(this, VarName, PinType, false, false, FString());
 				if (VariableIndex != INDEX_NONE)
 				{
 					AddedMemberVariableMap.Add(Description.Name, VariableIndex);
@@ -1875,7 +1875,7 @@ void UControlRigBlueprint::CreateMemberVariablesOnLoad()
 				}
 
 				FName VarName = FindCRMemberVariableUniqueName(NameValidator, Description.Name.ToString());
-				int32 VariableIndex = AddCRMemberVariable(this, VarName, PinType, true, !Description.bIsInput);
+				int32 VariableIndex = AddCRMemberVariable(this, VarName, PinType, true, !Description.bIsInput, FString());
 				if (VariableIndex != INDEX_NONE)
 				{
 					AddedMemberVariableMap.Add(Description.Name, VariableIndex);
@@ -1919,7 +1919,7 @@ FName UControlRigBlueprint::FindCRMemberVariableUniqueName(TSharedPtr<FKismetNam
 	return *KismetName;
 }
 
-int32 UControlRigBlueprint::AddCRMemberVariable(UControlRigBlueprint* InBlueprint, const FName& InVarName, FEdGraphPinType InVarType, bool bIsPublic, bool bIsReadOnly)
+int32 UControlRigBlueprint::AddCRMemberVariable(UControlRigBlueprint* InBlueprint, const FName& InVarName, FEdGraphPinType InVarType, bool bIsPublic, bool bIsReadOnly, FString InDefaultValue)
 {
 	FBPVariableDescription NewVar;
 
@@ -1952,10 +1952,12 @@ int32 UControlRigBlueprint::AddCRMemberVariable(UControlRigBlueprint* InBlueprin
 	// Text variables, etc. should default to multiline
 	NewVar.SetMetaData(TEXT("MultiLine"), TEXT("true"));
 
+	NewVar.DefaultValue = InDefaultValue;
+
 	return InBlueprint->NewVariables.Add(NewVar);
 }
 
-FName UControlRigBlueprint::AddCRMemberVariableFromExternal(FRigVMExternalVariable InVariableToCreate)
+FName UControlRigBlueprint::AddCRMemberVariableFromExternal(FRigVMExternalVariable InVariableToCreate, FString InDefaultValue)
 {
 	FEdGraphPinType PinType = UControlRig::GetPinTypeFromExternalVariable(InVariableToCreate);
 	if (!PinType.PinCategory.IsValid())
@@ -1967,7 +1969,7 @@ FName UControlRigBlueprint::AddCRMemberVariableFromExternal(FRigVMExternalVariab
 
 	TSharedPtr<FKismetNameValidator> NameValidator = MakeShareable(new FKismetNameValidator(this, NAME_None, nullptr));
 	FName VarName = FindCRMemberVariableUniqueName(NameValidator, InVariableToCreate.Name.ToString());
-	int32 VariableIndex = AddCRMemberVariable(this, VarName, PinType, InVariableToCreate.bIsPublic, InVariableToCreate.bIsReadOnly);
+	int32 VariableIndex = AddCRMemberVariable(this, VarName, PinType, InVariableToCreate.bIsPublic, InVariableToCreate.bIsReadOnly, InDefaultValue);
 	if (VariableIndex != INDEX_NONE)
 	{
 		FBlueprintEditorUtils::MarkBlueprintAsStructurallyModified(this);
diff --git a/Engine/Plugins/Experimental/ControlRig/Source/ControlRigDeveloper/Public/ControlRigBlueprint.h b/Engine/Plugins/Experimental/ControlRig/Source/ControlRigDeveloper/Public/ControlRigBlueprint.h
--- a/Engine/Plugins/Experimental/ControlRig/Source/ControlRigDeveloper/Public/ControlRigBlueprint.h
+++ b/Engine/Plugins/Experimental/ControlRig/Source/ControlRigDeveloper/Public/ControlRigBlueprint.h
@@ -191,7 +191,7 @@ public:
 
 #if WITH_EDITOR
 	UFUNCTION(BlueprintCallable, Category = "Variables")
-	FName AddMemberVariable(const FName& InName, const FString& InCPPType, bool bIsPublic = false, bool bIsReadOnly = false);
+	FName AddMemberVariable(const FName& InName, const FString& InCPPType, bool bIsPublic = false, bool bIsReadOnly = false, FString InDefaultValue = TEXT(""));
 #endif
 
 #if WITH_EDITORONLY_DATA
@@ -308,8 +308,8 @@ private:
 	void CreateMemberVariablesOnLoad();
 #if WITH_EDITOR
 	static FName FindCRMemberVariableUniqueName(TSharedPtr<FKismetNameValidator> InNameValidator, const FString& InBaseName);
-	static int32 AddCRMemberVariable(UControlRigBlueprint* InBlueprint, const FName& InVarName, FEdGraphPinType InVarType, bool bIsPublic, bool bIsReadOnly);
-	FName AddCRMemberVariableFromExternal(FRigVMExternalVariable InVariableToCreate);
+	static int32 AddCRMemberVariable(UControlRigBlueprint* InBlueprint, const FName& InVarName, FEdGraphPinType InVarType, bool bIsPublic, bool bIsReadOnly, FString InDefaultValue);
+	FName AddCRMemberVariableFromExternal(FRigVMExternalVariable InVariableToCreate, FString InDefaultValue = FString());
 #endif
 	void PatchVariableNodesOnLoad();
 
diff --git a/Engine/Source/Developer/RigVMDeveloper/Private/RigVMModel/RigVMController.cpp b/Engine/Source/Developer/RigVMDeveloper/Private/RigVMModel/RigVMController.cpp
--- a/Engine/Source/Developer/RigVMDeveloper/Private/RigVMModel/RigVMController.cpp
+++ b/Engine/Source/Developer/RigVMDeveloper/Private/RigVMModel/RigVMController.cpp
@@ -4935,7 +4935,7 @@ bool URigVMController::PromotePinToVariable(URigVMPin* InPin, bool bCreateVariab
 		}
 
 		VariableForPin = InPin->ToExternalVariable();
-		FName VariableName = UnitNodeCreatedContext.GetCreateExternalVariableDelegate().Execute(VariableForPin);
+		FName VariableName = UnitNodeCreatedContext.GetCreateExternalVariableDelegate().Execute(VariableForPin, InPin->GetDefaultValue());
 		if (VariableName.IsNone())
 		{
 			return false;
@@ -7593,7 +7593,7 @@ void URigVMController::RepopulatePinsOnNode(URigVMNode* InNode, bool bFollowCore
 
 #endif
 
-void URigVMController::SetupDefaultUnitNodeDelegates(TDelegate<FName(FRigVMExternalVariable)> InCreateExternalVariableDelegate)
+void URigVMController::SetupDefaultUnitNodeDelegates(TDelegate<FName(FRigVMExternalVariable, FString)> InCreateExternalVariableDelegate)
 {
 	TWeakObjectPtr<URigVMController> WeakThis(this);
 
diff --git a/Engine/Source/Developer/RigVMDeveloper/Public/RigVMModel/RigVMController.h b/Engine/Source/Developer/RigVMDeveloper/Public/RigVMModel/RigVMController.h
--- a/Engine/Source/Developer/RigVMDeveloper/Public/RigVMModel/RigVMController.h
+++ b/Engine/Source/Developer/RigVMDeveloper/Public/RigVMModel/RigVMController.h
@@ -543,7 +543,7 @@ public:
 
 	// Wires the unit node delegates to the default controller delegates.
 	// this is used only within the Control Rig Editor currently.
-	void SetupDefaultUnitNodeDelegates(TDelegate<FName(FRigVMExternalVariable)> InCreateExternalVariableDelegate);
+	void SetupDefaultUnitNodeDelegates(TDelegate<FName(FRigVMExternalVariable, FString)> InCreateExternalVariableDelegate);
 	void ResetUnitNodeDelegates();
 
 	// A flag that can be used to turn off pin default value validation if necessary
diff --git a/Engine/Source/Runtime/RigVM/Private/RigVMCore/RigVMStruct.cpp b/Engine/Source/Runtime/RigVM/Private/RigVMCore/RigVMStruct.cpp
--- a/Engine/Source/Runtime/RigVM/Private/RigVMCore/RigVMStruct.cpp
+++ b/Engine/Source/Runtime/RigVM/Private/RigVMCore/RigVMStruct.cpp
@@ -25,11 +25,11 @@ TArray<FRigVMExternalVariable> FRigVMUnitNodeCreatedContext::GetExternalVariable
 	return ExternalVariables;
 }
 
-FName FRigVMUnitNodeCreatedContext::AddExternalVariable(const FRigVMExternalVariable& InVariableToCreate)
+FName FRigVMUnitNodeCreatedContext::AddExternalVariable(const FRigVMExternalVariable& InVariableToCreate, FString InDefaultValue)
 {
 	if (CreateExternalVariableDelegate.IsBound())
 	{
-		return CreateExternalVariableDelegate.Execute(InVariableToCreate);
+		return CreateExternalVariableDelegate.Execute(InVariableToCreate, InDefaultValue);
 	}
 	return NAME_None;
 }
diff --git a/Engine/Source/Runtime/RigVM/Public/RigVMCore/RigVMStruct.h b/Engine/Source/Runtime/RigVM/Public/RigVMCore/RigVMStruct.h
--- a/Engine/Source/Runtime/RigVM/Public/RigVMCore/RigVMStruct.h
+++ b/Engine/Source/Runtime/RigVM/Public/RigVMCore/RigVMStruct.h
@@ -10,7 +10,7 @@
 
 // delegates used for variable introspection / creation
 DECLARE_DELEGATE_RetVal(TArray<FRigVMExternalVariable>, FRigVMGetExternalVariablesDelegate)
-DECLARE_DELEGATE_RetVal_OneParam(FName, FRigVMCreateExternalVariableDelegate, FRigVMExternalVariable)
+DECLARE_DELEGATE_RetVal_TwoParams(FName, FRigVMCreateExternalVariableDelegate, FRigVMExternalVariable, FString)
 DECLARE_DELEGATE_RetVal_TwoParams(bool, FRigVMBindPinToExternalVariableDelegate, FString, FString)
 
 /** Context as of why the node was created */
@@ -63,7 +63,7 @@ public:
 	TArray<FRigVMExternalVariable> GetExternalVariables() const;
 
 	/** Creates a new variable within the host of this VM */
-	FName AddExternalVariable(const FRigVMExternalVariable& InVariableToCreate);
+	FName AddExternalVariable(const FRigVMExternalVariable& InVariableToCreate, FString InDefaultValue = FString());
 
 	/** Binds a pin to an external variable on the created node */
 	bool BindPinToExternalVariable(FString InPinPath, FString InVariablePath);
