Control Rig: Support scale in basic Ik

diff --git a/Engine/Plugins/Experimental/ControlRig/Source/ControlRig/Private/Units/Highlevel/Hierarchy/RigUnit_TwoBoneIKSimple.cpp b/Engine/Plugins/Experimental/ControlRig/Source/ControlRig/Private/Units/Highlevel/Hierarchy/RigUnit_TwoBoneIKSimple.cpp
--- a/Engine/Plugins/Experimental/ControlRig/Source/ControlRig/Private/Units/Highlevel/Hierarchy/RigUnit_TwoBoneIKSimple.cpp
+++ b/Engine/Plugins/Experimental/ControlRig/Source/ControlRig/Private/Units/Highlevel/Hierarchy/RigUnit_TwoBoneIKSimple.cpp
@@ -91,15 +91,27 @@ FRigUnit_TwoBoneIKSimplePerItem_Execute()
 
 	if (LengthA < SMALL_NUMBER)
 	{
-		FVector Diff = Hierarchy->GetInitialGlobalTransform(CachedItemAIndex).GetLocation() - Hierarchy->GetInitialGlobalTransform(CachedItemBIndex).GetLocation();
-		Diff = Diff * TransformA.GetScale3D();
+		FTransform InitialTransformA = Hierarchy->GetInitialGlobalTransform(CachedItemAIndex);
+		FVector Scale = FVector::OneVector;
+		if (InitialTransformA.GetScale3D().SizeSquared() > SMALL_NUMBER)
+		{
+			Scale = TransformA.GetScale3D() / InitialTransformA.GetScale3D();
+		}
+		FVector Diff = InitialTransformA.GetLocation() - Hierarchy->GetInitialGlobalTransform(CachedItemBIndex).GetLocation();
+		Diff = Diff * Scale;
 		LengthA = Diff.Size();
 	}
 
 	if (LengthB < SMALL_NUMBER && CachedEffectorItemIndex != INDEX_NONE)
 	{
-		FVector Diff = Hierarchy->GetInitialGlobalTransform(CachedItemBIndex).GetLocation() - Hierarchy->GetInitialGlobalTransform(CachedEffectorItemIndex).GetLocation();
-		Diff = Diff * TransformB.GetScale3D();
+		FTransform InitialTransformB = Hierarchy->GetInitialGlobalTransform(CachedItemBIndex);
+		FVector Scale = FVector::OneVector;
+		if (InitialTransformB.GetScale3D().SizeSquared() > SMALL_NUMBER)
+		{
+			Scale = TransformB.GetScale3D() / InitialTransformB.GetScale3D();
+		}
+		FVector Diff = InitialTransformB.GetLocation() - Hierarchy->GetInitialGlobalTransform(CachedEffectorItemIndex).GetLocation();
+		Diff = Diff * Scale;
 		LengthB = Diff.Size();
 	}
 
