Control Rig: Enable animation scrub widget for editor viewport

diff --git a/Engine/Plugins/Experimental/ControlRig/Source/ControlRigEditor/Private/Editor/ControlRigEditorMode.cpp b/Engine/Plugins/Experimental/ControlRig/Source/ControlRigEditor/Private/Editor/ControlRigEditorMode.cpp
--- a/Engine/Plugins/Experimental/ControlRig/Source/ControlRigEditor/Private/Editor/ControlRigEditorMode.cpp
+++ b/Engine/Plugins/Experimental/ControlRig/Source/ControlRigEditor/Private/Editor/ControlRigEditorMode.cpp
@@ -31,7 +31,7 @@ FControlRigEditorMode::FControlRigEditorMode(const TSharedRef<FControlRigEditor>
  	ViewportArgs.BlueprintEditor = InControlRigEditor;
  	ViewportArgs.bShowStats = false;
 	ViewportArgs.bShowPlaySpeedMenu = false;
-	ViewportArgs.bShowTimeline = false;
+	ViewportArgs.bShowTimeline = true;
 	ViewportArgs.bShowTurnTable = false;
 	ViewportArgs.bAlwaysShowTransformToolbar = true;
 	ViewportArgs.OnViewportCreated = FOnViewportCreated::CreateSP(InControlRigEditor, &FControlRigEditor::HandleViewportCreated);
