CIS FIX, disable static analysis here

diff --git a/Engine/Plugins/Experimental/ControlRig/Source/ControlRig/Private/Tools/ControlRigPose.cpp b/Engine/Plugins/Experimental/ControlRig/Source/ControlRig/Private/Tools/ControlRigPose.cpp
--- a/Engine/Plugins/Experimental/ControlRig/Source/ControlRig/Private/Tools/ControlRigPose.cpp
+++ b/Engine/Plugins/Experimental/ControlRig/Source/ControlRig/Private/Tools/ControlRigPose.cpp
@@ -91,7 +91,7 @@ void FControlRigControlPose::PastePoseInternal(UControlRig* ControlRig, bool bDo
 				{
 					if (bDoMirror == false)
 					{
-						if (bDoLocal)
+						if (bDoLocal) // -V547  
 						{
 							ControlRig->SetControlLocalTransform(RigControl.Name, CopyRigControl->LocalTransform, true, Context);
 						}
@@ -192,7 +192,7 @@ void FControlRigControlPose::BlendWithInitialPoses(FControlRigControlPose& Initi
 					{
 						if (bDoMirror == false)
 						{
-							if (bDoLocal == true)
+							if (bDoLocal == true)    // -V547  
 							{
 								FTransform Val = CopyRigControl->LocalTransform;
 								FTransform InitialVal = InitialFound->LocalTransform;
