Control Rig: Pose library, litte updates to get it to work with local spaces.

diff --git a/Engine/Plugins/Experimental/ControlRig/Source/ControlRig/Private/Tools/ControlRigPose.cpp b/Engine/Plugins/Experimental/ControlRig/Source/ControlRig/Private/Tools/ControlRigPose.cpp
--- a/Engine/Plugins/Experimental/ControlRig/Source/ControlRig/Private/Tools/ControlRigPose.cpp
+++ b/Engine/Plugins/Experimental/ControlRig/Source/ControlRig/Private/Tools/ControlRigPose.cpp
@@ -35,25 +35,17 @@ void FControlRigControlPose::PastePose(UControlRig* ControlRig, bool bDoKey, boo
 
 void FControlRigControlPose::SetControlMirrorTransform(bool bDoLocal, UControlRig* ControlRig, const FName& Name, bool bIsMatched, const FVector& GlobalTranslation, const FQuat& GlobalRotation, const FVector& LocalTranslation, const FQuat& LocalRotation, bool bNotify, const  FRigControlModifiedContext& Context)
 {
-	if (bDoLocal)
+	if (bIsMatched)
 	{
 		FTransform NewLocalTransform(LocalRotation, LocalTranslation);
 		ControlRig->SetControlLocalTransform(Name, NewLocalTransform, bNotify, Context);
+
 	}
 	else
 	{
-		if (bIsMatched)
-		{
-			FTransform NewLocalTransform(LocalRotation, LocalTranslation);
-			ControlRig->SetControlLocalTransform(Name, NewLocalTransform, bNotify, Context);
-
-		}
-		else
-		{
-			FTransform NewGlobalTransform(GlobalRotation, GlobalTranslation);
-			ControlRig->SetControlGlobalTransform(Name, NewGlobalTransform, bNotify, Context);
-		}
-	}
+		FTransform NewGlobalTransform(GlobalRotation, GlobalTranslation);
+		ControlRig->SetControlGlobalTransform(Name, NewGlobalTransform, bNotify, Context);
+	}	
 }
 
 void FControlRigControlPose::PastePoseInternal(UControlRig* ControlRig, bool bDoKey, bool bDoMirror, const TArray<FRigControlCopy>& ControlsToPaste)
diff --git a/Engine/Plugins/Experimental/ControlRig/Source/ControlRig/Private/Tools/ControlRigPoseMirrorTable.cpp b/Engine/Plugins/Experimental/ControlRig/Source/ControlRig/Private/Tools/ControlRigPoseMirrorTable.cpp
--- a/Engine/Plugins/Experimental/ControlRig/Source/ControlRig/Private/Tools/ControlRigPoseMirrorTable.cpp
+++ b/Engine/Plugins/Experimental/ControlRig/Source/ControlRig/Private/Tools/ControlRigPoseMirrorTable.cpp
@@ -85,7 +85,7 @@ void FControlRigPoseMirrorTable::GetMirrorTransform(const FRigControlCopy& Contr
 	OutGlobalRotation = GlobalTransform.GetRotation();
 	if (Settings)
 	{
-		if (bIsMatched  && (OutGlobalRotation.IsIdentity() == false && OutLocalRotation.IsIdentity() == false))
+		if (!bIsMatched  && (OutGlobalRotation.IsIdentity() == false || OutLocalRotation.IsIdentity() == false))
 		{
 			FRigMirrorSettings MirrorSettings;
 			MirrorSettings.MirrorAxis = Settings->MirrorAxis;
