Control Rig: Base a new argument on the details of the last one

diff --git a/Engine/Plugins/Experimental/ControlRig/Source/ControlRigEditor/Private/ControlRigGraphDetails.cpp b/Engine/Plugins/Experimental/ControlRig/Source/ControlRigEditor/Private/ControlRigGraphDetails.cpp
--- a/Engine/Plugins/Experimental/ControlRig/Source/ControlRigEditor/Private/ControlRigGraphDetails.cpp
+++ b/Engine/Plugins/Experimental/ControlRig/Source/ControlRigEditor/Private/ControlRigGraphDetails.cpp
@@ -656,7 +656,24 @@ FReply FControlRigGraphDetails::OnAddNewInputClicked()
 			FString CPPType = TEXT("bool");
 			FName CPPTypeObjectPath = NAME_None;
 			FString DefaultValue = TEXT("False");
-			// todo: base decisions on types on last argument
+
+			if (URigVMLibraryNode* LibraryNode = Cast<URigVMLibraryNode>(Model->GetOuter()))
+			{
+				if (LibraryNode->GetPins().Num() > 0)
+				{
+					URigVMPin* LastPin = LibraryNode->GetPins().Last();
+					if (!LastPin->IsExecuteContext())
+					{
+						ArgumentName = LastPin->GetFName();
+						CPPType = LastPin->GetCPPType();
+						if (LastPin->GetCPPTypeObject())
+						{
+							CPPTypeObjectPath = *LastPin->GetCPPTypeObject()->GetPathName();
+						}
+						DefaultValue = LastPin->GetDefaultValue();
+					}
+				}
+			}
 
 			Controller->AddExposedPin(ArgumentName, ERigVMPinDirection::Input, CPPType, CPPTypeObjectPath, DefaultValue, true);
 		}
