Control Rig: Added two new SphereTrace nodes that support tracing by trace channels and by object types, deprecated the old SphereTrace Node

diff --git a/Engine/Plugins/Experimental/ControlRig/Source/ControlRig/Private/Units/Collision/RigUnit_WorldCollision.cpp b/Engine/Plugins/Experimental/ControlRig/Source/ControlRig/Private/Units/Collision/RigUnit_WorldCollision.cpp
--- a/Engine/Plugins/Experimental/ControlRig/Source/ControlRig/Private/Units/Collision/RigUnit_WorldCollision.cpp
+++ b/Engine/Plugins/Experimental/ControlRig/Source/ControlRig/Private/Units/Collision/RigUnit_WorldCollision.cpp
@@ -45,3 +45,96 @@ FRigUnit_SphereTraceWorld_Execute()
 		HitNormal = Context.ToWorldSpaceTransform.InverseTransformVector(HitResult.ImpactNormal);
 	}
 }
+
+FRigUnit_SphereTraceByTraceChannel_Execute()
+{
+	DECLARE_SCOPE_HIERARCHICAL_COUNTER_RIGUNIT()
+
+	bHit = false;
+	HitLocation = FVector::ZeroVector;
+	HitNormal = FVector(0.f, 0.f, 1.f);
+
+	if (Context.World == nullptr)
+	{
+		return;
+	}
+
+	const ECollisionChannel CollisionChannel = UEngineTypes::ConvertToCollisionChannel(TraceChannel); 
+	const FCollisionShape CollisionShape = FCollisionShape::MakeSphere(Radius);
+
+	FCollisionQueryParams QueryParams;
+	QueryParams.bTraceComplex = true;
+
+	if (Context.OwningActor)
+	{
+		QueryParams.AddIgnoredActor(Context.OwningActor);
+	}
+	else if (UPrimitiveComponent* PrimitiveComponent = Cast<UPrimitiveComponent>(Context.OwningComponent))
+	{
+		QueryParams.AddIgnoredComponent(PrimitiveComponent);
+	}
+
+	FHitResult HitResult;
+	bHit = Context.World->SweepSingleByChannel(HitResult, Context.ToWorldSpace(Start), Context.ToWorldSpace(End), FQuat::Identity, CollisionChannel, CollisionShape, QueryParams);
+	
+	if (bHit)
+	{
+		HitLocation = Context.ToRigSpace(HitResult.ImpactPoint);
+		HitNormal = Context.ToWorldSpaceTransform.InverseTransformVector(HitResult.ImpactNormal);
+	}
+}
+
+FRigUnit_SphereTraceByObjectTypes_Execute()
+{
+	DECLARE_SCOPE_HIERARCHICAL_COUNTER_RIGUNIT()
+
+	bHit = false;
+	HitLocation = FVector::ZeroVector;
+	HitNormal = FVector(0.f, 0.f, 1.f);
+
+	if (Context.World == nullptr)
+	{
+		return;
+	}
+
+	const FCollisionShape CollisionShape = FCollisionShape::MakeSphere(Radius); 
+
+	// similar to KismetTraceUtils.h -> ConfigureCollisionObjectParams
+	TArray<TEnumAsByte<ECollisionChannel>> CollisionChannels;
+	CollisionChannels.AddUninitialized(ObjectTypes.Num());
+
+	for (int Index = 0; Index < ObjectTypes.Num(); Index++)
+	{ 
+		CollisionChannels[Index] = UEngineTypes::ConvertToCollisionChannel(ObjectTypes[Index]);
+	}
+
+	FCollisionObjectQueryParams ObjectParams;
+	for (const ECollisionChannel& Channel : CollisionChannels)
+	{
+		if (FCollisionObjectQueryParams::IsValidObjectQuery(Channel))
+		{
+			ObjectParams.AddObjectTypesToQuery(Channel);
+		}
+	}
+
+	FCollisionQueryParams QueryParams;
+	QueryParams.bTraceComplex = true;
+
+	if (Context.OwningActor)
+	{
+		QueryParams.AddIgnoredActor(Context.OwningActor);
+	}
+	else if (UPrimitiveComponent* PrimitiveComponent = Cast<UPrimitiveComponent>(Context.OwningComponent))
+	{
+		QueryParams.AddIgnoredComponent(PrimitiveComponent);
+	} 
+
+	FHitResult HitResult;
+	bHit = Context.World->SweepSingleByObjectType(HitResult, Context.ToWorldSpace(Start), Context.ToWorldSpace(End), FQuat::Identity, ObjectParams, CollisionShape, QueryParams);
+
+	if (bHit)
+	{
+		HitLocation = Context.ToRigSpace(HitResult.ImpactPoint);
+		HitNormal = Context.ToWorldSpaceTransform.InverseTransformVector(HitResult.ImpactNormal);
+	}
+}
diff --git a/Engine/Plugins/Experimental/ControlRig/Source/ControlRig/Private/Units/Collision/RigUnit_WorldCollision.h b/Engine/Plugins/Experimental/ControlRig/Source/ControlRig/Private/Units/Collision/RigUnit_WorldCollision.h
--- a/Engine/Plugins/Experimental/ControlRig/Source/ControlRig/Private/Units/Collision/RigUnit_WorldCollision.h
+++ b/Engine/Plugins/Experimental/ControlRig/Source/ControlRig/Private/Units/Collision/RigUnit_WorldCollision.h
@@ -9,7 +9,7 @@
 /**
  * Sweeps a sphere against the world and return the first blocking hit using a specific channel
  */
-USTRUCT(meta=(DisplayName="Sphere Trace", Category="Collision", DocumentationPolicy = "Strict", Keywords="Sweep,Raytrace,Collision,Collide,Trace", Varying, NodeColor = "0.2 0.4 0.7"))
+USTRUCT(meta=(DisplayName="Sphere Trace", Category="Collision", DocumentationPolicy = "Strict", Keywords="Sweep,Raytrace,Collision,Collide,Trace", Varying, NodeColor = "0.2 0.4 0.7", Deprecated = "5.0"))
 struct CONTROLRIG_API FRigUnit_SphereTraceWorld : public FRigUnit
 {
 	GENERATED_BODY()
@@ -57,3 +57,109 @@ struct CONTROLRIG_API FRigUnit_SphereTraceWorld : public FRigUnit
 	UPROPERTY(meta = (Output))
 	FVector HitNormal;
 };
+
+/**
+ * Sweeps a sphere against the world and return the first blocking hit using a specific channel. Target objects can have different object types, but they need to have the same trace channel set to "block" in their collision response settings. 
+ * You can create custom trace channels in Project Setting - Collision. 
+ */
+USTRUCT(meta = (DisplayName = "Sphere Trace By Trace Channel", Category = "Collision", DocumentationPolicy = "Strict", Keywords = "Sweep,Raytrace,Collision,Collide,Trace", Varying, NodeColor = "0.2 0.4 0.7"))
+struct CONTROLRIG_API FRigUnit_SphereTraceByTraceChannel : public FRigUnit
+{
+	GENERATED_BODY()
+
+	FRigUnit_SphereTraceByTraceChannel()
+		: Start(EForceInit::ForceInitToZero)
+		, End(EForceInit::ForceInitToZero)
+		, TraceChannel(TraceTypeQuery1)
+		, Radius(5.f)
+		, bHit(false)
+		, HitLocation(EForceInit::ForceInitToZero)
+		, HitNormal(0.f, 0.f, 1.f)
+	{
+
+	}
+
+	RIGVM_METHOD()
+	virtual void Execute(const FRigUnitContext& Context) override;
+
+	/** Start of the trace in rig / global space */
+	UPROPERTY(meta = (Input))
+	FVector Start;
+
+	/** End of the trace in rig / global space */
+	UPROPERTY(meta = (Input))
+	FVector End;
+
+	/** The 'channel' that this trace is in, used to determine which components to hit */
+	UPROPERTY(meta = (Input))
+	TEnumAsByte<ETraceTypeQuery> TraceChannel;
+
+	/** Radius of the sphere to use for sweeping / tracing */
+	UPROPERTY(meta = (Input, UIMin = "0.0", UIMax = "100.0"))
+	float Radius;
+
+	/** Returns true if there was a hit */
+	UPROPERTY(meta = (Output))
+	bool bHit;
+
+	/** Hit location in rig / global Space */
+	UPROPERTY(meta = (Output))
+	FVector HitLocation;
+
+	/** Hit normal in rig / global Space */
+	UPROPERTY(meta = (Output))
+	FVector HitNormal;
+};
+
+/**
+ * Sweeps a sphere against the world and return the first blocking hit. The trace is filtered by object types only, the collision response settings are ignored.
+ * You can create custom object types in Project Setting - Collision
+ */
+USTRUCT(meta = (DisplayName = "Sphere Trace By Object Types", Category = "Collision", DocumentationPolicy = "Strict", Keywords = "Sweep,Raytrace,Collision,Collide,Trace", Varying, NodeColor = "0.2 0.4 0.7"))
+struct CONTROLRIG_API FRigUnit_SphereTraceByObjectTypes : public FRigUnit
+{
+	GENERATED_BODY()
+
+		FRigUnit_SphereTraceByObjectTypes()
+		: Start(EForceInit::ForceInitToZero)
+		, End(EForceInit::ForceInitToZero)
+		, ObjectTypes()
+		, Radius(5.f)
+		, bHit(false)
+		, HitLocation(EForceInit::ForceInitToZero)
+		, HitNormal(0.f, 0.f, 1.f)
+	{
+
+	}
+
+	RIGVM_METHOD()
+	virtual void Execute(const FRigUnitContext& Context) override;
+
+	/** Start of the trace in rig / global space */
+	UPROPERTY(meta = (Input))
+	FVector Start;
+
+	/** End of the trace in rig / global space */
+	UPROPERTY(meta = (Input))
+	FVector End;
+
+	/** The types of objects that this trace can hit */
+	UPROPERTY(meta = (Input))
+	TArray<TEnumAsByte<EObjectTypeQuery>> ObjectTypes;
+
+	/** Radius of the sphere to use for sweeping / tracing */
+	UPROPERTY(meta = (Input, UIMin = "0.0", UIMax = "100.0"))
+	float Radius;
+
+	/** Returns true if there was a hit */
+	UPROPERTY(meta = (Output))
+	bool bHit;
+
+	/** Hit location in rig / global Space */
+	UPROPERTY(meta = (Output))
+	FVector HitLocation;
+
+	/** Hit normal in rig / global Space */
+	UPROPERTY(meta = (Output))
+	FVector HitNormal;
+};
