Control Rig: Fix Item Replace node does not work with blank values on New

diff --git a/Engine/Plugins/Experimental/ControlRig/Source/ControlRig/Private/Units/Core/RigUnit_Name.cpp b/Engine/Plugins/Experimental/ControlRig/Source/ControlRig/Private/Units/Core/RigUnit_Name.cpp
--- a/Engine/Plugins/Experimental/ControlRig/Source/ControlRig/Private/Units/Core/RigUnit_Name.cpp
+++ b/Engine/Plugins/Experimental/ControlRig/Source/ControlRig/Private/Units/Core/RigUnit_Name.cpp
@@ -52,9 +52,14 @@ FRigUnit_NameReplace_Execute()
 
 	Result = Name;
 
-	if (!Old.IsNone() && !New.IsNone())
+	if (!Old.IsNone())
 	{
-		Result = *Name.ToString().Replace(*Old.ToString(), *New.ToString(), ESearchCase::CaseSensitive);
+		FString NewString(TEXT(""));
+		if (!New.IsNone())
+		{
+			NewString = *New.ToString();
+		}
+		Result = *Name.ToString().Replace(*Old.ToString(), *NewString, ESearchCase::CaseSensitive);
 	}
 }
 
diff --git a/Engine/Plugins/Experimental/ControlRig/Source/ControlRig/Private/Units/Execution/RigUnit_Item.cpp b/Engine/Plugins/Experimental/ControlRig/Source/ControlRig/Private/Units/Execution/RigUnit_Item.cpp
--- a/Engine/Plugins/Experimental/ControlRig/Source/ControlRig/Private/Units/Execution/RigUnit_Item.cpp
+++ b/Engine/Plugins/Experimental/ControlRig/Source/ControlRig/Private/Units/Execution/RigUnit_Item.cpp
@@ -42,14 +42,32 @@ FRigUnit_ItemReplace_Execute()
 
 IMPLEMENT_RIGUNIT_AUTOMATION_TEST(FRigUnit_ItemReplace)
 {
-	Unit.Item.Name = FName("OldItemName");
+	Unit.Item.Name = FName(TEXT("OldItemName"));
 	Unit.Item.Type = ERigElementType::Bone;
 
-	Unit.Old = FName("Old");
-	Unit.New = FName("New");
+	Unit.Old = FName(TEXT("Old"));
+	Unit.New = FName(TEXT("New"));
 	
 	Execute();
-	AddErrorIfFalse(Unit.Result == FRigElementKey("NewItemName", ERigElementType::Bone), TEXT("unexpected result"));
+	AddErrorIfFalse(Unit.Result == FRigElementKey(TEXT("NewItemName"), ERigElementType::Bone), TEXT("unexpected result"));
+
+	Unit.Item.Name = FName(TEXT("OldItemName"));
+	Unit.Item.Type = ERigElementType::Bone;
+
+	Unit.Old = FName(TEXT("Old"));
+	Unit.New = NAME_None;
+
+	Execute();
+	AddErrorIfFalse(Unit.Result == FRigElementKey(TEXT("ItemName"), ERigElementType::Bone), TEXT("unexpected result when New is None"));
+
+	Unit.Item.Name = FName(TEXT("OldItemName"));
+	Unit.Item.Type = ERigElementType::Bone;
+
+	Unit.Old = NAME_None;
+	Unit.New = FName(TEXT("New")); 
+
+	Execute();
+	AddErrorIfFalse(Unit.Result == FRigElementKey(TEXT("OldItemName"), ERigElementType::Bone), TEXT("unexpected result when Old is None"));
 	return true;
 }
 #endif
\ No newline at end of file
