Control Rig: Fix compliation error in RigUnit_WorldCollision.cpp for sphere trace nodes

diff --git a/Engine/Plugins/Experimental/ControlRig/Source/ControlRig/Private/Units/Collision/RigUnit_WorldCollision.cpp b/Engine/Plugins/Experimental/ControlRig/Source/ControlRig/Private/Units/Collision/RigUnit_WorldCollision.cpp
--- a/Engine/Plugins/Experimental/ControlRig/Source/ControlRig/Private/Units/Collision/RigUnit_WorldCollision.cpp
+++ b/Engine/Plugins/Experimental/ControlRig/Source/ControlRig/Private/Units/Collision/RigUnit_WorldCollision.cpp
@@ -100,7 +100,7 @@ FRigUnit_SphereTraceByObjectTypes_Execute()
 	const FCollisionShape CollisionShape = FCollisionShape::MakeSphere(Radius); 
 
 	// similar to KismetTraceUtils.h -> ConfigureCollisionObjectParams
-	TArray<TEnumAsByte<ECollisionChannel>> CollisionChannels;
+	TArray<ECollisionChannel> CollisionChannels;
 	CollisionChannels.AddUninitialized(ObjectTypes.Num());
 
 	for (int Index = 0; Index < ObjectTypes.Num(); Index++)
