Control Rig: Move instruction index lookup from a per-node field to a 1-n map. 1 instruction can represent 1 node, but 1 node can represent multiple instructions.

This fixes the labels on the execution stack for nested nodes.
Also fixes the cross selection issues between the stack and the graph views.
Finally also implements the "focus on selection" feature, which will now open the right tab and focus on the selected nodes.

diff --git a/Engine/Plugins/Experimental/ControlRig/Source/ControlRigDeveloper/Private/ControlRigBlueprint.cpp b/Engine/Plugins/Experimental/ControlRig/Source/ControlRigDeveloper/Private/ControlRigBlueprint.cpp
--- a/Engine/Plugins/Experimental/ControlRig/Source/ControlRigDeveloper/Private/ControlRigBlueprint.cpp
+++ b/Engine/Plugins/Experimental/ControlRig/Source/ControlRigDeveloper/Private/ControlRigBlueprint.cpp
@@ -609,6 +609,27 @@ URigVMController* UControlRigBlueprint::GetOrCreateController(URigVMGraph* InGra
 
 	});
 
+
+	// this delegate is used by the controller to retrieve the current bytecode of the VM
+	Controller->GetCurrentByteCodeDelegate.BindLambda([WeakThis]() -> const FRigVMByteCode* {
+
+		if (WeakThis.IsValid())
+		{
+			if (UControlRigBlueprintGeneratedClass* RigClass = WeakThis->GetControlRigBlueprintGeneratedClass())
+			{
+				if (UControlRig* CDO = Cast<UControlRig>(RigClass->GetDefaultObject(false)))
+				{
+					if (CDO->VM)
+					{
+						return &CDO->VM->GetByteCode();
+					}
+				}
+			}
+		}
+		return nullptr;
+
+	});
+
 #if WITH_EDITOR
 
 	// this sets up three delegates:
diff --git a/Engine/Plugins/Experimental/ControlRig/Source/ControlRigDeveloper/Private/Graph/ControlRigGraph.cpp b/Engine/Plugins/Experimental/ControlRig/Source/ControlRigDeveloper/Private/Graph/ControlRigGraph.cpp
--- a/Engine/Plugins/Experimental/ControlRig/Source/ControlRigDeveloper/Private/Graph/ControlRigGraph.cpp
+++ b/Engine/Plugins/Experimental/ControlRig/Source/ControlRigDeveloper/Private/Graph/ControlRigGraph.cpp
@@ -10,6 +10,10 @@
 #include "Units/RigUnit.h"
 #include "EdGraphNode_Comment.h"
 #include "ControlRig/Private/Units/Execution/RigUnit_BeginExecution.h"
+#include "RigVMModel/Nodes/RigVMLibraryNode.h"
+#include "RigVMModel/Nodes/RigVMFunctionEntryNode.h"
+#include "RigVMModel/Nodes/RigVMFunctionReturnNode.h"
+#include "RigVMModel/Nodes/RigVMRerouteNode.h"
 
 #if WITH_EDITOR
 #include "Kismet2/BlueprintEditorUtils.h"
@@ -35,6 +39,8 @@ void UControlRigGraph::Initialize(UControlRigBlueprint* InBlueprint)
 
 	InBlueprint->OnModified().RemoveAll(this);
 	InBlueprint->OnModified().AddUObject(this, &UControlRigGraph::HandleModifiedEvent);
+	InBlueprint->OnVMCompiled().RemoveAll(this);
+	InBlueprint->OnVMCompiled().AddUObject(this, &UControlRigGraph::HandleVMCompiledEvent);
 }
 
 const UControlRigGraphSchema* UControlRigGraph::GetControlRigGraphSchema()
@@ -621,6 +627,74 @@ void UControlRigGraph::HandleModifiedEvent(ERigVMGraphNotifType InNotifType, URi
 	}
 }
 
+int32 UControlRigGraph::GetInstructionIndex(UControlRigGraphNode* InNode)
+{
+	if (const int32* FoundIndex = CachedInstructionIndices.Find(InNode))
+	{
+		return *FoundIndex;
+	}
+
+	struct Local
+	{
+		static int32 GetInstructionIndex(URigVMNode* InModelNode, const FRigVMByteCode* InByteCode)
+		{
+			if (InModelNode == nullptr)
+			{
+				return INDEX_NONE;
+			}
+
+			int32 InstructionIndex = InByteCode->GetFirstInstructionIndexForSubject(InModelNode);
+			if (InstructionIndex != INDEX_NONE)
+			{
+				return InstructionIndex;
+			}
+
+			TArray<URigVMNode*> NodesToCheck;
+			if (URigVMLibraryNode* LibraryNode = Cast<URigVMLibraryNode>(InModelNode))
+			{
+				NodesToCheck.Append(LibraryNode->GetContainedNodes());
+			}
+
+			if (InModelNode->IsA<URigVMFunctionReturnNode>() ||
+				InModelNode->IsA<URigVMRerouteNode>())
+			{
+				NodesToCheck.Append(InModelNode->GetLinkedSourceNodes());
+			}
+
+			if (InModelNode->IsA<URigVMFunctionEntryNode>() ||
+				InModelNode->IsA<URigVMRerouteNode>())
+			{
+				NodesToCheck.Append(InModelNode->GetLinkedTargetNodes());
+			}
+
+			int32 MinimumInstructionIndex = INDEX_NONE;
+			for (URigVMNode* NodeToCheck : NodesToCheck)
+			{
+				int32 ContainedInstructionIndex = GetInstructionIndex(NodeToCheck, InByteCode);
+				if (ContainedInstructionIndex != INDEX_NONE)
+				{
+					if (ContainedInstructionIndex < MinimumInstructionIndex || MinimumInstructionIndex == INDEX_NONE)
+					{
+						MinimumInstructionIndex = ContainedInstructionIndex;
+					}
+				}
+			}
+			return MinimumInstructionIndex;
+		}
+	};
+
+	if (const FRigVMByteCode* ByteCode = GetController()->GetCurrentByteCode())
+	{
+		int32 InstructionIndex = Local::GetInstructionIndex(InNode->GetModelNode(), ByteCode);
+		if (InstructionIndex != INDEX_NONE)
+		{
+			CachedInstructionIndices.Add(InNode, InstructionIndex);
+		}
+	}
+
+	return INDEX_NONE;
+}
+
 UEdGraphNode* UControlRigGraph::FindNodeForModelNodeName(const FName& InModelNodeName)
 {
 	DECLARE_SCOPE_HIERARCHICAL_COUNTER_FUNC()
@@ -683,6 +757,11 @@ URigVMController* UControlRigGraph::GetTemplateController()
 	return TemplateController;
 }
 
+void UControlRigGraph::HandleVMCompiledEvent(UBlueprint* InBlueprint, URigVM* InVM)
+{
+	CachedInstructionIndices.Reset();
+}
+
 #endif
 
 #undef LOCTEXT_NAMESPACE
diff --git a/Engine/Plugins/Experimental/ControlRig/Source/ControlRigDeveloper/Private/Graph/ControlRigGraphNode.cpp b/Engine/Plugins/Experimental/ControlRig/Source/ControlRigDeveloper/Private/Graph/ControlRigGraphNode.cpp
--- a/Engine/Plugins/Experimental/ControlRig/Source/ControlRigDeveloper/Private/Graph/ControlRigGraphNode.cpp
+++ b/Engine/Plugins/Experimental/ControlRig/Source/ControlRigDeveloper/Private/Graph/ControlRigGraphNode.cpp
@@ -282,6 +282,15 @@ void UControlRigGraphNode::HandleInsertArrayElement(FString InPinPath)
 	}
 }
 
+int32 UControlRigGraphNode::GetInstructionIndex()
+{
+	if (UControlRigGraph* RigGraph = Cast<UControlRigGraph>(GetGraph()))
+	{
+		return RigGraph->GetInstructionIndex(this);
+	}
+	return INDEX_NONE;
+}
+
 void UControlRigGraphNode::AllocateDefaultPins()
 {
 	ExecutePins.Reset();
@@ -508,9 +517,16 @@ FLinearColor UControlRigGraphNode::GetNodeOpacityColor() const
 		{
 			return FLinearColor::White;
 		}
-		if (ModelNode->GetInstructionIndex() == INDEX_NONE)
+
+		if (URigVMController* Controller = GetController())
 		{
-			return FLinearColor(0.35f, 0.35f, 0.35f, 0.35f);
+			if (const FRigVMByteCode* ByteCode = Controller->GetCurrentByteCode())
+			{
+				if (ByteCode->GetFirstInstructionIndexForSubject(ModelNode) == INDEX_NONE)
+				{
+					return FLinearColor(0.35f, 0.35f, 0.35f, 0.35f);
+				}
+			}
 		}
 	}
 	return FLinearColor::White;
diff --git a/Engine/Plugins/Experimental/ControlRig/Source/ControlRigDeveloper/Private/Graph/ControlRigGraphSchema.cpp b/Engine/Plugins/Experimental/ControlRig/Source/ControlRigDeveloper/Private/Graph/ControlRigGraphSchema.cpp
--- a/Engine/Plugins/Experimental/ControlRig/Source/ControlRigDeveloper/Private/Graph/ControlRigGraphSchema.cpp
+++ b/Engine/Plugins/Experimental/ControlRig/Source/ControlRigDeveloper/Private/Graph/ControlRigGraphSchema.cpp
@@ -215,8 +215,10 @@ const FPinConnectionResponse UControlRigGraphSchema::CanCreateConnection(const U
 				PinB = Temp;
 			}
 
+			const FRigVMByteCode* ByteCode = RigNodeA->GetController()->GetCurrentByteCode();
+
 			FString FailureReason;
-			bool bResult = RigNodeA->GetModel()->CanLink(PinA, PinB, &FailureReason);
+			bool bResult = RigNodeA->GetModel()->CanLink(PinA, PinB, &FailureReason, ByteCode);
 			if (!bResult)
 			{
 				return FPinConnectionResponse(CONNECT_RESPONSE_DISALLOW, FText::FromString(FailureReason));
diff --git a/Engine/Plugins/Experimental/ControlRig/Source/ControlRigDeveloper/Public/Graph/ControlRigGraph.h b/Engine/Plugins/Experimental/ControlRig/Source/ControlRigDeveloper/Public/Graph/ControlRigGraph.h
--- a/Engine/Plugins/Experimental/ControlRig/Source/ControlRigDeveloper/Public/Graph/ControlRigGraph.h
+++ b/Engine/Plugins/Experimental/ControlRig/Source/ControlRigDeveloper/Public/Graph/ControlRigGraph.h
@@ -6,6 +6,7 @@
 #include "Graph/ControlRigGraphNode.h"
 #include "Rigs/RigHierarchyContainer.h"
 #include "RigVMModel/RigVMGraph.h"
+#include "RigVMCore/RigVM.h"
 #include "Drawing/ControlRigDrawContainer.h"
 #include "ControlRigGraph.generated.h"
 
@@ -59,6 +60,8 @@ public:
 
 	void HandleModifiedEvent(ERigVMGraphNotifType InNotifType, URigVMGraph* InGraph, UObject* InSubject);
 
+	int32 GetInstructionIndex(UControlRigGraphNode* InNode);
+
 	UPROPERTY()
 	FString ModelNodePath;
 
@@ -98,6 +101,8 @@ private:
 
 	FControlRigGraphNodeClicked OnGraphNodeClicked;
 
+	TMap<UControlRigGraphNode*, int32> CachedInstructionIndices;
+
 #endif
 #if WITH_EDITORONLY_DATA
 
@@ -108,6 +113,7 @@ private:
 #if WITH_EDITOR
 
 	URigVMController* GetTemplateController();
+	void HandleVMCompiledEvent(UBlueprint* InBlueprint, URigVM* InVM);
 
 	friend class UControlRigUnitNodeSpawner;
 	friend class UControlRigVariableNodeSpawner;
@@ -119,9 +125,11 @@ private:
 	friend class UControlRigPrototypeNodeSpawner;
 	friend class UControlRigEnumNodeSpawner;
 	friend class UControlRigFunctionRefNodeSpawner;
+
 #endif
 	friend class UControlRigGraphNode;
 	friend class FControlRigEditor;
 	friend class SControlRigGraphNode;
+	friend class UControlRigBlueprint;
 };
 
diff --git a/Engine/Plugins/Experimental/ControlRig/Source/ControlRigDeveloper/Public/Graph/ControlRigGraphNode.h b/Engine/Plugins/Experimental/ControlRig/Source/ControlRigDeveloper/Public/Graph/ControlRigGraphNode.h
--- a/Engine/Plugins/Experimental/ControlRig/Source/ControlRigDeveloper/Public/Graph/ControlRigGraphNode.h
+++ b/Engine/Plugins/Experimental/ControlRig/Source/ControlRigDeveloper/Public/Graph/ControlRigGraphNode.h
@@ -141,6 +141,8 @@ public:
 
 	FNodeTitleDirtied& GetNodeTitleDirtied() { return NodeTitleDirtied; }
 
+	int32 GetInstructionIndex();
+
 protected:
 
 	FLinearColor GetNodeOpacityColor() const;
diff --git a/Engine/Plugins/Experimental/ControlRig/Source/ControlRigEditor/Private/Editor/ControlRigEditor.cpp b/Engine/Plugins/Experimental/ControlRig/Source/ControlRigEditor/Private/Editor/ControlRigEditor.cpp
--- a/Engine/Plugins/Experimental/ControlRig/Source/ControlRigEditor/Private/Editor/ControlRigEditor.cpp
+++ b/Engine/Plugins/Experimental/ControlRig/Source/ControlRigEditor/Private/Editor/ControlRigEditor.cpp
@@ -1832,12 +1832,12 @@ void FControlRigEditor::HandleModifiedEvent(ERigVMGraphNotifType InNotifType, UR
 		case ERigVMGraphNotifType::NodeSelected:
 		case ERigVMGraphNotifType::NodeDeselected:
 		{
-			URigVMNode* Node = Cast<URigVMNode>(InSubject);
-
-			if (FocusedGraphEdPtr.IsValid())
+			if (UControlRigGraph* RigGraph = Cast<UControlRigGraph>(ControlRigBlueprint->GetEdGraph(InGraph)))
 			{
-				TSharedPtr<SGraphEditor> FocusedGraphEd = FocusedGraphEdPtr.Pin();
-				if (UControlRigGraph* RigGraph = Cast<UControlRigGraph>(FocusedGraphEd->GetCurrentGraph()))
+				TSharedPtr<SGraphEditor> GraphEd = GetGraphEditor(RigGraph);
+				URigVMNode* Node = Cast<URigVMNode>(InSubject);
+
+				if (GraphEd.IsValid() && Node != nullptr)
 				{
 					if (InNotifType == ERigVMGraphNotifType::NodeSelected)
 					{
@@ -1863,7 +1863,7 @@ void FControlRigEditor::HandleModifiedEvent(ERigVMGraphNotifType InNotifType, UR
 						{
 							if (UEdGraphNode* EdNode = RigGraph->FindNodeForModelNodeName(ModelNode->GetFName()))
 							{
-								FocusedGraphEd->SetNodeSelection(EdNode, InNotifType == ERigVMGraphNotifType::NodeSelected);
+								GraphEd->SetNodeSelection(EdNode, InNotifType == ERigVMGraphNotifType::NodeSelected);
 							}
 						}
 					}
@@ -4195,7 +4195,7 @@ void FControlRigEditor::UpdateGraphCompilerErrors()
 	DECLARE_SCOPE_HIERARCHICAL_COUNTER_FUNC()
 
 	UControlRigBlueprint* Blueprint = Cast<UControlRigBlueprint>(GetBlueprintObj());
-	if (Blueprint)
+	if (Blueprint && ControlRig && ControlRig->VM)
 	{
 		if (Blueprint->Status == BS_Error)
 		{
@@ -4207,6 +4207,8 @@ void FControlRigEditor::UpdateGraphCompilerErrors()
 			return;
 		}
 
+		const FRigVMByteCode& ByteCode = ControlRig->VM->GetByteCode();
+
 		TArray<UEdGraph*> EdGraphs;
 		Blueprint->GetAllGraphs(EdGraphs);
 
@@ -4221,20 +4223,12 @@ void FControlRigEditor::UpdateGraphCompilerErrors()
 			// reset all nodes and store them in the map
 			bool bFoundWarning = false;
 			bool bFoundError = false;
-			TMap<int32, UEdGraphNode*> InstructionIndexToNode;
 			for (UEdGraphNode* GraphNode : Graph->Nodes)
 			{
 				if (UControlRigGraphNode* ControlRigGraphNode = Cast<UControlRigGraphNode>(GraphNode))
 				{
 					bFoundError = bFoundError || ControlRigGraphNode->ErrorType <= (int32)EMessageSeverity::Error;
 					bFoundWarning = bFoundWarning || ControlRigGraphNode->ErrorType <= (int32)EMessageSeverity::Warning;
-					if (URigVMNode* ModelNode = ControlRigGraphNode->GetModelNode())
-					{
-						if (ModelNode->GetInstructionIndex() != INDEX_NONE)
-						{
-							InstructionIndexToNode.Add(ModelNode->GetInstructionIndex(), GraphNode);
-						}
-					}
 				}
 				GraphNode->ErrorType = int32(EMessageSeverity::Info) + 1;
 			}
@@ -4243,12 +4237,17 @@ void FControlRigEditor::UpdateGraphCompilerErrors()
 			bool bFoundErrorOrWarningInLog = false;
 			for (const FControlRigLog::FLogEntry& Entry : ControlRigLog.Entries)
 			{
-				UEdGraphNode** GraphNodePtr = InstructionIndexToNode.Find(Entry.InstructionIndex);
-				if (GraphNodePtr == nullptr)
+				URigVMNode* ModelNode = Cast<URigVMNode>(ByteCode.GetSubjectForInstruction(Entry.InstructionIndex));
+				if (ModelNode == nullptr)
+				{
+					continue;
+				}
+
+				UEdGraphNode* GraphNode = RigGraph->FindNodeForModelNodeName(ModelNode->GetFName());
+				if (GraphNode == nullptr)
 				{
 					continue;
 				}
-				UEdGraphNode* GraphNode = *GraphNodePtr;
 
 				bFoundError = bFoundError || Entry.Severity <= EMessageSeverity::Error;
 				bFoundWarning = bFoundWarning || Entry.Severity <= EMessageSeverity::Warning;
@@ -4467,4 +4466,22 @@ URigVMController* FControlRigEditor::GetFocusedController() const
 	return Blueprint->GetController(GetFocusedModel());
 }
 
+TSharedPtr<SGraphEditor> FControlRigEditor::GetGraphEditor(UEdGraph* InEdGraph) const
+{
+	TArray< TSharedPtr<SDockTab> > GraphEditorTabs;
+	DocumentManager->FindAllTabsForFactory(GraphEditorTabFactoryPtr, /*out*/ GraphEditorTabs);
+
+	for (TSharedPtr<SDockTab>& GraphEditorTab : GraphEditorTabs)
+	{
+		TSharedRef<SGraphEditor> Editor = StaticCastSharedRef<SGraphEditor>((GraphEditorTab)->GetContent());
+		if (Editor->GetCurrentGraph() == InEdGraph)
+		{
+			return Editor;
+		}
+	}
+
+	return TSharedPtr<SGraphEditor>();
+}
+
+
 #undef LOCTEXT_NAMESPACE
diff --git a/Engine/Plugins/Experimental/ControlRig/Source/ControlRigEditor/Private/Editor/ControlRigEditor.h b/Engine/Plugins/Experimental/ControlRig/Source/ControlRigEditor/Private/Editor/ControlRigEditor.h
--- a/Engine/Plugins/Experimental/ControlRig/Source/ControlRigEditor/Private/Editor/ControlRigEditor.h
+++ b/Engine/Plugins/Experimental/ControlRig/Source/ControlRigEditor/Private/Editor/ControlRigEditor.h
@@ -356,6 +356,7 @@ protected:
 
 	URigVMGraph* GetFocusedModel() const;
 	URigVMController* GetFocusedController() const;
+	TSharedPtr<SGraphEditor> GetGraphEditor(UEdGraph* InEdGraph) const;
 
 	URigVMController* ActiveController;
 	bool bControlRigEditorInitialized;
diff --git a/Engine/Plugins/Experimental/ControlRig/Source/ControlRigEditor/Private/Editor/SControlRigStackView.cpp b/Engine/Plugins/Experimental/ControlRig/Source/ControlRigEditor/Private/Editor/SControlRigStackView.cpp
--- a/Engine/Plugins/Experimental/ControlRig/Source/ControlRigEditor/Private/Editor/SControlRigStackView.cpp
+++ b/Engine/Plugins/Experimental/ControlRig/Source/ControlRigEditor/Private/Editor/SControlRigStackView.cpp
@@ -258,26 +258,44 @@ void SControlRigStackView::OnSelectionChanged(TSharedPtr<FRigStackEntry> Selecti
 			return;
 		}
 
-		TMap<int32, FName> InstructionIndexToNodeName;
-		for (URigVMNode* ModelNode : ControlRigEditor.Pin()->GetFocusedModel()->GetNodes())
+		UControlRig* ControlRig = ControlRigEditor.Pin()->ControlRig;
+		if (ControlRig == nullptr || ControlRig->GetVM() == nullptr)
 		{
-			if (ModelNode->GetInstructionIndex() != INDEX_NONE)
-			{
-				InstructionIndexToNodeName.Add(ModelNode->GetInstructionIndex(), ModelNode->GetFName());
-			}
+			return;
 		}
 
-		TArray<FName> SelectedNodes;
+		const FRigVMByteCode& ByteCode = ControlRig->GetVM()->GetByteCode();
+
+		TMap<URigVMGraph*, TArray<FName>> SelectedNodesPerGraph;
 		for (TSharedPtr<FRigStackEntry>& Entry : SelectedItems)
 		{
-			const FName* NodeName = InstructionIndexToNodeName.Find(Entry->InstructionIndex);
-			if (NodeName)
+			UObject* Subject = ByteCode.GetSubjectForInstruction(Entry->InstructionIndex);
+			URigVMGraph* SubjectGraph = nullptr;
+
+			FName NodeName = NAME_None;
+			if (URigVMNode* Node = Cast<URigVMNode>(Subject))
+			{
+				NodeName = Node->GetFName();
+				SubjectGraph = Node->GetGraph();
+			}
+			else if (URigVMPin* Pin = Cast<URigVMPin>(Subject))
+			{
+				NodeName = Pin->GetNode()->GetFName();
+				SubjectGraph = Pin->GetGraph();
+			}
+
+			if (NodeName.IsNone() || SubjectGraph == nullptr)
 			{
-				SelectedNodes.Add(*NodeName);
+				continue;
 			}
+
+			SelectedNodesPerGraph.FindOrAdd(SubjectGraph).AddUnique(NodeName);
 		}
 
-		ControlRigEditor.Pin()->GetFocusedController()->SetNodeSelection(SelectedNodes);
+		for (const TPair< URigVMGraph*, TArray<FName> >& Pair : SelectedNodesPerGraph)
+		{
+			ControlRigBlueprint->GetOrCreateController(Pair.Key)->SetNodeSelection(Pair.Value);
+		}
 	}
 }
 
@@ -303,39 +321,38 @@ void SControlRigStackView::PopulateStackView(URigVM* InVM)
 	if (InVM)
 	{
 		FRigVMInstructionArray Instructions = InVM->GetInstructions();
-		const TArray<URigVMNode*>& Nodes = ControlRigBlueprint->GetModel()->GetNodes();
-		
+
+		const FRigVMByteCode& ByteCode = InVM->GetByteCode();
+
 		// 1. cache information about instructions/nodes, which will be used later 
-		TArray<int32> InstructionIndexToNodeIndex;
-		InstructionIndexToNodeIndex.Init(-1, Instructions.Num());
 		TMap<FString, FString> NodeNameToDisplayName;
-		for (int32 NodeIndex = 0; NodeIndex < Nodes.Num(); NodeIndex++)
+		for (int32 InstructionIndex = 0; InstructionIndex < ByteCode.Num(); InstructionIndex++)
 		{
-			URigVMNode* Node = Nodes[NodeIndex];
-			// only struct nodes among all nodes has StaticExecute() that generates actual instructions
-			if (URigVMUnitNode* UnitNode = Cast<URigVMUnitNode>(Node))
+			UObject* Subject = ByteCode.GetSubjectForInstruction(InstructionIndex);
+			if (URigVMNode* Node = Cast<URigVMNode>(Subject))
 			{
-				int32 InstructionIndex = UnitNode->GetInstructionIndex();
-				// only active nodes are assigned a valid instruction index, hence the check here
-				if (InstructionIndex >= 0 && InstructionIndex < InstructionIndexToNodeIndex.Num())
+				FString DisplayName = Node->GetName();
+
+				// only unit nodes among all nodes has StaticExecute() that generates actual instructions
+				if (URigVMUnitNode* UnitNode = Cast<URigVMUnitNode>(Subject))
 				{
-					InstructionIndexToNodeIndex[InstructionIndex] = NodeIndex;
-				}
-				FString DisplayName = UnitNode->GetNodeTitle();
+					DisplayName = UnitNode->GetNodeTitle();
 #if WITH_EDITOR
-				UScriptStruct* Struct = UnitNode->GetScriptStruct();
-				FString MenuDescSuffixMetadata;
-				if (Struct)
-				{
-					Struct->GetStringMetaDataHierarchical(FRigVMStruct::MenuDescSuffixMetaName, &MenuDescSuffixMetadata);
-				}
-				if (!MenuDescSuffixMetadata.IsEmpty())
-				{
-					DisplayName = FString::Printf(TEXT("%s %s"), *UnitNode->GetNodeTitle(), *MenuDescSuffixMetadata);
-				}
+					UScriptStruct* Struct = UnitNode->GetScriptStruct();
+					FString MenuDescSuffixMetadata;
+					if (Struct)
+					{
+						Struct->GetStringMetaDataHierarchical(FRigVMStruct::MenuDescSuffixMetaName, &MenuDescSuffixMetadata);
+					}
+					if (!MenuDescSuffixMetadata.IsEmpty())
+					{
+						DisplayName = FString::Printf(TEXT("%s %s"), *UnitNode->GetNodeTitle(), *MenuDescSuffixMetadata);
+					}
 #endif
+				}
+
 				// this is needed for name replacement later
-				NodeNameToDisplayName.Add(UnitNode->GetName(), DisplayName);
+				NodeNameToDisplayName.Add(Node->GetName(), DisplayName);
 			}
 		}
 
@@ -365,13 +382,13 @@ void SControlRigStackView::PopulateStackView(URigVM* InVM)
 		for (int32 InstructionIndex = 0; InstructionIndex < Labels.Num(); InstructionIndex++)
 		{
 			FString Label = Labels[InstructionIndex];
-			int32 NodeIndex = InstructionIndexToNodeIndex[InstructionIndex];
-			// note that some instructions don't map to a node, hence the check here
-			if (NodeIndex >= 0 && NodeIndex < Nodes.Num())
+
+			UObject* Subject = ByteCode.GetSubjectForInstruction(InstructionIndex);
+			if(URigVMNode* Node = Cast<URigVMNode>(Subject))
 			{
-				URigVMNode* Node = Nodes[NodeIndex];
 				Label = NodeNameToDisplayName[Node->GetName()];
 			}
+
 			// add the entry with the new label to the stack view
 			const FRigVMInstruction& Instruction = Instructions[InstructionIndex];
 			if (FilterText.IsEmpty() || Label.Contains(FilterText.ToString()))
@@ -462,7 +479,28 @@ TSharedPtr< SWidget > SControlRigStackView::CreateContextMenu()
 void SControlRigStackView::HandleFocusOnSelectedGraphNode()
 {
 	OnSelectionChanged(TSharedPtr<FRigStackEntry>(), ESelectInfo::Direct);
-	ControlRigEditor.Pin()->ZoomToSelection_Clicked();
+
+	TArray<TSharedPtr<FRigStackEntry>> SelectedItems = TreeView->GetSelectedItems();
+	if (SelectedItems.Num() > 0)
+	{
+		UControlRig* ControlRig = ControlRigEditor.Pin()->ControlRig;
+		if (ControlRig == nullptr || ControlRig->GetVM() == nullptr)
+		{
+			return;
+		}
+
+		const FRigVMByteCode& ByteCode = ControlRig->GetVM()->GetByteCode();
+		URigVMNode* SelectedNode = Cast<URigVMNode>(ByteCode.GetSubjectForInstruction(SelectedItems[0]->InstructionIndex));
+		if (SelectedNode)
+		{
+			if (UEdGraph* EdGraph = ControlRigBlueprint->GetEdGraph(SelectedNode->GetGraph()))
+			{
+				ControlRigEditor.Pin()->OpenGraphAndBringToFront(EdGraph, true);
+				ControlRigEditor.Pin()->ZoomToSelection_Clicked();
+				ControlRigEditor.Pin()->HandleModifiedEvent(ERigVMGraphNotifType::NodeSelected, SelectedNode->GetGraph(), SelectedNode);
+			}
+		}
+	}
 }
 
 void SControlRigStackView::OnVMCompiled(UBlueprint* InCompiledBlueprint, URigVM* InCompiledVM)
@@ -489,22 +527,29 @@ void SControlRigStackView::HandleModifiedEvent(ERigVMGraphNotifType InNotifType,
 		return;
 	}
 
+	UControlRig* ControlRig = ControlRigEditor.Pin()->ControlRig;
+	if (ControlRig == nullptr || ControlRig->GetVM() == nullptr)
+	{
+		return;
+	}
+
+	const FRigVMByteCode& ByteCode = ControlRig->GetVM()->GetByteCode();
+
 	switch (InNotifType)
 	{
 		case ERigVMGraphNotifType::NodeSelected:
 		case ERigVMGraphNotifType::NodeDeselected:
 		{
-			URigVMNode* Node = Cast<URigVMNode>(InSubject);
-			if (Node)
+			TArray<int32> InstructionIndices = ByteCode.GetAllInstructionIndicesForSubject(InSubject);
+			for (int32 InstructionIndex : InstructionIndices)
 			{
-				if (Node->GetInstructionIndex() != INDEX_NONE)
+				if (InstructionIndex >= Operators.Num())
 				{
-					if (Operators.Num() > Node->GetInstructionIndex())
-					{
-						TGuardValue<bool> SuspendNotifs(bSuspendModelNotifications, true);
-						TreeView->SetItemSelection(Operators[Node->GetInstructionIndex()], InNotifType == ERigVMGraphNotifType::NodeSelected, ESelectInfo::Direct);
-					}
+					break;
 				}
+
+				TGuardValue<bool> SuspendNotifs(bSuspendModelNotifications, true);
+				TreeView->SetItemSelection(Operators[InstructionIndex], InNotifType == ERigVMGraphNotifType::NodeSelected, ESelectInfo::Direct);
 			}
 			break;
 		}
diff --git a/Engine/Plugins/Experimental/ControlRig/Source/ControlRigEditor/Private/Graph/ControlRigConnectionDrawingPolicy.cpp b/Engine/Plugins/Experimental/ControlRig/Source/ControlRigEditor/Private/Graph/ControlRigConnectionDrawingPolicy.cpp
--- a/Engine/Plugins/Experimental/ControlRig/Source/ControlRigEditor/Private/Graph/ControlRigConnectionDrawingPolicy.cpp
+++ b/Engine/Plugins/Experimental/ControlRig/Source/ControlRigEditor/Private/Graph/ControlRigConnectionDrawingPolicy.cpp
@@ -284,9 +284,9 @@ void FControlRigConnectionDrawingPolicy::DetermineWiringStyle(UEdGraphPin* Outpu
 
 		if (OutputModelPin && InputModelPin)
 		{
-			int32 OutputInstructionIndex = OutputModelPin->GetNode()->GetInstructionIndex();
-			int32 InputInstructionIndex = InputModelPin->GetNode()->GetInstructionIndex();
 			bool bVisited = false;
+			int32 OutputInstructionIndex = OutputNode->GetInstructionIndex();
+			int32 InputInstructionIndex = InputNode->GetInstructionIndex();
 
 			if (OutputInstructionIndex != INDEX_NONE && InputInstructionIndex != INDEX_NONE)
 			{
diff --git a/Engine/Source/Developer/RigVMDeveloper/Private/RigVMCompiler/RigVMCompiler.cpp b/Engine/Source/Developer/RigVMDeveloper/Private/RigVMCompiler/RigVMCompiler.cpp
--- a/Engine/Source/Developer/RigVMDeveloper/Private/RigVMCompiler/RigVMCompiler.cpp
+++ b/Engine/Source/Developer/RigVMDeveloper/Private/RigVMCompiler/RigVMCompiler.cpp
@@ -117,14 +117,6 @@ bool URigVMCompiler::Compile(URigVMGraph* InGraph, URigVMController* InControlle
 
 	OutVM->Reset();
 
-	if (Settings.SetupNodeInstructionIndex)
-	{
-		for (URigVMNode* Node : InGraph->Nodes)
-		{
-			Node->InstructionIndex = INDEX_NONE;
-		}
-	}
-
 	TMap<FString, FRigVMOperand> LocalOperands;
 	if (OutOperands == nullptr)
 	{
@@ -224,90 +216,6 @@ bool URigVMCompiler::Compile(URigVMGraph* InGraph, URigVMController* InControlle
 
 	WorkData.VM->GetByteCode().AlignByteCode();
 
-	// loop over all nodes once more and setup the instruction index for reroute nodes
-	if (Settings.SetupNodeInstructionIndex)
-	{
-		struct Local
-		{
-			static int32 GetInstructionIndex(URigVMNode* InNode, TMap<URigVMNode*, int32>& InOutKnownInstructionIndex)
-			{
-				if (InNode->GetInstructionIndex() != INDEX_NONE)
-				{
-					return InNode->GetInstructionIndex();
-				}
-
-				if (int32* InstructionIndexPtr = InOutKnownInstructionIndex.Find(InNode))
-				{
-					return *InstructionIndexPtr;
-				}
-
-				InOutKnownInstructionIndex.Add(InNode, INDEX_NONE);
-
-				if (URigVMRerouteNode* RerouteNode = Cast<URigVMRerouteNode>(InNode))
-				{
-					TArray<URigVMNode*> TargetNodes = RerouteNode->GetLinkedTargetNodes();
-					if (TargetNodes.Num() > 0)
-					{
-						TArray<URigVMNode*> SourceNodes = RerouteNode->GetLinkedSourceNodes();
-						if (SourceNodes.Num() == 0)
-						{
-							for (URigVMNode* TargetNode : TargetNodes)
-							{
-								int32 InstructionIndex = GetInstructionIndex(TargetNode, InOutKnownInstructionIndex);
-								InOutKnownInstructionIndex.FindOrAdd(TargetNode) = InstructionIndex;
-
-								if (InstructionIndex != INDEX_NONE)
-								{
-									InOutKnownInstructionIndex.FindOrAdd(InNode) = InstructionIndex;
-									return InstructionIndex;
-								}
-							}
-						}
-						else
-						{
-							for (URigVMNode* SourceNode : SourceNodes)
-							{
-								int32 InstructionIndex = GetInstructionIndex(SourceNode, InOutKnownInstructionIndex);
-								InOutKnownInstructionIndex.FindOrAdd(SourceNode) = InstructionIndex;
-
-								if (InstructionIndex != INDEX_NONE)
-								{
-									InOutKnownInstructionIndex.FindOrAdd(InNode) = InstructionIndex;
-									return InstructionIndex;
-								}
-							}
-						}
-					}
-				}
-				else if (URigVMCollapseNode* CollapseNode = Cast<URigVMCollapseNode>(InNode))
-				{
-					int32 FirstInstructionIndex = INDEX_NONE;
-					for (URigVMNode* ContainedNode : CollapseNode->GetContainedNodes())
-					{
-						int32 ContainedInstructionIndex = GetInstructionIndex(ContainedNode, InOutKnownInstructionIndex);
-						if (ContainedInstructionIndex != INDEX_NONE)
-						{
-							if (FirstInstructionIndex == INDEX_NONE || ContainedInstructionIndex < FirstInstructionIndex)
-							{
-								FirstInstructionIndex = ContainedInstructionIndex;
-							}
-						}
-					}
-					return FirstInstructionIndex;
-				}
-
-				return INDEX_NONE;
-			}
-		};
-
-		TMap<URigVMNode*, int32> KnownInstructionIndex;
-		for (URigVMNode* Node : InGraph->Nodes)
-		{
-			Node->InstructionIndex = Local::GetInstructionIndex(Node, KnownInstructionIndex);
-		}
-	}
-
-
 	UE_LOG_RIGVMMEMORY(TEXT("RigVMCompiler: Finished '%s'."), *InGraph->GetPathName());
 
 	return true;
@@ -477,7 +385,7 @@ void URigVMCompiler::TraverseEntry(const FRigVMEntryExprAST* InExpr, FRigVMCompi
 
 		if (Settings.SetupNodeInstructionIndex)
 		{
-			UnitNode->InstructionIndex = EntryInstructionIndex;
+			WorkData.VM->GetByteCode().SetSubject(EntryInstructionIndex, UnitNode);
 		}
 
 		TraverseChildren(InExpr, WorkData);
@@ -560,7 +468,7 @@ int32 URigVMCompiler::TraverseCallExtern(const FRigVMCallExternExprAST* InExpr,
 		InstructionIndex = WorkData.VM->GetByteCode().GetNumInstructions() - 1;
 		if (Settings.SetupNodeInstructionIndex)
 		{
-			UnitNode->InstructionIndex = InstructionIndex;
+			WorkData.VM->GetByteCode().SetSubject(InstructionIndex, UnitNode);
 		}
 
 		ensure(InExpr->NumChildren() == UnitNode->Pins.Num());
@@ -814,7 +722,7 @@ void URigVMCompiler::TraverseAssign(const FRigVMAssignExprAST* InExpr, FRigVMCom
 				{
 					if (URigVMVariableNode* VariableNode = Cast<URigVMVariableNode>(SourcePin->GetNode()))
 					{
-						VariableNode->InstructionIndex = InstructionIndex;
+						WorkData.VM->GetByteCode().SetSubject(InstructionIndex, VariableNode);
 					}
 				}
 			}
@@ -824,7 +732,7 @@ void URigVMCompiler::TraverseAssign(const FRigVMAssignExprAST* InExpr, FRigVMCom
 				{
 					if (URigVMVariableNode* VariableNode = Cast<URigVMVariableNode>(TargetPin->GetNode()))
 					{
-						VariableNode->InstructionIndex = InstructionIndex;
+						WorkData.VM->GetByteCode().SetSubject(InstructionIndex, VariableNode);
 					}
 				}
 			}
@@ -909,7 +817,7 @@ void URigVMCompiler::TraverseBranch(const FRigVMBranchExprAST* InExpr, FRigVMCom
 	int32 JumpToFalseInstruction = WorkData.VM->GetByteCode().GetNumInstructions() - 1;
 	if (Settings.SetupNodeInstructionIndex)
 	{
-		BranchNode->InstructionIndex = JumpToFalseInstruction;
+		WorkData.VM->GetByteCode().SetSubject(JumpToFalseInstruction, BranchNode);
 	}
 
 	// traverse the true case
@@ -963,7 +871,7 @@ void URigVMCompiler::TraverseIf(const FRigVMIfExprAST* InExpr, FRigVMCompilerWor
 	int32 JumpToFalseInstruction = WorkData.VM->GetByteCode().GetNumInstructions() - 1;
 	if (Settings.SetupNodeInstructionIndex)
 	{
-		IfNode->InstructionIndex = JumpToFalseInstruction;
+		WorkData.VM->GetByteCode().SetSubject(JumpToFalseInstruction, IfNode);
 	}
 
 	// traverse the true case
@@ -1081,7 +989,7 @@ void URigVMCompiler::TraverseSelect(const FRigVMSelectExprAST* InExpr, FRigVMCom
 		WorkData.VM->GetByteCode().AddEqualsOp(IndexOperand, WorkData.IntegerLiterals.FindChecked(CaseIndex), WorkData.ComparisonOperand);
 		if (CaseIndex == 1 && Settings.SetupNodeInstructionIndex)
 		{
-			SelectNode->InstructionIndex = WorkData.VM->GetByteCode().GetNumInstructions() - 1;
+			WorkData.VM->GetByteCode().SetSubject(WorkData.VM->GetByteCode().GetNumInstructions() - 1, SelectNode);
 		}
 
 		uint64 JumpByte = WorkData.VM->GetByteCode().AddJumpIfOp(ERigVMOpCode::JumpForwardIf, 1, WorkData.ComparisonOperand, true);
diff --git a/Engine/Source/Developer/RigVMDeveloper/Private/RigVMModel/RigVMController.cpp b/Engine/Source/Developer/RigVMDeveloper/Private/RigVMModel/RigVMController.cpp
--- a/Engine/Source/Developer/RigVMDeveloper/Private/RigVMModel/RigVMController.cpp
+++ b/Engine/Source/Developer/RigVMDeveloper/Private/RigVMModel/RigVMController.cpp
@@ -5051,9 +5051,9 @@ bool URigVMController::AddLink(URigVMPin* OutputPin, URigVMPin* InputPin, bool b
 		TGuardValue<FString> InputBoundVariableGuard(InputPin->BoundVariablePath, FString());
 
 		FString FailureReason;
-		if (!Graph->CanLink(OutputPin, InputPin, &FailureReason))
+		if (!Graph->CanLink(OutputPin, InputPin, &FailureReason, GetCurrentByteCode()))
 		{
-			ReportErrorf(TEXT("Cannot link '%s' to '%s': %s."), *OutputPin->GetPinPath(), *InputPin->GetPinPath(), *FailureReason);
+			ReportErrorf(TEXT("Cannot link '%s' to '%s': %s."), *OutputPin->GetPinPath(), *InputPin->GetPinPath(), *FailureReason, GetCurrentByteCode());
 			return false;
 		}
 	}
@@ -8180,6 +8180,15 @@ TArray<FRigVMExternalVariable> URigVMController::GetExternalVariables()
 	return TArray<FRigVMExternalVariable>();
 }
 
+const FRigVMByteCode* URigVMController::GetCurrentByteCode() const
+{
+	if (GetCurrentByteCodeDelegate.IsBound())
+	{
+		return GetCurrentByteCodeDelegate.Execute();
+	}
+	return nullptr;
+}
+
 void URigVMController::RefreshFunctionReferences(URigVMLibraryNode* InFunctionDefinition, bool bSetupUndoRedo)
 {
 	check(InFunctionDefinition);
diff --git a/Engine/Source/Developer/RigVMDeveloper/Private/RigVMModel/RigVMGraph.cpp b/Engine/Source/Developer/RigVMDeveloper/Private/RigVMModel/RigVMGraph.cpp
--- a/Engine/Source/Developer/RigVMDeveloper/Private/RigVMModel/RigVMGraph.cpp
+++ b/Engine/Source/Developer/RigVMDeveloper/Private/RigVMModel/RigVMGraph.cpp
@@ -305,9 +305,9 @@ void URigVMGraph::PrepareCycleChecking(URigVMPin* InPin, bool bAsInput)
 	GetDiagnosticsAST(false, LinksToSkip)->PrepareCycleChecking(InPin);
 }
 
-bool URigVMGraph::CanLink(URigVMPin* InSourcePin, URigVMPin* InTargetPin, FString* OutFailureReason)
+bool URigVMGraph::CanLink(URigVMPin* InSourcePin, URigVMPin* InTargetPin, FString* OutFailureReason, const FRigVMByteCode* InByteCode)
 {
-	if (!URigVMPin::CanLink(InSourcePin, InTargetPin, OutFailureReason))
+	if (!URigVMPin::CanLink(InSourcePin, InTargetPin, OutFailureReason, InByteCode))
 	{
 		return false;
 	}
diff --git a/Engine/Source/Developer/RigVMDeveloper/Private/RigVMModel/RigVMNode.cpp b/Engine/Source/Developer/RigVMDeveloper/Private/RigVMModel/RigVMNode.cpp
--- a/Engine/Source/Developer/RigVMDeveloper/Private/RigVMModel/RigVMNode.cpp
+++ b/Engine/Source/Developer/RigVMDeveloper/Private/RigVMModel/RigVMNode.cpp
@@ -12,8 +12,6 @@ URigVMNode::URigVMNode()
 , Position(FVector2D::ZeroVector)
 , Size(FVector2D::ZeroVector)
 , NodeColor(FLinearColor::Black)
-, InstructionIndex(INDEX_NONE)
-, BlockIndex(INDEX_NONE)
 , GetSliceContextBracket(0)
 {
 
@@ -101,16 +99,6 @@ int32 URigVMNode::GetNodeIndex() const
 	return Index;
 }
 
-int32 URigVMNode::GetInstructionIndex() const
-{
-	return InstructionIndex;
-}
-
-int32 URigVMNode::GetBlockIndex() const
-{
-	return BlockIndex;
-}
-
 const TArray<URigVMPin*>& URigVMNode::GetPins() const
 {
 	return Pins;
diff --git a/Engine/Source/Developer/RigVMDeveloper/Private/RigVMModel/RigVMPin.cpp b/Engine/Source/Developer/RigVMDeveloper/Private/RigVMModel/RigVMPin.cpp
--- a/Engine/Source/Developer/RigVMDeveloper/Private/RigVMModel/RigVMPin.cpp
+++ b/Engine/Source/Developer/RigVMDeveloper/Private/RigVMModel/RigVMPin.cpp
@@ -1021,7 +1021,7 @@ URigVMGraph* URigVMPin::GetGraph() const
 	return nullptr;
 }
 
-bool URigVMPin::CanLink(URigVMPin* InSourcePin, URigVMPin* InTargetPin, FString* OutFailureReason)
+bool URigVMPin::CanLink(URigVMPin* InSourcePin, URigVMPin* InTargetPin, FString* OutFailureReason, const FRigVMByteCode* InByteCode)
 {
 	if (InSourcePin == nullptr || InTargetPin == nullptr)
 	{
@@ -1202,34 +1202,37 @@ bool URigVMPin::CanLink(URigVMPin* InSourcePin, URigVMPin* InTargetPin, FString*
 	TArray<URigVMNode*> SourceNodes;
 	SourceNodes.Add(SourceNode);
 
-	int32 TargetNodeInstructionIndex = TargetNode->GetInstructionIndex();
-	if (TargetNodeInstructionIndex != INDEX_NONE)
+	if (InByteCode)
 	{
-		for (int32 SourceNodeIndex = 0; SourceNodeIndex < SourceNodes.Num(); SourceNodeIndex++)
+		int32 TargetNodeInstructionIndex = InByteCode->GetFirstInstructionIndexForSubject(TargetNode);
+		if (TargetNodeInstructionIndex != INDEX_NONE)
 		{
-			bool bNodeCanLinkAnywhere = SourceNodes[SourceNodeIndex]->IsA<URigVMRerouteNode>();
-			if (!bNodeCanLinkAnywhere)
+			for (int32 SourceNodeIndex = 0; SourceNodeIndex < SourceNodes.Num(); SourceNodeIndex++)
 			{
-				if (URigVMUnitNode* UnitNode = Cast<URigVMUnitNode>(SourceNodes[SourceNodeIndex]))
+				bool bNodeCanLinkAnywhere = SourceNodes[SourceNodeIndex]->IsA<URigVMRerouteNode>();
+				if (!bNodeCanLinkAnywhere)
 				{
-					// pure / immutable nodes can be connected to any input in any order.
-					// since a new link is going to change the abstract syntax tree 
-					if (!UnitNode->IsMutable())
+					if (URigVMUnitNode* UnitNode = Cast<URigVMUnitNode>(SourceNodes[SourceNodeIndex]))
 					{
-						bNodeCanLinkAnywhere = true;
+						// pure / immutable nodes can be connected to any input in any order.
+						// since a new link is going to change the abstract syntax tree 
+						if (!UnitNode->IsMutable())
+						{
+							bNodeCanLinkAnywhere = true;
+						}
 					}
 				}
-			}
 
-			if (!bNodeCanLinkAnywhere)
-			{
-				int32 SourceNodeInstructionIndex = SourceNodes[SourceNodeIndex]->GetInstructionIndex();
-				if (SourceNodeInstructionIndex != INDEX_NONE &&
-					SourceNodeInstructionIndex > TargetNodeInstructionIndex)
+				if (!bNodeCanLinkAnywhere)
 				{
-					return false;
+					int32 SourceNodeInstructionIndex = InByteCode->GetFirstInstructionIndexForSubject(SourceNodes[SourceNodeIndex]);
+					if (SourceNodeInstructionIndex != INDEX_NONE &&
+						SourceNodeInstructionIndex > TargetNodeInstructionIndex)
+					{
+						return false;
+					}
+					SourceNodes.Append(SourceNodes[SourceNodeIndex]->GetLinkedSourceNodes());
 				}
-				SourceNodes.Append(SourceNodes[SourceNodeIndex]->GetLinkedSourceNodes());
 			}
 		}
 	}
diff --git a/Engine/Source/Developer/RigVMDeveloper/Public/RigVMModel/RigVMController.h b/Engine/Source/Developer/RigVMDeveloper/Public/RigVMModel/RigVMController.h
--- a/Engine/Source/Developer/RigVMDeveloper/Public/RigVMModel/RigVMController.h
+++ b/Engine/Source/Developer/RigVMDeveloper/Public/RigVMModel/RigVMController.h
@@ -22,6 +22,7 @@ class URigVMActionStack;
 
 DECLARE_DELEGATE_RetVal_OneParam(bool, FRigVMController_ShouldStructUnfoldDelegate, const UStruct*)
 DECLARE_DELEGATE_RetVal(TArray<FRigVMExternalVariable>, FRigVMController_GetExternalVariablesDelegate)
+DECLARE_DELEGATE_RetVal(const FRigVMByteCode*, FRigVMController_GetByteCodeDelegate)
 
 /**
  * The Controller is the sole authority to perform changes
@@ -524,6 +525,9 @@ public:
 	// A delegate to retrieve the list of external variables
 	FRigVMController_GetExternalVariablesDelegate GetExternalVariablesDelegate;
 
+	// A delegate to retrieve the current bytecode of the graph
+	FRigVMController_GetByteCodeDelegate GetCurrentByteCodeDelegate;
+
 	int32 DetachLinksFromPinObjects(const TArray<URigVMLink*>* InLinks = nullptr, bool bNotify = false);
 	int32 ReattachLinksToPinObjects(bool bFollowCoreRedirectors = false, const TArray<URigVMLink*>* InLinks = nullptr, bool bNotify = false);
 	void AddPinRedirector(bool bInput, bool bOutput, const FString& OldPinPath, const FString& NewPinPath);
@@ -549,6 +553,8 @@ public:
 	// A flag that can be used to turn off pin default value validation if necessary
 	bool bValidatePinDefaults;
 
+	const FRigVMByteCode* GetCurrentByteCode() const;
+
 private:
 
 	UPROPERTY(BlueprintReadOnly, Category = RigVMController, meta = (ScriptName = "ModifiedEvent", AllowPrivateAccess = "true"))
diff --git a/Engine/Source/Developer/RigVMDeveloper/Public/RigVMModel/RigVMGraph.h b/Engine/Source/Developer/RigVMDeveloper/Public/RigVMModel/RigVMGraph.h
--- a/Engine/Source/Developer/RigVMDeveloper/Public/RigVMModel/RigVMGraph.h
+++ b/Engine/Source/Developer/RigVMDeveloper/Public/RigVMModel/RigVMGraph.h
@@ -110,7 +110,7 @@ public:
 	FRigVMGraphModifiedEvent& OnModified();
 
 	void PrepareCycleChecking(URigVMPin* InPin, bool bAsInput);
-	virtual bool CanLink(URigVMPin* InSourcePin, URigVMPin* InTargetPin, FString* OutFailureReason = nullptr);
+	virtual bool CanLink(URigVMPin* InSourcePin, URigVMPin* InTargetPin, FString* OutFailureReason, const FRigVMByteCode* InByteCode);
 	
 	TSharedPtr<FRigVMParserAST> GetDiagnosticsAST(bool bForceRefresh = false, TArray<URigVMLink*> InLinksToSkip = TArray<URigVMLink*>());
 	TSharedPtr<FRigVMParserAST> GetRuntimeAST(const FRigVMParserASTSettings& InSettings = FRigVMParserASTSettings::Optimized(), bool bForceRefresh = false);
diff --git a/Engine/Source/Developer/RigVMDeveloper/Public/RigVMModel/RigVMNode.h b/Engine/Source/Developer/RigVMDeveloper/Public/RigVMModel/RigVMNode.h
--- a/Engine/Source/Developer/RigVMDeveloper/Public/RigVMModel/RigVMNode.h
+++ b/Engine/Source/Developer/RigVMDeveloper/Public/RigVMModel/RigVMNode.h
@@ -54,14 +54,6 @@ public:
 	UFUNCTION(BlueprintCallable, Category = RigVMNode)
 	int32 GetNodeIndex() const;
 
-	// Returns the current index of the instruction in the stack (or INDEX_NONE)
-	UFUNCTION(BlueprintCallable, Category = RigVMNode)
-	int32 GetInstructionIndex() const;
-
-	// Returns the index of the block this node belongs to
-	UFUNCTION(BlueprintCallable, Category = RigVMNode)
-	int32 GetBlockIndex() const;
-
 	// Returns all of the top-level Pins of this Node.
 	UFUNCTION(BlueprintCallable, Category = RigVMNode)
 	virtual const TArray<URigVMPin*>& GetPins() const;
@@ -220,11 +212,6 @@ protected:
 	FName PreviousName;
 
 private:
-	UPROPERTY(transient)
-	int32 InstructionIndex;
-
-	UPROPERTY(transient)
-	int32 BlockIndex;
 
 	UPROPERTY()
 	TArray<URigVMPin*> Pins;
diff --git a/Engine/Source/Developer/RigVMDeveloper/Public/RigVMModel/RigVMPin.h b/Engine/Source/Developer/RigVMDeveloper/Public/RigVMModel/RigVMPin.h
--- a/Engine/Source/Developer/RigVMDeveloper/Public/RigVMModel/RigVMPin.h
+++ b/Engine/Source/Developer/RigVMDeveloper/Public/RigVMModel/RigVMPin.h
@@ -5,6 +5,7 @@
 #include "CoreMinimal.h"
 #include "UObject/ObjectMacros.h"
 #include "RigVMCore/RigVMRegistry.h"
+#include "RigVMCore/RigVMByteCode.h"
 #include "RigVMCompiler/RigVMASTProxy.h"
 #include "RigVMPin.generated.h"
 
@@ -304,7 +305,7 @@ public:
 
 	// Returns true is the two provided source and target Pins
 	// can be linked to one another.
-	static bool CanLink(URigVMPin* InSourcePin, URigVMPin* InTargetPin, FString* OutFailureReason = nullptr);
+	static bool CanLink(URigVMPin* InSourcePin, URigVMPin* InTargetPin, FString* OutFailureReason, const FRigVMByteCode* InByteCode);
 
 	// Returns true if this pin has injected nodes
 	bool HasInjectedNodes() const { return InjectionInfos.Num() > 0; }
diff --git a/Engine/Source/Runtime/RigVM/Private/RigVMCore/RigVMByteCode.cpp b/Engine/Source/Runtime/RigVM/Private/RigVMCore/RigVMByteCode.cpp
--- a/Engine/Source/Runtime/RigVM/Private/RigVMCore/RigVMByteCode.cpp
+++ b/Engine/Source/Runtime/RigVM/Private/RigVMCore/RigVMByteCode.cpp
@@ -117,6 +117,8 @@ void FRigVMInstructionArray::Empty()
 	Instructions.Empty();
 }
 
+TArray<int32> FRigVMByteCode::EmptyInstructionIndices;
+
 FRigVMByteCode::FRigVMByteCode()
 	: NumInstructions(0)
 	, bByteCodeIsAligned(false)
@@ -528,6 +530,11 @@ void FRigVMByteCode::Reset()
 	bByteCodeIsAligned = false;
 	NumInstructions = 0;
 	Entries.Reset();
+
+#if WITH_EDITORONLY_DATA
+	SubjectPerInstruction.Reset();
+	SubjectToInstructions.Reset();
+#endif
 }
 
 void FRigVMByteCode::Empty()
@@ -536,6 +543,11 @@ void FRigVMByteCode::Empty()
 	bByteCodeIsAligned = false;
 	NumInstructions = 0;
 	Entries.Empty();
+
+#if WITH_EDITORONLY_DATA
+	SubjectPerInstruction.Empty();
+	SubjectToInstructions.Empty();
+#endif
 }
 
 uint64 FRigVMByteCode::Num() const
@@ -1221,4 +1233,46 @@ void FRigVMByteCode::AlignByteCode()
 
 	Swap(ByteCode, AlignedByteCode);
 	bByteCodeIsAligned = true;
-}
\ No newline at end of file
+}
+
+#if WITH_EDITOR
+
+UObject* FRigVMByteCode::GetSubjectForInstruction(int32 InInstructionIndex) const
+{
+	if (SubjectPerInstruction.IsValidIndex(InInstructionIndex))
+	{
+		return SubjectPerInstruction[InInstructionIndex];
+	}
+	return nullptr;
+}
+
+int32 FRigVMByteCode::GetFirstInstructionIndexForSubject(UObject* InSubject) const
+{
+	const TArray<int32>& InstructionIndices = GetAllInstructionIndicesForSubject(InSubject);
+	if (InstructionIndices.Num() > 0)
+	{
+		return InstructionIndices[0];
+	}
+	return INDEX_NONE;
+}
+
+const TArray<int32>& FRigVMByteCode::GetAllInstructionIndicesForSubject(UObject* InSubject) const
+{
+	if (const TArray<int32>* InstructionIndices = SubjectToInstructions.Find(InSubject))
+	{
+		return *InstructionIndices;
+	}
+	return EmptyInstructionIndices;
+}
+
+void FRigVMByteCode::SetSubject(int32 InInstructionIndex, UObject* InSubject)
+{
+	if (SubjectPerInstruction.Num() <= InInstructionIndex)
+	{
+		SubjectPerInstruction.AddZeroed(1 + InInstructionIndex - SubjectPerInstruction.Num());
+	}
+	SubjectPerInstruction[InInstructionIndex] = InSubject;
+	SubjectToInstructions.FindOrAdd(InSubject).Add(InInstructionIndex);
+}
+
+#endif
diff --git a/Engine/Source/Runtime/RigVM/Public/RigVMCore/RigVMByteCode.h b/Engine/Source/Runtime/RigVM/Public/RigVMCore/RigVMByteCode.h
--- a/Engine/Source/Runtime/RigVM/Public/RigVMCore/RigVMByteCode.h
+++ b/Engine/Source/Runtime/RigVM/Public/RigVMCore/RigVMByteCode.h
@@ -662,6 +662,19 @@ public:
 
 	FString DumpToText() const;
 
+#if WITH_EDITOR
+
+	// returns the subject which was used to inject a given instruction
+	UObject* GetSubjectForInstruction(int32 InInstructionIndex) const;
+
+	// returns the first hit instruction index for a given subject (or INDEX_NONE)
+	int32 GetFirstInstructionIndexForSubject(UObject* InSubject) const;
+
+	// returns all found instruction indices for a given subject
+	const TArray<int32>& GetAllInstructionIndicesForSubject(UObject* InSubject) const;
+
+#endif
+
 private:
 
 	template<class OpType>
@@ -683,6 +696,19 @@ private:
 	// number of instructions stored here
 	UPROPERTY()
 	int32 NumInstructions;
+	
+#if WITH_EDITORONLY_DATA
+
+	TArray<UObject*> SubjectPerInstruction;
+	TMap<UObject*, TArray<int32>> SubjectToInstructions;
+
+#endif
+
+#if WITH_EDITOR
+
+	void SetSubject(int32 InInstructionIndex, UObject* InSubject);
+
+#endif
 
 	// a look up table from entry name to instruction index
 	UPROPERTY()
@@ -691,6 +717,8 @@ private:
 	// if this is set to true the stored bytecode is aligned / padded
 	bool bByteCodeIsAligned;
 
+	static TArray<int32> EmptyInstructionIndices;
+
 	friend class URigVMCompiler;
 	friend class FRigVMByteCodeTest;
 	friend class URigVM;
