Control Rig: Undeprecated SetControlVector, it was deprecated in CL13656155 when we introduced SetTransform node, which does not replace SetControlVector

diff --git a/Engine/Plugins/Experimental/ControlRig/Source/ControlRig/Private/Units/Hierarchy/RigUnit_SetControlTransform.h b/Engine/Plugins/Experimental/ControlRig/Source/ControlRig/Private/Units/Hierarchy/RigUnit_SetControlTransform.h
--- a/Engine/Plugins/Experimental/ControlRig/Source/ControlRig/Private/Units/Hierarchy/RigUnit_SetControlTransform.h
+++ b/Engine/Plugins/Experimental/ControlRig/Source/ControlRig/Private/Units/Hierarchy/RigUnit_SetControlTransform.h
@@ -376,7 +376,7 @@ struct CONTROLRIG_API FRigUnit_SetMultiControlVector2D : public FRigUnitMutable
 /**
  * SetControlVector is used to perform a change in the hierarchy by setting a single control's Vector value.
  */
-USTRUCT(meta=(DisplayName="Set Control Vector", Category="Hierarchy", DocumentationPolicy="Strict", Keywords = "SetControlVector,SetGizmoVector", PrototypeName = "SetControlValue", Deprecated = "4.25"))
+USTRUCT(meta=(DisplayName="Set Control Vector", Category="Hierarchy", DocumentationPolicy="Strict", Keywords = "SetControlVector,SetGizmoVector", PrototypeName = "SetControlValue"))
 struct CONTROLRIG_API FRigUnit_SetControlVector : public FRigUnitMutable
 {
 	GENERATED_BODY()
