Control Rig: Transform Manipulator now immediately updates for Spaces when using Set Initial Transform To Closest Bone

diff --git a/Engine/Plugins/Experimental/ControlRig/Source/ControlRigDeveloper/Private/ControlRigBlueprint.cpp b/Engine/Plugins/Experimental/ControlRig/Source/ControlRigDeveloper/Private/ControlRigBlueprint.cpp
--- a/Engine/Plugins/Experimental/ControlRig/Source/ControlRigDeveloper/Private/ControlRigBlueprint.cpp
+++ b/Engine/Plugins/Experimental/ControlRig/Source/ControlRigDeveloper/Private/ControlRigBlueprint.cpp
@@ -1168,6 +1168,24 @@ void UControlRigBlueprint::ClearTransientControls()
 	}
 }
 
+void UControlRigBlueprint::SetTransientControlValue(const FRigElementKey& InElement)
+{
+	UControlRigBlueprintGeneratedClass* RigClass = GetControlRigBlueprintGeneratedClass();
+	UControlRig* CDO = Cast<UControlRig>(RigClass->GetDefaultObject(true /* create if needed */));
+
+	TArray<FRigControl> PreviousControls;
+	TArray<UObject*> ArchetypeInstances;
+	CDO->GetArchetypeInstances(ArchetypeInstances);
+	for (UObject* ArchetypeInstance : ArchetypeInstances)
+	{
+		UControlRig* InstancedControlRig = Cast<UControlRig>(ArchetypeInstance);
+		if (InstancedControlRig)
+		{
+			InstancedControlRig->SetTransientControlValue(InElement);
+		}
+	}
+}
+
 #endif
 
 void UControlRigBlueprint::PopulateModelFromGraphForBackwardsCompatibility(UControlRigGraph* InGraph)
diff --git a/Engine/Plugins/Experimental/ControlRig/Source/ControlRigDeveloper/Public/ControlRigBlueprint.h b/Engine/Plugins/Experimental/ControlRig/Source/ControlRigDeveloper/Public/ControlRigBlueprint.h
--- a/Engine/Plugins/Experimental/ControlRig/Source/ControlRigDeveloper/Public/ControlRigBlueprint.h
+++ b/Engine/Plugins/Experimental/ControlRig/Source/ControlRigDeveloper/Public/ControlRigBlueprint.h
@@ -221,6 +221,9 @@ public:
 	/** Removes all  transient / temporary control used to interact with pins */
 	void ClearTransientControls();
 
+	/** update the value of the transient / temporary control to the current value of the rig element*/
+	void SetTransientControlValue(const FRigElementKey& InElement);
+
 #endif
 
 	UPROPERTY(EditAnywhere, Category = "Influence Map")
diff --git a/Engine/Plugins/Experimental/ControlRig/Source/ControlRigEditor/Private/Editor/SRigHierarchy.cpp b/Engine/Plugins/Experimental/ControlRig/Source/ControlRigEditor/Private/Editor/SRigHierarchy.cpp
--- a/Engine/Plugins/Experimental/ControlRig/Source/ControlRigEditor/Private/Editor/SRigHierarchy.cpp
+++ b/Engine/Plugins/Experimental/ControlRig/Source/ControlRigEditor/Private/Editor/SRigHierarchy.cpp
@@ -2904,6 +2904,7 @@ void SRigHierarchy::HandleSetInitialTransformFromClosestBone()
 						GetHierarchyContainer()->SetLocalTransform(Key, LocalTransform);
 						DebuggedContainer->SetInitialTransform(Key, LocalTransform);
 						DebuggedContainer->SetLocalTransform(Key, LocalTransform);
+						Blueprint->SetTransientControlValue(Key);
 					}
 					if (Key.Type == ERigElementType::Control)
 					{
