Control Rig: Fix compilation error for CL15209823

diff --git a/Engine/Plugins/Experimental/ControlRig/Source/ControlRig/Private/ControlRig.cpp b/Engine/Plugins/Experimental/ControlRig/Source/ControlRig/Private/ControlRig.cpp
--- a/Engine/Plugins/Experimental/ControlRig/Source/ControlRig/Private/ControlRig.cpp
+++ b/Engine/Plugins/Experimental/ControlRig/Source/ControlRig/Private/ControlRig.cpp
@@ -445,12 +445,13 @@ void UControlRig::Execute(const EControlRigState InState, const FName& InEventNa
 				VM->AddExternalVariable(ExternalVariable);
 			}
 		}
-
+#if WITH_EDITOR
 		// default to always clear data after each execution
 		// only set a valid first entry event later when execution
 		// has passed the initialization stage and there are multiple events present in one evaluation
 		// first entry event is used to determined when to clear data during an evaluation
 		VM->SetFirstEntryEventInEventQueue(NAME_None);
+#endif
 	}
 
 	bool bJustRanInit = false;
