Control Rig: Editing for node color, category and keywords

diff --git a/Engine/Plugins/Experimental/ControlRig/Source/ControlRigDeveloper/Private/ControlRigBlueprint.cpp b/Engine/Plugins/Experimental/ControlRig/Source/ControlRigDeveloper/Private/ControlRigBlueprint.cpp
--- a/Engine/Plugins/Experimental/ControlRig/Source/ControlRigDeveloper/Private/ControlRigBlueprint.cpp
+++ b/Engine/Plugins/Experimental/ControlRig/Source/ControlRigDeveloper/Private/ControlRigBlueprint.cpp
@@ -1780,6 +1780,12 @@ void UControlRigBlueprint::HandleModifiedEvent(ERigVMGraphNotifType InNotifType,
 				}
 				break;
 			}
+			case ERigVMGraphNotifType::NodeCategoryChanged:
+			case ERigVMGraphNotifType::NodeKeywordsChanged:
+			{
+				FBlueprintEditorUtils::MarkBlueprintAsStructurallyModified(this);
+				break;
+			}
 			default:
 			{
 				break;
diff --git a/Engine/Plugins/Experimental/ControlRig/Source/ControlRigDeveloper/Private/Graph/ControlRigGraph.cpp b/Engine/Plugins/Experimental/ControlRig/Source/ControlRigDeveloper/Private/Graph/ControlRigGraph.cpp
--- a/Engine/Plugins/Experimental/ControlRig/Source/ControlRigDeveloper/Private/Graph/ControlRigGraph.cpp
+++ b/Engine/Plugins/Experimental/ControlRig/Source/ControlRigDeveloper/Private/Graph/ControlRigGraph.cpp
@@ -400,10 +400,16 @@ void UControlRigGraph::HandleModifiedEvent(ERigVMGraphNotifType InNotifType, URi
 		{
 			if (URigVMNode* ModelNode = Cast<URigVMNode>(InSubject))
 			{
-				UEdGraphNode_Comment* EdNode = Cast<UEdGraphNode_Comment>(FindNodeForModelNodeName(ModelNode->GetFName()));
-				if (EdNode)
+				if (ModelNode->IsA<URigVMLibraryNode>())
 				{
-					EdNode->CommentColor = ModelNode->GetNodeColor();
+					if (UControlRigGraphNode* RigNode = Cast<UControlRigGraphNode>(FindNodeForModelNodeName(ModelNode->GetFName())))
+					{
+						RigNode->SetColorFromModel(ModelNode->GetNodeColor());
+					}
+				}
+				else if(UEdGraphNode_Comment * EdComment = Cast<UEdGraphNode_Comment>(FindNodeForModelNodeName(ModelNode->GetFName())))
+				{
+					EdComment->CommentColor = ModelNode->GetNodeColor();
 				}
 			}
 			break;
@@ -629,23 +635,31 @@ void UControlRigGraph::HandleModifiedEvent(ERigVMGraphNotifType InNotifType, URi
 
 int32 UControlRigGraph::GetInstructionIndex(UControlRigGraphNode* InNode)
 {
-	if (const int32* FoundIndex = CachedInstructionIndices.Find(InNode))
+	if (const int32* FoundIndex = CachedInstructionIndices.Find(InNode->GetModelNode()))
 	{
 		return *FoundIndex;
 	}
 
 	struct Local
 	{
-		static int32 GetInstructionIndex(URigVMNode* InModelNode, const FRigVMByteCode* InByteCode)
+		static int32 GetInstructionIndex(URigVMNode* InModelNode, const FRigVMByteCode* InByteCode, TMap<URigVMNode*, int32>& Indices)
 		{
 			if (InModelNode == nullptr)
 			{
 				return INDEX_NONE;
 			}
 
+			if (const int32* ExistingIndex = Indices.Find(InModelNode))
+			{
+				return *ExistingIndex;
+			}
+
+			Indices.Add(InModelNode, INDEX_NONE);
+
 			int32 InstructionIndex = InByteCode->GetFirstInstructionIndexForSubject(InModelNode);
 			if (InstructionIndex != INDEX_NONE)
 			{
+				Indices.FindOrAdd(InModelNode) = InstructionIndex;
 				return InstructionIndex;
 			}
 
@@ -670,7 +684,7 @@ int32 UControlRigGraph::GetInstructionIndex(UControlRigGraphNode* InNode)
 			int32 MinimumInstructionIndex = INDEX_NONE;
 			for (URigVMNode* NodeToCheck : NodesToCheck)
 			{
-				int32 ContainedInstructionIndex = GetInstructionIndex(NodeToCheck, InByteCode);
+				int32 ContainedInstructionIndex = GetInstructionIndex(NodeToCheck, InByteCode, Indices);
 				if (ContainedInstructionIndex != INDEX_NONE)
 				{
 					if (ContainedInstructionIndex < MinimumInstructionIndex || MinimumInstructionIndex == INDEX_NONE)
@@ -679,17 +693,15 @@ int32 UControlRigGraph::GetInstructionIndex(UControlRigGraphNode* InNode)
 					}
 				}
 			}
+
+			Indices.FindOrAdd(InModelNode) = MinimumInstructionIndex;
 			return MinimumInstructionIndex;
 		}
 	};
 
 	if (const FRigVMByteCode* ByteCode = GetController()->GetCurrentByteCode())
 	{
-		int32 InstructionIndex = Local::GetInstructionIndex(InNode->GetModelNode(), ByteCode);
-		if (InstructionIndex != INDEX_NONE)
-		{
-			CachedInstructionIndices.Add(InNode, InstructionIndex);
-		}
+		return Local::GetInstructionIndex(InNode->GetModelNode(), ByteCode, CachedInstructionIndices);
 	}
 
 	return INDEX_NONE;
diff --git a/Engine/Plugins/Experimental/ControlRig/Source/ControlRigDeveloper/Private/Graph/ControlRigGraphSchema.cpp b/Engine/Plugins/Experimental/ControlRig/Source/ControlRigDeveloper/Private/Graph/ControlRigGraphSchema.cpp
--- a/Engine/Plugins/Experimental/ControlRig/Source/ControlRigDeveloper/Private/Graph/ControlRigGraphSchema.cpp
+++ b/Engine/Plugins/Experimental/ControlRig/Source/ControlRigDeveloper/Private/Graph/ControlRigGraphSchema.cpp
@@ -303,13 +303,13 @@ FReply UControlRigGraphSchema::BeginGraphDragAction(TSharedPtr<FEdGraphSchemaAct
 	if (InAction->GetTypeId() == FEdGraphSchemaAction_K2Graph::StaticGetTypeId())
 	{
 		FEdGraphSchemaAction_K2Graph* FuncAction = (FEdGraphSchemaAction_K2Graph*)InAction.Get();
-		if (UControlRigGraph* RigGraph = Cast<UControlRigGraph>(FuncAction->EdGraph))
-		{
-			if (UControlRigBlueprint* RigBlueprint = Cast<UControlRigBlueprint>(FBlueprintEditorUtils::FindBlueprintForGraph(RigGraph)))
-			{
-				return FReply::Handled().BeginDragDrop(FControlRigFunctionDragDropAction::New(InAction, RigBlueprint, RigGraph));
-			}
-		}
+if (UControlRigGraph* RigGraph = Cast<UControlRigGraph>(FuncAction->EdGraph))
+{
+	if (UControlRigBlueprint* RigBlueprint = Cast<UControlRigBlueprint>(FBlueprintEditorUtils::FindBlueprintForGraph(RigGraph)))
+	{
+		return FReply::Handled().BeginDragDrop(FControlRigFunctionDragDropAction::New(InAction, RigBlueprint, RigGraph));
+	}
+}
 	}
 	return FReply::Unhandled();
 }
@@ -383,7 +383,7 @@ void UControlRigGraphSchema::SetNodePosition(UEdGraphNode* Node, const FVector2D
 {
 	if (UControlRigGraphNode* RigNode = Cast<UControlRigGraphNode>(Node))
 	{
-		RigNode->GetController()->SetNodePosition(RigNode->GetModelNode(), Position, true, false); 
+		RigNode->GetController()->SetNodePosition(RigNode->GetModelNode(), Position, true, false);
 	}
 }
 
@@ -394,7 +394,7 @@ void UControlRigGraphSchema::GetGraphDisplayInformation(const UEdGraph& Graph, /
 	if (UControlRigGraph* RigGraph = Cast<UControlRigGraph>((UEdGraph*)&Graph))
 	{
 		TArray<FString> NodePathParts;
-		if(URigVMNode::SplitNodePath(RigGraph->ModelNodePath, NodePathParts))
+		if (URigVMNode::SplitNodePath(RigGraph->ModelNodePath, NodePathParts))
 		{
 			DisplayInfo.DisplayName = FText::FromString(NodePathParts.Last());
 			DisplayInfo.PlainName = DisplayInfo.DisplayName;
@@ -402,24 +402,76 @@ void UControlRigGraphSchema::GetGraphDisplayInformation(const UEdGraph& Graph, /
 	}
 }
 
-bool UControlRigGraphSchema::TryDeleteGraph(UEdGraph* GraphToDelete) const
+FText UControlRigGraphSchema::GetGraphCategory(const UEdGraph* InGraph) const
 {
-	if (UControlRigGraph* RigGraph = Cast<UControlRigGraph>(GraphToDelete))
+	if (UControlRigGraph* RigGraph = Cast<UControlRigGraph>((UEdGraph*)InGraph))
+	{
+		if (URigVMGraph* Model = RigGraph->GetModel())
+		{
+			if (URigVMCollapseNode* CollapseNode = Cast<URigVMCollapseNode>(Model->GetOuter()))
+			{
+				return FText::FromString(CollapseNode->GetNodeCategory());
+			}
+		}
+	}
+	return FText();
+}
+
+FReply UControlRigGraphSchema::TrySetGraphCategory(const UEdGraph* InGraph, const FText& InCategory)
+{
+	if (UControlRigGraph* RigGraph = Cast<UControlRigGraph>((UEdGraph*)InGraph))
 	{
 		if (UControlRigBlueprint* RigBlueprint = Cast<UControlRigBlueprint>(FBlueprintEditorUtils::FindBlueprintForGraph(RigGraph)))
 		{
-			if (URigVMGraph* Model = RigBlueprint->GetModel())
+			if (URigVMGraph* Model = RigGraph->GetModel())
 			{
-				URigVMLibraryNode* LibraryNode = Cast<URigVMLibraryNode>(Model->FindNode(RigGraph->ModelNodePath));
-				if (LibraryNode == nullptr && RigBlueprint->GetLocalFunctionLibrary())
+				if (URigVMCollapseNode* CollapseNode = Cast<URigVMCollapseNode>(Model->GetOuter()))
 				{
-					LibraryNode = Cast<URigVMLibraryNode>(RigBlueprint->GetLocalFunctionLibrary()->FindFunction(*RigGraph->ModelNodePath));
+					if (URigVMController* Controller = RigBlueprint->GetOrCreateController(CollapseNode->GetGraph()))
+					{
+						if (Controller->SetNodeCategory(CollapseNode, InCategory.ToString()))
+						{
+							return FReply::Handled();
+						}
+					}
 				}
+			}
+		}
+	}
+	return FReply::Unhandled();
+}
 
-				if (LibraryNode)
+bool UControlRigGraphSchema::TryDeleteGraph(UEdGraph* GraphToDelete) const
+{
+	if (UControlRigGraph* RigGraph = Cast<UControlRigGraph>(GraphToDelete))
+	{
+		if (UControlRigBlueprint* RigBlueprint = Cast<UControlRigBlueprint>(FBlueprintEditorUtils::FindBlueprintForGraph(RigGraph)))
+		{
+			if (URigVMGraph* Model = RigBlueprint->GetModel(GraphToDelete))
+			{
+				if (URigVMCollapseNode* LibraryNode = Cast<URigVMCollapseNode>(Model->GetOuter()))
 				{
 					if (URigVMController* Controller = RigBlueprint->GetOrCreateController(LibraryNode->GetGraph()))
 					{
+						// check if there is a "bulk remove function" transaction going on.
+						// which implies that a category is being deleted
+						if (GEditor->CanTransact())
+						{
+							if (GEditor->Trans->GetQueueLength() > 0)
+							{
+								const FTransaction* LastTransaction = GEditor->Trans->GetTransaction(GEditor->Trans->GetQueueLength() - 1);
+								if (LastTransaction)
+								{
+									if (LastTransaction->GetTitle().ToString() == TEXT("Bulk Remove Functions"))
+									{
+										// instead of deleting the graph, let's set its category to none
+										// and thus moving it to the top of the tree
+										return Controller->SetNodeCategory(LibraryNode, FString());
+									}
+								}
+							}
+						}
+						
 						return Controller->RemoveNode(LibraryNode);
 					}
 				}
diff --git a/Engine/Plugins/Experimental/ControlRig/Source/ControlRigDeveloper/Public/Graph/ControlRigGraph.h b/Engine/Plugins/Experimental/ControlRig/Source/ControlRigDeveloper/Public/Graph/ControlRigGraph.h
--- a/Engine/Plugins/Experimental/ControlRig/Source/ControlRigDeveloper/Public/Graph/ControlRigGraph.h
+++ b/Engine/Plugins/Experimental/ControlRig/Source/ControlRigDeveloper/Public/Graph/ControlRigGraph.h
@@ -101,7 +101,7 @@ private:
 
 	FControlRigGraphNodeClicked OnGraphNodeClicked;
 
-	TMap<UControlRigGraphNode*, int32> CachedInstructionIndices;
+	TMap<URigVMNode*, int32> CachedInstructionIndices;
 
 #endif
 #if WITH_EDITORONLY_DATA
diff --git a/Engine/Plugins/Experimental/ControlRig/Source/ControlRigDeveloper/Public/Graph/ControlRigGraphSchema.h b/Engine/Plugins/Experimental/ControlRig/Source/ControlRigDeveloper/Public/Graph/ControlRigGraphSchema.h
--- a/Engine/Plugins/Experimental/ControlRig/Source/ControlRigDeveloper/Public/Graph/ControlRigGraphSchema.h
+++ b/Engine/Plugins/Experimental/ControlRig/Source/ControlRigDeveloper/Public/Graph/ControlRigGraphSchema.h
@@ -119,6 +119,8 @@ public:
 	virtual bool IsStructEditable(UStruct* InStruct) const;
 	virtual void SetNodePosition(UEdGraphNode* Node, const FVector2D& Position) const override;
 	virtual void GetGraphDisplayInformation(const UEdGraph& Graph, /*out*/ FGraphDisplayInfo& DisplayInfo) const override;
+	virtual FText GetGraphCategory(const UEdGraph* InGraph) const override;
+	virtual FReply TrySetGraphCategory(const UEdGraph* InGraph, const FText& InCategory) override;
 	virtual bool TryDeleteGraph(UEdGraph* GraphToDelete) const override;
 	virtual bool TryRenameGraph(UEdGraph* GraphToRename, const FName& InNewName) const override;
 	virtual bool CanDuplicateGraph(UEdGraph* InSourceGraph) const { return false; }
diff --git a/Engine/Plugins/Experimental/ControlRig/Source/ControlRigEditor/ControlRigEditor.Build.cs b/Engine/Plugins/Experimental/ControlRig/Source/ControlRigEditor/ControlRigEditor.Build.cs
--- a/Engine/Plugins/Experimental/ControlRig/Source/ControlRigEditor/ControlRigEditor.Build.cs
+++ b/Engine/Plugins/Experimental/ControlRig/Source/ControlRigEditor/ControlRigEditor.Build.cs
@@ -23,8 +23,9 @@ namespace UnrealBuildTool.Rules
 				new string[]
 				{
 					"CurveEditor",
+					"AppFramework",
 				}
-	        );
+			);
 
             PrivateDependencyModuleNames.AddRange(
                 new string[]
diff --git a/Engine/Plugins/Experimental/ControlRig/Source/ControlRigEditor/Private/ControlRigGraphDetails.cpp b/Engine/Plugins/Experimental/ControlRig/Source/ControlRigEditor/Private/ControlRigGraphDetails.cpp
--- a/Engine/Plugins/Experimental/ControlRig/Source/ControlRigEditor/Private/ControlRigGraphDetails.cpp
+++ b/Engine/Plugins/Experimental/ControlRig/Source/ControlRigEditor/Private/ControlRigGraphDetails.cpp
@@ -10,6 +10,8 @@
 #include "SPinTypeSelector.h"
 #include "Widgets/Input/SButton.h"
 #include "Widgets/Text/STextBlock.h"
+#include "Widgets/Colors/SColorPicker.h"
+#include "Widgets/Text/SInlineEditableTextBlock.h"
 #include "PropertyCustomizationHelpers.h"
 #include "NodeFactory.h"
 #include "Graph/ControlRigGraphNode.h"
@@ -497,6 +499,7 @@ void FControlRigArgumentDefaultNode::HandleModifiedEvent(ERigVMGraphNotifType In
 			break;
 		}
 		case ERigVMGraphNotifType::NodeRenamed:
+		case ERigVMGraphNotifType::NodeColorChanged:
 		{
 			URigVMNode* Node = CastChecked<URigVMNode>(InSubject);
 			if (Node == LibraryNode)
@@ -528,6 +531,8 @@ TSharedPtr<IDetailCustomization> FControlRigGraphDetails::MakeInstance(TSharedPt
 
 void FControlRigGraphDetails::CustomizeDetails(IDetailLayoutBuilder& DetailLayout)
 {
+	bIsPickingColor = false;
+
 	TArray<TWeakObjectPtr<UObject>> Objects;
 	DetailLayout.GetObjectsBeingCustomized(Objects);
 
@@ -626,6 +631,72 @@ void FControlRigGraphDetails::CustomizeDetails(IDetailLayoutBuilder& DetailLayou
 	];
 	OutputsCategory.HeaderContent(OutputsHeaderContentWidget);
 
+	IDetailCategoryBuilder& SettingsCategory = DetailLayout.EditCategory("NodeSettings", LOCTEXT("FunctionDetailsNodeSettings", "Node Settings"));
+
+	bool bIsFunction = false;
+	if (Model)
+	{
+		if (URigVMLibraryNode* LibraryNode = Cast<URigVMLibraryNode>(Model->GetOuter()))
+		{
+			bIsFunction = LibraryNode->GetGraph()->IsA<URigVMFunctionLibrary>();
+		}
+	}
+
+	if(bIsFunction)
+	{
+		// node category
+		SettingsCategory.AddCustomRow(FText::GetEmpty())
+		.NameContent()
+		[
+			SNew(STextBlock)
+			.Text(FText::FromString(TEXT("Category")))
+			.Font(IDetailLayoutBuilder::GetDetailFont())
+		]
+		.ValueContent()
+		[
+			SNew(SEditableTextBox)
+			.Font(IDetailLayoutBuilder::GetDetailFont())
+			.Text(this, &FControlRigGraphDetails::GetNodeCategory)
+			.OnTextCommitted(this, &FControlRigGraphDetails::SetNodeCategory)
+		];
+
+		// node keywords
+		SettingsCategory.AddCustomRow(FText::GetEmpty())
+		.NameContent()
+		[
+			SNew(STextBlock)
+			.Text(FText::FromString(TEXT("Keywords")))
+			.Font(IDetailLayoutBuilder::GetDetailFont())
+		]
+		.ValueContent()
+		[
+			SNew(SEditableTextBox)
+			.Font(IDetailLayoutBuilder::GetDetailFont())
+			.Text(this, &FControlRigGraphDetails::GetNodeKeywords)
+			.OnTextCommitted(this, &FControlRigGraphDetails::SetNodeKeywords)
+		];
+	}
+
+	// node color
+	SettingsCategory.AddCustomRow(FText::GetEmpty())
+	.NameContent()
+	[
+		SNew(STextBlock)
+		.Text(FText::FromString(TEXT("Color")))
+		.Font(IDetailLayoutBuilder::GetDetailFont())
+	]
+	.ValueContent()
+	[
+		SNew(SButton)
+		.ButtonStyle(FEditorStyle::Get(), "Menu.Button")
+		.OnClicked(this, &FControlRigGraphDetails::OnNodeColorClicked)
+		[
+			SAssignNew(ColorBlock, SColorBlock)
+			.Color(this, &FControlRigGraphDetails::GetNodeColor)
+			.Size(FVector2D(77, 16))
+		]
+	];
+
 	IDetailCategoryBuilder& DefaultsCategory = DetailLayout.EditCategory("NodeDefaults", LOCTEXT("FunctionDetailsNodeDefaults", "Node Defaults"));
 	TSharedRef<FControlRigArgumentDefaultNode> DefaultsArgumentNode = MakeShareable(new FControlRigArgumentDefaultNode(
 		Model,
@@ -701,4 +772,144 @@ FReply FControlRigGraphDetails::OnAddNewOutputClicked()
 	return FReply::Unhandled();
 }
 
+FText FControlRigGraphDetails::GetNodeCategory() const
+{
+	if (GraphPtr.IsValid() && ControlRigBlueprintPtr.IsValid())
+	{
+		UControlRigBlueprint* Blueprint = ControlRigBlueprintPtr.Get();
+		if (URigVMGraph* Model = Blueprint->GetModel(GraphPtr.Get()))
+		{
+			if (URigVMCollapseNode* OuterNode = Cast<URigVMCollapseNode>(Model->GetOuter()))
+			{
+				return FText::FromString(OuterNode->GetNodeCategory());
+			}
+		}
+	}
+
+	return FText();
+}
+
+void FControlRigGraphDetails::SetNodeCategory(const FText& InNewText, ETextCommit::Type InCommitType)
+{
+	if (GraphPtr.IsValid() && ControlRigBlueprintPtr.IsValid())
+	{
+		UControlRigBlueprint* Blueprint = ControlRigBlueprintPtr.Get();
+		if (URigVMGraph* Model = Blueprint->GetModel(GraphPtr.Get()))
+		{
+			if (URigVMCollapseNode* OuterNode = Cast<URigVMCollapseNode>(Model->GetOuter()))
+			{
+				if (URigVMController* Controller = Blueprint->GetOrCreateController(OuterNode->GetGraph()))
+				{
+					Controller->SetNodeCategory(OuterNode, InNewText.ToString());
+				}
+			}
+		}
+	}
+}
+
+FText FControlRigGraphDetails::GetNodeKeywords() const
+{
+	if (GraphPtr.IsValid() && ControlRigBlueprintPtr.IsValid())
+	{
+		UControlRigBlueprint* Blueprint = ControlRigBlueprintPtr.Get();
+		if (URigVMGraph* Model = Blueprint->GetModel(GraphPtr.Get()))
+		{
+			if (URigVMCollapseNode* OuterNode = Cast<URigVMCollapseNode>(Model->GetOuter()))
+			{
+				return FText::FromString(OuterNode->GetNodeKeywords());
+			}
+		}
+	}
+
+	return FText();
+}
+
+void FControlRigGraphDetails::SetNodeKeywords(const FText& InNewText, ETextCommit::Type InCommitType)
+{
+	if (GraphPtr.IsValid() && ControlRigBlueprintPtr.IsValid())
+	{
+		UControlRigBlueprint* Blueprint = ControlRigBlueprintPtr.Get();
+		if (URigVMGraph* Model = Blueprint->GetModel(GraphPtr.Get()))
+		{
+			if (URigVMCollapseNode* OuterNode = Cast<URigVMCollapseNode>(Model->GetOuter()))
+			{
+				if (URigVMController* Controller = Blueprint->GetOrCreateController(OuterNode->GetGraph()))
+				{
+					Controller->SetNodeKeywords(OuterNode, InNewText.ToString());
+				}
+			}
+		}
+	}
+}
+
+FLinearColor FControlRigGraphDetails::GetNodeColor() const
+{
+	if (GraphPtr.IsValid() && ControlRigBlueprintPtr.IsValid())
+	{
+		UControlRigBlueprint* Blueprint = ControlRigBlueprintPtr.Get();
+		if (URigVMGraph* Model = Blueprint->GetModel(GraphPtr.Get()))
+		{
+			if (URigVMCollapseNode* OuterNode = Cast<URigVMCollapseNode>(Model->GetOuter()))
+			{
+				return OuterNode->GetNodeColor();
+			}
+		}
+	}
+	return FLinearColor::White;
+}
+
+void FControlRigGraphDetails::SetNodeColor(FLinearColor InColor, bool bSetupUndoRedo)
+{
+	TargetColor = InColor;
+
+	if (GraphPtr.IsValid() && ControlRigBlueprintPtr.IsValid())
+	{
+		UControlRigBlueprint* Blueprint = ControlRigBlueprintPtr.Get();
+		if (URigVMGraph* Model = Blueprint->GetModel(GraphPtr.Get()))
+		{
+			if (URigVMCollapseNode* OuterNode = Cast<URigVMCollapseNode>(Model->GetOuter()))
+			{
+				if (URigVMController* Controller = Blueprint->GetOrCreateController(OuterNode->GetGraph()))
+				{
+					Controller->SetNodeColor(OuterNode, TargetColor, bSetupUndoRedo, bIsPickingColor);
+				}
+			}
+		}
+	}
+}
+
+void FControlRigGraphDetails::OnNodeColorBegin()
+{
+	bIsPickingColor = true;
+}
+void FControlRigGraphDetails::OnNodeColorEnd()
+{ 
+	bIsPickingColor = false; 
+}
+
+void FControlRigGraphDetails::OnNodeColorCancelled(FLinearColor OriginalColor)
+{
+	SetNodeColor(OriginalColor, true);
+}
+
+FReply FControlRigGraphDetails::OnNodeColorClicked()
+{
+	TargetColor = GetNodeColor();
+	TargetColors.Reset();
+	TargetColors.Add(&TargetColor);
+
+	FColorPickerArgs PickerArgs;
+	PickerArgs.ParentWidget = ColorBlock;
+	PickerArgs.bUseAlpha = false;
+	PickerArgs.DisplayGamma = false;
+	PickerArgs.InitialColorOverride = TargetColor;
+	PickerArgs.LinearColorArray = &TargetColors;
+	PickerArgs.OnInteractivePickBegin = FSimpleDelegate::CreateSP(this, &FControlRigGraphDetails::OnNodeColorBegin);
+	PickerArgs.OnInteractivePickEnd = FSimpleDelegate::CreateSP(this, &FControlRigGraphDetails::OnNodeColorEnd);
+	PickerArgs.OnColorCommitted = FOnLinearColorValueChanged::CreateSP(this, &FControlRigGraphDetails::SetNodeColor, true);
+	PickerArgs.OnColorPickerCancelled = FOnColorPickerCancelled::CreateSP(this, &FControlRigGraphDetails::OnNodeColorCancelled);
+	OpenColorPicker(PickerArgs);
+	return FReply::Handled();
+}
+
 #undef LOCTEXT_NAMESPACE
diff --git a/Engine/Plugins/Experimental/ControlRig/Source/ControlRigEditor/Private/ControlRigGraphDetails.h b/Engine/Plugins/Experimental/ControlRig/Source/ControlRigEditor/Private/ControlRigGraphDetails.h
--- a/Engine/Plugins/Experimental/ControlRig/Source/ControlRigEditor/Private/ControlRigGraphDetails.h
+++ b/Engine/Plugins/Experimental/ControlRig/Source/ControlRigEditor/Private/ControlRigGraphDetails.h
@@ -11,6 +11,7 @@
 #include "ControlRigEditor.h"
 #include "SGraphPin.h"
 #include "Graph/SControlRigGraphNode.h"
+#include "Widgets/Colors/SColorBlock.h"
 
 class FMobilityCustomization;
 class IDetailLayoutBuilder;
@@ -163,8 +164,19 @@ public:
 	EVisibility GetAddNewInputOutputVisibility() const;
 	FReply OnAddNewInputClicked();
 	FReply OnAddNewOutputClicked();
+	FText GetNodeCategory() const;
+	void SetNodeCategory(const FText& InNewText, ETextCommit::Type InCommitType);
+	FText GetNodeKeywords() const;
+	void SetNodeKeywords(const FText& InNewText, ETextCommit::Type InCommitType);
+	FLinearColor GetNodeColor() const;
+	void SetNodeColor(FLinearColor InColor, bool bSetupUndoRedo);
+	void OnNodeColorBegin();
+	void OnNodeColorEnd();
+	void OnNodeColorCancelled(FLinearColor OriginalColor);
+	FReply OnNodeColorClicked();
 
 private:
+
 	/** The Blueprint editor we are embedded in */
 	TWeakPtr<IControlRigEditor> ControlRigEditorPtr;
 
@@ -173,4 +185,16 @@ private:
 
 	/** The blueprint we are editing */
 	TWeakObjectPtr<UControlRigBlueprint> ControlRigBlueprintPtr;
+
+	/** The color block widget */
+	TSharedPtr<SColorBlock> ColorBlock;
+
+	/** The color to change */
+	FLinearColor TargetColor;
+
+	/** The color array to change */
+	TArray<FLinearColor*> TargetColors;
+
+	/** Set to true if the UI is currently picking a color */
+	bool bIsPickingColor;
 };
\ No newline at end of file
diff --git a/Engine/Plugins/Experimental/ControlRig/Source/ControlRigEditor/Private/Graph/NodeSpawners/ControlRigFunctionRefNodeSpawner.cpp b/Engine/Plugins/Experimental/ControlRig/Source/ControlRigEditor/Private/Graph/NodeSpawners/ControlRigFunctionRefNodeSpawner.cpp
--- a/Engine/Plugins/Experimental/ControlRig/Source/ControlRigEditor/Private/Graph/NodeSpawners/ControlRigFunctionRefNodeSpawner.cpp
+++ b/Engine/Plugins/Experimental/ControlRig/Source/ControlRigEditor/Private/Graph/NodeSpawners/ControlRigFunctionRefNodeSpawner.cpp
@@ -36,7 +36,7 @@ UControlRigFunctionRefNodeSpawner* UControlRigFunctionRefNodeSpawner::CreateFrom
 
 	FBlueprintActionUiSpec& MenuSignature = NodeSpawner->DefaultMenuSignature;
 
-	FString Category = InFunction->GetCategory();
+	FString Category = InFunction->GetNodeCategory();
 	FString CategoryPrefix;
 
 	if (UControlRigBlueprint* Blueprint = Cast<UControlRigBlueprint>(InFunction->GetLibrary()->GetOuter()))
@@ -59,9 +59,9 @@ UControlRigFunctionRefNodeSpawner* UControlRigFunctionRefNodeSpawner::CreateFrom
 	}
 
 	MenuSignature.MenuName = FText::FromString(InFunction->GetName());
-	MenuSignature.Tooltip = FText::FromString(InFunction->GetTooltip());
+	MenuSignature.Tooltip = InFunction->GetToolTipText();
 	MenuSignature.Category = FText::FromString(CategoryPrefix + Category);
-	MenuSignature.Keywords = FText::FromString(InFunction->GetKeywords());
+	MenuSignature.Keywords = FText::FromString(InFunction->GetNodeKeywords());
 
 	// add at least one character, so that PrimeDefaultUiSpec() doesn't 
 	// attempt to query the template node
diff --git a/Engine/Source/Developer/RigVMDeveloper/Private/RigVMModel/Nodes/RigVMCollapseNode.cpp b/Engine/Source/Developer/RigVMDeveloper/Private/RigVMModel/Nodes/RigVMCollapseNode.cpp
--- a/Engine/Source/Developer/RigVMDeveloper/Private/RigVMModel/Nodes/RigVMCollapseNode.cpp
+++ b/Engine/Source/Developer/RigVMDeveloper/Private/RigVMModel/Nodes/RigVMCollapseNode.cpp
@@ -3,18 +3,6 @@
 #include "RigVMModel/Nodes/RigVMCollapseNode.h"
 #include "RigVMModel/RigVMFunctionLibrary.h"
 
-FString URigVMCollapseNode::GetNodeTitle() const
-{
-	// todo
-	return URigVMNode::GetNodeTitle();
-}
-
-FText URigVMCollapseNode::GetToolTipText() const
-{
-	// todo
-	return URigVMNode::GetToolTipText();
-}
-
 URigVMFunctionLibrary* URigVMCollapseNode::GetLibrary() const
 {
 	return Cast<URigVMFunctionLibrary>(GetOuter());
diff --git a/Engine/Source/Developer/RigVMDeveloper/Private/RigVMModel/Nodes/RigVMFunctionReferenceNode.cpp b/Engine/Source/Developer/RigVMDeveloper/Private/RigVMModel/Nodes/RigVMFunctionReferenceNode.cpp
--- a/Engine/Source/Developer/RigVMDeveloper/Private/RigVMModel/Nodes/RigVMFunctionReferenceNode.cpp
+++ b/Engine/Source/Developer/RigVMDeveloper/Private/RigVMModel/Nodes/RigVMFunctionReferenceNode.cpp
@@ -21,13 +21,22 @@ FLinearColor URigVMFunctionReferenceNode::GetNodeColor() const
 	return Super::GetNodeColor();
 }
 
-FText URigVMFunctionReferenceNode::GetToolTipText() const
+FString URigVMFunctionReferenceNode::GetNodeCategory() const
 {
 	if (URigVMLibraryNode* ReferencedNode = GetReferencedNode())
 	{
-		return ReferencedNode->GetToolTipText();
+		return ReferencedNode->GetNodeCategory();
 	}
-	return Super::GetToolTipText();
+	return Super::GetNodeCategory();
+}
+
+FString URigVMFunctionReferenceNode::GetNodeKeywords() const
+{
+	if (URigVMLibraryNode* ReferencedNode = GetReferencedNode())
+	{
+		return ReferencedNode->GetNodeKeywords();
+	}
+	return Super::GetNodeKeywords();
 }
 
 URigVMFunctionLibrary* URigVMFunctionReferenceNode::GetLibrary() const
diff --git a/Engine/Source/Developer/RigVMDeveloper/Private/RigVMModel/RigVMController.cpp b/Engine/Source/Developer/RigVMDeveloper/Private/RigVMModel/RigVMController.cpp
--- a/Engine/Source/Developer/RigVMDeveloper/Private/RigVMModel/RigVMController.cpp
+++ b/Engine/Source/Developer/RigVMDeveloper/Private/RigVMModel/RigVMController.cpp
@@ -3635,6 +3635,26 @@ bool URigVMController::SetNodeColor(URigVMNode* InNode, const FLinearColor& InCo
 	InNode->NodeColor = InColor;
 	Notify(ERigVMGraphNotifType::NodeColorChanged, InNode);
 
+	if (URigVMLibraryNode* LibraryNode = Cast<URigVMLibraryNode>(InNode))
+	{
+		if (URigVMFunctionLibrary* FunctionLibrary = Cast<URigVMFunctionLibrary>(LibraryNode->GetGraph()))
+		{
+			FRigVMFunctionReferenceArray* ReferencesEntry = FunctionLibrary->FunctionReferences.Find(LibraryNode);
+			if (ReferencesEntry)
+			{
+				for (TSoftObjectPtr<URigVMFunctionReferenceNode> FunctionReferencePtr : ReferencesEntry->FunctionReferences)
+				{
+					if (FunctionReferencePtr.IsValid())
+					{
+						URigVMNode* ReferenceNode = FunctionReferencePtr.Get();
+						FRigVMControllerGraphGuard GraphGuard(this, ReferenceNode->GetGraph(), false);
+						Notify(ERigVMGraphNotifType::NodeColorChanged, ReferenceNode);
+					}
+				}
+			}
+		}
+	}
+
 	if (bSetupUndoRedo)
 	{
 		ActionStack->EndAction(Action, bMergeUndoAction);
@@ -3657,6 +3677,96 @@ bool URigVMController::SetNodeColorByName(const FName& InNodeName, const FLinear
 	return SetNodeColor(Node, InColor, bSetupUndoRedo, bMergeUndoAction);
 }
 
+bool URigVMController::SetNodeCategory(URigVMCollapseNode* InNode, const FString& InCategory, bool bSetupUndoRedo, bool bMergeUndoAction)
+{
+	if (!IsValidNodeForGraph(InNode))
+	{
+		return false;
+	}
+
+	if (InNode->GetNodeCategory() == InCategory)
+	{
+		return false;
+	}
+
+	FRigVMSetNodeCategoryAction Action;
+	if (bSetupUndoRedo)
+	{
+		Action = FRigVMSetNodeCategoryAction(InNode, InCategory);
+		Action.Title = FString::Printf(TEXT("Set Node Category"));
+		ActionStack->BeginAction(Action);
+	}
+
+	InNode->NodeCategory = InCategory;
+	Notify(ERigVMGraphNotifType::NodeCategoryChanged, InNode);
+
+	if (bSetupUndoRedo)
+	{
+		ActionStack->EndAction(Action, bMergeUndoAction);
+	}
+
+	return true;
+}
+
+bool URigVMController::SetNodeCategoryByName(const FName& InNodeName, const FString& InCategory, bool bSetupUndoRedo, bool bMergeUndoAction)
+{
+	if (!IsValidGraph())
+	{
+		return false;
+	}
+
+	URigVMGraph* Graph = GetGraph();
+	check(Graph);
+
+	URigVMCollapseNode* Node = Cast<URigVMCollapseNode>(Graph->FindNodeByName(InNodeName));
+	return SetNodeCategory(Node, InCategory, bSetupUndoRedo, bMergeUndoAction);
+}
+
+bool URigVMController::SetNodeKeywords(URigVMCollapseNode* InNode, const FString& InKeywords, bool bSetupUndoRedo, bool bMergeUndoAction)
+{
+	if (!IsValidNodeForGraph(InNode))
+	{
+		return false;
+	}
+
+	if (InNode->GetNodeKeywords() == InKeywords)
+	{
+		return false;
+	}
+
+	FRigVMSetNodeKeywordsAction Action;
+	if (bSetupUndoRedo)
+	{
+		Action = FRigVMSetNodeKeywordsAction(InNode, InKeywords);
+		Action.Title = FString::Printf(TEXT("Set Node Keywords"));
+		ActionStack->BeginAction(Action);
+	}
+
+	InNode->NodeKeywords = InKeywords;
+	Notify(ERigVMGraphNotifType::NodeKeywordsChanged, InNode);
+
+	if (bSetupUndoRedo)
+	{
+		ActionStack->EndAction(Action, bMergeUndoAction);
+	}
+
+	return true;
+}
+
+bool URigVMController::SetNodeKeywordsByName(const FName& InNodeName, const FString& InKeywords, bool bSetupUndoRedo, bool bMergeUndoAction)
+{
+	if (!IsValidGraph())
+	{
+		return false;
+	}
+
+	URigVMGraph* Graph = GetGraph();
+	check(Graph);
+
+	URigVMCollapseNode* Node = Cast<URigVMCollapseNode>(Graph->FindNodeByName(InNodeName));
+	return SetNodeKeywords(Node, InKeywords, bSetupUndoRedo, bMergeUndoAction);
+}
+
 bool URigVMController::SetCommentText(URigVMNode* InNode, const FString& InCommentText, bool bSetupUndoRedo)
 {
 	if (!IsValidNodeForGraph(InNode))
diff --git a/Engine/Source/Developer/RigVMDeveloper/Private/RigVMModel/RigVMControllerActions.cpp b/Engine/Source/Developer/RigVMDeveloper/Private/RigVMModel/RigVMControllerActions.cpp
--- a/Engine/Source/Developer/RigVMDeveloper/Private/RigVMModel/RigVMControllerActions.cpp
+++ b/Engine/Source/Developer/RigVMDeveloper/Private/RigVMModel/RigVMControllerActions.cpp
@@ -911,6 +911,90 @@ bool FRigVMSetNodeColorAction::Redo(URigVMController* InController)
 	return FRigVMBaseAction::Redo(InController);
 }
 
+FRigVMSetNodeCategoryAction::FRigVMSetNodeCategoryAction(URigVMCollapseNode* InNode, const FString& InNewCategory)
+	: NodePath(InNode->GetNodePath())
+	, OldCategory(InNode->GetNodeCategory())
+	, NewCategory(InNewCategory)
+{
+}
+
+bool FRigVMSetNodeCategoryAction::Merge(const FRigVMBaseAction* Other)
+{
+	if (!FRigVMBaseAction::Merge(Other))
+	{
+		return false;
+	}
+
+	const FRigVMSetNodeCategoryAction* Action = (const FRigVMSetNodeCategoryAction*)Other;
+	if (NodePath != Action->NodePath)
+	{
+		return false;
+	}
+
+	NewCategory = Action->NewCategory;
+	return true;
+}
+
+bool FRigVMSetNodeCategoryAction::Undo(URigVMController* InController)
+{
+	if (!FRigVMBaseAction::Undo(InController))
+	{
+		return false;
+	}
+	return InController->SetNodeCategoryByName(*NodePath, OldCategory, false);
+}
+
+bool FRigVMSetNodeCategoryAction::Redo(URigVMController* InController)
+{
+	if (!InController->SetNodeCategoryByName(*NodePath, NewCategory, false))
+	{
+		return false;
+	}
+	return FRigVMBaseAction::Redo(InController);
+}
+
+FRigVMSetNodeKeywordsAction::FRigVMSetNodeKeywordsAction(URigVMCollapseNode* InNode, const FString& InNewKeywords)
+	: NodePath(InNode->GetNodePath())
+	, OldKeywords(InNode->GetNodeKeywords())
+	, NewKeywords(InNewKeywords)
+{
+}
+
+bool FRigVMSetNodeKeywordsAction::Merge(const FRigVMBaseAction* Other)
+{
+	if (!FRigVMBaseAction::Merge(Other))
+	{
+		return false;
+	}
+
+	const FRigVMSetNodeKeywordsAction* Action = (const FRigVMSetNodeKeywordsAction*)Other;
+	if (NodePath != Action->NodePath)
+	{
+		return false;
+	}
+
+	NewKeywords = Action->NewKeywords;
+	return true;
+}
+
+bool FRigVMSetNodeKeywordsAction::Undo(URigVMController* InController)
+{
+	if (!FRigVMBaseAction::Undo(InController))
+	{
+		return false;
+	}
+	return InController->SetNodeKeywordsByName(*NodePath, OldKeywords, false);
+}
+
+bool FRigVMSetNodeKeywordsAction::Redo(URigVMController* InController)
+{
+	if (!InController->SetNodeKeywordsByName(*NodePath, NewKeywords, false))
+	{
+		return false;
+	}
+	return FRigVMBaseAction::Redo(InController);
+}
+
 FRigVMSetCommentTextAction::FRigVMSetCommentTextAction(URigVMCommentNode* InNode, const FString& InNewText)
 : NodePath(InNode->GetNodePath())
 , OldText(InNode->GetCommentText())
diff --git a/Engine/Source/Developer/RigVMDeveloper/Public/RigVMModel/Nodes/RigVMCollapseNode.h b/Engine/Source/Developer/RigVMDeveloper/Public/RigVMModel/Nodes/RigVMCollapseNode.h
--- a/Engine/Source/Developer/RigVMDeveloper/Public/RigVMModel/Nodes/RigVMCollapseNode.h
+++ b/Engine/Source/Developer/RigVMDeveloper/Public/RigVMModel/Nodes/RigVMCollapseNode.h
@@ -17,11 +17,9 @@ class RIGVMDEVELOPER_API URigVMCollapseNode : public URigVMLibraryNode
 
 public:
 
-	// Override node functions
-	virtual FString GetNodeTitle() const override;
-	virtual FText GetToolTipText() const override;
-
 	// Library node interface
+	virtual FString GetNodeCategory() const override { return NodeCategory; }
+	virtual FString GetNodeKeywords() const override { return NodeKeywords; }
 	virtual URigVMFunctionLibrary* GetLibrary() const override;
 	virtual URigVMGraph* GetContainedGraph() const override { return ContainedGraph; }
 
@@ -32,6 +30,12 @@ private:
 	UPROPERTY()
 	URigVMGraph* ContainedGraph;
 
+	UPROPERTY()
+	FString NodeCategory;
+
+	UPROPERTY()
+	FString NodeKeywords;
+
 	friend class URigVMController;
 };
 
diff --git a/Engine/Source/Developer/RigVMDeveloper/Public/RigVMModel/Nodes/RigVMFunctionReferenceNode.h b/Engine/Source/Developer/RigVMDeveloper/Public/RigVMModel/Nodes/RigVMFunctionReferenceNode.h
--- a/Engine/Source/Developer/RigVMDeveloper/Public/RigVMModel/Nodes/RigVMFunctionReferenceNode.h
+++ b/Engine/Source/Developer/RigVMDeveloper/Public/RigVMModel/Nodes/RigVMFunctionReferenceNode.h
@@ -22,10 +22,11 @@ public:
 	// URigVMNode interface
 	virtual FString GetNodeTitle() const override;
 	virtual FLinearColor GetNodeColor() const override;
-	virtual FText GetToolTipText() const override;
 	// end URigVMNode interface
 
 	// URigVMLibraryNode interface
+	virtual FString GetNodeCategory() const override;
+	virtual FString GetNodeKeywords() const override;
 	virtual URigVMFunctionLibrary* GetLibrary() const override;
 	virtual URigVMGraph* GetContainedGraph() const override;
 	// end URigVMLibraryNode interface
diff --git a/Engine/Source/Developer/RigVMDeveloper/Public/RigVMModel/Nodes/RigVMLibraryNode.h b/Engine/Source/Developer/RigVMDeveloper/Public/RigVMModel/Nodes/RigVMLibraryNode.h
--- a/Engine/Source/Developer/RigVMDeveloper/Public/RigVMModel/Nodes/RigVMLibraryNode.h
+++ b/Engine/Source/Developer/RigVMDeveloper/Public/RigVMModel/Nodes/RigVMLibraryNode.h
@@ -28,9 +28,8 @@ public:
 	virtual bool IsDefinedAsVarying() const override;
 
 	// Library node interface
-	virtual FString GetCategory() const { return Category; }
-	virtual FString GetTooltip() const { return Tooltip; }
-	virtual FString GetKeywords() const { return Keywords; }
+	virtual FString GetNodeCategory() const { return FString(); }
+	virtual FString GetNodeKeywords() const { return FString(); }
 	virtual URigVMFunctionLibrary* GetLibrary() const { return nullptr; }
 	virtual URigVMGraph* GetContainedGraph() const { return nullptr; }
 	virtual const TArray<URigVMNode*>& GetContainedNodes() const;
@@ -40,15 +39,6 @@ public:
 
 protected:
 
-	UPROPERTY()
-	FString Category;
-
-	UPROPERTY()
-	FString Tooltip;
-
-	UPROPERTY()
-	FString Keywords;
-
 	const static TArray<URigVMNode*> EmptyNodes;
 	const static TArray<URigVMLink*> EmptyLinks;
 
diff --git a/Engine/Source/Developer/RigVMDeveloper/Public/RigVMModel/RigVMController.h b/Engine/Source/Developer/RigVMDeveloper/Public/RigVMModel/RigVMController.h
--- a/Engine/Source/Developer/RigVMDeveloper/Public/RigVMModel/RigVMController.h
+++ b/Engine/Source/Developer/RigVMDeveloper/Public/RigVMModel/RigVMController.h
@@ -363,6 +363,26 @@ public:
 	UFUNCTION(BlueprintCallable, Category = RigVMController)
 	bool SetNodeColorByName(const FName& InNodeName, const FLinearColor& InColor, bool bSetupUndoRedo = true, bool bMergeUndoAction = false);
 
+	// Sets the category of a node in the graph.
+	// This causes a NodeCategoryChanged modified event.
+	UFUNCTION(BlueprintCallable, Category = RigVMController)
+	bool SetNodeCategory(URigVMCollapseNode* InNode, const FString& InCategory, bool bSetupUndoRedo = true, bool bMergeUndoAction = false);
+
+	// Sets the category of a node in the graph.
+	// This causes a NodeCategoryChanged modified event.
+	UFUNCTION(BlueprintCallable, Category = RigVMController)
+	bool SetNodeCategoryByName(const FName& InNodeName, const FString& InCategory, bool bSetupUndoRedo = true, bool bMergeUndoAction = false);
+
+	// Sets the keywords of a node in the graph.
+	// This causes a NodeKeywordsChanged modified event.
+	UFUNCTION(BlueprintCallable, Category = RigVMController)
+	bool SetNodeKeywords(URigVMCollapseNode* InNode, const FString& InKeywords, bool bSetupUndoRedo = true, bool bMergeUndoAction = false);
+
+	// Sets the keywords of a node in the graph.
+	// This causes a NodeKeywordsChanged modified event.
+	UFUNCTION(BlueprintCallable, Category = RigVMController)
+	bool SetNodeKeywordsByName(const FName& InNodeName, const FString& InKeywords, bool bSetupUndoRedo = true, bool bMergeUndoAction = false);
+
 	// Sets the comment text of a comment node in the graph.
 	// This causes a CommentTextChanged modified event.
 	UFUNCTION(BlueprintCallable, Category = RigVMController)
diff --git a/Engine/Source/Developer/RigVMDeveloper/Public/RigVMModel/RigVMControllerActions.h b/Engine/Source/Developer/RigVMDeveloper/Public/RigVMModel/RigVMControllerActions.h
--- a/Engine/Source/Developer/RigVMDeveloper/Public/RigVMModel/RigVMControllerActions.h
+++ b/Engine/Source/Developer/RigVMDeveloper/Public/RigVMModel/RigVMControllerActions.h
@@ -779,6 +779,67 @@ public:
 	FLinearColor NewColor;
 };
 
+/**
+ * An action setting a node's category in the graph.
+ */
+USTRUCT()
+struct FRigVMSetNodeCategoryAction : public FRigVMBaseAction
+{
+	GENERATED_BODY()
+
+public:
+
+	FRigVMSetNodeCategoryAction()
+	{
+		OldCategory = NewCategory = FString();
+	}
+	FRigVMSetNodeCategoryAction(URigVMCollapseNode* InNode, const FString& InNewCategory);
+	virtual ~FRigVMSetNodeCategoryAction() {};
+	virtual bool Merge(const FRigVMBaseAction* Other);
+	virtual bool Undo(URigVMController* InController) override;
+	virtual bool Redo(URigVMController* InController) override;
+
+	UPROPERTY()
+	FString NodePath;
+
+	UPROPERTY()
+	FString OldCategory;
+
+	UPROPERTY()
+	FString NewCategory;
+};
+
+
+/**
+ * An action setting a node's keywords in the graph.
+ */
+USTRUCT()
+struct FRigVMSetNodeKeywordsAction : public FRigVMBaseAction
+{
+	GENERATED_BODY()
+
+public:
+
+	FRigVMSetNodeKeywordsAction()
+	{
+		OldKeywords = NewKeywords = FString();
+	}
+	FRigVMSetNodeKeywordsAction(URigVMCollapseNode* InNode, const FString& InNewKeywords);
+	virtual ~FRigVMSetNodeKeywordsAction() {};
+	virtual bool Merge(const FRigVMBaseAction* Other);
+	virtual bool Undo(URigVMController* InController) override;
+	virtual bool Redo(URigVMController* InController) override;
+
+	UPROPERTY()
+	FString NodePath;
+
+	UPROPERTY()
+	FString OldKeywords;
+
+	UPROPERTY()
+	FString NewKeywords;
+};
+
 /**
  * An action setting a comment node's text in the graph.
  */
diff --git a/Engine/Source/Developer/RigVMDeveloper/Public/RigVMModel/RigVMNotifications.h b/Engine/Source/Developer/RigVMDeveloper/Public/RigVMModel/RigVMNotifications.h
--- a/Engine/Source/Developer/RigVMDeveloper/Public/RigVMModel/RigVMNotifications.h
+++ b/Engine/Source/Developer/RigVMDeveloper/Public/RigVMModel/RigVMNotifications.h
@@ -49,6 +49,8 @@ enum class ERigVMGraphNotifType : uint8
 	InteractionBracketCanceled, // A bracket has been canceled (Subject == nullptr)
 	PinBoundVariableChanged, // A pin has been bound or unbound to / from a variable (Subject == URigVMPin)
 	NodeRenamed, // A node has been renamed in the graph (Subject == URigVMNode)
+	NodeCategoryChanged, // A node's category has changed (Subject == URigVMNode)
+	NodeKeywordsChanged, // A node's keywords have changed (Subject == URigVMNode)
 	Invalid // The max for this enum (used for guarding)
 };
 
diff --git a/Engine/Source/Editor/Kismet/Private/SMyBlueprint.cpp b/Engine/Source/Editor/Kismet/Private/SMyBlueprint.cpp
--- a/Engine/Source/Editor/Kismet/Private/SMyBlueprint.cpp
+++ b/Engine/Source/Editor/Kismet/Private/SMyBlueprint.cpp
@@ -637,7 +637,41 @@ void SMyBlueprint::OnCategoryNameCommitted(const FText& InNewText, ETextCommit::
 					// Don't allow changing the category of a graph who's parent is not the current Blueprint
 					if(GraphAction && !FBlueprintEditorUtils::IsPaletteActionReadOnly(Actions[i], BlueprintEditorPtr.Pin()) && FBlueprintEditorUtils::FindBlueprintForGraph(GraphAction->EdGraph) == GetBlueprintObj())
 					{
-						GraphAction->MovePersistentItemToCategory(CategoryName);
+						FReply ReplyBySchema = FReply::Unhandled();
+						if (GraphAction->EdGraph)
+						{
+							if (UEdGraphSchema* Schema = (UEdGraphSchema*)GraphAction->EdGraph->GetSchema())
+							{
+								TArray<FString> CategoryParts;
+								TWeakPtr< FGraphActionNode > CurrentActionNode = InAction;
+								while (CurrentActionNode.IsValid())
+								{
+									FString CurrentDisplayName = CurrentActionNode.Pin()->GetDisplayName().ToString();
+									if (!CurrentDisplayName.IsEmpty())
+									{
+										CategoryParts.Insert(CurrentDisplayName, 0);
+									}
+									CurrentActionNode = CurrentActionNode.Pin()->GetParentNode();
+								}
+
+								FString OldCategoryPath = FString::Join(CategoryParts, TEXT("|"));
+								CategoryParts.Last() = CategoryName.ToString();
+								FString NewCategoryPath = FString::Join(CategoryParts, TEXT("|"));
+
+								FString CurrentCategoryPath = GraphAction->GetCategory().ToString();
+								if (CurrentCategoryPath == OldCategoryPath || CurrentCategoryPath.StartsWith(OldCategoryPath + TEXT("|"), ESearchCase::CaseSensitive))
+								{
+									NewCategoryPath = NewCategoryPath + CurrentCategoryPath.RightChop(OldCategoryPath.Len());
+								}
+
+								ReplyBySchema = Schema->TrySetGraphCategory(GraphAction->EdGraph, FText::FromString(NewCategoryPath));
+							}
+						}
+
+						if (!ReplyBySchema.IsEventHandled())
+						{
+							GraphAction->MovePersistentItemToCategory(CategoryName);
+						}
 					}
 				}
 			}
@@ -1272,8 +1306,8 @@ void SMyBlueprint::CollectAllActions(FGraphActionListBuilderBase& OutAllActions)
 		FGraphDisplayInfo DisplayInfo;
 		Graph->GetSchema()->GetGraphDisplayInformation(*Graph, DisplayInfo);
 
-		FText FunctionCategory;
-		if (BlueprintObj->SkeletonGeneratedClass != nullptr)
+		FText FunctionCategory = Graph->GetSchema()->GetGraphCategory(Graph);
+		if (FunctionCategory.IsEmpty() && BlueprintObj->SkeletonGeneratedClass != nullptr)
 		{
 			UFunction* Function = BlueprintObj->SkeletonGeneratedClass->FindFunctionByName(Graph->GetFName());
 			if (Function != nullptr)
diff --git a/Engine/Source/Runtime/Engine/Classes/EdGraph/EdGraphSchema.h b/Engine/Source/Runtime/Engine/Classes/EdGraph/EdGraphSchema.h
--- a/Engine/Source/Runtime/Engine/Classes/EdGraph/EdGraphSchema.h
+++ b/Engine/Source/Runtime/Engine/Classes/EdGraph/EdGraphSchema.h
@@ -968,6 +968,21 @@ class ENGINE_API UEdGraphSchema : public UObject
 	 */
 	virtual void GetGraphDisplayInformation(const UEdGraph& Graph, /*out*/ FGraphDisplayInfo& DisplayInfo) const;
 
+	/**
+	 * Returns an optional category for a graph
+	 *
+	 * @param	InGraph				Graph to get the category for
+	 */
+	virtual FText GetGraphCategory(const UEdGraph* InGraph) const { return FText(); }
+
+	/**
+	 * Tentatively sets the category for a given graph
+	 *
+	 * @param	InGraph				Graph to set the category for
+	 * @param	InCategory			Pipe "|" separated category for the graph.
+	 */
+	virtual FReply TrySetGraphCategory(const UEdGraph* InGraph, const FText& InCategory) { return FReply::Unhandled(); }
+
 	/** Called when asset(s) are dropped onto a graph background. */
 	virtual void DroppedAssetsOnGraph(const TArray<struct FAssetData>& Assets, const FVector2D& GraphPosition, UEdGraph* Graph) const {}
 
