Control Rig: Make Control Rig non-blueprintable and guard against non-ControlRigGraph graphs to avoid crashes

diff --git a/Engine/Plugins/Experimental/ControlRig/Source/ControlRig/Public/ControlRig.h b/Engine/Plugins/Experimental/ControlRig/Source/ControlRig/Public/ControlRig.h
--- a/Engine/Plugins/Experimental/ControlRig/Source/ControlRig/Public/ControlRig.h
+++ b/Engine/Plugins/Experimental/ControlRig/Source/ControlRig/Public/ControlRig.h
@@ -37,7 +37,7 @@ struct FRigControl;
 CONTROLRIG_API DECLARE_LOG_CATEGORY_EXTERN(LogControlRig, Log, All);
 
 /** Runs logic for mapping input data to transforms (the "Rig") */
-UCLASS(Blueprintable, Abstract, editinlinenew)
+UCLASS(BlueprintType, Abstract, editinlinenew)
 class CONTROLRIG_API UControlRig : public UObject, public INodeMappingProviderInterface, public IInterface_AssetUserData
 {
 	GENERATED_UCLASS_BODY()
diff --git a/Engine/Plugins/Experimental/ControlRig/Source/ControlRigEditor/Private/Graph/NodeSpawners/ControlRigBranchNodeSpawner.cpp b/Engine/Plugins/Experimental/ControlRig/Source/ControlRigEditor/Private/Graph/NodeSpawners/ControlRigBranchNodeSpawner.cpp
--- a/Engine/Plugins/Experimental/ControlRig/Source/ControlRigEditor/Private/Graph/NodeSpawners/ControlRigBranchNodeSpawner.cpp
+++ b/Engine/Plugins/Experimental/ControlRig/Source/ControlRigEditor/Private/Graph/NodeSpawners/ControlRigBranchNodeSpawner.cpp
@@ -56,7 +56,7 @@ UEdGraphNode* UControlRigBranchNodeSpawner::Invoke(UEdGraph* ParentGraph, FBindi
 
 	// First create a backing member for our node
 	UControlRigGraph* RigGraph = Cast<UControlRigGraph>(ParentGraph);
-	check(RigGraph);
+	if(RigGraph == nullptr) return nullptr;
 	UControlRigBlueprint* RigBlueprint = Cast<UControlRigBlueprint>(FBlueprintEditorUtils::FindBlueprintForGraph(ParentGraph));
 	check(RigBlueprint);
 
diff --git a/Engine/Plugins/Experimental/ControlRig/Source/ControlRigEditor/Private/Graph/NodeSpawners/ControlRigEnumNodeSpawner.cpp b/Engine/Plugins/Experimental/ControlRig/Source/ControlRigEditor/Private/Graph/NodeSpawners/ControlRigEnumNodeSpawner.cpp
--- a/Engine/Plugins/Experimental/ControlRig/Source/ControlRigEditor/Private/Graph/NodeSpawners/ControlRigEnumNodeSpawner.cpp
+++ b/Engine/Plugins/Experimental/ControlRig/Source/ControlRigEditor/Private/Graph/NodeSpawners/ControlRigEnumNodeSpawner.cpp
@@ -78,7 +78,7 @@ UEdGraphNode* UControlRigEnumNodeSpawner::Invoke(UEdGraph* ParentGraph, FBinding
 
 	// First create a backing member for our node
 	UControlRigGraph* RigGraph = Cast<UControlRigGraph>(ParentGraph);
-	check(RigGraph);
+	if(RigGraph == nullptr) return nullptr;
 	UControlRigBlueprint* RigBlueprint = Cast<UControlRigBlueprint>(FBlueprintEditorUtils::FindBlueprintForGraph(ParentGraph));
 	check(RigBlueprint);
 
diff --git a/Engine/Plugins/Experimental/ControlRig/Source/ControlRigEditor/Private/Graph/NodeSpawners/ControlRigIfNodeSpawner.cpp b/Engine/Plugins/Experimental/ControlRig/Source/ControlRigEditor/Private/Graph/NodeSpawners/ControlRigIfNodeSpawner.cpp
--- a/Engine/Plugins/Experimental/ControlRig/Source/ControlRigEditor/Private/Graph/NodeSpawners/ControlRigIfNodeSpawner.cpp
+++ b/Engine/Plugins/Experimental/ControlRig/Source/ControlRigEditor/Private/Graph/NodeSpawners/ControlRigIfNodeSpawner.cpp
@@ -81,7 +81,7 @@ UEdGraphNode* UControlRigIfNodeSpawner::Invoke(UEdGraph* ParentGraph, FBindingSe
 
 	// First create a backing member for our node
 	UControlRigGraph* RigGraph = Cast<UControlRigGraph>(ParentGraph);
-	check(RigGraph);
+	if(RigGraph == nullptr) return nullptr;
 	UControlRigBlueprint* RigBlueprint = Cast<UControlRigBlueprint>(FBlueprintEditorUtils::FindBlueprintForGraph(ParentGraph));
 	check(RigBlueprint);
 
diff --git a/Engine/Plugins/Experimental/ControlRig/Source/ControlRigEditor/Private/Graph/NodeSpawners/ControlRigRerouteNodeSpawner.cpp b/Engine/Plugins/Experimental/ControlRig/Source/ControlRigEditor/Private/Graph/NodeSpawners/ControlRigRerouteNodeSpawner.cpp
--- a/Engine/Plugins/Experimental/ControlRig/Source/ControlRigEditor/Private/Graph/NodeSpawners/ControlRigRerouteNodeSpawner.cpp
+++ b/Engine/Plugins/Experimental/ControlRig/Source/ControlRigEditor/Private/Graph/NodeSpawners/ControlRigRerouteNodeSpawner.cpp
@@ -66,7 +66,7 @@ UEdGraphNode* UControlRigRerouteNodeSpawner::Invoke(UEdGraph* ParentGraph, FBind
 
 	// First create a backing member for our node
 	UControlRigGraph* RigGraph = Cast<UControlRigGraph>(ParentGraph);
-	check(RigGraph);
+	if(RigGraph == nullptr) return nullptr;
 	UControlRigBlueprint* RigBlueprint = Cast<UControlRigBlueprint>(FBlueprintEditorUtils::FindBlueprintForGraph(ParentGraph));
 	check(RigBlueprint);
 
diff --git a/Engine/Plugins/Experimental/ControlRig/Source/ControlRigEditor/Private/Graph/NodeSpawners/ControlRigSelectNodeSpawner.cpp b/Engine/Plugins/Experimental/ControlRig/Source/ControlRigEditor/Private/Graph/NodeSpawners/ControlRigSelectNodeSpawner.cpp
--- a/Engine/Plugins/Experimental/ControlRig/Source/ControlRigEditor/Private/Graph/NodeSpawners/ControlRigSelectNodeSpawner.cpp
+++ b/Engine/Plugins/Experimental/ControlRig/Source/ControlRigEditor/Private/Graph/NodeSpawners/ControlRigSelectNodeSpawner.cpp
@@ -82,7 +82,7 @@ UEdGraphNode* UControlRigSelectNodeSpawner::Invoke(UEdGraph* ParentGraph, FBindi
 
 	// First create a backing member for our node
 	UControlRigGraph* RigGraph = Cast<UControlRigGraph>(ParentGraph);
-	check(RigGraph);
+	if(RigGraph == nullptr) return nullptr;
 	UControlRigBlueprint* RigBlueprint = Cast<UControlRigBlueprint>(FBlueprintEditorUtils::FindBlueprintForGraph(ParentGraph));
 	check(RigBlueprint);
 
diff --git a/Engine/Plugins/Experimental/ControlRig/Source/ControlRigEditor/Private/Graph/NodeSpawners/ControlRigVariableNodeSpawner.cpp b/Engine/Plugins/Experimental/ControlRig/Source/ControlRigEditor/Private/Graph/NodeSpawners/ControlRigVariableNodeSpawner.cpp
--- a/Engine/Plugins/Experimental/ControlRig/Source/ControlRigEditor/Private/Graph/NodeSpawners/ControlRigVariableNodeSpawner.cpp
+++ b/Engine/Plugins/Experimental/ControlRig/Source/ControlRigEditor/Private/Graph/NodeSpawners/ControlRigVariableNodeSpawner.cpp
@@ -83,7 +83,7 @@ UEdGraphNode* UControlRigVariableNodeSpawner::Invoke(UEdGraph* ParentGraph, FBin
 
 	// First create a backing member for our node
 	UControlRigGraph* RigGraph = Cast<UControlRigGraph>(ParentGraph);
-	check(RigGraph);
+	if(RigGraph == nullptr) return nullptr;
 	UControlRigBlueprint* RigBlueprint = Cast<UControlRigBlueprint>(FBlueprintEditorUtils::FindBlueprintForGraph(ParentGraph));
 	check(RigBlueprint);
 
