Control Rig: Reverting minor change on Control Rig header to make it blueprintable

Reverting parts of 15216821

diff --git a/Engine/Plugins/Experimental/ControlRig/Source/ControlRig/Public/ControlRig.h b/Engine/Plugins/Experimental/ControlRig/Source/ControlRig/Public/ControlRig.h
--- a/Engine/Plugins/Experimental/ControlRig/Source/ControlRig/Public/ControlRig.h
+++ b/Engine/Plugins/Experimental/ControlRig/Source/ControlRig/Public/ControlRig.h
@@ -37,7 +37,7 @@ struct FRigControl;
 CONTROLRIG_API DECLARE_LOG_CATEGORY_EXTERN(LogControlRig, Log, All);
 
 /** Runs logic for mapping input data to transforms (the "Rig") */
-UCLASS(BlueprintType, Abstract, editinlinenew)
+UCLASS(Blueprintable, Abstract, editinlinenew)
 class CONTROLRIG_API UControlRig : public UObject, public INodeMappingProviderInterface, public IInterface_AssetUserData
 {
 	GENERATED_UCLASS_BODY()
