Control Rig: Search and Replace now report a message for all names that were not modified

diff --git a/Engine/Plugins/Experimental/ControlRig/Source/ControlRigEditor/Private/ControlRigEditorModule.cpp b/Engine/Plugins/Experimental/ControlRig/Source/ControlRigEditor/Private/ControlRigEditorModule.cpp
--- a/Engine/Plugins/Experimental/ControlRig/Source/ControlRigEditor/Private/ControlRigEditorModule.cpp
+++ b/Engine/Plugins/Experimental/ControlRig/Source/ControlRigEditor/Private/ControlRigEditorModule.cpp
@@ -1658,6 +1658,7 @@ void FControlRigEditorModule::GetContextMenuActions(const UControlRigGraphSchema
 
 							Controller->OpenUndoBracket(TEXT("Mirroring Graph"));
 							int32 ReplacedNames = 0;
+							TArray<FString> UnchangedItems;
 
 							for (const TPair<const URigVMPin*, FRigElementKey>& Pair : PinToKey)
 							{
@@ -1679,11 +1680,33 @@ void FControlRigEditorModule::GetContextMenuActions(const UControlRigGraphSchema
 										Controller->SetPinDefaultValue(Pin->GetPinPath(), NewNameStr, false);
 										ReplacedNames++;
 									}
+									else
+									{
+										// save the names of the items that we skipped during this search & replace
+										UnchangedItems.AddUnique(OldNameStr);
+									} 
 								}
 							}
 
-							if (ReplacedNames > 0)
+							if (UnchangedItems.Num() > 0)
 							{
+								FString ListOfUnchangedItems;
+								for (int Index = 0; Index < UnchangedItems.Num(); Index++)
+								{
+									// construct the string "item1, item2, item3"
+									ListOfUnchangedItems += UnchangedItems[Index];
+									if (Index != UnchangedItems.Num() - 1)
+									{
+										ListOfUnchangedItems += TEXT(", ");
+									}
+								}
+								
+								// inform the user that some items were skipped due to invalid new names
+								Controller->ReportAndNotifyError(FString::Printf(TEXT("Invalid Names after Search & Replace, action skipped for %s"), *ListOfUnchangedItems));
+							}
+
+							if (ReplacedNames > 0)
+							{ 
 								Controller->CloseUndoBracket();
 							}
 							else
diff --git a/Engine/Source/Developer/RigVMDeveloper/Private/RigVMModel/RigVMController.cpp b/Engine/Source/Developer/RigVMDeveloper/Private/RigVMModel/RigVMController.cpp
--- a/Engine/Source/Developer/RigVMDeveloper/Private/RigVMModel/RigVMController.cpp
+++ b/Engine/Source/Developer/RigVMDeveloper/Private/RigVMModel/RigVMController.cpp
@@ -7940,7 +7940,8 @@ void URigVMController::ReportAndNotifyError(const FString& InMessage)
 	Info.Image = FEditorStyle::GetBrush(TEXT("MessageLog.Warning"));
 	Info.bFireAndForget = true;
 	Info.bUseThrobber = true;
-	Info.FadeOutDuration = 5.0f;
+	// longer message needs more time to read
+	Info.FadeOutDuration = FMath::Clamp(0.1f * InMessage.Len(), 5.0f, 20.0f);
 	Info.ExpireDuration = Info.FadeOutDuration;
 	TSharedPtr<SNotificationItem> NotificationPtr = FSlateNotificationManager::Get().AddNotification(Info);
 	if (NotificationPtr)
