Control Rig: Added cast ufunction to cast to a Control Rig Blueprint

diff --git a/Engine/Plugins/Experimental/ControlRig/Source/ControlRigEditor/Private/ControlRigBlueprintEditorLibrary.cpp b/Engine/Plugins/Experimental/ControlRig/Source/ControlRigEditor/Private/ControlRigBlueprintEditorLibrary.cpp
--- a/Engine/Plugins/Experimental/ControlRig/Source/ControlRigEditor/Private/ControlRigBlueprintEditorLibrary.cpp
+++ b/Engine/Plugins/Experimental/ControlRig/Source/ControlRigEditor/Private/ControlRigBlueprintEditorLibrary.cpp
@@ -2,6 +2,17 @@
 
 #include "ControlRigBlueprintEditorLibrary.h"
 
+void UControlRigBlueprintEditorLibrary::CastToControlRigBlueprint(
+	UObject* Object,
+	ECastToControlRigBlueprintCases& Branches,
+	UControlRigBlueprint*& AsControlRigBlueprint)
+{
+	AsControlRigBlueprint = Cast<UControlRigBlueprint>(Object);
+	Branches = AsControlRigBlueprint == nullptr ? 
+		ECastToControlRigBlueprintCases::CastFailed : 
+		ECastToControlRigBlueprintCases::CastSucceeded;
+}
+
 void UControlRigBlueprintEditorLibrary::SetPreviewMesh(UControlRigBlueprint* InRigBlueprint, USkeletalMesh* PreviewMesh, bool bMarkAsDirty)
 {
 	if(InRigBlueprint == nullptr)
diff --git a/Engine/Plugins/Experimental/ControlRig/Source/ControlRigEditor/Public/ControlRigBlueprintEditorLibrary.h b/Engine/Plugins/Experimental/ControlRig/Source/ControlRigEditor/Public/ControlRigBlueprintEditorLibrary.h
--- a/Engine/Plugins/Experimental/ControlRig/Source/ControlRigEditor/Public/ControlRigBlueprintEditorLibrary.h
+++ b/Engine/Plugins/Experimental/ControlRig/Source/ControlRigEditor/Public/ControlRigBlueprintEditorLibrary.h
@@ -13,6 +13,13 @@
 #include "Kismet/BlueprintFunctionLibrary.h"
 #include "ControlRigBlueprintEditorLibrary.generated.h"
 
+UENUM(BlueprintType)
+enum class ECastToControlRigBlueprintCases : uint8
+{
+	CastSucceeded,
+	CastFailed
+};
+
 UCLASS(meta=(ScriptName="ControlRigBlueprintLibrary"))
 class CONTROLRIGEDITOR_API UControlRigBlueprintEditorLibrary : public UBlueprintFunctionLibrary
 {
@@ -20,6 +27,12 @@ class CONTROLRIGEDITOR_API UControlRigBlueprintEditorLibrary : public UBlueprint
 
 public:
 
+	UFUNCTION(BlueprintCallable, Category = "Control Rig Blueprint", Meta = (ExpandEnumAsExecs = "Branches"))
+	static void CastToControlRigBlueprint(
+		UObject* Object, 
+		ECastToControlRigBlueprintCases& Branches,
+		UControlRigBlueprint*& AsControlRigBlueprint);
+
 	UFUNCTION(BlueprintCallable, Category = "Control Rig Blueprint")
 	static void SetPreviewMesh(UControlRigBlueprint* InRigBlueprint, USkeletalMesh* PreviewMesh, bool bMarkAsDirty = true);
 	
@@ -38,19 +51,19 @@ public:
 	UFUNCTION(BlueprintCallable, Category = "Control Rig Blueprint")
 	static void RequestControlRigInit(UControlRigBlueprint* InRigBlueprint);
 
-	UFUNCTION(BlueprintCallable, BlueprintPure, Category = "VM")
+	UFUNCTION(BlueprintCallable, BlueprintPure, Category = "Control Rig Blueprint")
 	static URigVMGraph* GetModel(UControlRigBlueprint* InRigBlueprint);
 
-	UFUNCTION(BlueprintCallable, BlueprintPure, Category = "VM")
+	UFUNCTION(BlueprintCallable, BlueprintPure, Category = "Control Rig Blueprint")
 	static URigVMController* GetController(UControlRigBlueprint* InRigBlueprint);
 
-	UFUNCTION(BlueprintCallable, BlueprintPure, Category = "VM")
+	UFUNCTION(BlueprintCallable, BlueprintPure, Category = "Control Rig Blueprint")
 	static TArray<UControlRigBlueprint*> GetCurrentlyOpenRigBlueprints();
 
-	UFUNCTION(BlueprintCallable, BlueprintPure, Category = "VM")
+	UFUNCTION(BlueprintCallable, BlueprintPure, Category = "Control Rig Blueprint")
 	static TArray<UStruct*> GetAvailableRigUnits();
 
-	UFUNCTION(BlueprintCallable, BlueprintPure, Category = "Hierarchy")
+	UFUNCTION(BlueprintCallable, BlueprintPure, Category = "Control Rig Blueprint")
 	static UControlRigHierarchyModifier* GetHierarchyModifier(UControlRigBlueprint* InRigBlueprint);
 };
 
