Control Rig: Guard against invalid frustrums when box selecting

diff --git a/Engine/Plugins/Experimental/ControlRig/Source/ControlRigEditor/Private/EditMode/ControlRigEditMode.cpp b/Engine/Plugins/Experimental/ControlRig/Source/ControlRigEditor/Private/EditMode/ControlRigEditMode.cpp
--- a/Engine/Plugins/Experimental/ControlRig/Source/ControlRigEditor/Private/EditMode/ControlRigEditMode.cpp
+++ b/Engine/Plugins/Experimental/ControlRig/Source/ControlRigEditor/Private/EditMode/ControlRigEditMode.cpp
@@ -981,6 +981,12 @@ bool FControlRigEditMode::FrustumSelect(const FConvexVolume& InFrustum, FEditorV
 	// Find intersection points and project to screen space to determine the bounding rectangle of the selection box
 	for (int32 PlaneIndex = 0; PlaneIndex < 4; ++PlaneIndex)
 	{
+		if (!InFrustum.Planes.IsValidIndex(PlaneIndex) ||
+			!InFrustum.Planes.IsValidIndex((PlaneIndex + 1) % 4))
+		{
+			break;
+		}
+
 		const FPlane& Plane1 = InFrustum.Planes[PlaneIndex];
 		const FPlane& Plane2 = InFrustum.Planes[(PlaneIndex + 1) % 4];
 		FVector I, D;
