Control Rig: Support for copy & pasting functions

diff --git a/Engine/Plugins/Experimental/ControlRig/Source/ControlRigDeveloper/Private/ControlRigBlueprint.cpp b/Engine/Plugins/Experimental/ControlRig/Source/ControlRigDeveloper/Private/ControlRigBlueprint.cpp
--- a/Engine/Plugins/Experimental/ControlRig/Source/ControlRigDeveloper/Private/ControlRigBlueprint.cpp
+++ b/Engine/Plugins/Experimental/ControlRig/Source/ControlRigDeveloper/Private/ControlRigBlueprint.cpp
@@ -132,6 +132,67 @@ void UControlRigBlueprint::LoadModulesRequiredForCompilation()
 {
 }
 
+bool UControlRigBlueprint::ExportGraphToText(UEdGraph* InEdGraph, FString& OutText)
+{
+	OutText.Empty();
+
+	if (URigVMGraph* RigGraph = GetModel(InEdGraph))
+	{
+		if (URigVMCollapseNode* CollapseNode = Cast<URigVMCollapseNode>(RigGraph->GetOuter()))
+		{
+			if (URigVMController* Controller = GetOrCreateController(CollapseNode->GetGraph()))
+			{
+				TArray<FName> NodeNamesToExport;
+				NodeNamesToExport.Add(CollapseNode->GetFName());
+				OutText = Controller->ExportNodesToText(NodeNamesToExport);
+			}
+		}
+	}
+
+	// always return true so that the default mechanism doesn't take over
+	return true;
+}
+
+bool UControlRigBlueprint::CanImportGraphFromText(const FString& InClipboardText)
+{
+	return GetTemplateController()->CanImportNodesFromText(InClipboardText);
+}
+
+bool UControlRigBlueprint::TryImportGraphFromText(const FString& InClipboardText, UEdGraph** OutGraphPtr)
+{
+	if (OutGraphPtr)
+	{
+		*OutGraphPtr = nullptr;
+	}
+
+	if (URigVMController* FunctionLibraryController = GetOrCreateController(GetLocalFunctionLibrary()))
+	{
+		TArray<FName> ImportedNodeNames = FunctionLibraryController->ImportNodesFromText(InClipboardText, true);
+		if (ImportedNodeNames.Num() == 0)
+		{
+			return false;
+		}
+
+		URigVMCollapseNode* CollapseNode = Cast<URigVMCollapseNode>(GetLocalFunctionLibrary()->FindFunction(ImportedNodeNames[0]));
+		if (ImportedNodeNames.Num() > 1 || CollapseNode == nullptr || CollapseNode->GetContainedGraph() == nullptr)
+		{
+			FunctionLibraryController->Undo();
+			return false;
+		}
+
+		UEdGraph* EdGraph = GetEdGraph(CollapseNode->GetContainedGraph());
+		if (OutGraphPtr)
+		{
+			*OutGraphPtr = EdGraph;
+		}
+
+		BroadcastGraphImported(EdGraph);
+	}
+
+	// always return true so that the default mechanism doesn't take over
+	return true;
+}
+
 USkeletalMesh* UControlRigBlueprint::GetPreviewMesh() const
 {
 	DECLARE_SCOPE_HIERARCHICAL_COUNTER_FUNC()
@@ -2585,6 +2646,11 @@ void UControlRigBlueprint::BroadcastNodeDoubleClicked(URigVMNode* InNode)
 	NodeDoubleClickedEvent.Broadcast(this, InNode);
 }
 
+void UControlRigBlueprint::BroadcastGraphImported(UEdGraph* InGraph)
+{
+	GraphImportedEvent.Broadcast(InGraph);
+}
+
 #endif
 
 void UControlRigBlueprint::CreateEdGraphForCollapseNodeIfNeeded(URigVMCollapseNode* InNode, bool bForce)
diff --git a/Engine/Plugins/Experimental/ControlRig/Source/ControlRigDeveloper/Public/ControlRigBlueprint.h b/Engine/Plugins/Experimental/ControlRig/Source/ControlRigDeveloper/Public/ControlRigBlueprint.h
--- a/Engine/Plugins/Experimental/ControlRig/Source/ControlRigDeveloper/Public/ControlRigBlueprint.h
+++ b/Engine/Plugins/Experimental/ControlRig/Source/ControlRigDeveloper/Public/ControlRigBlueprint.h
@@ -35,6 +35,7 @@ DECLARE_EVENT_OneParam(UControlRigBlueprint, FOnRefreshEditorEvent, UControlRigB
 DECLARE_EVENT_FourParams(UControlRigBlueprint, FOnVariableDroppedEvent, UObject*, FProperty*, const FVector2D&, const FVector2D&);
 DECLARE_EVENT_OneParam(UControlRigBlueprint, FOnExternalVariablesChanged, const TArray<FRigVMExternalVariable>&);
 DECLARE_EVENT_TwoParams(UControlRigBlueprint, FOnNodeDoubleClicked, UControlRigBlueprint*, URigVMNode*);
+DECLARE_EVENT_OneParam(UControlRigBlueprint, FOnGraphImported, UEdGraph*);
 
 UCLASS(BlueprintType, meta=(IgnoreClassThumbnail))
 class CONTROLRIGDEVELOPER_API UControlRigBlueprint : public UBlueprint, public IInterface_PreviewMeshProvider
@@ -74,7 +75,9 @@ public:
 	virtual bool SupportsDelegates() const override { return false; }
 	virtual bool SupportsEventGraphs() const override { return false; }
 	virtual bool SupportsAnimLayers() const override { return false; }
-
+	virtual bool ExportGraphToText(UEdGraph* InEdGraph, FString& OutText) override;
+	virtual bool TryImportGraphFromText(const FString& InClipboardText, UEdGraph** OutGraphPtr = nullptr) override;
+	virtual bool CanImportGraphFromText(const FString& InClipboardText) override;
 
 #endif	// #if WITH_EDITOR
 
@@ -372,12 +375,16 @@ public:
 	FOnNodeDoubleClicked& OnNodeDoubleClicked() { return NodeDoubleClickedEvent; }
 	void BroadcastNodeDoubleClicked(URigVMNode* InNode);
 
+	FOnGraphImported& OnGraphImported() { return GraphImportedEvent; }
+	void BroadcastGraphImported(UEdGraph* InGraph);
+
 private:
 
 	FOnExternalVariablesChanged ExternalVariablesChangedEvent;
 	void BroadcastExternalVariablesChangedEvent();
 
 	FOnNodeDoubleClicked NodeDoubleClickedEvent;
+	FOnGraphImported GraphImportedEvent;
 
 #endif
 
diff --git a/Engine/Plugins/Experimental/ControlRig/Source/ControlRigEditor/Private/Editor/ControlRigEditor.cpp b/Engine/Plugins/Experimental/ControlRig/Source/ControlRigEditor/Private/Editor/ControlRigEditor.cpp
--- a/Engine/Plugins/Experimental/ControlRig/Source/ControlRigEditor/Private/Editor/ControlRigEditor.cpp
+++ b/Engine/Plugins/Experimental/ControlRig/Source/ControlRigEditor/Private/Editor/ControlRigEditor.cpp
@@ -128,6 +128,7 @@ FControlRigEditor::~FControlRigEditor()
 		RigBlueprint->OnRefreshEditor().RemoveAll(this);
 		RigBlueprint->OnVariableDropped().RemoveAll(this);
 		RigBlueprint->OnNodeDoubleClicked().RemoveAll(this);
+		RigBlueprint->OnGraphImported().RemoveAll(this);
 	}
 
 	if (NodeDetailBuffer.Num() > 0 && NodeDetailStruct != nullptr)
@@ -363,6 +364,7 @@ void FControlRigEditor::InitControlRigEditor(const EToolkitMode::Type Mode, cons
 		}
 
 		InControlRigBlueprint->OnNodeDoubleClicked().AddSP(this, &FControlRigEditor::OnNodeDoubleClicked);
+		InControlRigBlueprint->OnGraphImported().AddSP(this, &FControlRigEditor::OnGraphImported);
 	}
 
 	UpdateStaleWatchedPins();
@@ -4181,6 +4183,14 @@ void FControlRigEditor::OnNodeDoubleClicked(UControlRigBlueprint* InBlueprint, U
 	}
 }
 
+void FControlRigEditor::OnGraphImported(UEdGraph* InEdGraph)
+{
+	check(InEdGraph);
+
+	OpenDocument(InEdGraph, FDocumentTracker::OpenNewDocument);
+	RenameNewlyAddedAction(InEdGraph->GetFName());
+}
+
 bool FControlRigEditor::OnActionMatchesName(FEdGraphSchemaAction* InAction, const FName& InName) const
 {
 	if (InAction->GetMenuDescription().ToString() == InName.ToString())
diff --git a/Engine/Plugins/Experimental/ControlRig/Source/ControlRigEditor/Private/Editor/ControlRigEditor.h b/Engine/Plugins/Experimental/ControlRig/Source/ControlRigEditor/Private/Editor/ControlRigEditor.h
--- a/Engine/Plugins/Experimental/ControlRig/Source/ControlRigEditor/Private/Editor/ControlRigEditor.h
+++ b/Engine/Plugins/Experimental/ControlRig/Source/ControlRigEditor/Private/Editor/ControlRigEditor.h
@@ -307,6 +307,8 @@ private:
 
 	 void OnNodeDoubleClicked(UControlRigBlueprint* InBlueprint, URigVMNode* InNode);
 
+	 void OnGraphImported(UEdGraph* InEdGraph);
+
 	 virtual bool OnActionMatchesName(FEdGraphSchemaAction* InAction, const FName& InName) const override;
 
 protected:
diff --git a/Engine/Source/Editor/Kismet/Private/SMyBlueprint.cpp b/Engine/Source/Editor/Kismet/Private/SMyBlueprint.cpp
--- a/Engine/Source/Editor/Kismet/Private/SMyBlueprint.cpp
+++ b/Engine/Source/Editor/Kismet/Private/SMyBlueprint.cpp
@@ -3227,9 +3227,12 @@ void SMyBlueprint::OnCopy()
 	}
 	else if (FEdGraphSchemaAction_K2Graph* GraphAction = SelectionAsGraph())
 	{
-		FBPGraphClipboardData FuncData(GraphAction->EdGraph);
-		FBPGraphClipboardData::StaticStruct()->ExportText(OutputString, &FuncData, &FuncData, nullptr, 0, nullptr, false);
-		OutputString = GRAPH_PREFIX + OutputString;
+		if (!Blueprint->ExportGraphToText(GraphAction->EdGraph, OutputString))
+		{
+			FBPGraphClipboardData FuncData(GraphAction->EdGraph);
+			FBPGraphClipboardData::StaticStruct()->ExportText(OutputString, &FuncData, &FuncData, nullptr, 0, nullptr, false);
+			OutputString = GRAPH_PREFIX + OutputString;
+		}
 	}
 
 	if (!OutputString.IsEmpty())
@@ -3429,6 +3432,12 @@ void SMyBlueprint::OnPasteFunction()
 {
 	FString ClipboardText;
 	FPlatformApplicationMisc::ClipboardPaste(ClipboardText);
+
+	if (Blueprint->TryImportGraphFromText(ClipboardText))
+	{
+		return;
+	}
+
 	if (!ensure(ClipboardText.StartsWith(GRAPH_PREFIX, ESearchCase::CaseSensitive)))
 	{
 		return;
@@ -3473,6 +3482,12 @@ bool SMyBlueprint::CanPasteFunction() const
 
 	FString ClipboardText;
 	FPlatformApplicationMisc::ClipboardPaste(ClipboardText);
+
+	if (Blueprint->CanImportGraphFromText(ClipboardText))
+	{
+		return true;
+	}
+
 	if (ClipboardText.StartsWith(GRAPH_PREFIX, ESearchCase::CaseSensitive))
 	{
 		FBPGraphClipboardData FuncData;
diff --git a/Engine/Source/Runtime/Engine/Classes/Engine/Blueprint.h b/Engine/Source/Runtime/Engine/Classes/Engine/Blueprint.h
--- a/Engine/Source/Runtime/Engine/Classes/Engine/Blueprint.h
+++ b/Engine/Source/Runtime/Engine/Classes/Engine/Blueprint.h
@@ -1104,6 +1104,23 @@ public:
 	 */
 	virtual bool SupportsAnimLayers() const { return true; }
 
+	/**
+	 * Copies a given graph into a text buffer. Returns true if successful.
+	 */
+	virtual bool ExportGraphToText(UEdGraph* InEdGraph, FString& OutText) { return false; }
+
+	/**
+	 * Returns true if the blueprint can import a given InClipboardText.
+	 * If this return false the default BP functionality will be used.
+	 */
+	virtual bool CanImportGraphFromText(const FString& InClipboardText) { return false; }
+
+	/**
+	 * Returns a new ed graph if the blueprint's specialization imported a graph based on 
+	 * a clipboard text content an nullptr if that's not successful.
+	 */
+	virtual bool TryImportGraphFromText(const FString& InClipboardText, UEdGraph** OutGraphPtr = nullptr) { return false; }
+
 #endif
 };
 
