Control Rig: Disable the ability to close the RigHierarchyTab

diff --git a/Engine/Plugins/Experimental/ControlRig/Source/ControlRigEditor/Private/Editor/RigHierarchyTabSummoner.cpp b/Engine/Plugins/Experimental/ControlRig/Source/ControlRigEditor/Private/Editor/RigHierarchyTabSummoner.cpp
--- a/Engine/Plugins/Experimental/ControlRig/Source/ControlRigEditor/Private/Editor/RigHierarchyTabSummoner.cpp
+++ b/Engine/Plugins/Experimental/ControlRig/Source/ControlRigEditor/Private/Editor/RigHierarchyTabSummoner.cpp
@@ -38,4 +38,14 @@ TSharedRef<SWidget> FRigHierarchyTabSummoner::CreateTabBody(const FWorkflowTabSp
 	return SNew(SRigHierarchy, ControlRigEditor.Pin().ToSharedRef());
 }
 
+TSharedRef<SDockTab> FRigHierarchyTabSummoner::SpawnTab(const FWorkflowTabSpawnInfo& Info) const
+{
+	TSharedRef<SDockTab>  DockTab = FWorkflowTabFactory::SpawnTab(Info);
+	DockTab->SetCanCloseTab(SDockTab::FCanCloseTab::CreateLambda([]()
+    {
+        return false;
+    }));
+	return DockTab;
+}
+
 #undef LOCTEXT_NAMESPACE 
diff --git a/Engine/Plugins/Experimental/ControlRig/Source/ControlRigEditor/Private/Editor/RigHierarchyTabSummoner.h b/Engine/Plugins/Experimental/ControlRig/Source/ControlRigEditor/Private/Editor/RigHierarchyTabSummoner.h
--- a/Engine/Plugins/Experimental/ControlRig/Source/ControlRigEditor/Private/Editor/RigHierarchyTabSummoner.h
+++ b/Engine/Plugins/Experimental/ControlRig/Source/ControlRigEditor/Private/Editor/RigHierarchyTabSummoner.h
@@ -17,6 +17,7 @@ public:
 	
 	virtual FTabSpawnerEntry& RegisterTabSpawner(TSharedRef<FTabManager> TabManager, const FApplicationMode* CurrentApplicationMode) const;
 	virtual TSharedRef<SWidget> CreateTabBody(const FWorkflowTabSpawnInfo& Info) const override;
+	virtual TSharedRef<SDockTab> SpawnTab(const FWorkflowTabSpawnInfo& Info) const override;
 
 protected:
 	TWeakPtr<FControlRigEditor> ControlRigEditor;
