Control Rig: Support to rename / search for key names also for function nodes

diff --git a/Engine/Plugins/Experimental/ControlRig/Source/ControlRigEditor/Private/ControlRigEditorModule.cpp b/Engine/Plugins/Experimental/ControlRig/Source/ControlRigEditor/Private/ControlRigEditorModule.cpp
--- a/Engine/Plugins/Experimental/ControlRig/Source/ControlRigEditor/Private/ControlRigEditorModule.cpp
+++ b/Engine/Plugins/Experimental/ControlRig/Source/ControlRigEditor/Private/ControlRigEditorModule.cpp
@@ -1578,39 +1578,66 @@ void FControlRigEditorModule::GetContextMenuActions(const UControlRigGraphSchema
 								RigElementsToSelect.AddUnique(Key);
 								PinToKey.Add(Pin, Key);
 							}
-							else if (Pin->GetCPPTypeObject() == FRigElementKey::StaticStruct() && StructMemory != nullptr)
+							else if (Pin->GetCPPTypeObject() == FRigElementKey::StaticStruct())
 							{
-								check(ScriptStruct);
-								if (const FProperty* Property = ScriptStruct->FindPropertyByName(Pin->GetFName()))
+								if (StructMemory == nullptr)
 								{
-									const FRigElementKey& Key = *Property->ContainerPtrToValuePtr<FRigElementKey>(StructMemory);
-
-									if (Key.IsValid())
+									FString DefaultValue = Pin->GetDefaultValue();
+									if (!DefaultValue.IsEmpty())
 									{
-										RigElementsToSelect.AddUnique(Key);
+										FRigElementKey Key;
+										FRigElementKey::StaticStruct()->ImportText(*DefaultValue, &Key, nullptr, EPropertyPortFlags::PPF_None, nullptr, FRigElementKey::StaticStruct()->GetName(), true);
+										if (Key.IsValid())
+										{
+											RigElementsToSelect.AddUnique(Key);
+											if (URigVMPin* NamePin = Pin->FindSubPin(TEXT("Name")))
+											{
+												PinToKey.Add(NamePin, Key);
+											}
+										}
+									}
+								}
+								else
+								{
+									check(ScriptStruct);
+									if (const FProperty* Property = ScriptStruct->FindPropertyByName(Pin->GetFName()))
+									{
+										const FRigElementKey& Key = *Property->ContainerPtrToValuePtr<FRigElementKey>(StructMemory);
 
-										if (URigVMPin* NamePin = Pin->FindSubPin(TEXT("Name")))
+										if (Key.IsValid())
 										{
-											PinToKey.Add(NamePin, Key);
+											RigElementsToSelect.AddUnique(Key);
+
+											if (URigVMPin* NamePin = Pin->FindSubPin(TEXT("Name")))
+											{
+												PinToKey.Add(NamePin, Key);
+											}
 										}
 									}
 								}
 							}
 							else if (Pin->GetCPPTypeObject() == FRigElementKeyCollection::StaticStruct() && Pin->GetDirection() == ERigVMPinDirection::Output && StructMemory != nullptr)
 							{
-								check(ScriptStruct);
-								if (const FProperty* Property = ScriptStruct->FindPropertyByName(Pin->GetFName()))
+								if (StructMemory == nullptr)
 								{
-									const FRigElementKeyCollection& Collection = *Property->ContainerPtrToValuePtr<FRigElementKeyCollection>(StructMemory);
-
-									if (Collection.Num() > 0)
+									// not supported for now
+								}
+								else
+								{
+									check(ScriptStruct);
+									if (const FProperty* Property = ScriptStruct->FindPropertyByName(Pin->GetFName()))
 									{
-										RigElementsToSelect.Reset();
-										for (const FRigElementKey& Item : Collection)
+										const FRigElementKeyCollection& Collection = *Property->ContainerPtrToValuePtr<FRigElementKeyCollection>(StructMemory);
+
+										if (Collection.Num() > 0)
 										{
-											RigElementsToSelect.AddUnique(Item);
+											RigElementsToSelect.Reset();
+											for (const FRigElementKey& Item : Collection)
+											{
+												RigElementsToSelect.AddUnique(Item);
+											}
+											break;
 										}
-										break;
 									}
 								}
 							}
