Control Rig: Added new menu action, Set Gizmo Transform From Current

diff --git a/Engine/Plugins/Experimental/ControlRig/Source/ControlRigEditor/Private/ControlRigHierarchyCommands.cpp b/Engine/Plugins/Experimental/ControlRig/Source/ControlRigEditor/Private/ControlRigHierarchyCommands.cpp
--- a/Engine/Plugins/Experimental/ControlRig/Source/ControlRigEditor/Private/ControlRigHierarchyCommands.cpp
+++ b/Engine/Plugins/Experimental/ControlRig/Source/ControlRigEditor/Private/ControlRigHierarchyCommands.cpp
@@ -19,6 +19,7 @@ void FControlRigHierarchyCommands::RegisterCommands()
 	UI_COMMAND(PasteGlobalTransforms, "Paste Global Transform", "Paste the global transforms.", EUserInterfaceActionType::Button, FInputChord(EKeys::V, EModifierKey::Control | EModifierKey::Shift));
 	UI_COMMAND(SetInitialTransformFromClosestBone, "Set Initial Transform from Closest Bone", "Find the Closest Bone to Initial Transform of the Selected Bones/Spaces or To Offset Transform of Selected Controls", EUserInterfaceActionType::Button, FInputChord());
 	UI_COMMAND(SetInitialTransformFromCurrentTransform, "Set Initial Transform from Current", "Save the Current Transform To Initial Transform of the Selected Bones/Spaces or To Offset Transform of Selected Controls", EUserInterfaceActionType::Button, FInputChord());
+	UI_COMMAND(SetGizmoTransformFromCurrent, "Set Gizmo Transform From Current", "Transfer the Current Local Transform of the Control to its Gizmo Transform, the Control's Transform Resets", EUserInterfaceActionType::Button, FInputChord());
 	UI_COMMAND(ResetTransform, "Reset Transform", "Reset the Transform", EUserInterfaceActionType::Button, FInputChord(EKeys::G, EModifierKey::Control));
 	UI_COMMAND(ResetAllTransforms, "Reset All Transforms", "Resets all Transforms", EUserInterfaceActionType::Button, FInputChord(EKeys::G, EModifierKey::Control | EModifierKey::Shift));
 	UI_COMMAND(ResetSpace, "Reset Space", "Resets or injects a Space below the Control", EUserInterfaceActionType::Button, FInputChord());
diff --git a/Engine/Plugins/Experimental/ControlRig/Source/ControlRigEditor/Private/ControlRigHierarchyCommands.h b/Engine/Plugins/Experimental/ControlRig/Source/ControlRigEditor/Private/ControlRigHierarchyCommands.h
--- a/Engine/Plugins/Experimental/ControlRig/Source/ControlRigEditor/Private/ControlRigHierarchyCommands.h
+++ b/Engine/Plugins/Experimental/ControlRig/Source/ControlRigEditor/Private/ControlRigHierarchyCommands.h
@@ -66,6 +66,9 @@ public:
 	/* set initial transform from closest bone */
 	TSharedPtr<FUICommandInfo> SetInitialTransformFromClosestBone;
 
+	/* set gizmo transform from the control's current local transform and reset the control's transform */
+	TSharedPtr<FUICommandInfo> SetGizmoTransformFromCurrent;
+
 	/* frames the selection in the tree */
 	TSharedPtr<FUICommandInfo> FrameSelection;
 
diff --git a/Engine/Plugins/Experimental/ControlRig/Source/ControlRigEditor/Private/Editor/SRigHierarchy.cpp b/Engine/Plugins/Experimental/ControlRig/Source/ControlRigEditor/Private/Editor/SRigHierarchy.cpp
--- a/Engine/Plugins/Experimental/ControlRig/Source/ControlRigEditor/Private/Editor/SRigHierarchy.cpp
+++ b/Engine/Plugins/Experimental/ControlRig/Source/ControlRigEditor/Private/Editor/SRigHierarchy.cpp
@@ -506,6 +506,11 @@ void SRigHierarchy::BindCommands()
 		FExecuteAction::CreateSP(this, &SRigHierarchy::HandleSetInitialTransformFromCurrentTransform),
 		FCanExecuteAction::CreateSP(this, &SRigHierarchy::IsMultiSelected));
 
+	CommandList->MapAction(
+		Commands.SetGizmoTransformFromCurrent,
+		FExecuteAction::CreateSP(this, &SRigHierarchy::HandleSetGizmoTransformFromCurrent),
+		FCanExecuteAction::CreateSP(this, &SRigHierarchy::IsControlSelected));
+
 	CommandList->MapAction(
 		Commands.FrameSelection,
 		FExecuteAction::CreateSP(this, &SRigHierarchy::HandleFrameSelection),
@@ -1298,6 +1303,7 @@ void SRigHierarchy::FillContextMenu(class FMenuBuilder& MenuBuilder)
 		MenuBuilder.AddMenuEntry(Actions.ResetAllTransforms);
 		MenuBuilder.AddMenuEntry(Actions.SetInitialTransformFromCurrentTransform);
 		MenuBuilder.AddMenuEntry(Actions.SetInitialTransformFromClosestBone);
+		MenuBuilder.AddMenuEntry(Actions.SetGizmoTransformFromCurrent);
 		MenuBuilder.AddMenuEntry(Actions.Unparent);
 		MenuBuilder.EndSection();
 
@@ -2922,4 +2928,48 @@ void SRigHierarchy::HandleSetInitialTransformFromClosestBone()
 	}
 }
 
+void SRigHierarchy::HandleSetGizmoTransformFromCurrent()
+{
+	if (IsControlSelected())
+	{
+		if (UControlRigBlueprint* Blueprint = ControlRigEditor.Pin()->GetControlRigBlueprint())
+		{
+			if (FRigHierarchyContainer* DebuggedContainer = GetDebuggedHierarchyContainer())
+			{
+				FScopedTransaction Transaction(LOCTEXT("HierarchySetInitialTransforms", "Set Initial Transforms"));
+				Blueprint->Modify();
+
+				TArray<TSharedPtr<FRigTreeElement>> SelectedItems = TreeView->GetSelectedItems();
+				for (TSharedPtr<FRigTreeElement> SelectedItem : SelectedItems)
+				{
+					FRigElementKey Key = SelectedItem->Key;
+					if (Key.Type == ERigElementType::Control)
+					{
+						FRigControl& Control = GetHierarchyContainer()->ControlHierarchy[Key.Name];
+						FRigControl& DebuggedControl = DebuggedContainer->ControlHierarchy[Key.Name];
+						
+						if (Control.ControlType == ERigControlType::Transform
+							|| Control.ControlType == ERigControlType::TransformNoScale
+							|| Control.ControlType == ERigControlType::EulerTransform)
+						{
+							if (Control.bGizmoEnabled)
+							{
+								Control.GizmoTransform = Control.GetTransformFromValue(ERigControlValueType::Current);
+								GetHierarchyContainer()->SetLocalTransform(Key, FTransform::Identity);
+								DebuggedControl.GizmoTransform = DebuggedControl.GetTransformFromValue(ERigControlValueType::Current);
+								DebuggedContainer->SetLocalTransform(Key, FTransform::Identity);
+							}
+						}
+
+						if (FControlRigEditorEditMode* EditMode = ControlRigEditor.Pin()->GetEditMode())
+						{
+							EditMode->RequestToRecreateGizmoActors();
+						}
+					}
+				}
+			}
+		}
+	}
+}
+
 #undef LOCTEXT_NAMESPACE
\ No newline at end of file
diff --git a/Engine/Plugins/Experimental/ControlRig/Source/ControlRigEditor/Private/Editor/SRigHierarchy.h b/Engine/Plugins/Experimental/ControlRig/Source/ControlRigEditor/Private/Editor/SRigHierarchy.h
--- a/Engine/Plugins/Experimental/ControlRig/Source/ControlRigEditor/Private/Editor/SRigHierarchy.h
+++ b/Engine/Plugins/Experimental/ControlRig/Source/ControlRigEditor/Private/Editor/SRigHierarchy.h
@@ -299,6 +299,7 @@ private:
 	void HandleResetSpace();
 	void HandleSetInitialTransformFromCurrentTransform();
 	void HandleSetInitialTransformFromClosestBone();
+	void HandleSetGizmoTransformFromCurrent();
 	void HandleFrameSelection();
 	void HandleControlBoneOrSpaceTransform();
 	void HandleUnparent();
