Control Rig: Making UAnimGraphNode_ControlRig public

diff --git a/Engine/Plugins/Experimental/ControlRig/Source/ControlRigDeveloper/Private/AnimGraphNode_ControlRig.h b/Engine/Plugins/Experimental/ControlRig/Source/ControlRigDeveloper/Private/AnimGraphNode_ControlRig.h
deleted file mode 100644
diff --git a/Engine/Plugins/Experimental/ControlRig/Source/ControlRigDeveloper/Public/AnimGraphNode_ControlRig.h b/Engine/Plugins/Experimental/ControlRig/Source/ControlRigDeveloper/Public/AnimGraphNode_ControlRig.h
new file mode 100644
--- /dev/null
+++ b/Engine/Plugins/Experimental/ControlRig/Source/ControlRigDeveloper/Public/AnimGraphNode_ControlRig.h
@@ -0,0 +1,60 @@
+// Copyright Epic Games, Inc. All Rights Reserved.
+
+#pragma once
+
+#include "AnimGraphNode_CustomProperty.h"
+#include "AnimNode_ControlRig.h"
+#include "SSearchableComboBox.h"
+#include "AnimGraphNode_ControlRig.generated.h"
+
+struct FVariableMappingInfo;
+
+UCLASS(MinimalAPI)
+class UAnimGraphNode_ControlRig : public UAnimGraphNode_CustomProperty
+{
+	GENERATED_UCLASS_BODY()
+
+	UPROPERTY(EditAnywhere, Category = Settings)
+	FAnimNode_ControlRig Node;
+
+	virtual void CustomizeDetails(IDetailLayoutBuilder& DetailBuilder) override;
+	virtual void PostEditChangeProperty(FPropertyChangedEvent& PropertyChangedEvent) override;
+
+private:
+	// UEdGraphNode interface
+	virtual FText GetNodeTitle(ENodeTitleType::Type TitleType) const override;
+	virtual void CustomizePinData(UEdGraphPin* Pin, FName SourcePropertyName, int32 ArrayIndex) const override;
+	virtual FText GetTooltipText() const override;
+
+	virtual FAnimNode_CustomProperty* GetCustomPropertyNode() override { return &Node; }
+	virtual const FAnimNode_CustomProperty* GetCustomPropertyNode() const override { return &Node; }
+
+	// property related things
+	void GetVariables(bool bInput, TMap<FName, FRigVMExternalVariable>& OutParameters) const;
+
+	TMap<FName, FRigVMExternalVariable> InputVariables;
+	TMap<FName, FRigVMExternalVariable> OutputVariables;
+
+	// we have to override both of it
+	// Rebuild is about rebuilding internal data structre
+	// Getter is about getting only properties, so that it can reconstruct node
+	virtual void GetExposableProperties(TArray<FProperty*>& OutExposableProperties) const override {}
+	virtual void RebuildExposedProperties() override;
+	virtual void ReallocatePinsDuringReconstruction(TArray<UEdGraphPin*>& OldPins) override;
+	virtual void ValidateAnimNodeDuringCompilation(USkeleton* ForSkeleton, FCompilerResultsLog& MessageLog) override;
+
+	// pin option related
+	bool IsPropertyExposeEnabled(FName PropertyName) const;
+	virtual ECheckBoxState IsPropertyExposed(FName PropertyName) const override;
+	virtual void OnPropertyExposeCheckboxChanged(ECheckBoxState NewState, FName PropertyName) override;
+
+	// SVariableMappingWidget related
+	void OnVariableMappingChanged(const FName& PathName, const FName& Curve, bool bInput);
+	FName GetVariableMapping(const FName& PathName, bool bInput);
+	void GetAvailableMapping(const FName& PathName, TArray<FName>& OutArray, bool bInput);
+	void CreateVariableMapping(const FString& FilteredText, TArray< TSharedPtr<FVariableMappingInfo> >& OutArray, bool bInput);
+
+	bool IsAvailableToMapToCurve(const FName& PropertyName, bool bInput) const;
+	bool IsInputProperty(const FName& PropertyName) const;
+};
+
