Control Rig: Ensure to maintain offsets when unparenting controls

diff --git a/Engine/Plugins/Experimental/ControlRig/Source/ControlRigEditor/Private/Editor/SRigHierarchy.cpp b/Engine/Plugins/Experimental/ControlRig/Source/ControlRigEditor/Private/Editor/SRigHierarchy.cpp
--- a/Engine/Plugins/Experimental/ControlRig/Source/ControlRigEditor/Private/Editor/SRigHierarchy.cpp
+++ b/Engine/Plugins/Experimental/ControlRig/Source/ControlRigEditor/Private/Editor/SRigHierarchy.cpp
@@ -699,6 +699,19 @@ void SRigHierarchy::RefreshTreeView()
 	}
 }
 
+TArray<FRigElementKey> SRigHierarchy::GetSelectedKeys() const
+{
+	TArray<TSharedPtr<FRigTreeElement>> SelectedItems = TreeView->GetSelectedItems();
+	
+	TArray<FRigElementKey> SelectedKeys;
+	for (const TSharedPtr<FRigTreeElement>& SelectedItem : SelectedItems)
+	{
+		SelectedKeys.AddUnique(SelectedItem->Key);
+	}
+
+	return SelectedKeys;
+}
+
 void SRigHierarchy::SetExpansionRecursive(TSharedPtr<FRigTreeElement> InElement, bool bTowardsParent, bool bShouldBeExpanded)
 {
 	TreeView->SetItemExpansion(InElement, bShouldBeExpanded);
@@ -745,13 +758,7 @@ void SRigHierarchy::OnSelectionChanged(TSharedPtr<FRigTreeElement> Selection, ES
 		TGuardValue<bool> GuardRigHierarchyChanges(bIsChangingRigHierarchy, true);
 
 		TArray<FRigElementKey> OldSelection = Hierarchy->CurrentSelection();
-		TArray<FRigElementKey> NewSelection;
-
-		TArray<TSharedPtr<FRigTreeElement>> SelectedItems = TreeView->GetSelectedItems();
-		for (const TSharedPtr<FRigTreeElement>& SelectedItem : SelectedItems)
-		{
-			NewSelection.Add(SelectedItem->Key);
-		}
+		TArray<FRigElementKey> NewSelection = GetSelectedKeys();
 
 		for (const FRigElementKey& PreviouslySelected : OldSelection)
 		{
@@ -1473,12 +1480,12 @@ void SRigHierarchy::ImportHierarchy(const FAssetData& InAssetData)
 
 bool SRigHierarchy::IsMultiSelected() const
 {
-	return TreeView->GetNumItemsSelected() > 0;
+	return GetSelectedKeys().Num() > 0;
 }
 
 bool SRigHierarchy::IsSingleSelected() const
 {
-	return TreeView->GetNumItemsSelected() == 1;
+	return GetSelectedKeys().Num() == 1;
 }
 
 bool SRigHierarchy::IsSingleBoneSelected() const
@@ -1487,7 +1494,7 @@ bool SRigHierarchy::IsSingleBoneSelected() const
 	{
 		return false;
 	}
-	return TreeView->GetSelectedItems()[0]->Key.Type == ERigElementType::Bone;
+	return GetSelectedKeys()[0].Type == ERigElementType::Bone;
 }
 
 bool SRigHierarchy::IsSingleSpaceSelected() const
@@ -1496,15 +1503,15 @@ bool SRigHierarchy::IsSingleSpaceSelected() const
 	{
 		return false;
 	}
-	return TreeView->GetSelectedItems()[0]->Key.Type == ERigElementType::Space;
+	return GetSelectedKeys()[0].Type == ERigElementType::Space;
 }
 
 bool SRigHierarchy::IsControlSelected() const
 {
-	TArray<TSharedPtr<FRigTreeElement>> SelectedItems = TreeView->GetSelectedItems();
-	for (TSharedPtr<FRigTreeElement> SelectedItem : SelectedItems)
+	TArray<FRigElementKey> SelectedKeys = GetSelectedKeys();
+	for (const FRigElementKey& SelectedKey : SelectedKeys)
 	{
-		if (SelectedItem->Key.Type == ERigElementType::Control)
+		if (SelectedKey.Type == ERigElementType::Control)
 		{
 			return true;
 		}
@@ -1514,14 +1521,14 @@ bool SRigHierarchy::IsControlSelected() const
 
 bool SRigHierarchy::IsControlOrSpaceSelected() const
 {
-	TArray<TSharedPtr<FRigTreeElement>> SelectedItems = TreeView->GetSelectedItems();
-	for (TSharedPtr<FRigTreeElement> SelectedItem : SelectedItems)
+	TArray<FRigElementKey> SelectedKeys = GetSelectedKeys();
+	for (const FRigElementKey& SelectedKey : SelectedKeys)
 	{
-		if (SelectedItem->Key.Type == ERigElementType::Control)
+		if (SelectedKey.Type == ERigElementType::Control)
 		{
 			return true;
 		}
-		if (SelectedItem->Key.Type == ERigElementType::Space)
+		if (SelectedKey.Type == ERigElementType::Space)
 		{
 			return true;
 		}
@@ -1546,13 +1553,12 @@ void SRigHierarchy::HandleDeleteItem()
 		bool bConfirmedByUser = false;
 		bool bDeleteImportedBones = false;
 
-		TArray<TSharedPtr<FRigTreeElement>> SelectedItems = TreeView->GetSelectedItems();
-		for (int32 ItemIndex = 0; ItemIndex < SelectedItems.Num(); ++ItemIndex)
+ 		TArray<FRigElementKey> SelectedKeys = GetSelectedKeys();
+		for (const FRigElementKey& SelectedKey : SelectedKeys)
 		{
 			TGuardValue<bool> GuardRigHierarchyChanges(bIsChangingRigHierarchy, true);
 			TGuardValue<bool> SuspendBlueprintNotifs(ControlRigBlueprint->bSuspendAllNotifications, true);
 
-			FRigElementKey SelectedKey = SelectedItems[ItemIndex]->Key;
 			switch (SelectedKey.Type)
 			{
 				case ERigElementType::Bone:
@@ -1643,10 +1649,10 @@ void SRigHierarchy::HandleNewItem(ERigElementType InElementType)
 		FRigElementKey ParentKey;
 		FTransform ParentTransform = FTransform::Identity;
 
-		TArray<TSharedPtr<FRigTreeElement>> SelectedItems = TreeView->GetSelectedItems();
-		if (SelectedItems.Num() > 0)
+		TArray<FRigElementKey> SelectedKeys = GetSelectedKeys();
+		if (SelectedKeys.Num() > 0)
 		{
-			ParentKey = SelectedItems[0]->Key;
+			ParentKey = SelectedKeys[0];
 			ParentTransform = Hierarchy->GetGlobalTransform(ParentKey);
 		}
 
@@ -1755,13 +1761,7 @@ void SRigHierarchy::HandleDuplicateItem()
 			FScopedTransaction Transaction(LOCTEXT("HierarchyTreeDuplicateSelected", "Duplicate selected items from hierarchy"));
 			ControlRigBlueprint->Modify();
 
-			TArray<TSharedPtr<FRigTreeElement>> SelectedItems = TreeView->GetSelectedItems();
-			TArray<FRigElementKey> KeysToDuplicate;
-			for (int32 Index = 0; Index < SelectedItems.Num(); ++Index)
-			{
-				KeysToDuplicate.Add(SelectedItems[Index]->Key);
-			}
-
+			TArray<FRigElementKey> KeysToDuplicate = GetSelectedKeys();
 			Hierarchy->DuplicateItems(KeysToDuplicate, true);
 		}
 		ControlRigBlueprint->PropagateHierarchyFromBPToInstances(true);
@@ -1798,13 +1798,7 @@ void SRigHierarchy::HandleMirrorItem()
 			FScopedTransaction Transaction(LOCTEXT("HierarchyTreeMirrorSelected", "Mirror selected items from hierarchy"));
 			ControlRigBlueprint->Modify();
 
-			TArray<TSharedPtr<FRigTreeElement>> SelectedItems = TreeView->GetSelectedItems();
-			TArray<FRigElementKey> KeysToMirror;
-			for (int32 Index = 0; Index < SelectedItems.Num(); ++Index)
-			{
-				KeysToMirror.Add(SelectedItems[Index]->Key);
-			}
-
+			TArray<FRigElementKey> KeysToMirror = GetSelectedKeys();
 			Hierarchy->MirrorItems(KeysToMirror, Settings, true);
 		}
 		ControlRigBlueprint->PropagateHierarchyFromBPToInstances(true);
@@ -2026,13 +2020,7 @@ void SRigHierarchy::PostUndo(bool bSuccess)
 
 FReply SRigHierarchy::OnDragDetected(const FGeometry& MyGeometry, const FPointerEvent& MouseEvent)
 {
-	TArray<FRigElementKey> DraggedElements;
-	TArray<TSharedPtr<FRigTreeElement>> SelectedItems = TreeView->GetSelectedItems();
-	for (const TSharedPtr<FRigTreeElement>& SelectedItem : SelectedItems)
-	{
-		DraggedElements.Add(SelectedItem->Key);
-	}
-
+	TArray<FRigElementKey> DraggedElements = GetSelectedKeys();
 	if (MouseEvent.IsMouseButtonDown(EKeys::LeftMouseButton) && DraggedElements.Num() > 0)
 	{
 		if (ControlRigEditor.IsValid())
@@ -2451,13 +2439,7 @@ void SRigHierarchy::HandleResetTransform(bool bSelectionOnly)
 				FScopedTransaction Transaction(LOCTEXT("HierarchyResetTransforms", "Reset Transforms"));
 				Blueprint->Modify();
 
-				TArray<TSharedPtr<FRigTreeElement>> SelectedItems = TreeView->GetSelectedItems();
-				TArray<FRigElementKey> KeysToReset;
-				for (TSharedPtr<FRigTreeElement> SelectedItem : SelectedItems)
-				{
-					KeysToReset.Add(SelectedItem->Key);
-				}
-
+				TArray<FRigElementKey> KeysToReset = GetSelectedKeys();
 				if (!bSelectionOnly)
 				{
 					for (const FRigControl& Control : DebuggedContainer->ControlHierarchy)
@@ -2511,30 +2493,29 @@ void SRigHierarchy::HandleResetSpace()
 				FScopedTransaction Transaction(LOCTEXT("HierarchyResetSpace", "Reset Space"));
 				Blueprint->Modify();
 
-				TArray<TSharedPtr<FRigTreeElement>> SelectedItems = TreeView->GetSelectedItems();
-				for (TSharedPtr<FRigTreeElement> SelectedItem : SelectedItems)
+				TArray<FRigElementKey> SelectedKeys = GetSelectedKeys();
+				for (const FRigElementKey& SelectedKey : SelectedKeys)
 				{
-					FRigElementKey Key = SelectedItem->Key;
-					if (Key.Type == ERigElementType::Control)
+					if (SelectedKey.Type == ERigElementType::Control)
 					{
-						FRigControl& Control = DebuggedContainer->ControlHierarchy[Key.Name];
+						FRigControl& Control = DebuggedContainer->ControlHierarchy[SelectedKey.Name];
 						if (FRigElementKey SpaceKey = Control.GetSpaceElementKey())
 						{
 							FRigSpace& Space = DebuggedContainer->SpaceHierarchy[SpaceKey.Name];
 
-							FTransform Transform = DebuggedContainer->GetGlobalTransform(Key);
+							FTransform Transform = DebuggedContainer->GetGlobalTransform(SelectedKey);
 							FTransform SpaceParentTransform = DebuggedContainer->GetGlobalTransform(Space.GetParentElementKey());
 							FTransform LocalTransform = Transform.GetRelativeTransform(SpaceParentTransform);
 
 							GetHierarchyContainer()->SetInitialTransform(SpaceKey, LocalTransform);
 							GetHierarchyContainer()->SetLocalTransform(SpaceKey, LocalTransform);
-							GetHierarchyContainer()->SetInitialTransform(Key, FTransform::Identity);
-							GetHierarchyContainer()->SetLocalTransform(Key, FTransform::Identity);
+							GetHierarchyContainer()->SetInitialTransform(SelectedKey, FTransform::Identity);
+							GetHierarchyContainer()->SetLocalTransform(SelectedKey, FTransform::Identity);
 
 							DebuggedContainer->SetInitialTransform(SpaceKey, LocalTransform);
 							DebuggedContainer->SetLocalTransform(SpaceKey, LocalTransform);
-							DebuggedContainer->SetInitialTransform(Key, FTransform::Identity);
-							DebuggedContainer->SetLocalTransform(Key, FTransform::Identity);
+							DebuggedContainer->SetInitialTransform(SelectedKey, FTransform::Identity);
+							DebuggedContainer->SetLocalTransform(SelectedKey, FTransform::Identity);
 						}
 						else
 						{
@@ -2574,28 +2555,27 @@ void SRigHierarchy::HandleSetInitialTransformFromCurrentTransform()
 				FScopedTransaction Transaction(LOCTEXT("HierarchySetInitialTransforms", "Set Initial Transforms"));
 				Blueprint->Modify();
 
-				TArray<TSharedPtr<FRigTreeElement>> SelectedItems = TreeView->GetSelectedItems();
+				TArray<FRigElementKey> SelectedKeys = GetSelectedKeys();
 				TMap<FRigElementKey, FTransform> GlobalTransforms;
 				TMap<FRigElementKey, FTransform> ParentGlobalTransforms;
 
-				for (TSharedPtr<FRigTreeElement> SelectedItem : SelectedItems)
+				for (const FRigElementKey& SelectedKey : SelectedKeys)
 				{
-					FRigElementKey Key = SelectedItem->Key;
-					if (Key.Type == ERigElementType::Control || Key.Type == ERigElementType::Space || Key.Type == ERigElementType::Bone)
+					if (SelectedKey.Type == ERigElementType::Control || SelectedKey.Type == ERigElementType::Space || SelectedKey.Type == ERigElementType::Bone)
 					{
-						GlobalTransforms.FindOrAdd(Key) = DebuggedContainer->GetGlobalTransform(Key);
+						GlobalTransforms.FindOrAdd(SelectedKey) = DebuggedContainer->GetGlobalTransform(SelectedKey);
 						FRigElementKey ParentKey;
-						switch (Key.Type)
+						switch (SelectedKey.Type)
 						{
 							case ERigElementType::Bone:
 							{
-								const FRigBone& Element = DebuggedContainer->BoneHierarchy[Key.Name];
+								const FRigBone& Element = DebuggedContainer->BoneHierarchy[SelectedKey.Name];
 								ParentKey = Element.GetParentElementKey();
 								break;
 							}
 							case ERigElementType::Control:
 							{
-								const FRigControl& Element = DebuggedContainer->ControlHierarchy[Key.Name];
+								const FRigControl& Element = DebuggedContainer->ControlHierarchy[SelectedKey.Name];
 								if (Element.SpaceName != NAME_None)
 								{
 									ParentKey = Element.GetSpaceElementKey();
@@ -2608,7 +2588,7 @@ void SRigHierarchy::HandleSetInitialTransformFromCurrentTransform()
 							}
 							case ERigElementType::Space:
 							{
-								const FRigSpace& Element = DebuggedContainer->SpaceHierarchy[Key.Name];
+								const FRigSpace& Element = DebuggedContainer->SpaceHierarchy[SelectedKey.Name];
 								ParentKey = Element.GetParentElementKey();
 								break;
 							}
@@ -2625,53 +2605,51 @@ void SRigHierarchy::HandleSetInitialTransformFromCurrentTransform()
 
 						if (ParentKey)
 						{
-							ParentGlobalTransforms.FindOrAdd(Key) = DebuggedContainer->GetGlobalTransform(ParentKey);
+							ParentGlobalTransforms.FindOrAdd(SelectedKey) = DebuggedContainer->GetGlobalTransform(ParentKey);
 						}
 						else
 						{
-							ParentGlobalTransforms.FindOrAdd(Key) = FTransform::Identity;
+							ParentGlobalTransforms.FindOrAdd(SelectedKey) = FTransform::Identity;
 						}
 					}
 
 				}
 
-				for (TSharedPtr<FRigTreeElement> SelectedItem : SelectedItems)
+				for (const FRigElementKey& SelectedKey : SelectedKeys)
 				{
-					FRigElementKey Key = SelectedItem->Key;
+					FTransform GlobalTransform = GlobalTransforms[SelectedKey];
+					FTransform LocalTransform = GlobalTransform.GetRelativeTransform(ParentGlobalTransforms[SelectedKey]);
 
-					FTransform GlobalTransform = GlobalTransforms[Key];
-					FTransform LocalTransform = GlobalTransform.GetRelativeTransform(ParentGlobalTransforms[Key]);
-
-					if (Key.Type == ERigElementType::Control )
+					if (SelectedKey.Type == ERigElementType::Control )
 					{
-						GetHierarchyContainer()->ControlHierarchy[Key.Name].OffsetTransform = LocalTransform;
-						GetHierarchyContainer()->SetLocalTransform(Key, FTransform::Identity);
-						GetHierarchyContainer()->SetInitialTransform(Key, FTransform::Identity);
-						DebuggedContainer->ControlHierarchy[Key.Name].OffsetTransform = LocalTransform;
-						DebuggedContainer->SetLocalTransform(Key, FTransform::Identity);
-						DebuggedContainer->SetInitialTransform(Key, FTransform::Identity);
+						GetHierarchyContainer()->ControlHierarchy[SelectedKey.Name].OffsetTransform = LocalTransform;
+						GetHierarchyContainer()->SetLocalTransform(SelectedKey, FTransform::Identity);
+						GetHierarchyContainer()->SetInitialTransform(SelectedKey, FTransform::Identity);
+						DebuggedContainer->ControlHierarchy[SelectedKey.Name].OffsetTransform = LocalTransform;
+						DebuggedContainer->SetLocalTransform(SelectedKey, FTransform::Identity);
+						DebuggedContainer->SetInitialTransform(SelectedKey, FTransform::Identity);
 					}
-					else if (Key.Type == ERigElementType::Space)
+					else if (SelectedKey.Type == ERigElementType::Space)
 					{
-						GetHierarchyContainer()->SetInitialTransform(Key, LocalTransform);
-						GetHierarchyContainer()->SetLocalTransform(Key, LocalTransform);
-						DebuggedContainer->SetInitialTransform(Key, LocalTransform);
-						DebuggedContainer->SetLocalTransform(Key, LocalTransform);
+						GetHierarchyContainer()->SetInitialTransform(SelectedKey, LocalTransform);
+						GetHierarchyContainer()->SetLocalTransform(SelectedKey, LocalTransform);
+						DebuggedContainer->SetInitialTransform(SelectedKey, LocalTransform);
+						DebuggedContainer->SetLocalTransform(SelectedKey, LocalTransform);
 					}
-					else if (Key.Type == ERigElementType::Bone)
+					else if (SelectedKey.Type == ERigElementType::Bone)
 					{
 						FTransform InitialTransform = GlobalTransform;
 						if (ControlRigEditor.Pin()->PreviewInstance)
 						{
-							if (FAnimNode_ModifyBone* ModifyBone = ControlRigEditor.Pin()->PreviewInstance->FindModifiedBone(Key.Name))
+							if (FAnimNode_ModifyBone* ModifyBone = ControlRigEditor.Pin()->PreviewInstance->FindModifiedBone(SelectedKey.Name))
 							{
 								InitialTransform.SetTranslation(ModifyBone->Translation);
 								InitialTransform.SetRotation(FQuat(ModifyBone->Rotation));
 								InitialTransform.SetScale3D(ModifyBone->Scale);
-								InitialTransform = InitialTransform * ParentGlobalTransforms[Key];
+								InitialTransform = InitialTransform * ParentGlobalTransforms[SelectedKey];
 							}
 						}
-						GetHierarchyContainer()->SetInitialGlobalTransform(Key, InitialTransform);
+						GetHierarchyContainer()->SetInitialGlobalTransform(SelectedKey, InitialTransform);
 						Blueprint->PropagateHierarchyFromBPToInstances(false, false);
 					}
 				}
@@ -2702,15 +2680,15 @@ void SRigHierarchy::HandleControlBoneOrSpaceTransform()
 		return;
 	}
 
-	TArray<TSharedPtr<FRigTreeElement>> SelectedItems = TreeView->GetSelectedItems();
-	for (TSharedPtr<FRigTreeElement> SelectedItem : SelectedItems)
+	TArray<FRigElementKey> SelectedKeys = GetSelectedKeys();
+	for (const FRigElementKey& SelectedKey : SelectedKeys)
 	{
-		if (SelectedItem->Key.Type == ERigElementType::Bone ||
-			SelectedItem->Key.Type == ERigElementType::Space)
+		if (SelectedKey.Type == ERigElementType::Bone ||
+			SelectedKey.Type == ERigElementType::Space)
 		{
-			Blueprint->AddTransientControl(SelectedItem->Key);
+			Blueprint->AddTransientControl(SelectedKey);
 			Blueprint->HierarchyContainer.ClearSelection();
-			Blueprint->HierarchyContainer.Select(SelectedItem->Key, true);
+			Blueprint->HierarchyContainer.Select(SelectedKey, true);
 			return;
 		}
 	}
@@ -2730,8 +2708,18 @@ void SRigHierarchy::HandleUnparent()
 	bool bUnparentImportedBones = false;
 	bool bConfirmedByUser = false;
 
-	TArray<TSharedPtr<FRigTreeElement>> SelectedItems = TreeView->GetSelectedItems();
-	for (TSharedPtr<FRigTreeElement> SelectedItem : SelectedItems)
+	TArray<FRigElementKey> SelectedKeys = GetSelectedKeys();
+	TMap<FRigElementKey, FTransform> InitialTransforms;
+	TMap<FRigElementKey, FTransform> GlobalTransforms;
+
+	for (const FRigElementKey& SelectedKey : SelectedKeys)
+	{
+		FRigHierarchyContainer* Container = GetHierarchyContainer();
+		InitialTransforms.Add(SelectedKey, Container->GetInitialGlobalTransform(SelectedKey));
+		GlobalTransforms.Add(SelectedKey, Container->GetGlobalTransform(SelectedKey));
+	}
+
+	for (const FRigElementKey& SelectedKey : SelectedKeys)
 	{
 		TGuardValue<bool> GuardRigHierarchyChanges(bIsChangingRigHierarchy, true);
 		TGuardValue<bool> SuspendBlueprintNotifs(ControlRigBlueprint->bSuspendAllNotifications, true);
@@ -2741,14 +2729,14 @@ void SRigHierarchy::HandleUnparent()
 		FRigControlHierarchy& ControlHierarchy = Container->ControlHierarchy;
 		FRigSpaceHierarchy& SpaceHierarchy = Container->SpaceHierarchy;
 
-		FTransform InitialTransform = Container->GetInitialGlobalTransform(SelectedItem->Key);
-		FTransform GlobalTransform = Container->GetGlobalTransform(SelectedItem->Key);
+		const FTransform& InitialTransform = InitialTransforms.FindChecked(SelectedKey);
+		const FTransform& GlobalTransform = GlobalTransforms.FindChecked(SelectedKey);
 
-		switch (SelectedItem->Key.Type)
+		switch (SelectedKey.Type)
 		{
 			case ERigElementType::Bone:
 			{
-				bool bIsImportedBone = BoneHierarchy[SelectedItem->Key.Name].Type == ERigBoneType::Imported;
+				bool bIsImportedBone = BoneHierarchy[SelectedKey.Name].Type == ERigBoneType::Imported;
 				if (bIsImportedBone && !bConfirmedByUser)
 				{
 					FText ConfirmUnparent = LOCTEXT("ConfirmUnparentBoneHierarchy", "Unparenting imported(white) bones can cause issues with animation - are you sure ?");
@@ -2764,19 +2752,19 @@ void SRigHierarchy::HandleUnparent()
 
 				if (bUnparentImportedBones || !bIsImportedBone)
 				{
-					BoneHierarchy.Reparent(SelectedItem->Key.Name, NAME_None);
+					BoneHierarchy.Reparent(SelectedKey.Name, NAME_None);
 				}
 				break;
 			}
 			case ERigElementType::Space:
 			{
-				SpaceHierarchy.Reparent(SelectedItem->Key.Name, ERigSpaceType::Global, NAME_None);
+				SpaceHierarchy.Reparent(SelectedKey.Name, ERigSpaceType::Global, NAME_None);
 				break;
 			}
 			case ERigElementType::Control:
 			{
-				ControlHierarchy.Reparent(SelectedItem->Key.Name, NAME_None);
-				ControlHierarchy.SetSpace(SelectedItem->Key.Name, NAME_None);
+				ControlHierarchy.Reparent(SelectedKey.Name, NAME_None);
+				ControlHierarchy.SetSpace(SelectedKey.Name, NAME_None);
 			}
 			default:
 			{
@@ -2784,8 +2772,20 @@ void SRigHierarchy::HandleUnparent()
 			}
 		}
 
-		Container->SetInitialGlobalTransform(SelectedItem->Key, InitialTransform);
-		Container->SetGlobalTransform(SelectedItem->Key, GlobalTransform);
+		if(SelectedKey.Type == ERigElementType::Control)
+		{
+			int32 ControlIndex = Container->GetIndex(SelectedKey);
+			FTransform ParentTransform = Container->ControlHierarchy.GetParentTransform(ControlIndex, false);
+			FTransform OffsetTransform = GlobalTransform.GetRelativeTransform(ParentTransform);
+			Container->ControlHierarchy[ControlIndex].OffsetTransform = OffsetTransform;
+			Container->SetLocalTransform(SelectedKey, FTransform::Identity);
+			Container->SetInitialTransform(SelectedKey, FTransform::Identity);
+		}
+		else
+		{
+			Container->SetInitialGlobalTransform(SelectedKey, InitialTransform);
+			Container->SetGlobalTransform(SelectedKey, GlobalTransform);
+		}
 	}
 
 	ControlRigBlueprint->PropagateHierarchyFromBPToInstances(true);
@@ -2829,16 +2829,15 @@ void SRigHierarchy::HandleSetInitialTransformFromClosestBone()
 				FScopedTransaction Transaction(LOCTEXT("HierarchySetInitialTransforms", "Set Initial Transforms"));
 				Blueprint->Modify();
 
-				TArray<TSharedPtr<FRigTreeElement>> SelectedItems = TreeView->GetSelectedItems();
+				TArray<FRigElementKey> SelectedKeys = GetSelectedKeys();
 				TMap<FRigElementKey, FTransform> ClosestTransforms;
 				TMap<FRigElementKey, FTransform> ParentGlobalTransforms;
 
-				for (TSharedPtr<FRigTreeElement> SelectedItem : SelectedItems)
+				for (const FRigElementKey& SelectedKey : SelectedKeys)
 				{
-					FRigElementKey Key = SelectedItem->Key;
-					if (Key.Type == ERigElementType::Control || Key.Type == ERigElementType::Space)
+					if (SelectedKey.Type == ERigElementType::Control || SelectedKey.Type == ERigElementType::Space)
 					{
-						FTransform GlobalTransform = DebuggedContainer->GetGlobalTransform(Key);
+						FTransform GlobalTransform = DebuggedContainer->GetGlobalTransform(SelectedKey);
 						FTransform ClosestTransform;
 						FName ClosestRigElement;
 
@@ -2847,14 +2846,14 @@ void SRigHierarchy::HandleSetInitialTransformFromClosestBone()
 							continue;
 						}
 
-						ClosestTransforms.FindOrAdd(Key) = ClosestTransform;
+						ClosestTransforms.FindOrAdd(SelectedKey) = ClosestTransform;
 
 						FRigElementKey ParentKey;
-						switch (Key.Type)
+						switch (SelectedKey.Type)
 						{
 							case ERigElementType::Control:
 							{
-								const FRigControl& Element = DebuggedContainer->ControlHierarchy[Key.Name];
+								const FRigControl& Element = DebuggedContainer->ControlHierarchy[SelectedKey.Name];
 								if (Element.SpaceName != NAME_None)
 								{
 									ParentKey = Element.GetSpaceElementKey();
@@ -2867,7 +2866,7 @@ void SRigHierarchy::HandleSetInitialTransformFromClosestBone()
 							}
 							case ERigElementType::Space:
 							{
-								const FRigSpace& Element = DebuggedContainer->SpaceHierarchy[Key.Name];
+								const FRigSpace& Element = DebuggedContainer->SpaceHierarchy[SelectedKey.Name];
 								ParentKey = Element.GetParentElementKey();
 								break;
 							}
@@ -2885,42 +2884,41 @@ void SRigHierarchy::HandleSetInitialTransformFromClosestBone()
 
 						if (ParentKey)
 						{
-							ParentGlobalTransforms.FindOrAdd(Key) = DebuggedContainer->GetGlobalTransform(ParentKey);
+							ParentGlobalTransforms.FindOrAdd(SelectedKey) = DebuggedContainer->GetGlobalTransform(ParentKey);
 						}
 						else
 						{
-							ParentGlobalTransforms.FindOrAdd(Key) = FTransform::Identity;
+							ParentGlobalTransforms.FindOrAdd(SelectedKey) = FTransform::Identity;
 						}
 					}
 
 				}
 
-				for (TSharedPtr<FRigTreeElement> SelectedItem : SelectedItems)
+				for (const FRigElementKey& SelectedKey : SelectedKeys)
 				{
-					FRigElementKey Key = SelectedItem->Key;
-					if (!ClosestTransforms.Contains(Key))
+					if (!ClosestTransforms.Contains(SelectedKey))
 					{
 						continue;
 					}
 
-					if (Key.Type == ERigElementType::Space || Key.Type == ERigElementType::Bone)
+					if (SelectedKey.Type == ERigElementType::Space || SelectedKey.Type == ERigElementType::Bone)
 					{
-						FTransform LocalTransform = ClosestTransforms[Key].GetRelativeTransform(ParentGlobalTransforms[Key]);
-						GetHierarchyContainer()->SetInitialTransform(Key, LocalTransform);
-						GetHierarchyContainer()->SetLocalTransform(Key, LocalTransform);
-						DebuggedContainer->SetInitialTransform(Key, LocalTransform);
-						DebuggedContainer->SetLocalTransform(Key, LocalTransform);
-						Blueprint->SetTransientControlValue(Key);
+						FTransform LocalTransform = ClosestTransforms[SelectedKey].GetRelativeTransform(ParentGlobalTransforms[SelectedKey]);
+						GetHierarchyContainer()->SetInitialTransform(SelectedKey, LocalTransform);
+						GetHierarchyContainer()->SetLocalTransform(SelectedKey, LocalTransform);
+						DebuggedContainer->SetInitialTransform(SelectedKey, LocalTransform);
+						DebuggedContainer->SetLocalTransform(SelectedKey, LocalTransform);
+						Blueprint->SetTransientControlValue(SelectedKey);
 					}
-					if (Key.Type == ERigElementType::Control)
+					if (SelectedKey.Type == ERigElementType::Control)
 					{
-						FTransform LocalTransform = ClosestTransforms[Key].GetRelativeTransform(ParentGlobalTransforms[Key]);
-						GetHierarchyContainer()->ControlHierarchy[Key.Name].OffsetTransform = LocalTransform;
-						GetHierarchyContainer()->SetLocalTransform(Key, FTransform::Identity);
-						GetHierarchyContainer()->SetInitialTransform(Key, FTransform::Identity);
-						DebuggedContainer->ControlHierarchy[Key.Name].OffsetTransform = LocalTransform;
-						DebuggedContainer->SetLocalTransform(Key, FTransform::Identity);
-						DebuggedContainer->SetInitialTransform(Key, FTransform::Identity);
+						FTransform LocalTransform = ClosestTransforms[SelectedKey].GetRelativeTransform(ParentGlobalTransforms[SelectedKey]);
+						GetHierarchyContainer()->ControlHierarchy[SelectedKey.Name].OffsetTransform = LocalTransform;
+						GetHierarchyContainer()->SetLocalTransform(SelectedKey, FTransform::Identity);
+						GetHierarchyContainer()->SetInitialTransform(SelectedKey, FTransform::Identity);
+						DebuggedContainer->ControlHierarchy[SelectedKey.Name].OffsetTransform = LocalTransform;
+						DebuggedContainer->SetLocalTransform(SelectedKey, FTransform::Identity);
+						DebuggedContainer->SetInitialTransform(SelectedKey, FTransform::Identity);
 					}
 				}
 			}
diff --git a/Engine/Plugins/Experimental/ControlRig/Source/ControlRigEditor/Private/Editor/SRigHierarchy.h b/Engine/Plugins/Experimental/ControlRig/Source/ControlRigEditor/Private/Editor/SRigHierarchy.h
--- a/Engine/Plugins/Experimental/ControlRig/Source/ControlRigEditor/Private/Editor/SRigHierarchy.h
+++ b/Engine/Plugins/Experimental/ControlRig/Source/ControlRigEditor/Private/Editor/SRigHierarchy.h
@@ -171,6 +171,9 @@ private:
 	/** Rebuild the tree view */
 	void RefreshTreeView();
 
+	/** Return all selected keys */
+	TArray<FRigElementKey> GetSelectedKeys() const;
+
 	/** Make a row widget for the table */
 	TSharedRef<ITableRow> MakeTableRowWidget(TSharedPtr<FRigTreeElement> InItem, const TSharedRef<STableViewBase>& OwnerTable);
 
