Avoid reusing the same DataSourceRegistry instance #jira FROST-923 #rb helge.mathee

diff --git a/Engine/Plugins/Experimental/ControlRig/Source/ControlRig/Private/ControlRig.cpp b/Engine/Plugins/Experimental/ControlRig/Source/ControlRig/Private/ControlRig.cpp
--- a/Engine/Plugins/Experimental/ControlRig/Source/ControlRig/Private/ControlRig.cpp
+++ b/Engine/Plugins/Experimental/ControlRig/Source/ControlRig/Private/ControlRig.cpp
@@ -1022,6 +1022,13 @@ void UControlRig::GetMappableNodeData(TArray<FName>& OutNames, TArray<FNodeItem>
 
 UAnimationDataSourceRegistry* UControlRig::GetDataSourceRegistry()
 {
+	if (DataSourceRegistry)
+	{
+		if (DataSourceRegistry->GetOuter() != this)
+		{
+			DataSourceRegistry = nullptr;
+		}
+	}
 	if (DataSourceRegistry == nullptr)
 	{
 		DataSourceRegistry = NewObject<UAnimationDataSourceRegistry>(this, TEXT("DataSourceRegistry"));
