Control Rig: Fix in-control-rig-editor focus issues

diff --git a/Engine/Plugins/Experimental/ControlRig/Source/ControlRigEditor/Private/EditMode/ControlRigEditMode.cpp b/Engine/Plugins/Experimental/ControlRig/Source/ControlRigEditor/Private/EditMode/ControlRigEditMode.cpp
--- a/Engine/Plugins/Experimental/ControlRig/Source/ControlRigEditor/Private/EditMode/ControlRigEditMode.cpp
+++ b/Engine/Plugins/Experimental/ControlRig/Source/ControlRigEditor/Private/EditMode/ControlRigEditMode.cpp
@@ -41,7 +41,9 @@
 #include "ControlRigComponent.h"
 #include "EngineUtils.h"
 #include "ControlRig/Private/Units/Execution/RigUnit_BeginExecution.h"
-
+//#include "IPersonaPreviewScene.h"
+//#include "Animation/DebugSkelMeshComponent.h"
+//#include "Persona/Private/AnimationEditorViewportClient.h"
 
 void UControlRigEditModeDelegateHelper::OnPoseInitialized()
 {
@@ -867,9 +869,6 @@ bool FControlRigEditMode::HandleClick(FEditorViewportClient* InViewportClient, H
 		}
 	}
 
-		
-	
-
 	// for now we show this menu all the time if body is selected
 	// if we want some global menu, we'll have to move this
 	if (Click.GetKey() == EKeys::RightMouseButton)
@@ -888,6 +887,39 @@ bool FControlRigEditMode::HandleClick(FEditorViewportClient* InViewportClient, H
 	// clear selected controls
 	ClearRigElementSelection(FRigElementTypeHelper::ToMask(ERigElementType::All));
 
+	/*
+	if(!InViewportClient->IsLevelEditorClient() && !InViewportClient->IsSimulateInEditorViewport())
+	{
+		bool bHandled = false;
+		const bool bSelectingSections = GetAnimPreviewScene().AllowMeshHitProxies();
+
+		USkeletalMeshComponent* MeshComponent = GetAnimPreviewScene().GetPreviewMeshComponent();
+
+		if ( HitProxy )
+		{
+			if ( HitProxy->IsA( HPersonaBoneProxy::StaticGetType() ) )
+			{			
+				SetRigElementSelection(ERigElementType::Bone, static_cast<HPersonaBoneProxy*>(HitProxy)->BoneName, true);
+				bHandled = true;
+			}
+		}
+		
+		if ( !bHandled && !bSelectingSections )
+		{
+			// Cast for phys bodies if we didn't get any hit proxies
+			FHitResult Result(1.0f);
+			UDebugSkelMeshComponent* PreviewMeshComponent = GetAnimPreviewScene().GetPreviewMeshComponent();
+			bool bHit = PreviewMeshComponent->LineTraceComponent(Result, Click.GetOrigin(), Click.GetOrigin() + Click.GetDirection() * 10000.0f, FCollisionQueryParams(NAME_None, FCollisionQueryParams::GetUnknownStatId(),true));
+			
+			if(bHit)
+			{
+				SetRigElementSelection(ERigElementType::Bone, Result.BoneName, true);
+				bHandled = true;
+			}
+		}
+	}
+	*/
+	
 	return FEdMode::HandleClick(InViewportClient, HitProxy, Click);
 }
 
@@ -1504,6 +1536,8 @@ bool FControlRigEditMode::GetRigElementGlobalTransform(const FRigElementKey& InE
 
 bool FControlRigEditMode::CanFrameSelection()
 {
+	return SelectedRigElements.Num() > 0;
+	/*
 	for (const FRigElementKey& SelectedKey : SelectedRigElements)
 	{
 		if (SelectedKey.Type == ERigElementType::Control)
@@ -1512,10 +1546,24 @@ bool FControlRigEditMode::CanFrameSelection()
 		}
 	}
 	return false;
+	*/
 }
 
 void FControlRigEditMode::FrameSelection()
 {
+	if(CurrentViewportClient)
+	{
+		FSphere Sphere(EForceInit::ForceInit);
+		if(GetCameraTarget(Sphere))
+		{
+			FBox Bounds(EForceInit::ForceInit);
+			Bounds += Sphere.Center;
+			Bounds += Sphere.Center + FVector::OneVector * Sphere.W;
+			Bounds += Sphere.Center - FVector::OneVector * Sphere.W;
+			CurrentViewportClient->FocusViewportOnBox(Bounds);
+		}
+    }
+
 	TArray<AActor*> Actors;
 	for (const FRigElementKey& SelectedKey : SelectedRigElements)
 	{
@@ -1529,7 +1577,7 @@ void FControlRigEditMode::FrameSelection()
 		}
 	}
 
-	if (Actors.Num() )
+	if (Actors.Num())
 	{
 		TArray<UPrimitiveComponent*> SelectedComponents;
 		GEditor->MoveViewportCamerasToActor(Actors, SelectedComponents, true);
diff --git a/Engine/Plugins/Experimental/ControlRig/Source/ControlRigEditor/Private/Editor/ControlRigEditorEditMode.cpp b/Engine/Plugins/Experimental/ControlRig/Source/ControlRigEditor/Private/Editor/ControlRigEditorEditMode.cpp
--- a/Engine/Plugins/Experimental/ControlRig/Source/ControlRigEditor/Private/Editor/ControlRigEditorEditMode.cpp
+++ b/Engine/Plugins/Experimental/ControlRig/Source/ControlRigEditor/Private/Editor/ControlRigEditorEditMode.cpp
@@ -108,16 +108,18 @@ bool FControlRigEditorEditMode::GetCameraTarget(FSphere& OutTarget) const
 
 	for (int32 Index = 0; Index < SelectedRigElements.Num(); ++Index)
 	{
+		static const float Radius = 20.f;
 		if (SelectedRigElements[Index].Type == ERigElementType::Bone || SelectedRigElements[Index].Type == ERigElementType::Space)
 		{
 			FTransform Transform = OnGetRigElementTransformDelegate.Execute(SelectedRigElements[Index], false, true);
-			Box += Transform.GetLocation();
+			Box += Transform.TransformPosition(FVector::OneVector * Radius);
+			Box += Transform.TransformPosition(FVector::OneVector * -Radius);
 		}
 		else if (SelectedRigElements[Index].Type == ERigElementType::Control)
 		{
 			FTransform Transform = OnGetRigElementTransformDelegate.Execute(SelectedRigElements[Index], false, true);
-			Box += Transform.TransformPosition(FVector::OneVector * 50.f); // we assume 100cm as the base size for a control
-			Box += Transform.TransformPosition(FVector::OneVector * -50.f);
+			Box += Transform.TransformPosition(FVector::OneVector * Radius);
+			Box += Transform.TransformPosition(FVector::OneVector * -Radius);
 		}
 	}
 
