Control Rig: Pose Library: Use Local Transforms for non matched mirrors and set axis to X.

diff --git a/Engine/Plugins/Experimental/ControlRig/Source/ControlRig/Private/Tools/ControlRigPose.cpp b/Engine/Plugins/Experimental/ControlRig/Source/ControlRig/Private/Tools/ControlRigPose.cpp
--- a/Engine/Plugins/Experimental/ControlRig/Source/ControlRig/Private/Tools/ControlRigPose.cpp
+++ b/Engine/Plugins/Experimental/ControlRig/Source/ControlRig/Private/Tools/ControlRigPose.cpp
@@ -35,7 +35,7 @@ void FControlRigControlPose::PastePose(UControlRig* ControlRig, bool bDoKey, boo
 
 void FControlRigControlPose::SetControlMirrorTransform(bool bDoLocal, UControlRig* ControlRig, const FName& Name, bool bIsMatched, const FVector& GlobalTranslation, const FQuat& GlobalRotation, const FVector& LocalTranslation, const FQuat& LocalRotation, bool bNotify, const  FRigControlModifiedContext& Context)
 {
-	if (bIsMatched)
+	if (bDoLocal || bIsMatched)
 	{
 		FTransform NewLocalTransform(LocalRotation, LocalTranslation);
 		ControlRig->SetControlLocalTransform(Name, NewLocalTransform, bNotify, Context);
@@ -374,6 +374,5 @@ bool UControlRigPoseAsset::DoesMirrorMatch(UControlRig* ControlRig, const FName&
 }
 
 
-
 #undef LOCTEXT_NAMESPACE
 
diff --git a/Engine/Plugins/Experimental/ControlRig/Source/ControlRig/Private/Tools/ControlRigPoseMirrorSettings.cpp b/Engine/Plugins/Experimental/ControlRig/Source/ControlRig/Private/Tools/ControlRigPoseMirrorSettings.cpp
--- a/Engine/Plugins/Experimental/ControlRig/Source/ControlRig/Private/Tools/ControlRigPoseMirrorSettings.cpp
+++ b/Engine/Plugins/Experimental/ControlRig/Source/ControlRig/Private/Tools/ControlRigPoseMirrorSettings.cpp
@@ -6,5 +6,5 @@ UControlRigPoseMirrorSettings::UControlRigPoseMirrorSettings()
 	RightSide = TEXT("_r_");
 	LeftSide = TEXT("_l_");
 	MirrorAxis = EAxis::X;
-	AxisToFlip = EAxis::Z;
+	AxisToFlip = EAxis::X;
 }
diff --git a/Engine/Plugins/Experimental/ControlRig/Source/ControlRig/Private/Tools/ControlRigPoseMirrorTable.cpp b/Engine/Plugins/Experimental/ControlRig/Source/ControlRig/Private/Tools/ControlRigPoseMirrorTable.cpp
--- a/Engine/Plugins/Experimental/ControlRig/Source/ControlRig/Private/Tools/ControlRigPoseMirrorTable.cpp
+++ b/Engine/Plugins/Experimental/ControlRig/Source/ControlRig/Private/Tools/ControlRigPoseMirrorTable.cpp
@@ -93,6 +93,10 @@ void FControlRigPoseMirrorTable::GetMirrorTransform(const FRigControlCopy& Contr
 			FTransform NewTransform = MirrorSettings.MirrorTransform(GlobalTransform);
 			OutGlobalTranslation = NewTransform.GetTranslation();
 			OutGlobalRotation = NewTransform.GetRotation();
+
+			FTransform NewLocalTransform = MirrorSettings.MirrorTransform(LocalTransform);
+			OutLocalTranslation = NewLocalTransform.GetTranslation();
+			OutLocalRotation = NewLocalTransform.GetRotation();
 			return;
 		}
 	}
