Control Rig: Hide "All" Rig Element Type from Node UI, This is needed because we currently don't support creating a list of all rig elements since elements of different types can have the same name

diff --git a/Engine/Plugins/Experimental/ControlRig/Source/ControlRig/Public/Rigs/RigHierarchyDefines.h b/Engine/Plugins/Experimental/ControlRig/Source/ControlRig/Public/Rigs/RigHierarchyDefines.h
--- a/Engine/Plugins/Experimental/ControlRig/Source/ControlRig/Public/Rigs/RigHierarchyDefines.h
+++ b/Engine/Plugins/Experimental/ControlRig/Source/ControlRig/Public/Rigs/RigHierarchyDefines.h
@@ -19,7 +19,7 @@ enum class ERigElementType : uint8
 	Space = 0x002,
 	Control = 0x004,
 	Curve = 0x008,
-	All = Bone | Space | Control | Curve,
+	All = Bone | Space | Control | Curve UMETA(Hidden),
 };
 
 UENUM()
