Control Rig: Color invalid item names red in UI, properly report errors at all times

diff --git a/Engine/Plugins/Experimental/ControlRig/Source/ControlRig/Private/ControlRig.cpp b/Engine/Plugins/Experimental/ControlRig/Source/ControlRig/Private/ControlRig.cpp
--- a/Engine/Plugins/Experimental/ControlRig/Source/ControlRig/Private/ControlRig.cpp
+++ b/Engine/Plugins/Experimental/ControlRig/Source/ControlRig/Private/ControlRig.cpp
@@ -559,7 +559,7 @@ void UControlRig::Execute(const EControlRigState InState, const FName& InEventNa
 	Context.Log = ControlRigLog;
 	if (ControlRigLog != nullptr)
 	{
-		ControlRigLog->Entries.Reset();
+		ControlRigLog->Reset();
 	}
 #endif
 
@@ -1031,7 +1031,7 @@ UAnimationDataSourceRegistry* UControlRig::GetDataSourceRegistry()
 	}
 	if (DataSourceRegistry == nullptr)
 	{
-		DataSourceRegistry = NewObject<UAnimationDataSourceRegistry>(this, TEXT("DataSourceRegistry"));
+		DataSourceRegistry = NewObject<UAnimationDataSourceRegistry>(this);
 	}
 	return DataSourceRegistry;
 }
diff --git a/Engine/Plugins/Experimental/ControlRig/Source/ControlRig/Private/ControlRigLog.cpp b/Engine/Plugins/Experimental/ControlRig/Source/ControlRig/Private/ControlRigLog.cpp
--- a/Engine/Plugins/Experimental/ControlRig/Source/ControlRig/Private/ControlRigLog.cpp
+++ b/Engine/Plugins/Experimental/ControlRig/Source/ControlRig/Private/ControlRigLog.cpp
@@ -2,6 +2,14 @@
 
 #include "ControlRigLog.h"
 
+void FControlRigLog::Reset()
+{
+#if WITH_EDITOR
+	Entries.Reset();
+	KnownMessages.Reset();
+#endif
+}
+
 void FControlRigLog::Report(EMessageSeverity::Type InSeverity, const FName& InOperatorName, int32 InInstructionIndex, const FString& InMessage)
 {
 #if WITH_EDITOR
diff --git a/Engine/Plugins/Experimental/ControlRig/Source/ControlRig/Private/Units/Hierarchy/RigUnit_GetTransform.cpp b/Engine/Plugins/Experimental/ControlRig/Source/ControlRig/Private/Units/Hierarchy/RigUnit_GetTransform.cpp
--- a/Engine/Plugins/Experimental/ControlRig/Source/ControlRig/Private/Units/Hierarchy/RigUnit_GetTransform.cpp
+++ b/Engine/Plugins/Experimental/ControlRig/Source/ControlRig/Private/Units/Hierarchy/RigUnit_GetTransform.cpp
@@ -26,10 +26,7 @@ FRigUnit_GetTransform_Execute()
 			{
 				if (!CachedIndex.UpdateCache(Item, Hierarchy))
 				{
-					if(Context.State != EControlRigState::Init)
-					{
-						UE_CONTROLRIG_RIGUNIT_REPORT_WARNING(TEXT("Item '%s' is not valid."), *Item.ToString());
-					}
+					UE_CONTROLRIG_RIGUNIT_REPORT_WARNING(TEXT("Item '%s' is not valid."), *Item.ToString());
 				}
 				else
 				{
diff --git a/Engine/Plugins/Experimental/ControlRig/Source/ControlRig/Private/Units/Hierarchy/RigUnit_PropagateTransform.cpp b/Engine/Plugins/Experimental/ControlRig/Source/ControlRig/Private/Units/Hierarchy/RigUnit_PropagateTransform.cpp
--- a/Engine/Plugins/Experimental/ControlRig/Source/ControlRig/Private/Units/Hierarchy/RigUnit_PropagateTransform.cpp
+++ b/Engine/Plugins/Experimental/ControlRig/Source/ControlRig/Private/Units/Hierarchy/RigUnit_PropagateTransform.cpp
@@ -27,10 +27,7 @@ FRigUnit_PropagateTransform_Execute()
 			{
 				if (!CachedIndex.UpdateCache(Item, Hierarchy))
 				{
-					if(Context.State != EControlRigState::Init)
-					{
-						UE_CONTROLRIG_RIGUNIT_REPORT_WARNING(TEXT("Item '%s' is not valid."), *Item.ToString());
-					}
+					UE_CONTROLRIG_RIGUNIT_REPORT_WARNING(TEXT("Item '%s' is not valid."), *Item.ToString());
 				}
 				else
 				{
diff --git a/Engine/Plugins/Experimental/ControlRig/Source/ControlRig/Private/Units/Hierarchy/RigUnit_SetTransform.cpp b/Engine/Plugins/Experimental/ControlRig/Source/ControlRig/Private/Units/Hierarchy/RigUnit_SetTransform.cpp
--- a/Engine/Plugins/Experimental/ControlRig/Source/ControlRig/Private/Units/Hierarchy/RigUnit_SetTransform.cpp
+++ b/Engine/Plugins/Experimental/ControlRig/Source/ControlRig/Private/Units/Hierarchy/RigUnit_SetTransform.cpp
@@ -39,10 +39,7 @@ FRigUnit_SetTransform_Execute()
 			{
 				if (!CachedIndex.UpdateCache(Item, Hierarchy))
 				{
-					if(Context.State != EControlRigState::Init)
-					{
-						UE_CONTROLRIG_RIGUNIT_REPORT_WARNING(TEXT("Item '%s' is not valid."), *Item.ToString());
-					}
+					UE_CONTROLRIG_RIGUNIT_REPORT_WARNING(TEXT("Item '%s' is not valid."), *Item.ToString());
 				}
 				else
 				{
diff --git a/Engine/Plugins/Experimental/ControlRig/Source/ControlRig/Public/ControlRigLog.h b/Engine/Plugins/Experimental/ControlRig/Source/ControlRig/Public/ControlRigLog.h
--- a/Engine/Plugins/Experimental/ControlRig/Source/ControlRig/Public/ControlRigLog.h
+++ b/Engine/Plugins/Experimental/ControlRig/Source/ControlRig/Public/ControlRigLog.h
@@ -31,5 +31,6 @@ public:
 	TMap<FString, bool> KnownMessages;
 #endif
 
+	virtual void Reset();
 	virtual void Report(EMessageSeverity::Type InSeverity, const FName& InFunctionName, int32 InInstructionIndex, const FString& InMessage);
 };
diff --git a/Engine/Plugins/Experimental/ControlRig/Source/ControlRigEditor/Private/Editor/ControlRigEditor.cpp b/Engine/Plugins/Experimental/ControlRig/Source/ControlRigEditor/Private/Editor/ControlRigEditor.cpp
--- a/Engine/Plugins/Experimental/ControlRig/Source/ControlRigEditor/Private/Editor/ControlRigEditor.cpp
+++ b/Engine/Plugins/Experimental/ControlRig/Source/ControlRigEditor/Private/Editor/ControlRigEditor.cpp
@@ -3193,7 +3193,8 @@ void FControlRigEditor::OnRigElementRemoved(FRigHierarchyContainer* Container, c
 		return;
 	}
 
-	FString NoneStr = FName(NAME_None).ToString();
+	CacheNameLists();
+
 	FString RemovedElementName = InKey.Name.ToString();
 	ERigElementType RemovedElementType = InKey.Type;
 
@@ -3224,7 +3225,8 @@ void FControlRigEditor::OnRigElementRemoved(FRigHierarchyContainer* Container, c
 						{
 							if (ModelPin->GetDefaultValue() == RemovedElementName)
 							{
-								GetFocusedController()->SetPinDefaultValue(ModelPin->GetPinPath(), FName(NAME_None).ToString());
+								RigNode->ReconstructNode();
+								break;
 							}
 						}
 						else if (ModelPin->GetCPPTypeObject() == FRigElementKey::StaticStruct())
@@ -3241,7 +3243,8 @@ void FControlRigEditor::OnRigElementRemoved(FRigHierarchyContainer* Container, c
 										FString NameStr = NamePin->GetDefaultValue();
 										if (NameStr == RemovedElementName)
 										{
-											GetFocusedController()->SetPinDefaultValue(NamePin->GetPinPath(), NoneStr);
+											RigNode->ReconstructNode();
+											break;
 										}
 									}
 								}
@@ -3251,8 +3254,6 @@ void FControlRigEditor::OnRigElementRemoved(FRigHierarchyContainer* Container, c
 				}
 			}
 		}
-
-		CacheNameLists();
 	}
 
 	OnHierarchyChanged();
diff --git a/Engine/Plugins/Experimental/ControlRig/Source/ControlRigEditor/Private/Graph/SControlRigGraphPinNameList.cpp b/Engine/Plugins/Experimental/ControlRig/Source/ControlRigEditor/Private/Graph/SControlRigGraphPinNameList.cpp
--- a/Engine/Plugins/Experimental/ControlRig/Source/ControlRigEditor/Private/Graph/SControlRigGraphPinNameList.cpp
+++ b/Engine/Plugins/Experimental/ControlRig/Source/ControlRigEditor/Private/Graph/SControlRigGraphPinNameList.cpp
@@ -13,7 +13,9 @@ void SControlRigGraphPinNameList::Construct(const FArguments& InArgs, UEdGraphPi
 {
 	this->ModelPin = InArgs._ModelPin;
 	this->OnGetNameListContent = InArgs._OnGetNameListContent;
+	this->bMarkupInvalidItems = InArgs._MarkupInvalidItems;
 
+	CurrentList = GetNameList();
 	SGraphPin::Construct(SGraphPin::FArguments(), InGraphPinObj);
 }
 
@@ -42,6 +44,7 @@ TSharedRef<SWidget>	SControlRigGraphPinNameList::GetDefaultValueWidget()
 				[
 					SNew(STextBlock)
 					.Text(this, &SControlRigGraphPinNameList::GetNameListText)
+					.ColorAndOpacity(this, &SControlRigGraphPinNameList::GetNameColor)
 				]
 		];
 }
@@ -70,6 +73,30 @@ void SControlRigGraphPinNameList::SetNameListText(const FText& NewTypeInValue, E
 	}
 }
 
+FSlateColor SControlRigGraphPinNameList::GetNameColor() const
+{
+	if(bMarkupInvalidItems)
+	{
+		FString CurrentItem = GetNameListText().ToString();
+		
+		bool bFound = false;
+		for (TSharedPtr<FString> Item : CurrentList)
+		{
+			if (Item->Equals(CurrentItem))
+			{
+				bFound = true;
+				break;
+			}
+		}
+
+		if(!bFound || CurrentItem.IsEmpty() || CurrentItem == FName(NAME_None).ToString())
+		{
+			return FSlateColor(FLinearColor::Red);
+		}
+	}
+	return FSlateColor::UseForeground();
+}
+
 TSharedRef<SWidget> SControlRigGraphPinNameList::MakeNameListItemWidget(TSharedPtr<FString> InItem)
 {
 	return 	SNew(STextBlock).Text(FText::FromString(*InItem));// .Font(FEditorStyle::GetFontStyle(TEXT("PropertyWindow.NormalFont")));
diff --git a/Engine/Plugins/Experimental/ControlRig/Source/ControlRigEditor/Private/Graph/SControlRigGraphPinNameList.h b/Engine/Plugins/Experimental/ControlRig/Source/ControlRigEditor/Private/Graph/SControlRigGraphPinNameList.h
--- a/Engine/Plugins/Experimental/ControlRig/Source/ControlRigEditor/Private/Graph/SControlRigGraphPinNameList.h
+++ b/Engine/Plugins/Experimental/ControlRig/Source/ControlRigEditor/Private/Graph/SControlRigGraphPinNameList.h
@@ -15,10 +15,13 @@ public:
 
 	DECLARE_DELEGATE_RetVal_OneParam( const TArray<TSharedPtr<FString>>&, FOnGetNameListContent, URigVMPin*);
 
-	SLATE_BEGIN_ARGS(SControlRigGraphPinNameList){}
+	SLATE_BEGIN_ARGS(SControlRigGraphPinNameList)
+		: _MarkupInvalidItems(true)
+	{}
 
 		SLATE_ARGUMENT(URigVMPin*, ModelPin)
 		SLATE_EVENT(FOnGetNameListContent, OnGetNameListContent)
+		SLATE_ARGUMENT(bool, MarkupInvalidItems)
 
 	SLATE_END_ARGS()
 
@@ -31,6 +34,7 @@ protected:
 
 	const TArray<TSharedPtr<FString>>& GetNameList() const;
 	FText GetNameListText() const;
+	FSlateColor GetNameColor() const;
 	virtual void SetNameListText(const FText& NewTypeInValue, ETextCommit::Type CommitInfo);
 
 	TSharedRef<SWidget> MakeNameListItemWidget(TSharedPtr<FString> InItem);
@@ -42,4 +46,5 @@ protected:
 	TSharedPtr<SControlRigGraphPinNameListValueWidget> NameListComboBox;
 	TArray<TSharedPtr<FString>> EmptyList;
 	TArray<TSharedPtr<FString>> CurrentList;
+	bool bMarkupInvalidItems;
 };
