Control Rig: Added additional keywords for GetControlVector

diff --git a/Engine/Plugins/Experimental/ControlRig/Source/ControlRig/Private/Units/Hierarchy/RigUnit_GetControlTransform.h b/Engine/Plugins/Experimental/ControlRig/Source/ControlRig/Private/Units/Hierarchy/RigUnit_GetControlTransform.h
--- a/Engine/Plugins/Experimental/ControlRig/Source/ControlRig/Private/Units/Hierarchy/RigUnit_GetControlTransform.h
+++ b/Engine/Plugins/Experimental/ControlRig/Source/ControlRig/Private/Units/Hierarchy/RigUnit_GetControlTransform.h
@@ -160,9 +160,9 @@ struct CONTROLRIG_API FRigUnit_GetControlVector2D : public FRigUnit
 };
 
 /**
- * GetControlVector is used to retrieve a single Vector from a hierarchy.
+ * GetControlVector is used to retrieve a single Vector from a hierarchy, can be used for Controls of type "Position" or "Scale".
  */
-USTRUCT(meta=(DisplayName="Get Control Vector", Category="Controls", DocumentationPolicy = "Strict", Keywords="GetControlVector", PrototypeName = "GetControlValue", Varying))
+USTRUCT(meta=(DisplayName="Get Control Vector", Category="Controls", DocumentationPolicy = "Strict", Keywords="GetControlVector, GetControlPosition, GetControlScale", PrototypeName = "GetControlValue", Varying))
 struct CONTROLRIG_API FRigUnit_GetControlVector : public FRigUnit
 {
 	GENERATED_BODY()
