Control Rig: Ensure to pass correct graph to repopulate, avoid invalid errors

diff --git a/Engine/Plugins/Experimental/ControlRig/Source/ControlRigDeveloper/Private/ControlRigBlueprint.cpp b/Engine/Plugins/Experimental/ControlRig/Source/ControlRigDeveloper/Private/ControlRigBlueprint.cpp
--- a/Engine/Plugins/Experimental/ControlRig/Source/ControlRigDeveloper/Private/ControlRigBlueprint.cpp
+++ b/Engine/Plugins/Experimental/ControlRig/Source/ControlRigDeveloper/Private/ControlRigBlueprint.cpp
@@ -313,7 +313,7 @@ void UControlRigBlueprint::PostLoad()
 			TArray<URigVMNode*> Nodes = GraphToDetach->GetNodes();
 			for (URigVMNode* Node : Nodes)
 			{
-				GetOrCreateController()->RepopulatePinsOnNode(Node);
+				GetOrCreateController(GraphToDetach)->RepopulatePinsOnNode(Node);
 			}
 		}
 		SetupPinRedirectorsForBackwardsCompatibility();
@@ -1252,6 +1252,8 @@ void UControlRigBlueprint::SetTransientControlValue(const FRigElementKey& InElem
 void UControlRigBlueprint::PopulateModelFromGraphForBackwardsCompatibility(UControlRigGraph* InGraph)
 {
 	DECLARE_SCOPE_HIERARCHICAL_COUNTER_FUNC()
+	
+	
 
 	// temporarily disable default value validation during load time, serialized values should always be accepted
 	TGuardValue<bool> DisablePinDefaultValueValidation(GetOrCreateController()->bValidatePinDefaults, false);
diff --git a/Engine/Source/Developer/RigVMDeveloper/Private/RigVMModel/RigVMController.cpp b/Engine/Source/Developer/RigVMDeveloper/Private/RigVMModel/RigVMController.cpp
--- a/Engine/Source/Developer/RigVMDeveloper/Private/RigVMModel/RigVMController.cpp
+++ b/Engine/Source/Developer/RigVMDeveloper/Private/RigVMModel/RigVMController.cpp
@@ -4810,14 +4810,14 @@ bool URigVMController::BindPinToVariable(URigVMPin* InPin, const FString& InNewB
 		return false;
 	}
 
-	if (GetGraph()->IsA<URigVMFunctionLibrary>())
+	if (InPin->GetBoundVariablePath() == InNewBoundVariablePath)
 	{
-		ReportError(TEXT("Cannot bind pins to variables in function library graphs."));
 		return false;
 	}
 
-	if (InPin->GetBoundVariablePath() == InNewBoundVariablePath)
+	if (GetGraph()->IsA<URigVMFunctionLibrary>())
 	{
+		ReportError(TEXT("Cannot bind pins to variables in function library graphs."));
 		return false;
 	}
 
