Control Rig: Introduce LRU cache for Name operations within the rig

diff --git a/Engine/Plugins/Experimental/ControlRig/Source/ControlRig/Private/ControlRig.cpp b/Engine/Plugins/Experimental/ControlRig/Source/ControlRig/Private/ControlRig.cpp
--- a/Engine/Plugins/Experimental/ControlRig/Source/ControlRig/Private/ControlRig.cpp
+++ b/Engine/Plugins/Experimental/ControlRig/Source/ControlRig/Private/ControlRig.cpp
@@ -484,6 +484,7 @@ void UControlRig::Execute(const EControlRigState InState, const FName& InEventNa
 	if (InState == EControlRigState::Init)
 	{
 		AbsoluteTime = DeltaTime = 0.f;
+		NameCache.Reset();
 	}
 
 	Context.DeltaTime = DeltaTime;
@@ -497,6 +498,7 @@ void UControlRig::Execute(const EControlRigState InState, const FName& InEventNa
 	Context.OwningComponent = nullptr;
 	Context.OwningActor = nullptr;
 	Context.World = nullptr;
+	Context.NameCache = &NameCache;
 
 	if (!OuterSceneComponent.IsValid())
 	{
diff --git a/Engine/Plugins/Experimental/ControlRig/Source/ControlRig/Private/Rigs/RigNameCache.cpp b/Engine/Plugins/Experimental/ControlRig/Source/ControlRig/Private/Rigs/RigNameCache.cpp
new file mode 100644
--- /dev/null
+++ b/Engine/Plugins/Experimental/ControlRig/Source/ControlRig/Private/Rigs/RigNameCache.cpp
@@ -0,0 +1,356 @@
+// Copyright Epic Games, Inc. All Rights Reserved.
+
+#include "Rigs/RigNameCache.h"
+
+#if WITH_EDITOR
+#include "ControlRig.h"
+TAutoConsoleVariable<int32> CVarRigNameCacheMaxSize(TEXT("a.ControlRig.NameCacheMaxSize"), 512, TEXT("Change to control how many names are cached per rig instance."));
+#endif
+
+////////////////////////////////////////////////////////////////////////////////
+// FRigNameOp
+////////////////////////////////////////////////////////////////////////////////
+
+FRigNameOp FRigNameOp::Concat(const FName& InA, const FName& InB)
+{
+	FRigNameOp Op;
+	Op.Type = ERigNameOp::Concat;
+	Op.A = InA.GetComparisonIndex().ToUnstableInt();
+	Op.B = InB.GetComparisonIndex().ToUnstableInt();
+	return Op;
+}
+
+FRigNameOp FRigNameOp::Left(const FName& InA, const uint32 InCount)
+{
+	FRigNameOp Op;
+	Op.Type = ERigNameOp::Left;
+	Op.A = InA.GetComparisonIndex().ToUnstableInt();
+	Op.B = InCount;
+	return Op;
+}
+
+FRigNameOp FRigNameOp::Right(const FName& InA, const uint32 InCount)
+{
+	FRigNameOp Op;
+	Op.Type = ERigNameOp::Right;
+	Op.A = InA.GetComparisonIndex().ToUnstableInt();
+	Op.B = InCount;
+	return Op;
+}
+
+FRigNameOp FRigNameOp::LeftChop(const FName& InA, const uint32 InCount)
+{
+	FRigNameOp Op;
+	Op.Type = ERigNameOp::LeftChop;
+	Op.A = InA.GetComparisonIndex().ToUnstableInt();
+	Op.B = InCount;
+	return Op;
+}
+
+FRigNameOp FRigNameOp::RightChop(const FName& InA, const uint32 InCount)
+{
+	FRigNameOp Op;
+	Op.Type = ERigNameOp::RightChop;
+	Op.A = InA.GetComparisonIndex().ToUnstableInt();
+	Op.B = InCount;
+	return Op;
+}
+
+FRigNameOp FRigNameOp::Replace(const FName& InA, const FName& InB, const FName& InC, const ESearchCase::Type InSearchCase)
+{
+	FRigNameOp Op;
+	Op.Type = InSearchCase == ESearchCase::CaseSensitive ? ERigNameOp::ReplaceCase : ERigNameOp::ReplaceNoCase;
+	Op.A = InA.GetComparisonIndex().ToUnstableInt();
+	Op.B = InB.GetComparisonIndex().ToUnstableInt();
+	Op.C = InC.GetComparisonIndex().ToUnstableInt();
+	return Op;
+}
+
+FRigNameOp FRigNameOp::EndsWith(const FName& InA, const FName& InB, const ESearchCase::Type InSearchCase)
+{
+	FRigNameOp Op;
+	Op.Type = InSearchCase == ESearchCase::CaseSensitive ? ERigNameOp::EndsWithCase : ERigNameOp::EndsWithNoCase;
+	Op.A = InA.GetComparisonIndex().ToUnstableInt();
+	Op.B = InB.GetComparisonIndex().ToUnstableInt();
+	return Op;
+}
+
+FRigNameOp FRigNameOp::StartsWith(const FName& InA, const FName& InB, const ESearchCase::Type InSearchCase)
+{
+	FRigNameOp Op;
+	Op.Type = InSearchCase == ESearchCase::CaseSensitive ? ERigNameOp::StartsWithCase : ERigNameOp::StartsWithNoCase;
+	Op.A = InA.GetComparisonIndex().ToUnstableInt();
+	Op.B = InB.GetComparisonIndex().ToUnstableInt();
+	return Op;
+}
+
+FRigNameOp FRigNameOp::Contains(const FName& InA, const FName& InB, const ESearchCase::Type InSearchCase)
+{
+	FRigNameOp Op;
+	Op.Type = InSearchCase == ESearchCase::CaseSensitive ? ERigNameOp::ContainsCase : ERigNameOp::ContainsNoCase;
+	Op.A = InA.GetComparisonIndex().ToUnstableInt();
+	Op.B = InB.GetComparisonIndex().ToUnstableInt();
+	return Op;
+}
+
+////////////////////////////////////////////////////////////////////////////////
+// FRigNameCache
+////////////////////////////////////////////////////////////////////////////////
+
+FName FRigNameCache::Concat(const FName& InA, const FName& InB)
+{
+	if(InA.IsNone())
+	{
+		return InB;
+	}
+
+	if(InB.IsNone())
+	{
+		return InA;
+	}
+	
+	const FRigNameOp Op = FRigNameOp::Concat(InA, InB);
+	if(const FName* ExistingName = NameCache.Find(Op))
+	{
+		return *ExistingName;
+	}
+	const FName NewName = *(InA.ToString() + InB.ToString());
+
+	if(CheckCacheSize())
+	{
+		NameCache.Add(Op, NewName);
+	}
+	return NewName;
+}
+
+FName FRigNameCache::Left(const FName& InA, const uint32 InCount)
+{
+	if(InCount == 0 || InA.IsNone())
+	{
+		return NAME_None;
+	}
+
+	if(InCount >= InA.GetStringLength())
+	{
+		return InA;
+	}
+	
+	const FRigNameOp Op = FRigNameOp::Left(InA, InCount);
+	if(const FName* ExistingName = NameCache.Find(Op))
+	{
+		return *ExistingName;
+	}
+	const FName NewName = *InA.ToString().Left(InCount);
+
+	if(CheckCacheSize())
+	{
+		NameCache.Add(Op, NewName);
+	}
+	
+	return NewName;
+}
+
+FName FRigNameCache::Right(const FName& InA, const uint32 InCount)
+{
+	if(InCount == 0 || InA.IsNone())
+	{
+		return NAME_None;
+	}
+
+	if(InCount >= InA.GetStringLength())
+	{
+		return InA;
+	}
+
+	const FRigNameOp Op = FRigNameOp::Right(InA, InCount);
+	if(const FName* ExistingName = NameCache.Find(Op))
+	{
+		return *ExistingName;
+	}
+	const FName NewName = *InA.ToString().Right(InCount);
+
+	if(CheckCacheSize())
+	{
+		NameCache.Add(Op, NewName);
+	}
+	
+	return NewName;
+}
+
+FName FRigNameCache::LeftChop(const FName& InA, const uint32 InCount)
+{
+	if(InCount == 0 || InA.IsNone())
+	{
+		return InA;
+	}
+	
+	const FRigNameOp Op = FRigNameOp::LeftChop(InA, InCount);
+	if(const FName* ExistingName = NameCache.Find(Op))
+	{
+		return *ExistingName;
+	}
+	const FName NewName = *InA.ToString().LeftChop(InCount);
+
+	if(CheckCacheSize())
+	{
+		NameCache.Add(Op, NewName);
+	}
+	
+	return NewName;
+}
+
+FName FRigNameCache::RightChop(const FName& InA, const uint32 InCount)
+{
+	if(InCount == 0 || InA.IsNone())
+	{
+		return InA;
+	}
+
+	const FRigNameOp Op = FRigNameOp::RightChop(InA, InCount);
+	if(const FName* ExistingName = NameCache.Find(Op))
+	{
+		return *ExistingName;
+	}
+	const FName NewName = *InA.ToString().RightChop(InCount);
+
+	if(CheckCacheSize())
+	{
+		NameCache.Add(Op, NewName);
+	}
+	
+	return NewName;
+}
+
+FName FRigNameCache::Replace(const FName& InA, const FName& InB, const FName& InC, const ESearchCase::Type InSearchCase)
+{
+	if(InA.IsNone() || InB.IsNone())
+	{
+		return InA;
+	}
+
+	const FRigNameOp Op = FRigNameOp::Replace(InA, InB, InC, InSearchCase);
+	if(const FName* ExistingName = NameCache.Find(Op))
+	{
+		return *ExistingName;
+	}
+
+	FString NewString(TEXT(""));
+	if (!InC.IsNone())
+	{
+		NewString = *InC.ToString();
+	}
+	const FName NewName = *InA.ToString().Replace(*InB.ToString(), *NewString, InSearchCase);
+
+	if(CheckCacheSize())
+	{
+		NameCache.Add(Op, NewName);
+	}
+	
+	return NewName;
+}
+
+bool FRigNameCache::EndsWith(const FName& InA, const FName& InB, const ESearchCase::Type InSearchCase)
+{
+	if(InA.IsNone() || InB.IsNone())
+	{
+		return false;
+	}
+	
+	const FRigNameOp Op = FRigNameOp::EndsWith(InA, InB, InSearchCase);
+	if(const bool* ExistingBool = BoolCache.Find(Op))
+	{
+		return *ExistingBool;
+	}
+	const bool bResult = InA.ToString().EndsWith(InB.ToString(), InSearchCase);
+
+	if(CheckCacheSize())
+	{
+		BoolCache.Add(Op, bResult);
+	}
+	
+	return bResult;
+}
+
+bool FRigNameCache::StartsWith(const FName& InA, const FName& InB, const ESearchCase::Type InSearchCase)
+{
+	if(InA.IsNone() || InB.IsNone())
+	{
+		return false;
+	}
+
+	const FRigNameOp Op = FRigNameOp::StartsWith(InA, InB, InSearchCase);
+	if(const bool* ExistingBool = BoolCache.Find(Op))
+	{
+		return *ExistingBool;
+	}
+	const bool bResult = InA.ToString().StartsWith(InB.ToString(), InSearchCase);
+
+	if(CheckCacheSize())
+	{
+		BoolCache.Add(Op, bResult);
+	}
+	
+	return bResult;
+}
+
+bool FRigNameCache::Contains(const FName& InA, const FName& InB, const ESearchCase::Type InSearchCase)
+{
+	if(InA.IsNone() || InB.IsNone())
+	{
+		return false;
+	}
+
+	const FRigNameOp Op = FRigNameOp::Contains(InA, InB, InSearchCase);
+	if(const bool* ExistingBool = BoolCache.Find(Op))
+	{
+		return *ExistingBool;
+	}
+	const bool bResult = InA.ToString().Contains(InB.ToString(), InSearchCase);
+
+	if(CheckCacheSize())
+	{
+		BoolCache.Add(Op, bResult);
+	}
+	
+	return bResult;
+}
+
+TArray<FRigNameOp> FRigNameCache::GetNameOps() const
+{
+	TArray<FRigNameOp> Result;
+	NameCache.GenerateKeyArray(Result);
+	return Result;
+}
+
+TArray<FName> FRigNameCache::GetNameValues() const
+{
+	TArray<FName> Result;
+	NameCache.GenerateValueArray(Result);
+	return Result;
+}
+
+TArray<FRigNameOp> FRigNameCache::GetBoolOps() const
+{
+	TArray<FRigNameOp> Result;
+	BoolCache.GenerateKeyArray(Result);
+	return Result;
+}
+
+TArray<bool> FRigNameCache::GetBoolValues() const
+{
+	TArray<bool> Result;
+	BoolCache.GenerateValueArray(Result);
+	return Result;
+}
+
+#if WITH_EDITOR
+bool FRigNameCache::CheckCacheSize() const
+{
+	if(CVarRigNameCacheMaxSize.GetValueOnAnyThread() == NameCache.Num() || CVarRigNameCacheMaxSize.GetValueOnAnyThread() == BoolCache.Num())
+	{
+		UE_LOG(LogControlRig, Warning, TEXT("FRigNameCache exceeded maximum size of %d. You can change it using the 'a.ControlRig.NameCacheMaxSize' console variable."), CVarRigNameCacheMaxSize.GetValueOnAnyThread());
+		return false;
+	}
+	return true;
+}
+#endif
diff --git a/Engine/Plugins/Experimental/ControlRig/Source/ControlRig/Private/Units/Core/RigUnit_Name.cpp b/Engine/Plugins/Experimental/ControlRig/Source/ControlRig/Private/Units/Core/RigUnit_Name.cpp
--- a/Engine/Plugins/Experimental/ControlRig/Source/ControlRig/Private/Units/Core/RigUnit_Name.cpp
+++ b/Engine/Plugins/Experimental/ControlRig/Source/ControlRig/Private/Units/Core/RigUnit_Name.cpp
@@ -5,75 +5,252 @@
 
 FRigUnit_NameConcat_Execute()
 {
-    DECLARE_SCOPE_HIERARCHICAL_COUNTER_RIGUNIT()
+	check(Context.NameCache);
+	
+	Result = Context.NameCache->Concat(A, B);
+}
 
-	if(A.IsNone())
-	{
-		Result = B;
-	}
-	else if(B.IsNone())
-	{
-		Result = A;
-	}
-	else
-	{
-		Result = *(A.ToString() + B.ToString());
-	}
+#if WITH_DEV_AUTOMATION_TESTS
+#include "Units/RigUnitTest.h"
+
+IMPLEMENT_RIGUNIT_AUTOMATION_TEST(FRigUnit_NameConcat)
+{
+	const FName Control(TEXT("Control"));
+	const FName Rig(TEXT("Rig"));
+	const FName ControlRig(TEXT("ControlRig"));
+	
+	Unit.A = Control;
+	Unit.B = Rig;
+	Execute();
+	AddErrorIfFalse(Unit.Result == ControlRig, TEXT("unexpected name concat result"));
+
+	Unit.A = Control;
+	Unit.B = NAME_None;
+	Execute();
+	AddErrorIfFalse(Unit.Result == Control, TEXT("unexpected name concat result"));
+
+	Unit.A = NAME_None;
+	Unit.B = Rig;
+	Execute();
+	AddErrorIfFalse(Unit.Result == Rig, TEXT("unexpected name concat result"));
+	return true;
 }
 
+#endif
+
 FRigUnit_NameTruncate_Execute()
 {
-    DECLARE_SCOPE_HIERARCHICAL_COUNTER_RIGUNIT()
+	check(Context.NameCache);
 
 	Remainder = Name;
 	Chopped = NAME_None;
 
-    if(Name.IsNone() || Count <= 0)
-    {
-    	return;
-    }
+	if(Name.IsNone() || Count <= 0)
+	{
+		return;
+	}
 
-	FString String = Name.ToString();
 	if (FromEnd)
 	{
-		Remainder = *String.LeftChop(Count);
-		Chopped = *String.Right(Count);
+		Remainder = Context.NameCache->LeftChop(Name, Count);
+		Chopped = Context.NameCache->Right(Name, Count);
 	}
 	else
 	{
-		Remainder = *String.RightChop(Count);
-		Chopped = *String.Left(Count);
+		Remainder = Context.NameCache->RightChop(Name, Count);
+		Chopped = Context.NameCache->Left(Name, Count);
 	}
 }
 
+#if WITH_DEV_AUTOMATION_TESTS
+#include "Units/RigUnitTest.h"
+
+IMPLEMENT_RIGUNIT_AUTOMATION_TEST(FRigUnit_NameTruncate)
+{
+	const FName Control(TEXT("Control"));
+	const FName Rig(TEXT("Rig"));
+	const FName ControlRig(TEXT("ControlRig"));
+	const FName Con(TEXT("Con"));
+	const FName trolRig(TEXT("trolRig"));
+	
+	Unit.Name = ControlRig;
+	Unit.Count = 3;
+	Unit.FromEnd = true;
+	Execute();
+	AddErrorIfFalse(Unit.Chopped == Rig, *FString::Printf(TEXT("unexpected name truncate result 1 '%s'"), *Unit.Chopped.ToString()));
+	AddErrorIfFalse(Unit.Remainder == Control, *FString::Printf(TEXT("unexpected name truncate result 2 '%s'"), *Unit.Remainder.ToString()));
+
+	Unit.Name = ControlRig;
+	Unit.Count = 7;
+	Unit.FromEnd = false;
+	Execute();
+	AddErrorIfFalse(Unit.Chopped == Control, *FString::Printf(TEXT("unexpected name truncate result 3 '%s'"), *Unit.Chopped.ToString()));
+	AddErrorIfFalse(Unit.Remainder == Rig, *FString::Printf(TEXT("unexpected name truncate result 4 '%s'"), *Unit.Remainder.ToString()));
+
+	Unit.Name = ControlRig;
+	Unit.Count = 3;
+	Unit.FromEnd = false;
+	Execute();
+	AddErrorIfFalse(Unit.Chopped == Con, *FString::Printf(TEXT("unexpected name truncate result 5 '%s'"), *Unit.Chopped.ToString()));
+	AddErrorIfFalse(Unit.Remainder == trolRig, *FString::Printf(TEXT("unexpected name truncate result 6 '%s'"), *Unit.Remainder.ToString()));
+
+	Unit.Name = ControlRig;
+	Unit.Count = 0;
+	Unit.FromEnd = false;
+	Execute();
+	AddErrorIfFalse(Unit.Chopped == NAME_None, *FString::Printf(TEXT("unexpected name truncate result 7 '%s'"), *Unit.Chopped.ToString()));
+	AddErrorIfFalse(Unit.Remainder == ControlRig, *FString::Printf(TEXT("unexpected name truncate result 8 '%s'"), *Unit.Remainder.ToString()));
+
+	Unit.Name = NAME_None;
+	Unit.Count = 6;
+	Unit.FromEnd = false;
+	Execute();
+	AddErrorIfFalse(Unit.Chopped == NAME_None, *FString::Printf(TEXT("unexpected name truncate result 9 '%s'"), *Unit.Chopped.ToString()));
+	AddErrorIfFalse(Unit.Remainder == NAME_None, *FString::Printf(TEXT("unexpected name truncate result 10 '%s'"), *Unit.Remainder.ToString()));
+	return true;
+}
+#endif
+
 FRigUnit_NameReplace_Execute()
 {
-	DECLARE_SCOPE_HIERARCHICAL_COUNTER_RIGUNIT()
+	check(Context.NameCache);
 
-	Result = Name;
+	Result = Context.NameCache->Replace(Name, Old, New, ESearchCase::CaseSensitive);
+}
 
-	if (!Old.IsNone())
-	{
-		FString NewString(TEXT(""));
-		if (!New.IsNone())
-		{
-			NewString = *New.ToString();
-		}
-		Result = *Name.ToString().Replace(*Old.ToString(), *NewString, ESearchCase::CaseSensitive);
-	}
+#if WITH_DEV_AUTOMATION_TESTS
+#include "Units/RigUnitTest.h"
+
+IMPLEMENT_RIGUNIT_AUTOMATION_TEST(FRigUnit_NameReplace)
+{
+	const FName Control(TEXT("Control"));
+	const FName Rig(TEXT("Rig"));
+	const FName ControlRig(TEXT("ControlRig"));
+	const FName Foo(TEXT("Foo"));
+	const FName FooRig(TEXT("FooRig"));
+	const FName ControlFoo(TEXT("ControlFoo"));
+	
+	Unit.Name = ControlRig;
+	Unit.Old = Control;
+	Unit.New = Foo;
+	Execute();
+	AddErrorIfFalse(Unit.Result == FooRig, TEXT("unexpected name replace result"));
+
+	Unit.Name = ControlRig;
+	Unit.Old = Rig;
+	Unit.New = Foo;
+	Execute();
+	AddErrorIfFalse(Unit.Result == ControlFoo, TEXT("unexpected name replace result"));
+
+	Unit.Name = ControlRig;
+	Unit.Old = NAME_None;
+	Unit.New = Foo;
+	Execute();
+	AddErrorIfFalse(Unit.Result == ControlRig, TEXT("unexpected name replace result"));
+
+	Unit.Name = ControlRig;
+	Unit.Old = Control;
+	Unit.New = NAME_None;
+	Execute();
+	AddErrorIfFalse(Unit.Result == Rig, TEXT("unexpected name replace result"));
+
+	return true;
 }
+#endif
 
 FRigUnit_EndsWith_Execute()
 {
-	Result = Name.ToString().EndsWith(Ending.ToString());
+	check(Context.NameCache);
+
+	Result = Context.NameCache->EndsWith(Name, Ending, ESearchCase::CaseSensitive);
+}
+
+#if WITH_DEV_AUTOMATION_TESTS
+#include "Units/RigUnitTest.h"
+
+IMPLEMENT_RIGUNIT_AUTOMATION_TEST(FRigUnit_EndsWith)
+{
+	const FName Control(TEXT("Control"));
+	const FName Rig(TEXT("Rig"));
+	const FName ControlRig(TEXT("ControlRig"));
+	
+	Unit.Name = ControlRig;
+	Unit.Ending = Control;
+	Execute();
+	AddErrorIfFalse(Unit.Result == false, TEXT("unexpected name endswith result"));
+
+	Unit.Name = ControlRig;
+	Unit.Ending = Rig;
+	Execute();
+	AddErrorIfFalse(Unit.Result == true, TEXT("unexpected name endswith result"));
+
+	return true;
 }
 
+#endif
+
 FRigUnit_StartsWith_Execute()
 {
-	Result = Name.ToString().StartsWith(Start.ToString());
+	check(Context.NameCache);
+
+	Result = Context.NameCache->StartsWith(Name, Start, ESearchCase::CaseSensitive);
+}
+
+#if WITH_DEV_AUTOMATION_TESTS
+#include "Units/RigUnitTest.h"
+
+IMPLEMENT_RIGUNIT_AUTOMATION_TEST(FRigUnit_StartsWith)
+{
+	const FName Control(TEXT("Control"));
+	const FName Rig(TEXT("Rig"));
+	const FName ControlRig(TEXT("ControlRig"));
+	
+	Unit.Name = ControlRig;
+	Unit.Start = Control;
+	Execute();
+	AddErrorIfFalse(Unit.Result == true, TEXT("unexpected name startswith result"));
+
+	Unit.Name = ControlRig;
+	Unit.Start = Rig;
+	Execute();
+	AddErrorIfFalse(Unit.Result == false, TEXT("unexpected name startswith result"));
+
+	return true;
 }
+#endif
 
 FRigUnit_Contains_Execute()
 {
-	Result = Name.ToString().Contains(Search.ToString());
-}
\ No newline at end of file
+	check(Context.NameCache);
+
+	Result = Context.NameCache->Contains(Name, Search, ESearchCase::CaseSensitive);
+}
+
+#if WITH_DEV_AUTOMATION_TESTS
+#include "Units/RigUnitTest.h"
+
+IMPLEMENT_RIGUNIT_AUTOMATION_TEST(FRigUnit_Contains)
+{
+	const FName Control(TEXT("Control"));
+	const FName Rig(TEXT("Rig"));
+	const FName Foo(TEXT("Foo"));
+	const FName ControlRig(TEXT("ControlRig"));
+	
+	Unit.Name = ControlRig;
+	Unit.Search = Control;
+	Execute();
+	AddErrorIfFalse(Unit.Result == true, TEXT("unexpected name contains result"));
+
+	Unit.Name = ControlRig;
+	Unit.Search = Rig;
+	Execute();
+	AddErrorIfFalse(Unit.Result == true, TEXT("unexpected name contains result"));
+
+	Unit.Name = ControlRig;
+	Unit.Search = Foo;
+	Execute();
+	AddErrorIfFalse(Unit.Result == false, TEXT("unexpected name contains result"));
+
+	return true;
+}
+#endif
\ No newline at end of file
diff --git a/Engine/Plugins/Experimental/ControlRig/Source/ControlRig/Public/ControlRig.h b/Engine/Plugins/Experimental/ControlRig/Source/ControlRig/Public/ControlRig.h
--- a/Engine/Plugins/Experimental/ControlRig/Source/ControlRig/Public/ControlRig.h
+++ b/Engine/Plugins/Experimental/ControlRig/Source/ControlRig/Public/ControlRig.h
@@ -442,6 +442,8 @@ private:
 	UPROPERTY()
 	TMap<FName, FCachedPropertyPath> OutputProperties_DEPRECATED;
 
+	FRigNameCache NameCache;
+
 private:
 	// Controls for the container
 	void HandleOnControlModified(UControlRig* Subject, const FRigControl& Control, const FRigControlModifiedContext& Context);
diff --git a/Engine/Plugins/Experimental/ControlRig/Source/ControlRig/Public/Rigs/RigNameCache.h b/Engine/Plugins/Experimental/ControlRig/Source/ControlRig/Public/Rigs/RigNameCache.h
new file mode 100644
--- /dev/null
+++ b/Engine/Plugins/Experimental/ControlRig/Source/ControlRig/Public/Rigs/RigNameCache.h
@@ -0,0 +1,141 @@
+// Copyright Epic Games, Inc. All Rights Reserved.
+
+#pragma once
+
+#include "CoreMinimal.h"
+
+enum class ERigNameOp : uint8
+{
+	None,
+	Concat,
+	Left,
+    Right,
+    LeftChop,
+	RightChop,
+	ReplaceCase,
+	ReplaceNoCase,
+	EndsWithCase,
+	EndsWithNoCase,
+    StartsWithCase,
+    StartsWithNoCase,
+    ContainsCase,
+    ContainsNoCase
+};
+
+struct CONTROLRIG_API FRigNameOp
+{
+public:
+
+	FORCEINLINE FRigNameOp()
+	: A(INDEX_NONE)
+	, B(INDEX_NONE)
+	, C(INDEX_NONE)
+	, Type(ERigNameOp::None)
+	{}
+
+	static FRigNameOp Concat(const FName& InA, const FName& InB);
+	static FRigNameOp Left(const FName& InA, const uint32 InCount);
+	static FRigNameOp Right(const FName& InA, const uint32 InCount);
+	static FRigNameOp LeftChop(const FName& InA, const uint32 InCount);
+	static FRigNameOp RightChop(const FName& InA, const uint32 InCount);
+	static FRigNameOp Replace(const FName& InA, const FName& InB, const FName& InC, const ESearchCase::Type InSearchCase);
+    static FRigNameOp EndsWith(const FName& InA, const FName& InB, const ESearchCase::Type InSearchCase);
+    static FRigNameOp StartsWith(const FName& InA, const FName& InB, const ESearchCase::Type InSearchCase);
+    static FRigNameOp Contains(const FName& InA, const FName& InB, const ESearchCase::Type InSearchCase);
+
+	friend FORCEINLINE uint32 GetTypeHash(const FRigNameOp& Op)
+	{
+		return HashCombine(Op.A, HashCombine(Op.B, HashCombine(Op.C, (uint32)Op.Type)));
+	}
+
+	FORCEINLINE bool operator ==(const FRigNameOp& Other) const
+	{
+		return Type == Other.Type && A == Other.A && B == Other.B && B == Other.B && C == Other.C;
+	}
+
+	FORCEINLINE bool operator !=(const FRigNameOp& Other) const
+	{
+		return Type != Other.Type || A != Other.A || B != Other.B || B != Other.B || C != Other.C;
+	}
+
+	FORCEINLINE bool operator <(const FRigNameOp& Other) const
+	{
+		if (Type < Other.Type)
+		{
+			return true;
+		}
+		if(A < Other.A)
+		{
+			return true;
+		}
+		if(B < Other.B)
+		{
+			return true;
+		}
+		return C < Other.C;
+	}
+
+	FORCEINLINE bool operator >(const FRigNameOp& Other) const
+	{
+		if (Type > Other.Type)
+		{
+			return true;
+		}
+		if(A > Other.A)
+		{
+			return true;
+		}
+		if(B > Other.B)
+		{
+			return true;
+		}
+		return C > Other.C;
+	}
+
+
+private:
+
+	uint32 A;
+	uint32 B;
+	uint32 C;
+	ERigNameOp Type;
+};
+
+class CONTROLRIG_API FRigNameCache
+{
+public:
+
+	FORCEINLINE void Reset()
+	{
+		NameCache.Reset();
+		BoolCache.Reset();
+	}
+
+	FName Concat(const FName& InA, const FName& InB);
+	FName Left(const FName& InA, const uint32 InCount);
+	FName Right(const FName& InA, const uint32 InCount);
+	FName LeftChop(const FName& InA, const uint32 InCount);
+	FName RightChop(const FName& InA, const uint32 InCount);
+	FName Replace(const FName& InA, const FName& InB, const FName& InC, const ESearchCase::Type InSearchCase);
+	bool EndsWith(const FName& InA, const FName& InB, const ESearchCase::Type InSearchCase);
+	bool StartsWith(const FName& InA, const FName& InB, const ESearchCase::Type InSearchCase);
+	bool Contains(const FName& InA, const FName& InB, const ESearchCase::Type InSearchCase);
+
+	FORCEINLINE const TMap<FRigNameOp, FName>& GetNameCache() const { return NameCache; }
+	FORCEINLINE const TMap<FRigNameOp, bool>&  GetBoolCache() const { return BoolCache; }
+	TArray<FRigNameOp> GetNameOps() const;
+	TArray<FName> GetNameValues() const;
+	TArray<FRigNameOp> GetBoolOps() const;
+	TArray<bool> GetBoolValues() const;
+
+private:
+
+#if WITH_EDITOR
+	bool CheckCacheSize() const;
+#else
+	FORCEINLINE bool CheckCacheSize() const { return true; } 
+#endif
+
+	TMap<FRigNameOp, FName> NameCache;
+	TMap<FRigNameOp, bool> BoolCache;
+};
diff --git a/Engine/Plugins/Experimental/ControlRig/Source/ControlRig/Public/Units/RigUnitContext.h b/Engine/Plugins/Experimental/ControlRig/Source/ControlRig/Public/Units/RigUnitContext.h
--- a/Engine/Plugins/Experimental/ControlRig/Source/ControlRig/Public/Units/RigUnitContext.h
+++ b/Engine/Plugins/Experimental/ControlRig/Source/ControlRig/Public/Units/RigUnitContext.h
@@ -5,6 +5,7 @@
 #include "CoreMinimal.h"
 #include "Rigs/RigHierarchyContainer.h"
 #include "Rigs/RigCurveContainer.h"
+#include "Rigs/RigNameCache.h"
 #include "ControlRigLog.h"
 #include "AnimationDataSource.h"
 #include "Drawing/ControlRigDrawInterface.h"
@@ -42,6 +43,7 @@ struct FRigUnitContext
 #if WITH_EDITOR
 		, Log(nullptr)
 #endif
+		, NameCache(nullptr)
 	{
 	}
 
@@ -89,6 +91,9 @@ struct FRigUnitContext
 	FControlRigLog* Log;
 #endif
 
+	/** A container to store all names */
+	FRigNameCache* NameCache;
+
 	FORCEINLINE const FRigBoneHierarchy* GetBones() const
 	{
 		if (Hierarchy != nullptr)
diff --git a/Engine/Plugins/Experimental/ControlRig/Source/ControlRig/Public/Units/RigUnitTest.h b/Engine/Plugins/Experimental/ControlRig/Source/ControlRig/Public/Units/RigUnitTest.h
--- a/Engine/Plugins/Experimental/ControlRig/Source/ControlRig/Public/Units/RigUnitTest.h
+++ b/Engine/Plugins/Experimental/ControlRig/Source/ControlRig/Public/Units/RigUnitTest.h
@@ -21,6 +21,7 @@ public:
 	{
 		Context.Hierarchy = &HierarchyContainer;
 		ExecuteContext.Hierarchy= &HierarchyContainer;
+		Context.NameCache = &NameCache;
 	}
 
 	FRigHierarchyContainer HierarchyContainer;
@@ -29,6 +30,7 @@ public:
 	FRigControlHierarchy& ControlHierarchy;
 	FRigCurveContainer& CurveContainer;
 	FControlRigExecuteContext ExecuteContext;
+	FRigNameCache NameCache;
 	FRigUnitContext Context;
 };
 
