UE-107811 - Fix static analysis issue with duplicate conditions

diff --git a/Engine/Plugins/Experimental/ControlRig/Source/ControlRig/Public/Rigs/RigNameCache.h b/Engine/Plugins/Experimental/ControlRig/Source/ControlRig/Public/Rigs/RigNameCache.h
--- a/Engine/Plugins/Experimental/ControlRig/Source/ControlRig/Public/Rigs/RigNameCache.h
+++ b/Engine/Plugins/Experimental/ControlRig/Source/ControlRig/Public/Rigs/RigNameCache.h
@@ -50,12 +50,12 @@ public:
 
 	FORCEINLINE bool operator ==(const FRigNameOp& Other) const
 	{
-		return Type == Other.Type && A == Other.A && B == Other.B && B == Other.B && C == Other.C;
+		return Type == Other.Type && A == Other.A && B == Other.B && C == Other.C;
 	}
 
 	FORCEINLINE bool operator !=(const FRigNameOp& Other) const
 	{
-		return Type != Other.Type || A != Other.A || B != Other.B || B != Other.B || C != Other.C;
+		return Type != Other.Type || A != Other.A || B != Other.B || C != Other.C;
 	}
 
 	FORCEINLINE bool operator <(const FRigNameOp& Other) const
