ControlRig: Removed redundant equality check.

#trivial
#rnx

diff --git a/Engine/Plugins/Experimental/ControlRig/Source/ControlRigEditor/Private/ControlRigEditorModule.cpp b/Engine/Plugins/Experimental/ControlRig/Source/ControlRigEditor/Private/ControlRigEditorModule.cpp
--- a/Engine/Plugins/Experimental/ControlRig/Source/ControlRigEditor/Private/ControlRigEditorModule.cpp
+++ b/Engine/Plugins/Experimental/ControlRig/Source/ControlRigEditor/Private/ControlRigEditorModule.cpp
@@ -1631,7 +1631,7 @@ void FControlRigEditorModule::GetContextMenuActions(const UControlRigGraphSchema
 									}
 								}
 							}
-							else if (Pin->GetCPPTypeObject() == FRigElementKeyCollection::StaticStruct() && Pin->GetDirection() == ERigVMPinDirection::Output && StructMemory != nullptr)
+							else if (Pin->GetCPPTypeObject() == FRigElementKeyCollection::StaticStruct() && Pin->GetDirection() == ERigVMPinDirection::Output)
 							{
 								if (StructMemory == nullptr)
 								{
