Fixed bug, imported hierarchy stripped "-" characters from bone names breaking skeletal meshes that use them.
JIRA: UE-107965

diff --git a/Engine/Plugins/Experimental/ControlRig/Source/ControlRig/Private/Rigs/RigHierarchyContainer.cpp b/Engine/Plugins/Experimental/ControlRig/Source/ControlRig/Private/Rigs/RigHierarchyContainer.cpp
--- a/Engine/Plugins/Experimental/ControlRig/Source/ControlRig/Private/Rigs/RigHierarchyContainer.cpp
+++ b/Engine/Plugins/Experimental/ControlRig/Source/ControlRig/Private/Rigs/RigHierarchyContainer.cpp
@@ -1495,7 +1495,7 @@ void FRigHierarchyContainer::SanitizeName(FString& InOutName)
 
 		const bool bGoodChar =
 			((C >= 'A') && (C <= 'Z')) || ((C >= 'a') && (C <= 'z')) ||		// A-Z (upper and lowercase) anytime
-			(C == '_') ||													// _ anytime
+			(C == '_') || (C == '-') ||										// _ or - anytime
 			((i > 0) && (C >= '0') && (C <= '9'));							// 0-9 after the first character
 
 		if (!bGoodChar)
