Added a missing include to fix non-unity builds.

#trivial
#rnx

diff --git a/Engine/Plugins/Experimental/ControlRig/Source/ControlRigEditor/Private/Editor/RigHierarchyTabSummoner.cpp b/Engine/Plugins/Experimental/ControlRig/Source/ControlRigEditor/Private/Editor/RigHierarchyTabSummoner.cpp
--- a/Engine/Plugins/Experimental/ControlRig/Source/ControlRigEditor/Private/Editor/RigHierarchyTabSummoner.cpp
+++ b/Engine/Plugins/Experimental/ControlRig/Source/ControlRigEditor/Private/Editor/RigHierarchyTabSummoner.cpp
@@ -4,6 +4,7 @@
 #include "SRigHierarchy.h"
 #include "ControlRigEditorStyle.h"
 #include "ControlRigEditor.h"
+#include "Widgets/Docking/SDockTab.h"
 
 #define LOCTEXT_NAMESPACE "RigHierarchyTabSummoner"
 
