Added an explcit include of Engine/SkeletalMesh.h for non-editor builds.

#trivial
#rnx

diff --git a/Engine/Plugins/Experimental/ControlRig/Source/ControlRig/Private/ControlRig.cpp b/Engine/Plugins/Experimental/ControlRig/Source/ControlRig/Private/ControlRig.cpp
--- a/Engine/Plugins/Experimental/ControlRig/Source/ControlRig/Private/ControlRig.cpp
+++ b/Engine/Plugins/Experimental/ControlRig/Source/ControlRig/Private/ControlRig.cpp
@@ -1,6 +1,7 @@
 // Copyright Epic Games, Inc. All Rights Reserved.
 
 #include "ControlRig.h"
+#include "Engine/SkeletalMesh.h"
 #include "GameFramework/Actor.h"
 #include "Misc/RuntimeErrors.h"
 #include "IControlRigObjectBinding.h"
