Fix crash in ControlRigComponent from accessing stale memory.

diff --git a/Engine/Plugins/Experimental/ControlRig/Source/ControlRig/Private/ControlRigComponent.cpp b/Engine/Plugins/Experimental/ControlRig/Source/ControlRig/Private/ControlRigComponent.cpp
--- a/Engine/Plugins/Experimental/ControlRig/Source/ControlRig/Private/ControlRigComponent.cpp
+++ b/Engine/Plugins/Experimental/ControlRig/Source/ControlRig/Private/ControlRigComponent.cpp
@@ -26,13 +26,13 @@ TMap<FString, TSharedPtr<SNotificationItem>> UControlRigComponent::EditorNotific
 
 struct FSkeletalMeshToMap
 {
-	USkeletalMeshComponent* SkeletalMeshComponent;
+	TWeakObjectPtr<USkeletalMeshComponent> SkeletalMeshComponent;
 	TArray<FControlRigComponentMappedBone> Bones;
 	TArray<FControlRigComponentMappedCurve> Curves;
 };
 
 FCriticalSection gPendingSkeletalMeshesLock;
-TMap<UControlRigComponent*, TArray<FSkeletalMeshToMap> > gPendingSkeletalMeshes;
+TMap<TWeakObjectPtr<UControlRigComponent>, TArray<FSkeletalMeshToMap> > gPendingSkeletalMeshes;
 
 UControlRigComponent::UControlRigComponent(const FObjectInitializer& ObjectInitializer)
 	: Super(ObjectInitializer)
@@ -70,6 +70,14 @@ void UControlRigComponent::PostEditChangeProperty(FPropertyChangedEvent& Propert
 }
 #endif
 
+void UControlRigComponent::BeginDestroy()
+{
+	Super::BeginDestroy();
+
+	FScopeLock Lock(&gPendingSkeletalMeshesLock);
+	gPendingSkeletalMeshes.Remove(this);
+}
+
 void UControlRigComponent::OnRegister()
 {
 	Super::OnRegister();
@@ -145,7 +153,7 @@ void UControlRigComponent::TickComponent(float DeltaTime, enum ELevelTick TickTy
 			for (const FSkeletalMeshToMap& SkeletalMeshToMap : *PendingSkeletalMeshes)
 			{
 				AddMappedSkeletalMesh(
-					SkeletalMeshToMap.SkeletalMeshComponent,
+					SkeletalMeshToMap.SkeletalMeshComponent.Get(),
 					SkeletalMeshToMap.Bones,
 					SkeletalMeshToMap.Curves
 				);
diff --git a/Engine/Plugins/Experimental/ControlRig/Source/ControlRig/Public/ControlRigComponent.h b/Engine/Plugins/Experimental/ControlRig/Source/ControlRig/Public/ControlRigComponent.h
--- a/Engine/Plugins/Experimental/ControlRig/Source/ControlRig/Public/ControlRigComponent.h
+++ b/Engine/Plugins/Experimental/ControlRig/Source/ControlRig/Public/ControlRigComponent.h
@@ -202,15 +202,19 @@ public:
 	UPROPERTY(BlueprintAssignable, Category = "ControlRig", meta = (DisplayName = "On Post Forwards Solve"))
 	FControlRigComponentDelegate OnPostUpdateDelegate;
 
-	// UObject interface
+	//~ Begin UObject interface
 #if WITH_EDITOR
 	virtual void PostEditChangeProperty(FPropertyChangedEvent& PropertyChangedEvent) override;
 #endif
 
-	// UActorComponent interface
+	virtual void BeginDestroy() override;
+	//~ End UObject Interface
+	
+	//~ Begin UActorComponent Interface
 	virtual void OnRegister() override;
 	virtual void OnUnregister() override;
 	virtual void TickComponent(float DeltaTime, enum ELevelTick TickType, FActorComponentTickFunction* ThisTickFunction) override;
+	//~ End UActorComponent Interface.
 
 	//~ Begin UPrimitiveComponent Interface.
 	virtual FPrimitiveSceneProxy* CreateSceneProxy() override;
