Sequencer: Control Rig: Reconstruct Channels on PostLoad if not intialized. Don't need that binding check at all.

diff --git a/Engine/Plugins/Experimental/ControlRig/Source/ControlRig/Private/Sequencer/MovieSceneControlRigParameterTrack.cpp b/Engine/Plugins/Experimental/ControlRig/Source/ControlRig/Private/Sequencer/MovieSceneControlRigParameterTrack.cpp
--- a/Engine/Plugins/Experimental/ControlRig/Source/ControlRig/Private/Sequencer/MovieSceneControlRigParameterTrack.cpp
+++ b/Engine/Plugins/Experimental/ControlRig/Source/ControlRig/Private/Sequencer/MovieSceneControlRigParameterTrack.cpp
@@ -319,7 +319,7 @@ UMovieSceneSection* UMovieSceneControlRigParameterTrack::GetSectionToKey() const
 void UMovieSceneControlRigParameterTrack::PostLoad()
 {
 	Super::PostLoad();
-	if (ControlRig && ControlRig->GetObjectBinding() && !ControlRig->HasAnyFlags(RF_NeedLoad | RF_NeedPostLoad | RF_NeedInitialization))
+	if (ControlRig  && !ControlRig->HasAnyFlags(RF_NeedLoad | RF_NeedPostLoad | RF_NeedInitialization))
 	{
 		ControlRig->ConditionalPostLoad();
 		ControlRig->Initialize();
