Sequencer: Control Rig: Don't Key on PostEditUndo

diff --git a/Engine/Plugins/Experimental/ControlRig/Source/ControlRigEditor/Private/EditMode/ControlRigControlsProxy.cpp b/Engine/Plugins/Experimental/ControlRig/Source/ControlRigEditor/Private/EditMode/ControlRigControlsProxy.cpp
--- a/Engine/Plugins/Experimental/ControlRig/Source/ControlRigEditor/Private/EditMode/ControlRigControlsProxy.cpp
+++ b/Engine/Plugins/Experimental/ControlRig/Source/ControlRigEditor/Private/EditMode/ControlRigControlsProxy.cpp
@@ -171,7 +171,7 @@ void UControlRigTransformControlProxy::PostEditUndo()
 			CheckEditModeOnSelectionChange(ControlRig.Get());
 		}
 		ControlRig->SelectControl(ControlName, bSelected);
-		ControlRig->SetControlValue<FTransform>(ControlName, Transform, true, EControlRigSetKey::DoNotCare);
+		ControlRig->SetControlValue<FTransform>(ControlName, Transform, true, EControlRigSetKey::Never);
 	}
 }
 #endif
@@ -224,7 +224,7 @@ void UControlRigTransformNoScaleControlProxy::PostEditUndo()
 			CheckEditModeOnSelectionChange(ControlRig.Get());
 		}
 		ControlRig->SelectControl(ControlName, bSelected);
-		ControlRig->SetControlValue<FTransformNoScale>(ControlName, Transform, true, EControlRigSetKey::DoNotCare);
+		ControlRig->SetControlValue<FTransformNoScale>(ControlName, Transform, true, EControlRigSetKey::Never);
 	}
 }
 #endif
@@ -279,7 +279,7 @@ void UControlRigEulerTransformControlProxy::PostEditUndo()
 			CheckEditModeOnSelectionChange(ControlRig.Get());
 		}
 		ControlRig->SelectControl(ControlName, bSelected);
-		ControlRig->SetControlValue<FEulerTransform>(ControlName, Transform, true, EControlRigSetKey::DoNotCare);
+		ControlRig->SetControlValue<FEulerTransform>(ControlName, Transform, true, EControlRigSetKey::Never);
 	}
 }
 #endif
@@ -331,7 +331,7 @@ void UControlRigFloatControlProxy::PostEditUndo()
 			CheckEditModeOnSelectionChange(ControlRig.Get());
 		}
 		ControlRig->SelectControl(ControlName, bSelected);
-		ControlRig->SetControlValue<float>(ControlName, Float, true, EControlRigSetKey::DoNotCare);
+		ControlRig->SetControlValue<float>(ControlName, Float, true, EControlRigSetKey::Never);
 	}
 }
 #endif
@@ -383,7 +383,7 @@ void UControlRigIntegerControlProxy::PostEditUndo()
 			CheckEditModeOnSelectionChange(ControlRig.Get());
 		}
 		ControlRig->SelectControl(ControlName, bSelected);
-		ControlRig->SetControlValue<int32>(ControlName, Integer, true, EControlRigSetKey::DoNotCare);
+		ControlRig->SetControlValue<int32>(ControlName, Integer, true, EControlRigSetKey::Never);
 	}
 }
 #endif
@@ -438,7 +438,7 @@ void UControlRigEnumControlProxy::PostEditUndo()
 			CheckEditModeOnSelectionChange(ControlRig.Get());
 		}
 		ControlRig->SelectControl(ControlName, bSelected);
-		ControlRig->SetControlValue<int32>(ControlName, Enum.EnumIndex, true, EControlRigSetKey::DoNotCare);
+		ControlRig->SetControlValue<int32>(ControlName, Enum.EnumIndex, true, EControlRigSetKey::Never);
 	}
 }
 #endif
@@ -491,7 +491,7 @@ void UControlRigVectorControlProxy::PostEditUndo()
 			CheckEditModeOnSelectionChange(ControlRig.Get());
 		}
 		ControlRig->SelectControl(ControlName, bSelected);
-		ControlRig->SetControlValue<FVector>(ControlName, Vector, true, EControlRigSetKey::DoNotCare);
+		ControlRig->SetControlValue<FVector>(ControlName, Vector, true, EControlRigSetKey::Never);
 	}
 }
 #endif
@@ -545,7 +545,7 @@ void UControlRigVector2DControlProxy::PostEditUndo()
 			CheckEditModeOnSelectionChange(ControlRig.Get());
 		}
 		ControlRig->SelectControl(ControlName, bSelected);
-		ControlRig->SetControlValue<FVector2D>(ControlName, Vector2D, true, EControlRigSetKey::DoNotCare);
+		ControlRig->SetControlValue<FVector2D>(ControlName, Vector2D, true, EControlRigSetKey::Never);
 	}
 }
 #endif
@@ -598,7 +598,7 @@ void UControlRigBoolControlProxy::PostEditUndo()
 			CheckEditModeOnSelectionChange(ControlRig.Get());
 		}
 		ControlRig->SelectControl(ControlName, bSelected);
-		ControlRig->SetControlValue<bool>(ControlName, Bool, true, EControlRigSetKey::DoNotCare);
+		ControlRig->SetControlValue<bool>(ControlName, Bool, true, EControlRigSetKey::Never);
 	}
 }
 #endif
