Control Rig: New rig hierarchy data structure refactoring

diff --git a/Engine/Plugins/Animation/RigLogic/Source/RigLogicModule/Private/RigUnit_RigLogic.cpp b/Engine/Plugins/Animation/RigLogic/Source/RigLogicModule/Private/RigUnit_RigLogic.cpp
--- a/Engine/Plugins/Animation/RigLogic/Source/RigLogicModule/Private/RigUnit_RigLogic.cpp
+++ b/Engine/Plugins/Animation/RigLogic/Source/RigLogicModule/Private/RigUnit_RigLogic.cpp
@@ -67,7 +67,7 @@ FRigUnit_RigLogic_Data& FRigUnit_RigLogic_Data::operator=(const FRigUnit_RigLogi
 	HierarchyBoneIndices = Other.HierarchyBoneIndices;
 	MorphTargetCurveIndices = Other.MorphTargetCurveIndices;
 	BlendShapeIndices = Other.BlendShapeIndices;
-	CurveContainerIndicesForAnimMaps = Other.CurveContainerIndicesForAnimMaps;
+	CurveElementIndicesForAnimMaps = Other.CurveElementIndicesForAnimMaps;
 	RigLogicIndicesForAnimMaps = Other.RigLogicIndicesForAnimMaps;
 	CurrentLOD = Other.CurrentLOD;
 	return *this;
@@ -78,7 +78,7 @@ bool FRigUnit_RigLogic_Data::IsRigLogicInitialized()
 	return (SharedRigRuntimeContext != nullptr) && SharedRigRuntimeContext->RigLogic.IsValid() && RigInstance.IsValid();
 }
 
-void FRigUnit_RigLogic_Data::InitializeRigLogic(const FRigBoneHierarchy* BoneHierarchy, const FRigCurveContainer* CurveContainer)
+void FRigUnit_RigLogic_Data::InitializeRigLogic(const URigHierarchy* InHierarchy)
 {
 	if (!SharedRigRuntimeContext->BehaviorReader.IsValid() || SharedRigRuntimeContext->BehaviorReader->GetJointCount() == 0u)
 	{
@@ -98,15 +98,15 @@ void FRigUnit_RigLogic_Data::InitializeRigLogic(const FRigBoneHierarchy* BoneHie
 		RigInstance = MakeUnique<FRigInstance>(SharedRigRuntimeContext->RigLogic.Get());
 		RigInstance->SetLOD(CurrentLOD);
 
-		MapJoints(BoneHierarchy);
-		MapInputCurveIndices(CurveContainer);
-		MapMorphTargets(CurveContainer);
-		MapMaskMultipliers(CurveContainer);
+		MapJoints(InHierarchy);
+		MapInputCurveIndices(InHierarchy);
+		MapMorphTargets(InHierarchy);
+		MapMaskMultipliers(InHierarchy);
 	}
 }
 
 //maps indices of input curves from dna file to control rig curves
-void FRigUnit_RigLogic_Data::MapInputCurveIndices(const FRigCurveContainer* CurveContainer)
+void FRigUnit_RigLogic_Data::MapInputCurveIndices(const URigHierarchy* InHierarchy)
 {
 	const IBehaviorReader* DNABehavior = SharedRigRuntimeContext->BehaviorReader.Get();
 	const uint32 ControlCount = DNABehavior->GetRawControlCount();
@@ -120,12 +120,12 @@ void FRigUnit_RigLogic_Data::MapInputCurveIndices(const FRigCurveContainer* Curv
 			return;
 		}
 		const FName ControlFName(*AnimatedControlName);
-		const int32 CurveIndex = CurveContainer ? CurveContainer->GetIndex(ControlFName) : INDEX_NONE;
+		const int32 CurveIndex = InHierarchy ? InHierarchy->GetIndex(FRigElementKey(ControlFName, ERigElementType::Curve)) : INDEX_NONE;
 		InputCurveIndices.Add(CurveIndex); //can be INDEX_NONE
 	}
 }
 
-void FRigUnit_RigLogic_Data::MapJoints(const FRigBoneHierarchy* Hierarchy)
+void FRigUnit_RigLogic_Data::MapJoints(const URigHierarchy* Hierarchy)
 {
 	const IBehaviorReader* DNABehavior = SharedRigRuntimeContext->BehaviorReader.Get();
 	const uint16 JointCount = DNABehavior->GetJointCount();
@@ -134,7 +134,7 @@ void FRigUnit_RigLogic_Data::MapJoints(const FRigBoneHierarchy* Hierarchy)
 	{
 		const FString RLJointName = DNABehavior->GetJointName(JointIndex);
 		const FName JointFName = FName(*RLJointName);
-		const int32 BoneIndex = Hierarchy->GetIndex(JointFName);
+		const int32 BoneIndex = Hierarchy->GetIndex(FRigElementKey(JointFName, ERigElementType::Bone));
 		HierarchyBoneIndices.Add(BoneIndex);
 	}
 }
@@ -160,7 +160,7 @@ void FRigUnit_RigLogic_Data::CacheVariableJointIndices()
 	}
 }
 
-void FRigUnit_RigLogic_Data::MapMorphTargets(const FRigCurveContainer* CurveContainer)
+void FRigUnit_RigLogic_Data::MapMorphTargets(const URigHierarchy* InHierarchy)
 {
 	const IBehaviorReader* DNABehavior = SharedRigRuntimeContext->BehaviorReader.Get();
 	const uint16 LODCount = DNABehavior->GetLODCount();
@@ -185,19 +185,19 @@ void FRigUnit_RigLogic_Data::MapMorphTargets(const FRigCurveContainer* CurveCont
 			const FString MeshStr = DNABehavior->GetMeshName(MeshIndex);
 			const FString MorphTargetStr = MeshStr + TEXT("__") + BlendShapeStr;
 			const FName MorphTargetName(*MorphTargetStr);
-			const int32 MorphTargetIndex = CurveContainer->GetIndex(MorphTargetName);
+			const int32 MorphTargetIndex = InHierarchy->GetIndex(FRigElementKey(MorphTargetName, ERigElementType::Curve));
 			MorphTargetCurveIndices[LodIndex].Values.Add(MorphTargetIndex);
 			BlendShapeIndices[LodIndex].Values.Add(Mapping.BlendShapeChannelIndex);
 		}
 	}
 }
 
-void FRigUnit_RigLogic_Data::MapMaskMultipliers(const FRigCurveContainer* CurveContainer)
+void FRigUnit_RigLogic_Data::MapMaskMultipliers(const URigHierarchy* InHierarchy)
 {
 	const IBehaviorReader* DNABehavior = SharedRigRuntimeContext->BehaviorReader.Get();
 	const uint16 LODCount = DNABehavior->GetLODCount();
-	CurveContainerIndicesForAnimMaps.Reset();
-	CurveContainerIndicesForAnimMaps.AddDefaulted(LODCount);
+	CurveElementIndicesForAnimMaps.Reset();
+	CurveElementIndicesForAnimMaps.AddDefaulted(LODCount);
 
 	RigLogicIndicesForAnimMaps.Reset();
 	RigLogicIndicesForAnimMaps.AddDefaulted(LODCount);
@@ -205,7 +205,7 @@ void FRigUnit_RigLogic_Data::MapMaskMultipliers(const FRigCurveContainer* CurveC
 	for (uint16 LodIndex = 0; LodIndex < LODCount; ++LodIndex)
 	{		
 		TArrayView<const uint16> AnimMapIndicesPerLOD = DNABehavior->GetAnimatedMapIndicesForLOD(LodIndex);
-		CurveContainerIndicesForAnimMaps[LodIndex].Values.Reserve(AnimMapIndicesPerLOD.Num());
+		CurveElementIndicesForAnimMaps[LodIndex].Values.Reserve(AnimMapIndicesPerLOD.Num());
 		RigLogicIndicesForAnimMaps[LodIndex].Values.Reserve(AnimMapIndicesPerLOD.Num());
 		for (uint16 AnimMapIndexPerLOD: AnimMapIndicesPerLOD)
 		{
@@ -215,14 +215,14 @@ void FRigUnit_RigLogic_Data::MapMaskMultipliers(const FRigCurveContainer* CurveC
 				return;
 			}
 			const FName MaskMultiplierFName(*MaskMultiplierNameStr);
-			const int32 CurveIndex = CurveContainer->GetIndex(MaskMultiplierFName);
-			CurveContainerIndicesForAnimMaps[LodIndex].Values.Add(CurveIndex); //can be INDEX_NONE if curve was not found
+			const int32 CurveIndex = InHierarchy->GetIndex(FRigElementKey(MaskMultiplierFName, ERigElementType::Curve));
+			CurveElementIndicesForAnimMaps[LodIndex].Values.Add(CurveIndex); //can be INDEX_NONE if curve was not found
 			RigLogicIndicesForAnimMaps[LodIndex].Values.Add(AnimMapIndexPerLOD);
 		}
 	}
 }
 
-void FRigUnit_RigLogic_Data::CalculateRigLogic(const FRigCurveContainer* CurveContainer)
+void FRigUnit_RigLogic_Data::CalculateRigLogic(const URigHierarchy* InHierarchy)
 {
 	QUICK_SCOPE_CYCLE_COUNTER(STAT_RigUnit_RigLogic_Calculate);
 
@@ -233,13 +233,13 @@ void FRigUnit_RigLogic_Data::CalculateRigLogic(const FRigCurveContainer* CurveCo
 	for (int32 ControlIndex = 0; ControlIndex < RawControlCount; ++ControlIndex)
 	{
 		const uint32 CurveIndex = InputCurveIndices[ControlIndex];
-		const float Value = CurveContainer->GetValue(CurveIndex);
+		const float Value = InHierarchy->GetCurveValue(CurveIndex);
 		RigInstance->SetRawControl(ControlIndex, Value);
 	}
 	SharedRigRuntimeContext->RigLogic->Calculate(RigInstance.Get());
 }
 
-void FRigUnit_RigLogic_Data::UpdateJoints(FRigHierarchyContainer* Hierarchy, const FRigUnit_RigLogic_JointUpdateParams& JointUpdateParams)
+void FRigUnit_RigLogic_Data::UpdateJoints(URigHierarchy* Hierarchy, const FRigUnit_RigLogic_JointUpdateParams& JointUpdateParams)
 {
 	for (const uint16 JointIndex : SharedRigRuntimeContext->VariableJointIndices[CurrentLOD].Values)
 	{
@@ -254,13 +254,12 @@ void FRigUnit_RigLogic_Data::UpdateJoints(FRigHierarchyContainer* Hierarchy, con
 				Neutral.GetTranslation() + Delta.GetTranslation(),
 				Neutral.GetScale3D() + Delta.GetScale3D()  // Neutral scale is always 1.0
 			};
-			Hierarchy->BoneHierarchy[BoneIndex].LocalTransform = Transform;
+			Hierarchy->SetLocalTransform(BoneIndex, Transform);
 		}
 	}
-	Hierarchy->BoneHierarchy.RecomputeGlobalTransforms();
 }
 
-void FRigUnit_RigLogic_Data::UpdateBlendShapeCurves(FRigCurveContainer* CurveContainer, TArrayView<const float> BlendShapeValues)
+void FRigUnit_RigLogic_Data::UpdateBlendShapeCurves(URigHierarchy* InHierarchy, TArrayView<const float> BlendShapeValues)
 {
 	// set output blend shapes
 	if (BlendShapeIndices.IsValidIndex(CurrentLOD) && MorphTargetCurveIndices.IsValidIndex(CurrentLOD))
@@ -276,7 +275,7 @@ void FRigUnit_RigLogic_Data::UpdateBlendShapeCurves(FRigCurveContainer* CurveCon
 				if (MorphTargetCurveIndex != INDEX_NONE)
 				{
 					const float Value = BlendShapeValues[BlendShapeIndex];
-					CurveContainer->SetValue(MorphTargetCurveIndex, Value);
+					InHierarchy->SetCurveValue(MorphTargetCurveIndex, Value);
 				}
 			}
 		}
@@ -287,22 +286,22 @@ void FRigUnit_RigLogic_Data::UpdateBlendShapeCurves(FRigCurveContainer* CurveCon
 	}
 }
 
-void FRigUnit_RigLogic_Data::UpdateAnimMapCurves(FRigCurveContainer* CurveContainer, TArrayView<const float> AnimMapOutputs)
+void FRigUnit_RigLogic_Data::UpdateAnimMapCurves(URigHierarchy* InHierarchy, TArrayView<const float> AnimMapOutputs)
 {
 	// set output mask multipliers
 	// In case curves are not imported yet into CL, AnimatedMapsCurveIndices will be empty, so we need to check
 	// array bounds before trying to access it:
-	if (RigLogicIndicesForAnimMaps.IsValidIndex(CurrentLOD) && CurveContainerIndicesForAnimMaps.IsValidIndex(CurrentLOD))
+	if (RigLogicIndicesForAnimMaps.IsValidIndex(CurrentLOD) && CurveElementIndicesForAnimMaps.IsValidIndex(CurrentLOD))
 	{
 		const uint32 AnimMapPerLODCount = RigLogicIndicesForAnimMaps[CurrentLOD].Values.Num();
 		for (uint32 AnimMapIndexForLOD = 0; AnimMapIndexForLOD < AnimMapPerLODCount; ++AnimMapIndexForLOD)
 		{
 			const int32 RigLogicAnimMapIndex = RigLogicIndicesForAnimMaps[CurrentLOD].Values[AnimMapIndexForLOD];
-			const int32 CurveContainerAnimMapIndex = CurveContainerIndicesForAnimMaps[CurrentLOD].Values[AnimMapIndexForLOD];
-			if (CurveContainerAnimMapIndex != INDEX_NONE)
+			const int32 InHierarchyAnimMapIndex = CurveElementIndicesForAnimMaps[CurrentLOD].Values[AnimMapIndexForLOD];
+			if (InHierarchyAnimMapIndex != INDEX_NONE)
 			{
 				const float Value = AnimMapOutputs[RigLogicAnimMapIndex];
-				CurveContainer->SetValue(CurveContainerAnimMapIndex, Value);
+				InHierarchy->SetCurveValue(InHierarchyAnimMapIndex, Value);
 			}
 		}
 	}
@@ -327,7 +326,7 @@ FRigUnit_RigLogic_Execute()
 {
 	QUICK_SCOPE_CYCLE_COUNTER(STAT_RigUnit_RigLogic_Execute);
 
- 	FRigHierarchyContainer* Hierarchy = ExecuteContext.Hierarchy;
+ 	URigHierarchy* Hierarchy = ExecuteContext.Hierarchy;
 	if (Hierarchy)
 	{
 		switch (Context.State)
@@ -351,20 +350,18 @@ FRigUnit_RigLogic_Execute()
 				Data.SharedRigRuntimeContext = GetSharedRigRuntimeContext(Data.SkelMeshComponent->SkeletalMesh);
 				// Context is initialized with a BehaviorReader, which can be imported into SkeletalMesh from DNA file
 				// or overwritten by GeneSplicer when making a new character
-				FRigCurveContainer* CurveContainer = ExecuteContext.GetCurves();
-				Data.InitializeRigLogic(&Hierarchy->BoneHierarchy, CurveContainer);
+				Data.InitializeRigLogic(Hierarchy);
 				break; 
 			}
 			case EControlRigState::Update:
 			{
 				// Fetch shared runtime context of rig from DNAAsset
-				FRigCurveContainer* CurveContainer = ExecuteContext.GetCurves();
-				if (!Data.SkelMeshComponent.IsValid() || !Data.IsRigLogicInitialized() || CurveContainer == nullptr)
+				if (!Data.SkelMeshComponent.IsValid() || !Data.IsRigLogicInitialized() || Hierarchy == nullptr)
 				{
 					return;
 				}
 				Data.CurrentLOD = Data.SkelMeshComponent->PredictedLODLevel;
-				Data.CalculateRigLogic(CurveContainer);
+				Data.CalculateRigLogic(Hierarchy);
 
 				//Filing a struct so we can call the same method for updating joints from tests
 				FRigUnit_RigLogic_JointUpdateParams JointUpdateParamsTemp
@@ -376,10 +373,10 @@ FRigUnit_RigLogic_Execute()
 				Data.UpdateJoints(Hierarchy, JointUpdateParamsTemp);
 
 				TArrayView<const float> BlendShapeValues = Data.RigInstance->GetBlendShapeOutputs();
-				Data.UpdateBlendShapeCurves(CurveContainer, BlendShapeValues);
+				Data.UpdateBlendShapeCurves(Hierarchy, BlendShapeValues);
 
 				TArrayView<const float> AnimMapOutputs = Data.RigInstance->GetAnimatedMapOutputs();
-				Data.UpdateAnimMapCurves(CurveContainer, AnimMapOutputs);
+				Data.UpdateAnimMapCurves(Hierarchy, AnimMapOutputs);
 			}
 			default:
 			{
diff --git a/Engine/Plugins/Animation/RigLogic/Source/RigLogicModule/Private/RigUnit_RigLogic_Test.cpp b/Engine/Plugins/Animation/RigLogic/Source/RigLogicModule/Private/RigUnit_RigLogic_Test.cpp
--- a/Engine/Plugins/Animation/RigLogic/Source/RigLogicModule/Private/RigUnit_RigLogic_Test.cpp
+++ b/Engine/Plugins/Animation/RigLogic/Source/RigLogicModule/Private/RigUnit_RigLogic_Test.cpp
@@ -31,9 +31,9 @@ TSharedPtr<TestBehaviorReader> FRigUnit_RigLogic::TestAccessor::CreateBehaviorRe
 	return MakeShared<TestBehaviorReader>();
 }
 
-TUniquePtr<FRigCurveContainer> FRigUnit_RigLogic::TestAccessor::CreateCurveContainerEmpty()
+TStrongObjectPtr<URigHierarchy> FRigUnit_RigLogic::TestAccessor::CreateCurveContainerEmpty()
 {
-	return MakeUnique<FRigCurveContainer>();
+	return TStrongObjectPtr<URigHierarchy>(NewObject<URigHierarchy>());
 }
 
 TSharedPtr<TestBehaviorReader> FRigUnit_RigLogic::TestAccessor::CreateBehaviorReaderOneCurve(FString ControlNameStr)
@@ -45,34 +45,35 @@ TSharedPtr<TestBehaviorReader> FRigUnit_RigLogic::TestAccessor::CreateBehaviorRe
 	return BehaviorReader;
 }
 
-TUniquePtr<FRigCurveContainer> FRigUnit_RigLogic::TestAccessor::CreateCurveContainerOneCurve(FString CurveNameStr)
+TStrongObjectPtr<URigHierarchy> FRigUnit_RigLogic::TestAccessor::CreateCurveContainerOneCurve(FString CurveNameStr)
 {
-	TUniquePtr<FRigCurveContainer> ValidCurveContainer = MakeUnique<FRigCurveContainer>();
-	ValidCurveContainer->Add(FName(*CurveNameStr));
-	ValidCurveContainer->Initialize();
+	TStrongObjectPtr<URigHierarchy> ValidCurveContainer(NewObject<URigHierarchy>());
+	URigHierarchyController* Controller = NewObject<URigHierarchyController>();
+	Controller->SetHierarchy(ValidCurveContainer.Get());
+	Controller->AddCurve(FName(*CurveNameStr));
 	return ValidCurveContainer;
 }
 
-void FRigUnit_RigLogic::TestAccessor::Exec_MapInputCurve(FRigCurveContainer* TestCurveContainer)
+void FRigUnit_RigLogic::TestAccessor::Exec_MapInputCurve(URigHierarchy* TestHierarchy)
 {
-	Unit->Data.MapInputCurveIndices(TestCurveContainer); //inside the method so we can access data, which is a private member
+	Unit->Data.MapInputCurveIndices(TestHierarchy); //inside the method so we can access data, which is a private member
 }
 
 /** ====== Map Joints ===== **/
 
 
-TUniquePtr<FRigBoneHierarchy> FRigUnit_RigLogic::TestAccessor::CreateBoneHierarchyEmpty()
+TStrongObjectPtr<URigHierarchy> FRigUnit_RigLogic::TestAccessor::CreateBoneHierarchyEmpty()
 {
-	return MakeUnique<FRigBoneHierarchy>();
+	return TStrongObjectPtr<URigHierarchy>(NewObject<URigHierarchy>());
 }
 
-TUniquePtr<FRigBoneHierarchy> FRigUnit_RigLogic::TestAccessor::CreateBoneHierarchyTwoBones(FString Bone1NameStr, FString Bone2NameStr)
+TStrongObjectPtr<URigHierarchy> FRigUnit_RigLogic::TestAccessor::CreateBoneHierarchyTwoBones(FString Bone1NameStr, FString Bone2NameStr)
 {
-	TUniquePtr<FRigBoneHierarchy> TestHierarchy = MakeUnique<FRigBoneHierarchy>();
-	TestHierarchy->Reset();
-	TestHierarchy->Add(*Bone1NameStr, NAME_None, ERigBoneType::User, FTransform(FVector(1.f, 0.f, 0.f)));
-	TestHierarchy->Add(*Bone2NameStr, *Bone1NameStr, ERigBoneType::User, FTransform(FVector(1.f, 2.f, 0.f)));
-	TestHierarchy->Initialize();
+	TStrongObjectPtr<URigHierarchy> TestHierarchy(NewObject<URigHierarchy>());
+	URigHierarchyController* Controller = NewObject<URigHierarchyController>();
+	Controller->SetHierarchy(TestHierarchy.Get());
+	const FRigElementKey Bone1Key = Controller->AddBone(*Bone1NameStr, FRigElementKey(), FTransform(FVector(1.f, 0.f, 0.f)), true, ERigBoneType::User);
+	Controller->AddBone(*Bone2NameStr, Bone1Key, FTransform(FVector(1.f, 2.f, 3.f)), true, ERigBoneType::User);
 	return TestHierarchy;
 }
 
@@ -85,7 +86,7 @@ TSharedPtr<TestBehaviorReader> FRigUnit_RigLogic::TestAccessor::CreateBehaviorRe
 	return TestReader;
 }
 
-void FRigUnit_RigLogic::TestAccessor::Exec_MapJoints(FRigBoneHierarchy* TestHierarchy)
+void FRigUnit_RigLogic::TestAccessor::Exec_MapJoints(URigHierarchy* TestHierarchy)
 {
 	Unit->Data.MapJoints(TestHierarchy);
 }
@@ -126,24 +127,26 @@ TSharedPtr<TestBehaviorReader> FRigUnit_RigLogic::TestAccessor::CreateBehaviorRe
 	return BehaviorReader;
 }
 
-TUniquePtr<FRigCurveContainer> FRigUnit_RigLogic::TestAccessor::CreateCurveContainerOneMorphTarget(FString MorphTargetStr)
+TStrongObjectPtr<URigHierarchy> FRigUnit_RigLogic::TestAccessor::CreateCurveContainerOneMorphTarget(FString MorphTargetStr)
 {
-	TUniquePtr<FRigCurveContainer> ValidCurveContainer = MakeUnique<FRigCurveContainer>();
-	ValidCurveContainer->Add(FName(*MorphTargetStr));
-	ValidCurveContainer->Initialize();
+	TStrongObjectPtr<URigHierarchy> ValidCurveContainer(NewObject<URigHierarchy>());
+	URigHierarchyController* Controller = NewObject<URigHierarchyController>();
+	Controller->SetHierarchy(ValidCurveContainer.Get());
+	Controller->AddCurve(FName(*MorphTargetStr));
 	return ValidCurveContainer;
 }
 
-TUniquePtr<FRigCurveContainer> FRigUnit_RigLogic::TestAccessor::CreateCurveContainerTwoMorphTargets(FString MorphTarget1Str, FString MorphTarget2Str)
+TStrongObjectPtr<URigHierarchy> FRigUnit_RigLogic::TestAccessor::CreateCurveContainerTwoMorphTargets(FString MorphTarget1Str, FString MorphTarget2Str)
 {
-	TUniquePtr<FRigCurveContainer> ValidCurveContainer = MakeUnique<FRigCurveContainer>();
-	ValidCurveContainer->Add(FName(*MorphTarget1Str));
-	ValidCurveContainer->Add(FName(*MorphTarget2Str));
-	ValidCurveContainer->Initialize();
+	TStrongObjectPtr<URigHierarchy> ValidCurveContainer(NewObject<URigHierarchy>());
+	URigHierarchyController* Controller = NewObject<URigHierarchyController>();
+	Controller->SetHierarchy(ValidCurveContainer.Get());
+	Controller->AddCurve(FName(*MorphTarget1Str));
+	Controller->AddCurve(FName(*MorphTarget2Str));
 	return ValidCurveContainer;
 }
 
-void FRigUnit_RigLogic::TestAccessor::Exec_MapMorphTargets(FRigCurveContainer* TestCurveContainer)
+void FRigUnit_RigLogic::TestAccessor::Exec_MapMorphTargets(URigHierarchy* TestCurveContainer)
 {
 	//put into a separate method so we can access the private Data member
 	Unit->Data.MapMorphTargets(TestCurveContainer); 
@@ -160,9 +163,9 @@ TSharedPtr<TestBehaviorReader> FRigUnit_RigLogic::TestAccessor::CreateBehaviorRe
 	return BehaviorReader;
 }
 
-void FRigUnit_RigLogic::TestAccessor::Exec_MapMaskMultipliers(FRigCurveContainer* TestCurveContainer)
+void FRigUnit_RigLogic::TestAccessor::Exec_MapMaskMultipliers(URigHierarchy* TestHierarchy)
 {
-	Unit->Data.MapMaskMultipliers(TestCurveContainer); //inside the method so we can access data, which is a private member
+	Unit->Data.MapMaskMultipliers(TestHierarchy); //inside the method so we can access data, which is a private member
 }
 
 void FRigUnit_RigLogic::TestAccessor::AddToTransformArray(float* InArray, FTransform& Transform)
@@ -237,9 +240,9 @@ TArrayView<const uint16> FRigUnit_RigLogic::TestAccessor::CreateTwoJointVariable
 	return TArrayView<const uint16>(InVariableAttributeIndices, FRigUnit_RigLogic::TestAccessor::MAX_ATTRS_PER_JOINT);
 }
 
-void FRigUnit_RigLogic::TestAccessor::Exec_UpdateJoints(FRigHierarchyContainer* TestHierarchyContainer, FRigUnit_RigLogic_JointUpdateParams& JointUpdateParams)
+void FRigUnit_RigLogic::TestAccessor::Exec_UpdateJoints(URigHierarchy* TestHierarchy, FRigUnit_RigLogic_JointUpdateParams& JointUpdateParams)
 {
-	Unit->Data.UpdateJoints(TestHierarchyContainer, JointUpdateParams);
+	Unit->Data.UpdateJoints(TestHierarchy, JointUpdateParams);
 }
 
 FSharedRigRuntimeContext* FRigUnit_RigLogic::TestAccessor::GetSharedRigRuntimeContext(USkeletalMesh* SkelMesh)
@@ -273,7 +276,7 @@ IMPLEMENT_RIGUNIT_AUTOMATION_TEST(FRigUnit_RigLogic)
 	}
 
 	Test.GetData()->SkelMeshComponent = SkelMeshComponent;
-	Unit.ExecuteContext.Hierarchy = &HierarchyContainer;
+	Unit.ExecuteContext.Hierarchy = Hierarchy;
 
 	//Test
 	InitAndExecute();
@@ -296,7 +299,7 @@ IMPLEMENT_RIGUNIT_AUTOMATION_TEST(FRigUnit_RigLogic)
 
 		// Check mask multipliers
 		Test.GetData()->RigLogicIndicesForAnimMaps.Num() > 0 &&
-		Test.GetData()->CurveContainerIndicesForAnimMaps.Num() > 0,
+		Test.GetData()->CurveElementIndicesForAnimMaps.Num() > 0,
 
 		TEXT("InitAndExecute failed to initialize rig logic.")
 	);
@@ -307,7 +310,7 @@ IMPLEMENT_RIGUNIT_AUTOMATION_TEST(FRigUnit_RigLogic)
 
 	//Prepare
 	TSharedPtr<TestBehaviorReader> TestReaderValid = Test.CreateBehaviorReaderOneCurve("CTRL_Expressions.Some_Control");
-	TUniquePtr<FRigCurveContainer> TestCurveContainerNameMismatch = Test.CreateCurveContainerOneCurve("CTRL_Expressions_NOT_ThatControl");
+	TStrongObjectPtr<URigHierarchy> TestCurveContainerNameMismatch = Test.CreateCurveContainerOneCurve("CTRL_Expressions_NOT_ThatControl");
 	SharedRigRuntimeContext->BehaviorReader = TestReaderValid;
 	//Test
 	Test.Exec_MapInputCurve(TestCurveContainerNameMismatch.Get());
@@ -322,7 +325,7 @@ IMPLEMENT_RIGUNIT_AUTOMATION_TEST(FRigUnit_RigLogic)
 
 	//Prepare
 	TSharedPtr<TestBehaviorReader> TestReaderEmpty = Test.CreateBehaviorReaderEmpty();
-	TUniquePtr<FRigCurveContainer> TestCurveContainerValid = Test.CreateCurveContainerOneCurve("CTRL_Expressions_Some_Control");
+	TStrongObjectPtr<URigHierarchy> TestCurveContainerValid = Test.CreateCurveContainerOneCurve("CTRL_Expressions_Some_Control");
 	SharedRigRuntimeContext->BehaviorReader = TestReaderEmpty;
 	//Test
 	Test.Exec_MapInputCurve(TestCurveContainerValid.Get());
@@ -335,8 +338,7 @@ IMPLEMENT_RIGUNIT_AUTOMATION_TEST(FRigUnit_RigLogic)
 	//=== MapInputCurve ValidReader EmptyCurveContainer ===
 
 	//Prepare
-	TUniquePtr<FRigCurveContainer> TestCurveContainerEmpty = Test.CreateCurveContainerEmpty();
-	TestCurveContainerEmpty->Initialize();
+	TStrongObjectPtr<URigHierarchy> TestCurveContainerEmpty = Test.CreateCurveContainerEmpty();
 	SharedRigRuntimeContext->BehaviorReader = TestReaderValid;
 	//Test
 	Test.Exec_MapInputCurve(TestCurveContainerEmpty.Get());
@@ -380,7 +382,7 @@ IMPLEMENT_RIGUNIT_AUTOMATION_TEST(FRigUnit_RigLogic)
 
 	//=== MapJoints EmptyInputs ===
 	//Prepare
-	TUniquePtr<FRigBoneHierarchy> TestHierarchyEmpty = Test.CreateBoneHierarchyEmpty();
+	TStrongObjectPtr<URigHierarchy> TestHierarchyEmpty = Test.CreateBoneHierarchyEmpty();
 	SharedRigRuntimeContext->BehaviorReader = TestReaderEmpty;
 	//Test
 	Test.Exec_MapJoints(TestHierarchyEmpty.Get());
@@ -392,7 +394,7 @@ IMPLEMENT_RIGUNIT_AUTOMATION_TEST(FRigUnit_RigLogic)
 
 	//=== MapJoints EmptyReader TwoBones ===
 	//Prepare
-	TUniquePtr<FRigBoneHierarchy> TestHierarchyTwoBones = Test.CreateBoneHierarchyTwoBones("BoneA", "BoneB");
+	TStrongObjectPtr<URigHierarchy> TestHierarchyTwoBones = Test.CreateBoneHierarchyTwoBones("BoneA", "BoneB");
 	SharedRigRuntimeContext->BehaviorReader = TestReaderEmpty;
 	//Test
 	Test.Exec_MapJoints(TestHierarchyTwoBones.Get());
@@ -450,7 +452,7 @@ IMPLEMENT_RIGUNIT_AUTOMATION_TEST(FRigUnit_RigLogic)
 	//=== MapMorphTargets ValidReader MorphTargetWithNameMismatch ===
 
 	//Prepare
-	TUniquePtr<FRigCurveContainer> TestMorphTargetNameMismatch = Test.CreateCurveContainerOneMorphTarget("head_NOT_that_blendshape");
+	TStrongObjectPtr<URigHierarchy> TestMorphTargetNameMismatch = Test.CreateCurveContainerOneMorphTarget("head_NOT_that_blendshape");
 	TSharedPtr<TestBehaviorReader> TestReaderBlendshapeValid = Test.CreateBehaviorReaderOneBlendShape("head", "blendshape");
 	SharedRigRuntimeContext->BehaviorReader = TestReaderBlendshapeValid;
 	//Test
@@ -470,7 +472,7 @@ IMPLEMENT_RIGUNIT_AUTOMATION_TEST(FRigUnit_RigLogic)
 
 	//Prepare
 	//Empty reader (no meshes, no blendshapes)
-	TUniquePtr<FRigCurveContainer> TestMorphTargetCurveValid = Test.CreateCurveContainerOneMorphTarget("head__blendshape");
+	TStrongObjectPtr<URigHierarchy> TestMorphTargetCurveValid = Test.CreateCurveContainerOneMorphTarget("head__blendshape");
 	SharedRigRuntimeContext->BehaviorReader = TestReaderEmpty;
 	//Test
 	Test.Exec_MapMorphTargets(TestMorphTargetCurveValid.Get());
@@ -534,7 +536,7 @@ IMPLEMENT_RIGUNIT_AUTOMATION_TEST(FRigUnit_RigLogic)
 	//=== MapMorphTargets ValidReader InvalidMorphTargetCurve ===
 
 	//Prepare
-	TUniquePtr<FRigCurveContainer> TestMorphTargetCurvesInvalid = Test.CreateCurveContainerOneMorphTarget("");
+	TStrongObjectPtr<URigHierarchy> TestMorphTargetCurvesInvalid = Test.CreateCurveContainerOneMorphTarget("");
 	SharedRigRuntimeContext->BehaviorReader = TestReaderBlendshapeValid;
 	//Test
 	Test.Exec_MapMorphTargets(TestMorphTargetCurvesInvalid.Get());
@@ -570,7 +572,7 @@ IMPLEMENT_RIGUNIT_AUTOMATION_TEST(FRigUnit_RigLogic)
 
 	//Prepare
 	TSharedPtr<TestBehaviorReader> TestReaderBlendshapes_LOD0AB_LOD1A = Test.CreateBehaviorReaderTwoBlendShapes("head", "blendshapeA", "blendshapeB");
-	TUniquePtr<FRigCurveContainer> TestMorphTargetTwoCurves = Test.CreateCurveContainerTwoMorphTargets("head__blendshapeA", "head__blendshapeB");
+	TStrongObjectPtr<URigHierarchy> TestMorphTargetTwoCurves = Test.CreateCurveContainerTwoMorphTargets("head__blendshapeA", "head__blendshapeB");
 	//NOTE: indices in the first param here are not blendshape indices, but rather mappings from blendshapes to meshes
 	//in this test, they will correspond to blendshape indices
 
@@ -687,15 +689,15 @@ IMPLEMENT_RIGUNIT_AUTOMATION_TEST(FRigUnit_RigLogic)
 
 	//Prepare
 	TSharedPtr<TestBehaviorReader> TestReaderAnimMapsValid = Test.CreateBehaviorReaderOneAnimatedMap("CTRL_AnimMap.Some_Multiplier");
-	TUniquePtr<FRigCurveContainer> TestCurveContainerForAnimMapsNameMismatch = Test.CreateCurveContainerOneCurve("CTRL_AnimMap_NOT_ThatMultiploer");
+	TStrongObjectPtr<URigHierarchy> TestCurveContainerForAnimMapsNameMismatch = Test.CreateCurveContainerOneCurve("CTRL_AnimMap_NOT_ThatMultiploer");
 	//Test
 	SharedRigRuntimeContext->BehaviorReader = TestReaderAnimMapsValid;
 	Test.Exec_MapMaskMultipliers(TestCurveContainerForAnimMapsNameMismatch.Get());
 	//Assert
 	AddErrorIfFalse(
-		Test.GetData()->CurveContainerIndicesForAnimMaps.Num() == 1 &&
-		Test.GetData()->CurveContainerIndicesForAnimMaps[0].Values.Num() == 1 &&
-		Test.GetData()->CurveContainerIndicesForAnimMaps[0].Values[0] == INDEX_NONE &&
+		Test.GetData()->CurveElementIndicesForAnimMaps.Num() == 1 &&
+		Test.GetData()->CurveElementIndicesForAnimMaps[0].Values.Num() == 1 &&
+		Test.GetData()->CurveElementIndicesForAnimMaps[0].Values[0] == INDEX_NONE &&
 		Test.GetData()->RigLogicIndicesForAnimMaps.Num() == 1 &&
 		Test.GetData()->RigLogicIndicesForAnimMaps[0].Values.Num() == 1 &&
 		Test.GetData()->RigLogicIndicesForAnimMaps[0].Values[0] == 0,
@@ -705,13 +707,13 @@ IMPLEMENT_RIGUNIT_AUTOMATION_TEST(FRigUnit_RigLogic)
 	//=== MapMaskMultipliers EmptyReader ValidAnimatedMap ===
 
 	//Prepare
-	TUniquePtr<FRigCurveContainer> TestCurveContainerForAnimMapsValid = Test.CreateCurveContainerOneCurve("CTRL_AnimMap_Some_Multiplier");
+	TStrongObjectPtr<URigHierarchy> TestCurveContainerForAnimMapsValid = Test.CreateCurveContainerOneCurve("CTRL_AnimMap_Some_Multiplier");
 	//Test
 	SharedRigRuntimeContext->BehaviorReader = TestReaderEmpty;
 	Test.Exec_MapMaskMultipliers(TestCurveContainerForAnimMapsValid.Get());
 	//Assert
 	AddErrorIfFalse(
-		Test.GetData()->CurveContainerIndicesForAnimMaps.Num() == 0 &&
+		Test.GetData()->CurveElementIndicesForAnimMaps.Num() == 0 &&
 		Test.GetData()->RigLogicIndicesForAnimMaps.Num() == 0,
 		TEXT("MapMaskMultipliers - EmptyReader ValidAnimatedMap")
 	);
@@ -723,9 +725,9 @@ IMPLEMENT_RIGUNIT_AUTOMATION_TEST(FRigUnit_RigLogic)
 	Test.Exec_MapMaskMultipliers(TestCurveContainerEmpty.Get());
 	//Assert
 	AddErrorIfFalse(
-		Test.GetData()->CurveContainerIndicesForAnimMaps.Num() == 1 &&
-		Test.GetData()->CurveContainerIndicesForAnimMaps[0].Values.Num() == 1 &&
-		Test.GetData()->CurveContainerIndicesForAnimMaps[0].Values[0] == INDEX_NONE &&
+		Test.GetData()->CurveElementIndicesForAnimMaps.Num() == 1 &&
+		Test.GetData()->CurveElementIndicesForAnimMaps[0].Values.Num() == 1 &&
+		Test.GetData()->CurveElementIndicesForAnimMaps[0].Values[0] == INDEX_NONE &&
 		Test.GetData()->RigLogicIndicesForAnimMaps.Num() == 1 &&
 		Test.GetData()->RigLogicIndicesForAnimMaps[0].Values.Num() == 1 &&
 		Test.GetData()->RigLogicIndicesForAnimMaps[0].Values[0] == 0,
@@ -741,9 +743,9 @@ IMPLEMENT_RIGUNIT_AUTOMATION_TEST(FRigUnit_RigLogic)
 	Test.Exec_MapMaskMultipliers(TestCurveContainerForAnimMapsValid.Get());
 	//Assert
 	AddErrorIfFalse(
-		Test.GetData()->CurveContainerIndicesForAnimMaps.Num() == 1 &&
-		Test.GetData()->CurveContainerIndicesForAnimMaps[0].Values.Num() == 1 &&
-		Test.GetData()->CurveContainerIndicesForAnimMaps[0].Values[0] == 0 &&
+		Test.GetData()->CurveElementIndicesForAnimMaps.Num() == 1 &&
+		Test.GetData()->CurveElementIndicesForAnimMaps[0].Values.Num() == 1 &&
+		Test.GetData()->CurveElementIndicesForAnimMaps[0].Values[0] == 0 &&
 		Test.GetData()->RigLogicIndicesForAnimMaps.Num() == 1 &&
 		Test.GetData()->RigLogicIndicesForAnimMaps[0].Values.Num() == 1 &&
 		Test.GetData()->RigLogicIndicesForAnimMaps[0].Values[0] == 0,
@@ -751,10 +753,7 @@ IMPLEMENT_RIGUNIT_AUTOMATION_TEST(FRigUnit_RigLogic)
 	);
 
 	//BoneHierarchy belongs to HierarchyContainer
-	BoneHierarchy = *TestHierarchyTwoBones;
-	BoneHierarchy.Initialize();
-	BoneHierarchy.ResetTransforms();
-	FRigHierarchyContainer* HierarchyContainerPtr = &HierarchyContainer;
+	TestHierarchyTwoBones->ResetPoseToInitial();
 
 	//Prepare
 	//-----
@@ -782,11 +781,11 @@ IMPLEMENT_RIGUNIT_AUTOMATION_TEST(FRigUnit_RigLogic)
 		TwoJointNeutralTransforms,
 		DeltaTransforms);
 	//Test
-	Test.Exec_UpdateJoints(HierarchyContainerPtr, TestJointUpdateParamsTwoJoints_LOD0);
+	Test.Exec_UpdateJoints(TestHierarchyTwoBones.Get(), TestJointUpdateParamsTwoJoints_LOD0);
 	//Assert
 	//Note that BoneB.GlobalTransform.Z should be zero since the scale.Z is zero. Also, translation Y becomes -Y 
-	AddErrorIfFalse(BoneHierarchy.GetGlobalTransform(0).GetTranslation().Equals(FVector(1.f, 0.f, 0.f)), TEXT("UpdateJoints LOD0 Bone 01 - unexpected transform"));
-	AddErrorIfFalse(BoneHierarchy.GetGlobalTransform(1).GetTranslation().Equals(FVector(2.f, -2.f, 0.f)), TEXT("UpdateJoints LOD0 Bone 02 - unexpected transform"));
+	AddErrorIfFalse(TestHierarchyTwoBones->GetGlobalTransform(0).GetTranslation().Equals(FVector(1.f, 0.f, 0.f)), TEXT("UpdateJoints LOD0 Bone 01 - unexpected transform"));
+	AddErrorIfFalse(TestHierarchyTwoBones->GetGlobalTransform(1).GetTranslation().Equals(FVector(2.f, -2.f, 0.f)), TEXT("UpdateJoints LOD0 Bone 02 - unexpected transform"));
 
 	return true;
 }
diff --git a/Engine/Plugins/Animation/RigLogic/Source/RigLogicModule/Public/RigUnit_RigLogic.h b/Engine/Plugins/Animation/RigLogic/Source/RigLogicModule/Public/RigUnit_RigLogic.h
--- a/Engine/Plugins/Animation/RigLogic/Source/RigLogicModule/Public/RigUnit_RigLogic.h
+++ b/Engine/Plugins/Animation/RigLogic/Source/RigLogicModule/Public/RigUnit_RigLogic.h
@@ -93,7 +93,7 @@ struct FRigUnit_RigLogic_Data
 
 	/** RL animated map index to ControlRig's output anim map curve index for each LOD **/
 	UPROPERTY(transient)
-	TArray<FRigUnit_RigLogic_IntArray> CurveContainerIndicesForAnimMaps;
+	TArray<FRigUnit_RigLogic_IntArray> CurveElementIndicesForAnimMaps;
 
 	/** RL animated map index to RL anim map curve index for each LOD **/
 	UPROPERTY(transient)
@@ -106,29 +106,29 @@ struct FRigUnit_RigLogic_Data
 	static const uint8 MAX_ATTRS_PER_JOINT;
 
 	bool IsRigLogicInitialized();
-	void InitializeRigLogic(const FRigBoneHierarchy* BoneHierarchy, const FRigCurveContainer* CurveContainer);
+	void InitializeRigLogic(const URigHierarchy* InHierarchy);
 
 	/** Makes a map of input curve indices from DNA file to the control rig curves **/
-	void MapInputCurveIndices(const FRigCurveContainer* CurveContainer);
+	void MapInputCurveIndices(const URigHierarchy* InHierarchy);
 	/** Uses names to map joint indices from DNA file to the indices of bones in control rig hierarchy **/
-	void MapJoints(const FRigBoneHierarchy* Hierarchy);
+	void MapJoints(const URigHierarchy* Hierarchy);
 	/** Cache the joint indices that change per each LOD **/
 	void CacheVariableJointIndices();
 	/** Uses names of blend shapes and meshes in DNA file, for all LODs, to map their indices to the indices of
 	  * morph target curves in the curve container; curve name format is <mesh>__<blendshape> **/
-	void MapMorphTargets(const FRigCurveContainer* CurveContainer);
+	void MapMorphTargets(const URigHierarchy* InHierarchy);
 	/** Uses names to map mask multiplier indices from DNA file, for all LODs, to the indices of curves in the
 	  * control rig's curve container **/
-	void MapMaskMultipliers(const FRigCurveContainer* CurveContainer);
+	void MapMaskMultipliers(const URigHierarchy* InHierarchy);
 
 	/** Calculates joint positions, orientation and scale based on inputs curves of the control rig **/
-	void CalculateRigLogic(const FRigCurveContainer* CurveContainer);
+	void CalculateRigLogic(const URigHierarchy* InHierarchy);
 	/** Updates joint positions in the hierarchy based on inputs curves of the control rig **/
-	void UpdateJoints(FRigHierarchyContainer* Hierarchy, const FRigUnit_RigLogic_JointUpdateParams& JointUpdateParams);
+	void UpdateJoints(URigHierarchy* Hierarchy, const FRigUnit_RigLogic_JointUpdateParams& JointUpdateParams);
 	/** Updates morph target curve values based on values of input curves of the control rig **/
-	void UpdateBlendShapeCurves(FRigCurveContainer* CurveContainer, TArrayView<const float> BlendShapeValues);
+	void UpdateBlendShapeCurves(URigHierarchy* InHierarchy, TArrayView<const float> BlendShapeValues);
 	/** Updates anim map curve values based on values of input curves of the control rig **/
-	void UpdateAnimMapCurves(FRigCurveContainer* CurveContainer, TArrayView<const float> AnimMapOutputs);
+	void UpdateAnimMapCurves(URigHierarchy* InHierarchy, TArrayView<const float> AnimMapOutputs);
 };
 
 
diff --git a/Engine/Plugins/Animation/RigLogic/Source/RigLogicModule/Public/RigUnit_RigLogic_Test.h b/Engine/Plugins/Animation/RigLogic/Source/RigLogicModule/Public/RigUnit_RigLogic_Test.h
--- a/Engine/Plugins/Animation/RigLogic/Source/RigLogicModule/Public/RigUnit_RigLogic_Test.h
+++ b/Engine/Plugins/Animation/RigLogic/Source/RigLogicModule/Public/RigUnit_RigLogic_Test.h
@@ -4,6 +4,8 @@
 
 #include "CoreTypes.h"
 #include "CoreMinimal.h"
+#include "UObject/StrongObjectPtr.h"
+#include "Rigs/RigHierarchy.h"
 
 #include "DNAReader.h"
 
@@ -491,33 +493,33 @@ struct FRigUnit_RigLogic::TestAccessor
 	TSharedPtr<TestBehaviorReader> CreateBehaviorReaderEmpty();
 	TSharedPtr<TestBehaviorReader> CreateBehaviorReaderOneCurve(FString ControlNameStr);
 
-	TUniquePtr<FRigCurveContainer> CreateCurveContainerEmpty();
-	TUniquePtr<FRigCurveContainer> CreateCurveContainerOneCurve(FString CurveNameStr);
+	TStrongObjectPtr<URigHierarchy> CreateCurveContainerEmpty();
+	TStrongObjectPtr<URigHierarchy> CreateCurveContainerOneCurve(FString CurveNameStr);
 
-	void Exec_MapInputCurve(FRigCurveContainer* TestCurveContainer);
+	void Exec_MapInputCurve(URigHierarchy* TestHierarchy);
 
 	/** MapJoints Tests **/
 
 	TSharedPtr<TestBehaviorReader> CreateBehaviorReaderTwoJoints(FString Joint1NameStr, FString Joint2NameStr);
-	TUniquePtr<FRigBoneHierarchy> CreateBoneHierarchyEmpty();
-	TUniquePtr<FRigBoneHierarchy> CreateBoneHierarchyTwoBones(FString Bone1NameStr, FString Bone2NameStr);
+	TStrongObjectPtr<URigHierarchy> CreateBoneHierarchyEmpty();
+	TStrongObjectPtr<URigHierarchy> CreateBoneHierarchyTwoBones(FString Bone1NameStr, FString Bone2NameStr);
 
-	void Exec_MapJoints(FRigBoneHierarchy* MockHierachy);
+	void Exec_MapJoints(URigHierarchy* TestHierarchy);
 
 	/** MapBlendShapes Tests **/
 
 	TSharedPtr<TestBehaviorReader> CreateBehaviorReaderNoBlendshapes(FString MeshNameStr);
 	TSharedPtr<TestBehaviorReader> CreateBehaviorReaderOneBlendShape(FString MeshNameStr, FString BlendShapeNameStr);
-	TUniquePtr<FRigCurveContainer> CreateCurveContainerOneMorphTarget(FString MorphTargetStr);
+	TStrongObjectPtr<URigHierarchy> CreateCurveContainerOneMorphTarget(FString MorphTargetStr);
 	TSharedPtr<TestBehaviorReader> CreateBehaviorReaderTwoBlendShapes(FString MeshNameStr, FString BlendShape1Str, FString BlendShape2Str);
-	TUniquePtr<FRigCurveContainer> CreateCurveContainerTwoMorphTargets(FString MorphTarget1Str, FString MorphTarget2Str);
+	TStrongObjectPtr<URigHierarchy> CreateCurveContainerTwoMorphTargets(FString MorphTarget1Str, FString MorphTarget2Str);
 
-	void Exec_MapMorphTargets(FRigCurveContainer* TestCurveContainer);
+	void Exec_MapMorphTargets(URigHierarchy* TestHierarchy);
 
 	/** MapMaskMultipliers Tests **/
 	TSharedPtr<TestBehaviorReader> CreateBehaviorReaderOneAnimatedMap(FString AnimatedMapNameStr);
 
-	void Exec_MapMaskMultipliers(FRigCurveContainer* TestCurveContainer);
+	void Exec_MapMaskMultipliers(URigHierarchy* TestHierarchy);
 
 	static const uint8 MAX_ATTRS_PER_JOINT;
 
@@ -526,14 +528,13 @@ struct FRigUnit_RigLogic::TestAccessor
 	FTransformArrayView CreateTwoJointNeutralTransforms(float *InValueArray);
 	TArrayView<const uint16> CreateTwoJointVariableAttributes(uint16* InVariableAttributeIndices, uint8 LOD);
 
-	void Exec_UpdateJoints(FRigHierarchyContainer* TestHierarchyContainer, FRigUnit_RigLogic_JointUpdateParams& JointUpdateParams);
+	void Exec_UpdateJoints(URigHierarchy* TestHierarchy, FRigUnit_RigLogic_JointUpdateParams& JointUpdateParams);
 
 	FSharedRigRuntimeContext* GetSharedRigRuntimeContext(USkeletalMesh* SkelMesh);
 
 private:
 	FName RawCtrlName;
-	FRigCurveContainer CurveContainer;
-	FRigBoneHierarchy Hierarchy;
+	TStrongObjectPtr<URigHierarchy> Hierarchy;
 	FRigUnit_RigLogic* Unit;
 
 public:
diff --git a/Engine/Plugins/Experimental/ControlRig/Config/BaseControlRig.ini b/Engine/Plugins/Experimental/ControlRig/Config/BaseControlRig.ini
--- a/Engine/Plugins/Experimental/ControlRig/Config/BaseControlRig.ini
+++ b/Engine/Plugins/Experimental/ControlRig/Config/BaseControlRig.ini
@@ -8,7 +8,6 @@
 
 
 +StructRedirects=(OldName="RigJoint",NewName="/Script/ControlRig.RigBone")
-+StructRedirects=(OldName="RigHierarchy",NewName="/Script/ControlRig.RigBoneHierarchy")
 +PropertyRedirects=(OldName="RigBoneHierarchy.Joints",NewName="RigBoneHierarchy.Bones")
 +PropertyRedirects=(OldName="RigHierarchyContainer.BaseHierarchy",NewName="RigHierarchyContainer.BoneHierarchy")
 
diff --git a/Engine/Plugins/Experimental/ControlRig/Source/ControlRig/Private/AnimNode_ControlRigBase.cpp b/Engine/Plugins/Experimental/ControlRig/Source/ControlRig/Private/AnimNode_ControlRigBase.cpp
--- a/Engine/Plugins/Experimental/ControlRig/Source/ControlRig/Private/AnimNode_ControlRigBase.cpp
+++ b/Engine/Plugins/Experimental/ControlRig/Source/ControlRig/Private/AnimNode_ControlRigBase.cpp
@@ -100,7 +100,7 @@ void FAnimNode_ControlRigBase::UpdateInput(UControlRig* ControlRig, const FPoseC
 		// reset transforms here to prevent additive transforms from accumulating to INF
 		// we only update transforms from the mesh pose for bones in the current LOD, 
 		// so the reset here ensures excluded bones are also reset
-		ControlRig->GetBoneHierarchy().ResetTransforms();
+		ControlRig->GetHierarchy()->ResetPoseToInitial(ERigElementType::Bone);
 
 		// @re-think - now control rig contains init pose from their default hierarchy and current pose from this instance.
 		// we may need this init pose somewhere (instance refpose)
@@ -108,6 +108,7 @@ void FAnimNode_ControlRigBase::UpdateInput(UControlRig* ControlRig, const FPoseC
 		{
 			const FName& Name = Iter.Key();
 			const uint16 Index = Iter.Value();
+			const FRigElementKey Key(Name, ERigElementType::Bone);
 
 			FTransform ComponentTransform = MeshPoses.GetComponentSpaceTransform(FCompactPoseBoneIndex(Index));
 			if (NodeMappingContainer.IsValid())
@@ -115,7 +116,7 @@ void FAnimNode_ControlRigBase::UpdateInput(UControlRig* ControlRig, const FPoseC
 				ComponentTransform = NodeMappingContainer->GetSourceToTargetTransform(Name).GetRelativeTransformReverse(ComponentTransform);
 			}
 
-			ControlRig->SetGlobalTransform(Name, ComponentTransform, false);
+			ControlRig->GetHierarchy()->SetGlobalTransform(Key, ComponentTransform, false);
 		}
 	}
 
@@ -126,8 +127,9 @@ void FAnimNode_ControlRigBase::UpdateInput(UControlRig* ControlRig, const FPoseC
 		{
 			const FName& Name = Iter.Key();
 			const uint16 Index = Iter.Value();
+			const FRigElementKey Key(Name, ERigElementType::Bone);
 
-			ControlRig->SetCurveValue(Name, InOutput.Curve.Get(Index));
+			ControlRig->GetHierarchy()->SetCurveValue(Key, InOutput.Curve.Get(Index));
 		}
 	}
 }
@@ -149,9 +151,10 @@ void FAnimNode_ControlRigBase::UpdateOutput(UControlRig* ControlRig, FPoseContex
 		{
 			const FName& Name = Iter.Key();
 			const uint16 Index = Iter.Value();
+			const FRigElementKey Key(Name, ERigElementType::Bone);
 
 			FCompactPoseBoneIndex CompactPoseIndex(Index);
-			FTransform ComponentTransform = ControlRig->GetGlobalTransform(Name);
+			FTransform ComponentTransform = ControlRig->GetHierarchy()->GetGlobalTransform(Key);
 			if (NodeMappingContainer.IsValid())
 			{
 				ComponentTransform = NodeMappingContainer->GetSourceToTargetTransform(Name) * ComponentTransform;
@@ -171,8 +174,9 @@ void FAnimNode_ControlRigBase::UpdateOutput(UControlRig* ControlRig, FPoseContex
 		{
 			const FName& Name = Iter.Key();
 			const uint16 Index = Iter.Value();
+			const FRigElementKey Key(Name, ERigElementType::Bone);
 
-			const float Value = ControlRig->GetCurveValue(Name);
+			const float Value = ControlRig->GetHierarchy()->GetCurveValue(Key);
 			InOutput.Curve.Set(Index, Value);
 		}
 	}
@@ -261,8 +265,12 @@ struct FControlRigControlScope
 	{
 		if (ControlRig.IsValid())
 		{
-			CopyOfControls = ControlRig->AvailableControls();
-			
+			URigHierarchy* Hierarchy = ControlRig->GetHierarchy();
+			Hierarchy->ForEach<FRigControlElement>([this, Hierarchy](FRigControlElement* ControlElement) -> bool
+			{
+				ControlValues.Add(ControlElement->GetKey(), Hierarchy->GetControlValueByIndex(ControlElement->GetIndex()));
+				return true; // continue
+			});
 		}
 	}
 
@@ -270,76 +278,15 @@ struct FControlRigControlScope
 	{
 		if (ControlRig.IsValid())
 		{
-			for (const FRigControl& CopyRigControl: CopyOfControls)
+			URigHierarchy* Hierarchy = ControlRig->GetHierarchy();
+			for (const TPair<FRigElementKey, FRigControlValue>& Pair: ControlValues)
 			{
-				FRigControl* RigControl = ControlRig->FindControl(CopyRigControl.Name);
-				if (RigControl)
-				{
-					if (CopyRigControl.ControlType == RigControl->ControlType)
-					{
-						switch (RigControl->ControlType)
-						{
-						case ERigControlType::Transform:
-						{
-							FTransform Val = CopyRigControl.GetValue(ERigControlValueType::Current).Get<FTransform>();
-							RigControl->GetValue(ERigControlValueType::Current).Set<FTransform>(Val);
-							break;
-						}
-						case ERigControlType::TransformNoScale:
-						{
-							FTransformNoScale Val = CopyRigControl.GetValue(ERigControlValueType::Current).Get<FTransformNoScale>();
-							RigControl->GetValue(ERigControlValueType::Current).Set<FTransformNoScale>(Val);
-							break;
-						}
-						case ERigControlType::EulerTransform:
-						{
-							FEulerTransform Val = CopyRigControl.GetValue(ERigControlValueType::Current).Get<FEulerTransform>();
-							RigControl->GetValue(ERigControlValueType::Current).Set<FEulerTransform>(Val);
-							break;
-						}
-						case ERigControlType::Float:
-						{
-							float Val = CopyRigControl.GetValue(ERigControlValueType::Current).Get<float>();
-							RigControl->GetValue(ERigControlValueType::Current).Set<float>(Val);
-							break;
-						}
-						case ERigControlType::Bool:
-						{
-							bool Val = CopyRigControl.GetValue(ERigControlValueType::Current).Get<bool>();
-							RigControl->GetValue(ERigControlValueType::Current).Set<bool>(Val);
-							break;
-						}
-						case ERigControlType::Integer:
-						{
-							int32 Val = CopyRigControl.GetValue(ERigControlValueType::Current).Get<int32>();
-							RigControl->GetValue(ERigControlValueType::Current).Set<int32>(Val);
-							break;
-						}
-						case ERigControlType::Vector2D:
-						{
-							FVector2D Val = CopyRigControl.GetValue(ERigControlValueType::Current).Get<FVector2D>();
-							RigControl->GetValue(ERigControlValueType::Current).Set<FVector2D>(Val);
-							break;
-						}
-						case ERigControlType::Position:
-						case ERigControlType::Scale:
-						case ERigControlType::Rotator:
-						{
-							FVector Val = CopyRigControl.GetValue(ERigControlValueType::Current).Get<FVector>();
-							RigControl->GetValue(ERigControlValueType::Current).Set<FVector>(Val);
-							break;
-						}
-						default:
-							break;
-
-						};
-					}
-				}
+				Hierarchy->SetControlValue(Pair.Key, Pair.Value);
 			}
 		}
 	}
 
-	TArray<FRigControl> CopyOfControls;
+	TMap<FRigElementKey, FRigControlValue> ControlValues;
 	TWeakObjectPtr<UControlRig> ControlRig;
 };
 
@@ -403,11 +350,11 @@ void FAnimNode_ControlRigBase::CacheBones_AnyThread(const FAnimationCacheBonesCo
 			
 			// we just support curves by name only
 			TArray<FName> const& CurveNames = RequiredBones.GetUIDToNameLookupTable();
-			const FRigCurveContainer& RigCurveContainer = ControlRig->GetCurveContainer();
+			URigHierarchy* Hierarchy = ControlRig->GetHierarchy();
 			for (uint16 Index = 0; Index < CurveNames.Num(); ++Index)
 			{
 				// see if the curve name exists in the control rig
-				if (RigCurveContainer.GetIndex(CurveNames[Index]) != INDEX_NONE)
+				if (Hierarchy->GetIndex(FRigElementKey(CurveNames[Index], ERigElementType::Curve)) != INDEX_NONE)
 				{
 					ControlRigCurveMapping.Add(CurveNames[Index], Index);
 				}
diff --git a/Engine/Plugins/Experimental/ControlRig/Source/ControlRig/Private/ControlRig.cpp b/Engine/Plugins/Experimental/ControlRig/Source/ControlRig/Private/ControlRig.cpp
--- a/Engine/Plugins/Experimental/ControlRig/Source/ControlRig/Private/ControlRig.cpp
+++ b/Engine/Plugins/Experimental/ControlRig/Source/ControlRig/Private/ControlRig.cpp
@@ -8,6 +8,7 @@
 #include "HelperUtil.h"
 #include "ControlRigBlueprintGeneratedClass.h"
 #include "ControlRigObjectVersion.h"
+#include "Rigs/RigHierarchyController.h"
 #include "Units/Execution/RigUnit_BeginExecution.h"
 #include "Units/Execution/RigUnit_PrepareForExecution.h"
 #include "Units/Execution/RigUnit_InverseExecution.h"
@@ -127,7 +128,6 @@ void UControlRig::Initialize(bool bInitRigUnits)
 	InstantiateVMFromCDO();
 
 	// should refresh mapping 
-	Hierarchy.Initialize();
 	RequestSetup();
 
 	if (bInitRigUnits)
@@ -135,11 +135,10 @@ void UControlRig::Initialize(bool bInitRigUnits)
 		RequestInit();
 	}
 
-	Hierarchy.ControlHierarchy.OnControlSelected.RemoveAll(this);
-	Hierarchy.ControlHierarchy.OnControlSelected.AddUObject(this, &UControlRig::HandleOnControlSelected);
-	Hierarchy.ControlHierarchy.OnControlUISettingsChanged.AddUObject(this, &UControlRig::HandleOnControlUISettingChanged);
-	Hierarchy.OnEventReceived.RemoveAll(this);
-	Hierarchy.OnEventReceived.AddUObject(this, &UControlRig::HandleOnRigEvent);
+	DynamicHierarchy->OnModified().RemoveAll(this);
+	DynamicHierarchy->OnModified().AddUObject(this, &UControlRig::HandleHierarchyModified);
+	DynamicHierarchy->OnEventReceived().RemoveAll(this);
+	DynamicHierarchy->OnEventReceived().AddUObject(this, &UControlRig::HandleHierarchyEvent);
 }
 
 void UControlRig::InitializeFromCDO()
@@ -152,8 +151,8 @@ void UControlRig::InitializeFromCDO()
 		UControlRig* CDO = GetClass()->GetDefaultObject<UControlRig>();
 
 		// copy hierarchy
-		Hierarchy = CDO->Hierarchy;
-		Hierarchy.Initialize();
+		GetHierarchy()->CopyHierarchy(CDO->GetHierarchy());
+		GetHierarchy()->ResetPoseToInitial(ERigElementType::All);
 
 		// copy draw container
 		DrawContainer = CDO->DrawContainer;
@@ -493,7 +492,7 @@ void UControlRig::Execute(const EControlRigState InState, const FName& InEventNa
 	Context.FramesPerSecond = GetCurrentFramesPerSecond();
 	Context.bDuringInteraction = IsInteracting();
 	Context.State = InState;
-	Context.Hierarchy = &Hierarchy;
+	Context.Hierarchy = DynamicHierarchy;
 
 	Context.ToWorldSpaceTransform = FTransform::Identity;
 	Context.OwningComponent = nullptr;
@@ -577,7 +576,7 @@ void UControlRig::Execute(const EControlRigState InState, const FName& InEventNa
 			{
 				if (UControlRig* CDO = Cast<UControlRig>(GetClass()->GetDefaultObject()))
 				{
-					Hierarchy.CopyInitialTransforms(CDO->Hierarchy);
+					DynamicHierarchy->CopyPose(CDO->DynamicHierarchy, false, true);
 				}
 			}
 
@@ -597,7 +596,7 @@ void UControlRig::Execute(const EControlRigState InState, const FName& InEventNa
 
 			if (bSetupModeEnabled)
 			{
-				Hierarchy.BoneHierarchy.ResetTransforms();
+				DynamicHierarchy->ResetPoseToInitial(ERigElementType::Bone);
 			}
 		}
 		else
@@ -709,44 +708,52 @@ void UControlRig::Execute(const EControlRigState InState, const FName& InEventNa
 	{
 		Context.DrawInterface->Instructions.Append(Context.DrawContainer->Instructions);
 
-		for (const FRigControl& Control : Hierarchy.ControlHierarchy)
+		DynamicHierarchy->ForEach<FRigControlElement>([this](FRigControlElement* ControlElement) -> bool
 		{
-			if (Control.bGizmoEnabled && Control.bGizmoVisible && !Control.bIsTransientControl && Control.bDrawLimits && (Control.bLimitTranslation || Control.bLimitRotation || Control.bLimitScale))
+			const FRigControlSettings& Settings = ControlElement->Settings;
+			
+			if (Settings.bGizmoEnabled &&
+				Settings.bGizmoVisible &&
+				!Settings.bIsTransientControl &&
+				Settings.bDrawLimits &&
+				(Settings.bLimitTranslation
+					|| Settings.bLimitRotation
+					|| Settings.bLimitScale))
 			{
 				// for now we don't draw rotational limits
-				if(!Control.bLimitTranslation)
+				if(!Settings.bLimitTranslation)
 				{
-					continue;
+					return true;
 				}
 
-				FTransform Transform = Hierarchy.ControlHierarchy.GetParentTransform(Control.Index);
-				FControlRigDrawInstruction Instruction(EControlRigDrawSettings::Lines, Control.GizmoColor, 0.f, Transform);
+				FTransform Transform = DynamicHierarchy->GetGlobalControlOffsetTransformByIndex(ControlElement->GetIndex());
+				FControlRigDrawInstruction Instruction(EControlRigDrawSettings::Lines, Settings.GizmoColor, 0.f, Transform);
 
-				switch (Control.ControlType)
+				switch (Settings.ControlType)
 				{
 					case ERigControlType::Float:
 					{
 						FVector MinPos = FVector::ZeroVector;
 						FVector MaxPos = FVector::ZeroVector;
 
-						switch (Control.PrimaryAxis)
+						switch (Settings.PrimaryAxis)
 						{
 							case ERigControlAxis::X:
 							{
-								MinPos.X = Control.MinimumValue.Get<float>();
-								MaxPos.X = Control.MaximumValue.Get<float>();
+								MinPos.X = Settings.MinimumValue.Get<float>();
+								MaxPos.X = Settings.MaximumValue.Get<float>();
 								break;
 							}
 							case ERigControlAxis::Y:
 							{
-								MinPos.Y = Control.MinimumValue.Get<float>();
-								MaxPos.Y = Control.MaximumValue.Get<float>();
+								MinPos.Y = Settings.MinimumValue.Get<float>();
+								MaxPos.Y = Settings.MaximumValue.Get<float>();
 								break;
 							}
 							case ERigControlAxis::Z:
 							{
-								MinPos.Z = Control.MinimumValue.Get<float>();
-								MaxPos.Z = Control.MaximumValue.Get<float>();
+								MinPos.Z = Settings.MinimumValue.Get<float>();
+								MaxPos.Z = Settings.MaximumValue.Get<float>();
 								break;
 							}
 						}
@@ -760,24 +767,24 @@ void UControlRig::Execute(const EControlRigState InState, const FName& InEventNa
 						FVector MinPos = FVector::ZeroVector;
 						FVector MaxPos = FVector::ZeroVector;
 
-						switch (Control.PrimaryAxis)
+						switch (Settings.PrimaryAxis)
 						{
 							case ERigControlAxis::X:
 							{
-								MinPos.X = (float)Control.MinimumValue.Get<int32>();
-								MaxPos.X = (float)Control.MaximumValue.Get<int32>();
+								MinPos.X = (float)Settings.MinimumValue.Get<int32>();
+								MaxPos.X = (float)Settings.MaximumValue.Get<int32>();
 								break;
 							}
 							case ERigControlAxis::Y:
 							{
-								MinPos.Y = (float)Control.MinimumValue.Get<int32>();
-								MaxPos.Y = (float)Control.MaximumValue.Get<int32>();
+								MinPos.Y = (float)Settings.MinimumValue.Get<int32>();
+								MaxPos.Y = (float)Settings.MaximumValue.Get<int32>();
 								break;
 							}
 							case ERigControlAxis::Z:
 							{
-								MinPos.Z = (float)Control.MinimumValue.Get<int32>();
-								MaxPos.Z = (float)Control.MaximumValue.Get<int32>();
+								MinPos.Z = (float)Settings.MinimumValue.Get<int32>();
+								MaxPos.Z = (float)Settings.MaximumValue.Get<int32>();
 								break;
 							}
 						}
@@ -789,10 +796,10 @@ void UControlRig::Execute(const EControlRigState InState, const FName& InEventNa
 					case ERigControlType::Vector2D:
 					{
 						Instruction.PrimitiveType = EControlRigDrawSettings::LineStrip;
-						FVector2D MinPos = Control.MinimumValue.Get<FVector2D>();
-						FVector2D MaxPos = Control.MaximumValue.Get<FVector2D>();
+						FVector2D MinPos = Settings.MinimumValue.Get<FVector2D>();
+						FVector2D MaxPos = Settings.MaximumValue.Get<FVector2D>();
 
-						switch (Control.PrimaryAxis)
+						switch (Settings.PrimaryAxis)
 						{
 							case ERigControlAxis::X:
 							{
@@ -832,31 +839,31 @@ void UControlRig::Execute(const EControlRigState InState, const FName& InEventNa
 					{
 						FVector MinPos, MaxPos;
 
-						switch (Control.ControlType)
+						switch (Settings.ControlType)
 						{
 							case ERigControlType::Position:
 							case ERigControlType::Scale:
 							{
-								MinPos = Control.MinimumValue.Get<FVector>();
-								MaxPos = Control.MaximumValue.Get<FVector>();
+								MinPos = Settings.MinimumValue.Get<FVector>();
+								MaxPos = Settings.MaximumValue.Get<FVector>();
 								break;
 							}
 							case ERigControlType::Transform:
 							{
-								MinPos = Control.MinimumValue.Get<FTransform>().GetLocation();
-								MaxPos = Control.MaximumValue.Get<FTransform>().GetLocation();
+								MinPos = Settings.MinimumValue.Get<FTransform>().GetLocation();
+								MaxPos = Settings.MaximumValue.Get<FTransform>().GetLocation();
 								break;
 							}
 							case ERigControlType::TransformNoScale:
 							{
-								MinPos = Control.MinimumValue.Get<FTransformNoScale>().Location;
-								MaxPos = Control.MaximumValue.Get<FTransformNoScale>().Location;
+								MinPos = Settings.MinimumValue.Get<FTransformNoScale>().Location;
+								MaxPos = Settings.MaximumValue.Get<FTransformNoScale>().Location;
 								break;
 							}
 							case ERigControlType::EulerTransform:
 							{
-								MinPos = Control.MinimumValue.Get<FEulerTransform>().Location;
-								MaxPos = Control.MaximumValue.Get<FEulerTransform>().Location;
+								MinPos = Settings.MinimumValue.Get<FEulerTransform>().Location;
+								MaxPos = Settings.MaximumValue.Get<FEulerTransform>().Location;
 								break;
 							}
 						}
@@ -896,7 +903,9 @@ void UControlRig::Execute(const EControlRigState InState, const FName& InEventNa
 					DrawInterface.Instructions.Add(Instruction);
 				}
 			}
-		}
+
+			return true;
+		});
 	}
 }
 
@@ -945,67 +954,6 @@ URigVM* UControlRig::GetVM()
 	return VM;
 }
 
-FTransform UControlRig::GetGlobalTransform(const FRigElementKey& InKey) const
-{
-	return Hierarchy.GetGlobalTransform(InKey);
-}
-
-void UControlRig::SetGlobalTransform(const FRigElementKey& InKey, const FTransform& InTransform, bool bPropagateTransform)
-{
-	Hierarchy.SetGlobalTransform(InKey, InTransform, bPropagateTransform);
-}
-
-FTransform UControlRig::GetGlobalTransform(const FName& BoneName) const
-{
-	int32 Index = Hierarchy.BoneHierarchy.GetIndex(BoneName);
-	if (Index != INDEX_NONE)
-	{
-		return Hierarchy.BoneHierarchy.GetGlobalTransform(Index);
-	}
-
-	return FTransform::Identity;
-
-}
-
-void UControlRig::SetGlobalTransform(const FName& BoneName, const FTransform& InTransform, bool bPropagateTransform)
-{
-	int32 Index = Hierarchy.BoneHierarchy.GetIndex(BoneName);
-	if (Index != INDEX_NONE)
-	{
-		return Hierarchy.BoneHierarchy.SetGlobalTransform(Index, InTransform, bPropagateTransform);
-	}
-}
-
-FTransform UControlRig::GetGlobalTransform(const int32 BoneIndex) const
-{
-	return Hierarchy.BoneHierarchy.GetGlobalTransform(BoneIndex);
-}
-
-void UControlRig::SetGlobalTransform(const int32 BoneIndex, const FTransform& InTransform, bool bPropagateTransform)
-{
-	Hierarchy.BoneHierarchy.SetGlobalTransform(BoneIndex, InTransform, bPropagateTransform);
-}
-
-float UControlRig::GetCurveValue(const FName& CurveName) const
-{
-	return Hierarchy.CurveContainer.GetValue(CurveName);
-}
-
-void UControlRig::SetCurveValue(const FName& CurveName, const float CurveValue)
-{
-	Hierarchy.CurveContainer.SetValue(CurveName, CurveValue);
-}
-
-float UControlRig::GetCurveValue(const int32 CurveIndex) const
-{
-	return Hierarchy.CurveContainer.GetValue(CurveIndex);
-}
-
-void UControlRig::SetCurveValue(const int32 CurveIndex, const float CurveValue)
-{
-	Hierarchy.CurveContainer.SetValue(CurveIndex, CurveValue);
-}
-
 void UControlRig::GetMappableNodeData(TArray<FName>& OutNames, TArray<FNodeItem>& OutNodeItems) const
 {
 	DECLARE_SCOPE_HIERARCHICAL_COUNTER_FUNC()
@@ -1014,13 +962,19 @@ void UControlRig::GetMappableNodeData(TArray<FName>& OutNames, TArray<FNodeItem>
 	OutNodeItems.Reset();
 
 	// now add all nodes
-	const FRigBoneHierarchy& BoneHierarchy = Hierarchy.BoneHierarchy;
+	DynamicHierarchy->ForEach<FRigBoneElement>([&OutNames, &OutNodeItems, this](FRigBoneElement* BoneElement) -> bool
+    {
+		OutNames.Add(BoneElement->GetName());
+		FRigElementKey ParentKey = DynamicHierarchy->GetFirstParent(BoneElement->GetKey());
+		if(ParentKey.Type != ERigElementType::Bone)
+		{
+			ParentKey.Name = NAME_None;
+		}
 
-	for (const FRigBone& Bone : BoneHierarchy)
-	{
-		OutNames.Add(Bone.Name);
-		OutNodeItems.Add(FNodeItem(Bone.ParentName, Bone.InitialTransform));
-	}
+		const FTransform GlobalInitial = DynamicHierarchy->GetGlobalTransformByIndex(BoneElement->GetIndex(), true);
+		OutNodeItems.Add(FNodeItem(ParentKey.Name, GlobalInitial));
+		return true;
+	});
 }
 
 UAnimationDataSourceRegistry* UControlRig::GetDataSourceRegistry()
@@ -1073,86 +1027,33 @@ void UControlRig::Serialize(FArchive& Ar)
 	Super::Serialize(Ar);
 
 	Ar.UsingCustomVersion(FControlRigObjectVersion::GUID);
-
-	if (Ar.IsLoading())
-	{
-		Hierarchy.ControlHierarchy.PostLoad();
-	}
-}
-
-const TArray<FRigSpace>& UControlRig::AvailableSpaces() const
-{
-	return Hierarchy.SpaceHierarchy.GetSpaces();
 }
 
-// do we need to return pointer? They can't save this
-FRigSpace* UControlRig::FindSpace(const FName& InSpaceName)
+TArray<FRigControlElement*> UControlRig::AvailableControls() const
 {
-	const int32 SpaceIndex = Hierarchy.SpaceHierarchy.GetIndex(InSpaceName);
-	if (SpaceIndex != INDEX_NONE)
+	if(DynamicHierarchy)
 	{
-		return &Hierarchy.SpaceHierarchy[SpaceIndex];
+		return DynamicHierarchy->GetElementsOfType<FRigControlElement>();
 	}
-
-	return nullptr;
-}
-
-FTransform UControlRig::GetSpaceGlobalTransform(const FName& InSpaceName)
-{
-	return Hierarchy.SpaceHierarchy.GetGlobalTransform(InSpaceName);
-}
-
-bool UControlRig::SetSpaceGlobalTransform(const FName& InSpaceName, const FTransform& InTransform)
-{
-	Hierarchy.SpaceHierarchy.SetGlobalTransform(InSpaceName, InTransform);
-	return true;
-}
-
-const TArray<FRigControl>& UControlRig::AvailableControls() const
-{
-#if WITH_EDITOR
-	if (AvailableControlsOverride.Num() > 0)
-	{
-		return AvailableControlsOverride;
-	}
-#endif
-	return Hierarchy.ControlHierarchy.GetControls();
+	return TArray<FRigControlElement*>();
 }
 
-FRigControl* UControlRig::FindControl(const FName& InControlName)
+FRigControlElement* UControlRig::FindControl(const FName& InControlName) const
 {
-	const int32 ControlIndex = Hierarchy.ControlHierarchy.GetIndex(InControlName);
-	if (ControlIndex != INDEX_NONE)
-	{
-		return &Hierarchy.ControlHierarchy[ControlIndex];
-	}
-
-#if WITH_EDITOR
-	if (TransientControls.Num() > 0)
-	{
-		for (int32 Index = 0; Index < TransientControls.Num(); Index++)
-		{
-			if (TransientControls[Index].Name == InControlName)
-			{
-				return &TransientControls[Index];
-			}
-		}
-	}
-#endif
-
-	return nullptr;
+	return DynamicHierarchy->Find<FRigControlElement>(FRigElementKey(InControlName, ERigElementType::Control));
 }
 
 FTransform UControlRig::SetupControlFromGlobalTransform(const FName& InControlName, const FTransform& InGlobalTransform)
 {
 	if (IsSetupModeEnabled())
 	{
-		FRigControl* Control = FindControl(InControlName);
-		if (Control && !Control->bIsTransientControl)
+		FRigControlElement* ControlElement = FindControl(InControlName);
+		if (ControlElement && !ControlElement->Settings.bIsTransientControl)
 		{
-			FTransform ParentTransform = Hierarchy.ControlHierarchy.GetParentTransform(Control->Index, false);
-			FTransform OffsetTransform = InGlobalTransform.GetRelativeTransform(ParentTransform);
-			Hierarchy.ControlHierarchy.SetControlOffset(Control->Index, OffsetTransform);
+			const FTransform ParentTransform = DynamicHierarchy->GetParentTransform(ControlElement, ERigTransformType::CurrentGlobal);
+			const FTransform OffsetTransform = InGlobalTransform.GetRelativeTransform(ParentTransform);
+			DynamicHierarchy->SetControlOffsetTransform(ControlElement, OffsetTransform, ERigTransformType::InitialLocal, true, true);
+			ControlElement->Offset.Current = ControlElement->Offset.Initial; 
 		}
 	}
 	return InGlobalTransform;
@@ -1161,72 +1062,60 @@ FTransform UControlRig::SetupControlFromGlobalTransform(const FName& InControlNa
 void UControlRig::CreateRigControlsForCurveContainer()
 {
 	const bool bCreateFloatControls = CVarControlRigCreateFloatControlsForCurves->GetInt() == 0 ? false : true;
-	if(bCreateFloatControls)
+	if(bCreateFloatControls && DynamicHierarchy)
 	{
-		FRigCurveContainer& CurveContainer = GetCurveContainer();
-		FString CTRLName(TEXT("CTRL_"));
-		for (FRigCurve& Curve : CurveContainer)
+		URigHierarchyController* Controller = DynamicHierarchy->GetController(true);
+		if(Controller == nullptr)
 		{
-			FString Name = Curve.Name.ToString();
-			if (Name.Contains(CTRLName) && Hierarchy.ControlHierarchy.GetIndex(Curve.Name) == INDEX_NONE) //-V1051
+			return;
+		}
+		static const FString CtrlPrefix(TEXT("CTRL_"));
+
+		DynamicHierarchy->ForEach<FRigCurveElement>([this, Controller](FRigCurveElement* CurveElement) -> bool
+        {
+			const FString Name = CurveElement->GetName().ToString();
+			
+			if (Name.Contains(CtrlPrefix) && !DynamicHierarchy->Contains(FRigElementKey(*Name, ERigElementType::Curve))) //-V1051
 			{
+				FRigControlSettings Settings;
+				Settings.ControlType = ERigControlType::Float;
+				Settings.bIsCurve = true;
+				Settings.bAnimatable = true;
+				Settings.bDrawLimits = false;
+				Settings.bGizmoEnabled = false;
+				Settings.bGizmoVisible = false;
+				Settings.GizmoColor = FLinearColor::Red;
+
 				FRigControlValue Value;
-				Value.Set<float>(Curve.Value);
-				FRigControl& Control = Hierarchy.ControlHierarchy.Add(Curve.Name, ERigControlType::Float, NAME_None, NAME_None, FTransform::Identity, Value
-	#if WITH_EDITORONLY_DATA
-					, Curve.Name,
-					FTransform::Identity,
-					FLinearColor::Red
-	#endif
-				);
-				Control.bIsCurve = true;
+				Value.Set<float>(CurveElement->Value);
+
+				Controller->AddControl(CurveElement->GetName(), FRigElementKey(), Settings, Value, FTransform::Identity, FTransform::Identity); 
 			}
-		}
+
+			return true;
+		});
 
 		ControlModified().AddUObject(this, &UControlRig::HandleOnControlModified);
 	}
 }
 
-void UControlRig::HandleOnControlModified(UControlRig* Subject, const FRigControl& Control, const FRigControlModifiedContext& Context)
+void UControlRig::HandleOnControlModified(UControlRig* Subject, FRigControlElement* Control, const FRigControlModifiedContext& Context)
 {
-	if (Control.bIsCurve)
+	if (Control->Settings.bIsCurve && DynamicHierarchy)
 	{
-		SetCurveValue(Control.Name, Control.Value.Get<float>());
+		const FRigControlValue Value = DynamicHierarchy->GetControlValue(Control, IsSetupModeEnabled() ? ERigControlValueType::Initial : ERigControlValueType::Current);
+		DynamicHierarchy->SetCurveValue(FRigElementKey(Control->GetName(), ERigElementType::Curve), Value.Get<float>());
 	}	
 }
 
-bool UControlRig::IsCurveControl(const FRigControl* InControl) const
+bool UControlRig::IsCurveControl(const FRigControlElement* InControlElement) const
 {
-	return InControl->bIsCurve;
+	return InControlElement->Settings.bIsCurve;
 }
 
 FTransform UControlRig::GetControlGlobalTransform(const FName& InControlName) const
 {
-#if WITH_EDITOR
-	if (TransientControls.Num() > 0)
-	{
-		for (int32 Index = 0; Index < TransientControls.Num(); Index++)
-		{
-			const FRigControl& Control = TransientControls[Index];
-			if (Control.Name == InControlName)
-			{
-				FTransform ParentTransform = FTransform::Identity;
-				if (Control.ParentIndex != INDEX_NONE)
-				{
-					ParentTransform = Hierarchy.BoneHierarchy.GetGlobalTransform(Control.ParentIndex);
-				}
-// 				else if (Control.SpaceIndex != INDEX_NONE)
-// 				{
-// 					ParentTransform = Hierarchy.SpaceHierarchy.GetGlobalTransform(Control.SpaceIndex);
-// 				}
-
-				FTransform Transform = Control.GetTransformFromValue(ERigControlValueType::Current);
-				return Transform * Control.OffsetTransform * ParentTransform;
-			}
-		}
-	}
-#endif
-	return Hierarchy.ControlHierarchy.GetGlobalTransform(InControlName);
+	return DynamicHierarchy->GetGlobalTransform(FRigElementKey(InControlName, ERigElementType::Control), false);
 }
 
 bool UControlRig::SetControlGlobalTransform(const FName& InControlName, const FTransform& InGlobalTransform, bool bNotify, const FRigControlModifiedContext& Context)
@@ -1240,10 +1129,10 @@ bool UControlRig::SetControlGlobalTransform(const FName& InControlName, const FT
 	FRigControlValue Value = GetControlValueFromGlobalTransform(InControlName, GlobalTransform);
 	if (OnFilterControl.IsBound())
 	{
-		FRigControl* Control = FindControl(InControlName);
+		FRigControlElement* Control = FindControl(InControlName);
 		if (Control)
 		{
-			OnFilterControl.Broadcast(this, *Control, Value);
+			OnFilterControl.Broadcast(this, Control, Value);
 		}
 	}
 
@@ -1253,72 +1142,42 @@ bool UControlRig::SetControlGlobalTransform(const FName& InControlName, const FT
 
 FRigControlValue UControlRig::GetControlValueFromGlobalTransform(const FName& InControlName, const FTransform& InGlobalTransform)
 {
-	FRigControlValue RetVal;
+	FRigControlValue Value;
 
-	if (FRigControl* Control = FindControl(InControlName))
+	if (FRigControlElement* ControlElement = FindControl(InControlName))
 	{
-		FTransform ParentTransform = FTransform::Identity;
-		if ((Control->ParentIndex != INDEX_NONE || Control->SpaceIndex != INDEX_NONE) && Control->bIsTransientControl)
-		{
-			if (Control->ParentIndex != INDEX_NONE)
-			{
-				ParentTransform = Hierarchy.BoneHierarchy.GetGlobalTransform(Control->ParentIndex);
-			}
-		}
-		else if(Control->Index != INDEX_NONE)
-		{
-			ParentTransform = Hierarchy.ControlHierarchy.GetParentTransform(Control->Index);
-		}
-
+		FTransform ParentTransform = DynamicHierarchy->GetControlOffsetTransform(ControlElement, ERigTransformType::CurrentGlobal);
 		FTransform Transform = InGlobalTransform.GetRelativeTransform(ParentTransform);
-
-		FRigControlValue PreviousVal = Control->Value;
-		Control->SetValueFromTransform(Transform, ERigControlValueType::Current);
-		RetVal = Control->Value;
-		Control->Value = PreviousVal;
+		Value.SetFromTransform(Transform, ControlElement->Settings.ControlType, ControlElement->Settings.PrimaryAxis);
 
 		if (ShouldApplyLimits())
 		{
-			Control->ApplyLimits(RetVal);
+			ControlElement->Settings.ApplyLimits(Value);
 		}
 	}
 
-	return RetVal;
+	return Value;
 }
 
 void UControlRig::SetControlLocalTransform(const FName& InControlName, const FTransform& InLocalTransform, bool bNotify, const FRigControlModifiedContext& Context)
 {
-
-	if (FRigControl* Control = FindControl(InControlName))
+	if (FRigControlElement* ControlElement = FindControl(InControlName))
 	{
-		FRigControlValue PreviousVal = Control->Value;
-		Control->SetValueFromTransform(InLocalTransform, ERigControlValueType::Current);
-		FRigControlValue ValueToSet = Control->Value;
-		Control->Value = PreviousVal;
+		FRigControlValue Value;
+		Value.SetFromTransform(InLocalTransform, ControlElement->Settings.ControlType, ControlElement->Settings.PrimaryAxis);
+
 		if (OnFilterControl.IsBound())
 		{
-			OnFilterControl.Broadcast(this, *Control, ValueToSet);
+			OnFilterControl.Broadcast(this, ControlElement, Value);
 			
 		}
-		SetControlValue(InControlName, ValueToSet, bNotify, Context);
+		SetControlValue(InControlName, Value, bNotify, Context);
 	}
-
 }
 
 FTransform UControlRig::GetControlLocalTransform(const FName& InControlName)
 {
-	FRigControl* Control = FindControl(InControlName);
-	if (Control->bIsTransientControl)
-	{
-		return Control->GetTransformFromValue(ERigControlValueType::Current);
-	}
-	return GetControlHierarchy().GetLocalTransform(InControlName);
-}
-
-bool UControlRig::SetControlSpace(const FName& InControlName, const FName& InSpaceName)
-{
-	Hierarchy.ControlHierarchy.SetSpace(InControlName, InSpaceName);
-	return true;
+	return DynamicHierarchy->GetLocalTransform(FRigElementKey(InControlName, ERigElementType::Control));
 }
 
 UControlRigGizmoLibrary* UControlRig::GetGizmoLibrary() const
@@ -1337,60 +1196,78 @@ UControlRigGizmoLibrary* UControlRig::GetGizmoLibrary() const
 	return nullptr;
 }
 
-void UControlRig::HandleOnControlUISettingChanged(FRigHierarchyContainer* Container, const FRigElementKey& InKey)
+void UControlRig::SelectControl(const FName& InControlName, bool bSelect)
 {
-	if (InKey.Type == ERigElementType::Control)
+	if(URigHierarchyController* Controller = DynamicHierarchy->GetController(true))
 	{
-		if (FRigControl* Control = FindControl(InKey.Name))
-		{
-			ControlModified().Broadcast(this, *Control, FRigControlModifiedContext(EControlRigSetKey::Never));
-		}
+		Controller->SelectElement(FRigElementKey(InControlName, ERigElementType::Control), bSelect);
 	}
 }
 
-void UControlRig::SelectControl(const FName& InControlName, bool bSelect)
-{
-	Hierarchy.ControlHierarchy.Select(InControlName, bSelect);
-}
-
 bool UControlRig::ClearControlSelection()
 {
-	return Hierarchy.ControlHierarchy.ClearSelection();
+	if(URigHierarchyController* Controller = DynamicHierarchy->GetController(true))
+	{
+		return Controller->ClearSelection();
+	}
+	return false;
 }
 
 TArray<FName> UControlRig::CurrentControlSelection() const
 {
-	return Hierarchy.ControlHierarchy.CurrentSelection();
+	TArray<FName> SelectedControlNames;
+	TArray<FRigBaseElement*> SelectedControls = DynamicHierarchy->GetSelectedElements(ERigElementType::Control);
+	for (FRigBaseElement* SelectedControl : SelectedControls)
+	{
+		SelectedControlNames.Add(SelectedControl->GetName());
+	}
+	return SelectedControlNames;
 }
 
 bool UControlRig::IsControlSelected(const FName& InControlName)const
 {
-	return Hierarchy.ControlHierarchy.IsSelected(InControlName);
+	if(FRigControlElement* ControlElement = FindControl(InControlName))
+	{
+		return DynamicHierarchy->IsSelected(ControlElement);
+	}
+	return false;
 }
 
-void UControlRig::HandleOnControlSelected(FRigHierarchyContainer* InContainer, const FRigElementKey& InKey, bool bSelected)
+void UControlRig::HandleHierarchyModified(ERigHierarchyNotification InNotification, URigHierarchy* InHierarchy,
+    const FRigBaseElement* InElement)
 {
-	if (InKey.Type == ERigElementType::Control)
+	switch(InNotification)
 	{
-		FRigControl* Control = FindControl(InKey.Name);
-		if (Control)
+		case ERigHierarchyNotification::ElementSelected:
+		case ERigHierarchyNotification::ElementDeselected:
+		{
+			if(FRigControlElement* ControlElement = Cast<FRigControlElement>((FRigBaseElement*)InElement))
+			{
+				const bool bSelected = InNotification == ERigHierarchyNotification::ElementSelected;
+				ControlSelected().Broadcast(this, ControlElement, bSelected);
+			}
+			break;
+		}
+		case ERigHierarchyNotification::ControlSettingChanged:
 		{
-			ControlSelected().Broadcast(this, *Control, bSelected);
+			if(FRigControlElement* ControlElement = Cast<FRigControlElement>((FRigBaseElement*)InElement))
+			{
+				ControlModified().Broadcast(this, ControlElement, FRigControlModifiedContext(EControlRigSetKey::Never));
+			}
+			break;
+		}
+		default:
+		{
+			break;
 		}
 	}
 }
 
 #if WITH_EDITOR
 
-void UControlRig::UpdateAvailableControls()
-{
-	AvailableControlsOverride = Hierarchy.ControlHierarchy.GetControls();
-	AvailableControlsOverride.Append(TransientControls);
-}
-
 FName UControlRig::AddTransientControl(URigVMPin* InPin, FRigElementKey SpaceKey)
 {
-	if (InPin == nullptr)
+	if ((InPin == nullptr) || (DynamicHierarchy == nullptr))
 	{
 		return NAME_None;
 	}
@@ -1404,94 +1281,110 @@ FName UControlRig::AddTransientControl(URigVMPin* InPin, FRigElementKey SpaceKey
 
 	RemoveTransientControl(InPin);
 
-	FName ControlName = *InPin->GetPinPath();
-
-	FRigControl NewControl;
-	NewControl.Name = ControlName;
-	NewControl.bIsTransientControl = true;
-	NewControl.GizmoTransform.SetScale3D(FVector::ZeroVector);
+	URigHierarchyController* Controller = DynamicHierarchy->GetController(true);
+	if(Controller == nullptr)
+	{
+		return NAME_None;
+	}
 
 	URigVMPin* PinForLink = InPin->GetPinForLink();
+
+	const FName ControlName = GetNameForTransientControl(InPin);
+	FTransform GizmoTransform = FTransform::Identity;
+	GizmoTransform.SetScale3D(FVector::ZeroVector);
+
+	FRigControlSettings Settings;
 	if (URigVMPin* ColorPin = PinForLink->GetNode()->FindPin(TEXT("Color")))
 	{
 		if (ColorPin->GetCPPType() == TEXT("FLinearColor"))
 		{
 			FRigControlValue Value;
-			NewControl.GizmoColor = Value.SetFromString<FLinearColor>(ColorPin->GetDefaultValue());
+			Settings.GizmoColor = Value.SetFromString<FLinearColor>(ColorPin->GetDefaultValue());
 		}
 	}
+	Settings.bIsTransientControl = true;
 
+	FRigElementKey Parent;
 	if(SpaceKey.IsValid() && SpaceKey.Type == ERigElementType::Bone)
 	{
-		NewControl.ParentIndex = Hierarchy.BoneHierarchy.GetIndex(SpaceKey.Name);
-		if (NewControl.ParentIndex != INDEX_NONE)
-		{
-			NewControl.ParentName = Hierarchy.BoneHierarchy[NewControl.ParentIndex].Name;
-		}
+		Parent = DynamicHierarchy->GetFirstParent(SpaceKey);
 	}
 
-	TransientControls.Add(NewControl);
+	Controller->ClearSelection();
+
+    const FRigElementKey ControlKey = Controller->AddControl(
+    	ControlName,
+    	Parent,
+    	Settings,
+    	FRigControlValue::Make(FTransform::Identity),
+    	FTransform::Identity,
+    	GizmoTransform, false);
 
 	SetTransientControlValue(InPin);
-	UpdateAvailableControls();
+
+	if(const FRigBaseElement* Element = DynamicHierarchy->Find(ControlKey))
+	{
+		DynamicHierarchy->Notify(ERigHierarchyNotification::ElementSelected, Element);
+	}
 
 	return ControlName;
 }
 
 bool UControlRig::SetTransientControlValue(URigVMPin* InPin)
 {
-	FString PinPath = InPin->GetPinForLink()->GetPinPath();
-	FString OriginalPinPath = InPin->GetOriginalPinFromInjectedNode()->GetPinPath();
-	for (FRigControl& Control : TransientControls)
+	const FName ControlName = GetNameForTransientControl(InPin);
+	if (FRigControlElement* ControlElement = FindControl(ControlName))
 	{
-		FString ControlName = Control.Name.ToString();
-		if (ControlName == PinPath || ControlName == OriginalPinPath)
+		FString DefaultValue = InPin->GetPinForLink()->GetDefaultValue();
+		if (!DefaultValue.IsEmpty())
 		{
-			FString DefaultValue = InPin->GetPinForLink()->GetDefaultValue();
-			if (!DefaultValue.IsEmpty())
+			if (InPin->GetCPPType() == TEXT("FVector"))
 			{
-				if (InPin->GetCPPType() == TEXT("FVector"))
-				{
-					Control.ControlType = ERigControlType::Position;
-					FVector Value = Control.Value.SetFromString<FVector>(DefaultValue);
-					Control.InitialValue.Set<FVector>(Value);
-				}
-				else if (InPin->GetCPPType() == TEXT("FQuat"))
-				{
-					Control.ControlType = ERigControlType::Rotator;
-					FQuat Value = Control.Value.SetFromString<FQuat>(DefaultValue);
-					Control.Value.Set<FRotator>(Value.Rotator());
-					Control.InitialValue.Set<FRotator>(Value.Rotator());
-				}
-				else
-				{
-					Control.ControlType = ERigControlType::Transform;
-					FTransform Value = Control.Value.SetFromString<FTransform>(DefaultValue);
-					Control.InitialValue.Set<FTransform>(Value);
-				}
+				ControlElement->Settings.ControlType = ERigControlType::Position;
+				FRigControlValue Value;
+				Value.SetFromString<FVector>(DefaultValue);
+				DynamicHierarchy->SetControlValue(ControlElement, Value, ERigControlValueType::Current, false);
+			}
+			else if (InPin->GetCPPType() == TEXT("FQuat"))
+			{
+				ControlElement->Settings.ControlType = ERigControlType::Rotator;
+				FRigControlValue Value;
+				Value.SetFromString<FRotator>(DefaultValue);
+				DynamicHierarchy->SetControlValue(ControlElement, Value, ERigControlValueType::Current, false);
+			}
+			else
+			{
+				ControlElement->Settings.ControlType = ERigControlType::Transform;
+				FRigControlValue Value;
+				Value.SetFromString<FTransform>(DefaultValue);
+				DynamicHierarchy->SetControlValue(ControlElement, Value, ERigControlValueType::Current, false);
 			}
-			return true;
 		}
+		return true;
 	}
 	return false;
 }
 
 FName UControlRig::RemoveTransientControl(URigVMPin* InPin)
 {
-	if (InPin == nullptr)
+	if ((InPin == nullptr) || (DynamicHierarchy == nullptr))
+	{
+		return NAME_None;
+	}
+
+	URigHierarchyController* Controller = DynamicHierarchy->GetController(true);
+	if(Controller == nullptr)
 	{
 		return NAME_None;
 	}
 
-	FName ControlName = *InPin->GetPinPath();
-	for (int32 Index = 0; Index < TransientControls.Num(); Index++)
+	const FName ControlName = GetNameForTransientControl(InPin);
+	if(FRigControlElement* ControlElement = FindControl(ControlName))
 	{
-		if (TransientControls[Index].Name == ControlName)
+		//Hierarchy->Notify(ERigHierarchyNotification::ElementDeselected, ControlElement);
+		if(Controller->RemoveElement(ControlElement))
 		{
-			FRigControl ControlToRemove = TransientControls[Index];
-			TransientControls.RemoveAt(Index);
-			UpdateAvailableControls();
-			return ControlToRemove.Name;
+			return ControlName;
 		}
 	}
 
@@ -1505,49 +1398,82 @@ FName UControlRig::AddTransientControl(const FRigElementKey& InElement)
 		return NAME_None;
 	}
 
-	RemoveTransientControl(InElement);
+	if(DynamicHierarchy == nullptr)
+	{
+		return NAME_None;
+	}
+
+	URigHierarchyController* Controller = DynamicHierarchy->GetController(true);
+	if(Controller == nullptr)
+	{
+		return NAME_None;
+	}
 
-	int32 ElementIndex = Hierarchy.GetIndex(InElement);
+	RemoveTransientControl(InElement);
+	const int32 ElementIndex = DynamicHierarchy->GetIndex(InElement);
 	if (ElementIndex == INDEX_NONE)
 	{
 		return NAME_None;
 	}
 
-	FName ControlName = GetNameForTransientControl(InElement);
+	const FName ControlName = GetNameForTransientControl(InElement);
+
+	FTransform GizmoTransform = FTransform::Identity;
+	GizmoTransform.SetScale3D(FVector::ZeroVector);
 
-	FRigControl NewControl;
-	NewControl.Name = ControlName;
-	NewControl.bIsTransientControl = true;
-	NewControl.GizmoTransform.SetScale3D(FVector::ZeroVector);
+	FRigControlSettings Settings;
+	Settings.bIsTransientControl = true;
 
+	FRigElementKey Parent;
 	switch (InElement.Type)
 	{
 		case ERigElementType::Bone:
 		{
-			NewControl.ParentIndex = Hierarchy.BoneHierarchy[ElementIndex].ParentIndex;
-			NewControl.ParentName = Hierarchy.BoneHierarchy[ElementIndex].ParentName;
+			Parent = DynamicHierarchy->GetFirstParent(InElement);
 			break;
 		}
 		case ERigElementType::Space:
 		{
-			NewControl.SpaceIndex = ElementIndex;
-			NewControl.SpaceName = InElement.Name;
+			Parent = InElement;
+			break;
+		}
+		default:
+		{
 			break;
 		}
 	}
 
-	TransientControls.Add(NewControl);
+	TArray<FRigElementKey> SelectedControls = DynamicHierarchy->GetSelectedKeys(ERigElementType::Control);
+	for(const FRigElementKey& SelectedControl : SelectedControls)
+	{
+		Controller->DeselectElement(SelectedControl);
+	}
+
+	const FRigElementKey ControlKey = Controller->AddControl(
+        ControlName,
+        Parent,
+        Settings,
+        FRigControlValue::Make(FTransform::Identity),
+        FTransform::Identity,
+        GizmoTransform, false);
 
 	if (InElement.Type == ERigElementType::Bone)
 	{
-		if (PreviewInstance && Hierarchy.BoneHierarchy[ElementIndex].Type == ERigBoneType::Imported)
+		if(FRigBoneElement* BoneElement = DynamicHierarchy->Find<FRigBoneElement>(InElement))
 		{
-			PreviewInstance->ModifyBone(InElement.Name);
+			if (PreviewInstance && BoneElement->BoneType == ERigBoneType::Imported)
+			{
+				PreviewInstance->ModifyBone(InElement.Name);
+			}
 		}
 	}
 
 	SetTransientControlValue(InElement);
-	UpdateAvailableControls();
+
+	if(const FRigBaseElement* Element = DynamicHierarchy->Find(ControlKey))
+	{
+		DynamicHierarchy->Notify(ERigHierarchyNotification::ElementSelected, Element);
+	}
 
 	return ControlName;
 }
@@ -1559,39 +1485,34 @@ bool UControlRig::SetTransientControlValue(const FRigElementKey& InElement)
 		return false;
 	}
 
-	FName ControlName = GetNameForTransientControl(InElement);
-	for (FRigControl& Control : TransientControls)
+	const FName ControlName = GetNameForTransientControl(InElement);
+	if(FRigControlElement* ControlElement = FindControl(ControlName))
 	{
-		if (Control.Name == ControlName)
+		if (InElement.Type == ERigElementType::Bone)
 		{
-			Control.ControlType = ERigControlType::Transform;
+			const FTransform LocalTransform = DynamicHierarchy->GetLocalTransform(InElement);
+			DynamicHierarchy->SetTransform(ControlElement, LocalTransform, ERigTransformType::InitialLocal, true, false);
+			DynamicHierarchy->SetTransform(ControlElement, LocalTransform, ERigTransformType::CurrentLocal, true, false);
 
-			if (InElement.Type == ERigElementType::Bone)
+			if (PreviewInstance)
 			{
-				FTransform LocalTransform = Hierarchy.GetLocalTransform(InElement);
-				Control.Value.Set<FTransform>(LocalTransform);
-				Control.InitialValue.Set<FTransform>(LocalTransform);
-
-				if (PreviewInstance)
+				if (FAnimNode_ModifyBone* Modify = PreviewInstance->FindModifiedBone(InElement.Name))
 				{
-					if (FAnimNode_ModifyBone* Modify = PreviewInstance->FindModifiedBone(InElement.Name))
-					{
-						Modify->Translation = LocalTransform.GetTranslation();
-						Modify->Rotation = LocalTransform.GetRotation().Rotator();
-						Modify->TranslationSpace = EBoneControlSpace::BCS_ParentBoneSpace;
-						Modify->RotationSpace = EBoneControlSpace::BCS_ParentBoneSpace;
-					}
+					Modify->Translation = LocalTransform.GetTranslation();
+					Modify->Rotation = LocalTransform.GetRotation().Rotator();
+					Modify->TranslationSpace = EBoneControlSpace::BCS_ParentBoneSpace;
+					Modify->RotationSpace = EBoneControlSpace::BCS_ParentBoneSpace;
 				}
 			}
-			else if (InElement.Type == ERigElementType::Space)
-			{
-				FTransform GlobalTransform = Hierarchy.GetGlobalTransform(InElement);
-				Control.Value.Set<FTransform>(GlobalTransform);
-				Control.InitialValue.Set<FTransform>(GlobalTransform);
-			}
-
-			return true;
 		}
+		else if (InElement.Type == ERigElementType::Space)
+		{
+			const FTransform GlobalTransform = DynamicHierarchy->GetGlobalTransform(InElement);
+			DynamicHierarchy->SetTransform(ControlElement, GlobalTransform, ERigTransformType::InitialGlobal, true, false);
+			DynamicHierarchy->SetTransform(ControlElement, GlobalTransform, ERigTransformType::CurrentGlobal, true, false);
+		}
+
+		return true;
 	}
 	return false;
 }
@@ -1603,20 +1524,42 @@ FName UControlRig::RemoveTransientControl(const FRigElementKey& InElement)
 		return NAME_None;
 	}
 
-	FName ControlName = GetNameForTransientControl(InElement);
-	for (int32 Index = 0; Index < TransientControls.Num(); Index++)
+	URigHierarchyController* Controller = DynamicHierarchy->GetController(true);
+	if(Controller == nullptr)
 	{
-		if (TransientControls[Index].Name == ControlName)
+		return NAME_None;
+	}
+
+	const FName ControlName = GetNameForTransientControl(InElement);
+	if(FRigControlElement* ControlElement = FindControl(ControlName))
+	{
+		//Hierarchy->Notify(ERigHierarchyNotification::ElementDeselected, ControlElement);
+		if(Controller->RemoveElement(ControlElement))
 		{
-			FRigControl ControlToRemove = TransientControls[Index];
-			TransientControls.RemoveAt(Index);
-			UpdateAvailableControls();
-			return ControlToRemove.Name;
+			return ControlName;
 		}
 	}
 
 	return NAME_None;
+}
+
+FName UControlRig::GetNameForTransientControl(URigVMPin* InPin) const
+{
+	check(InPin);
+	check(DynamicHierarchy);
+	
+	const FString OriginalPinPath = InPin->GetOriginalPinFromInjectedNode()->GetPinPath();
+	return DynamicHierarchy->GetSanitizedName(FString::Printf(TEXT("ControlForPin_%s"), *OriginalPinPath));
+}
 
+FString UControlRig::GetPinNameFromTransientControl(const FRigElementKey& InKey)
+{
+	FString Name = InKey.Name.ToString();
+	if(Name.StartsWith(TEXT("ControlForPin_")))
+	{
+		Name.RightChopInline(14);
+	}
+	return Name;
 }
 
 FName UControlRig::GetNameForTransientControl(const FRigElementKey& InElement)
@@ -1626,59 +1569,120 @@ FName UControlRig::GetNameForTransientControl(const FRigElementKey& InElement)
 		return InElement.Name;
 	}
 
-	FName EnumName = StaticEnum<ERigElementType>()->GetNameByValue((int64)InElement.Type);
+	const FName EnumName = *StaticEnum<ERigElementType>()->GetDisplayNameTextByValue((int64)InElement.Type).ToString();
 	return *FString::Printf(TEXT("ControlForRigElement_%s_%s"), *EnumName.ToString(), *InElement.Name.ToString());
 }
 
+FRigElementKey UControlRig::GetElementKeyFromTransientControl(const FRigElementKey& InKey)
+{
+	static FString ControlRigForElementBoneName;
+	static FString ControlRigForElementSpaceName;
+
+	if (ControlRigForElementBoneName.IsEmpty())
+	{
+		ControlRigForElementBoneName = FString::Printf(TEXT("ControlForRigElement_%s_"),
+            *StaticEnum<ERigElementType>()->GetDisplayNameTextByValue((int64)ERigElementType::Bone).ToString());
+		ControlRigForElementSpaceName = FString::Printf(TEXT("ControlForRigElement_%s_"),
+            *StaticEnum<ERigElementType>()->GetDisplayNameTextByValue((int64)ERigElementType::Space).ToString());
+	}
+	
+	FString Name = InKey.Name.ToString();
+	if(Name.StartsWith(ControlRigForElementBoneName))
+	{
+		Name.RightChopInline(ControlRigForElementBoneName.Len());
+		return FRigElementKey(*Name, ERigElementType::Bone);
+	}
+	if(Name.StartsWith(ControlRigForElementSpaceName))
+	{
+		Name.RightChopInline(ControlRigForElementSpaceName.Len());
+		return FRigElementKey(*Name, ERigElementType::Space);
+	}
+	
+	return FRigElementKey();;
+}
+
 void UControlRig::ClearTransientControls()
 {
-	for (int32 Index = TransientControls.Num() - 1; Index >= 0; Index--)
+	if(DynamicHierarchy == nullptr)
+	{
+		return;
+	}
+
+	URigHierarchyController* Controller = DynamicHierarchy->GetController(true);
+	if(Controller == nullptr)
 	{
-		FRigControl ControlToRemove = TransientControls[Index];
-		TransientControls.RemoveAt(Index);
-		UpdateAvailableControls();
+		return;
+	}
+
+	const TArray<FRigControlElement*> ControlsToRemove = DynamicHierarchy->GetTransientControls();
+	for (FRigControlElement* ControlToRemove : ControlsToRemove)
+	{
+		DynamicHierarchy->Notify(ERigHierarchyNotification::ElementDeselected, ControlToRemove);
+		Controller->RemoveElement(ControlToRemove);
 	}
 }
 
 #endif
 
-void UControlRig::HandleOnRigEvent(FRigHierarchyContainer* InContainer, const FRigEventContext& InEvent)
+void UControlRig::HandleHierarchyEvent(URigHierarchy* InHierarchy, const FRigEventContext& InEvent)
 {
 	if (RigEventDelegate.IsBound())
 	{
-		RigEventDelegate.Broadcast(InContainer, InEvent);
+		RigEventDelegate.Broadcast(InHierarchy, InEvent);
 	}
 
 	switch (InEvent.Event)
 	{
 		case ERigEvent::RequestAutoKey:
 		{
-			int32 Index = InContainer->GetIndex(InEvent.Key);
+			int32 Index = InHierarchy->GetIndex(InEvent.Key);
 			if (Index != INDEX_NONE && InEvent.Key.Type == ERigElementType::Control)
 			{
-				const FRigControl& Control = InContainer->ControlHierarchy[Index];
-
-				FRigControlModifiedContext Context;
-				Context.SetKey = EControlRigSetKey::Always;
-				Context.LocalTime = InEvent.LocalTime;
-				Context.EventName = InEvent.SourceEventName;
-				ControlModified().Broadcast(this, Control, Context);
+				if(FRigControlElement* ControlElement = InHierarchy->GetChecked<FRigControlElement>(Index))
+				{
+					FRigControlModifiedContext Context;
+					Context.SetKey = EControlRigSetKey::Always;
+					Context.LocalTime = InEvent.LocalTime;
+					Context.EventName = InEvent.SourceEventName;
+					ControlModified().Broadcast(this, ControlElement, Context);
+				}
 			}
 		}
+		default:
+		{
+			break;
+		}
 	}
 }
 
-void UControlRig::GetControlsInOrder(TArray<FRigControl>& SortedControls) const
+void UControlRig::GetControlsInOrder(TArray<FRigControlElement*>& SortedControls) const
 {
-	TArray<FRigElementKey> Keys = Hierarchy.GetAllItems(true);
-	for (const FRigElementKey& Key : Keys)
-	{ 
-		if (Key.Type == ERigElementType::Control)
-		{
-			const FRigControl& RigControl = Hierarchy.ControlHierarchy[Key.Name];
-			FRigControl NewControl(RigControl);
-			SortedControls.Add(NewControl);
-		}
+	SortedControls.Reset();
+
+	if(DynamicHierarchy == nullptr)
+	{
+		return;
+	}
+
+	TArray<bool> ElementVisited;
+	ElementVisited.AddZeroed(DynamicHierarchy->Num());
+		
+	for (FRigBaseElement* Element : *DynamicHierarchy)
+	{
+		DynamicHierarchy->Traverse(Element, true, [&ElementVisited, &SortedControls](FRigBaseElement* InElement, bool& bContinue)
+        {
+            bContinue = !ElementVisited[InElement->GetIndex()];
+
+            if(bContinue)
+            {
+            	if(FRigControlElement* ControlElement = Cast<FRigControlElement>(InElement))
+            	{
+            		SortedControls.Add(ControlElement);
+            	}
+            	
+                ElementVisited[InElement->GetIndex()] = true;
+            }
+        });
 	}
 }
 
@@ -1845,36 +1849,40 @@ const TArray<UAssetUserData*>* UControlRig::GetAssetUserDataArray() const
 
 void UControlRig::CopyPoseFromOtherRig(UControlRig* Subject)
 {
+	check(DynamicHierarchy);
 	check(Subject);
+	URigHierarchy* OtherHierarchy = Subject->GetHierarchy();
+	check(OtherHierarchy);
 
-	const FRigBoneHierarchy& SourceBones = Subject->GetBoneHierarchy();
-	FRigBoneHierarchy& TargetBones = GetBoneHierarchy();
-
-	for (const FRigBone& SourceBone : SourceBones)
+	for (FRigBaseElement* Element : *DynamicHierarchy)
 	{
-		int32 Index = TargetBones.GetIndex(SourceBone.Name);
-		if (Index != INDEX_NONE)
+		FRigBaseElement* OtherElement = OtherHierarchy->Find(Element->GetKey());
+		if(OtherElement == nullptr)
 		{
-			TargetBones[Index].LocalTransform = SourceBone.LocalTransform;
+			continue;
 		}
-	}
-
-	TargetBones.RecomputeGlobalTransforms();
 
-	const FRigCurveContainer& SourceCurves = Subject->GetCurveContainer();
-	FRigCurveContainer& TargetCurves = GetCurveContainer();
+		if(OtherElement->GetType() != Element->GetType())
+		{
+			continue;
+		}
 
-	for (const FRigCurve& SourceCurve : SourceCurves)
-	{
-		int32 Index = TargetCurves.GetIndex(SourceCurve.Name);
-		if (Index != INDEX_NONE)
+		if(FRigBoneElement* BoneElement = Cast<FRigBoneElement>(Element))
 		{
-			TargetCurves[Index].Value = SourceCurve.Value;
+			FRigBoneElement* OtherBoneElement = CastChecked<FRigBoneElement>(OtherElement);
+			const FTransform Transform = OtherHierarchy->GetTransform(OtherBoneElement, ERigTransformType::CurrentLocal);
+			DynamicHierarchy->SetTransform(BoneElement, Transform, ERigTransformType::CurrentLocal, true, false);
+		}
+		else if(FRigCurveElement* CurveElement = Cast<FRigCurveElement>(Element))
+		{
+			FRigCurveElement* OtherCurveElement = CastChecked<FRigCurveElement>(OtherElement);
+			const float Value = OtherHierarchy->GetCurveValue(OtherCurveElement);
+			DynamicHierarchy->SetCurveValue(CurveElement, Value, false);
 		}
 	}
 }
 
-void UControlRig::HandleInteractionRigControlModified(UControlRig* Subject, const FRigControl& Control, const FRigControlModifiedContext& Context)
+void UControlRig::HandleInteractionRigControlModified(UControlRig* Subject, FRigControlElement* Control, const FRigControlModifiedContext& Context)
 {
 	check(Subject);
 
@@ -1891,15 +1899,15 @@ void UControlRig::HandleInteractionRigControlModified(UControlRig* Subject, cons
 
 	if (const FRigInfluenceMap* InfluenceMap = Subject->FindInfluenceMap(Context.EventName))
 	{
-		if (const FRigInfluenceEntry* InfluenceEntry = InfluenceMap->Find(Control.GetElementKey()))
+		if (const FRigInfluenceEntry* InfluenceEntry = InfluenceMap->Find(Control->GetKey()))
 		{
 			for (const FRigElementKey& AffectedKey : *InfluenceEntry)
 			{
 				if (AffectedKey.Type == ERigElementType::Control)
 				{
-					if (const FRigControl* AffectedControl = FindControl(AffectedKey.Name))
+					if (FRigControlElement* AffectedControl = FindControl(AffectedKey.Name))
 					{
-						QueuedModifiedControls.Add(*AffectedControl);
+						QueuedModifiedControls.Add(AffectedControl->GetKey());
 					}
 				}
 				else if (
@@ -1908,9 +1916,9 @@ void UControlRig::HandleInteractionRigControlModified(UControlRig* Subject, cons
 				{
 					// special case controls with a CONTROL suffix
 					FName BoneControlName = *FString::Printf(TEXT("%s_CONTROL"), *AffectedKey.Name.ToString());
-					if (const FRigControl* AffectedControl = FindControl(BoneControlName))
+					if (FRigControlElement* AffectedControl = FindControl(BoneControlName))
 					{
-						QueuedModifiedControls.Add(*AffectedControl);
+						QueuedModifiedControls.Add(AffectedControl->GetKey());
 					}
 				}
 			}
@@ -1948,13 +1956,16 @@ void UControlRig::HandleInteractionRigExecuted(UControlRig* Subject, EControlRig
 	Context.EventName = FRigUnit_InverseExecution::EventName;
 	Context.SetKey = EControlRigSetKey::DoNotCare;
 
-	for (const FRigControl& QueuedModifiedControl : QueuedModifiedControls)
+	for (const FRigElementKey& QueuedModifiedControl : QueuedModifiedControls)
 	{
-		ControlModified().Broadcast(this, QueuedModifiedControl, Context);
+		if(FRigControlElement* ControlElement = FindControl(QueuedModifiedControl.Name))
+		{
+			ControlModified().Broadcast(this, ControlElement, Context);
+		}
 	}
 }
 
-void UControlRig::HandleInteractionRigControlSelected(UControlRig* Subject, const FRigControl& InControl, bool bSelected, bool bInverted)
+void UControlRig::HandleInteractionRigControlSelected(UControlRig* Subject, FRigControlElement* Control, bool bSelected, bool bInverted)
 {
 	check(Subject);
 
@@ -2025,14 +2036,14 @@ void UControlRig::HandleInteractionRigControlSelected(UControlRig* Subject, cons
 			}
 		};
 
-		Local::SelectAffectedElements(this, InfluenceMap, InControl.GetElementKey(), bSelected, bInverted);
+		Local::SelectAffectedElements(this, InfluenceMap, Control->GetKey(), bSelected, bInverted);
 
 		if (bInverted)
 		{
-			FString ControlName = InControl.Name.ToString();
+			const FString ControlName = Control->GetName().ToString();
 			if (ControlName.EndsWith(TEXT("_CONTROL")))
 			{
-				FString BaseName = ControlName.Left(ControlName.Len() - 8);
+				const FString BaseName = ControlName.Left(ControlName.Len() - 8);
 				Local::SelectAffectedElements(this, InfluenceMap, FRigElementKey(*BaseName, ERigElementType::Bone), bSelected, bInverted);
 				Local::SelectAffectedElements(this, InfluenceMap, FRigElementKey(*BaseName, ERigElementType::Curve), bSelected, bInverted);
 			}
@@ -2165,8 +2176,8 @@ FRigVMExternalVariable UControlRig::GetExternalVariableFromPinType(const FName&
 
 FRigVMExternalVariable UControlRig::GetExternalVariableFromDescription(const FBPVariableDescription& InVariableDescription)
 {
-	bool bIsPublic = !((InVariableDescription.PropertyFlags & CPF_DisableEditOnInstance) == CPF_DisableEditOnInstance);
-	bool bIsReadOnly = ((InVariableDescription.PropertyFlags & CPF_BlueprintReadOnly) == CPF_BlueprintReadOnly);
+	const bool bIsPublic = !((InVariableDescription.PropertyFlags & CPF_DisableEditOnInstance) == CPF_DisableEditOnInstance);
+	const bool bIsReadOnly = ((InVariableDescription.PropertyFlags & CPF_BlueprintReadOnly) == CPF_BlueprintReadOnly);
 	return GetExternalVariableFromPinType(InVariableDescription.VarName, InVariableDescription.VarType, bIsPublic, bIsReadOnly);
 }
 
@@ -2180,15 +2191,21 @@ void UControlRig::SetBoneInitialTransformsFromSkeletalMesh(USkeletalMesh* InSkel
 
 void UControlRig::SetBoneInitialTransformsFromRefSkeleton(const FReferenceSkeleton& InReferenceSkeleton)
 {
-	for (const FRigBone& Bone : GetBoneHierarchy())
+	check(DynamicHierarchy);
+
+	DynamicHierarchy->ForEach<FRigBoneElement>([this, InReferenceSkeleton](FRigBoneElement* BoneElement) -> bool
 	{
-		int32 BoneIndex = InReferenceSkeleton.FindBoneIndex(Bone.Name);
-		if (BoneIndex != INDEX_NONE)
+		if(BoneElement->BoneType == ERigBoneType::Imported)
 		{
-			FTransform LocalInitialTransform = InReferenceSkeleton.GetRefBonePose()[BoneIndex];
-			GetBoneHierarchy().SetInitialLocalTransform(Bone.Index, LocalInitialTransform);
+			const int32 BoneIndex = InReferenceSkeleton.FindBoneIndex(BoneElement->GetName());
+			if (BoneIndex != INDEX_NONE)
+			{
+				const FTransform LocalInitialTransform = InReferenceSkeleton.GetRefBonePose()[BoneIndex];
+				DynamicHierarchy->SetTransform(BoneElement, LocalInitialTransform, ERigTransformType::InitialLocal, true, false);
+			}
 		}
-	}
+		return true;
+	});
 	bResetInitialTransformsBeforeSetup = false;
 }
 
