Control Rig: New rig hierarchy data structure refactoring

diff --git a/Engine/Plugins/Experimental/ControlRig/Source/ControlRig/Private/ControlRigComponent.cpp b/Engine/Plugins/Experimental/ControlRig/Source/ControlRig/Private/ControlRigComponent.cpp
--- a/Engine/Plugins/Experimental/ControlRig/Source/ControlRig/Private/ControlRigComponent.cpp
+++ b/Engine/Plugins/Experimental/ControlRig/Source/ControlRig/Private/ControlRigComponent.cpp
@@ -194,14 +194,17 @@ FBoxSphereBounds UControlRigComponent::CalcBounds(const FTransform& LocalToWorld
 			}
 		}
 
-		FTransform Transform = GetComponentToWorld();
+		const FTransform Transform = GetComponentToWorld();
 
 		// Get bounding box for bones
-		const FRigBoneHierarchy& BoneHierarchy = ControlRig->GetBoneHierarchy();
-		for (const FRigBone& Bone : BoneHierarchy)
+		URigHierarchy* Hierarchy = ControlRig->GetHierarchy();
+		check(Hierarchy);
+
+		Hierarchy->ForEach<FRigTransformElement>([&BBox, Transform, Hierarchy](FRigTransformElement* TransformElement) -> bool
 		{
-			BBox += Transform.TransformPosition(Bone.GlobalTransform.GetLocation());
-		}
+			BBox += Transform.TransformPosition(Hierarchy->GetTransform(TransformElement, ERigTransformType::CurrentGlobal).GetLocation());
+			return true;
+        });
 	}
 
 	if (BBox.IsValid)
@@ -270,7 +273,7 @@ void UControlRigComponent::Initialize()
 		else
 		{
 			CR->DrawInterface.Reset();
-			CR->GetHierarchy()->Initialize(true);
+			CR->GetHierarchy()->ResetPoseToInitial(ERigElementType::All);
 			CR->RequestInit();
 		}
 	}
@@ -294,7 +297,7 @@ void UControlRigComponent::Update(float DeltaTime)
 
 			if (bResetTransformBeforeTick)
 			{
-				CR->GetBoneHierarchy().ResetTransforms();
+				CR->GetHierarchy()->ResetPoseToInitial(ERigElementType::Bone);
 			}
 
 #if WITH_EDITOR
@@ -328,39 +331,11 @@ TArray<FName> UControlRigComponent::GetElementNames(ERigElementType ElementType)
 
 	if (UControlRig* CR = SetupControlRigIfRequired())
 	{
-		switch (ElementType)
+		for (FRigBaseElement* Element : *CR->GetHierarchy())
 		{
-			case ERigElementType::Bone:
-			{
-				for (const FRigBone& Element : CR->GetBoneHierarchy())
-				{
-					Names.Add(Element.Name);
-				}
-				break;
-			}
-			case ERigElementType::Space:
+			if(Element->IsTypeOf(ElementType))
 			{
-				for (const FRigSpace& Element : CR->GetSpaceHierarchy())
-				{
-					Names.Add(Element.Name);
-				}
-				break;
-			}
-			case ERigElementType::Control:
-			{
-				for (const FRigControl& Element : CR->GetControlHierarchy())
-				{
-					Names.Add(Element.Name);
-				}
-				break;
-			}
-			case ERigElementType::Curve:
-			{
-				for (const FRigCurve& Element : CR->GetCurveContainer())
-				{
-					Names.Add(Element.Name);
-				}
-				break;
+				Names.Add(Element->GetName());
 			}
 		}
 	}
@@ -480,16 +455,17 @@ void UControlRigComponent::AddMappedSkeletalMesh(USkeletalMeshComponent* Skeleta
 		{
 			if (const USkeleton* Skeleton = SkeletalMesh->GetSkeleton())
 			{
-				for (const FRigBone& RigBone : CR->GetBoneHierarchy())
+				CR->GetHierarchy()->ForEach<FRigBoneElement>([Skeleton, &BonesToMap](FRigBoneElement* BoneElement) -> bool
 				{
-					if (Skeleton->GetReferenceSkeleton().FindBoneIndex(RigBone.Name) != INDEX_NONE)
+					if (Skeleton->GetReferenceSkeleton().FindBoneIndex(BoneElement->GetName()) != INDEX_NONE)
 					{
 						FControlRigComponentMappedBone BoneToMap;
-						BoneToMap.Source = RigBone.Name;
-						BoneToMap.Target = RigBone.Name;
+						BoneToMap.Source = BoneElement->GetName();
+						BoneToMap.Target = BoneElement->GetName();
 						BonesToMap.Add(BoneToMap);
 					}
-				}
+					return true;
+				});
 			}
 			else
 			{
@@ -505,21 +481,22 @@ void UControlRigComponent::AddMappedSkeletalMesh(USkeletalMeshComponent* Skeleta
 		{
 			if (USkeleton* Skeleton = SkeletalMesh->GetSkeleton())
 			{
-				for (const FRigCurve& RigCurve : CR->GetCurveContainer())
-				{
-					const FSmartNameMapping* CurveNameMapping = Skeleton->GetSmartNameContainer(USkeleton::AnimCurveMappingName);
+				CR->GetHierarchy()->ForEach<FRigCurveElement>([Skeleton, &CurvesToMap](FRigCurveElement* CurveElement) -> bool
+                {
+                    const FSmartNameMapping* CurveNameMapping = Skeleton->GetSmartNameContainer(USkeleton::AnimCurveMappingName);
 					if (CurveNameMapping)
 					{
 						FSmartName SmartName;
-						if (CurveNameMapping->FindSmartName(RigCurve.Name, SmartName))
+						if (CurveNameMapping->FindSmartName(CurveElement->GetName(), SmartName))
 						{
 							FControlRigComponentMappedCurve CurveToMap;
-							CurveToMap.Source = RigCurve.Name;
-							CurveToMap.Target = RigCurve.Name;
+							CurveToMap.Source = CurveElement->GetName();
+							CurveToMap.Target = CurveElement->GetName();
 							CurvesToMap.Add(CurveToMap);
 						}
 					}
-				}
+					return true;
+				});
 			}
 			else
 			{
@@ -590,16 +567,16 @@ FTransform UControlRigComponent::GetBoneTransform(FName BoneName, EControlRigCom
 {
 	if (UControlRig* CR = SetupControlRigIfRequired())
 	{
-		int32 BoneIndex = CR->GetBoneHierarchy().GetIndex(BoneName);
+		const int32 BoneIndex = CR->GetHierarchy()->GetIndex(FRigElementKey(BoneName, ERigElementType::Bone));
 		if (BoneIndex != INDEX_NONE)
 		{
 			if (Space == EControlRigComponentSpace::LocalSpace)
 			{
-				return CR->GetBoneHierarchy().GetLocalTransform(BoneIndex);
+				return CR->GetHierarchy()->GetLocalTransform(BoneIndex);
 			}
 			else
 			{
-				FTransform RootTransform = CR->GetBoneHierarchy().GetGlobalTransform(BoneIndex);
+				FTransform RootTransform = CR->GetHierarchy()->GetGlobalTransform(BoneIndex);
 				ConvertTransformFromRigSpace(RootTransform, Space);
 				return RootTransform;
 			}
@@ -613,16 +590,16 @@ FTransform UControlRigComponent::GetInitialBoneTransform(FName BoneName, EContro
 {
 	if (UControlRig* CR = SetupControlRigIfRequired())
 	{
-		int32 BoneIndex = CR->GetBoneHierarchy().GetIndex(BoneName);
+		const int32 BoneIndex = CR->GetHierarchy()->GetIndex(FRigElementKey(BoneName, ERigElementType::Bone));
 		if (BoneIndex != INDEX_NONE)
 		{
 			if (Space == EControlRigComponentSpace::LocalSpace)
 			{
-				return CR->GetHierarchy()->GetInitialTransform(ERigElementType::Bone, BoneIndex);
+				return CR->GetHierarchy()->GetInitialLocalTransform(BoneIndex);
 			}
 			else
 			{
-				FTransform RootTransform = CR->GetHierarchy()->GetInitialGlobalTransform(ERigElementType::Bone, BoneIndex);
+				FTransform RootTransform = CR->GetHierarchy()->GetInitialGlobalTransform(BoneIndex);
 				ConvertTransformFromRigSpace(RootTransform, Space);
 				return RootTransform;
 			}
@@ -641,7 +618,7 @@ void UControlRigComponent::SetBoneTransform(FName BoneName, FTransform Transform
 
 	if (UControlRig* CR = SetupControlRigIfRequired())
 	{
-		int32 BoneIndex = CR->GetBoneHierarchy().GetIndex(BoneName);
+		int32 BoneIndex = CR->GetHierarchy()->GetIndex(FRigElementKey(BoneName, ERigElementType::Bone));
 		if (BoneIndex != INDEX_NONE)
 		{
 			ConvertTransformToRigSpace(Transform, Space);
@@ -650,26 +627,26 @@ void UControlRigComponent::SetBoneTransform(FName BoneName, FTransform Transform
 			{
 				if (Weight >= 1.f - SMALL_NUMBER)
 				{
-					CR->GetBoneHierarchy().SetLocalTransform(BoneIndex, Transform, bPropagateToChildren);
+					CR->GetHierarchy()->SetLocalTransform(BoneIndex, Transform, bPropagateToChildren);
 				}
 				else
 				{
-					FTransform PreviousTransform = CR->GetBoneHierarchy().GetLocalTransform(BoneIndex);
+					FTransform PreviousTransform = CR->GetHierarchy()->GetLocalTransform(BoneIndex);
 					FTransform BlendedTransform = FControlRigMathLibrary::LerpTransform(PreviousTransform, Transform, Weight);
-					CR->GetBoneHierarchy().SetLocalTransform(BoneIndex, Transform, bPropagateToChildren);
+					CR->GetHierarchy()->SetLocalTransform(BoneIndex, Transform, bPropagateToChildren);
 				}
 			}
 			else
 			{
 				if (Weight >= 1.f - SMALL_NUMBER)
 				{
-					CR->GetBoneHierarchy().SetGlobalTransform(BoneIndex, Transform, bPropagateToChildren);
+					CR->GetHierarchy()->SetGlobalTransform(BoneIndex, Transform, bPropagateToChildren);
 				}
 				else
 				{
-					FTransform PreviousTransform = CR->GetBoneHierarchy().GetGlobalTransform(BoneIndex);
+					FTransform PreviousTransform = CR->GetHierarchy()->GetGlobalTransform(BoneIndex);
 					FTransform BlendedTransform = FControlRigMathLibrary::LerpTransform(PreviousTransform, Transform, Weight);
-					CR->GetBoneHierarchy().SetGlobalTransform(BoneIndex, Transform, bPropagateToChildren);
+					CR->GetHierarchy()->SetGlobalTransform(BoneIndex, Transform, bPropagateToChildren);
 				}
 			}
 		}
@@ -680,8 +657,7 @@ void UControlRigComponent::SetInitialBoneTransform(FName BoneName, FTransform In
 {
 	if (UControlRig* CR = SetupControlRigIfRequired())
 	{
-		FRigBoneHierarchy& HierarchyRef = CR->GetBoneHierarchy();
-		int32 BoneIndex = HierarchyRef.GetIndex(BoneName);
+		const int32 BoneIndex = CR->GetHierarchy()->GetIndex(FRigElementKey(BoneName, ERigElementType::Bone));
 		if (BoneIndex != INDEX_NONE)
 		{
 			if(!CR->IsRunningPreSetup() && !CR->IsRunningPostSetup())
@@ -694,14 +670,14 @@ void UControlRigComponent::SetInitialBoneTransform(FName BoneName, FTransform In
 
 			if (Space == EControlRigComponentSpace::LocalSpace)
 			{
-				int32 ParentIndex = HierarchyRef[BoneIndex].ParentIndex;
+				const int32 ParentIndex = CR->GetHierarchy()->GetFirstParent(BoneIndex);
 				if(ParentIndex != INDEX_NONE)
 				{
-					InitialTransform = InitialTransform * HierarchyRef[ParentIndex].InitialTransform;
+					InitialTransform = InitialTransform * CR->GetHierarchy()->GetInitialGlobalTransform(ParentIndex);
 				}
 			}
 
-			HierarchyRef.SetInitialGlobalTransform(BoneIndex, InitialTransform, bPropagateToChildren);
+			CR->GetHierarchy()->SetInitialGlobalTransform(BoneIndex, InitialTransform, bPropagateToChildren);
 		}
 	}	
 }
@@ -710,11 +686,11 @@ bool UControlRigComponent::GetControlBool(FName ControlName)
 {
 	if (UControlRig* CR = SetupControlRigIfRequired())
 	{
-		if (FRigControl* Control = CR->FindControl(ControlName))
+		if(FRigControlElement* ControlElement = CR->GetHierarchy()->Find<FRigControlElement>(FRigElementKey(ControlName, ERigElementType::Control)))
 		{
-			if (Control->ControlType == ERigControlType::Bool)
+			if (ControlElement->Settings.ControlType == ERigControlType::Bool)
 			{
-				return Control->GetValue().Get<bool>();
+				return CR->GetHierarchy()->GetControlValue(ControlElement, ERigControlValueType::Current).Get<bool>();
 			}
 		}
 	}
@@ -726,11 +702,11 @@ float UControlRigComponent::GetControlFloat(FName ControlName)
 {
 	if (UControlRig* CR = SetupControlRigIfRequired())
 	{
-		if (FRigControl* Control = CR->FindControl(ControlName))
+		if(FRigControlElement* ControlElement = CR->GetHierarchy()->Find<FRigControlElement>(FRigElementKey(ControlName, ERigElementType::Control)))
 		{
-			if (Control->ControlType == ERigControlType::Float)
+			if (ControlElement->Settings.ControlType == ERigControlType::Float)
 			{
-				return Control->GetValue().Get<float>();
+				return CR->GetHierarchy()->GetControlValue(ControlElement, ERigControlValueType::Current).Get<float>();
 			}
 		}
 	}
@@ -742,11 +718,11 @@ int32 UControlRigComponent::GetControlInt(FName ControlName)
 {
 	if (UControlRig* CR = SetupControlRigIfRequired())
 	{
-		if (FRigControl* Control = CR->FindControl(ControlName))
+		if(FRigControlElement* ControlElement = CR->GetHierarchy()->Find<FRigControlElement>(FRigElementKey(ControlName, ERigElementType::Control)))
 		{
-			if (Control->ControlType == ERigControlType::Integer)
+			if (ControlElement->Settings.ControlType == ERigControlType::Integer)
 			{
-				return Control->GetValue().Get<int32>();
+				return CR->GetHierarchy()->GetControlValue(ControlElement, ERigControlValueType::Current).Get<int>();
 			}
 		}
 	}
@@ -758,11 +734,11 @@ FVector2D UControlRigComponent::GetControlVector2D(FName ControlName)
 {
 	if (UControlRig* CR = SetupControlRigIfRequired())
 	{
-		if (FRigControl* Control = CR->FindControl(ControlName))
+		if(FRigControlElement* ControlElement = CR->GetHierarchy()->Find<FRigControlElement>(FRigElementKey(ControlName, ERigElementType::Control)))
 		{
-			if (Control->ControlType == ERigControlType::Vector2D)
+			if (ControlElement->Settings.ControlType == ERigControlType::Vector2D)
 			{
-				return Control->GetValue().Get<FVector2D>();
+				return CR->GetHierarchy()->GetControlValue(ControlElement, ERigControlValueType::Current).Get<FVector2D>();
 			}
 		}
 	}
@@ -772,115 +748,36 @@ FVector2D UControlRigComponent::GetControlVector2D(FName ControlName)
 
 FVector UControlRigComponent::GetControlPosition(FName ControlName, EControlRigComponentSpace Space)
 {
-	if (UControlRig* CR = SetupControlRigIfRequired())
-	{
-		if (FRigControl* Control = CR->FindControl(ControlName))
-		{
-			if (Space == EControlRigComponentSpace::LocalSpace)
-			{
-				if (Control->ControlType == ERigControlType::Position)
-				{
-					return Control->GetValue().Get<FVector>();
-				}
-				else if (Control->ControlType == ERigControlType::TransformNoScale)
-				{
-					return Control->GetValue().Get<FTransformNoScale>().Location;
-				}
-				else if (Control->ControlType == ERigControlType::Transform)
-				{
-					return Control->GetValue().Get<FTransform>().GetLocation();
-				}
-			}
-			else
-			{
-				FTransform RootTransform = CR->GetHierarchy()->GetGlobalTransform(ERigElementType::Control, Control->Index);
-				ConvertTransformFromRigSpace(RootTransform, Space);
-				return RootTransform.GetLocation();
-			}
-		}
-	}
-
-	return FVector::ZeroVector;
+	return GetControlTransform(ControlName, Space).GetLocation();
 }
 
 FRotator UControlRigComponent::GetControlRotator(FName ControlName, EControlRigComponentSpace Space)
 {
-	if (UControlRig* CR = SetupControlRigIfRequired())
-	{
-		if (FRigControl* Control = CR->FindControl(ControlName))
-		{
-			if (Space == EControlRigComponentSpace::LocalSpace)
-			{
-				if (Control->ControlType == ERigControlType::Rotator)
-				{
-					return Control->GetValue().Get<FRotator>();
-				}
-				else if (Control->ControlType == ERigControlType::TransformNoScale)
-				{
-					return Control->GetValue().Get<FTransformNoScale>().Rotation.Rotator();
-				}
-				else if (Control->ControlType == ERigControlType::Transform)
-				{
-					return Control->GetValue().Get<FTransform>().GetRotation().Rotator();
-				}
-			}
-			else
-			{
-				FTransform RootTransform = CR->GetHierarchy()->GetGlobalTransform(ERigElementType::Control, Control->Index);
-				ConvertTransformFromRigSpace(RootTransform, Space);
-				return RootTransform.GetRotation().Rotator();
-			}
-		}
-	}
-
-	return FRotator::ZeroRotator;
+	return GetControlTransform(ControlName, Space).GetRotation().Rotator();
 }
 
 FVector UControlRigComponent::GetControlScale(FName ControlName, EControlRigComponentSpace Space)
 {
-	if (UControlRig* CR = SetupControlRigIfRequired())
-	{
-		if (FRigControl* Control = CR->FindControl(ControlName))
-		{
-			if (Space == EControlRigComponentSpace::LocalSpace)
-			{
-				if (Control->ControlType == ERigControlType::Scale)
-				{
-					return Control->GetValue().Get<FVector>();
-				}
-				else if (Control->ControlType == ERigControlType::Transform)
-				{
-					return Control->GetValue().Get<FTransform>().GetScale3D();
-				}
-			}
-			else
-			{
-				FTransform RootTransform = CR->GetHierarchy()->GetGlobalTransform(ERigElementType::Control, Control->Index);
-				ConvertTransformFromRigSpace(RootTransform, Space);
-				return RootTransform.GetScale3D();
-			}
-		}
-	}
-
-	return FVector::OneVector;
+	return GetControlTransform(ControlName, Space).GetScale3D();
 }
 
 FTransform UControlRigComponent::GetControlTransform(FName ControlName, EControlRigComponentSpace Space)
 {
 	if (UControlRig* CR = SetupControlRigIfRequired())
 	{
-		if (FRigControl* Control = CR->FindControl(ControlName))
+		if(FRigControlElement* ControlElement = CR->GetHierarchy()->Find<FRigControlElement>(FRigElementKey(ControlName, ERigElementType::Control)))
 		{
+			FTransform Transform;
 			if (Space == EControlRigComponentSpace::LocalSpace)
 			{
-				return CR->GetHierarchy()->GetLocalTransform(ERigElementType::Control, Control->Index);
+				Transform = CR->GetHierarchy()->GetTransform(ControlElement, ERigTransformType::CurrentLocal);
 			}
 			else
 			{
-				FTransform RootTransform = CR->GetHierarchy()->GetGlobalTransform(ERigElementType::Control, Control->Index);
-				ConvertTransformFromRigSpace(RootTransform, Space);
-				return RootTransform;
+				Transform = CR->GetHierarchy()->GetTransform(ControlElement, ERigTransformType::CurrentGlobal);
+				ConvertTransformFromRigSpace(Transform, Space);
 			}
+			return Transform;
 		}
 	}
 
@@ -923,34 +820,25 @@ void UControlRigComponent::SetControlPosition(FName ControlName, FVector Value,
 {
 	if (UControlRig* CR = SetupControlRigIfRequired())
 	{
-		if (Space == EControlRigComponentSpace::LocalSpace)
+		if(FRigControlElement* ControlElement = CR->GetHierarchy()->Find<FRigControlElement>(FRigElementKey(ControlName, ERigElementType::Control)))
 		{
-			if (FRigControl* Control = CR->FindControl(ControlName))
+			FTransform InputTransform = FTransform::Identity;
+			InputTransform.SetLocation(Value);
+
+			FTransform PreviousTransform = FTransform::Identity;
+
+			if (Space == EControlRigComponentSpace::LocalSpace)
 			{
-				if (Control->ControlType == ERigControlType::Position)
-				{
-					CR->SetControlValue<FVector>(ControlName, Value);
-				}
-				else if (Control->ControlType == ERigControlType::TransformNoScale)
-				{
-					FTransformNoScale Previous = Control->GetValue().Get<FTransformNoScale>();
-					Previous.Location = Value;
-					CR->SetControlValue<FTransformNoScale>(ControlName, Previous);
-				}
-				else if (Control->ControlType == ERigControlType::Transform)
-				{
-					FTransform Previous = Control->GetValue().Get<FTransform>();
-					Previous.SetLocation(Value);
-					CR->SetControlValue<FTransform>(ControlName, Previous);
-				}
+				PreviousTransform = CR->GetHierarchy()->GetTransform(ControlElement, ERigTransformType::CurrentLocal); 
 			}
-		}
-		else
-		{
-			FTransform Transform = FTransform::Identity;
-			Transform.SetLocation(Value);
-			ConvertTransformToRigSpace(Transform, Space);
-			CR->SetControlGlobalTransform(ControlName, Transform);
+			else
+			{
+				PreviousTransform = CR->GetHierarchy()->GetTransform(ControlElement, ERigTransformType::CurrentGlobal); 
+				ConvertTransformToRigSpace(InputTransform, Space);
+			}
+
+			PreviousTransform.SetTranslation(InputTransform.GetTranslation());
+			SetControlTransform(ControlName, PreviousTransform, Space);
 		}
 	}
 }
@@ -959,34 +847,25 @@ void UControlRigComponent::SetControlRotator(FName ControlName, FRotator Value,
 {
 	if (UControlRig* CR = SetupControlRigIfRequired())
 	{
-		if (Space == EControlRigComponentSpace::LocalSpace)
+		if(FRigControlElement* ControlElement = CR->GetHierarchy()->Find<FRigControlElement>(FRigElementKey(ControlName, ERigElementType::Control)))
 		{
-			if (FRigControl* Control = CR->FindControl(ControlName))
+			FTransform InputTransform = FTransform::Identity;
+			InputTransform.SetRotation(FQuat(Value));
+
+			FTransform PreviousTransform = FTransform::Identity;
+
+			if (Space == EControlRigComponentSpace::LocalSpace)
 			{
-				if (Control->ControlType == ERigControlType::Rotator)
-				{
-					CR->SetControlValue<FRotator>(ControlName, Value);
-				}
-				else if (Control->ControlType == ERigControlType::TransformNoScale)
-				{
-					FTransformNoScale Previous = Control->GetValue().Get<FTransformNoScale>();
-					Previous.Rotation = FQuat(Value);
-					CR->SetControlValue<FTransformNoScale>(ControlName, Previous);
-				}
-				else if (Control->ControlType == ERigControlType::Transform)
-				{
-					FTransform Previous = Control->GetValue().Get<FTransform>();
-					Previous.SetRotation(FQuat(Value));
-					CR->SetControlValue<FTransform>(ControlName, Previous);
-				}
+				PreviousTransform = CR->GetHierarchy()->GetTransform(ControlElement, ERigTransformType::CurrentLocal); 
 			}
-		}
-		else
-		{
-			FTransform Transform = FTransform::Identity;
-			Transform.SetRotation(FQuat(Value));
-			ConvertTransformToRigSpace(Transform, Space);
-			CR->SetControlGlobalTransform(ControlName, Transform);
+			else
+			{
+				PreviousTransform = CR->GetHierarchy()->GetTransform(ControlElement, ERigTransformType::CurrentGlobal); 
+				ConvertTransformToRigSpace(InputTransform, Space);
+			}
+
+			PreviousTransform.SetRotation(InputTransform.GetRotation());
+			SetControlTransform(ControlName, PreviousTransform, Space);
 		}
 	}
 }
@@ -995,28 +874,25 @@ void UControlRigComponent::SetControlScale(FName ControlName, FVector Value, ECo
 {
 	if (UControlRig* CR = SetupControlRigIfRequired())
 	{
-		if (Space == EControlRigComponentSpace::LocalSpace)
+		if(FRigControlElement* ControlElement = CR->GetHierarchy()->Find<FRigControlElement>(FRigElementKey(ControlName, ERigElementType::Control)))
 		{
-			if (FRigControl* Control = CR->FindControl(ControlName))
+			FTransform InputTransform = FTransform::Identity;
+			InputTransform.SetScale3D(Value);
+
+			FTransform PreviousTransform = FTransform::Identity;
+
+			if (Space == EControlRigComponentSpace::LocalSpace)
 			{
-				if (Control->ControlType == ERigControlType::Scale)
-				{
-					CR->SetControlValue<FVector>(ControlName, Value);
-				}
-				else if (Control->ControlType == ERigControlType::Transform)
-				{
-					FTransform Previous = Control->GetValue().Get<FTransform>();
-					Previous.SetScale3D(Value);
-					CR->SetControlValue<FTransform>(ControlName, Previous);
-				}
+				PreviousTransform = CR->GetHierarchy()->GetTransform(ControlElement, ERigTransformType::CurrentLocal); 
 			}
-		}
-		else
-		{
-			FTransform Transform = FTransform::Identity;
-			Transform.SetScale3D(Value);
-			ConvertTransformToRigSpace(Transform, Space);
-			CR->SetControlGlobalTransform(ControlName, Transform);
+			else
+			{
+				PreviousTransform = CR->GetHierarchy()->GetTransform(ControlElement, ERigTransformType::CurrentGlobal); 
+				ConvertTransformToRigSpace(InputTransform, Space);
+			}
+
+			PreviousTransform.SetScale3D(InputTransform.GetScale3D());
+			SetControlTransform(ControlName, PreviousTransform, Space);
 		}
 	}
 }
@@ -1025,16 +901,16 @@ void UControlRigComponent::SetControlTransform(FName ControlName, FTransform Val
 {
 	if (UControlRig* CR = SetupControlRigIfRequired())
 	{
-		if (FRigControl* Control = CR->FindControl(ControlName))
+		if(FRigControlElement* ControlElement = CR->GetHierarchy()->Find<FRigControlElement>(FRigElementKey(ControlName, ERigElementType::Control)))
 		{
 			if (Space == EControlRigComponentSpace::LocalSpace)
 			{
-				CR->GetHierarchy()->SetLocalTransform(ERigElementType::Control, Control->Index, Value);
+				CR->GetHierarchy()->SetTransform(ControlElement, Value, ERigTransformType::CurrentLocal, true);
 			}
 			else
 			{
 				ConvertTransformToRigSpace(Value, Space);
-				CR->GetHierarchy()->SetGlobalTransform(ERigElementType::Control, Control->Index, Value);
+				CR->GetHierarchy()->SetTransform(ControlElement, Value, ERigTransformType::CurrentGlobal, true);
 			}
 		}
 	}
@@ -1044,15 +920,15 @@ FTransform UControlRigComponent::GetControlOffset(FName ControlName, EControlRig
 {
 	if (UControlRig* CR = SetupControlRigIfRequired())
 	{
-		if (FRigControl* Control = CR->FindControl(ControlName))
+		if(FRigControlElement* ControlElement = CR->GetHierarchy()->Find<FRigControlElement>(FRigElementKey(ControlName, ERigElementType::Control)))
 		{
 			if (Space == EControlRigComponentSpace::LocalSpace)
 			{
-				return Control->OffsetTransform;
+				return CR->GetHierarchy()->GetControlOffsetTransform(ControlElement, ERigTransformType::CurrentLocal);
 			}
 			else
 			{
-				FTransform RootTransform = CR->GetControlHierarchy().GetParentInitialTransform(Control->Index, true);
+				FTransform RootTransform = CR->GetHierarchy()->GetControlOffsetTransform(ControlElement, ERigTransformType::CurrentGlobal);
 				ConvertTransformFromRigSpace(RootTransform, Space);
 				return RootTransform;
 			}
@@ -1066,17 +942,15 @@ void UControlRigComponent::SetControlOffset(FName ControlName, FTransform Offset
 {
 	if (UControlRig* CR = SetupControlRigIfRequired())
 	{
-		if (FRigControl* Control = CR->FindControl(ControlName))
+		if(FRigControlElement* ControlElement = CR->GetHierarchy()->Find<FRigControlElement>(FRigElementKey(ControlName, ERigElementType::Control)))
 		{
 			if (Space != EControlRigComponentSpace::LocalSpace)
 			{
 				ConvertTransformToRigSpace(OffsetTransform, Space);
-
-				FTransform ParentTransform = CR->GetControlHierarchy().GetParentInitialTransform(Control->Index, false);
-				OffsetTransform = OffsetTransform.GetRelativeTransform(ParentTransform);
 			}
 
-			CR->GetControlHierarchy().SetControlOffset(Control->Index, OffsetTransform);
+			CR->GetHierarchy()->SetControlOffsetTransform(ControlElement, OffsetTransform,
+				Space == EControlRigComponentSpace::LocalSpace ? ERigTransformType::CurrentLocal : ERigTransformType::CurrentGlobal, true, false); 
 		}
 	}
 }
@@ -1085,17 +959,15 @@ FTransform UControlRigComponent::GetSpaceTransform(FName SpaceName, EControlRigC
 {
 	if (UControlRig* CR = SetupControlRigIfRequired())
 	{
-		int32 SpaceIndex = CR->GetSpaceHierarchy().GetIndex(SpaceName);
-		if (SpaceIndex != INDEX_NONE)
+		if(FRigSpaceElement* SpaceElement = CR->GetHierarchy()->Find<FRigSpaceElement>(FRigElementKey(SpaceName, ERigElementType::Control)))
 		{
-			const FRigSpace& SpaceElement = CR->GetSpaceHierarchy()[SpaceIndex];
 			if (Space == EControlRigComponentSpace::LocalSpace)
 			{
-				return CR->GetHierarchy()->GetLocalTransform(ERigElementType::Space, SpaceElement.Index);
+				return CR->GetHierarchy()->GetTransform(SpaceElement, ERigTransformType::CurrentLocal);
 			}
 			else
 			{
-				FTransform RootTransform = CR->GetHierarchy()->GetGlobalTransform(ERigElementType::Space, SpaceElement.Index);
+				FTransform RootTransform = CR->GetHierarchy()->GetTransform(SpaceElement, ERigTransformType::CurrentGlobal);
 				ConvertTransformFromRigSpace(RootTransform, Space);
 				return RootTransform;
 			}
@@ -1109,17 +981,15 @@ FTransform UControlRigComponent::GetInitialSpaceTransform(FName SpaceName, ECont
 {
 	if (UControlRig* CR = SetupControlRigIfRequired())
 	{
-		int32 SpaceIndex = CR->GetSpaceHierarchy().GetIndex(SpaceName);
-		if (SpaceIndex != INDEX_NONE)
+		if(FRigSpaceElement* SpaceElement = CR->GetHierarchy()->Find<FRigSpaceElement>(FRigElementKey(SpaceName, ERigElementType::Control)))
 		{
-			const FRigSpace& SpaceElement = CR->GetSpaceHierarchy()[SpaceIndex];
 			if (Space == EControlRigComponentSpace::LocalSpace)
 			{
-				return CR->GetHierarchy()->GetInitialTransform(ERigElementType::Space, SpaceElement.Index);
+				return CR->GetHierarchy()->GetTransform(SpaceElement, ERigTransformType::InitialLocal);
 			}
 			else
 			{
-				FTransform RootTransform = CR->GetHierarchy()->GetInitialGlobalTransform(ERigElementType::Space, SpaceElement.Index);
+				FTransform RootTransform = CR->GetHierarchy()->GetTransform(SpaceElement, ERigTransformType::InitialGlobal);
 				ConvertTransformFromRigSpace(RootTransform, Space);
 				return RootTransform;
 			}
@@ -1133,20 +1003,17 @@ void UControlRigComponent::SetInitialSpaceTransform(FName SpaceName, FTransform
 {
 	if (UControlRig* CR = SetupControlRigIfRequired())
 	{
-		int32 SpaceIndex = CR->GetSpaceHierarchy().GetIndex(SpaceName);
-		if (SpaceIndex != INDEX_NONE)
+		if(FRigSpaceElement* SpaceElement = CR->GetHierarchy()->Find<FRigSpaceElement>(FRigElementKey(SpaceName, ERigElementType::Control)))
 		{
-			const FRigSpace& SpaceElement = CR->GetSpaceHierarchy()[SpaceIndex];
-
-			if (Space != EControlRigComponentSpace::LocalSpace)
+			if (Space == EControlRigComponentSpace::LocalSpace)
+			{
+				CR->GetHierarchy()->SetTransform(SpaceElement, InitialTransform, ERigTransformType::InitialLocal, true, false);
+			}
+			else
 			{
 				ConvertTransformToRigSpace(InitialTransform, Space);
-
-				FTransform ParentTransform = CR->GetHierarchy()->GetInitialGlobalTransform(SpaceElement.GetParentElementKey());
-				InitialTransform = InitialTransform.GetRelativeTransform(ParentTransform);
+				CR->GetHierarchy()->SetTransform(SpaceElement, InitialTransform, ERigTransformType::InitialGlobal, true, false);
 			}
-
-			CR->GetHierarchy()->SetInitialTransform(SpaceElement.GetElementKey(), InitialTransform);
 		}
 	}
 }
@@ -1350,15 +1217,7 @@ void UControlRigComponent::TransferInputs()
 			Transform = MappedElement.Offset * Transform;
 
 			ConvertTransformToRigSpace(Transform, MappedElement.Space);
-
-			if (MappedElement.ElementType == ERigElementType::Control)
-			{
-				ControlRig->SetControlGlobalTransform(MappedElement.ElementName, Transform);
-			}
-			else
-			{
-				ControlRig->GetHierarchy()->SetGlobalTransform(MappedElement.ElementType, MappedElement.ElementIndex, Transform);
-			}
+			ControlRig->GetHierarchy()->SetGlobalTransform(MappedElement.ElementIndex, Transform);
 		}
 	}
 }
@@ -1397,7 +1256,7 @@ void UControlRigComponent::TransferOutputs()
 				MappedElement.ElementType == ERigElementType::Control ||
 				MappedElement.ElementType == ERigElementType::Space)
 			{
-				FTransform Transform = ControlRig->GetHierarchy()->GetGlobalTransform(MappedElement.ElementType, MappedElement.ElementIndex);
+				FTransform Transform = ControlRig->GetHierarchy()->GetGlobalTransform(MappedElement.ElementIndex);
 				ConvertTransformFromRigSpace(Transform, MappedElement.Space);
 
 				Transform = MappedElement.Offset * Transform;
@@ -1458,7 +1317,7 @@ void UControlRigComponent::TransferOutputs()
 					if (Proxy)
 					{
 						ComponentsToTick.AddUnique(Cast<USkeletalMeshComponent>(MappedElement.SceneComponent));
-						Proxy->StoredCurves.FindOrAdd((SmartName::UID_Type)MappedElement.SubIndex) = ControlRig->GetCurveContainer()[MappedElement.ElementIndex].Value;
+						Proxy->StoredCurves.FindOrAdd((SmartName::UID_Type)MappedElement.SubIndex) = ControlRig->GetHierarchy()->GetCurveValue(MappedElement.ElementIndex);
 					}
 				}
 			}
@@ -1783,25 +1642,25 @@ void FControlRigSceneProxy::GetDynamicMeshElements(const TArray<const FSceneView
 
 			if (bShouldDrawBones)
 			{
-				FTransform Transform = ControlRigComponent->GetComponentToWorld();
+				const FTransform Transform = ControlRigComponent->GetComponentToWorld();
 				const float MaxDrawRadius = ControlRigComponent->Bounds.SphereRadius * 0.02f;
-				const FRigBoneHierarchy& Hierarchy = ControlRigComponent->ControlRig->GetBoneHierarchy();
 
-				for (const FRigBone& Bone : Hierarchy)
-				{
-					const int32 ParentIndex = Bone.ParentIndex;
+				URigHierarchy* Hierarchy = ControlRigComponent->ControlRig->GetHierarchy();
+				Hierarchy->ForEach<FRigBoneElement>([PDI, Hierarchy, Transform, MaxDrawRadius](FRigBoneElement* BoneElement) -> bool
+                {
+                    const int32 ParentIndex = Hierarchy->GetFirstParent(BoneElement->GetIndex());
 					const FLinearColor LineColor = FLinearColor::White;
 
 					FVector Start, End;
 					if (ParentIndex >= 0)
 					{
-						Start = Hierarchy[ParentIndex].GlobalTransform.GetLocation();
-						End = Bone.GlobalTransform.GetLocation();
+						Start = Hierarchy->GetGlobalTransform(ParentIndex).GetLocation();
+						End = Hierarchy->GetGlobalTransform(BoneElement->GetIndex()).GetLocation();
 					}
 					else
 					{
 						Start = FVector::ZeroVector;
-						End = Bone.GlobalTransform.GetLocation();
+						End = Hierarchy->GetGlobalTransform(BoneElement->GetIndex()).GetLocation();
 					}
 
 					Start = Transform.TransformPosition(Start);
@@ -1813,7 +1672,9 @@ void FControlRigSceneProxy::GetDynamicMeshElements(const TArray<const FSceneView
 
 					//Render Sphere for bone end point and a cone between it and its parent.
 					SkeletalDebugRendering::DrawWireBone(PDI, Start, End, LineColor, SDPG_Foreground, Radius);
-				}
+
+					return true;
+				});
 			}
 
 			if (ControlRigComponent->bShowDebugDrawing)
diff --git a/Engine/Plugins/Experimental/ControlRig/Source/ControlRig/Private/ControlRigControlActor.cpp b/Engine/Plugins/Experimental/ControlRig/Source/ControlRig/Private/ControlRigControlActor.cpp
--- a/Engine/Plugins/Experimental/ControlRig/Source/ControlRig/Private/ControlRigControlActor.cpp
+++ b/Engine/Plugins/Experimental/ControlRig/Source/ControlRig/Private/ControlRigControlActor.cpp
@@ -143,15 +143,15 @@ void AControlRigControlActor::Refresh()
 			ColorParameterName = GizmoLibrary->MaterialColorParameter;
 		}
 
-		const FRigControlHierarchy& ControlHierarchy = ControlRig->GetControlHierarchy();
-		for (const FRigControl& Control : ControlHierarchy)
-		{
-			if (!Control.bGizmoEnabled)
+		URigHierarchy* Hierarchy = ControlRig->GetHierarchy();
+		Hierarchy->ForEach<FRigControlElement>([this, GizmoLibrary, BaseMaterial, Hierarchy](FRigControlElement* ControlElement) -> bool
+        {
+			if (!ControlElement->Settings.bGizmoEnabled)
 			{
-				continue;
+				return true;
 			}
 
-			switch (Control.ControlType)
+			switch (ControlElement->Settings.ControlType)
 			{
 				case ERigControlType::Float:
 				case ERigControlType::Integer:
@@ -163,7 +163,7 @@ void AControlRigControlActor::Refresh()
 				case ERigControlType::TransformNoScale:
 				case ERigControlType::EulerTransform:
 				{
-					if (const FControlRigGizmoDefinition* GizmoDef = GizmoLibrary->GetGizmoByName(Control.GizmoName))
+					if (const FControlRigGizmoDefinition* GizmoDef = GizmoLibrary->GetGizmoByName(ControlElement->Settings.GizmoName))
 					{
 						UStaticMesh* StaticMesh = GizmoDef->StaticMesh.LoadSynchronous();
 						UStaticMeshComponent* Component = NewObject< UStaticMeshComponent>(ActorRootComponent);
@@ -177,8 +177,8 @@ void AControlRigControlActor::Refresh()
 						UMaterialInstanceDynamic* MaterialInstance = UMaterialInstanceDynamic::Create(BaseMaterial, Component);
 						Component->SetMaterial(0, MaterialInstance);
 
-						ControlNames.Add(Control.Name);
-						GizmoTransforms.Add(Control.GizmoTransform * GizmoDef->Transform);
+						ControlNames.Add(ControlElement->GetName());
+						GizmoTransforms.Add(Hierarchy->GetControlGizmoTransform(ControlElement, ERigTransformType::CurrentLocal) * GizmoDef->Transform);
 						Components.Add(Component);
 						Materials.Add(MaterialInstance);
 					}
@@ -188,24 +188,25 @@ void AControlRigControlActor::Refresh()
 					break;
 				}
 			}
-		}
+			return true;
+		});
 	}
 
 	for (int32 GizmoIndex = 0; GizmoIndex < ControlNames.Num(); GizmoIndex++)
 	{
-		const FRigControlHierarchy& ControlHierarchy = ControlRig->GetControlHierarchy();
-		int32 ControlIndex = ControlHierarchy.GetIndex(ControlNames[GizmoIndex]);
-		if (ControlIndex == INDEX_NONE)
+		URigHierarchy* Hierarchy = ControlRig->GetHierarchy();
+
+		const FRigElementKey ControlKey(ControlNames[GizmoIndex], ERigElementType::Control);
+		FRigControlElement* ControlElement = Hierarchy->Find<FRigControlElement>(ControlKey);
+		if(ControlElement == nullptr)
 		{
 			Components[GizmoIndex]->SetVisibility(false);
 			continue;
 		}
 
-		const FRigControl& Control = ControlHierarchy[ControlIndex];
-
 		FTransform ControlTransform = ControlRig->GetControlGlobalTransform(ControlNames[GizmoIndex]);
 		Components[GizmoIndex]->SetRelativeTransform(GizmoTransforms[GizmoIndex] * ControlTransform);
-		Materials[GizmoIndex]->SetVectorParameterValue(ColorParameterName, FVector(Control.GizmoColor));
+		Materials[GizmoIndex]->SetVectorParameterValue(ColorParameterName, FVector(ControlElement->Settings.GizmoColor));
 	}
 }
 
diff --git a/Engine/Plugins/Experimental/ControlRig/Source/ControlRig/Private/Drawing/ControlRigDrawInterface.cpp b/Engine/Plugins/Experimental/ControlRig/Source/ControlRig/Private/Drawing/ControlRigDrawInterface.cpp
--- a/Engine/Plugins/Experimental/ControlRig/Source/ControlRig/Private/Drawing/ControlRigDrawInterface.cpp
+++ b/Engine/Plugins/Experimental/ControlRig/Source/ControlRig/Private/Drawing/ControlRigDrawInterface.cpp
@@ -178,7 +178,7 @@ void FControlRigDrawInterface::DrawBezier(const FTransform& WorldOffset, const F
 	Instructions.Add(Instruction);
 }
 
-void FControlRigDrawInterface::DrawHierarchy(const FTransform& WorldOffset, const FRigBoneHierarchy& Hierarchy, EControlRigDrawHierarchyMode::Type Mode, float Scale, const FLinearColor& Color, float Thickness)
+void FControlRigDrawInterface::DrawHierarchy(const FTransform& WorldOffset, URigHierarchy* Hierarchy, EControlRigDrawHierarchyMode::Type Mode, float Scale, const FLinearColor& Color, float Thickness)
 {
 	switch (Mode)
 	{
@@ -188,18 +188,18 @@ void FControlRigDrawInterface::DrawHierarchy(const FTransform& WorldOffset, cons
 			FControlRigDrawInstruction InstructionY(EControlRigDrawSettings::Lines, FLinearColor::Green, Thickness, WorldOffset);
 			FControlRigDrawInstruction InstructionZ(EControlRigDrawSettings::Lines, FLinearColor::Blue, Thickness, WorldOffset);
 			FControlRigDrawInstruction InstructionParent(EControlRigDrawSettings::Lines, Color, Thickness, WorldOffset);
-			InstructionX.Positions.Reserve(Hierarchy.Num() * 2);
-			InstructionY.Positions.Reserve(Hierarchy.Num() * 2);
-			InstructionZ.Positions.Reserve(Hierarchy.Num() * 2);
-			InstructionParent.Positions.Reserve(Hierarchy.Num() * 6);
+			InstructionX.Positions.Reserve(Hierarchy->Num() * 2);
+			InstructionY.Positions.Reserve(Hierarchy->Num() * 2);
+			InstructionZ.Positions.Reserve(Hierarchy->Num() * 2);
+			InstructionParent.Positions.Reserve(Hierarchy->Num() * 6);
 
-			for (const FRigBone& Bone : Hierarchy)
+			Hierarchy->ForEach<FRigTransformElement>([&](FRigTransformElement* Child) -> bool
 			{
-				FTransform Transform = Bone.GlobalTransform;
-				FVector P0 = Transform.GetLocation();
-				FVector PX = Transform.TransformPosition(FVector(Scale, 0.f, 0.f));
-				FVector PY = Transform.TransformPosition(FVector(0.f, Scale, 0.f));
-				FVector PZ = Transform.TransformPosition(FVector(0.f, 0.f, Scale));
+				const FTransform Transform = Hierarchy->GetTransform(Child, ERigTransformType::CurrentGlobal);
+				const FVector P0 = Transform.GetLocation();
+				const FVector PX = Transform.TransformPosition(FVector(Scale, 0.f, 0.f));
+				const FVector PY = Transform.TransformPosition(FVector(0.f, Scale, 0.f));
+				const FVector PZ = Transform.TransformPosition(FVector(0.f, 0.f, Scale));
 				InstructionX.Positions.Add(P0);
 				InstructionX.Positions.Add(PX);
 				InstructionY.Positions.Add(P0);
@@ -207,14 +207,19 @@ void FControlRigDrawInterface::DrawHierarchy(const FTransform& WorldOffset, cons
 				InstructionZ.Positions.Add(P0);
 				InstructionZ.Positions.Add(PZ);
 
-				if (Bone.ParentIndex != INDEX_NONE)
+				TArray<FRigBaseElement*> Parents = Hierarchy->GetParents(Child);
+				for (FRigBaseElement* Parent : Parents)
 				{
-					FTransform ParentTransform = Hierarchy[Bone.ParentIndex].GlobalTransform;
-					FVector P1 = ParentTransform.GetLocation();
-					InstructionParent.Positions.Add(P0);
-					InstructionParent.Positions.Add(P1);
+					if(FRigTransformElement* ParentTransformElement = Cast<FRigTransformElement>(Parent))
+					{
+						const FTransform ParentTransform = Hierarchy->GetTransform(ParentTransformElement, ERigTransformType::CurrentGlobal);
+						const FVector P1 = ParentTransform.GetLocation();
+						InstructionParent.Positions.Add(P0);
+						InstructionParent.Positions.Add(P1);
+					}
 				}
-			}
+				return true;
+			});
 
 			Instructions.Add(InstructionX);
 			Instructions.Add(InstructionY);
@@ -222,6 +227,10 @@ void FControlRigDrawInterface::DrawHierarchy(const FTransform& WorldOffset, cons
 			Instructions.Add(InstructionParent);
 			break;
 		}
+		default:
+		{
+			break;
+		}
 	}
 }
 
diff --git a/Engine/Plugins/Experimental/ControlRig/Source/ControlRig/Private/Rigs/AdditiveControlRig.cpp b/Engine/Plugins/Experimental/ControlRig/Source/ControlRig/Private/Rigs/AdditiveControlRig.cpp
--- a/Engine/Plugins/Experimental/ControlRig/Source/ControlRig/Private/Rigs/AdditiveControlRig.cpp
+++ b/Engine/Plugins/Experimental/ControlRig/Source/ControlRig/Private/Rigs/AdditiveControlRig.cpp
@@ -5,6 +5,7 @@
 #include "Engine/SkeletalMesh.h"
 #include "IControlRigObjectBinding.h"
 #include "Components/SkeletalMeshComponent.h"
+#include "Rigs/RigHierarchyController.h"
 #include "Units/Execution/RigUnit_BeginExecution.h"
 
 #define LOCTEXT_NAMESPACE "AdditiveControlRig"
@@ -44,12 +45,11 @@ void UAdditiveControlRig::ExecuteUnits(FRigUnitContext& InOutContext, const FNam
 		return;
 	}
 
-	FRigControlHierarchy& ControlHierarchy = GetControlHierarchy();
 	for (FRigUnit_AddBoneTransform& Unit : AddBoneRigUnits)
 	{
 		FName ControlName = GetControlName(Unit.Bone);
-		const int32 Index = ControlHierarchy.GetIndex(ControlName);
-		Unit.Transform = ControlHierarchy.GetLocalTransform(Index);
+		const int32 Index = GetHierarchy()->GetIndex(FRigElementKey(ControlName, ERigElementType::Control));
+		Unit.Transform = GetHierarchy()->GetLocalTransform(Index);
 		Unit.ExecuteContext.Hierarchy = GetHierarchy();
 		Unit.ExecuteContext.EventName = InEventName;
 		Unit.Execute(InOutContext);
@@ -74,15 +74,15 @@ void UAdditiveControlRig::Initialize(bool bInitRigUnits /*= true*/)
 
 	// add units and initialize
 	AddBoneRigUnits.Reset();
-	FRigBoneHierarchy& BoneHierarchy = GetBoneHierarchy();
 
-	for (int32 BoneIndex = 0; BoneIndex < BoneHierarchy.Num(); ++BoneIndex)
+	GetHierarchy()->ForEach<FRigBoneElement>([&](FRigBoneElement* BoneElement) -> bool
 	{
 		FRigUnit_AddBoneTransform NewUnit;
-		NewUnit.Bone = BoneHierarchy[BoneIndex].Name;
-		NewUnit.bPropagateToChildren = true;
-		AddBoneRigUnits.Add(NewUnit);
-	}
+        NewUnit.Bone = BoneElement->GetName();
+        NewUnit.bPropagateToChildren = true;
+        AddBoneRigUnits.Add(NewUnit);
+		return true;
+    });
 
 	// execute init
 	Execute(EControlRigState::Init, FRigUnit_BeginExecution::EventName);
@@ -90,53 +90,54 @@ void UAdditiveControlRig::Initialize(bool bInitRigUnits /*= true*/)
 
 void UAdditiveControlRig::CreateRigElements(const FReferenceSkeleton& InReferenceSkeleton, const FSmartNameMapping* InSmartNameMapping)
 {
-	FRigHierarchyContainer* Container = GetHierarchy();
-	Container->Reset();
-	FRigBoneHierarchy& BoneHierarchy = Container->BoneHierarchy;
-	BoneHierarchy.ImportSkeleton(InReferenceSkeleton, NAME_None, false, false, true, false);
+	if(Controller == nullptr)
+	{
+		Controller = NewObject<URigHierarchyController>(this);
+		Controller->SetHierarchy(GetHierarchy());
+	}
+	
+	GetHierarchy()->Reset();
+	Controller->ImportBones(InReferenceSkeleton, NAME_None, false, false, true, false);
 
 	if (InSmartNameMapping)
 	{
-		FRigCurveContainer& CurveContainer = Container->CurveContainer;
 		TArray<FName> NameArray;
 		InSmartNameMapping->FillNameArray(NameArray);
 		for (int32 Index = 0; Index < NameArray.Num(); ++Index)
 		{
-			CurveContainer.Add(NameArray[Index]);
+			Controller->AddCurve(NameArray[Index], 0.f, false);
 		}
 	}
 
-	// add control for all bone hierarchy 
-	for (int32 BoneIndex = 0; BoneIndex < BoneHierarchy.Num(); ++BoneIndex)
-	{
-
-		const FRigBone& RigBone = BoneHierarchy[BoneIndex];
-		FName BoneName = RigBone.Name;
-		FName ParentName = RigBone.ParentName;
-		FName SpaceName = GetSpaceName(BoneName);// name conflict?
-		FName ControlName = GetControlName(BoneName); // name conflict?
-		if (ParentName != NAME_None)
+	// add control for all bone hierarchy
+	GetHierarchy()->ForEach<FRigBoneElement>([&](FRigBoneElement* BoneElement) -> bool
+    {
+		const FName BoneName = BoneElement->GetName();
+		const int32 ParentIndex = GetHierarchy()->GetFirstParent(BoneElement->GetIndex());
+		const FName SpaceName = GetSpaceName(BoneName);// name conflict?
+		const FName ControlName = GetControlName(BoneName); // name conflict?
+
+		FRigElementKey SpaceKey;
+		
+		if (ParentIndex != INDEX_NONE)
 		{
-			FTransform Transform = BoneHierarchy.GetGlobalTransform(BoneName);
-			FTransform ParentTransform = BoneHierarchy.GetGlobalTransform(ParentName);
-			FTransform LocalTransform = Transform.GetRelativeTransform(ParentTransform);
-			FRigSpace& Space = Container->SpaceHierarchy.Add(SpaceName, ERigSpaceType::Bone, ParentName);
-			Space.InitialTransform = LocalTransform;
+			FTransform GlobalTransform = GetHierarchy()->GetGlobalTransform(BoneElement->GetIndex());
+			FTransform ParentTransform = GetHierarchy()->GetGlobalTransform(ParentIndex);
+			FTransform LocalTransform = GlobalTransform.GetRelativeTransform(ParentTransform);
+			SpaceKey = Controller->AddSpace(SpaceName, GetHierarchy()->GetKey(ParentIndex), LocalTransform, false, false);
 		}
 		else
 		{
-			FTransform Transform = BoneHierarchy.GetGlobalTransform(BoneName);
-			FTransform ParentTransform = FTransform::Identity;
-			FTransform LocalTransform = Transform.GetRelativeTransform(ParentTransform);
-			FRigSpace& Space = Container->SpaceHierarchy.Add(SpaceName, ERigSpaceType::Global, ParentName);
-			Space.InitialTransform = LocalTransform;
+			FTransform GlobalTransform = GetHierarchy()->GetGlobalTransform(BoneElement->GetIndex());
+			SpaceKey = Controller->AddSpace(SpaceName, FRigElementKey(), GlobalTransform, true, false);
 		}
 
-		FRigControl& RigControl = Container->ControlHierarchy.Add(ControlName, ERigControlType::Transform, NAME_None, SpaceName);
-		RigControl.DisplayName = BoneName;
-	}
+		FRigControlSettings Settings;
+		Settings.DisplayName = BoneName;
+		Controller->AddControl(ControlName, SpaceKey, Settings, FRigControlValue::Make(FTransform::Identity), FTransform::Identity, FTransform::Identity, false);
 
-	Container->Initialize(true);
+		return true;
+	});
 }
 
 void UAdditiveControlRig::CreateRigElements(const USkeletalMesh* InReferenceMesh)
diff --git a/Engine/Plugins/Experimental/ControlRig/Source/ControlRig/Private/Rigs/FKControlRig.cpp b/Engine/Plugins/Experimental/ControlRig/Source/ControlRig/Private/Rigs/FKControlRig.cpp
--- a/Engine/Plugins/Experimental/ControlRig/Source/ControlRig/Private/Rigs/FKControlRig.cpp
+++ b/Engine/Plugins/Experimental/ControlRig/Source/ControlRig/Private/Rigs/FKControlRig.cpp
@@ -50,88 +50,93 @@ void UFKControlRig::ExecuteUnits(FRigUnitContext& InOutContext, const FName& InE
 	{
 		FRigVMExecuteContext VMContext;
 
-		FRigBoneHierarchy& BoneHierarchy = GetBoneHierarchy();
-		FRigCurveContainer& CurveContainer = GetCurveContainer();
-		const FRigControlHierarchy& ControlHierarchy = GetControlHierarchy();
-
-		for(FRigBone& Bone : BoneHierarchy)
-		{
-			FName ControlName = GetControlName(Bone.Name);
-			const int32 ControlIndex = ControlHierarchy.GetIndex(ControlName);
+		GetHierarchy()->ForEach<FRigBoneElement>([&](FRigBoneElement* BoneElement) -> bool
+        {
+			const FName ControlName = GetControlName(BoneElement->GetName());
+			const FRigElementKey ControlKey(ControlName, ERigElementType::Control);
+			const int32 ControlIndex = GetHierarchy()->GetIndex(ControlKey);
 			if (IsControlActive[ControlIndex])
 			{
-				FTransform Transform = ControlHierarchy[ControlIndex].GetValue(ERigControlValueType::Current).Get<FEulerTransform>().ToFTransform();
+				const FTransform Transform = GetHierarchy()->GetLocalTransform(ControlIndex);
 				switch (ApplyMode)
 				{
 					case EControlRigFKRigExecuteMode::Replace:
 					{
-						BoneHierarchy.SetLocalTransform(Bone.Index, Transform, false);
+						GetHierarchy()->SetTransform(BoneElement, Transform, ERigTransformType::CurrentLocal, true, false);
 						break;
 					}
 					case EControlRigFKRigExecuteMode::Additive:
 					{
-						FTransform PreviousTransform = BoneHierarchy.GetLocalTransform(Bone.Index);
-						BoneHierarchy.SetLocalTransform(Bone.Index, Transform * PreviousTransform, false);
+						const FTransform PreviousTransform = GetHierarchy()->GetTransform(BoneElement, ERigTransformType::CurrentLocal);
+						GetHierarchy()->SetTransform(BoneElement, Transform * PreviousTransform, ERigTransformType::CurrentLocal, true, false);
 						break;
 					}
 				}
 			}
-		}
-
-		BoneHierarchy.RecomputeGlobalTransforms();
-
-		for (FRigCurve& Curve : CurveContainer)
-		{
-			FName ControlName = GetControlName(Curve.Name);
-			const int32 ControlIndex = ControlHierarchy.GetIndex(ControlName);
+			return true;
+		});
+
+		GetHierarchy()->ForEach<FRigCurveElement>([&](FRigCurveElement* CurveElement) -> bool
+        {
+			const FName ControlName = GetControlName(CurveElement->GetName());
+			const FRigElementKey ControlKey(ControlName, ERigElementType::Control);
+			const int32 ControlIndex = GetHierarchy()->GetIndex(ControlKey);
+			
 			if (IsControlActive[ControlIndex])
 			{
-				float CurveValue = ControlHierarchy[ControlIndex].GetValue(ERigControlValueType::Current).Get<float>();
+				const float CurveValue = GetHierarchy()->GetControlValue(ControlIndex).Get<float>();
+
 				switch (ApplyMode)
 				{
-				case EControlRigFKRigExecuteMode::Replace:
-				{
-					CurveContainer.SetValue(Curve.Index, CurveValue);
-					break;
-				}
-				case EControlRigFKRigExecuteMode::Additive:
-				{
-					float PreviousValue = CurveContainer.GetValue(Curve.Index);
-					CurveContainer.SetValue(Curve.Index, PreviousValue + CurveValue);
-					break;
-				}
+					case EControlRigFKRigExecuteMode::Replace:
+					{
+						GetHierarchy()->SetCurveValue(CurveElement, CurveValue, false);
+						break;
+					}
+					case EControlRigFKRigExecuteMode::Additive:
+					{
+						const float PreviousValue = GetHierarchy()->GetCurveValue(CurveElement);
+						GetHierarchy()->SetCurveValue(CurveElement, PreviousValue + CurveValue, false);
+						break;
+					}
 				}
 			}
-		}
+			return true;
+		});
 	}
 	else if (InEventName == FRigUnit_InverseExecution::EventName)
 	{
 		FRigVMExecuteContext VMContext;
 
-		const FRigBoneHierarchy& BoneHierarchy = GetBoneHierarchy();
-		const FRigCurveContainer& CurveContainer = GetCurveContainer();
-		FRigControlHierarchy& ControlHierarchy = GetControlHierarchy();
-
-		for (const FRigBone& Bone : BoneHierarchy)
-		{
-			FName ControlName = GetControlName(Bone.Name);
-			const int32 ControlIndex = ControlHierarchy.GetIndex(ControlName);
+		GetHierarchy()->ForEach<FRigBoneElement>([&](FRigBoneElement* BoneElement) -> bool
+        {
+            const FName ControlName = GetControlName(BoneElement->GetName());
+			const FRigElementKey ControlKey(ControlName, ERigElementType::Control);
+			const int32 ControlIndex = GetHierarchy()->GetIndex(ControlKey);
+			
 			if (IsControlActive[ControlIndex])
 			{
-				FEulerTransform EulerTransform(BoneHierarchy.GetLocalTransform(Bone.Index));
-				SetControlValue<FEulerTransform>(ControlName, EulerTransform);
+				const FEulerTransform EulerTransform(GetHierarchy()->GetTransform(BoneElement, ERigTransformType::CurrentLocal));
+				SetControlValue(ControlName, FRigControlValue::Make(EulerTransform));
 			}
-		}
+			
+			return true;
+		});
+
+		GetHierarchy()->ForEach<FRigCurveElement>([&](FRigCurveElement* CurveElement) -> bool
+        {
+            const FName ControlName = GetControlName(CurveElement->GetName());
+			const FRigElementKey ControlKey(ControlName, ERigElementType::Control);
+			const int32 ControlIndex = GetHierarchy()->GetIndex(ControlKey);
 
-		for (const FRigCurve& Curve : CurveContainer)
-		{
-			FName ControlName = GetControlName(Curve.Name);
-			const int32 ControlIndex = ControlHierarchy.GetIndex(ControlName);
 			if (IsControlActive[ControlIndex])
 			{
-				SetControlValue<float>(ControlName, CurveContainer.GetValue(Curve.Index));
+				const float CurveValue = GetHierarchy()->GetCurveValue(CurveElement);
+				SetControlValue(ControlName, FRigControlValue::Make(CurveValue));
 			}
-		}
+
+			return true;
+		});
 	}
 }
 
@@ -156,11 +161,13 @@ void UFKControlRig::Initialize(bool bInitRigUnits /*= true*/)
 }
 TArray<FName> UFKControlRig::GetControlNames()
 {
-	FRigHierarchyContainer* Container = GetHierarchy();
 	TArray<FName> Names;
-	for (FRigControl& Control : Container->ControlHierarchy)
+	for (FRigBaseElement* Element : *GetHierarchy())
 	{
-		Names.Add(Control.Name);
+		if(Element->IsTypeOf(ERigElementType::Control))
+		{
+			Names.Add(Element->GetName());
+		}
 	}
 	return Names;
 }
@@ -174,7 +181,6 @@ bool UFKControlRig::GetControlActive(int32 Index) const
 	return false;
 }
 
-
 void UFKControlRig::SetControlActive(int32 Index, bool bActive)
 {
 	if (Index >= 0 && Index < IsControlActive.Num())
@@ -193,11 +199,14 @@ void UFKControlRig::SetControlActive(const TArray<FFKBoneCheckInfo>& BoneChecks)
 
 void UFKControlRig::CreateRigElements(const FReferenceSkeleton& InReferenceSkeleton, const FSmartNameMapping* InSmartNameMapping)
 {
-	FRigHierarchyContainer* Container = GetHierarchy();
-	Container->Reset();
-	FRigBoneHierarchy& BoneHierarchy = Container->BoneHierarchy;
-	BoneHierarchy.ImportSkeleton(InReferenceSkeleton, NAME_None, false, false, true, false);
-	FRigCurveContainer& CurveContainer = Container->CurveContainer;
+	if(Controller == nullptr)
+	{
+		Controller = NewObject<URigHierarchyController>(this);
+		Controller->SetHierarchy(GetHierarchy());
+	}
+	
+	GetHierarchy()->Reset();
+	Controller->ImportBones(InReferenceSkeleton, NAME_None, false, false, true, false);
 
 	if (InSmartNameMapping)
 	{
@@ -205,59 +214,66 @@ void UFKControlRig::CreateRigElements(const FReferenceSkeleton& InReferenceSkele
 		InSmartNameMapping->FillNameArray(NameArray);
 		for (int32 Index = 0; Index < NameArray.Num(); ++Index)
 		{
-			CurveContainer.Add(NameArray[Index]);
-		}
-	}
-
-	if (IsControlActive.Num() != (BoneHierarchy.Num() + CurveContainer.Num()))
-	{
-		IsControlActive.SetNum(BoneHierarchy.Num() + CurveContainer.Num());
-		for (bool& bIsActive : IsControlActive)
-		{
-			bIsActive = true;
+			Controller->AddCurve(NameArray[Index], 0.f, false);
 		}
 	}
 
 	// add control for all bone hierarchy 
 	int32 ControlIndex = 0;
-	for (; ControlIndex < BoneHierarchy.Num(); ++ControlIndex)
+
+	GetHierarchy()->ForEach<FRigBoneElement>([&](FRigBoneElement* BoneElement) -> bool
 	{
-		const FRigBone& RigBone = BoneHierarchy[ControlIndex];
-		FName BoneName = RigBone.Name;
-		FName ParentName = RigBone.ParentName;
-		FName SpaceName = GetSpaceName(BoneName);// name conflict?
-		FName ControlName = GetControlName(BoneName); // name conflict?
-		FTransform LocalTransform = FTransform::Identity;
-		if (ParentName != NAME_None)
+		const FName BoneName = BoneElement->GetName();
+		const int32 ParentIndex = GetHierarchy()->GetFirstParent(BoneElement->GetIndex());
+		const FName SpaceName = GetSpaceName(BoneName);// name conflict?
+		const FName ControlName = GetControlName(BoneName); // name conflict?
+
+		FRigElementKey SpaceKey;
+
+		FTransform LocalTransform;
+		if (ParentIndex != INDEX_NONE)
 		{
-			FTransform Transform = BoneHierarchy.GetGlobalTransform(BoneName);
-			FTransform ParentTransform = BoneHierarchy.GetGlobalTransform(ParentName);
-			LocalTransform = Transform.GetRelativeTransform(ParentTransform);
-			Container->SpaceHierarchy.Add(SpaceName, ERigSpaceType::Bone, ParentName);
+			FTransform GlobalTransform = GetHierarchy()->GetGlobalTransform(BoneElement->GetIndex());
+			FTransform ParentTransform = GetHierarchy()->GetGlobalTransform(ParentIndex);
+			LocalTransform = GlobalTransform.GetRelativeTransform(ParentTransform);
+			SpaceKey = Controller->AddSpace(SpaceName, GetHierarchy()->GetKey(ParentIndex), FTransform::Identity, false, false);
 		}
 		else
 		{
-			FTransform Transform = BoneHierarchy.GetGlobalTransform(BoneName);
-			FTransform ParentTransform = FTransform::Identity;
-			LocalTransform = Transform.GetRelativeTransform(ParentTransform);
-			Container->SpaceHierarchy.Add(SpaceName, ERigSpaceType::Global, ParentName);
+			LocalTransform = GetHierarchy()->GetLocalTransform(BoneElement->GetIndex());
+			SpaceKey = Controller->AddSpace(SpaceName, FRigElementKey(), FTransform::Identity, true, false);
 		}
 
-		FRigControl& RigControl = Container->ControlHierarchy.Add(ControlName, ERigControlType::EulerTransform, NAME_None, SpaceName); // NAME_None, LocalTransform);//SpaceName);  LocalTransform);
-		RigControl.DisplayName = BoneName;
-		RigControl.SetValueFromTransform(LocalTransform, ERigControlValueType::Initial);
-	}
-	for (int32 Index = 0; Index < CurveContainer.Num(); ++Index)
+		FRigControlSettings Settings;
+		Settings.ControlType = ERigControlType::EulerTransform;
+		Settings.DisplayName = BoneName;
+		const FRigElementKey ControlKey = Controller->AddControl(ControlName, SpaceKey, Settings, FRigControlValue::Make(FEulerTransform::Identity), FTransform::Identity, FTransform::Identity, false);
+		GetHierarchy()->SetLocalTransform(ControlKey, LocalTransform, true, true, false);
+
+		return true;
+	});
+	
+	GetHierarchy()->ForEach<FRigCurveElement>([&](FRigCurveElement* CurveElement) -> bool
+    {
+        const FName ControlName = GetControlName(CurveElement->GetName()); // name conflict?
+
+		FRigControlSettings Settings;
+		Settings.ControlType = ERigControlType::Float;
+		Settings.DisplayName = CurveElement->GetName();
+
+		Controller->AddControl(ControlName, FRigElementKey(), Settings, FRigControlValue::Make(CurveElement->Value), FTransform::Identity, FTransform::Identity, false);
+		
+		return true;
+	});
+
+	if (IsControlActive.Num() != GetHierarchy()->Num())
 	{
-		FRigCurve& RigCurve = CurveContainer[Index];
-		FName ControlName = GetControlName(RigCurve.Name); // name conflict?
-
-		FRigControl& RigControl = Container->ControlHierarchy.Add(ControlName, ERigControlType::Float, NAME_None, NAME_None); // NAME_None, LocalTransform);//SpaceName);  LocalTransform);
-		RigControl.DisplayName = RigCurve.Name;
-		RigControl.Value.Set<float>(RigCurve.Value);
+		IsControlActive.SetNum(GetHierarchy()->Num());
+		for (bool& bIsActive : IsControlActive)
+		{
+			bIsActive = true;
+		}
 	}
-
-	Container->Initialize(true);
 }
 
 void UFKControlRig::CreateRigElements(const USkeletalMesh* InReferenceMesh)
@@ -285,51 +301,44 @@ void UFKControlRig::ToggleApplyMode()
 		FRigControlModifiedContext Context;
 		Context.SetKey = EControlRigSetKey::Never;
 
-		FRigControlHierarchy& ControlHierarchy = GetControlHierarchy();
 		FTransform ZeroScale = FTransform::Identity;
 		ZeroScale.SetScale3D(FVector::ZeroVector);
-		FEulerTransform EulerZero(ZeroScale);
-		for (int32 ControlIndex = 0; ControlIndex < ControlHierarchy.Num(); ControlIndex++)
-		{
-			const FRigControl& ControlToChange = ControlHierarchy[ControlIndex];
-			if (ControlToChange.ControlType == ERigControlType::EulerTransform)
+		const FEulerTransform EulerZero(ZeroScale);
+
+		GetHierarchy()->ForEach<FRigControlElement>([&](FRigControlElement* ControlElement) -> bool
+        {
+            if (ControlElement->Settings.ControlType == ERigControlType::EulerTransform)
 			{
-				SetControlValue<FEulerTransform>(ControlToChange.Name, EulerZero, true, Context);
+				SetControlValue<FEulerTransform>(ControlElement->GetName(), EulerZero, true, Context);
 			}
-			else if (ControlToChange.ControlType == ERigControlType::Float)
+			else if (ControlElement->Settings.ControlType == ERigControlType::Float)
 			{
-				SetControlValue<float>(ControlToChange.Name, 0.f, true, Context);
+				SetControlValue<float>(ControlElement->GetName(), 0.f, true, Context);
 			}
-		}
+
+			return true;
+		});
 	}
 	else
 	{
-		FRigControlHierarchy& ControlHierarchy = GetControlHierarchy();
 		FRigControlModifiedContext Context;
 		Context.SetKey = EControlRigSetKey::Never;
 
-		for (int32 ControlIndex = 0; ControlIndex < ControlHierarchy.Num(); ControlIndex++)
-		{
-			const FRigControl& ControlToChange = ControlHierarchy[ControlIndex];
-			if (ControlToChange.ControlType == ERigControlType::EulerTransform)
+		GetHierarchy()->ForEach<FRigControlElement>([&](FRigControlElement* ControlElement) -> bool
+        {
+            if (ControlElement->Settings.ControlType == ERigControlType::EulerTransform)
 			{
-				const FRigControl* Control = FindControl(ControlToChange.Name);
-				if (Control)
-				{
-					FEulerTransform InitValue = Control->InitialValue.Get<FEulerTransform>();
-					SetControlValue<FEulerTransform>(ControlToChange.Name, InitValue, true, Context);
-				}
+				const FEulerTransform InitValue = GetHierarchy()->GetControlValue(ControlElement, ERigControlValueType::Initial).Get<FEulerTransform>();
+				SetControlValue<FEulerTransform>(ControlElement->GetName(), InitValue, true, Context);
 			}
-			else if (ControlToChange.ControlType == ERigControlType::Float)
+			else if (ControlElement->Settings.ControlType == ERigControlType::Float)
 			{
-				const FRigControl* Control = FindControl(ControlToChange.Name);
-				if (Control)
-				{
-					float InitValue = Control->InitialValue.Get<float>();
-					SetControlValue<float>(ControlToChange.Name, InitValue, true, Context);
-				}
+				const float InitValue = GetHierarchy()->GetControlValue(ControlElement, ERigControlValueType::Initial).Get<float>();
+				SetControlValue<float>(ControlElement->GetName(), InitValue, true, Context);
 			}
-		}
+
+			return true;
+		});
 	}
 }
 
diff --git a/Engine/Plugins/Experimental/ControlRig/Source/ControlRig/Private/Rigs/RigBoneHierarchy.cpp b/Engine/Plugins/Experimental/ControlRig/Source/ControlRig/Private/Rigs/RigBoneHierarchy.cpp
--- a/Engine/Plugins/Experimental/ControlRig/Source/ControlRig/Private/Rigs/RigBoneHierarchy.cpp
+++ b/Engine/Plugins/Experimental/ControlRig/Source/ControlRig/Private/Rigs/RigBoneHierarchy.cpp
@@ -1,941 +1,12 @@
 // Copyright Epic Games, Inc. All Rights Reserved.
 
 #include "Rigs/RigBoneHierarchy.h"
-#include "Rigs/RigHierarchyContainer.h"
-#include "ControlRig.h"
-#include "HelperUtil.h"
-#include "AnimationRuntime.h"
 
 ////////////////////////////////////////////////////////////////////////////////
 // FRigBoneHierarchy
 ////////////////////////////////////////////////////////////////////////////////
 
 FRigBoneHierarchy::FRigBoneHierarchy()
-	: Container(nullptr)
-	, bSuspendNotifications(false)
 {
 }
 
-FRigBoneHierarchy& FRigBoneHierarchy::operator= (const FRigBoneHierarchy &InOther)
-{
-#if WITH_EDITOR
-	if (!bSuspendNotifications)
-	{
-		for (int32 Index = Num() - 1; Index >= 0; Index--)
-		{
-			FRigBone BoneToRemove = Bones[Index];
-			OnBoneRemoved.Broadcast(Container, FRigElementKey(BoneToRemove.Name, ERigElementType::Bone));
-		}
-	}
-#endif
-
-	Bones.Reset();
-	Bones.Append(InOther.Bones);
-	NameToIndexMapping.Reset();
-	RefreshMapping();
-
-#if WITH_EDITOR
-	if (!bSuspendNotifications)
-	{
-		for (const FRigBone& BoneAdded : Bones)
-		{
-			OnBoneAdded.Broadcast(Container, FRigElementKey(BoneAdded.Name, ERigElementType::Bone));
-		}
-	}
-#endif
-
-	return *this;
-}
-
-bool FRigBoneHierarchy::IsNameAvailable(const FString& InPotentialNewName, FString* OutErrorMessage) const
-{
-	if (Container)
-	{
-		return Container->IsNameAvailable(InPotentialNewName, ERigElementType::Bone, OutErrorMessage);
-	}
-
-	// if we are not part of a container, fall back to index based validation
-	if (InPotentialNewName.Len() > FRigHierarchyContainer::GetMaxNameLength())
-	{
-		return false;
-	}
-	return GetIndex(*InPotentialNewName) == INDEX_NONE;
-}
-
-FName FRigBoneHierarchy::GetSafeNewName(const FString& InPotentialNewName) const
-{
-	return FRigHierarchyContainer::GetSafeNewName(InPotentialNewName, [this](const FString& InName) -> bool {
-		return IsNameAvailable(InName);
-	});
-}
-
-FRigBone& FRigBoneHierarchy::Add(const FName& InNewName, const FName& InParentName, ERigBoneType InType, const FTransform& InInitTransform)
-{
-	FRigBone NewBone;
-	NewBone.Name = GetSafeNewName(InNewName.ToString());
-	NewBone.ParentIndex = GetIndex(InParentName);
-	NewBone.ParentName = NewBone.ParentIndex == INDEX_NONE ? NAME_None : InParentName;
-	NewBone.InitialTransform = InInitTransform;
-	NewBone.GlobalTransform = InInitTransform;
-	NewBone.Type = InType;
-	RecalculateLocalTransform(NewBone);
-
-	FName NewBoneName = NewBone.Name;
-	Bones.Add(NewBone);
-	RefreshMapping();
-
-#if WITH_EDITOR
-	if (!bSuspendNotifications)
-	{
-		OnBoneAdded.Broadcast(Container, NewBone.GetElementKey());
-	}
-#endif
-
-	int32 Index = GetIndex(NewBoneName);
-	return Bones[Index];
-}
-
-FRigBone& FRigBoneHierarchy::Add(const FName& InNewName, const FName& InParentName, ERigBoneType InType, const FTransform& InInitTransform, const FTransform& InLocalTransform, const FTransform& InGlobalTransform)
-{
-	FRigBone& Bone = Add(InNewName, InParentName, InType, InInitTransform);
-	Bone.LocalTransform = InLocalTransform;
-	Bone.GlobalTransform = InGlobalTransform;
-	return Bone;
-}
-
-bool FRigBoneHierarchy::Reparent(const FName& InName, const FName& InNewParentName)
-{
-	int32 Index = GetIndex(InName);
-
-	if (Index != INDEX_NONE && InName != InNewParentName)
-	{
-		FRigBone& Bone= Bones[Index];
-
-#if WITH_EDITOR
-		FName OldParentName = Bone.ParentName;
-#endif
-
-		int32 ParentIndex = GetIndex(InNewParentName);
-		if (Container != nullptr)
-		{
-			if (Container->IsParentedTo(ERigElementType::Bone, ParentIndex, ERigElementType::Bone, Index))
-			{
-				ParentIndex = INDEX_NONE;
-			}
-		}
-
-		Bone.ParentIndex = ParentIndex;
-		Bone.ParentName = Bone.ParentIndex == INDEX_NONE ? NAME_None : InNewParentName;
-		RecalculateLocalTransform(Bone);
-
-#if WITH_EDITOR
-		FName NewParentName = Bone.ParentName;
-#endif
-
-		// we want to make sure parent is before the child
-		RefreshMapping();
-
-#if WITH_EDITOR
-		if (OldParentName != NewParentName)
-		{
-			if (!bSuspendNotifications)
-			{
-				OnBoneReparented.Broadcast(Container, FRigElementKey(InName, RigElementType()), OldParentName, NewParentName);
-			}
-		}
-#endif
-		return Bones[GetIndex(InName)].ParentName == InNewParentName;
-	}
-
-	return false;
-}
-
-FRigBone FRigBoneHierarchy::Remove(const FName& InNameToRemove)
-{
-	TArray<int32> Children;
-#if WITH_EDITOR
-	TArray<FName> RemovedChildBones;
-#endif
-	if (GetChildren(InNameToRemove, Children, true) > 0)
-	{
-		// sort by child index
-		Children.Sort([](const int32& A, const int32& B) { return A < B; });
-
-		// want to delete from end to the first 
-		for (int32 ChildIndex = Children.Num() - 1; ChildIndex >= 0; --ChildIndex)
-		{
-#if WITH_EDITOR
-			RemovedChildBones.Add(Bones[Children[ChildIndex]].Name);
-#endif
-			Bones.RemoveAt(Children[ChildIndex]);
-		}
-	}
-
-	int32 IndexToDelete = GetIndex(InNameToRemove);
-	if (IndexToDelete == INDEX_NONE)
-	{
-		return FRigBone();
-	}
-
-	Select(InNameToRemove, false);
-
-	FRigBone RemovedBone = Bones[IndexToDelete];
-	Bones.RemoveAt(IndexToDelete);
-
-	RefreshMapping();
-
-#if WITH_EDITOR
-	if (!bSuspendNotifications)
-	{
-		for (const FName& RemovedChildBone : RemovedChildBones)
-		{
-			OnBoneRemoved.Broadcast(Container, FRigElementKey(RemovedChildBone, RigElementType()));
-		}
-		OnBoneRemoved.Broadcast(Container, RemovedBone.GetElementKey());
-	}
-#endif
-
-	return RemovedBone;
-}
-
-// list of names of children - this is not cheap, and is supposed to be used only for one time set up
-int32 FRigBoneHierarchy::GetChildren(const FName& InName, TArray<int32>& OutChildren, bool bRecursively) const
-{
-	return GetChildren(GetIndex(InName), OutChildren, bRecursively);
-}
-
-int32 FRigBoneHierarchy::GetChildren(const int32 InIndex, TArray<int32>& OutChildren, bool bRecursively) const
-{
-	OutChildren.Reset();
-
-	if (InIndex != INDEX_NONE)
-	{
-		GetChildrenRecursive(InIndex, OutChildren, bRecursively);
-	}
-
-	return OutChildren.Num();
-}
-
-FName FRigBoneHierarchy::GetName(int32 InIndex) const
-{
-	if (Bones.IsValidIndex(InIndex))
-	{
-		return Bones[InIndex].Name;
-	}
-
-	return NAME_None;
-}
-
-int32 FRigBoneHierarchy::GetIndexSlow(const FName& InName) const
-{
-	for (int32 Index = 0; Index < Bones.Num(); ++Index)
-	{
-		if (Bones[Index].Name == InName)
-		{
-			return Index;
-		}
-	}
-
-	return INDEX_NONE;
-}
-
-void FRigBoneHierarchy::SetGlobalTransform(const FName& InName, const FTransform& InTransform, bool bPropagateTransform)
-{
-	SetGlobalTransform(GetIndex(InName), InTransform, bPropagateTransform);
-}
-
-void FRigBoneHierarchy::SetGlobalTransform(int32 InIndex, const FTransform& InTransform, bool bPropagateTransform)
-{
-	if (Bones.IsValidIndex(InIndex))
-	{
-		FRigBone& Bone = Bones[InIndex];
-		Bone.GlobalTransform = InTransform;
-		Bone.GlobalTransform.NormalizeRotation();
-		RecalculateLocalTransform(Bone);
-
-		if (bPropagateTransform)
-		{
-			PropagateTransform(InIndex);
-		}
-	}
-}
-
-FTransform FRigBoneHierarchy::GetGlobalTransform(const FName& InName) const
-{
-	return GetGlobalTransform(GetIndex(InName));
-}
-
-FTransform FRigBoneHierarchy::GetGlobalTransform(int32 InIndex) const
-{
-	if (Bones.IsValidIndex(InIndex))
-	{
-		return Bones[InIndex].GlobalTransform;
-	}
-
-	return FTransform::Identity;
-}
-
-void FRigBoneHierarchy::SetLocalTransform(const FName& InName, const FTransform& InTransform, bool bPropagateTransform)
-{
-	SetLocalTransform(GetIndex(InName), InTransform, bPropagateTransform);
-}
-
-void FRigBoneHierarchy::SetLocalTransform(int32 InIndex, const FTransform& InTransform, bool bPropagateTransform)
-{
-	if (Bones.IsValidIndex(InIndex))
-	{
-		FRigBone& Bone = Bones[InIndex];
-		Bone.LocalTransform = InTransform;
-		RecalculateGlobalTransform(Bone);
-
-		if (bPropagateTransform)
-		{
-			PropagateTransform(InIndex);
-		}
-	}
-}
-
-FTransform FRigBoneHierarchy::GetLocalTransform(const FName& InName) const
-{
-	return GetLocalTransform(GetIndex(InName));
-}
-
-FTransform FRigBoneHierarchy::GetLocalTransform(int32 InIndex) const
-{
-	if (Bones.IsValidIndex(InIndex))
-	{
-		return Bones[InIndex].LocalTransform;
-	}
-
-	return FTransform::Identity;
-}
-
-void FRigBoneHierarchy::SetInitialGlobalTransform(const FName& InName, const FTransform& InTransform, bool bPropagateTransform)
-{
-	SetInitialGlobalTransform(GetIndex(InName), InTransform);
-}
-
-void FRigBoneHierarchy::SetInitialGlobalTransform(int32 InIndex, const FTransform& InTransform, bool bPropagateTransform)
-{
-	if (Bones.IsValidIndex(InIndex))
-	{
-		FRigBone& Bone = Bones[InIndex];
-		
-		TArray<FTransform> LocalTransforms;
-		if (bPropagateTransform)
-		{
-			for (int32 Dependent : Bone.Dependents)
-			{
-				FTransform LocalTransform = Bones[Dependent].InitialTransform.GetRelativeTransform(Bone.InitialTransform);
-				LocalTransforms.Add(LocalTransform);
-			}
-		}
-
-		Bone.InitialTransform = InTransform;
-		Bone.InitialTransform.NormalizeRotation();
-
-		for (int32 LocalTransformIndex = 0; LocalTransformIndex < LocalTransforms.Num(); LocalTransformIndex++)
-		{
-			int32 Dependent = Bone.Dependents[LocalTransformIndex];
-			FTransform InitialTransform = LocalTransforms[LocalTransformIndex] * Bone.InitialTransform;
-			SetInitialGlobalTransform(Dependent, InitialTransform, bPropagateTransform);
-		}
-	}
-}
-
-void FRigBoneHierarchy::SetInitialLocalTransform(const FName& InName, const FTransform& InTransform, bool bPropagateTransform)
-{
-	SetInitialLocalTransform(GetIndex(InName), InTransform, bPropagateTransform);
-}
-
-void FRigBoneHierarchy::SetInitialLocalTransform(int32 InIndex, const FTransform& InTransform, bool bPropagateTransform)
-{
-	if (Bones.IsValidIndex(InIndex))
-	{
-		FRigBone& Bone = Bones[InIndex];
-
-		FTransform Transform = InTransform;
-		if (Bone.ParentIndex != INDEX_NONE)
-		{
-			FTransform ParentTransform = GetInitialGlobalTransform(Bone.ParentIndex);
-			Transform = Transform * ParentTransform;
-		}
-
-		SetInitialGlobalTransform(InIndex, Transform, bPropagateTransform);
-	}
-}
-
-FTransform FRigBoneHierarchy::GetInitialGlobalTransform(const FName& InName) const
-{
-	return GetInitialGlobalTransform(GetIndex(InName));
-}
-
-FTransform FRigBoneHierarchy::GetInitialGlobalTransform(int32 InIndex) const
-{
-	if (Bones.IsValidIndex(InIndex))
-	{
-		return Bones[InIndex].InitialTransform;
-	}
-
-	return FTransform::Identity;
-}
-
-FTransform FRigBoneHierarchy::GetInitialLocalTransform(const FName& InName) const
-{
-	return GetInitialLocalTransform(GetIndex(InName));
-}
-
-FTransform FRigBoneHierarchy::GetInitialLocalTransform(int32 InIndex) const
-{
-	if (Bones.IsValidIndex(InIndex))
-	{
-		FTransform Transform = GetInitialGlobalTransform(InIndex);
-		if (Bones[InIndex].ParentIndex != INDEX_NONE)
-		{
-			FTransform ParentTransform = GetInitialGlobalTransform(Bones[InIndex].ParentIndex);
-			Transform = Transform.GetRelativeTransform(ParentTransform);
-		}
-		return Transform;
-	}
-
-	return FTransform::Identity;
-}
-
-void FRigBoneHierarchy::RecalculateLocalTransform(int32 InIndex)
-{
-	RecalculateLocalTransform(Bones[InIndex]);
-}
-
-void FRigBoneHierarchy::RecalculateGlobalTransform(int32 InIndex)
-{
-	RecalculateGlobalTransform(Bones[InIndex]);
-}
-
-void FRigBoneHierarchy::RecalculateLocalTransform(FRigBone& InOutBone)
-{
-	bool bHasParent = InOutBone.ParentIndex != INDEX_NONE;
-	InOutBone.LocalTransform = (bHasParent) ? InOutBone.GlobalTransform.GetRelativeTransform(Bones[InOutBone.ParentIndex].GlobalTransform) : InOutBone.GlobalTransform;
-}
-
-void FRigBoneHierarchy::RecalculateGlobalTransform(FRigBone& InOutBone)
-{
-	bool bHasParent = InOutBone.ParentIndex != INDEX_NONE;
-	InOutBone.GlobalTransform = (bHasParent) ? InOutBone.LocalTransform * Bones[InOutBone.ParentIndex].GlobalTransform : InOutBone.LocalTransform;
-}
-
-FName FRigBoneHierarchy::Rename(const FName& InOldName, const FName& InNewName)
-{
-	if (InOldName != InNewName)
-	{
-		const int32 Found = GetIndex(InOldName);
-		if (Found != INDEX_NONE)
-		{
-			FName NewName = GetSafeNewName(InNewName.ToString());
-			bool bWasSelected = IsSelected(InOldName);
-			if(bWasSelected)
-			{
-				Select(InOldName, false);
-			}
-
-			Bones[Found].Name = NewName;
-
-			// go through find all children and rename them
-#if WITH_EDITOR
-			TArray<FName> ReparentedBones;
-#endif
-			for (int32 Index = 0; Index < Bones.Num(); ++Index)
-			{
-				if (Bones[Index].ParentName == InOldName)
-				{
-					Bones[Index].ParentName = NewName;
-					RecalculateLocalTransform(Bones[Index]);
-#if WITH_EDITOR
-					ReparentedBones.Add(Bones[Index].Name);
-#endif
-				}
-			}
-
-			RefreshMapping();
-
-#if WITH_EDITOR
-			if (!bSuspendNotifications)
-			{
-				OnBoneRenamed.Broadcast(Container, RigElementType(), InOldName, NewName);
-				for (const FName& ReparentedBone : ReparentedBones)
-				{
-					OnBoneReparented.Broadcast(Container, FRigElementKey(ReparentedBone, RigElementType()), InOldName, NewName);
-				}
-			}
-			if (bWasSelected)
-			{
-				Select(NewName, true);
-			}
-#endif
-			return NewName;
-		}
-	}
-
-	return NAME_None;
-}
-
-void FRigBoneHierarchy::Sort()
-{
-	TArray<int32> RootArray;
-	TMap<int32, TArray<int32>> HierarchyTree;
-
-	// first figure out children and find roots
-	for (int32 Index = 0; Index < Bones.Num(); ++Index)
-	{
-		int32 ParentIndex = GetIndexSlow(Bones[Index].ParentName);
-		if (ParentIndex != INDEX_NONE)
-		{
-			TArray<int32>& ChildIndices = HierarchyTree.FindOrAdd(ParentIndex);
-			ChildIndices.Add(Index);
-		}
-		else
-		{
-			// add as a root
-			HierarchyTree.FindOrAdd(Index);
-			RootArray.Add(Index);
-		}
-	}
-	
-	struct Local
-	{
-		static void VisitBone(int32 Index, const TMap<int32, TArray<int32>>& InHierarchyTree, TArray<int32>& SortedArray)
-		{
-			if (const TArray<int32>* ChildIndices = InHierarchyTree.Find(Index))
-			{
-				SortedArray.Append(*ChildIndices);
-				for (int32 ChildIndex : *ChildIndices)
-				{
-					VisitBone(ChildIndex, InHierarchyTree, SortedArray);
-				}
-			}
-		}
-	};
-
-	TArray<int32> SortedArray;
-	for (int32 RootIndex : RootArray)
-	{
-		SortedArray.Add(RootIndex);
-		Local::VisitBone(RootIndex, HierarchyTree, SortedArray);
-	}
-
-	check(SortedArray.Num() == Bones.Num());
-
-	// create new list with sorted
-	TArray<FRigBone> NewSortedList;
-	NewSortedList.AddDefaulted(Bones.Num());
-	for (int32 NewIndex = 0; NewIndex < SortedArray.Num(); ++NewIndex)
-	{
-		NewSortedList[NewIndex] = Bones[SortedArray[NewIndex]];
-	}
-
-	Bones = MoveTemp(NewSortedList);
-
-	// now fix up parent Index
-	for (int32 Index = 0; Index < Bones.Num(); ++Index)
-	{
-		Bones[Index].ParentIndex = GetIndexSlow(Bones[Index].ParentName);
-		// parent index always should be less than this index, even if invalid
-		check(Bones[Index].ParentIndex < Index);
-	}
-}
-
-void FRigBoneHierarchy::RefreshParentNames()
-{
-	for (FRigBone& Bone : Bones)
-	{
-		Bone.ParentIndex = INDEX_NONE;
-		if (Bone.ParentName != NAME_None)
-		{
-			Bone.ParentIndex = GetIndex(Bone.ParentName);
-		}
-	}
-}
-
-void FRigBoneHierarchy::RefreshMapping()
-{
-	// sort doesn't rely on the name map,
-	// since it uses GetIndexSlow.
-	Sort();
-
-	NameToIndexMapping.Empty();
-	for (int32 Index = 0; Index < Bones.Num(); ++Index)
-	{
-		Bones[Index].Index = Index;
-		NameToIndexMapping.Add(Bones[Index].Name, Index);
-	}
-
-	RefreshParentNames();
-
-	// update the dependents
-	for (int32 Index = 0; Index < Bones.Num(); ++Index)
-	{
-		GetChildren(Index, Bones[Index].Dependents, false);
-	}
-}
-
-void FRigBoneHierarchy::Initialize(bool bResetTransforms)
-{
-	RefreshMapping();
-
-	for (int32 Index = 0; Index < Bones.Num(); ++Index)
-	{
-		Bones[Index].ParentIndex = GetIndex(Bones[Index].ParentName);
-	}
-
-	// initialize transform
-	for (int32 Index = 0; Index < Bones.Num(); ++Index)
-	{
-		if (bResetTransforms)
-		{
-			Bones[Index].GlobalTransform = Bones[Index].InitialTransform;
-			RecalculateLocalTransform(Bones[Index]);
-		}
-
-		// update children
-		GetChildren(Index, Bones[Index].Dependents, false);
-	}
-}
-
-void FRigBoneHierarchy::Reset()
-{
-	Bones.Reset();
-}
-
-void FRigBoneHierarchy::ResetTransforms()
-{
-	// initialize transform
-	for (int32 Index = 0; Index < Bones.Num(); ++Index)
-	{
-		Bones[Index].GlobalTransform = Bones[Index].InitialTransform;
-		RecalculateLocalTransform(Bones[Index]);
-	}
-}
-
-void FRigBoneHierarchy::CopyInitialTransforms(const FRigBoneHierarchy& InOther)
-{
-	ensure(InOther.Num() == Num());
-
-	for (int32 Index = 0; Index < Bones.Num(); ++Index)
-	{
-		Bones[Index].InitialTransform = InOther.Bones[Index].InitialTransform;
-	}
-}
-
-void FRigBoneHierarchy::RecomputeGlobalTransforms()
-{
-	for (int32 Index = 0; Index < Bones.Num(); ++Index)
-	{
-		Bones[Index].GlobalTransform = Bones[Index].InitialTransform;
-		RecalculateGlobalTransform(Bones[Index]);
-	}
-}
-
-int32 FRigBoneHierarchy::GetChildrenRecursive(const int32 InIndex, TArray<int32>& OutChildren, bool bRecursively) const
-{
-	const int32 StartChildIndex = OutChildren.Num();
-
-	// all children should be later than parent
-	for (int32 ChildIndex = InIndex + 1; ChildIndex < Bones.Num(); ++ChildIndex)
-	{
-		if (Bones[ChildIndex].ParentIndex == InIndex)
-		{
-			OutChildren.AddUnique(ChildIndex);
-		}
-	}
-
-	if (bRecursively)
-	{
-		// since we keep appending inside of functions, we make sure not to go over original list
-		const int32 EndChildIndex = OutChildren.Num() - 1;
-		for (int32 ChildIndex = StartChildIndex; ChildIndex <= EndChildIndex; ++ChildIndex)
-		{
-			GetChildrenRecursive(OutChildren[ChildIndex], OutChildren, bRecursively);
-		}
-	}
-
-	return OutChildren.Num();
-}
-
-void FRigBoneHierarchy::PropagateTransform(int32 InIndex)
-{
-	const TArray<int32> Dependents = Bones[InIndex].Dependents;
-	for (int32 DependentIndex = 0; DependentIndex<Dependents.Num(); ++DependentIndex)
-	{
-		int32 Index = Dependents[DependentIndex];
-		RecalculateGlobalTransform(Bones[Index]);
-		PropagateTransform(Index);
-	}
-}
-
-void FRigBoneHierarchy::CopyGlobalToInitial(int32 InIndex, bool bPropagate)
-{
-	FTransform InitialGlobalTransform = GetGlobalTransform(InIndex);
-	SetInitialGlobalTransform(InIndex, InitialGlobalTransform);
-
-	if (bPropagate)
-	{
-		const TArray<int32> Dependents = Bones[InIndex].Dependents;
-		for (int32 DependentIndex = 0; DependentIndex < Dependents.Num(); ++DependentIndex)
-		{
-			CopyGlobalToInitial(Dependents[DependentIndex], bPropagate);
-		}
-	}
-}
-
-bool FRigBoneHierarchy::Select(const FName& InName, bool bSelect)
-{
-	if(GetIndex(InName) == INDEX_NONE)
-	{
-		return false;
-	}
-
-	if(bSelect == IsSelected(InName))
-	{
-		return false;
-	}
-
-	if(bSelect)
-	{
-		Selection.Add(InName);
-	}
-	else
-	{
-		Selection.Remove(InName);
-	}
-
-	if (!bSuspendNotifications)
-	{
-		OnBoneSelected.Broadcast(Container, FRigElementKey(InName, RigElementType()), bSelect);
-	}
-
-	return true;
-}
-
-bool FRigBoneHierarchy::ClearSelection()
-{
-	TArray<FName> TempSelection;
-	TempSelection.Append(Selection);
-	for(const FName& SelectedName : TempSelection)
-	{
-		Select(SelectedName, false);
-	}
-	return TempSelection.Num() > 0;
-}
-
-TArray<FName> FRigBoneHierarchy::CurrentSelection() const
-{
-	TArray<FName> TempSelection;
-	TempSelection.Append(Selection);
-	return TempSelection;
-}
-
-bool FRigBoneHierarchy::IsSelected(const FName& InName) const
-{
-	return Selection.Contains(InName);
-}
-
-TArray<FRigElementKey> FRigBoneHierarchy::ImportSkeleton(const FReferenceSkeleton& InSkeleton, const FName& InNameSpace, bool bReplaceExistingBones, bool bRemoveObsoleteBones, bool bSelectBones, bool bNotify)
-{
-	TGuardValue<bool> SuspendNotification(bSuspendNotifications, !bNotify);
-
-	TArray<FRigElementKey> AddedBones;
-	TArray<FName> BoneNamesToSelect;
-	TMap<FName, FName> BoneNameMap;
-
-	ResetTransforms();
-
-	const TArray<FMeshBoneInfo>& BoneInfos = InSkeleton.GetRefBoneInfo();
-	const TArray<FTransform>& BonePoses = InSkeleton.GetRefBonePose();
-
-	struct Local
-	{
-		static FName DetermineBoneName(const FName& InBoneName, const FName& InLocalNameSpace)
-		{
-			if (InLocalNameSpace == NAME_None || InBoneName == NAME_None)
-			{
-				return InBoneName;
-			}
-			return *FString::Printf(TEXT("%s_%s"), *InLocalNameSpace.ToString(), *InBoneName.ToString());
-		}
-	};
-
-	if (bReplaceExistingBones)
-	{
-		for (const FRigBone& ExistingBone : Bones)
-		{
-			BoneNameMap.Add(ExistingBone.Name, ExistingBone.Name);
-		}
-
-		for (int32 Index = 0; Index < BoneInfos.Num(); ++Index)
-		{
-			int32 ExistingBoneIndex = GetIndex(BoneInfos[Index].Name);
-			FName DesiredBoneName = Local::DetermineBoneName(BoneInfos[Index].Name, InNameSpace);
-			FName ParentName = (BoneInfos[Index].ParentIndex != INDEX_NONE) ? BoneInfos[BoneInfos[Index].ParentIndex].Name : NAME_None;
-			ParentName = Local::DetermineBoneName(ParentName, InNameSpace);
-
-			const FName* MappedParentNamePtr = BoneNameMap.Find(ParentName);
-			if (MappedParentNamePtr)
-			{
-				ParentName = *MappedParentNamePtr;
-			}
-
-			if (ExistingBoneIndex != INDEX_NONE)
-			{
-				FRigBone& ExistingBone = Bones[ExistingBoneIndex];
-				// check it's parent
-				if (ParentName != NAME_None)
-				{
-					ExistingBone.ParentName = ParentName;
-				}
-
-				ExistingBone.LocalTransform = BonePoses[Index];
-
-				BoneNamesToSelect.Add(ExistingBone.Name);
-			}
-			else
-			{
-				FRigBone& AddedBone = Add(DesiredBoneName, ParentName, ERigBoneType::Imported, FTransform::Identity);
-				BoneNameMap.Add(DesiredBoneName, AddedBone.Name);
-				AddedBones.Add(AddedBone.GetElementKey());
-				BoneNamesToSelect.Add(AddedBone.Name);
-
-				AddedBone.LocalTransform = BonePoses[Index];
-			}
-		}
-
-	}
-	else // import all as new
-	{
-		for (int32 Index = 0; Index < BoneInfos.Num(); ++Index)
-		{
-			FName DesiredBoneName = Local::DetermineBoneName(BoneInfos[Index].Name, InNameSpace);
-			FName ParentName = (BoneInfos[Index].ParentIndex != INDEX_NONE) ? BoneInfos[BoneInfos[Index].ParentIndex].Name : NAME_None;
-			ParentName = Local::DetermineBoneName(ParentName, InNameSpace);
-
-			const FName* MappedParentNamePtr = BoneNameMap.Find(ParentName);
-			if (MappedParentNamePtr)
-			{
-				ParentName = *MappedParentNamePtr;
-			}
-			
-			FRigBone& AddedBone = Add(DesiredBoneName, ParentName, ERigBoneType::Imported, FTransform::Identity);
-			BoneNameMap.Add(DesiredBoneName, AddedBone.Name);
-			AddedBones.Add(AddedBone.GetElementKey());
-			BoneNamesToSelect.Add(AddedBone.Name);
-
-			// FAnimationRuntime::GetComponentSpaceTransform(InSkeleton, BonePoses, Index)
-			AddedBone.LocalTransform = BonePoses[Index];
-		}
-	}
-
-	RefreshMapping();
-
-	RecomputeGlobalTransforms();
-
-	// copy the initial from the global
-	for (FRigBone& Bone : Bones)
-	{
-		Bone.InitialTransform = Bone.GlobalTransform;
-	}
-
-	if (bReplaceExistingBones && bRemoveObsoleteBones)
-	{
-		TMap<FName, int32> BoneNameToIndexInSkeleton;
-		for (const FMeshBoneInfo& BoneInfo : BoneInfos)
-		{
-			FName DesiredBoneName = Local::DetermineBoneName(BoneInfo.Name, InNameSpace);
-			BoneNameToIndexInSkeleton.Add(DesiredBoneName, BoneNameToIndexInSkeleton.Num());
-		}
-		
-		TArray<FName> BonesToDelete;
-		for (int32 ExistingBoneIndex = 0; ExistingBoneIndex < Bones.Num(); ExistingBoneIndex++)
-		{
-			const FRigBone& Bone = Bones[ExistingBoneIndex];
-			if (!BoneNameToIndexInSkeleton.Contains(Bone.Name))
-			{
-				if (Bone.Type == ERigBoneType::Imported)
-				{
-					BonesToDelete.Add(Bone.Name);
-				}
-			}
-		}
-
-		for (const FName& BoneToDelete : BonesToDelete)
-		{
-			const FRigBone& Bone = Bones[GetIndex(BoneToDelete)];
-			TArray<FName> Dependents;
-			for (int32 DependentIndex : Bone.Dependents)
-			{
-				if (Bones[DependentIndex].Type != ERigBoneType::Imported)
-				{
-					Dependents.Add(Bones[DependentIndex].Name);
-				}
-			}
-
-			Algo::Reverse(Dependents);
-			for (const FName& DependentName : Dependents)
-			{
-				Reparent(DependentName, NAME_None);
-			}
-		}
-
-		for (const FName& BoneToDelete : BonesToDelete)
-		{
-			Remove(BoneToDelete);
-			BoneNamesToSelect.Remove(BoneToDelete);
-		}
-	}
-
-	if (bSelectBones)
-	{
-		ClearSelection();
-		for (const FName& BoneNameToSelect : BoneNamesToSelect)
-		{
-			Select(BoneNameToSelect);
-		}
-	}
-
-	Initialize();
-
-	return AddedBones;
-}
-
-FRigPose FRigBoneHierarchy::GetPose() const
-{
-	FRigPose Pose;
-	AppendToPose(Pose);
-	return Pose;
-}
-
-void FRigBoneHierarchy::SetPose(FRigPose& InPose)
-{
-	for(FRigPoseElement& Element : InPose)
-	{
-		if(Element.Index.GetKey().Type == ERigElementType::Bone)
-		{
-			if(Element.Index.UpdateCache(Container))
-			{
-				Bones[Element.Index.GetIndex()].GlobalTransform = Element.GlobalTransform;
-				Bones[Element.Index.GetIndex()].LocalTransform = Element.LocalTransform;
-			}
-		}
-	}
-}
-
-void FRigBoneHierarchy::AppendToPose(FRigPose& InOutPose) const
-{
-	for(const FRigBone& Bone : Bones)
-	{
-		FRigPoseElement Element;
-		if(Element.Index.UpdateCache(Bone.GetElementKey(), Container))
-		{
-			Element.GlobalTransform = GetGlobalTransform(Element.Index.GetIndex());
-			Element.LocalTransform = GetLocalTransform(Element.Index.GetIndex());
-			InOutPose.Elements.Add(Element);
-		}
-	}
-}
\ No newline at end of file
