Control Rig: New rig hierarchy data structure refactoring

diff --git a/Engine/Plugins/Experimental/ControlRig/Source/ControlRig/Private/Rigs/RigControlHierarchy.cpp b/Engine/Plugins/Experimental/ControlRig/Source/ControlRig/Private/Rigs/RigControlHierarchy.cpp
--- a/Engine/Plugins/Experimental/ControlRig/Source/ControlRig/Private/Rigs/RigControlHierarchy.cpp
+++ b/Engine/Plugins/Experimental/ControlRig/Source/ControlRig/Private/Rigs/RigControlHierarchy.cpp
@@ -1,467 +1,24 @@
 // Copyright Epic Games, Inc. All Rights Reserved.
 
 #include "Rigs/RigControlHierarchy.h"
-#include "Rigs/RigHierarchyContainer.h"
-#include "ControlRig.h"
-#include "HelperUtil.h"
-#include "Settings/ControlRigSettings.h"
-#include "Async/TaskGraphInterfaces.h"
 
 ////////////////////////////////////////////////////////////////////////////////
 // FRigControl
 ////////////////////////////////////////////////////////////////////////////////
 
-void FRigControl::ApplyLimits(FRigControlValue& InOutValue)
+void FRigControl::ApplyLimits(FRigControlValue& InOutValue) const
 {
-	if (!bLimitTranslation && !bLimitRotation && !bLimitScale)
-	{
-		return;
-	}
-
-	switch(ControlType)
-	{
-		case ERigControlType::Float:
-		{
-			if (bLimitTranslation)
-			{
-				float& ValueRef = InOutValue.GetRef<float>();
-				ValueRef = Clamp(ValueRef, MinimumValue.Get<float>(), MaximumValue.Get<float>());
-			}
-			break;
-		}
-		case ERigControlType::Integer:
-		{
-			if (bLimitTranslation)
-			{
-				int32& ValueRef = InOutValue.GetRef<int32>();
-				ValueRef = Clamp(ValueRef, MinimumValue.Get<int32>(), MaximumValue.Get<int32>());
-			}
-			break;
-		}
-		case ERigControlType::Vector2D:
-		{
-			if (bLimitTranslation)
-			{
-				FVector2D& ValueRef = InOutValue.GetRef<FVector2D>();
-				const FVector2D& Min = MinimumValue.GetRef<FVector2D>();
-				const FVector2D& Max = MaximumValue.GetRef<FVector2D>();
-				ValueRef.X = Clamp(ValueRef.X, Min.X, Max.X);
-				ValueRef.Y = Clamp(ValueRef.Y, Min.Y, Max.Y);
-			}
-			break;
-		}
-		case ERigControlType::Position:
-		{
-			if (bLimitTranslation)
-			{
-				FVector& ValueRef = InOutValue.GetRef<FVector>();
-				const FVector& Min = MinimumValue.GetRef<FVector>();
-				const FVector& Max = MaximumValue.GetRef<FVector>();
-				ValueRef.X = Clamp(ValueRef.X, Min.X, Max.X);
-				ValueRef.Y = Clamp(ValueRef.Y, Min.Y, Max.Y);
-				ValueRef.Z = Clamp(ValueRef.Z, Min.Z, Max.Z);
-			}
-			break;
-		}
-		case ERigControlType::Scale:
-		{
-			if (bLimitScale)
-			{
-				FVector& ValueRef = InOutValue.GetRef<FVector>();
-				const FVector& Min = MinimumValue.GetRef<FVector>();
-				const FVector& Max = MaximumValue.GetRef<FVector>();
-				ValueRef.X = Clamp(ValueRef.X, Min.X, Max.X);
-				ValueRef.Y = Clamp(ValueRef.Y, Min.Y, Max.Y);
-				ValueRef.Z = Clamp(ValueRef.Z, Min.Z, Max.Z);
-			}
-			break;
-		}
-		case ERigControlType::Rotator:
-		{
-			if (bLimitRotation)
-			{
-				FRotator& ValueRef = InOutValue.GetRef<FRotator>();
-				const FRotator& Min = MinimumValue.GetRef<FRotator>();
-				const FRotator& Max = MaximumValue.GetRef<FRotator>();
-				ValueRef.Pitch = Clamp(ValueRef.Pitch, Min.Pitch, Max.Pitch);
-				ValueRef.Yaw = Clamp(ValueRef.Yaw, Min.Yaw, Max.Yaw);
-				ValueRef.Roll = Clamp(ValueRef.Roll, Min.Roll, Max.Roll);
-			}
-			break;
-		}
-		case ERigControlType::Transform:
-		{
-			FTransform& ValueRef = InOutValue.GetRef<FTransform>();
-			const FTransform& Min = MinimumValue.GetRef<FTransform>();
-			const FTransform& Max = MaximumValue.GetRef<FTransform>();
-
-			if (bLimitTranslation)
-			{
-				ValueRef.SetLocation(FVector(
-					Clamp(ValueRef.GetLocation().X, Min.GetLocation().X, Max.GetLocation().X),
-					Clamp(ValueRef.GetLocation().Y, Min.GetLocation().Y, Max.GetLocation().Y),
-					Clamp(ValueRef.GetLocation().Z, Min.GetLocation().Z, Max.GetLocation().Z)
-				));
-			}
-			if (bLimitRotation)
-			{
-				FRotator Rotator = ValueRef.GetRotation().Rotator();
-				FRotator MinRotator = Min.GetRotation().Rotator();
-				FRotator MaxRotator = Max.GetRotation().Rotator();
-
-				ValueRef.SetRotation(FQuat(FRotator(
-					Clamp(Rotator.Pitch, MinRotator.Pitch, MaxRotator.Pitch),
-					Clamp(Rotator.Yaw, MinRotator.Yaw, MaxRotator.Yaw),
-					Clamp(Rotator.Roll, MinRotator.Roll, MaxRotator.Roll)
-				)));
-			}
-			if (bLimitScale)
-			{
-				ValueRef.SetScale3D(FVector(
-					Clamp(ValueRef.GetScale3D().X, Min.GetScale3D().X, Max.GetScale3D().X),
-					Clamp(ValueRef.GetScale3D().Y, Min.GetScale3D().Y, Max.GetScale3D().Y),
-					Clamp(ValueRef.GetScale3D().Z, Min.GetScale3D().Z, Max.GetScale3D().Z)
-				));
-			}
-			break;
-		}
-		case ERigControlType::TransformNoScale:
-		{
-			FTransformNoScale& ValueRef = InOutValue.GetRef<FTransformNoScale>();
-			const FTransformNoScale& Min = MinimumValue.GetRef<FTransformNoScale>();
-			const FTransformNoScale& Max = MaximumValue.GetRef<FTransformNoScale>();
-
-			if (bLimitTranslation)
-			{
-				ValueRef.Location = FVector(
-					Clamp(ValueRef.Location.X, Min.Location.X, Max.Location.X),
-					Clamp(ValueRef.Location.Y, Min.Location.Y, Max.Location.Y),
-					Clamp(ValueRef.Location.Z, Min.Location.Z, Max.Location.Z)
-				);
-			}
-			if (bLimitRotation)
-			{
-				FRotator Rotator = ValueRef.Rotation.Rotator();
-				FRotator MinRotator = Min.Rotation.Rotator();
-				FRotator MaxRotator = Max.Rotation.Rotator();
-
-				ValueRef.Rotation = FQuat(FRotator(
-					Clamp(Rotator.Pitch, MinRotator.Pitch, MaxRotator.Pitch),
-					Clamp(Rotator.Yaw, MinRotator.Yaw, MaxRotator.Yaw),
-					Clamp(Rotator.Roll, MinRotator.Roll, MaxRotator.Roll)
-				));
-			}
-			break;
-		}
-
-		case ERigControlType::EulerTransform:
-		{
-			FEulerTransform& ValueRef = InOutValue.GetRef<FEulerTransform>();
-			const FEulerTransform& Min = MinimumValue.GetRef<FEulerTransform>();
-			const FEulerTransform& Max = MaximumValue.GetRef<FEulerTransform>();
-
-			if (bLimitTranslation)
-			{
-				ValueRef.Location = FVector(
-					Clamp(ValueRef.Location.X, Min.Location.X, Max.Location.X),
-					Clamp(ValueRef.Location.Y, Min.Location.Y, Max.Location.Y),
-					Clamp(ValueRef.Location.Z, Min.Location.Z, Max.Location.Z)
-				);
-			}
-			if (bLimitRotation)
-			{
-				FRotator Rotator = ValueRef.Rotation;
-				FRotator MinRotator = Min.Rotation;
-				FRotator MaxRotator = Max.Rotation;
-
-				ValueRef.Rotation = FRotator(
-					Clamp(Rotator.Pitch, MinRotator.Pitch, MaxRotator.Pitch),
-					Clamp(Rotator.Yaw, MinRotator.Yaw, MaxRotator.Yaw),
-					Clamp(Rotator.Roll, MinRotator.Roll, MaxRotator.Roll)
-				);
-			}
-			if (bLimitScale)
-			{
-				ValueRef.Location = FVector(
-					Clamp(ValueRef.Scale.X, Min.Scale.X, Max.Scale.X),
-					Clamp(ValueRef.Scale.Y, Min.Scale.Y, Max.Scale.Y),
-					Clamp(ValueRef.Scale.Z, Min.Scale.Z, Max.Scale.Z)
-				);
-			}
-			break;
-		}
-		case ERigControlType::Bool:
-		default:
-		{
-			break;
-		}
-	}
+	InOutValue.ApplyLimits(bLimitTranslation, bLimitRotation, bLimitScale, ControlType, MinimumValue, MaximumValue);
 }
 
 FTransform FRigControl::GetTransformFromValue(ERigControlValueType InValueType) const
 {
-	switch (ControlType)
-	{
-		case ERigControlType::Bool:
-		{
-			FTransform Transform;
-			Transform.SetLocation(FVector(GetValue(InValueType).Get<bool>() ? 1.f : 0.f, 0.f, 0.f));
-			return Transform;
-		}
-		case ERigControlType::Float:
-		{
-			float ValueToGet = GetValue(InValueType).Get<float>();
-			FTransform Transform = FTransform::Identity;
-			switch (PrimaryAxis)
-			{
-				case ERigControlAxis::X:
-				{
-					Transform.SetLocation(FVector(ValueToGet, 0.f, 0.f));
-					break;
-				}
-				case ERigControlAxis::Y:
-				{
-					Transform.SetLocation(FVector(0.f, ValueToGet, 0.f));
-					break;
-				}
-				case ERigControlAxis::Z:
-				{
-					Transform.SetLocation(FVector(0.f, 0.f, ValueToGet));
-					break;
-				}
-			}
-			return Transform;
-		}
-		case ERigControlType::Integer:
-		{
-			int32 ValueToGet = GetValue(InValueType).Get<int32>();
-			FTransform Transform = FTransform::Identity;
-			switch (PrimaryAxis)
-			{
-				case ERigControlAxis::X:
-				{
-					Transform.SetLocation(FVector((float)ValueToGet, 0.f, 0.f));
-					break;
-				}
-				case ERigControlAxis::Y:
-				{
-					Transform.SetLocation(FVector(0.f, (float)ValueToGet, 0.f));
-					break;
-				}
-				case ERigControlAxis::Z:
-				{
-					Transform.SetLocation(FVector(0.f, 0.f, (float)ValueToGet));
-					break;
-				}
-			}
-			return Transform;
-		}
-		case ERigControlType::Vector2D:
-		{
-			FVector2D ValueToGet = GetValue(InValueType).Get<FVector2D>();
-			FTransform Transform = FTransform::Identity;
-			switch (PrimaryAxis)
-			{
-				case ERigControlAxis::X:
-				{
-					Transform.SetLocation(FVector(0.f, ValueToGet.X, ValueToGet.Y));
-					break;
-				}
-				case ERigControlAxis::Y:
-				{
-					Transform.SetLocation(FVector(ValueToGet.X, 0.f, ValueToGet.Y));
-					break;
-				}
-				case ERigControlAxis::Z:
-				{
-					Transform.SetLocation(FVector(ValueToGet.X, ValueToGet.Y, 0.f));
-					break;
-				}
-			}
-			return Transform;
-		}
-		case ERigControlType::Position:
-		{
-			FTransform Transform;
-			Transform.SetLocation(GetValue(InValueType).Get<FVector>());
-			return Transform;
-		}
-		case ERigControlType::Scale:
-		{
-			FTransform Transform;
-			Transform.SetScale3D(GetValue(InValueType).Get<FVector>());
-			return Transform;
-		}
-		case ERigControlType::Rotator:
-		{
-			FTransform Transform;
-			Transform.SetRotation(FQuat(GetValue(InValueType).Get<FRotator>()));
-			return Transform;
-		}
-		case ERigControlType::Transform:
-		{
-			return GetValue(InValueType).Get<FTransform>();
-		}
-		case ERigControlType::TransformNoScale:
-		{
-			FTransformNoScale TransformNoScale = GetValue(InValueType).Get<FTransformNoScale>();
-			FTransform Transform = TransformNoScale;
-			Transform.NormalizeRotation();
-			return Transform;
-		}
-		case ERigControlType::EulerTransform:
-		{
-			FEulerTransform EulerTransform = GetValue(InValueType).Get<FEulerTransform>();
-			FTransform Transform(EulerTransform.ToFTransform());
-			Transform.NormalizeRotation();
-			return Transform;
-		}
-		default:
-		{
-			ensure(false);
-			break;
-		}
-	}
-	return FTransform::Identity;
+	return GetValue(InValueType).GetAsTransform(ControlType, PrimaryAxis);
 }
 
 void FRigControl::SetValueFromTransform(const FTransform& InTransform, ERigControlValueType InValueType)
 {
-	switch (ControlType)
-	{
-		case ERigControlType::Bool:
-		{
-			GetValue(InValueType).Set<bool>(InTransform.GetLocation().X > SMALL_NUMBER);
-			break;
-		}
-		case ERigControlType::Float:
-		{
-			switch (PrimaryAxis)
-			{
-				case ERigControlAxis::X:
-				{
-					GetValue(InValueType).Set<float>(InTransform.GetLocation().X);
-					break;
-				}
-				case ERigControlAxis::Y:
-				{
-					GetValue(InValueType).Set<float>(InTransform.GetLocation().Y);
-					break;
-				}
-				case ERigControlAxis::Z:
-				{
-					GetValue(InValueType).Set<float>(InTransform.GetLocation().Z);
-					break;
-				}
-			}
-			break;
-		}
-		case ERigControlType::Integer:
-		{
-			switch (PrimaryAxis)
-			{
-				case ERigControlAxis::X:
-				{
-					GetValue(InValueType).Set<int32>((int32)InTransform.GetLocation().X);
-					break;
-				}
-				case ERigControlAxis::Y:
-				{
-					GetValue(InValueType).Set<int32>((int32)InTransform.GetLocation().Y);
-					break;
-				}
-				case ERigControlAxis::Z:
-				{
-					GetValue(InValueType).Set<int32>((int32)InTransform.GetLocation().Z);
-					break;
-				}
-			}
-			break;
-		}
-		case ERigControlType::Vector2D:
-		{
-			FVector Location = InTransform.GetLocation();
-			switch (PrimaryAxis)
-			{
-				case ERigControlAxis::X:
-				{
-					GetValue(InValueType).Set<FVector2D>(FVector2D(Location.Y, Location.Z));
-					break;
-				}
-				case ERigControlAxis::Y:
-				{
-					GetValue(InValueType).Set<FVector2D>(FVector2D(Location.X, Location.Z));
-					break;
-				}
-				case ERigControlAxis::Z:
-				{
-					GetValue(InValueType).Set<FVector2D>(FVector2D(Location.X, Location.Y));
-					break;
-				}
-			}
-			break;
-		}
-		case ERigControlType::Position:
-		{
-			GetValue(InValueType).Set<FVector>(InTransform.GetLocation());
-			break;
-		}
-		case ERigControlType::Scale:
-		{
-			GetValue(InValueType).Set<FVector>(InTransform.GetScale3D());
-			break;
-		}
-		case ERigControlType::Rotator:
-		{
-			//allow for values ><180/-180 by getting diff and adding that back in.
-			FRotator CurrentRotator = GetValue(InValueType).Get<FRotator>();
-			FRotator CurrentRotWind, CurrentRotRem;
-			CurrentRotator.GetWindingAndRemainder(CurrentRotWind, CurrentRotRem);
-
-			//Get Diff
-			const FRotator NewRotator = FRotator(InTransform.GetRotation());
-			FRotator DeltaRot = NewRotator - CurrentRotRem;
-			DeltaRot.Normalize();
-
-			//Add Diff
-			CurrentRotator = CurrentRotator + DeltaRot;
-			GetValue(InValueType).Set<FRotator>(CurrentRotator);
-			break;
-		}
-		case ERigControlType::Transform:
-		{
-			GetValue(InValueType).Set<FTransform>(InTransform);
-			break;
-		}
-		case ERigControlType::TransformNoScale:
-		{
-			FTransformNoScale NoScale = InTransform;
-			GetValue(InValueType).Set<FTransformNoScale>(NoScale);
-			break;
-		}
-		case ERigControlType::EulerTransform:
-		{
-			//Find Diff of the rotation from current and just add that instead of setting so we can go over/under -180
-			FEulerTransform NewTransform(InTransform);
-
-			FEulerTransform CurrentEulerTransform = GetValue(InValueType).Get<FEulerTransform>();
-			FRotator CurrentWinding;
-			FRotator CurrentRotRemainder;
-			CurrentEulerTransform.Rotation.GetWindingAndRemainder(CurrentWinding, CurrentRotRemainder);
-			FRotator NewRotator = InTransform.GetRotation().Rotator();
-			FRotator DeltaRot = NewRotator - CurrentRotRemainder;
-			DeltaRot.Normalize();
-			const FRotator NewRotation(CurrentEulerTransform.Rotation + DeltaRot);
-			NewTransform.Rotation = NewRotation;
-			GetValue(InValueType).Set<FEulerTransform>(NewTransform);
-			break;
-		}
-		default:
-		{
-			ensure(false);
-			break;
-		}
-	}
+	GetValue(InValueType).SetFromTransform(InTransform, ControlType, PrimaryAxis);
 
 	if (InValueType == ERigControlValueType::Current)
 	{
@@ -474,921 +31,20 @@ void FRigControl::SetValueFromTransform(const FTransform& InTransform, ERigContr
 ////////////////////////////////////////////////////////////////////////////////
 
 FRigControlHierarchy::FRigControlHierarchy()
-	:Container(nullptr)
-{
-}
-
-FRigControlHierarchy& FRigControlHierarchy::operator= (const FRigControlHierarchy &InOther)
-{
-#if WITH_EDITOR
-	for (int32 Index = Num() - 1; Index >= 0; Index--)
-	{
-		FRigControl ControlToRemove = Controls[Index];
-		OnControlRemoved.Broadcast(Container, FRigElementKey(ControlToRemove.Name, ERigElementType::Control));
-	}
-#endif
-
-	Controls.Reset();
-	Controls.Append(InOther.Controls);
-	NameToIndexMapping.Reset();
-	RefreshMapping();
-
-#if WITH_EDITOR
-	for (const FRigControl& ControlAdded : Controls)
-	{
-		OnControlAdded.Broadcast(Container, FRigElementKey(ControlAdded.Name, ERigElementType::Control));
-	}
-#endif
-
-	return *this;
-}
-
-bool FRigControlHierarchy::IsNameAvailable(const FString& InPotentialNewName, FString* OutErrorMessage) const
-{
-	if (Container)
-	{
-		return Container->IsNameAvailable(InPotentialNewName, ERigElementType::Control, OutErrorMessage);
-	}
-
-	// if we are not part of a container, fall back to index based validation
-	if (InPotentialNewName.Len() > FRigHierarchyContainer::GetMaxNameLength())
-	{
-		return false;
-	}
-	return GetIndex(*InPotentialNewName) == INDEX_NONE;
-}
-
-FName FRigControlHierarchy::GetSafeNewName(const FString& InPotentialNewName) const
 {
-	return FRigHierarchyContainer::GetSafeNewName(InPotentialNewName, [this](const FString& InName) -> bool {
-		return IsNameAvailable(InName);
-	});
-}
-
-FRigControl& FRigControlHierarchy::Add(
-	const FName& InNewName,
-	ERigControlType InControlType,
-	const FName& InParentName,
-	const FName& InSpaceName,
-	const FTransform& InOffsetTransform,
-	const FRigControlValue& InValue,
-	const FName& InGizmoName,
-	const FTransform& InGizmoTransform,
-	const FLinearColor& InGizmoColor
-)
-{
-	FRigControl NewControl;
-	NewControl.Name = GetSafeNewName(InNewName.ToString());
-	NewControl.ControlType = InControlType;
-	NewControl.ParentIndex = GetIndex(InParentName);
-	NewControl.ParentName = NewControl.ParentIndex == INDEX_NONE ? NAME_None : InParentName;
-	NewControl.SpaceIndex = INDEX_NONE;
-	NewControl.SpaceName = NAME_None;
-	NewControl.OffsetTransform = InOffsetTransform;
-	NewControl.InitialValue = InValue;
-	NewControl.Value = FRigControlValue();
-	NewControl.GizmoName = InGizmoName;
-	NewControl.GizmoTransform = InGizmoTransform;
-	NewControl.GizmoColor = InGizmoColor;
-
-	if (!NewControl.InitialValue.IsValid())
-	{
-		NewControl.SetValueFromTransform(FTransform::Identity, ERigControlValueType::Initial);
-	}
-
-	FName NewControlName = NewControl.Name;
-	Controls.Add(NewControl);
-	RefreshMapping();
-
-#if WITH_EDITOR
-	OnControlAdded.Broadcast(Container, NewControl.GetElementKey());
-#endif
-
-	SetSpace(NewControlName, InSpaceName);
-
-	int32 Index = GetIndex(NewControlName);
-	return Controls[Index];
-}
-
-bool FRigControlHierarchy::Reparent(const FName& InName, const FName& InNewParentName)
-{
-	int32 Index = GetIndex(InName);
-	// can't parent to itself
-	if (Index != INDEX_NONE && InName != InNewParentName)
-	{
-		FRigControl& Control = Controls[Index];
-
-#if WITH_EDITOR
-		FName OldParentName = Control.ParentName;
-#endif
-
-		struct Local
-		{
-			static bool IsParentedTo(int32 Child, int32 Parent, const TArray<FRigControl>& Controls)
-			{
-				if (Parent == INDEX_NONE || Child == INDEX_NONE)
-				{
-					return false;
-				}
-
-				if (Child == Parent)
-				{
-					return true;
-				}
-
-				if (Controls[Child].ParentIndex == Parent)
-				{
-					return true;
-				}
-
-				return IsParentedTo(Controls[Child].ParentIndex, Parent, Controls);
-			}
-		};
-
-		int32 ParentIndex = GetIndex(InNewParentName);
-		if (Local::IsParentedTo(ParentIndex, Index, Controls))
-		{
-			ParentIndex = INDEX_NONE;
-		}
-
-		Control.ParentIndex = ParentIndex;
-		Control.ParentName = Control.ParentIndex == INDEX_NONE ? NAME_None : InNewParentName;
-
-#if WITH_EDITOR
-		FName NewParentName = Control.ParentName;
-#endif
-
-		RefreshMapping();
-
-#if WITH_EDITOR
-		if (OldParentName != NewParentName)
-		{
-			OnControlReparented.Broadcast(Container, FRigElementKey(InName, RigElementType()), OldParentName, NewParentName);
-		}
-#endif
-		return Controls[GetIndex(InName)].ParentName == InNewParentName;
-	}
-	return false;
 }
 
-void FRigControlHierarchy::SetSpace(const FName& InName, const FName& InNewSpaceName)
+void FRigControlHierarchy::PostLoad()
 {
-	int32 Index = GetIndex(InName);
-	if (Index != INDEX_NONE)
+	for (FRigControl& Control : Controls)
 	{
-		int32 SpaceIndex = GetSpaceIndex(InNewSpaceName);
-
-		if (SpaceIndex != INDEX_NONE)
+		for (int32 ValueType = 0; ValueType <= (int32)ERigControlValueType::Maximum; ValueType++)
 		{
-			if (Container != nullptr)
+			FRigControlValue& Value = Control.GetValue((ERigControlValueType)ValueType);
+			if (!Value.IsValid())
 			{
-				if (Container->IsParentedTo(ERigElementType::Space, SpaceIndex, ERigElementType::Control, Index))
-				{
-					SpaceIndex = INDEX_NONE;
-				}
+				Value.GetRef<FTransform>() = Value.Storage_DEPRECATED;
 			}
 		}
-
-		Controls[Index].SpaceIndex = SpaceIndex;
-		Controls[Index].SpaceName = Controls[Index].SpaceIndex == INDEX_NONE ? NAME_None : InNewSpaceName;
-
-#if WITH_EDITOR
-		OnControlReparented.Broadcast(Container, FRigElementKey(InName, RigElementType()), Controls[Index].ParentName, Controls[Index].ParentName);
-#endif
-	}
-}
-
-FRigControl FRigControlHierarchy::Remove(const FName& InNameToRemove)
-{
-	TArray<int32> Children;
-#if WITH_EDITOR
-	TArray<FName> RemovedChildControls;
-#endif
-	if (GetChildren(InNameToRemove, Children, true) > 0)
-	{
-		// sort by child index
-		Children.Sort([](const int32& A, const int32& B) { return A < B; });
-
-		// want to delete from end to the first 
-		for (int32 ChildIndex = Children.Num() - 1; ChildIndex >= 0; --ChildIndex)
-		{
-#if WITH_EDITOR
-			RemovedChildControls.Add(Controls[Children[ChildIndex]].Name);
-#endif
-			Controls.RemoveAt(Children[ChildIndex]);
-		}
-	}
-
-	int32 IndexToDelete = GetIndex(InNameToRemove);
-	Select(InNameToRemove, false);
-	FRigControl RemovedControl = Controls[IndexToDelete];
-	Controls.RemoveAt(IndexToDelete);
-
-	RefreshMapping();
-
-#if WITH_EDITOR
-	for (const FName& RemovedChildControl : RemovedChildControls)
-	{
-		OnControlRemoved.Broadcast(Container, FRigElementKey(RemovedChildControl, RigElementType()));
-	}
-	OnControlRemoved.Broadcast(Container, RemovedControl.GetElementKey());
-#endif
-
-	return RemovedControl;
-}
-
-// list of names of children - this is not cheap, and is supposed to be used only for one time set up
-int32 FRigControlHierarchy::GetChildren(const FName& InName, TArray<int32>& OutChildren, bool bRecursively) const
-{
-	return GetChildren(GetIndex(InName), OutChildren, bRecursively);
-}
-
-int32 FRigControlHierarchy::GetChildren(const int32 InIndex, TArray<int32>& OutChildren, bool bRecursively) const
-{
-	OutChildren.Reset();
-
-	if (InIndex != INDEX_NONE)
-	{
-		GetChildrenRecursive(InIndex, OutChildren, bRecursively);
-	}
-
-	return OutChildren.Num();
-}
-
-FName FRigControlHierarchy::GetName(int32 InIndex) const
-{
-	if (Controls.IsValidIndex(InIndex))
-	{
-		return Controls[InIndex].Name;
-	}
-
-	return NAME_None;
-}
-
-int32 FRigControlHierarchy::GetIndexSlow(const FName& InName) const
-{
-	for (int32 Index = 0; Index < Controls.Num(); ++Index)
-	{
-		if (Controls[Index].Name == InName)
-		{
-			return Index;
-		}
-	}
-
-	return INDEX_NONE;
-}
-
-void FRigControlHierarchy::SetGlobalTransform(const FName& InName, const FTransform& InTransform)
-{
-	SetGlobalTransform(GetIndex(InName), InTransform);
-}
-
-void FRigControlHierarchy::SetGlobalTransform(int32 InIndex, const FTransform& InTransform)
-{
-	if (Container == nullptr)
-	{
-		SetLocalTransform(InIndex, InTransform);
-		return;
-	}
-
-	if (Controls.IsValidIndex(InIndex))
-	{
-		const FRigControl& Control = Controls[InIndex];
-		FTransform ParentTransform = GetParentTransform(InIndex);
-		SetLocalTransform(InIndex, InTransform.GetRelativeTransform(ParentTransform));
-	}
-}
-
-FTransform FRigControlHierarchy::GetGlobalTransform(const FName& InName) const
-{
-	return GetGlobalTransform(GetIndex(InName));
-}
-
-FTransform FRigControlHierarchy::GetGlobalTransform(int32 InIndex) const
-{
-	if (Container == nullptr)
-	{
-		return GetLocalTransform(InIndex);
-	}
-
-	if (Controls.IsValidIndex(InIndex))
-	{
-		FTransform ParentTransform = GetParentTransform(InIndex);
-		FTransform Transform = GetLocalTransform(InIndex) * ParentTransform;
-		Transform.NormalizeRotation();
-		return Transform;
-	}
-
-	return FTransform::Identity;
-}
-
-void FRigControlHierarchy::SetLocalTransform(const FName& InName, const FTransform& InTransform, ERigControlValueType InValueType)
-{
-	SetLocalTransform(GetIndex(InName), InTransform, InValueType);
-}
-
-void FRigControlHierarchy::SetLocalTransform(int32 InIndex, const FTransform& InTransform, ERigControlValueType InValueType)
-{
-	FRigControl& Control = Controls[InIndex];
-	Control.SetValueFromTransform(InTransform, InValueType);
-}
-
-FTransform FRigControlHierarchy::GetLocalTransform(const FName& InName, ERigControlValueType InValueType) const
-{
-	return GetLocalTransform(GetIndex(InName), InValueType);
-}
-
-FTransform FRigControlHierarchy::GetLocalTransform(int32 InIndex, ERigControlValueType InValueType) const
-{
-	const FRigControl& Control = Controls[InIndex];
-	return Control.GetTransformFromValue(InValueType);
-}
-
-FTransform FRigControlHierarchy::GetParentTransform(int32 InIndex, bool bIncludeOffsetTransform) const
-{
-	const FRigControl& Control = Controls[InIndex];
-
-	FTransform ParentTransform = FTransform::Identity;
-	if (Control.SpaceIndex != INDEX_NONE && Container != nullptr)
-	{
-		ParentTransform = Container->GetGlobalTransform(ERigElementType::Space, Control.SpaceIndex);
-	}
-	else if (Control.ParentIndex != INDEX_NONE && Container != nullptr)
-	{
-		ParentTransform = GetGlobalTransform(Control.ParentIndex);
-	}
-
-	if (bIncludeOffsetTransform)
-	{
-		return Control.OffsetTransform * ParentTransform;
-	}
-	return ParentTransform;
-}
-
-FTransform FRigControlHierarchy::GetParentInitialTransform(int32 InIndex, bool bIncludeOffsetTransform) const
-{
-	const FRigControl& Control = Controls[InIndex];
-
-	FTransform ParentTransform = FTransform::Identity;
-	if (Control.SpaceIndex != INDEX_NONE && Container != nullptr)
-	{
-		ParentTransform = Container->GetInitialGlobalTransform(ERigElementType::Space, Control.SpaceIndex);
-	}
-	else if (Control.ParentIndex != INDEX_NONE && Container != nullptr)
-	{
-		ParentTransform = GetInitialGlobalTransform(Control.ParentIndex);
-	}
-
-	if (bIncludeOffsetTransform)
-	{
-		return Control.OffsetTransform * ParentTransform;
-	}
-	return ParentTransform;
-}
-
-void FRigControlHierarchy::SetValue(const FName& InName, const FRigControlValue& InValue, ERigControlValueType InValueType)
-{
-	SetValue(GetIndex(InName), InValue, InValueType);
-}
-
-void FRigControlHierarchy::SetValue(int32 InIndex, const FRigControlValue& InValue, ERigControlValueType InValueType)
-{
-	if (Controls.IsValidIndex(InIndex))
-	{
-		FRigControl& Control = Controls[InIndex];
-
-		Control.GetValue(InValueType) = InValue;
-
-		if (InValueType == ERigControlValueType::Current)
-		{
-			Control.ApplyLimits(Control.GetValue(InValueType));
-		}
-
-		if (Control.ControlType == ERigControlType::Transform)
-		{
-			Control.GetValue(InValueType).GetRef<FTransform>().NormalizeRotation();
-		}
-		else if (Control.ControlType == ERigControlType::TransformNoScale)
-		{
-			Control.GetValue(InValueType).GetRef<FTransformNoScale>().Rotation.Normalize();
-		}
-	}
-}
-
-FRigControlValue FRigControlHierarchy::GetValue(const FName& InName, ERigControlValueType InValueType) const
-{
-	return GetValue(GetIndex(InName), InValueType);
-}
-
-FRigControlValue FRigControlHierarchy::GetValue(int32 InIndex, ERigControlValueType InValueType) const
-{
-	if (Controls.IsValidIndex(InIndex))
-	{
-		return Controls[InIndex].GetValue(InValueType);
-	}
-
-	return FRigControlValue();
-}
-
-void FRigControlHierarchy::SetInitialGlobalTransform(const FName& InName, const FTransform& GlobalTransform)
-{
-	SetInitialGlobalTransform(GetIndex(InName), GlobalTransform);
-}
-
-// wip - should support all types that can provide space transform data (pos/rotation)
-void FRigControlHierarchy::SetInitialGlobalTransform(int32 InIndex, const FTransform& GlobalTransform)
-{
-	if (Controls.IsValidIndex(InIndex))
-	{
-		FRigControl& Control = Controls[InIndex];
-		FTransform ParentTransform = FTransform::Identity;
-		if (Container)
-		{
-			if (Control.SpaceName != NAME_None)
-			{
-				ParentTransform = Container->GetInitialGlobalTransform(Control.GetSpaceElementKey());
-			}
-			else
-			{
-				ParentTransform = Container->GetInitialGlobalTransform(Control.GetParentElementKey());
-			}
-		}
-
-		ParentTransform = Control.OffsetTransform * ParentTransform;
-
-		if (Control.ControlType == ERigControlType::Transform)
-		{
-			Control.InitialValue.Set<FTransform>(GlobalTransform.GetRelativeTransform(ParentTransform));
-		}
-	}
-}
-
-FTransform FRigControlHierarchy::GetInitialGlobalTransform(const FName& InName) const
-{
-	return GetInitialGlobalTransform(GetIndex(InName));
-}
-
-FTransform FRigControlHierarchy::GetInitialGlobalTransform(int32 InIndex) const
-{
-	// @todo: Templatize
-	if (Controls.IsValidIndex(InIndex))
-	{
-		const FRigControl& Control = Controls[InIndex];
-		FTransform ParentTransform = FTransform::Identity;
-		if (Container)
-		{
-			if (Control.SpaceName != NAME_None)
-			{
-				ParentTransform = Container->GetInitialGlobalTransform(Control.GetSpaceElementKey());
-			}
-			else
-			{
-				ParentTransform = Container->GetInitialGlobalTransform(Control.GetParentElementKey());
-			}
-		}
-
-		ParentTransform = Control.OffsetTransform * ParentTransform;
-
-		switch (Control.ControlType)
-		{
-			case ERigControlType::Bool:
-			{
-				break;
-			}
-			case ERigControlType::Float:
-			{
-				float Value = Control.InitialValue.Get<float>();
-				FTransform Transform = FTransform::Identity;
-				switch (Control.PrimaryAxis)
-				{
-					case ERigControlAxis::X:
-					{
-						Transform.SetLocation(FVector(Value, 0.f, 0.f));
-						break;
-					}
-					case ERigControlAxis::Y:
-					{
-						Transform.SetLocation(FVector(0.f, Value, 0.f));
-						break;
-					}
-					case ERigControlAxis::Z:
-					{
-						Transform.SetLocation(FVector(0.f, 0.f, Value));
-						break;
-					}
-				}
-				return Transform * ParentTransform;
-			}
-			case ERigControlType::Integer:
-			{
-				int32 Value = Control.InitialValue.Get<int32>();
-				FTransform Transform = FTransform::Identity;
-				switch (Control.PrimaryAxis)
-				{
-					case ERigControlAxis::X:
-					{
-						Transform.SetLocation(FVector((float)Value, 0.f, 0.f));
-						break;
-					}
-					case ERigControlAxis::Y:
-					{
-						Transform.SetLocation(FVector(0.f, (float)Value, 0.f));
-						break;
-					}
-					case ERigControlAxis::Z:
-					{
-						Transform.SetLocation(FVector(0.f, 0.f, (float)Value));
-						break;
-					}
-				}
-				return Transform * ParentTransform;
-			}
-			case ERigControlType::Vector2D:
-			{
-				FVector2D Value = Control.InitialValue.Get<FVector2D>();
-				FTransform Transform = FTransform::Identity;
-				switch (Control.PrimaryAxis)
-				{
-					case ERigControlAxis::X:
-					{
-						Transform.SetLocation(FVector(0.f, Value.X, Value.Y));
-						break;
-					}
-					case ERigControlAxis::Y:
-					{
-						Transform.SetLocation(FVector(Value.X, 0.f, Value.Y));
-						break;
-					}
-					case ERigControlAxis::Z:
-					{
-						Transform.SetLocation(FVector(Value.X, Value.Y, 0.f));
-						break;
-					}
-				}
-				return Transform * ParentTransform;
-			}
-			case ERigControlType::Position:
-			{
-				FTransform Transform = FTransform::Identity;
-				Transform.SetLocation(Control.InitialValue.Get<FVector>());
-				return Transform * ParentTransform;
-			}
-			case ERigControlType::Scale:
-			{
-				FTransform Transform = FTransform::Identity;
-				Transform.SetScale3D(Control.InitialValue.Get<FVector>());
-				return Transform * ParentTransform;
-			}
-			case ERigControlType::Rotator:
-			{
-				FTransform Transform = FTransform::Identity;
-				Transform.SetRotation(FQuat(Control.InitialValue.Get<FRotator>()));
-				return Transform * ParentTransform;
-			}
-			case ERigControlType::Transform:
-			{
-				return Control.InitialValue.Get<FTransform>() * ParentTransform;
-			}
-			case ERigControlType::TransformNoScale:
-			{
-				return Control.InitialValue.Get<FTransformNoScale>().ToFTransform() * ParentTransform;
-			}
-			case ERigControlType::EulerTransform:
-			{
-				return Control.InitialValue.Get<FEulerTransform>().ToFTransform() * ParentTransform;
-			}
-		}
-	}
-	return FTransform::Identity;
-}
-
-void FRigControlHierarchy::SetControlOffset(int32 InIndex, const FTransform& InOffsetTransform)
-{
-	Controls[InIndex].OffsetTransform = InOffsetTransform;
-
-	if (OnControlUISettingsChanged.IsBound())
-	{
-		FRigHierarchyContainer* LocalContainer = Container;
-		FRigElementChanged& Delegate = OnControlUISettingsChanged;
-		FRigElementKey Key = Controls[InIndex].GetElementKey();
-
-		FFunctionGraphTask::CreateAndDispatchWhenReady([LocalContainer, Delegate, Key]()
-		{
-			Delegate.Broadcast(LocalContainer, Key);
-		}, TStatId(), NULL, ENamedThreads::GameThread);
-	}
-}
-
-FName FRigControlHierarchy::Rename(const FName& InOldName, const FName& InNewName)
-{
-	if (InOldName != InNewName)
-	{
-		const int32 Found = GetIndex(InOldName);
-		if (Found != INDEX_NONE)
-		{
-			FName NewName = GetSafeNewName(InNewName.ToString());
-
-			bool bWasSelected = IsSelected(InOldName);
-			if(bWasSelected)
-			{
-				Select(InOldName, false);
-			}
-
-			Controls[Found].Name = NewName;
-
-			// go through find all children and rename them
-#if WITH_EDITOR
-			TArray<FName> ReparentedControls;
-#endif
-			for (int32 Index = 0; Index < Controls.Num(); ++Index)
-			{
-				if (Controls[Index].ParentName == InOldName)
-				{
-					Controls[Index].ParentName = NewName;
-#if WITH_EDITOR
-					ReparentedControls.Add(Controls[Index].Name);
-#endif
-				}
-			}
-
-			RefreshMapping();
-
-#if WITH_EDITOR
-			OnControlRenamed.Broadcast(Container, RigElementType(), InOldName, NewName);
-			for (const FName& ReparentedControl : ReparentedControls)
-			{
-				OnControlReparented.Broadcast(Container, FRigElementKey(ReparentedControl, RigElementType()), InOldName, NewName);
-			}
-#endif
-			if(bWasSelected)
-			{
-				Select(NewName, true);
-			}
-			return NewName;
-		}
-	}
-
-	return NAME_None;
-}
-
-void FRigControlHierarchy::RefreshMapping()
-{
-	NameToIndexMapping.Empty();
-	for (int32 Index = 0; Index < Controls.Num(); ++Index)
-	{
-		Controls[Index].Index = Index;
-		NameToIndexMapping.Add(Controls[Index].Name, Index);
-	}
-}
-
-void FRigControlHierarchy::Initialize(bool bResetTransforms)
-{
-	RefreshMapping();
-
-	// update parent index
-	for (int32 Index = 0; Index < Controls.Num(); ++Index)
-	{
-		Controls[Index].ParentIndex = GetIndex(Controls[Index].ParentName);
-		if (Container)
-		{
-			Controls[Index].SpaceIndex = Container->SpaceHierarchy.GetIndex(Controls[Index].SpaceName);
-		}
-	}
-
-	// initialize transform
-	for (int32 Index = 0; Index < Controls.Num(); ++Index)
-	{
-		if (bResetTransforms)
-		{
-			Controls[Index].Value = Controls[Index].InitialValue;
-		}
-
-		// update children
-		GetChildren(Index, Controls[Index].Dependents, false);
-	}
-}
-
-void FRigControlHierarchy::Reset()
-{
-	Controls.Reset();
-}
-
-void FRigControlHierarchy::ResetValues()
-{
-	// initialize transform
-	for (int32 Index = 0; Index < Controls.Num(); ++Index)
-	{
-		Controls[Index].Value = Controls[Index].InitialValue;
-	}
-}
-
-void FRigControlHierarchy::CopyOffsetTransforms(const FRigControlHierarchy& InOther)
-{
-	ensure(InOther.Num() == Num());
-
-	for (int32 Index = 0; Index < Controls.Num(); ++Index)
-	{
-		Controls[Index].OffsetTransform = InOther.Controls[Index].OffsetTransform;
-	}
-}
-
-int32 FRigControlHierarchy::GetChildrenRecursive(const int32 InIndex, TArray<int32>& OutChildren, bool bRecursively) const
-{
-	const int32 StartChildIndex = OutChildren.Num();
-
-	// all children should be later than parent
-	for (int32 ChildIndex = InIndex + 1; ChildIndex < Controls.Num(); ++ChildIndex)
-	{
-		if (Controls[ChildIndex].ParentIndex == InIndex)
-		{
-			OutChildren.AddUnique(ChildIndex);
-		}
-	}
-
-	if (bRecursively)
-	{
-		// since we keep appending inside of functions, we make sure not to go over original list
-		const int32 EndChildIndex = OutChildren.Num() - 1;
-		for (int32 ChildIndex = StartChildIndex; ChildIndex <= EndChildIndex; ++ChildIndex)
-		{
-			GetChildrenRecursive(OutChildren[ChildIndex], OutChildren, bRecursively);
-		}
-	}
-
-	return OutChildren.Num();
-}
-
-int32 FRigControlHierarchy::GetSpaceIndex(const FName& InName) const
-{
-	if (Container == nullptr || InName == NAME_None)
-	{
-		return INDEX_NONE;
-	}
-	return Container->GetIndex(FRigElementKey(InName, ERigElementType::Space));
-}
-
-bool FRigControlHierarchy::Select(const FName& InName, bool bSelect)
-{
-	if(GetIndex(InName) == INDEX_NONE)
-	{
-		return false;
-	}
-
-	if(bSelect == IsSelected(InName))
-	{
-		return false;
-	}
-
-	if(bSelect)
-	{
-		Selection.Add(InName);
-	}
-	else
-	{
-		Selection.Remove(InName);
-	}
-
-	OnControlSelected.Broadcast(Container, FRigElementKey(InName, RigElementType()), bSelect);
-
-	return true;
-}
-
-bool FRigControlHierarchy::ClearSelection()
-{
-	TArray<FName> TempSelection;
-	TempSelection.Append(Selection);
-	for(const FName& SelectedName : TempSelection)
-	{
-		Select(SelectedName, false);
-	}
-	return TempSelection.Num() > 0;
-}
-
-TArray<FName> FRigControlHierarchy::CurrentSelection() const
-{
-	TArray<FName> TempSelection;
-	TempSelection.Append(Selection);
-	return TempSelection;
-}
-
-bool FRigControlHierarchy::IsSelected(const FName& InName) const
-{
-	return Selection.Contains(InName);
-}
-
-void FRigControlHierarchy::HandleOnElementRemoved(FRigHierarchyContainer* InContainer, const FRigElementKey& InKey)
-{
-	if (Container == nullptr)
-	{
-		return;
-	}
-
-	switch (InKey.Type)
-	{
-		case ERigElementType::Space:
-		{
-			for (FRigControl& Control : Controls)
-			{
-				if (Control.SpaceName == InKey.Name)
-				{
-					Control.SpaceIndex = INDEX_NONE;
-					Control.SpaceName = NAME_None;
-#if WITH_EDITOR
-					OnControlReparented.Broadcast(Container, Control.GetElementKey(), Control.ParentName, Control.ParentName);
-#endif
-				}
-			}
-			break;
-		}
-		case ERigElementType::Bone:
-		case ERigElementType::Control:
-		case ERigElementType::Curve:
-		{
-			break;
-		}
-	}
-}
-
-void FRigControlHierarchy::HandleOnElementRenamed(FRigHierarchyContainer* InContainer, ERigElementType InElementType, const FName& InOldName, const FName& InNewName)
-{
-	if (Container == nullptr)
-	{
-		return;
-	}
-
-	switch (InElementType)
-	{
-		case ERigElementType::Space:
-		{
-			for (FRigControl& Control : Controls)
-			{
-				if (Control.SpaceName == InOldName)
-				{
-					Control.SpaceIndex = Container->SpaceHierarchy.GetIndex(InNewName);
-					Control.SpaceName = Control.SpaceIndex == INDEX_NONE ? NAME_None : InNewName;
-#if WITH_EDITOR
-					OnControlReparented.Broadcast(Container, Control.GetElementKey(), Control.ParentName, Control.ParentName);
-#endif
-				}
-			}
-			break;
-		}
-		case ERigElementType::Bone:
-		case ERigElementType::Control:
-		case ERigElementType::Curve:
-		{
-			break;
-		}
-	}
-}
-
-void FRigControlHierarchy::PostLoad()
-{
-	RefreshMapping();
-	for (FRigControl& Control : Controls)
-	{
-		for (int32 ValueType = 0; ValueType <= (int32)ERigControlValueType::Maximum; ValueType++)
-		{
-			FRigControlValue& Value = Control.GetValue((ERigControlValueType)ValueType);
-			if (!Value.IsValid())
-			{
-				Value.GetRef<FTransform>() = Value.Storage_DEPRECATED;
-			}
-		}
-	}
-}
-
-FRigPose FRigControlHierarchy::GetPose() const
-{
-	FRigPose Pose;
-	AppendToPose(Pose);
-	return Pose;
-}
-
-void FRigControlHierarchy::SetPose(FRigPose& InPose)
-{
-	for(FRigPoseElement& Element : InPose)
-	{
-		if(Element.Index.GetKey().Type == ERigElementType::Control)
-		{
-			if(Element.Index.UpdateCache(Container))
-			{
-				SetLocalTransform(Element.Index.GetIndex(), Element.LocalTransform);
-			}
-		}
-	}
-}
-
-void FRigControlHierarchy::AppendToPose(FRigPose& InOutPose) const
-{
-	for(const FRigControl& Control : Controls)
-	{
-		FRigPoseElement Element;
-		if(Element.Index.UpdateCache(Control.GetElementKey(), Container))
-		{
-			Element.GlobalTransform = GetGlobalTransform(Element.Index.GetIndex());
-			Element.LocalTransform = GetLocalTransform(Element.Index.GetIndex());
-			InOutPose.Elements.Add(Element);
-		}
 	}
 }
\ No newline at end of file
diff --git a/Engine/Plugins/Experimental/ControlRig/Source/ControlRig/Private/Rigs/RigCurveContainer.cpp b/Engine/Plugins/Experimental/ControlRig/Source/ControlRig/Private/Rigs/RigCurveContainer.cpp
--- a/Engine/Plugins/Experimental/ControlRig/Source/ControlRig/Private/Rigs/RigCurveContainer.cpp
+++ b/Engine/Plugins/Experimental/ControlRig/Source/ControlRig/Private/Rigs/RigCurveContainer.cpp
@@ -11,350 +11,5 @@
 ////////////////////////////////////////////////////////////////////////////////
 
 FRigCurveContainer::FRigCurveContainer()
-	: Container(nullptr)
-	, bSuspendNotifications(false)
 {
 }
-
-FRigCurveContainer& FRigCurveContainer::operator= (const FRigCurveContainer &InOther)
-{
-#if WITH_EDITOR
-	if (!bSuspendNotifications)
-	{
-		for (int32 Index = Num() - 1; Index >= 0; Index--)
-		{
-			FRigCurve CurveToRemove = Curves[Index];
-			OnCurveRemoved.Broadcast(Container, FRigElementKey(CurveToRemove.Name, ERigElementType::Curve));
-		}
-	}
-#endif
-
-	Curves.Reset();
-	Curves.Append(InOther.Curves);
-	NameToIndexMapping.Reset();
-	RefreshMapping();
-
-#if WITH_EDITOR
-	if (!bSuspendNotifications)
-	{
-		for (const FRigCurve& CurveAdded : Curves)
-		{
-			OnCurveAdded.Broadcast(Container, FRigElementKey(CurveAdded.Name, ERigElementType::Curve));
-		}
-	}
-#endif
-
-	return *this;
-}
-
-bool FRigCurveContainer::IsNameAvailable(const FString& InPotentialNewName, FString* OutErrorMessage) const
-{
-	if (Container)
-	{
-		return Container->IsNameAvailable(InPotentialNewName, ERigElementType::Curve, OutErrorMessage);
-	}
-
-	// if we are not part of a container, fall back to index based validation
-	if (InPotentialNewName.Len() > FRigHierarchyContainer::GetMaxNameLength())
-	{
-		return false;
-	}
-	return GetIndex(*InPotentialNewName) == INDEX_NONE;
-}
-
-FName FRigCurveContainer::GetSafeNewName(const FString& InPotentialNewName) const
-{
-	return FRigHierarchyContainer::GetSafeNewName(InPotentialNewName, [this](const FString& InName) -> bool {
-		return IsNameAvailable(InName);
-	});
-}
-
-FRigCurve& FRigCurveContainer::Add(const FName& InNewName, float InValue)
-{
-	DECLARE_SCOPE_HIERARCHICAL_COUNTER_FUNC()
-
-	FRigCurve NewCurve;
-	NewCurve.Name = GetSafeNewName(InNewName.ToString());
-	NewCurve.Value = InValue;
-	FName NewCurveName = NewCurve.Name;
-	Curves.Add(NewCurve);
-	RefreshMapping();
-
-#if WITH_EDITOR
-	if (!bSuspendNotifications)
-	{
-		OnCurveAdded.Broadcast(Container, NewCurve.GetElementKey());
-	}
-#endif
-
-	int32 Index = GetIndex(NewCurveName);
-	return Curves[Index];
-}
-
-FRigCurve FRigCurveContainer::Remove(const FName& InName)
-{
-	DECLARE_SCOPE_HIERARCHICAL_COUNTER_FUNC()
-
-	FRigCurve RemovedCurve;
-
-	int32 IndexToDelete = GetIndex(InName);
-	ensure(IndexToDelete != INDEX_NONE);
-	Select(InName, false);
-	RemovedCurve = Curves[IndexToDelete];
-	Curves.RemoveAt(IndexToDelete);
-	RefreshMapping();
-
-#if WITH_EDITOR
-	if (!bSuspendNotifications)
-	{
-		OnCurveRemoved.Broadcast(Container, RemovedCurve.GetElementKey());
-	}
-#endif
-
-	return RemovedCurve;
-}
-
-FName FRigCurveContainer::GetName(int32 InIndex) const
-{
-	if (Curves.IsValidIndex(InIndex))
-	{
-		return Curves[InIndex].Name;
-	}
-
-	return NAME_None;
-}
-
-int32 FRigCurveContainer::GetIndexSlow(const FName& InName) const
-{
-	for (int32 Index = 0; Index < Curves.Num(); ++Index)
-	{
-		if (Curves[Index].Name == InName)
-		{
-			return Index;
-		}
-	}
-
-	return INDEX_NONE;
-}
-
-void FRigCurveContainer::SetValue(const FName& InName, const float InValue)
-{
-	SetValue(GetIndex(InName), InValue);
-}
-
-void FRigCurveContainer::SetValue(int32 InIndex, const float InValue)
-{
-	if (Curves.IsValidIndex(InIndex))
-	{
-		FRigCurve& Curve = Curves[InIndex];
-		Curve.Value = InValue;
-	}
-}
-
-float FRigCurveContainer::GetValue(const FName& InName) const
-{
-	return GetValue(GetIndex(InName));
-}
-
-float FRigCurveContainer::GetValue(int32 InIndex) const
-{
-	if (Curves.IsValidIndex(InIndex))
-	{
-		return Curves[InIndex].Value;
-	}
-
-	return 0.f;
-}
-
-FName FRigCurveContainer::Rename(const FName& InOldName, const FName& InNewName)
-{
-	DECLARE_SCOPE_HIERARCHICAL_COUNTER_FUNC()
-
-	if (InOldName != InNewName)
-	{
-		const int32 Found = GetIndex(InOldName);
-		if (Found != INDEX_NONE)
-		{
-			FName NewName = GetSafeNewName(InNewName.ToString());
-
-			bool bWasSelected = IsSelected(InOldName);
-			if(bWasSelected)
-			{
-				Select(InOldName, false);
-			}
-
-			Curves[Found].Name = NewName;
-			RefreshMapping();
-
-#if WITH_EDITOR
-			if (!bSuspendNotifications)
-			{
-				OnCurveRenamed.Broadcast(Container, RigElementType(), InOldName, NewName);
-			}
-#endif
-			if(bWasSelected)
-			{
-				Select(NewName, true);
-			}
-			return NewName;
-		}
-	}
-	
-	return NAME_None;
-}
-
-void FRigCurveContainer::RefreshMapping()
-{
-	DECLARE_SCOPE_HIERARCHICAL_COUNTER_FUNC()
-
-	NameToIndexMapping.Empty();
-	for (int32 Index = 0; Index < Curves.Num(); ++Index)
-	{
-		Curves[Index].Index = Index;
-		NameToIndexMapping.Add(Curves[Index].Name, Index);
-	}
-}
-
-void FRigCurveContainer::Initialize()
-{
-	DECLARE_SCOPE_HIERARCHICAL_COUNTER_FUNC()
-
-	RefreshMapping();
-	ResetValues();
-}
-
-void FRigCurveContainer::Reset()
-{
-	DECLARE_SCOPE_HIERARCHICAL_COUNTER_FUNC()
-
-	Curves.Reset();
-}
-
-void FRigCurveContainer::ResetValues()
-{
-	DECLARE_SCOPE_HIERARCHICAL_COUNTER_FUNC()
-
-	// update parent index
-	for (int32 Index = 0; Index < Curves.Num(); ++Index)
-	{
-		Curves[Index].Value = 0.f;
-	}
-}
-
-bool FRigCurveContainer::Select(const FName& InName, bool bSelect)
-{
-	if(GetIndex(InName) == INDEX_NONE)
-	{
-		return false;
-	}
-
-	if(bSelect == IsSelected(InName))
-	{
-		return false;
-	}
-
-	if(bSelect)
-	{
-		Selection.Add(InName);
-	}
-	else
-	{
-		Selection.Remove(InName);
-	}
-
-	if (!bSuspendNotifications)
-	{
-		OnCurveSelected.Broadcast(Container, FRigElementKey(InName, RigElementType()), bSelect);
-	}
-
-	return true;
-}
-
-bool FRigCurveContainer::ClearSelection()
-{
-	TArray<FName> TempSelection;
-	TempSelection.Append(Selection);
-	for(const FName& SelectedName : TempSelection)
-	{
-		Select(SelectedName, false);
-	}
-	return TempSelection.Num() > 0;
-}
-
-TArray<FName> FRigCurveContainer::CurrentSelection() const
-{
-	TArray<FName> TempSelection;
-	TempSelection.Append(Selection);
-	return TempSelection;
-}
-
-bool FRigCurveContainer::IsSelected(const FName& InName) const
-{
-	return Selection.Contains(InName);
-}
-
-#if WITH_EDITOR
-
-TArray<FRigElementKey> FRigCurveContainer::ImportCurvesFromSkeleton(const USkeleton* InSkeleton, const FName& InNameSpace, bool bRemoveObsoleteCurves, bool bSelectCurves, bool bNotify)
-{
-	check(InSkeleton);
-
-	TGuardValue<bool> SuspendNotification(bSuspendNotifications, !bNotify);
-
-	TArray<FRigElementKey> Keys;
-
-	const FSmartNameMapping* SmartNameMapping = InSkeleton->GetSmartNameContainer(USkeleton::AnimCurveMappingName);
-
-	TArray<FName> NameArray;
-	SmartNameMapping->FillNameArray(NameArray);
-	for (int32 Index = 0; Index < NameArray.Num(); ++Index)
-	{
-		FName Name = NameArray[Index];
-		if (!InNameSpace.IsNone())
-		{
-			Name = *FString::Printf(TEXT("%s::%s"), *InNameSpace.ToString(), *Name.ToString());
-		}
-		Add(Name);
-		Select(Name, true);
-		Keys.Add(FRigElementKey(Name, ERigElementType::Curve));
-	}
-
-	return Keys;
-}
-
-
-#endif
-
-FRigPose FRigCurveContainer::GetPose() const
-{
-	FRigPose Pose;
-	AppendToPose(Pose);
-	return Pose;
-}
-
-void FRigCurveContainer::SetPose(FRigPose& InPose)
-{
-	for(FRigPoseElement& Element : InPose)
-	{
-		if(Element.Index.GetKey().Type == ERigElementType::Curve)
-		{
-			if(Element.Index.UpdateCache(Container))
-			{
-				Curves[Element.Index.GetIndex()].Value = Element.CurveValue;
-			}
-		}
-	}
-}
-
-void FRigCurveContainer::AppendToPose(FRigPose& InOutPose) const
-{
-	for(const FRigCurve& Curve : Curves)
-	{
-		FRigPoseElement Element;
-		if(Element.Index.UpdateCache(Curve.GetElementKey(), Container))
-		{
-			Element.CurveValue = Curve.Value;
-			InOutPose.Elements.Add(Element);
-		}
-	}
-}
\ No newline at end of file
