Control Rig: New rig hierarchy data structure refactoring

diff --git a/Engine/Plugins/Experimental/ControlRig/Source/ControlRig/Private/Rigs/RigHierarchy.cpp b/Engine/Plugins/Experimental/ControlRig/Source/ControlRig/Private/Rigs/RigHierarchy.cpp
new file mode 100644
--- /dev/null
+++ b/Engine/Plugins/Experimental/ControlRig/Source/ControlRig/Private/Rigs/RigHierarchy.cpp
@@ -0,0 +1,1904 @@
+// Copyright Epic Games, Inc. All Rights Reserved.
+
+#include "Rigs/RigHierarchy.h"
+#include "Rigs/RigHierarchyElements.h"
+#include "Rigs/RigHierarchyController.h"
+#include "Math/ControlRigMathLibrary.h"
+#include "UObject/AnimObjectVersion.h"
+#include "Serialization/MemoryReader.h"
+#include "Serialization/MemoryWriter.h"
+
+#if WITH_EDITOR
+	#include "ScopedTransaction.h"
+#endif
+
+////////////////////////////////////////////////////////////////////////////////
+// URigHierarchy
+////////////////////////////////////////////////////////////////////////////////
+
+const TArray<FRigBaseElement*> URigHierarchy::EmptyElementArray;
+
+URigHierarchy::~URigHierarchy()
+{
+	Reset();
+}
+
+void URigHierarchy::Serialize(FArchive& Ar)
+{
+	Ar.UsingCustomVersion(FAnimObjectVersion::GUID);
+
+	if (Ar.IsSaving() || Ar.IsObjectReferenceCollector())
+	{
+		Save(Ar);
+	}
+	else if (Ar.IsLoading())
+	{
+		Load(Ar);
+	}
+	else
+	{
+		checkNoEntry();
+	}
+}
+
+void URigHierarchy::Save(FArchive& Ar)
+{
+	if(Ar.IsTransacting())
+	{
+		Ar << TransformStackIndex;
+		if(bTransactingForTransformChange)
+		{
+			return;
+		}
+	}
+
+	int32 ElementCount = Elements.Num();
+	Ar << ElementCount;
+
+	for(int32 ElementIndex = 0; ElementIndex < ElementCount; ElementIndex++)
+	{
+		FRigBaseElement* Element = Elements[ElementIndex];
+
+		// store the key
+		FRigElementKey Key = Element->GetKey();
+		Ar << Key;
+
+		// allow the element to store more information
+		Element->Serialize(Ar, this, FRigBaseElement::StaticData);
+	}
+
+	for(int32 ElementIndex = 0; ElementIndex < ElementCount; ElementIndex++)
+	{
+		FRigBaseElement* Element = Elements[ElementIndex];
+		Element->Serialize(Ar, this, FRigBaseElement::InterElementData);
+	}
+}
+
+void URigHierarchy::Load(FArchive& Ar)
+{
+	if(Ar.IsTransacting())
+	{
+		Ar << TransformStackIndex;
+		if(TransformStackIndex != TransformUndoStack.Num())
+		{
+			return;
+		}
+	}
+
+	Reset();
+
+	int32 ElementCount = 0;
+	Ar << ElementCount;
+
+	for(int32 ElementIndex = 0; ElementIndex < ElementCount; ElementIndex++)
+	{
+		FRigElementKey Key;
+		Ar << Key;
+
+		FRigBaseElement* Element = nullptr;
+		switch(Key.Type)
+		{
+			case ERigElementType::Bone:
+			{
+				Element = new FRigBoneElement();
+				break;
+			}
+			case ERigElementType::Space:
+			{
+				Element = new FRigSpaceElement();
+				break;
+			}
+			case ERigElementType::Control:
+			{
+				Element = new FRigControlElement();
+				break;
+			}
+			case ERigElementType::Curve:
+			{
+				Element = new FRigCurveElement();
+				break;
+			}
+			default:
+			{
+				ensure(false);
+			}
+		}
+
+		check(Element);
+
+		Element->SubIndex = Num(Key.Type);
+		Element->Index = Elements.Add(Element);
+		IndexLookup.Add(Key, Element->Index);
+		
+		Element->Load(Ar, this, FRigBaseElement::StaticData);
+	}
+
+	TopologyVersion++;
+
+	for(int32 ElementIndex = 0; ElementIndex < ElementCount; ElementIndex++)
+	{
+		FRigBaseElement* Element = Elements[ElementIndex];
+		Element->Load(Ar, this, FRigBaseElement::InterElementData);
+	}
+
+	TopologyVersion++;
+
+	for(int32 ElementIndex = 0; ElementIndex < ElementCount; ElementIndex++)
+	{
+		if(FRigTransformElement* TransformElement = Cast<FRigTransformElement>(Elements[ElementIndex]))
+		{
+			TArray<FRigBaseElement*> CurrentParents = GetParents(TransformElement, true);
+			for (FRigBaseElement* CurrentParent : CurrentParents)
+			{
+				if(FRigTransformElement* TransformParent = Cast<FRigTransformElement>(CurrentParent))
+				{
+					TransformParent->ElementsToDirty.AddUnique(TransformElement);
+				}
+			}
+		}
+	}
+
+	TArray<int32> CountedAsChild;
+	CountedAsChild.AddZeroed(Elements.Num());
+
+	for (int32 ElementIndex = 0; ElementIndex < Elements.Num(); ElementIndex++)
+	{
+		TArray<FRigBaseElement*> Children = GetChildren(Elements[ElementIndex]);
+		for (FRigBaseElement* Child : Children)
+		{
+			CountedAsChild[Child->Index]++;
+		}
+	}
+
+	for (int32 ElementIndex = 0; ElementIndex < Elements.Num(); ElementIndex++)
+	{
+		FRigBaseElement* Element = Elements[ElementIndex];
+		if(GetNumberOfParents(Element) > 0)
+		{
+			continue;
+		}
+		
+		Traverse(Element, true, [this, &CountedAsChild](FRigBaseElement* InElement, bool& bContinue)
+		{
+
+			int32& Count = CountedAsChild[InElement->Index];
+			Count = FMath::Max(Count - 1, 0);
+
+			// if the count is 0 it means we have hit all parents
+			// and we can finally notify of our existence
+			if(Count == 0)
+			{
+				Notify(ERigHierarchyNotification::ElementAdded, InElement);
+				bContinue = true;
+			}
+			else if(Count > 0)
+			{
+				bContinue = false;
+			}
+		});
+	}
+}
+
+void URigHierarchy::Reset()
+{
+	TopologyVersion = 0;
+	bEnableDirtyPropagation = true;
+
+	for(int32 ElementIndex=0; ElementIndex<Elements.Num(); ElementIndex++)
+	{
+		delete Elements[ElementIndex];
+	}
+	Elements.Reset();
+	IndexLookup.Reset();
+
+	Notify(ERigHierarchyNotification::HierarchyReset, nullptr);
+}
+
+void URigHierarchy::CopyHierarchy(URigHierarchy* InHierarchy)
+{
+	check(InHierarchy);
+	
+	Reset();
+
+	FRigVMByteArray ArchiveBytes;
+	FMemoryWriter ArchiveWriter(ArchiveBytes);
+	InHierarchy->Save(ArchiveWriter);
+
+	FMemoryReader ArchiveReader(ArchiveBytes);
+	Load(ArchiveReader);
+}
+
+void URigHierarchy::CopyPose(URigHierarchy* InHierarchy, bool bCurrent, bool bInitial)
+{
+	check(InHierarchy);
+
+	for (int32 ElementIndex = 0; ElementIndex < Elements.Num(); ElementIndex++)
+	{
+		FRigBaseElement* Element = Elements[ElementIndex];
+		if(FRigBaseElement* OtherElement = InHierarchy->Find(Element->GetKey()))
+		{
+			Element->CopyPose(OtherElement, bCurrent, bInitial);
+		}
+	}
+}
+
+void URigHierarchy::ResetPoseToInitial(ERigElementType InTypeFilter)
+{
+	const bool bAppliesToAllTypes = InTypeFilter == ERigElementType::All;
+	for(int32 ElementIndex=0; ElementIndex<Elements.Num(); ElementIndex++)
+	{
+		if(!bAppliesToAllTypes)
+		{
+			if(!Elements[ElementIndex]->IsTypeOf(InTypeFilter))
+			{
+				continue;
+			}
+		}
+
+		if(FRigControlElement* ControlElement = Cast<FRigControlElement>(Elements[ElementIndex]))
+		{
+			if(!bAppliesToAllTypes)
+			{
+				const FTransform OffsetTransform = GetControlOffsetTransform(ControlElement, ERigTransformType::InitialLocal);
+				SetControlOffsetTransform(ControlElement, OffsetTransform, ERigTransformType::CurrentLocal, true);
+				const FTransform GizmoTransform = GetControlGizmoTransform(ControlElement, ERigTransformType::InitialLocal);
+				SetControlGizmoTransform(ControlElement, GizmoTransform, ERigTransformType::CurrentLocal, true);
+			}
+			else
+			{
+				ControlElement->Offset.Current = ControlElement->Offset.Initial;
+				ControlElement->Gizmo.Current = ControlElement->Gizmo.Initial;
+			}
+		}
+
+		if(FRigTransformElement* TransformElement = Cast<FRigTransformElement>(Elements[ElementIndex]))
+		{
+			if(!bAppliesToAllTypes)
+			{
+				const FTransform Transform = GetTransform(TransformElement, ERigTransformType::InitialLocal);
+				SetTransform(TransformElement, Transform, ERigTransformType::CurrentLocal, true);
+			}
+			else
+			{
+				TransformElement->Pose.Current = TransformElement->Pose.Initial;
+			}
+		}
+
+		if(FRigMultiParentElement* MultiParentElement = Cast<FRigMultiParentElement>(Elements[ElementIndex]))
+		{
+			if(!bAppliesToAllTypes)
+			{
+				MultiParentElement->Parent.MarkDirty(ERigTransformType::CurrentGlobal);
+			}
+			else
+			{
+				MultiParentElement->Parent.Current = MultiParentElement->Parent.Initial;
+			}
+		}
+	}
+}
+
+void URigHierarchy::ResetCurveValues()
+{
+	for(int32 ElementIndex=0; ElementIndex<Elements.Num(); ElementIndex++)
+	{
+		if(FRigCurveElement* CurveElement = Cast<FRigCurveElement>(Elements[ElementIndex]))
+		{
+			SetCurveValue(CurveElement, 0.f);
+		}
+	}
+}
+
+int32 URigHierarchy::Num(ERigElementType InElementType) const
+{
+	int32 Count = 0;
+	for (int32 ElementIndex = 0; ElementIndex < Elements.Num(); ElementIndex++)
+	{
+		FRigBaseElement* Element = Elements[ElementIndex];
+		if(Element->IsTypeOf(InElementType))
+		{
+			Count++;
+		}
+	}
+	return Count;
+}
+
+TArray<FRigBaseElement*> URigHierarchy::GetSelectedElements(ERigElementType InTypeFilter) const
+{
+	TArray<FRigBaseElement*> Selection;
+
+	if(URigHierarchy* HierarchyForSelection = HierarchyForSelectionPtr.Get())
+	{
+		TArray<FRigElementKey> SelectedKeys = HierarchyForSelection->GetSelectedKeys(InTypeFilter);
+		for(const FRigElementKey& SelectedKey : SelectedKeys)
+		{
+			if(const FRigBaseElement* Element = Find(SelectedKey))
+			{
+				Selection.Add((FRigBaseElement*)Element);
+			}
+		}
+		return Selection;
+	}
+
+	for (int32 ElementIndex = 0; ElementIndex < Elements.Num(); ElementIndex++)
+	{
+		FRigBaseElement* Element = Elements[ElementIndex];
+		if(Element->IsTypeOf(InTypeFilter))
+		{
+			if(IsSelected(Element))
+			{
+				Selection.Add(Element);
+			}
+		}
+	}
+	return Selection;
+}
+
+TArray<FRigElementKey> URigHierarchy::GetSelectedKeys(ERigElementType InTypeFilter) const
+{
+	if(URigHierarchy* HierarchyForSelection = HierarchyForSelectionPtr.Get())
+	{
+		return HierarchyForSelection->GetSelectedKeys(InTypeFilter);
+	}
+
+	TArray<FRigElementKey> Selection;
+	for (int32 ElementIndex = 0; ElementIndex < Elements.Num(); ElementIndex++)
+	{
+		FRigBaseElement* Element = Elements[ElementIndex];
+		if(Element->IsTypeOf(InTypeFilter))
+		{
+			if(IsSelected(Element))
+			{
+				Selection.Add(Element->GetKey());
+			}
+		}
+	}
+	return Selection;
+}
+
+void URigHierarchy::SanitizeName(FString& InOutName)
+{
+	// Sanitize the name
+	for (int32 i = 0; i < InOutName.Len(); ++i)
+	{
+		TCHAR& C = InOutName[i];
+
+		const bool bGoodChar =
+			((C >= 'A') && (C <= 'Z')) || ((C >= 'a') && (C <= 'z')) ||		// A-Z (upper and lowercase) anytime
+			(C == '_') || (C == '-') || (C == '.') ||						// _  - . anytime
+			((i > 0) && (C >= '0') && (C <= '9'));							// 0-9 after the first character
+
+		if (!bGoodChar)
+		{
+			C = '_';
+		}
+	}
+
+	if (InOutName.Len() > GetMaxNameLength())
+	{
+		InOutName.LeftChopInline(InOutName.Len() - GetMaxNameLength());
+	}
+}
+
+FName URigHierarchy::GetSanitizedName(const FString& InName)
+{
+	FString Name = InName;
+	SanitizeName(Name);
+
+	if (Name.IsEmpty())
+	{
+		return NAME_None;
+	}
+
+	return *Name;
+}
+
+bool URigHierarchy::IsNameAvailable(const FString& InPotentialNewName, ERigElementType InType, FString* OutErrorMessage) const
+{
+	FString UnsanitizedName = InPotentialNewName;
+	if (UnsanitizedName.Len() > GetMaxNameLength())
+	{
+		if (OutErrorMessage)
+		{
+			*OutErrorMessage = TEXT("Name too long.");
+		}
+		return false;
+	}
+
+	FString SanitizedName = UnsanitizedName;
+	SanitizeName(SanitizedName);
+
+	if (SanitizedName != UnsanitizedName)
+	{
+		if (OutErrorMessage)
+		{
+			*OutErrorMessage = TEXT("Name contains invalid characters.");
+		}
+		return false;
+	}
+
+	if (GetIndex(FRigElementKey(*InPotentialNewName, InType)) != INDEX_NONE)
+	{
+		if (OutErrorMessage)
+		{
+			*OutErrorMessage = TEXT("Name already used.");
+		}
+		return false;
+	}
+
+	return true;
+}
+
+FName URigHierarchy::GetSafeNewName(const FString& InPotentialNewName, ERigElementType InType) const
+{
+	FString SanitizedName = InPotentialNewName;
+	SanitizeName(SanitizedName);
+	FString Name = SanitizedName;
+
+	int32 Suffix = 1;
+	while (!IsNameAvailable(Name, InType))
+	{
+		FString BaseString = SanitizedName;
+		if (BaseString.Len() > GetMaxNameLength() - 4)
+		{
+			BaseString.LeftChopInline(BaseString.Len() - (GetMaxNameLength() - 4));
+		}
+		Name = *FString::Printf(TEXT("%s_%d"), *BaseString, ++Suffix);
+	}
+	return *Name;
+}
+
+TArray<FRigElementKey> URigHierarchy::GetChildren(FRigElementKey InKey, bool bRecursive) const
+{
+	TArray<FRigBaseElement*> LocalChildren;
+	const TArray<FRigBaseElement*>* ChildrenPtr = nullptr;
+	if(bRecursive)
+	{
+		LocalChildren = GetChildren(Find(InKey), true);
+		ChildrenPtr = & LocalChildren;
+	}
+	else
+	{
+		ChildrenPtr = &GetChildren(Find(InKey));
+	}
+
+	const TArray<FRigBaseElement*>& Children = *ChildrenPtr;
+
+	TArray<FRigElementKey> Keys;
+	for(const FRigBaseElement* Child : Children)
+	{
+		Keys.Add(Child->Key);
+	}
+	return Keys;
+}
+
+TArray<int32> URigHierarchy::GetChildren(int32 InIndex, bool bRecursive) const
+{
+	TArray<FRigBaseElement*> LocalChildren;
+	const TArray<FRigBaseElement*>* ChildrenPtr = nullptr;
+	if(bRecursive)
+	{
+		LocalChildren = GetChildren(Get(InIndex), true);
+		ChildrenPtr = & LocalChildren;
+	}
+	else
+	{
+		ChildrenPtr = &GetChildren(Get(InIndex));
+	}
+
+	const TArray<FRigBaseElement*>& Children = *ChildrenPtr;
+
+	TArray<int32> Indices;
+	for(const FRigBaseElement* Child : Children)
+	{
+		Indices.Add(Child->Index);
+	}
+	return Indices;
+}
+
+const TArray<FRigBaseElement*>& URigHierarchy::GetChildren(const FRigBaseElement* InElement) const
+{
+	if(InElement)
+	{
+		UpdateCachedChildren(InElement);
+		return InElement->CachedChildren;
+	}
+	return EmptyElementArray;
+}
+
+TArray<FRigBaseElement*> URigHierarchy::GetChildren(const FRigBaseElement* InElement, bool bRecursive) const
+{
+	// call the non-recursive variation
+	TArray<FRigBaseElement*> Children = GetChildren(InElement);
+	
+	if(bRecursive)
+	{
+		for(int32 ChildIndex = 0; ChildIndex < Children.Num(); ChildIndex++)
+		{
+			Children.Append(GetChildren(Children[ChildIndex], true));
+		}
+	}
+
+	return Children;
+}
+
+TArray<FRigElementKey> URigHierarchy::GetParents(FRigElementKey InKey, bool bRecursive) const
+{
+	const TArray<FRigBaseElement*>& Parents = GetParents(Find(InKey), bRecursive);
+	TArray<FRigElementKey> Keys;
+	for(const FRigBaseElement* Parent : Parents)
+	{
+		Keys.Add(Parent->Key);
+	}
+	return Keys;
+}
+
+TArray<int32> URigHierarchy::GetParents(int32 InIndex, bool bRecursive) const
+{
+	const TArray<FRigBaseElement*>& Parents = GetParents(Get(InIndex), bRecursive);
+	TArray<int32> Indices;
+	for(const FRigBaseElement* Parent : Parents)
+	{
+		Indices.Add(Parent->Index);
+	}
+	return Indices;
+}
+
+TArray<FRigBaseElement*> URigHierarchy::GetParents(const FRigBaseElement* InElement, bool bRecursive) const
+{
+	TArray<FRigBaseElement*> Parents;
+
+	if(const FRigSingleParentElement* SingleParentElement = Cast<FRigSingleParentElement>(InElement))
+	{
+		if(SingleParentElement->ParentElement)
+		{
+			Parents.AddUnique(SingleParentElement->ParentElement);
+		}
+	}
+	else if(const FRigMultiParentElement* MultiParentElement = Cast<FRigMultiParentElement>(InElement))
+	{
+		for(FRigTransformElement* ParentElement : MultiParentElement->ParentElements)
+		{
+			Parents.AddUnique(ParentElement);
+		}
+	}
+
+	if(bRecursive)
+	{
+		const int32 CurrentNumberParents = Parents.Num();
+		for(int32 ParentIndex = 0;ParentIndex < CurrentNumberParents; ParentIndex++)
+		{
+			const TArray<FRigBaseElement*> GrandParents = GetParents(Parents[ParentIndex], bRecursive);
+			for (FRigBaseElement* GrandParent : GrandParents)
+			{
+				Parents.AddUnique(GrandParent);
+			}
+		}
+	}
+
+	return Parents;
+}
+
+FRigElementKey URigHierarchy::GetFirstParent(FRigElementKey InKey) const
+{
+	if(FRigBaseElement* FirstParent = GetFirstParent(Find(InKey)))
+	{
+		return FirstParent->Key;
+	}
+	return FRigElementKey();
+}
+
+int32 URigHierarchy::GetFirstParent(int32 InIndex) const
+{
+	if(FRigBaseElement* FirstParent = GetFirstParent(Get(InIndex)))
+	{
+		return FirstParent->Index;
+	}
+	return INDEX_NONE;
+}
+
+FRigBaseElement* URigHierarchy::GetFirstParent(const FRigBaseElement* InElement) const
+{
+	if(const FRigSingleParentElement* SingleParentElement = Cast<FRigSingleParentElement>(InElement))
+	{
+		return SingleParentElement->ParentElement;
+	}
+	else if(const FRigMultiParentElement* MultiParentElement = Cast<FRigMultiParentElement>(InElement))
+	{
+		if(MultiParentElement->ParentElements.Num() > 0)
+		{
+			return MultiParentElement->ParentElements[0];
+		}
+	}
+	
+	return nullptr;
+}
+
+int32 URigHierarchy::GetNumberOfParents(FRigElementKey InKey) const
+{
+	return GetNumberOfParents(Find(InKey));
+}
+
+int32 URigHierarchy::GetNumberOfParents(int32 InIndex) const
+{
+	return GetNumberOfParents(Get(InIndex));
+}
+
+int32 URigHierarchy::GetNumberOfParents(const FRigBaseElement* InElement) const
+{
+	if(InElement == nullptr)
+	{
+		return 0;
+	}
+
+	if(const FRigSingleParentElement* SingleParentElement = Cast<FRigSingleParentElement>(InElement))
+	{
+		return SingleParentElement->ParentElement == nullptr ? 0 : 1;
+	}
+	else if(const FRigMultiParentElement* MultiParentElement = Cast<FRigMultiParentElement>(InElement))
+	{
+		return MultiParentElement->ParentElements.Num();
+	}
+
+	return 0;
+}
+
+float URigHierarchy::GetParentWeight(FRigElementKey InChild, FRigElementKey InParent, bool bInitial) const
+{
+	return GetParentWeight(Find(InChild), Find(InParent), bInitial);
+}
+
+float URigHierarchy::GetParentWeight(const FRigBaseElement* InChild, const FRigBaseElement* InParent, bool bInitial) const
+{
+	if(const FRigMultiParentElement* MultiParentElement = Cast<FRigMultiParentElement>(InChild))
+	{
+		if(const int32* ParentIndexPtr = MultiParentElement->IndexLookup.Find(InParent->GetKey()))
+		{
+			return GetParentWeight(InChild, *ParentIndexPtr, bInitial);
+		}
+	}
+	return FLT_MAX;
+}
+
+float URigHierarchy::GetParentWeight(const FRigBaseElement* InChild, int32 InParentIndex, bool bInitial) const
+{
+	if(const FRigMultiParentElement* MultiParentElement = Cast<FRigMultiParentElement>(InChild))
+	{
+		if(MultiParentElement->ParentWeights.IsValidIndex(InParentIndex))
+		{
+			if(bInitial)
+			{
+				return MultiParentElement->ParentWeightsInitial[InParentIndex];
+			}
+			else
+			{
+				return MultiParentElement->ParentWeights[InParentIndex];
+			}
+		}
+	}
+	return FLT_MAX;
+}
+
+bool URigHierarchy::SetParentWeight(FRigElementKey InChild, FRigElementKey InParent, float InWeight, bool bInitial, bool bAffectChildren)
+{
+	return SetParentWeight(Find(InChild), Find(InParent), InWeight, bInitial, bAffectChildren);
+}
+
+bool URigHierarchy::SetParentWeight(FRigBaseElement* InChild, const FRigBaseElement* InParent, float InWeight, bool bInitial, bool bAffectChildren)
+{
+	if(const FRigMultiParentElement* MultiParentElement = Cast<FRigMultiParentElement>(InChild))
+	{
+		if(const int32* ParentIndexPtr = MultiParentElement->IndexLookup.Find(InParent->GetKey()))
+		{
+			return SetParentWeight(InChild, *ParentIndexPtr, InWeight, bInitial, bAffectChildren);
+		}
+	}
+	return false;
+}
+
+bool URigHierarchy::SetParentWeight(FRigBaseElement* InChild, int32 InParentIndex, float InWeight, bool bInitial, bool bAffectChildren)
+{
+	using namespace ERigTransformType;
+
+	if(FRigMultiParentElement* MultiParentElement = Cast<FRigMultiParentElement>(InChild))
+	{
+		if(MultiParentElement->ParentWeights.IsValidIndex(InParentIndex))
+		{
+			const float InputWeight = FMath::Max(InWeight, 0.f);
+
+			float& TargetWeight = bInitial?
+				MultiParentElement->ParentWeightsInitial[InParentIndex] :
+			MultiParentElement->ParentWeights[InParentIndex];
+
+			if(FMath::IsNearlyZero(InputWeight - TargetWeight))
+			{
+				return false;
+			}
+			
+			const ERigTransformType::Type LocalType = bInitial ? InitialLocal : CurrentLocal;
+			const ERigTransformType::Type GlobalType = SwapLocalAndGlobal(LocalType);
+
+			if(bAffectChildren)
+			{
+				GetTransform(MultiParentElement, LocalType);
+				MultiParentElement->Pose.MarkDirty(GlobalType);
+			}
+			else
+			{
+				GetTransform(MultiParentElement, GlobalType);
+				MultiParentElement->Pose.MarkDirty(LocalType);
+			}
+
+			TargetWeight = InputWeight;
+			MultiParentElement->Parent.MarkDirty(GlobalType);
+
+			PropagateDirtyFlags(MultiParentElement, LocalType, bAffectChildren);
+			return true;
+		}
+	}
+	return false;
+}
+
+TArray<FRigElementKey> URigHierarchy::GetAllKeys(bool bTraverse, ERigElementType InElementType) const
+{
+	TArray<FRigElementKey> Keys;
+	Keys.Reserve(Elements.Num());
+
+	if(bTraverse)
+	{
+		TArray<bool> ElementVisited;
+		ElementVisited.AddZeroed(Elements.Num());
+		
+		for (int32 ElementIndex = 0; ElementIndex < Elements.Num(); ElementIndex++)
+		{
+			FRigBaseElement* Element = Elements[ElementIndex];
+			Traverse(Element, true, [&ElementVisited, &Keys, InElementType](FRigBaseElement* InElement, bool& bContinue)
+            {
+				bContinue = !ElementVisited[InElement->GetIndex()];
+
+				if(bContinue)
+				{
+					if(InElement->IsTypeOf(InElementType))
+					{
+						Keys.Add(InElement->GetKey());
+					}
+					ElementVisited[InElement->GetIndex()] = true;
+				}
+            });
+		}
+	}
+	else
+	{
+		for (int32 ElementIndex = 0; ElementIndex < Elements.Num(); ElementIndex++)
+		{
+			FRigBaseElement* Element = Elements[ElementIndex];
+			if(Element->IsTypeOf(InElementType))
+			{
+				Keys.Add(Element->GetKey());
+			}
+		}
+	}
+	return Keys;
+}
+
+void URigHierarchy::Traverse(FRigBaseElement* InElement, bool bTowardsChildren,
+                             TFunction<void(FRigBaseElement*, bool&)> PerElementFunction) const
+{
+	bool bContinue = true;
+	PerElementFunction(InElement, bContinue);
+
+	if(bContinue)
+	{
+		if(bTowardsChildren)
+		{
+			const TArray<FRigBaseElement*>& Children = GetChildren(InElement);
+			for (FRigBaseElement* Child : Children)
+			{
+				Traverse(Child, true, PerElementFunction);
+			}
+		}
+		else
+		{
+			TArray<FRigBaseElement*> Parents = GetParents(InElement);
+			for (FRigBaseElement* Parent : Parents)
+			{
+				Traverse(Parent, false, PerElementFunction);
+			}
+		}
+	}
+}
+
+void URigHierarchy::Traverse(TFunction<void(FRigBaseElement*, bool& /* continue */)> PerElementFunction, bool bTowardsChildren) const
+{
+	if(bTowardsChildren)
+	{
+		for (int32 ElementIndex = 0; ElementIndex < Elements.Num(); ElementIndex++)
+		{
+			FRigBaseElement* Element = Elements[ElementIndex];
+			if(GetNumberOfParents(Element) == 0)
+			{
+				Traverse(Element, bTowardsChildren, PerElementFunction);
+			}
+        }
+	}
+	else
+	{
+		for (int32 ElementIndex = 0; ElementIndex < Elements.Num(); ElementIndex++)
+		{
+			FRigBaseElement* Element = Elements[ElementIndex];
+			if(GetChildren(Element).Num() == 0)
+			{
+				Traverse(Element, bTowardsChildren, PerElementFunction);
+			}
+		}
+	}
+}
+
+bool URigHierarchy::Undo()
+{
+#if WITH_EDITOR
+	
+	if(TransformUndoStack.IsEmpty())
+	{
+		return false;
+	}
+
+	const FTransformStackEntry Entry = TransformUndoStack.Pop();
+	ApplyTransformFromStack(Entry, true);
+	TransformRedoStack.Push(Entry);
+	TransformStackIndex = TransformUndoStack.Num();
+	return true;
+	
+#else
+	
+	return false;
+	
+#endif
+}
+
+bool URigHierarchy::Redo()
+{
+#if WITH_EDITOR
+
+	if(TransformRedoStack.IsEmpty())
+	{
+		return false;
+	}
+
+	const FTransformStackEntry Entry = TransformRedoStack.Pop();
+	ApplyTransformFromStack(Entry, false);
+	TransformUndoStack.Push(Entry);
+	TransformStackIndex = TransformUndoStack.Num();
+	return true;
+	
+#else
+	
+	return false;
+	
+#endif
+}
+
+bool URigHierarchy::SetTransformStackIndex(int32 InTransformStackIndex)
+{
+#if WITH_EDITOR
+
+	while(TransformUndoStack.Num() > InTransformStackIndex)
+	{
+		if(TransformUndoStack.Num() == 0)
+		{
+			return false;
+		}
+
+		if(!Undo())
+		{
+			return false;
+		}
+	}
+	
+	while(TransformUndoStack.Num() < InTransformStackIndex)
+	{
+		if(TransformRedoStack.Num() == 0)
+		{
+			return false;
+		}
+
+		if(!Redo())
+		{
+			return false;
+		}
+	}
+
+	return InTransformStackIndex == TransformStackIndex;
+
+#else
+	
+	return false;
+	
+#endif
+}
+
+#if WITH_EDITOR
+
+void URigHierarchy::PostTransacted(const FTransactionObjectEvent& TransactionEvent)
+{
+	if (TransactionEvent.GetEventType() == ETransactionObjectEventType::UndoRedo)
+	{
+		const int32 DesiredStackIndex = TransformStackIndex;
+		TransformStackIndex = TransformUndoStack.Num();
+		if (DesiredStackIndex == TransformStackIndex)
+		{
+			return;
+		}
+		SetTransformStackIndex(DesiredStackIndex);
+	}
+}
+
+#endif
+
+void URigHierarchy::SendEvent(const FRigEventContext& InEvent, bool bAsynchronous)
+{
+	if(EventDelegate.IsBound())
+	{
+		URigHierarchy* LocalHierarchy = this;
+		FRigEventDelegate& Delegate = EventDelegate;
+
+		if (bAsynchronous)
+		{
+			FFunctionGraphTask::CreateAndDispatchWhenReady([LocalHierarchy, Delegate, InEvent]()
+            {
+                Delegate.Broadcast(LocalHierarchy, InEvent);
+            }, TStatId(), NULL, ENamedThreads::GameThread);
+		}
+		else
+		{
+			Delegate.Broadcast(LocalHierarchy, InEvent);
+		}
+	}
+
+}
+
+URigHierarchyController* URigHierarchy::GetController(bool bCreateIfNeeded)
+{
+	if(LastControllerPtr.IsValid())
+	{
+		return Cast<URigHierarchyController>(LastControllerPtr.Get());
+	}
+	else if(bCreateIfNeeded)
+	{
+		if(UObject* Outer = GetOuter())
+		{
+			URigHierarchyController* Controller = NewObject<URigHierarchyController>(Outer);
+			Controller->SetHierarchy(this);
+			LastControllerPtr = Controller;
+			return Controller;
+		}
+	}
+	return nullptr;
+}
+
+FRigPose URigHierarchy::GetPose(bool bInitial) const
+{
+	FRigPose Pose;
+
+	for (int32 ElementIndex = 0; ElementIndex < Elements.Num(); ElementIndex++)
+	{
+		FRigBaseElement* Element = Elements[ElementIndex];
+		FRigPoseElement PoseElement;
+		PoseElement.Index.UpdateCache(Element->GetKey(), this);
+		
+		if(FRigTransformElement* TransformElement = Cast<FRigTransformElement>(Element))
+		{
+			PoseElement.LocalTransform = GetTransform(TransformElement, bInitial ? ERigTransformType::InitialLocal : ERigTransformType::CurrentLocal);
+			PoseElement.GlobalTransform = GetTransform(TransformElement, bInitial ? ERigTransformType::InitialGlobal : ERigTransformType::CurrentGlobal);
+		}
+		else if(FRigCurveElement* CurveElement = Cast<FRigCurveElement>(Element))
+		{
+			PoseElement.CurveValue = GetCurveValue(CurveElement);
+		}
+		else
+		{
+			continue;
+		}
+		Pose.Elements.Add(PoseElement);
+	}
+	return Pose;
+}
+
+void URigHierarchy::SetPose(const FRigPose& InPose, ERigTransformType::Type InTransformType)
+{
+	for(const FRigPoseElement& PoseElement : InPose)
+	{
+		FCachedRigElement Index = PoseElement.Index;
+		if(Index.UpdateCache(this))
+		{
+			FRigBaseElement* Element = Get(Index.GetIndex());
+			if(FRigTransformElement* TransformElement = Cast<FRigTransformElement>(Element))
+			{
+				if(ERigTransformType::IsLocal(InTransformType))
+				{
+					SetTransform(TransformElement, PoseElement.LocalTransform, InTransformType, true);
+				}
+				else
+				{
+					SetTransform(TransformElement, PoseElement.GlobalTransform, InTransformType, true);
+				}
+			}
+			else if(FRigCurveElement* CurveElement = Cast<FRigCurveElement>(Element))
+			{
+				SetCurveValue(CurveElement, PoseElement.CurveValue);
+			}
+		}
+	}
+}
+
+void URigHierarchy::Notify(ERigHierarchyNotification InNotifType, const FRigBaseElement* InElement)
+{
+	ModifiedEvent.Broadcast(InNotifType, this, InElement);
+}
+
+FTransform URigHierarchy::GetTransform(FRigTransformElement* InTransformElement,
+	const ERigTransformType::Type InTransformType) const
+{
+	if(InTransformElement == nullptr)
+	{
+		return FTransform::Identity;
+	}
+	
+	if(InTransformElement->Pose.IsDirty(InTransformType))
+	{
+		const ERigTransformType::Type OpposedType = SwapLocalAndGlobal(InTransformType);
+		const ERigTransformType::Type GlobalType = MakeGlobal(InTransformType);
+		ensure(!InTransformElement->Pose.IsDirty(OpposedType));
+
+		FTransform ParentTransform;
+		if(FRigControlElement* ControlElement = Cast<FRigControlElement>(InTransformElement))
+		{
+			ParentTransform = GetControlOffsetTransform(ControlElement, GlobalType);
+		}
+		else
+		{
+			ParentTransform = GetParentTransform(InTransformElement, GlobalType);
+		}
+
+		if(IsLocal(InTransformType))
+		{
+			FTransform NewTransform = InTransformElement->Pose.Get(OpposedType).GetRelativeTransform(ParentTransform);
+			NewTransform.NormalizeRotation();
+			InTransformElement->Pose.Set(InTransformType, NewTransform);
+		}
+		else
+		{
+			FTransform NewTransform = InTransformElement->Pose.Get(OpposedType) * ParentTransform;
+			NewTransform.NormalizeRotation();
+			InTransformElement->Pose.Set(InTransformType, NewTransform);
+		}
+	}
+	return InTransformElement->Pose.Get(InTransformType);
+}
+
+void URigHierarchy::SetTransform(FRigTransformElement* InTransformElement, const FTransform& InTransform, const ERigTransformType::Type InTransformType, bool bAffectChildren, bool bSetupUndo, bool bForce)
+{
+	if(InTransformElement == nullptr)
+	{
+		return;
+	}
+
+	if(IsGlobal(InTransformType))
+	{
+		if(FRigControlElement* ControlElement = Cast<FRigControlElement>(InTransformElement))
+		{
+			const FTransform OffsetTransform = GetControlOffsetTransform(ControlElement, InTransformType);
+			FTransform LocalTransform = InTransform.GetRelativeTransform(OffsetTransform);
+			
+			ControlElement->Settings.ApplyLimits(LocalTransform);
+			SetTransform(ControlElement, LocalTransform, MakeLocal(InTransformType), bAffectChildren);
+			return;
+		}
+	}
+
+	if(!InTransformElement->Pose.IsDirty(InTransformType))
+	{
+		const FTransform PreviousTransform = InTransformElement->Pose.Get(InTransformType);
+		if(!bForce && FRigComputedTransform::Equals(PreviousTransform, InTransform))
+		{
+			return;
+		}
+	}
+
+	const FTransform PreviousTransform = GetTransform(InTransformElement, InTransformType);
+	PropagateDirtyFlags(InTransformElement, InTransformType, bAffectChildren);
+
+	const ERigTransformType::Type OpposedType = SwapLocalAndGlobal(InTransformType);
+	InTransformElement->Pose.Set(InTransformType, InTransform);
+	InTransformElement->Pose.MarkDirty(OpposedType);
+
+	if(FRigControlElement* ControlElement = Cast<FRigControlElement>(InTransformElement))
+	{
+		ControlElement->Gizmo.MarkDirty(MakeGlobal(InTransformType));
+	}
+
+#if WITH_EDITOR
+	if(bSetupUndo)
+	{
+		PushTransformToStack(
+			InTransformElement->GetKey(),
+			ETransformStackEntryType::TransformPose,
+			InTransformType,
+			PreviousTransform,
+			InTransformElement->Pose.Get(InTransformType),
+			bAffectChildren);
+	}
+#endif
+}
+
+FTransform URigHierarchy::GetControlOffsetTransform(FRigControlElement* InControlElement,
+	const ERigTransformType::Type InTransformType) const
+{
+	if(InControlElement == nullptr)
+	{
+		return FTransform::Identity;
+	}
+	
+	if(InControlElement->Offset.IsDirty(InTransformType))
+	{
+		const ERigTransformType::Type OpposedType = SwapLocalAndGlobal(InTransformType);
+		const ERigTransformType::Type GlobalType = MakeGlobal(InTransformType);
+		ensure(!InControlElement->Offset.IsDirty(OpposedType));
+
+		const FTransform ParentTransform =  GetParentTransform(InControlElement, GlobalType);
+		if(IsLocal(InTransformType))
+		{
+			InControlElement->Offset.Set(InTransformType, InControlElement->Offset.Get(OpposedType).GetRelativeTransform(ParentTransform));
+		}
+		else
+		{
+			InControlElement->Offset.Set(InTransformType, InControlElement->Offset.Get(OpposedType) * ParentTransform);
+		}
+	}
+	return InControlElement->Offset.Get(InTransformType);
+}
+
+void URigHierarchy::SetControlOffsetTransform(FRigControlElement* InControlElement, const FTransform& InTransform,
+	const ERigTransformType::Type InTransformType, bool bAffectChildren, bool bSetupUndo, bool bForce)
+{
+	if(InControlElement == nullptr)
+	{
+		return;
+	}
+
+	if(!InControlElement->Offset.IsDirty(InTransformType))
+	{
+		const FTransform PreviousTransform = InControlElement->Offset.Get(InTransformType);
+		if(!bForce && FRigComputedTransform::Equals(PreviousTransform, InTransform))
+		{
+			return;
+		}
+	}
+
+	const FTransform PreviousTransform = GetControlOffsetTransform(InControlElement, InTransformType);
+	PropagateDirtyFlags(InControlElement, InTransformType, bAffectChildren);
+
+	GetTransform(InControlElement, MakeLocal(InTransformType));
+	InControlElement->Pose.MarkDirty(MakeGlobal(InTransformType));
+
+	const ERigTransformType::Type OpposedType = SwapLocalAndGlobal(InTransformType);
+	InControlElement->Offset.Set(InTransformType, InTransform);
+	InControlElement->Offset.MarkDirty(OpposedType);
+	InControlElement->Gizmo.MarkDirty(MakeGlobal(InTransformType));
+
+#if WITH_EDITOR
+	if(bSetupUndo)
+	{
+		PushTransformToStack(
+            InControlElement->GetKey(),
+            ETransformStackEntryType::ControlOffset,
+            InTransformType,
+            PreviousTransform,
+            InControlElement->Offset.Get(InTransformType),
+            bAffectChildren);
+	}
+#endif
+}
+
+FTransform URigHierarchy::GetControlGizmoTransform(FRigControlElement* InControlElement,
+	const ERigTransformType::Type InTransformType) const
+{
+	if(InControlElement == nullptr)
+	{
+		return FTransform::Identity;
+	}
+	
+	if(InControlElement->Gizmo.IsDirty(InTransformType))
+	{
+		const ERigTransformType::Type OpposedType = SwapLocalAndGlobal(InTransformType);
+		const ERigTransformType::Type GlobalType = MakeGlobal(InTransformType);
+		ensure(!InControlElement->Gizmo.IsDirty(OpposedType));
+
+		const FTransform ParentTransform = GetTransform(InControlElement, GlobalType);
+		if(IsLocal(InTransformType))
+		{
+			InControlElement->Gizmo.Set(InTransformType, InControlElement->Gizmo.Get(OpposedType).GetRelativeTransform(ParentTransform));
+		}
+		else
+		{
+			InControlElement->Gizmo.Set(InTransformType, InControlElement->Gizmo.Get(OpposedType) * ParentTransform);
+		}
+	}
+	return InControlElement->Gizmo.Get(InTransformType);
+}
+
+void URigHierarchy::SetControlGizmoTransform(FRigControlElement* InControlElement, const FTransform& InTransform,
+	const ERigTransformType::Type InTransformType, bool bSetupUndo, bool bForce)
+{
+	if(InControlElement == nullptr)
+	{
+		return;
+	}
+
+	if(!InControlElement->Gizmo.IsDirty(InTransformType))
+	{
+		const FTransform PreviousTransform = InControlElement->Gizmo.Get(InTransformType);
+		if(!bForce && FRigComputedTransform::Equals(PreviousTransform, InTransform))
+		{
+			return;
+		}
+	}
+
+	const FTransform PreviousTransform = GetControlGizmoTransform(InControlElement, InTransformType);
+	const ERigTransformType::Type OpposedType = SwapLocalAndGlobal(InTransformType);
+	InControlElement->Gizmo.Set(InTransformType, InTransform);
+	InControlElement->Gizmo.MarkDirty(OpposedType);
+
+#if WITH_EDITOR
+	if(bSetupUndo)
+	{
+		PushTransformToStack(
+            InControlElement->GetKey(),
+            ETransformStackEntryType::ControlGizmo,
+            InTransformType,
+            PreviousTransform,
+            InControlElement->Gizmo.Get(InTransformType),
+            false);
+	}
+#endif
+
+	Notify(ERigHierarchyNotification::ControlSettingChanged, InControlElement);
+}
+
+FTransform URigHierarchy::GetParentTransform(FRigBaseElement* InElement, const ERigTransformType::Type InTransformType) const
+{
+	if(FRigSingleParentElement* SingleParentElement = Cast<FRigSingleParentElement>(InElement))
+	{
+		return GetTransform(SingleParentElement->ParentElement, InTransformType);
+	}
+	else if(FRigMultiParentElement* MultiParentElement = Cast<FRigMultiParentElement>(InElement))
+	{
+		FRigComputedTransform& Output = MultiParentElement->Parent[InTransformType];
+
+		if(Output.bDirty)
+		{
+			Output.Set(FTransform::Identity);
+
+			const bool bInitial = IsInitial(InTransformType);
+			const TArray<float>& ParentWeights = bInitial ? MultiParentElement->ParentWeightsInitial : MultiParentElement->ParentWeights;
+			const TArray<FRigTransformElement*>& ParentElements = MultiParentElement->ParentElements;
+
+			// if we have only one parent we behave like a single parent element
+			if(ParentElements.Num() == 1)
+			{
+				Output.Set(GetTransform(ParentElements[0], InTransformType));
+			}
+			// for exactly two parents we'll lerp
+			else if(ParentElements.Num() == 2)
+			{
+				float Weight = 0.f;
+				const float WeightA = ParentWeights[0];
+				const float WeightB = ParentWeights[1];
+				const float ClampedWeightA = FMath::Max(WeightA, 0.f);
+				const float ClampedWeightB = FMath::Max(WeightB, 0.f);
+				const float OverallWeight = ClampedWeightA + ClampedWeightB;
+				if(OverallWeight > SMALL_NUMBER)
+				{
+					Weight = ClampedWeightB / OverallWeight;
+				}
+
+				if(Weight <= SMALL_NUMBER)
+				{
+					Output.Set(GetTransform(ParentElements[0], InTransformType));
+				}
+				else if(Weight >= (1.0f - SMALL_NUMBER))
+				{
+					Output.Set(GetTransform(ParentElements[1], InTransformType));
+				}
+				else
+				{
+					const FTransform ParentTransformA = GetTransform(ParentElements[0], InTransformType);
+					const FTransform ParentTransformB = GetTransform(ParentElements[1], InTransformType);
+					Output.Set(FControlRigMathLibrary::LerpTransform(ParentTransformA, ParentTransformB, Weight));
+				}
+			}
+			else if(ParentElements.Num() > 2)
+			{
+				ensure(ParentElements.Num() == ParentWeights.Num());
+
+				// determine if there are more than one parent weighted
+				float OverallWeight = 0.f;
+				int32 FirstWeightedParent = INDEX_NONE;
+				int32 SecondWeightedParent = INDEX_NONE;
+				int32 NumWeightedParents = 0;
+				for(int32 ParentIndex = 0; ParentIndex < ParentWeights.Num(); ParentIndex++)
+				{
+					const float Weight = ParentWeights[ParentIndex];
+					const float ClampedWeight = FMath::Max(Weight, 0.f);
+					OverallWeight += ClampedWeight;
+
+					if(ClampedWeight > SMALL_NUMBER)
+					{
+						NumWeightedParents++;
+						if(FirstWeightedParent == INDEX_NONE)
+						{
+							FirstWeightedParent = ParentIndex;
+						}
+						else if(SecondWeightedParent == INDEX_NONE)
+						{
+							SecondWeightedParent = ParentIndex;
+						}
+					}
+				}
+
+				// if there is only weight on one parent, behave like a single parent element
+				if(NumWeightedParents == 1)
+				{
+					Output.Set(GetTransform(ParentElements[FirstWeightedParent], InTransformType));
+				}
+				else if(NumWeightedParents == 2) // for two weighted parents, special case once more
+				{
+					float Weight = 0.f;
+					const float WeightA = ParentWeights[FirstWeightedParent];
+					const float WeightB = ParentWeights[SecondWeightedParent];
+					const float ClampedWeightA = FMath::Max(WeightA, 0.f);
+					const float ClampedWeightB = FMath::Max(WeightB, 0.f);
+					if(OverallWeight > SMALL_NUMBER)
+					{
+						Weight = ClampedWeightB / OverallWeight;
+					}
+
+					const FTransform ParentTransformA = GetTransform(ParentElements[FirstWeightedParent], InTransformType);
+					const FTransform ParentTransformB = GetTransform(ParentElements[SecondWeightedParent], InTransformType);
+					Output.Set(FControlRigMathLibrary::LerpTransform(ParentTransformA, ParentTransformB, Weight));
+				}
+				else if(OverallWeight > SMALL_NUMBER)
+				{
+					FVector MixedTranslation = FVector::ZeroVector;
+					FQuat MixedRotation = FQuat(0.f, 0.f, 0.f, 0.f);;
+					FVector MixedScale3D = FVector::ZeroVector;
+
+					FQuat FirstRotation = FQuat::Identity;
+					int NumMixedRotations = 0;
+
+					for(int32 ParentIndex = 0; ParentIndex < ParentWeights.Num(); ParentIndex++)
+					{
+						const float Weight = ParentWeights[ParentIndex];
+						const float ClampedWeight = FMath::Max(Weight, 0.f);
+						if(ClampedWeight <= SMALL_NUMBER)
+						{
+							continue;
+						}
+						const float NormalizedWeight = ClampedWeight / OverallWeight;
+
+						const FTransform ParentTransform = GetTransform(ParentElements[ParentIndex], InTransformType);
+
+						// mix translation
+						MixedTranslation += ParentTransform.GetTranslation() * NormalizedWeight; 
+
+						// mix rotation
+						FQuat CurrentRotation = ParentTransform.GetRotation();
+						if(NumMixedRotations == 0)
+						{
+							FirstRotation = CurrentRotation; 
+						}
+						else if ((CurrentRotation | FirstRotation) <= 0.f)
+						{
+							// invert sign of rotation (NOT the same as .Inverse() )
+							CurrentRotation = FQuat(-CurrentRotation.X, -CurrentRotation.Y, -CurrentRotation.Z, -CurrentRotation.W);
+						}
+
+						// scale down the rotation if needed
+						if(NormalizedWeight < (1.f - SMALL_NUMBER))
+						{
+							CurrentRotation = FQuat::Slerp(FQuat::Identity, CurrentRotation, NormalizedWeight);
+						}
+
+						// accumulate part of current rotation
+						MixedRotation.W += CurrentRotation.W;
+						MixedRotation.X += CurrentRotation.X;
+						MixedRotation.Y += CurrentRotation.Y;
+						MixedRotation.Z += CurrentRotation.Z;
+
+						NumMixedRotations++;
+
+						// mix scale
+						MixedScale3D += ParentTransform.GetScale3D() * NormalizedWeight; 
+					}
+
+					float W = MixedRotation.W;
+					float X = MixedRotation.X;
+					float Y = MixedRotation.Y;
+					float Z = MixedRotation.Z;
+						
+					//Normalize. Note: experiment to see whether you can skip this step.
+					float D = 1.0f / (W * W + X * X + Y * Y + Z * Z);
+					W *= D;
+					X *= D;
+					Y *= D;
+					Z *= D;
+						
+					MixedRotation.X = X;
+					MixedRotation.Y = Y;
+					MixedRotation.Z = Z;
+					MixedRotation.W = W;
+
+					FTransform MixedTransform = FTransform::Identity;
+					MixedTransform.SetTranslation(MixedTranslation);
+					MixedTransform.SetRotation(MixedRotation.GetNormalized());
+					MixedTransform.SetScale3D(MixedScale3D);
+					Output.Set(MixedTransform);
+				}
+			}
+
+			Output.bDirty = false;
+		}
+		return Output.Transform;
+	}
+	return FTransform::Identity;
+}
+
+FRigControlValue URigHierarchy::GetControlValue(FRigControlElement* InControlElement, ERigControlValueType InValueType) const
+{
+	using namespace ERigTransformType;
+
+	FRigControlValue Value;
+
+	if(InControlElement != nullptr)
+	{
+		switch(InValueType)
+		{
+			case ERigControlValueType::Current:
+			{
+				Value.SetFromTransform(
+                    GetTransform(InControlElement, CurrentLocal),
+                    InControlElement->Settings.ControlType,
+                    InControlElement->Settings.PrimaryAxis
+                );
+				break;
+			}
+			case ERigControlValueType::Initial:
+			{
+				Value.SetFromTransform(
+                    GetTransform(InControlElement, InitialLocal),
+                    InControlElement->Settings.ControlType,
+                    InControlElement->Settings.PrimaryAxis
+                );
+				break;
+			}
+			case ERigControlValueType::Minimum:
+			{
+				return InControlElement->Settings.MinimumValue;
+			}
+			case ERigControlValueType::Maximum:
+			{
+				return InControlElement->Settings.MaximumValue;
+			}
+		}
+	}
+	return Value;
+}
+
+void URigHierarchy::SetControlValue(FRigControlElement* InControlElement, const FRigControlValue& InValue, ERigControlValueType InValueType, bool bSetupUndo, bool bForce)
+{
+	using namespace ERigTransformType;
+
+	if(InControlElement != nullptr)
+	{
+		switch(InValueType)
+		{
+			case ERigControlValueType::Current:
+			{
+				FRigControlValue Value = InValue;
+				InControlElement->Settings.ApplyLimits(Value);
+					
+				SetTransform(
+                    InControlElement,
+                    Value.GetAsTransform(
+                        InControlElement->Settings.ControlType,
+                        InControlElement->Settings.PrimaryAxis
+                    ),
+                    CurrentLocal,
+                    true,
+                    bSetupUndo,
+                    bForce
+                ); 
+				break;
+			}
+			case ERigControlValueType::Initial:
+			{
+				FRigControlValue Value = InValue;
+				InControlElement->Settings.ApplyLimits(Value);
+
+				SetTransform(
+                    InControlElement,
+                    Value.GetAsTransform(
+                        InControlElement->Settings.ControlType,
+                        InControlElement->Settings.PrimaryAxis
+                    ),
+                    InitialLocal,
+                    true,
+                    bSetupUndo,
+                    bForce
+                ); 
+				break;
+			}
+			case ERigControlValueType::Minimum:
+			{
+				InControlElement->Settings.MinimumValue = InValue;
+				InControlElement->Settings.ApplyLimits(InControlElement->Settings.MinimumValue);
+				Notify(ERigHierarchyNotification::ControlSettingChanged, InControlElement);
+				break;
+			}
+			case ERigControlValueType::Maximum:
+			{
+				InControlElement->Settings.MaximumValue = InValue;
+				InControlElement->Settings.ApplyLimits(InControlElement->Settings.MaximumValue);
+				Notify(ERigHierarchyNotification::ControlSettingChanged, InControlElement);
+				break;
+			}
+		}	
+	}
+}
+
+void URigHierarchy::SetControlVisibility(FRigControlElement* InControlElement, bool bVisibility)
+{
+	if(InControlElement == nullptr)
+	{
+		return;
+	}
+
+	InControlElement->Settings.bGizmoVisible = bVisibility;
+		Notify(ERigHierarchyNotification::ControlSettingChanged, InControlElement);
+}
+
+float URigHierarchy::GetCurveValue(FRigCurveElement* InCurveElement) const
+{
+	if(InCurveElement == nullptr)
+	{
+		return 0.f;
+	}
+	return InCurveElement->Value;
+}
+
+void URigHierarchy::SetCurveValue(FRigCurveElement* InCurveElement, float InValue, bool bSetupUndo, bool bForce)
+{
+	if(InCurveElement == nullptr)
+	{
+		return;
+	}
+
+	const float PreviousValue = InCurveElement->Value;
+	if(!bForce && FMath::IsNearlyZero(PreviousValue - InValue))
+	{
+		return;
+	}
+
+	InCurveElement->Value = InValue;
+
+#if WITH_EDITOR
+	if(bSetupUndo)
+	{
+		PushCurveToStack(InCurveElement->GetKey(), PreviousValue, InCurveElement->Value);
+	}
+#endif
+}
+
+FName URigHierarchy::GetPreviousName(const FRigElementKey& InKey) const
+{
+	if(const FRigElementKey* OldKeyPtr = PreviousNameMap.Find(InKey))
+	{
+		return OldKeyPtr->Name;
+	}
+	return NAME_None;
+}
+
+FRigElementKey URigHierarchy::GetPreviousParent(const FRigElementKey& InKey) const
+{
+	if(const FRigElementKey* OldParentPtr = PreviousParentMap.Find(InKey))
+	{
+		return *OldParentPtr;
+	}
+	return FRigElementKey();
+}
+
+bool URigHierarchy::IsParentedTo(FRigBaseElement* InChild, FRigBaseElement* InParent) const
+{
+	if((InChild == nullptr) || (InParent == nullptr))
+	{
+		return false;
+	}
+
+	if(FRigSingleParentElement* SingleParentElement = Cast<FRigSingleParentElement>(InChild))
+	{
+		if(SingleParentElement->ParentElement == InParent)
+		{
+			return true;
+		}
+		return IsParentedTo(SingleParentElement->ParentElement, InParent);
+	}
+
+	if(FRigMultiParentElement* MultiParentElement = Cast<FRigMultiParentElement>(InChild))
+	{
+		for(FRigTransformElement* MultiParentElementParent : MultiParentElement->ParentElements)
+		{
+			if(MultiParentElementParent == InParent)
+			{
+				return true;
+			}
+			if(IsParentedTo(MultiParentElementParent, InParent))
+			{
+				return true;
+			}
+		}
+	}
+
+	return false;
+}
+
+bool URigHierarchy::IsSelected(const FRigBaseElement* InElement) const
+{
+	if(InElement == nullptr)
+	{
+		return false;
+	}
+	if(URigHierarchy* HierarchyForSelection = HierarchyForSelectionPtr.Get())
+	{
+		return HierarchyForSelection->IsSelected(InElement->GetKey());
+	}
+	return InElement->IsSelected();
+}
+
+void URigHierarchy::ResetCachedChildren()
+{
+	for (int32 ElementIndex = 0; ElementIndex < Elements.Num(); ElementIndex++)
+	{
+		FRigBaseElement* Element = Elements[ElementIndex];
+		Element->CachedChildren.Reset();
+	}
+}
+
+void URigHierarchy::UpdateCachedChildren(const FRigBaseElement* InElement, bool bForce) const
+{
+	check(InElement);
+
+	if(InElement->TopologyVersion == TopologyVersion && !bForce)
+	{
+		return;
+	}
+	
+	InElement->CachedChildren.Reset();
+	
+	for (int32 ElementIndex = 0; ElementIndex < Elements.Num(); ElementIndex++)
+	{
+		FRigBaseElement* Element = Elements[ElementIndex];
+		if(FRigSingleParentElement* SingleParentElement = Cast<FRigSingleParentElement>(Element))
+		{
+			if(SingleParentElement->ParentElement == InElement)
+			{
+				InElement->CachedChildren.Add(SingleParentElement);
+			}
+		}
+		else if(FRigMultiParentElement* MultiParentElement = Cast<FRigMultiParentElement>(Element))
+		{
+			for(FRigTransformElement* ParentElement : MultiParentElement->ParentElements)
+			{
+				if(ParentElement == InElement)
+				{
+					InElement->CachedChildren.Add(MultiParentElement);
+					break;
+				}
+			}
+		}
+	}
+
+	InElement->TopologyVersion = TopologyVersion;
+}
+
+void URigHierarchy::PropagateDirtyFlags(FRigTransformElement* InTransformElement, const ERigTransformType::Type InTransformType, bool bAffectChildren) const
+{
+	if(!bEnableDirtyPropagation)
+	{
+		return;
+	}
+	
+	check(InTransformElement);
+
+	const ERigTransformType::Type LocalType = MakeLocal(InTransformType);
+	const ERigTransformType::Type GlobalType = MakeGlobal(InTransformType);
+	const ERigTransformType::Type TypeToCompute = bAffectChildren ? LocalType : GlobalType;
+	const ERigTransformType::Type TypeToDirty = SwapLocalAndGlobal(TypeToCompute);
+	
+	for(const FRigTransformElement::FElementToDirty& ElementToDirty : InTransformElement->ElementsToDirty)
+	{
+		if(!bAffectChildren && ElementToDirty.HierarchyDistance > 1)
+		{
+			continue;
+		}
+
+		GetTransform(ElementToDirty.Element, TypeToCompute); // make sure the local / global transform is up 2 date
+	}
+	
+	for(const FRigTransformElement::FElementToDirty& ElementToDirty : InTransformElement->ElementsToDirty)
+	{
+#if URIGHIERARCHY_RECURSIVE_DIRTY_PROPAGATION
+
+		if(ElementToDirty.Element->Pose.IsDirty(TypeToDirty))
+		{
+			return;
+		}
+		
+#else
+
+		if(!bAffectChildren && ElementToDirty.HierarchyDistance > 1)
+		{
+			return;
+		}
+		
+#endif
+
+		ElementToDirty.Element->Pose.MarkDirty(TypeToDirty);
+	
+		if(FRigMultiParentElement* MultiParentElement = Cast<FRigMultiParentElement>(ElementToDirty.Element))
+		{
+			MultiParentElement->Parent.MarkDirty(GlobalType);
+		}
+		if(FRigControlElement* ControlElement = Cast<FRigControlElement>(ElementToDirty.Element))
+		{
+			ControlElement->Offset.MarkDirty(GlobalType);
+			ControlElement->Gizmo.MarkDirty(GlobalType);
+		}
+
+#if URIGHIERARCHY_RECURSIVE_DIRTY_PROPAGATION
+
+		if(bAffectChildren)
+		{
+			PropagateDirtyFlags(ElementToDirty.Element, InTransformType, bAffectChildren);
+		}
+		
+#endif
+	}
+}
+
+void URigHierarchy::PushTransformToStack(const FRigElementKey& InKey, ETransformStackEntryType InEntryType,
+	ERigTransformType::Type InTransformType, const FTransform& InOldTransform, const FTransform& InNewTransform,
+	bool bAffectChildren)
+{
+#if WITH_EDITOR
+
+	static const FText TransformPoseTitle = NSLOCTEXT("RigHierarchy", "Set Pose Transform", "Set Pose Transform");
+	static const FText ControlOffsetTitle = NSLOCTEXT("RigHierarchy", "Set Control Offset", "Set Control Offset");
+	static const FText ControlGizmoTitle = NSLOCTEXT("RigHierarchy", "Set Control Gizo", "Set Control Gizo");
+	static const FText CurveValueTitle = NSLOCTEXT("RigHierarchy", "Set Curve Value", "Set Curve Value");
+	
+	FText Title;
+	switch(InEntryType)
+	{
+		case ETransformStackEntryType::TransformPose:
+		{
+			Title = TransformPoseTitle;
+			break;
+		}
+		case ETransformStackEntryType::ControlOffset:
+		{
+			Title = TransformPoseTitle;
+			break;
+		}
+		case ETransformStackEntryType::ControlGizmo:
+		{
+			Title = TransformPoseTitle;
+			break;
+		}
+		case ETransformStackEntryType::CurveValue:
+		{
+			Title = TransformPoseTitle;
+			break;
+		}
+	}
+
+	TGuardValue<bool> TransactingGuard(bTransactingForTransformChange, true);
+	FScopedTransaction Transaction(Title);
+	Modify();
+	TransformUndoStack.Add(
+		FTransformStackEntry(InKey, InEntryType, InTransformType, InOldTransform, InNewTransform, bAffectChildren));
+	TransformStackIndex = TransformUndoStack.Num();
+
+	TransformRedoStack.Reset();
+	
+#endif
+}
+
+void URigHierarchy::PushCurveToStack(const FRigElementKey& InKey, float InOldCurveValue, float InNewCurveValue)
+{
+#if WITH_EDITOR
+
+	FTransform OldTransform = FTransform::Identity;
+	FTransform NewTransform = FTransform::Identity;
+
+	OldTransform.SetTranslation(FVector(InOldCurveValue, 0.f, 0.f));
+	NewTransform.SetTranslation(FVector(InNewCurveValue, 0.f, 0.f));
+
+	PushTransformToStack(InKey, ETransformStackEntryType::CurveValue, ERigTransformType::CurrentLocal, OldTransform, NewTransform, false);
+
+#endif
+}
+
+bool URigHierarchy::ApplyTransformFromStack(const FTransformStackEntry& InEntry, bool bUndo)
+{
+#if WITH_EDITOR
+	
+	FRigBaseElement* Element = Find(InEntry.Key);
+	if(Element == nullptr)
+	{
+		return false;
+	}
+
+	const FTransform& Transform = bUndo ? InEntry.OldTransform : InEntry.NewTransform;
+	
+	switch(InEntry.EntryType)
+	{
+		case ETransformStackEntryType::TransformPose:
+		{
+			SetTransform(Cast<FRigTransformElement>(Element), Transform, InEntry.TransformType, InEntry.bAffectChildren, false); 
+			break;
+		}
+		case ETransformStackEntryType::ControlOffset:
+		{
+			SetControlOffsetTransform(Cast<FRigControlElement>(Element), Transform, InEntry.TransformType, InEntry.bAffectChildren, false); 
+			break;
+		}
+		case ETransformStackEntryType::ControlGizmo:
+		{
+			SetControlGizmoTransform(Cast<FRigControlElement>(Element), Transform, InEntry.TransformType, false); 
+			break;
+		}
+		case ETransformStackEntryType::CurveValue:
+		{
+			const float CurveValue = Transform.GetTranslation().X;
+			SetCurveValue(Cast<FRigCurveElement>(Element), CurveValue, false);
+			break;
+		}
+	}
+
+	return true;
+
+#endif
+
+	return false;
+}
+
diff --git a/Engine/Plugins/Experimental/ControlRig/Source/ControlRig/Private/Rigs/RigHierarchyCache.cpp b/Engine/Plugins/Experimental/ControlRig/Source/ControlRig/Private/Rigs/RigHierarchyCache.cpp
--- a/Engine/Plugins/Experimental/ControlRig/Source/ControlRig/Private/Rigs/RigHierarchyCache.cpp
+++ b/Engine/Plugins/Experimental/ControlRig/Source/ControlRig/Private/Rigs/RigHierarchyCache.cpp
@@ -1,80 +1,45 @@
 // Copyright Epic Games, Inc. All Rights Reserved.
 
 #include "Rigs/RigHierarchyCache.h"
-#include "Rigs/RigHierarchyContainer.h"
+#include "Rigs/RigHierarchy.h"
 
-bool FCachedRigElement::UpdateCache(const FRigHierarchyContainer* InContainer)
+bool FCachedRigElement::UpdateCache(const URigHierarchy* InHierarchy)
 {
-	if(InContainer)
+	if(InHierarchy)
 	{
-		if(!IsValid() || InContainer->Version != ContainerVersion)
+		if(!IsValid() || InHierarchy->GetTopologyVersion() != ContainerVersion)
 		{
-			return UpdateCache(GetKey(), InContainer);
+			return UpdateCache(GetKey(), InHierarchy);
 		}
 		return IsValid();
 	}
 	return false;
 }
 
-bool FCachedRigElement::UpdateCache(const FRigElementKey& InKey, const FRigHierarchyContainer* InContainer)
+bool FCachedRigElement::UpdateCache(const FRigElementKey& InKey, const URigHierarchy* InHierarchy)
 {
-	if(InContainer)
+	if(InHierarchy)
 	{
-		if(!IsValid() || !IsIdentical(InKey, InContainer))
+		if(!IsValid() || !IsIdentical(InKey, InHierarchy))
 		{
 			Reset();
 
-			int32 Idx = InContainer->GetIndex(InKey);
+			int32 Idx = InHierarchy->GetIndex(InKey);
 			if(Idx != INDEX_NONE)
 			{
 				Key = InKey;
 				Index = (uint16)Idx;
+				Element = InHierarchy->Get(Index);
 			}
 
-			ContainerVersion = InContainer->Version;
+			ContainerVersion = InHierarchy->GetTopologyVersion();
 		}
 		return IsValid();
 	}
 	return false;
 }
 
-bool FCachedRigElement::UpdateCache(const FName& InName, const FRigBoneHierarchy* InHierarchy)
+bool FCachedRigElement::IsIdentical(const FRigElementKey& InKey, const URigHierarchy* InHierarchy)
 {
-	if (InHierarchy && InHierarchy->Container)
-	{
-		return UpdateCache(FRigElementKey(InName, ERigElementType::Bone), InHierarchy->Container);
-	}
-	return false;
-}
-
-bool FCachedRigElement::UpdateCache(const FName& InName, const FRigSpaceHierarchy* InHierarchy)
-{
-	if (InHierarchy && InHierarchy->Container)
-	{
-		return UpdateCache(FRigElementKey(InName, ERigElementType::Space), InHierarchy->Container);
-	}
-	return false;
-}
-
-bool FCachedRigElement::UpdateCache(const FName& InName, const FRigControlHierarchy* InHierarchy)
-{
-	if (InHierarchy && InHierarchy->Container)
-	{
-		return UpdateCache(FRigElementKey(InName, ERigElementType::Control), InHierarchy->Container);
-	}
-	return false;
-}
-
-bool FCachedRigElement::UpdateCache(const FName& InName, const FRigCurveContainer* InHierarchy)
-{
-	if (InHierarchy && InHierarchy->Container)
-	{
-		return UpdateCache(FRigElementKey(InName, ERigElementType::Curve), InHierarchy->Container);
-	}
-	return false;
-}
-
-bool FCachedRigElement::IsIdentical(const FRigElementKey& InKey, const FRigHierarchyContainer* InContainer)
-{
-	return InKey == Key && InContainer->Version == ContainerVersion;
+	return InKey == Key && InHierarchy->GetTopologyVersion() == ContainerVersion;
 }
