Control Rig: New rig hierarchy data structure refactoring

diff --git a/Engine/Plugins/Experimental/ControlRig/Source/ControlRig/Private/Rigs/RigHierarchyContainer.cpp b/Engine/Plugins/Experimental/ControlRig/Source/ControlRig/Private/Rigs/RigHierarchyContainer.cpp
--- a/Engine/Plugins/Experimental/ControlRig/Source/ControlRig/Private/Rigs/RigHierarchyContainer.cpp
+++ b/Engine/Plugins/Experimental/ControlRig/Source/ControlRig/Private/Rigs/RigHierarchyContainer.cpp
@@ -5,63 +5,12 @@
 #include "HelperUtil.h"
 #include "UObject/PropertyPortFlags.h"
 
-////////////////////////////////////////////////////////////////////////////////
-// FRigMirrorSettings
-////////////////////////////////////////////////////////////////////////////////
-
-FTransform FRigMirrorSettings::MirrorTransform(const FTransform& InTransform) const
-{
-	FTransform Transform = InTransform;
-	FQuat Quat = Transform.GetRotation();
-
-	Transform.SetLocation(MirrorVector(Transform.GetLocation()));
-
-	switch (AxisToFlip)
-	{
-		case EAxis::X:
-		{
-			FVector Y = MirrorVector(Quat.GetAxisY());
-			FVector Z = MirrorVector(Quat.GetAxisZ());
-			FMatrix Rotation = FRotationMatrix::MakeFromYZ(Y, Z);
-			Transform.SetRotation(FQuat(Rotation));
-			break;
-		}
-		case EAxis::Y:
-		{
-			FVector X = MirrorVector(Quat.GetAxisX());
-			FVector Z = MirrorVector(Quat.GetAxisZ());
-			FMatrix Rotation = FRotationMatrix::MakeFromXZ(X, Z);
-			Transform.SetRotation(FQuat(Rotation));
-			break;
-		}
-		default:
-		{
-			FVector X = MirrorVector(Quat.GetAxisX());
-			FVector Y = MirrorVector(Quat.GetAxisY());
-			FMatrix Rotation = FRotationMatrix::MakeFromXY(X, Y);
-			Transform.SetRotation(FQuat(Rotation));
-			break;
-		}
-	}
-
-	return Transform;
-}
-
-FVector FRigMirrorSettings::MirrorVector(const FVector& InVector) const
-{
-	FVector Axis = FVector::ZeroVector;
-	Axis.SetComponentForAxis(MirrorAxis, 1.f);
-	return InVector.MirrorByVector(Axis);
-}
-
 ////////////////////////////////////////////////////////////////////////////////
 // FRigHierarchyContainer
 ////////////////////////////////////////////////////////////////////////////////
 
 FRigHierarchyContainer::FRigHierarchyContainer()
 {
-	Version = 0;
-	Initialize();
 }
 
 FRigHierarchyContainer::FRigHierarchyContainer(const FRigHierarchyContainer& InOther)
@@ -70,11 +19,6 @@ FRigHierarchyContainer::FRigHierarchyContainer(const FRigHierarchyContainer& InO
 	SpaceHierarchy = InOther.SpaceHierarchy;
 	ControlHierarchy = InOther.ControlHierarchy;
 	CurveContainer = InOther.CurveContainer;
-
-	DepthIndexByKey.Reset();
-	Version = InOther.Version + 1;
-
-	Initialize();
 }
 
 FRigHierarchyContainer& FRigHierarchyContainer::operator= (const FRigHierarchyContainer &InOther)
@@ -84,1499 +28,5 @@ FRigHierarchyContainer& FRigHierarchyContainer::operator= (const FRigHierarchyCo
 	ControlHierarchy = InOther.ControlHierarchy;
 	CurveContainer = InOther.CurveContainer;
 
-	DepthIndexByKey.Reset();
-	Version++;
-
-	Initialize();
-
 	return *this;
 }
-
-
-void FRigHierarchyContainer::Initialize(bool bResetTransforms)
-{
-	BoneHierarchy.Container = this;
-	SpaceHierarchy.Container = this;
-	ControlHierarchy.Container = this;
-	CurveContainer.Container = this;
-
-#if WITH_EDITOR
-	BoneHierarchy.OnBoneAdded.RemoveAll(this);
-	BoneHierarchy.OnBoneRemoved.RemoveAll(this);
-	BoneHierarchy.OnBoneRenamed.RemoveAll(this);
-	BoneHierarchy.OnBoneReparented.RemoveAll(this);
-#endif
-	BoneHierarchy.OnBoneSelected.RemoveAll(this);
-#if WITH_EDITOR
-
-	BoneHierarchy.OnBoneAdded.AddRaw(this, &FRigHierarchyContainer::HandleOnElementAdded);
-	BoneHierarchy.OnBoneRemoved.AddRaw(this, &FRigHierarchyContainer::HandleOnElementRemoved);
-	BoneHierarchy.OnBoneRenamed.AddRaw(this, &FRigHierarchyContainer::HandleOnElementRenamed);
-	BoneHierarchy.OnBoneReparented.AddRaw(this, &FRigHierarchyContainer::HandleOnElementReparented);
-#endif
-	BoneHierarchy.OnBoneSelected.AddRaw(this, &FRigHierarchyContainer::HandleOnElementSelected);
-#if WITH_EDITOR
-
-	SpaceHierarchy.OnSpaceAdded.RemoveAll(this);
-	SpaceHierarchy.OnSpaceRemoved.RemoveAll(this);
-	SpaceHierarchy.OnSpaceRenamed.RemoveAll(this);
-	SpaceHierarchy.OnSpaceReparented.RemoveAll(this);
-#endif
-	SpaceHierarchy.OnSpaceSelected.RemoveAll(this);
-#if WITH_EDITOR
-
-	SpaceHierarchy.OnSpaceAdded.AddRaw(this, &FRigHierarchyContainer::HandleOnElementAdded);
-	SpaceHierarchy.OnSpaceRemoved.AddRaw(this, &FRigHierarchyContainer::HandleOnElementRemoved);
-	SpaceHierarchy.OnSpaceRenamed.AddRaw(this, &FRigHierarchyContainer::HandleOnElementRenamed);
-	SpaceHierarchy.OnSpaceReparented.AddRaw(this, &FRigHierarchyContainer::HandleOnElementReparented);
-#endif
-	SpaceHierarchy.OnSpaceSelected.AddRaw(this, &FRigHierarchyContainer::HandleOnElementSelected);
-#if WITH_EDITOR
-
-	ControlHierarchy.OnControlAdded.RemoveAll(this);
-	ControlHierarchy.OnControlRemoved.RemoveAll(this);
-	ControlHierarchy.OnControlRenamed.RemoveAll(this);
-	ControlHierarchy.OnControlReparented.RemoveAll(this);
-#endif
-	ControlHierarchy.OnControlSelected.RemoveAll(this);
-#if WITH_EDITOR
-
-	ControlHierarchy.OnControlAdded.AddRaw(this, &FRigHierarchyContainer::HandleOnElementAdded);
-	ControlHierarchy.OnControlRemoved.AddRaw(this, &FRigHierarchyContainer::HandleOnElementRemoved);
-	ControlHierarchy.OnControlRenamed.AddRaw(this, &FRigHierarchyContainer::HandleOnElementRenamed);
-	ControlHierarchy.OnControlReparented.AddRaw(this, &FRigHierarchyContainer::HandleOnElementReparented);
-#endif
-	ControlHierarchy.OnControlSelected.AddRaw(this, &FRigHierarchyContainer::HandleOnElementSelected);
-#if WITH_EDITOR
-
-	CurveContainer.OnCurveAdded.RemoveAll(this);
-	CurveContainer.OnCurveRemoved.RemoveAll(this);
-	CurveContainer.OnCurveRenamed.RemoveAll(this);
-	CurveContainer.OnCurveSelected.RemoveAll(this);
-
-	CurveContainer.OnCurveAdded.AddRaw(this, &FRigHierarchyContainer::HandleOnElementAdded);
-	CurveContainer.OnCurveRemoved.AddRaw(this, &FRigHierarchyContainer::HandleOnElementRemoved);
-	CurveContainer.OnCurveRenamed.AddRaw(this, &FRigHierarchyContainer::HandleOnElementRenamed);
-	CurveContainer.OnCurveSelected.AddRaw(this, &FRigHierarchyContainer::HandleOnElementSelected);
-
-	// wire them between each other
-	BoneHierarchy.OnBoneRenamed.RemoveAll(&SpaceHierarchy);
-	SpaceHierarchy.OnSpaceRenamed.RemoveAll(&ControlHierarchy);
-	ControlHierarchy.OnControlRenamed.RemoveAll(&SpaceHierarchy);
-	BoneHierarchy.OnBoneRenamed.AddRaw(&SpaceHierarchy, &FRigSpaceHierarchy::HandleOnElementRenamed);
-	SpaceHierarchy.OnSpaceRenamed.AddRaw(&ControlHierarchy, &FRigControlHierarchy::HandleOnElementRenamed);
-	ControlHierarchy.OnControlRenamed.AddRaw(&SpaceHierarchy, &FRigSpaceHierarchy::HandleOnElementRenamed);
-	BoneHierarchy.OnBoneRemoved.AddRaw(&SpaceHierarchy, &FRigSpaceHierarchy::HandleOnElementRemoved);
-	ControlHierarchy.OnControlRemoved.AddRaw(&SpaceHierarchy, &FRigSpaceHierarchy::HandleOnElementRemoved);
-	SpaceHierarchy.OnSpaceRemoved.AddRaw(&ControlHierarchy, &FRigControlHierarchy::HandleOnElementRemoved);
-
-#endif
-
-	BoneHierarchy.Initialize(bResetTransforms);
-	SpaceHierarchy.Initialize(bResetTransforms);
-	ControlHierarchy.Initialize(bResetTransforms);
-	CurveContainer.Initialize();
-	DepthIndexByKey.Reset();
-	Version++;
-
-	if (bResetTransforms)
-	{
-		ResetTransforms();
-	}
-}
-
-void FRigHierarchyContainer::Reset()
-{
-	BoneHierarchy.Reset();
-	SpaceHierarchy.Reset();
-	ControlHierarchy.Reset();
-	CurveContainer.Reset();
-	DepthIndexByKey.Reset();
-
-	Initialize();
-}
-
-void FRigHierarchyContainer::ResetTransforms()
-{
-	BoneHierarchy.ResetTransforms();
-	SpaceHierarchy.ResetTransforms();
-	ControlHierarchy.ResetValues();
-	CurveContainer.ResetValues();
-}
-
-void FRigHierarchyContainer::CopyInitialTransforms(const FRigHierarchyContainer& InOther)
-{
-	BoneHierarchy.CopyInitialTransforms(InOther.BoneHierarchy);
-	SpaceHierarchy.CopyInitialTransforms(InOther.SpaceHierarchy);
-	ControlHierarchy.CopyOffsetTransforms(InOther.ControlHierarchy);
-}
-
-FTransform FRigHierarchyContainer::GetInitialTransform(ERigElementType InElementType, int32 InIndex) const
-{
-	if(InIndex == INDEX_NONE)
-	{
-		return FTransform::Identity;
-	}
-
-	switch(InElementType)
-	{
-		case ERigElementType::Bone:
-		{
-			return BoneHierarchy.GetInitialLocalTransform(InIndex);
-		}
-		case ERigElementType::Space:
-		{
-			return SpaceHierarchy.GetInitialTransform(InIndex);
-		}
-		case ERigElementType::Control:
-		{
-			return ControlHierarchy.GetLocalTransform(InIndex,ERigControlValueType::Initial);
-		}
-		case ERigElementType::Curve:
-		{
-			break;
-		}
-	}
-
-	return FTransform::Identity;
-}
-
-void FRigHierarchyContainer::SetInitialTransform(ERigElementType InElementType, int32 InIndex, const FTransform& InTransform)
-{
-	if (InIndex == INDEX_NONE)
-	{
-		return;
-	}
-
-	switch(InElementType)
-	{
-		case ERigElementType::Bone:
-		{
-			BoneHierarchy.SetInitialLocalTransform(InIndex, InTransform);
-			break;
-		}
-		case ERigElementType::Space:
-		{
-			SpaceHierarchy.SetInitialTransform(InIndex, InTransform);
-			break;
-		}
-		case ERigElementType::Control:
-		{
-			ControlHierarchy.SetLocalTransform(InIndex, InTransform,ERigControlValueType::Initial);
-			break;
-		}
-		case ERigElementType::Curve:
-		{
-			break;
-		}
-	}
-}
-
-FTransform FRigHierarchyContainer::GetInitialGlobalTransform(ERigElementType InElementType, int32 InIndex) const
-{
-	if (InIndex == INDEX_NONE)
-	{
-		return FTransform::Identity;
-	}
-
-	switch (InElementType)
-	{
-		case ERigElementType::Bone:
-		{
-			return BoneHierarchy.GetInitialGlobalTransform(InIndex);
-		}
-		case ERigElementType::Space:
-		{
-			return SpaceHierarchy.GetInitialGlobalTransform(InIndex);
-		}
-		case ERigElementType::Control:
-		{
-			return ControlHierarchy.GetInitialGlobalTransform(InIndex);
-		}
-		case ERigElementType::Curve:
-		{
-			break;
-		}
-	}
-
-	return FTransform::Identity;
-}
-
-void FRigHierarchyContainer::SetInitialGlobalTransform(ERigElementType InElementType, int32 InIndex, const FTransform& InTransform)
-{
-	if (InIndex == INDEX_NONE)
-	{
-		return;
-	}
-
-	switch (InElementType)
-	{
-		case ERigElementType::Bone:
-		{
-			BoneHierarchy.SetInitialGlobalTransform(InIndex, InTransform);
-			break;
-		}
-		case ERigElementType::Space:
-		{
-			SpaceHierarchy.SetInitialGlobalTransform(InIndex, InTransform);
-			break;
-		}
-		case ERigElementType::Control:
-		{
-			ControlHierarchy.SetInitialGlobalTransform(InIndex, InTransform);
-			break;
-		}
-		case ERigElementType::Curve:
-		{
-			break;
-		}
-	}
-}
-
-FTransform FRigHierarchyContainer::GetLocalTransform(ERigElementType InElementType, int32 InIndex) const
-{
-	if(InIndex == INDEX_NONE)
-	{
-		return FTransform::Identity;
-	}
-
-	switch(InElementType)
-	{
-		case ERigElementType::Bone:
-		{
-			return BoneHierarchy.GetLocalTransform(InIndex);
-		}
-		case ERigElementType::Space:
-		{
-			return SpaceHierarchy.GetLocalTransform(InIndex);
-		}
-		case ERigElementType::Control:
-		{
-			return ControlHierarchy.GetLocalTransform(InIndex);
-		}
-		case ERigElementType::Curve:
-		{
-			break;
-		}
-	}
-
-	return FTransform::Identity;
-}
-
-void FRigHierarchyContainer::SetLocalTransform(ERigElementType InElementType, int32 InIndex, const FTransform& InTransform, bool bPropagateToChildren)
-{
-	if (InIndex == INDEX_NONE)
-	{
-		return;
-	}
-
-	switch(InElementType)
-	{
-		case ERigElementType::Bone:
-		{
-			BoneHierarchy.SetLocalTransform(InIndex, InTransform, bPropagateToChildren);
-			break;
-		}
-		case ERigElementType::Space:
-		{
-			SpaceHierarchy.SetLocalTransform(InIndex, InTransform);
-			break;
-		}
-		case ERigElementType::Control:
-		{
-			ControlHierarchy.SetLocalTransform(InIndex, InTransform);
-			break;
-		}
-		case ERigElementType::Curve:
-		{
-			break;
-		}
-	}
-}
-
-FTransform FRigHierarchyContainer::GetGlobalTransform(ERigElementType InElementType, int32 InIndex) const
-{
-	if(InIndex == INDEX_NONE)
-	{
-		return FTransform::Identity;
-	}
-
-	switch(InElementType)
-	{
-		case ERigElementType::Bone:
-		{
-			return BoneHierarchy.GetGlobalTransform(InIndex);
-		}
-		case ERigElementType::Space:
-		{
-			return SpaceHierarchy.GetGlobalTransform(InIndex);
-		}
-		case ERigElementType::Control:
-		{
-			return ControlHierarchy.GetGlobalTransform(InIndex);
-		}
-		case ERigElementType::Curve:
-		{
-			break;
-		}
-	}
-
-	return FTransform::Identity;
-}
-
-void FRigHierarchyContainer::SetGlobalTransform(ERigElementType InElementType, int32 InIndex, const FTransform& InTransform, bool bPropagateToChildren)
-{
-	if (InIndex == INDEX_NONE)
-	{
-		return;
-	}
-
-	switch(InElementType)
-	{
-		case ERigElementType::Bone:
-		{
-			BoneHierarchy.SetGlobalTransform(InIndex, InTransform, bPropagateToChildren);
-			break;
-		}
-		case ERigElementType::Space:
-		{
-			SpaceHierarchy.SetGlobalTransform(InIndex, InTransform);
-			break;
-		}
-		case ERigElementType::Control:
-		{
-			ControlHierarchy.SetGlobalTransform(InIndex, InTransform);
-			break;
-		}
-		case ERigElementType::Curve:
-		{
-			break;
-		}
-	}
-}
-
-#if WITH_EDITOR
-
-void FRigHierarchyContainer::HandleOnElementAdded(FRigHierarchyContainer* InContainer, const FRigElementKey& InKey)
-{
-	DepthIndexByKey.Reset();
-	Version++;
-
-	OnElementAdded.Broadcast(InContainer, InKey);
-	OnElementChanged.Broadcast(InContainer, InKey);
-}
-
-void FRigHierarchyContainer::HandleOnElementRemoved(FRigHierarchyContainer* InContainer, const FRigElementKey& InKey)
-{
-	DepthIndexByKey.Reset();
-	Version++;
-
-	OnElementRemoved.Broadcast(InContainer, InKey);
-	OnElementChanged.Broadcast(InContainer, InKey);
-}
-
-void FRigHierarchyContainer::HandleOnElementRenamed(FRigHierarchyContainer* InContainer, ERigElementType InElementType, const FName& InOldName, const FName& InNewName)
-{
-	DepthIndexByKey.Reset();
-	Version++;
-
-	OnElementRenamed.Broadcast(InContainer, InElementType, InOldName, InNewName);
-	OnElementChanged.Broadcast(InContainer, FRigElementKey(InNewName, InElementType));
-}
-
-void FRigHierarchyContainer::HandleOnElementReparented(FRigHierarchyContainer* InContainer, const FRigElementKey& InKey, const FName& InOldParentName, const FName& InNewParentName)
-{
-	DepthIndexByKey.Reset();
-	Version++;
-
-	OnElementReparented.Broadcast(InContainer, InKey, InOldParentName, InNewParentName);
-	OnElementChanged.Broadcast(InContainer, InKey);
-}
-
-#endif
-
-void FRigHierarchyContainer::HandleOnElementSelected(FRigHierarchyContainer* InContainer, const FRigElementKey& InKey, bool bSelected)
-{
-	OnElementSelected.Broadcast(InContainer, InKey, bSelected);
-	OnElementChanged.Broadcast(InContainer, InKey);
-}
-
-FRigElementKey FRigHierarchyContainer::GetParentKey(const FRigElementKey& InKey) const
-{
-	switch (InKey.Type)
-	{
-		case ERigElementType::Bone:
-		{
-			return BoneHierarchy[InKey.Name].GetParentElementKey();
-		}
-		case ERigElementType::Space:
-		{
-			return SpaceHierarchy[InKey.Name].GetParentElementKey();
-		}
-		case ERigElementType::Control:
-		{
-			const FRigControl& Control = ControlHierarchy[InKey.Name];
-			FRigElementKey SpaceKey = Control.GetSpaceElementKey();
-			if (SpaceKey.IsValid())
-			{
-				return SpaceKey;
-			}
-			return Control.GetParentElementKey();
-		}
-	}
-
-	return FRigElementKey();
-}
-
-TArray<FRigElementKey> FRigHierarchyContainer::GetChildKeys(const FRigElementKey& InKey, bool bRecursive) const
-{
-	TArray<FRigElementKey> Children;
-	for (const FRigBone& Element : BoneHierarchy)
-	{
-		if (Element.GetParentElementKey() == InKey)
-		{
-			Children.Add(Element.GetElementKey());
-			if (bRecursive)
-			{
-				Children.Append(GetChildKeys(Children.Last()));
-			}
-		}
-	}
-	for (const FRigControl& Element : ControlHierarchy)
-	{
-		if (Element.GetSpaceElementKey() == InKey ||
-			Element.GetParentElementKey() == InKey)
-		{
-			Children.Add(Element.GetElementKey());
-			if (bRecursive)
-			{
-				Children.Append(GetChildKeys(Children.Last()));
-			}
-		}
-	}
-	for (const FRigSpace& Element : SpaceHierarchy)
-	{
-		if (Element.GetParentElementKey() == InKey)
-		{
-			Children.Add(Element.GetElementKey());
-			if (bRecursive)
-			{
-				Children.Append(GetChildKeys(Children.Last()));
-			}
-		}
-	}
-	return Children;
-}
-
-bool FRigHierarchyContainer::IsParentedTo(ERigElementType InChildType, int32 InChildIndex, ERigElementType InParentType, int32 InParentIndex) const
-{
-	if (InParentIndex == INDEX_NONE || InChildIndex == INDEX_NONE)
-	{
-		return false;
-	}
-
-	switch (InChildType)
-	{
-		case ERigElementType::Curve:
-		{
-			return false;
-		}
-		case ERigElementType::Bone:
-		{
-			switch (InParentType)
-			{
-				case ERigElementType::Bone:
-				{
-					if (BoneHierarchy[InChildIndex].ParentIndex != INDEX_NONE)
-					{
-						if (BoneHierarchy[InChildIndex].ParentIndex == InParentIndex)
-						{
-							return true;
-						}
-						return IsParentedTo(ERigElementType::Bone, BoneHierarchy[InChildIndex].ParentIndex, InParentType, InParentIndex);
-					}
-					// no break - fall through to next case
-				}
-				case ERigElementType::Space:
-				case ERigElementType::Control:
-				case ERigElementType::Curve:
-				{
-					return false;
-				}
-			}
-		}
-		case ERigElementType::Space:
-		{
-			const FRigSpace& ChildSpace = SpaceHierarchy[InChildIndex];
-			switch (ChildSpace.SpaceType)
-			{
-				case ERigSpaceType::Global:
-				{
-					return false;
-				}
-				case ERigSpaceType::Bone:
-				{
-					if (ChildSpace.ParentIndex == InParentIndex && InParentType == ERigElementType::Bone)
-					{
-						return true;
-					}
-					return IsParentedTo(ERigElementType::Bone, ChildSpace.ParentIndex, InParentType, InParentIndex);
-				}
-				case ERigSpaceType::Space:
-				{
-					if (ChildSpace.ParentIndex == InParentIndex && InParentType == ERigElementType::Space)
-					{
-						return true;
-					}
-					return IsParentedTo(ERigElementType::Space, ChildSpace.ParentIndex, InParentType, InParentIndex);
-				}
-				case ERigSpaceType::Control:
-				{
-					if (ChildSpace.ParentIndex == InParentIndex && InParentType == ERigElementType::Control)
-					{
-						return true;
-					}
-					return IsParentedTo(ERigElementType::Control, ChildSpace.ParentIndex, InParentType, InParentIndex);
-				}
-			}
-		}
-		case ERigElementType::Control:
-		{
-			const FRigControl& ChildControl = ControlHierarchy[InChildIndex];
-			switch (InParentType)
-			{
-				case ERigElementType::Space:
-				{
-					if (ChildControl.SpaceIndex == InParentIndex)
-					{
-						return true;
-					}
-					// no break - fall through to next cases
-				}
-				case ERigElementType::Control:
-				case ERigElementType::Bone:
-				{
-					if (ChildControl.SpaceIndex != INDEX_NONE)
-					{
-						return IsParentedTo(ERigElementType::Space, ChildControl.SpaceIndex, InParentType, InParentIndex);
-					}
-					else if (ChildControl.ParentIndex != INDEX_NONE)
-					{
-						if (ChildControl.ParentIndex == InParentIndex && InParentType == ERigElementType::Control)
-						{
-							return true;
-						}
-						return IsParentedTo(ERigElementType::Control, ChildControl.ParentIndex, InParentType, InParentIndex);
-					}
-				}
-				case ERigElementType::Curve:
-				{
-					return false;
-				}
-			}
-		}
-	}
-
-	return false;
-}
-
-void FRigHierarchyContainer::UpdateDepthIndexIfRequired() const
-{
-	if (DepthIndexByKey.Num() > 0)
-	{
-		return;
-	}
-
-	struct Local
-	{
-		static void CollectChildren(const FRigElementKey& Child, const FRigHierarchyContainer* Hierarchy, TMap<FRigElementKey, TArray<FRigElementKey>>& ChildMap)
-		{
-			if (!Child.IsValid())
-			{
-				return;
-			}
-
-			switch (Child.Type)
-			{
-				case ERigElementType::Bone:
-				{
-					const FRigBone& Bone = Hierarchy->BoneHierarchy[Child.Name];
-					CollectChildren(Child, Bone.GetParentElementKey(true), Hierarchy, ChildMap);
-					break;
-				}
-				case ERigElementType::Space:
-				{
-					const FRigSpace& Space = Hierarchy->SpaceHierarchy[Child.Name];
-					CollectChildren(Child, Space.GetParentElementKey(true), Hierarchy, ChildMap);
-					break;
-				}
-				case ERigElementType::Control:
-				{
-					const FRigControl& Control = Hierarchy->ControlHierarchy[Child.Name];
-					CollectChildren(Child, Control.GetSpaceElementKey(true), Hierarchy, ChildMap);
-					CollectChildren(Child, Control.GetParentElementKey(true), Hierarchy, ChildMap);
-					break;
-				}
-				default:
-				{
-					break;
-				}
-			}
-		}
-
-
-		static void CollectChildren(const FRigElementKey& Child, const FRigElementKey& Parent, const FRigHierarchyContainer* Hierarchy, TMap<FRigElementKey, TArray<FRigElementKey>>& ChildMap)
-		{
-			if (!Child.IsValid() || !Parent.IsValid())
-			{
-				return;
-			}
-			ChildMap.FindOrAdd(Parent).Add(Child);
-		}
-
-		static void VisitKey(const FRigElementKey& Key, const FRigHierarchyContainer* Hierarchy, const TMap<FRigElementKey, TArray<FRigElementKey>>& ChildMap, TMap<FRigElementKey, int32>& IndexByKey)
-		{
-			if (!Key.IsValid() || IndexByKey.Contains(Key))
-			{
-				return;
-			}
-
-			IndexByKey.Add(Key, IndexByKey.Num());
-
-			const TArray<FRigElementKey>* ChildrenPtr = ChildMap.Find(Key);
-			if (ChildrenPtr)
-			{
-				const TArray<FRigElementKey>& Children = *ChildrenPtr;
-				for (const FRigElementKey& Child : Children)
-				{
-					VisitKey(Child, Hierarchy, ChildMap, IndexByKey);
-				}
-			}
-		}
-	};
-
-	DepthIndexByKey.Reset();
-
-	TArray<FRigElementKey> AllKeys = GetAllItems(false);
-	TMap<FRigElementKey, TArray<FRigElementKey>> ChildMap;
-
-	for (const FRigElementKey& Key : AllKeys)
-	{
-		Local::CollectChildren(Key, this, ChildMap);
-	}
-
-	for (TPair<FRigElementKey, TArray<FRigElementKey>> Pair : ChildMap)
-	{
-		TArray<FRigElementKey>& Children = Pair.Value;
-
-		Algo::SortBy(Children, [&Children](int32 Index) -> int32
-		{
-			return GetTypeHash(Children[Index].Name);
-		});
-	}
-
-	for (const FRigElementKey& Key : AllKeys)
-	{
-		Local::VisitKey(Key, this, ChildMap, DepthIndexByKey);
-	}
-}
-
-void FRigHierarchyContainer::SortKeyArray(TArray<FRigElementKey>& InOutKeysToSort) const
-{
-	UpdateDepthIndexIfRequired();
-
-	InOutKeysToSort.Sort([&](const FRigElementKey& A, const FRigElementKey& B) {
-
-		const int32* IndexA = DepthIndexByKey.Find(A);
-		const int32* IndexB = DepthIndexByKey.Find(B);
-		if (IndexA == nullptr || IndexB == nullptr)
-		{
-			return false;
-		}
-
-		return *IndexA < * IndexB;
-	});
-}
-
-
-#if WITH_EDITOR
-
-FString FRigHierarchyContainer::ExportSelectionToText() const
-{
-	TArray<FRigElementKey> Selection = CurrentSelection();
-	return ExportToText(Selection);
-}
-
-FString FRigHierarchyContainer::ExportToText(const TArray<FRigElementKey>& InSelection) const
-{
-	TArray<FRigElementKey> KeysToExport;
-	KeysToExport.Append(InSelection);
-	SortKeyArray(KeysToExport);
-
-	FRigHierarchyCopyPasteContent Data;
-	for (const FRigElementKey& Key : KeysToExport)
-	{
-		Data.Types.Add(Key.Type);
-		FString Content;
-		switch (Key.Type)
-		{
-		case ERigElementType::Bone:
-		{
-			FRigBone DefaultElement;
-			FRigBone::StaticStruct()->ExportText(Content, &BoneHierarchy[Key.Name], &DefaultElement, nullptr, PPF_None, nullptr);
-			break;
-		}
-		case ERigElementType::Control:
-		{
-			FRigControl DefaultElement;
-			FRigControl::StaticStruct()->ExportText(Content, &ControlHierarchy[Key.Name], &DefaultElement, nullptr, PPF_None, nullptr);
-			break;
-		}
-		case ERigElementType::Space:
-		{
-			FRigSpace DefaultElement;
-			FRigSpace::StaticStruct()->ExportText(Content, &SpaceHierarchy[Key.Name], &DefaultElement, nullptr, PPF_None, nullptr);
-			break;
-		}
-		case ERigElementType::Curve:
-		{
-			FRigCurve DefaultElement;
-			FRigCurve::StaticStruct()->ExportText(Content, &CurveContainer[Key.Name], &DefaultElement, nullptr, PPF_None, nullptr);
-			break;
-		}
-		default:
-		{
-			ensure(false);
-			break;
-		}
-		}
-		Data.Contents.Add(Content);
-		Data.LocalTransforms.Add(GetLocalTransform(Key));
-		Data.GlobalTransforms.Add(GetGlobalTransform(Key));
-	}
-
-	FString ExportedText;
-	FRigHierarchyCopyPasteContent DefaultContent;
-	FRigHierarchyCopyPasteContent::StaticStruct()->ExportText(ExportedText, &Data, &DefaultContent, nullptr, PPF_None, nullptr);
-	return ExportedText;
-}
-
-class FRigHierarchyContainerImportErrorContext : public FOutputDevice
-{
-public:
-
-	int32 NumErrors;
-
-	FRigHierarchyContainerImportErrorContext()
-		: FOutputDevice()
-		, NumErrors(0)
-	{
-	}
-
-	virtual void Serialize(const TCHAR* V, ELogVerbosity::Type Verbosity, const class FName& Category) override
-	{
-		UE_LOG(LogControlRig, Error, TEXT("Error Importing To Hierarchy: %s"), V);
-		NumErrors++;
-	}
-};
-
-TArray<FRigElementKey> FRigHierarchyContainer::ImportFromText(const FString& InContent, ERigHierarchyImportMode InImportMode, bool bSelectNewElements, const TArray<FRigElementKey>& InSelection)
-{
-	TArray<FRigElementKey> PastedKeys;
-
-	FRigHierarchyCopyPasteContent Data;
-	FRigHierarchyContainerImportErrorContext ErrorPipe;
-	FRigHierarchyCopyPasteContent::StaticStruct()->ImportText(*InContent, &Data, nullptr, EPropertyPortFlags::PPF_None, &ErrorPipe, FRigHierarchyCopyPasteContent::StaticStruct()->GetName(), true);
-	if (ErrorPipe.NumErrors > 0)
-	{
-		return PastedKeys;
-	}
-
-	if (Data.Contents.Num() == 0 ||
-		(Data.Types.Num() != Data.Contents.Num()) ||
-		(Data.LocalTransforms.Num() != Data.Contents.Num()) ||
-		(Data.GlobalTransforms.Num() != Data.Contents.Num()))
-	{
-		return PastedKeys;
-	}
-
-	TMap<FRigElementKey, FRigElementKey> ElementMap;
-	for (const FRigBone& Element : BoneHierarchy)
-	{
-		ElementMap.Add(Element.GetElementKey(), Element.GetElementKey());
-	}
-	for (const FRigControl& Element : ControlHierarchy)
-	{
-		ElementMap.Add(Element.GetElementKey(), Element.GetElementKey());
-	}
-	for (const FRigSpace& Element : SpaceHierarchy)
-	{
-		ElementMap.Add(Element.GetElementKey(), Element.GetElementKey());
-	}
-	for (const FRigCurve& Element : CurveContainer)
-	{
-		ElementMap.Add(Element.GetElementKey(), Element.GetElementKey());
-	}
-
-	TArray<TSharedPtr<FRigElement>> Elements;
-	for (int32 Index = 0; Index < Data.Types.Num(); Index++)
-	{
-		ErrorPipe.NumErrors = 0;
-		TSharedPtr<FRigElement> NewElement;
-		switch (Data.Types[Index])
-		{
-			case ERigElementType::Bone:
-			{
-				NewElement = MakeShared<FRigBone>();
-				FRigBone::StaticStruct()->ImportText(*Data.Contents[Index], NewElement.Get(), nullptr, EPropertyPortFlags::PPF_None, &ErrorPipe, FRigBone::StaticStruct()->GetName(), true);
-				break;
-			}
-			case ERigElementType::Control:
-			{
-				NewElement = MakeShared<FRigControl>();
-				FRigControl::StaticStruct()->ImportText(*Data.Contents[Index], NewElement.Get(), nullptr, EPropertyPortFlags::PPF_None, &ErrorPipe, FRigControl::StaticStruct()->GetName(), true);
-				break;
-			}
-			case ERigElementType::Space:
-			{
-				NewElement = MakeShared<FRigSpace>();
-				FRigSpace::StaticStruct()->ImportText(*Data.Contents[Index], NewElement.Get(), nullptr, EPropertyPortFlags::PPF_None, &ErrorPipe, FRigSpace::StaticStruct()->GetName(), true);
-				break;
-			}
-			case ERigElementType::Curve:
-			{
-				NewElement = MakeShared<FRigCurve>();
-				FRigCurve::StaticStruct()->ImportText(*Data.Contents[Index], NewElement.Get(), nullptr, EPropertyPortFlags::PPF_None, &ErrorPipe, FRigCurve::StaticStruct()->GetName(), true);
-				break;
-			}
-			default:
-			{
-				ensure(false);
-				break;
-			}
-		}
-
-		if (ErrorPipe.NumErrors > 0)
-		{
-			return PastedKeys;
-		}
-
-		Elements.Add(NewElement);
-	}
-
-	TArray<FRigElementKey> Selection = InSelection;
-	if (InSelection.IsEmpty())
-	{
-		Selection = CurrentSelection();
-	}
-	else
-	{
-		bSelectNewElements = false;
-	}
-
-	switch (InImportMode)
-	{
-		case ERigHierarchyImportMode::Append:
-		{
-			for (int32 Index = 0; Index < Data.Types.Num(); Index++)
-			{
-				switch (Data.Types[Index])
-				{
-					case ERigElementType::Bone:
-					{
-						const FRigBone& Element = *static_cast<FRigBone*>(Elements[Index].Get());
-
-						FName ParentName = NAME_None;
-						if (const FRigElementKey* ParentKey = ElementMap.Find(Element.GetParentElementKey(true /* force */)))
-						{
-							ParentName = ParentKey->Name;
-						}
-
-						FRigBone& NewElement = BoneHierarchy.Add(Element.Name, ParentName, ERigBoneType::User, Element.InitialTransform, Element.LocalTransform, Data.GlobalTransforms[Index]);
-						ElementMap.FindOrAdd(Element.GetElementKey()) = NewElement.GetElementKey();
-						PastedKeys.Add(NewElement.GetElementKey());
-						break;
-					}
-					case ERigElementType::Control:
-					{
-						const FRigControl& Element = *static_cast<FRigControl*>(Elements[Index].Get());
-
-						FName ParentName = NAME_None;
-						if (const FRigElementKey* ParentKey = ElementMap.Find(Element.GetParentElementKey(true /* force */)))
-						{
-							ParentName = ParentKey->Name;
-						}
-
-						FName SpaceName = NAME_None;
-						if (const FRigElementKey* SpaceKey = ElementMap.Find(Element.GetSpaceElementKey(true /* force */)))
-						{
-							SpaceName = SpaceKey->Name;
-						}
-
-						FRigControl& NewElement = ControlHierarchy.Add(Element.Name, Element.ControlType, ParentName, SpaceName, Element.OffsetTransform, Element.InitialValue, Element.GizmoName, Element.GizmoTransform, Element.GizmoColor);
-
-						// copy additional members
-						NewElement.DisplayName = Element.DisplayName;
-						NewElement.bAnimatable = Element.bAnimatable;
-						NewElement.PrimaryAxis = Element.PrimaryAxis;
-						NewElement.bLimitTranslation = Element.bLimitTranslation;
-						NewElement.bLimitRotation = Element.bLimitRotation;
-						NewElement.bLimitScale= Element.bLimitScale;
-						NewElement.MinimumValue = Element.MinimumValue;
-						NewElement.MaximumValue = Element.MaximumValue;
-						NewElement.bDrawLimits = Element.bDrawLimits;
-						NewElement.bGizmoEnabled = Element.bGizmoEnabled;
-						NewElement.bGizmoVisible = Element.bGizmoVisible;
-						NewElement.ControlEnum = Element.ControlEnum;
-
-						ElementMap.FindOrAdd(Element.GetElementKey()) = NewElement.GetElementKey();
-						PastedKeys.Add(NewElement.GetElementKey());
-
-						break;
-					}
-					case ERigElementType::Space:
-					{
-						const FRigSpace& Element = *static_cast<FRigSpace*>(Elements[Index].Get());
-
-						FName ParentName = NAME_None;
-						if (const FRigElementKey* ParentKey = ElementMap.Find(Element.GetParentElementKey(true /* force */)))
-						{
-							ParentName = ParentKey->Name;
-						}
-
-						FRigSpace& NewElement = SpaceHierarchy.Add(Element.Name, Element.SpaceType, ParentName, Element.InitialTransform);
-						ElementMap.FindOrAdd(Element.GetElementKey()) = NewElement.GetElementKey();
-						PastedKeys.Add(NewElement.GetElementKey());
-						break;
-					}
-					case ERigElementType::Curve:
-					{
-						const FRigCurve& Element = *static_cast<FRigCurve*>(Elements[Index].Get());
-						FRigCurve& NewElement = CurveContainer.Add(Element.Name, Element.Value);
-						ElementMap.FindOrAdd(Element.GetElementKey()) = NewElement.GetElementKey();
-						PastedKeys.Add(NewElement.GetElementKey());
-						break;
-					}
-					default:
-					{
-						ensure(false);
-						break;
-					}
-				}
-			}
-			break;
-		}
-		case ERigHierarchyImportMode::ReplaceInitialLocalTransform:
-		case ERigHierarchyImportMode::ReplaceInitialGlobalTransform:
-		case ERigHierarchyImportMode::ReplaceCurrentLocalTransform:
-		case ERigHierarchyImportMode::ReplaceCurrentGlobalTransform:
-		{
-			if (Selection.Num() == 0)
-			{
-				for (const TSharedPtr<FRigElement>& Element : Elements)
-				{
-					Selection.Add(Element->GetElementKey());
-				}
-			}
-			else if (Selection.Num() != Data.Types.Num())
-			{
-				// todo: error message
-				return PastedKeys;
-			}
-
-			for (int32 Index = 0; Index < Data.Types.Num(); Index++)
-			{
-				if(InImportMode == ERigHierarchyImportMode::ReplaceInitialLocalTransform)
-				{
-					Data.LocalTransforms[Index].NormalizeRotation();
-					switch (Selection[Index].Type)
-					{
-						case ERigElementType::Bone:
-						case ERigElementType::Space:
-						{
-							SetInitialTransform(Selection[Index], Data.LocalTransforms[Index]);
-							SetLocalTransform(Selection[Index], Data.LocalTransforms[Index]);
-							break;
-						}
-						case ERigElementType::Control:
-						{
-							int32 ControlIndex = GetIndex(Selection[Index]);
-							ControlHierarchy.SetControlOffset(ControlIndex, Data.LocalTransforms[Index]);
-							ControlHierarchy[ControlIndex].SetValueFromTransform(FTransform::Identity, ERigControlValueType::Initial);
-							ControlHierarchy[ControlIndex].SetValueFromTransform(FTransform::Identity, ERigControlValueType::Current);
-							break;
-						}
-					}
-				}
-				else if (InImportMode == ERigHierarchyImportMode::ReplaceInitialGlobalTransform)
-				{
-					Data.GlobalTransforms[Index].NormalizeRotation();
-					switch (Selection[Index].Type)
-					{
-						case ERigElementType::Bone:
-						case ERigElementType::Space:
-						{
-							SetInitialGlobalTransform(Selection[Index], Data.GlobalTransforms[Index]);
-							SetGlobalTransform(Selection[Index], Data.GlobalTransforms[Index]);
-							break;
-						}
-						case ERigElementType::Control:
-						{
-							int32 ControlIndex = GetIndex(Selection[Index]);
-							FTransform ParentTransform = ControlHierarchy.GetParentTransform(ControlIndex, false);
-							ControlHierarchy.SetControlOffset(ControlIndex, Data.GlobalTransforms[Index].GetRelativeTransform(ParentTransform));
-							ControlHierarchy[ControlIndex].SetValueFromTransform(FTransform::Identity, ERigControlValueType::Initial);
-							ControlHierarchy[ControlIndex].SetValueFromTransform(FTransform::Identity, ERigControlValueType::Current);
-							break;
-						}
-					}
-				}
-				else if (InImportMode == ERigHierarchyImportMode::ReplaceCurrentLocalTransform)
-				{
-					Data.LocalTransforms[Index].NormalizeRotation();
-					SetLocalTransform(Selection[Index], Data.LocalTransforms[Index]);
-				}
-				else // (InImportMode == ERigHierarchyImportMode::ReplaceCurrentGlobalTransform)
-				{
-					Data.GlobalTransforms[Index].NormalizeRotation();
-					SetGlobalTransform(Selection[Index], Data.GlobalTransforms[Index]);
-				}
-				PastedKeys.Add(Selection[Index]);
-			}
-			break;
-		}
-		default:
-		{
-			ensure(false);
-			break;
-		}
-	}
-
-	if (PastedKeys.Num() > 0 && InImportMode == ERigHierarchyImportMode::Append)
-	{
-		BoneHierarchy.Initialize();
-		SpaceHierarchy.Initialize();
-		ControlHierarchy.Initialize();
-		CurveContainer.Initialize();
-	}
-
-	if (bSelectNewElements && PastedKeys.Num() > 0)
-	{
-		ClearSelection();
-		for (const FRigElementKey& Key : PastedKeys)
-		{
-			Select(Key, true);
-		}
-	}
-
-	return PastedKeys;
-}
-
-TArray<FRigElementKey> FRigHierarchyContainer::DuplicateItems(const TArray<FRigElementKey>& InSelection, bool bSelectNewElements)
-{
-	FString CopiedText = ExportToText(InSelection);
-	return ImportFromText(CopiedText, ERigHierarchyImportMode::Append, bSelectNewElements);
-}
-
-TArray<FRigElementKey> FRigHierarchyContainer::MirrorItems(const TArray<FRigElementKey>& InSelection, const FRigMirrorSettings& InSettings, bool bSelectNewElements)
-{
-	TArray<FRigElementKey> OriginalItems;
-	OriginalItems.Append(InSelection);
-	SortKeyArray(OriginalItems);
-	TArray<FRigElementKey> DuplicatedItems = DuplicateItems(OriginalItems, bSelectNewElements);
-
-	if (DuplicatedItems.Num() != OriginalItems.Num())
-	{
-		return DuplicatedItems;
-	}
-
-	for (int32 Index = 0; Index < OriginalItems.Num(); Index++)
-	{
-		if (DuplicatedItems[Index].Type != OriginalItems[Index].Type)
-		{
-			return DuplicatedItems;
-		}
-	}
-
-	// mirror the transforms
-	for (int32 Index = 0; Index < OriginalItems.Num(); Index++)
-	{
-		FTransform GlobalTransform = GetGlobalTransform(OriginalItems[Index]);
-		FTransform InitialTransform = GetInitialGlobalTransform(OriginalItems[Index]);
-
-		// also mirror the offset, limits and gizmo transform
-		if (OriginalItems[Index].Type == ERigElementType::Control)
-		{
-			int32 OriginalControlIndex = ControlHierarchy.GetIndex(OriginalItems[Index].Name);
-			ensure(OriginalControlIndex != INDEX_NONE);
-			int32 DuplicatedControlIndex = ControlHierarchy.GetIndex(DuplicatedItems[Index].Name);
-			ensure(DuplicatedControlIndex != INDEX_NONE);
-
-			FRigControl& DuplicatedControl = ControlHierarchy[DuplicatedControlIndex];
-			TGuardValue<bool> DisableTranslationLimit(DuplicatedControl.bLimitTranslation, false);
-			TGuardValue<bool> DisableRotationLimit(DuplicatedControl.bLimitRotation, false);
-			TGuardValue<bool> DisableScaleLimit(DuplicatedControl.bLimitScale, false);
-
-			// mirror offset
-			FTransform OriginalGlobalOffsetTransform = ControlHierarchy.GetParentTransform(OriginalControlIndex, true /* include offset */);
-			FTransform ParentTransform = ControlHierarchy.GetParentTransform(DuplicatedControlIndex, false /* include offset */);
-			FTransform OffsetTransform = InSettings.MirrorTransform(OriginalGlobalOffsetTransform).GetRelativeTransform(ParentTransform);
-			ControlHierarchy.SetControlOffset(DuplicatedControlIndex, OffsetTransform);
-
-			// mirror limits
-			FTransform DuplicateGlobalOffsetTransform = ControlHierarchy.GetParentTransform(DuplicatedControlIndex, true /* include offset */);
-
-			for (ERigControlValueType ValueType = ERigControlValueType::Minimum;
-				ValueType <= ERigControlValueType::Maximum;
-				ValueType = (ERigControlValueType)(uint8(ValueType) + 1)
-				)
-			{
-				FTransform LimitTransform = ControlHierarchy[DuplicatedControlIndex].GetTransformFromValue(ValueType);
-				FTransform GlobalLimitTransform = LimitTransform * OriginalGlobalOffsetTransform;
-				FTransform DuplicateLimitTransform = InSettings.MirrorTransform(GlobalLimitTransform).GetRelativeTransform(DuplicateGlobalOffsetTransform);
-				ControlHierarchy[DuplicatedControlIndex].SetValueFromTransform(DuplicateLimitTransform, ValueType);
-			}
-
-			// we need to do this here to make sure that the limits don't apply ( the tguardvalue is still active within here )
-			SetGlobalTransform(DuplicatedItems[Index], InSettings.MirrorTransform(GlobalTransform));
-			SetInitialGlobalTransform(DuplicatedItems[Index], InSettings.MirrorTransform(InitialTransform));
-
-			// mirror gizmo transform
-			FTransform GlobalGizmoTransform = ControlHierarchy[OriginalControlIndex].GizmoTransform * OriginalGlobalOffsetTransform;
-			DuplicatedControl.GizmoTransform = InSettings.MirrorTransform(GlobalGizmoTransform).GetRelativeTransform(DuplicateGlobalOffsetTransform);
-		}
-		else
-		{
-			SetGlobalTransform(DuplicatedItems[Index], InSettings.MirrorTransform(GlobalTransform));
-			SetInitialGlobalTransform(DuplicatedItems[Index], InSettings.MirrorTransform(InitialTransform));
-		}
-	}
-
-	// correct the names
-	if (!InSettings.OldName.IsEmpty() && !InSettings.NewName.IsEmpty())
-	{
-		for (int32 Index = 0; Index < DuplicatedItems.Num(); Index++)
-		{
-			FName OldName = OriginalItems[Index].Name;
-			FString OldNameStr = OldName.ToString();
-			FString NewNameStr = OldNameStr.Replace(*InSettings.OldName, *InSettings.NewName, ESearchCase::CaseSensitive);
-			if (NewNameStr != OldNameStr)
-			{
-				FName NewName(*NewNameStr);
-				switch (DuplicatedItems[Index].Type)
-				{
-					case ERigElementType::Bone:
-					{
-						DuplicatedItems[Index].Name = BoneHierarchy.Rename(DuplicatedItems[Index].Name, NewName);
-						break;
-					}
-					case ERigElementType::Space:
-					{
-						DuplicatedItems[Index].Name = SpaceHierarchy.Rename(DuplicatedItems[Index].Name, NewName);
-						break;
-					}
-					case ERigElementType::Control:
-					{
-						DuplicatedItems[Index].Name = ControlHierarchy.Rename(DuplicatedItems[Index].Name, NewName);
-						break;
-					}
-					default:
-					{
-						break;
-					}
-				}
-			}
-		}
-	}
-
-	return DuplicatedItems;
-}
-
-TArray<FRigElementKey> FRigHierarchyContainer::CurrentSelection() const
-{
-	TArray<FRigElementKey> Selection;
-	for (const FRigBone& Element : BoneHierarchy)
-	{
-		if (BoneHierarchy.IsSelected(Element.Name))
-		{
-			Selection.Add(Element.GetElementKey());
-		}
-	}
-	for (const FRigControl& Element : ControlHierarchy)
-	{
-		if (ControlHierarchy.IsSelected(Element.Name))
-		{
-			Selection.Add(Element.GetElementKey());
-		}
-	}
-	for (const FRigSpace& Element : SpaceHierarchy)
-	{
-		if (SpaceHierarchy.IsSelected(Element.Name))
-		{
-			Selection.Add(Element.GetElementKey());
-		}
-	}
-	for (const FRigCurve& Element : CurveContainer)
-	{
-		if (CurveContainer.IsSelected(Element.Name))
-		{
-			Selection.Add(Element.GetElementKey());
-		}
-	}
-
-	SortKeyArray(Selection);
-
-	return Selection;
-}
-
-#endif
-
-TArray<FRigElementKey> FRigHierarchyContainer::GetAllItems(bool bSort) const
-{
-	TArray<FRigElementKey> Items;
-
-	for (const FRigControl& Element : ControlHierarchy)
-	{
-		Items.Add(Element.GetElementKey());
-	}
-	for (const FRigBone& Element : BoneHierarchy)
-	{
-		Items.Add(Element.GetElementKey());
-	}
-	
-	for (const FRigSpace& Element : SpaceHierarchy)
-	{
-		Items.Add(Element.GetElementKey());
-	}
-	for (const FRigCurve& Element : CurveContainer)
-	{
-		Items.Add(Element.GetElementKey());
-	}
-	
-	if (bSort)
-	{
-		SortKeyArray(Items);
-	}
-
-	return Items;
-}
-
-bool FRigHierarchyContainer::Select(const FRigElementKey& InKey, bool bSelect)
-{
-	switch(InKey.Type)
-	{
-		case ERigElementType::Bone:
-		{
-			return BoneHierarchy.Select(InKey.Name, bSelect);
-		}
-		case ERigElementType::Space:
-		{
-			return SpaceHierarchy.Select(InKey.Name, bSelect);
-		}
-		case ERigElementType::Control:
-		{
-			return ControlHierarchy.Select(InKey.Name, bSelect);
-		}
-		case ERigElementType::Curve:
-		{
-			return CurveContainer.Select(InKey.Name, bSelect);
-		}
-	}
-	return false;
-}
-
-bool FRigHierarchyContainer::ClearSelection()
-{
-	ClearSelection(ERigElementType::Bone);
-	ClearSelection(ERigElementType::Control);
-	ClearSelection(ERigElementType::Space);
-	ClearSelection(ERigElementType::Curve);
-	return true;
-}
-
-bool FRigHierarchyContainer::ClearSelection(ERigElementType InElementType)
-{
-	switch(InElementType)
-	{
-		case ERigElementType::Bone:
-		{
-			return BoneHierarchy.ClearSelection();
-		}
-		case ERigElementType::Space:
-		{
-			return SpaceHierarchy.ClearSelection();
-		}
-		case ERigElementType::Control:
-		{
-			return ControlHierarchy.ClearSelection();
-		}
-		case ERigElementType::Curve:
-		{
-			return CurveContainer.ClearSelection();
-		}
-	}
-	return false;
-}
-
-bool FRigHierarchyContainer::IsSelected(const FRigElementKey& InKey) const
-{
-	switch(InKey.Type)
-	{
-		case ERigElementType::Bone:
-		{
-			return BoneHierarchy.IsSelected(InKey.Name);
-		}
-		case ERigElementType::Space:
-		{
-			return SpaceHierarchy.IsSelected(InKey.Name);
-		}
-		case ERigElementType::Control:
-		{
-			return ControlHierarchy.IsSelected(InKey.Name);
-		}
-		case ERigElementType::Curve:
-		{
-			return CurveContainer.IsSelected(InKey.Name);
-		}
-	}
-	return false;
-}
-
-void FRigHierarchyContainer::SendEvent(const FRigEventContext& InEvent, bool bAsyncronous)
-{
-	if(OnEventReceived.IsBound())
-	{
-		FRigHierarchyContainer* LocalContainer = this;
-		FRigEventDelegate& Delegate = OnEventReceived;
-
-		if (bAsyncronous)
-		{
-			FFunctionGraphTask::CreateAndDispatchWhenReady([LocalContainer, Delegate, InEvent]()
-			{
-				Delegate.Broadcast(LocalContainer, InEvent);
-			}, TStatId(), NULL, ENamedThreads::GameThread);
-		}
-		else
-		{
-			Delegate.Broadcast(LocalContainer, InEvent);
-		}
-	}
-}
-
-FRigPose FRigHierarchyContainer::GetPose() const
-{
-	FRigPose Pose;
-	AppendToPose(Pose);
-	return Pose;
-}
-
-void FRigHierarchyContainer::SetPose(FRigPose& InPose)
-{
-	BoneHierarchy.SetPose(InPose);
-	SpaceHierarchy.SetPose(InPose);
-	ControlHierarchy.SetPose(InPose);
-	CurveContainer.SetPose(InPose);
-}
-
-void FRigHierarchyContainer::AppendToPose(FRigPose& InOutPose) const
-{
-	BoneHierarchy.AppendToPose(InOutPose);
-	SpaceHierarchy.AppendToPose(InOutPose);
-	ControlHierarchy.AppendToPose(InOutPose);
-	CurveContainer.AppendToPose(InOutPose);
-}
-
-void FRigHierarchyContainer::SanitizeName(FString& InOutName)
-{
-	// Sanitize the name
-	for (int32 i = 0; i < InOutName.Len(); ++i)
-	{
-		TCHAR& C = InOutName[i];
-
-		const bool bGoodChar =
-			((C >= 'A') && (C <= 'Z')) || ((C >= 'a') && (C <= 'z')) ||		// A-Z (upper and lowercase) anytime
-			(C == '_') || (C == '-') ||										// _ or - anytime
-			((i > 0) && (C >= '0') && (C <= '9'));							// 0-9 after the first character
-
-		if (!bGoodChar)
-		{
-			C = '_';
-		}
-	}
-
-	if (InOutName.Len() > GetMaxNameLength())
-	{
-		InOutName.LeftChopInline(InOutName.Len() - GetMaxNameLength());
-	}
-}
-
-FName FRigHierarchyContainer::GetSanitizedName(const FString& InName)
-{
-	FString Name = InName;
-	SanitizeName(Name);
-
-	if (Name.IsEmpty())
-	{
-		return NAME_None;
-	}
-
-	return *Name;
-}
-
-FName FRigHierarchyContainer::GetSafeNewName(const FString& InName, TFunction<bool(const FString&)> IsNameAvailableFunction)
-{
-	check(IsNameAvailableFunction);
-
-	FString SanitizedName = InName;
-	SanitizeName(SanitizedName);
-	FString Name = SanitizedName;
-
-	int32 Suffix = 1;
-	while (!IsNameAvailableFunction(Name))
-	{
-		FString BaseString = SanitizedName;
-		if (BaseString.Len() > GetMaxNameLength() - 4)
-		{
-			BaseString.LeftChopInline(BaseString.Len() - (GetMaxNameLength() - 4));
-		}
-		Name = *FString::Printf(TEXT("%s_%d"), *BaseString, ++Suffix);
-	}
-	return *Name;
-}
-
-bool FRigHierarchyContainer::IsNameAvailable(const FString& InName, ERigElementType InElementType, FString* OutErrorMessage)
-{
-	FString UnsanitizedName = InName;
-	if (UnsanitizedName.Len() > GetMaxNameLength())
-	{
-		if (OutErrorMessage)
-		{
-			*OutErrorMessage = TEXT("Name too long.");
-		}
-		return false;
-	}
-
-	FString SanitizedName = UnsanitizedName;
-	SanitizeName(SanitizedName);
-
-	if (SanitizedName != UnsanitizedName)
-	{
-		if (OutErrorMessage)
-		{
-			*OutErrorMessage = TEXT("Name contains invalid characters.");
-		}
-		return false;
-	}
-
-	if (GetIndex(FRigElementKey(*InName, InElementType)) != INDEX_NONE)
-	{
-		if (OutErrorMessage)
-		{
-			*OutErrorMessage = TEXT("Name already used.");
-		}
-		return false;
-	}
-
-	return true;
-}
-
