Control Rig: New rig hierarchy data structure refactoring

diff --git a/Engine/Plugins/Experimental/ControlRig/Source/ControlRig/Private/Rigs/RigHierarchyController.cpp b/Engine/Plugins/Experimental/ControlRig/Source/ControlRig/Private/Rigs/RigHierarchyController.cpp
new file mode 100644
--- /dev/null
+++ b/Engine/Plugins/Experimental/ControlRig/Source/ControlRig/Private/Rigs/RigHierarchyController.cpp
@@ -0,0 +1,1777 @@
+// Copyright Epic Games, Inc. All Rights Reserved.
+
+#include "Rigs/RigHierarchyController.h"
+
+#if WITH_EDITOR
+#include "Framework/Notifications/NotificationManager.h"
+#include "Widgets/Notifications/SNotificationList.h"
+#include "EditorStyleSet.h"
+#include "ScopedTransaction.h"
+#include "Engine/SkeletalMesh.h"
+#endif
+
+#include "ControlRig.h"
+
+////////////////////////////////////////////////////////////////////////////////
+// URigHierarchyController
+////////////////////////////////////////////////////////////////////////////////
+
+URigHierarchyController::~URigHierarchyController()
+{
+}
+
+void URigHierarchyController::SetHierarchy(URigHierarchy* InHierarchy)
+{
+	if(Hierarchy)
+	{
+		if(!Hierarchy->HasAnyFlags(RF_BeginDestroyed) && Hierarchy->IsValidLowLevel())
+		{
+			Hierarchy->OnModified().RemoveAll(this);
+		}
+		Hierarchy->LastControllerPtr.Reset();
+	}
+
+	Hierarchy = InHierarchy;
+
+	if(Hierarchy)
+	{
+		Hierarchy->OnModified().AddUObject(this, &URigHierarchyController::HandleHierarchyModified);
+		Hierarchy->LastControllerPtr = this;
+	}
+}
+
+bool URigHierarchyController::SelectElement(FRigElementKey InKey, bool bSelect, bool bClearSelection)
+{
+	if(!IsValid())
+	{
+		return false;
+	}
+
+	if(bClearSelection)
+	{
+		TArray<FRigElementKey> KeysToSelect;
+		KeysToSelect.Add(InKey);
+		return SetSelection(KeysToSelect);
+	}
+
+	if(URigHierarchy* HierarchyForSelection = Hierarchy->HierarchyForSelectionPtr.Get())
+	{
+		if(URigHierarchyController* ControllerForSelection = HierarchyForSelection->GetController())
+		{
+			return ControllerForSelection->SelectElement(InKey, bSelect, bClearSelection);
+		}
+	}
+
+	FRigBaseElement* Element = Hierarchy->Find(InKey);
+	if(Element == nullptr)
+	{
+		return false;
+	}
+
+	if(Element->bSelected == bSelect)
+	{
+		return false;
+	}
+
+	Element->bSelected = bSelect;
+
+	if(Element->bSelected)
+	{
+		Notify(ERigHierarchyNotification::ElementSelected, Element);
+	}
+	else
+	{
+		Notify(ERigHierarchyNotification::ElementDeselected, Element);
+	}
+
+	return true;
+}
+
+bool URigHierarchyController::SetSelection(TArray<FRigElementKey> InKeys)
+{
+	if(!IsValid())
+	{
+		return false;
+	}
+
+	if(URigHierarchy* HierarchyForSelection = Hierarchy->HierarchyForSelectionPtr.Get())
+	{
+		if(URigHierarchyController* ControllerForSelection = HierarchyForSelection->GetController())
+		{
+			return ControllerForSelection->SetSelection(InKeys);
+		}
+	}
+
+	TArray<FRigElementKey> PreviousSelection = Hierarchy->GetSelectedKeys();
+
+	for(const FRigElementKey& KeyToDeselect : PreviousSelection)
+	{
+		if(!InKeys.Contains(KeyToDeselect))
+		{
+			if(!SelectElement(KeyToDeselect, false))
+			{
+				return false;
+			}
+		}
+	}
+
+	for(const FRigElementKey& KeyToSelect : InKeys)
+	{
+		if(!PreviousSelection.Contains(KeyToSelect))
+		{
+			if(!SelectElement(KeyToSelect, true))
+			{
+				return false;
+			}
+		}
+	}
+
+	return true;
+}
+
+FRigElementKey URigHierarchyController::AddBone(FName InName, FRigElementKey InParent, FTransform InTransform,
+                                                bool bTransformInGlobal, ERigBoneType InBoneType, bool bSetupUndo)
+{
+	if(!IsValid())
+	{
+		return FRigElementKey();
+	}
+
+#if WITH_EDITOR
+	TSharedPtr<FScopedTransaction> TransactionPtr;
+	if(bSetupUndo)
+	{
+		TransactionPtr = MakeShared<FScopedTransaction>(NSLOCTEXT("RigHierarchyController", "Add Bone", "Add Bone"));
+        Hierarchy->Modify();
+	}
+#endif
+
+	FRigBoneElement* NewElement = new FRigBoneElement();
+	NewElement->Key.Type = ERigElementType::Bone;
+	NewElement->Key.Name = Hierarchy->GetSafeNewName(InName.ToString(), NewElement->Key.Type);
+	NewElement->BoneType = InBoneType;
+	AddElement(NewElement, Hierarchy->Get(Hierarchy->GetIndex(InParent)), true);
+
+	if(bTransformInGlobal)
+	{
+		Hierarchy->SetTransform(NewElement, InTransform, ERigTransformType::InitialGlobal, true, false);
+	}
+	else
+	{
+		Hierarchy->SetTransform(NewElement, InTransform, ERigTransformType::InitialLocal, true, false);
+	}
+
+	NewElement->Pose.Current = NewElement->Pose.Initial;
+
+#if WITH_EDITOR
+	TransactionPtr.Reset();
+#endif
+
+	return NewElement->Key;
+}
+
+FRigElementKey URigHierarchyController::AddSpace(FName InName, FRigElementKey InParent, FTransform InTransform,
+	bool bTransformInGlobal, bool bSetupUndo)
+{
+	if(!IsValid())
+	{
+		return FRigElementKey();
+	}
+
+#if WITH_EDITOR
+	TSharedPtr<FScopedTransaction> TransactionPtr;
+	if(bSetupUndo)
+	{
+		TransactionPtr = MakeShared<FScopedTransaction>(NSLOCTEXT("RigHierarchyController", "Add Space", "Add Space"));
+		Hierarchy->Modify();
+	}
+#endif
+
+	FRigSpaceElement* NewElement = new FRigSpaceElement();
+	NewElement->Key.Type = ERigElementType::Space;
+	NewElement->Key.Name = Hierarchy->GetSafeNewName(InName.ToString(), NewElement->Key.Type);
+	AddElement(NewElement, Hierarchy->Get(Hierarchy->GetIndex(InParent)), false);
+
+	if(bTransformInGlobal)
+	{
+		Hierarchy->SetTransform(NewElement, InTransform, ERigTransformType::InitialGlobal, true, false);
+	}
+	else
+	{
+		Hierarchy->SetTransform(NewElement, InTransform, ERigTransformType::InitialLocal, true, false);
+	}
+
+	NewElement->Parent.MarkDirty(ERigTransformType::InitialGlobal);
+	NewElement->Parent.Current = NewElement->Parent.Initial;
+	NewElement->Pose.Current = NewElement->Pose.Initial;
+
+#if WITH_EDITOR
+	TransactionPtr.Reset();
+#endif
+
+	return NewElement->Key;
+}
+
+FRigElementKey URigHierarchyController::AddControl(
+	FName InName,
+	FRigElementKey InParent,
+	FRigControlSettings InSettings,
+	FRigControlValue InValue,
+	FTransform InOffsetTransform,
+	FTransform InGizmoTransform,
+	bool bSetupUndo)
+{
+	if(!IsValid())
+	{
+		return FRigElementKey();
+	}
+
+#if WITH_EDITOR
+	TSharedPtr<FScopedTransaction> TransactionPtr;
+	if(bSetupUndo)
+	{
+		TransactionPtr = MakeShared<FScopedTransaction>(NSLOCTEXT("RigHierarchyController", "Add Control", "Add Control"));
+		Hierarchy->Modify();
+	}
+#endif
+
+	FRigControlElement* NewElement = new FRigControlElement();
+	NewElement->Key.Type = ERigElementType::Control;
+	NewElement->Key.Name = Hierarchy->GetSafeNewName(InName.ToString(), NewElement->Key.Type);
+	NewElement->Settings = InSettings;
+	AddElement(NewElement, Hierarchy->Get(Hierarchy->GetIndex(InParent)), false);
+	
+	NewElement->Offset.Set(ERigTransformType::InitialLocal, InOffsetTransform);  
+	NewElement->Gizmo.Set(ERigTransformType::InitialLocal, InGizmoTransform);  
+	Hierarchy->SetControlValue(NewElement, InValue, ERigControlValueType::Initial, false);
+
+	NewElement->Parent.MarkDirty(ERigTransformType::InitialGlobal);
+	NewElement->Parent.Current = NewElement->Parent.Initial;
+	NewElement->Offset.Current = NewElement->Offset.Initial;
+	NewElement->Pose.Current = NewElement->Pose.Initial;
+	NewElement->Gizmo.Current = NewElement->Gizmo.Initial;
+
+#if WITH_EDITOR
+	TransactionPtr.Reset();
+#endif
+
+	return NewElement->Key;
+}
+
+FRigElementKey URigHierarchyController::AddCurve(FName InName, float InValue, bool bSetupUndo)
+{
+	if(!IsValid())
+	{
+		return FRigElementKey();
+	}
+
+#if WITH_EDITOR
+	TSharedPtr<FScopedTransaction> TransactionPtr;
+	if(bSetupUndo)
+	{
+		TransactionPtr = MakeShared<FScopedTransaction>(NSLOCTEXT("RigHierarchyController", "Add Curve", "Add Curve"));
+		Hierarchy->Modify();
+	}
+#endif
+
+	FRigCurveElement* NewElement = new FRigCurveElement();
+	NewElement->Key.Type = ERigElementType::Curve;
+	NewElement->Key.Name = Hierarchy->GetSafeNewName(InName.ToString(), NewElement->Key.Type);
+	NewElement->Value = InValue;
+	AddElement(NewElement, nullptr, false);
+
+#if WITH_EDITOR
+	TransactionPtr.Reset();
+#endif
+
+	return NewElement->Key;
+}
+
+FRigControlSettings URigHierarchyController::GetControlSettings(FRigElementKey InKey) const
+{
+	if(!IsValid())
+	{
+		return FRigControlSettings();
+	}
+
+	FRigControlElement* ControlElement = Hierarchy->Find<FRigControlElement>(InKey);
+	if(ControlElement == nullptr)
+	{
+		return FRigControlSettings();
+	}
+
+	return ControlElement->Settings;
+}
+
+bool URigHierarchyController::SetControlSettings(FRigElementKey InKey, FRigControlSettings InSettings, bool bSetupUndo) const
+{
+	if(!IsValid())
+	{
+		return false;
+	}
+
+	FRigControlElement* ControlElement = Hierarchy->Find<FRigControlElement>(InKey);
+	if(ControlElement == nullptr)
+	{
+		return false;
+	}
+
+	#if WITH_EDITOR
+	TSharedPtr<FScopedTransaction> TransactionPtr;
+	if(bSetupUndo)
+	{
+		TransactionPtr = MakeShared<FScopedTransaction>(NSLOCTEXT("RigHierarchyController", "SetControlSettings", "Set Control Settings"));
+		Hierarchy->Modify();
+	}
+#endif
+
+	ControlElement->Settings = InSettings;
+
+	FRigControlValue InitialValue = Hierarchy->GetControlValue(ControlElement, ERigControlValueType::Initial);
+	FRigControlValue CurrentValue = Hierarchy->GetControlValue(ControlElement, ERigControlValueType::Current);
+
+	ControlElement->Settings.ApplyLimits(InitialValue);
+	ControlElement->Settings.ApplyLimits(CurrentValue);
+
+	Hierarchy->Notify(ERigHierarchyNotification::ControlSettingChanged, ControlElement);
+
+	Hierarchy->SetControlValue(ControlElement, InitialValue, ERigControlValueType::Initial, bSetupUndo);
+	Hierarchy->SetControlValue(ControlElement, CurrentValue, ERigControlValueType::Current, bSetupUndo);
+
+#if WITH_EDITOR
+    TransactionPtr.Reset();
+#endif
+
+	return true;
+}
+
+TArray<FRigElementKey> URigHierarchyController::ImportBones(const FReferenceSkeleton& InSkeleton,
+                                                            const FName& InNameSpace, bool bReplaceExistingBones, bool bRemoveObsoleteBones, bool bSelectBones,
+                                                            bool bSetupUndo)
+{
+	TArray<FRigElementKey> AddedBones;
+
+	if(!IsValid())
+	{
+		return AddedBones;
+	}
+	
+	TArray<FRigElementKey> BonesToSelect;
+	TMap<FName, FName> BoneNameMap;
+
+	Hierarchy->ResetPoseToInitial();
+
+	const TArray<FMeshBoneInfo>& BoneInfos = InSkeleton.GetRefBoneInfo();
+	const TArray<FTransform>& BonePoses = InSkeleton.GetRefBonePose();
+
+	struct Local
+	{
+		static FName DetermineBoneName(const FName& InBoneName, const FName& InLocalNameSpace)
+		{
+			if (InLocalNameSpace == NAME_None || InBoneName == NAME_None)
+			{
+				return InBoneName;
+			}
+			return *FString::Printf(TEXT("%s_%s"), *InLocalNameSpace.ToString(), *InBoneName.ToString());
+		}
+	};
+
+	if (bReplaceExistingBones)
+	{
+		GetHierarchy()->ForEach<FRigBoneElement>([&BoneNameMap](FRigBoneElement* BoneElement) -> bool
+        {
+			BoneNameMap.Add(BoneElement->GetName(), BoneElement->GetName());
+			return true;
+		});
+
+		for (int32 Index = 0; Index < BoneInfos.Num(); ++Index)
+		{
+			const FRigElementKey ExistingBoneKey(BoneInfos[Index].Name, ERigElementType::Bone);
+			const int32 ExistingBoneIndex = Hierarchy->GetIndex(ExistingBoneKey);
+			
+			const FName DesiredBoneName = Local::DetermineBoneName(BoneInfos[Index].Name, InNameSpace);
+			FName ParentName = (BoneInfos[Index].ParentIndex != INDEX_NONE) ? BoneInfos[BoneInfos[Index].ParentIndex].Name : NAME_None;
+			ParentName = Local::DetermineBoneName(ParentName, InNameSpace);
+
+			const FName* MappedParentNamePtr = BoneNameMap.Find(ParentName);
+			if (MappedParentNamePtr)
+			{
+				ParentName = *MappedParentNamePtr;
+			}
+
+			const FRigElementKey ParentKey(ParentName, ERigElementType::Bone);
+
+			// if this bone already exists
+			if (ExistingBoneIndex != INDEX_NONE)
+			{
+				const int32 ParentIndex = Hierarchy->GetIndex(ParentKey);
+				
+				// check it's parent
+				if (ParentIndex != INDEX_NONE)
+				{
+					SetParent(ExistingBoneKey, ParentKey, bSetupUndo);
+				}
+
+				Hierarchy->SetLocalTransform(ExistingBoneIndex, BonePoses[Index], true, bSetupUndo);
+
+				BonesToSelect.Add(ExistingBoneKey);
+			}
+			else
+			{
+				const FRigElementKey AddedBoneKey = AddBone(DesiredBoneName, ParentKey, BonePoses[Index], false, ERigBoneType::Imported, bSetupUndo);
+				BoneNameMap.Add(DesiredBoneName, AddedBoneKey.Name);
+				AddedBones.Add(AddedBoneKey);
+				BonesToSelect.Add(AddedBoneKey);
+			}
+		}
+
+	}
+	else // import all as new
+	{
+		for (int32 Index = 0; Index < BoneInfos.Num(); ++Index)
+		{
+			FName DesiredBoneName = Local::DetermineBoneName(BoneInfos[Index].Name, InNameSpace);
+			FName ParentName = (BoneInfos[Index].ParentIndex != INDEX_NONE) ? BoneInfos[BoneInfos[Index].ParentIndex].Name : NAME_None;
+			ParentName = Local::DetermineBoneName(ParentName, InNameSpace);
+
+			const FName* MappedParentNamePtr = BoneNameMap.Find(ParentName);
+			if (MappedParentNamePtr)
+			{
+				ParentName = *MappedParentNamePtr;
+			}
+
+			const FRigElementKey ParentKey(ParentName, ERigElementType::Bone);
+			const FRigElementKey AddedBoneKey = AddBone(DesiredBoneName, ParentKey, BonePoses[Index], false, ERigBoneType::Imported, bSetupUndo);
+			BoneNameMap.Add(DesiredBoneName, AddedBoneKey.Name);
+			AddedBones.Add(AddedBoneKey);
+			BonesToSelect.Add(AddedBoneKey);
+		}
+	}
+
+	if (bReplaceExistingBones && bRemoveObsoleteBones)
+	{
+		TMap<FName, int32> BoneNameToIndexInSkeleton;
+		for (const FMeshBoneInfo& BoneInfo : BoneInfos)
+		{
+			FName DesiredBoneName = Local::DetermineBoneName(BoneInfo.Name, InNameSpace);
+			BoneNameToIndexInSkeleton.Add(DesiredBoneName, BoneNameToIndexInSkeleton.Num());
+		}
+		
+		TArray<FRigElementKey> BonesToDelete;
+		GetHierarchy()->ForEach<FRigBoneElement>([BoneNameToIndexInSkeleton, &BonesToDelete](FRigBoneElement* BoneElement) -> bool
+        {
+            if (!BoneNameToIndexInSkeleton.Contains(BoneElement->GetName()))
+			{
+				if (BoneElement->BoneType == ERigBoneType::Imported)
+				{
+					BonesToDelete.Add(BoneElement->GetKey());
+				}
+			}
+			return true;
+		});
+
+		for (const FRigElementKey& BoneToDelete : BonesToDelete)
+		{
+			TArray<FRigElementKey> Children = Hierarchy->GetChildren(BoneToDelete);
+			Algo::Reverse(Children);
+			
+			for (const FRigElementKey& Child : Children)
+			{
+				if(BonesToDelete.Contains(Child))
+				{
+					continue;
+				}
+				RemoveAllParents(Child, true, bSetupUndo);
+			}
+		}
+
+		for (const FRigElementKey& BoneToDelete : BonesToDelete)
+		{
+			RemoveElement(BoneToDelete);
+			BonesToSelect.Remove(BoneToDelete);
+		}
+	}
+
+	if (bSelectBones)
+	{
+		SetSelection(BonesToSelect);
+	}
+
+	return AddedBones;
+
+}
+
+#if WITH_EDITOR
+
+TArray<FRigElementKey> URigHierarchyController::ImportBonesFromAsset(FString InAssetPath, FName InNameSpace,
+	bool bReplaceExistingBones, bool bRemoveObsoleteBones, bool bSelectBones, bool bSetupUndo)
+{
+	if(USkeleton* Skeleton = GetSkeletonFromAssetPath(InAssetPath))
+	{
+		return ImportBones(Skeleton, InNameSpace, bReplaceExistingBones, bRemoveObsoleteBones, bSelectBones, bSetupUndo);
+	}
+	return TArray<FRigElementKey>();
+}
+
+TArray<FRigElementKey> URigHierarchyController::ImportCurvesFromAsset(FString InAssetPath, FName InNameSpace,
+	bool bSelectCurves, bool bSetupUndo)
+{
+	if(USkeleton* Skeleton = GetSkeletonFromAssetPath(InAssetPath))
+	{
+		return ImportCurves(Skeleton, InNameSpace, bSelectCurves, bSetupUndo);
+	}
+	return TArray<FRigElementKey>();
+}
+
+USkeleton* URigHierarchyController::GetSkeletonFromAssetPath(const FString& InAssetPath)
+{
+	UObject* AssetObject = StaticLoadObject(UObject::StaticClass(), NULL, *InAssetPath, NULL, LOAD_None, NULL);
+	if(AssetObject == nullptr)
+	{
+		return nullptr;
+    }
+
+	if(USkeletalMesh* SkeletalMesh = Cast<USkeletalMesh>(AssetObject))
+	{
+		return SkeletalMesh->GetSkeleton();
+	}
+
+	if(USkeleton* Skeleton = Cast<USkeleton>(AssetObject))
+	{
+		return Skeleton;
+	}
+
+	return nullptr;
+}
+
+#endif
+
+TArray<FRigElementKey> URigHierarchyController::ImportCurves(USkeleton* InSkeleton, FName InNameSpace,
+                                                             bool bSelectCurves, bool bSetupUndo)
+{
+	check(InSkeleton);
+
+	TArray<FRigElementKey> Keys;
+	if(!IsValid())
+	{
+		return Keys;
+	}
+
+	const FSmartNameMapping* SmartNameMapping = InSkeleton->GetSmartNameContainer(USkeleton::AnimCurveMappingName);
+
+	TArray<FName> NameArray;
+	SmartNameMapping->FillNameArray(NameArray);
+
+	for (int32 Index = 0; Index < NameArray.Num(); ++Index)
+	{
+		FName Name = NameArray[Index];
+		if (!InNameSpace.IsNone())
+		{
+			Name = *FString::Printf(TEXT("%s::%s"), *InNameSpace.ToString(), *Name.ToString());
+		}
+
+		const FRigElementKey ExpectedKey(Name, ERigElementType::Curve);
+		if(Hierarchy->Contains(ExpectedKey))
+		{
+			Keys.Add(ExpectedKey);
+			continue;
+		}
+		
+		const FRigElementKey CurveKey = AddCurve(Name, 0.f, bSetupUndo);
+		Keys.Add(FRigElementKey(Name, ERigElementType::Curve));
+	}
+
+	if(bSelectCurves)
+	{
+		SetSelection(Keys);
+	}
+
+	return Keys;
+}
+
+FString URigHierarchyController::ExportSelectionToText() const
+{
+	if(!IsValid())
+	{
+		return FString();
+	}
+	return ExportToText(Hierarchy->GetSelectedKeys());
+}
+
+FString URigHierarchyController::ExportToText(TArray<FRigElementKey> InKeys) const
+{
+	if(!IsValid() || InKeys.IsEmpty())
+	{
+		return FString();
+	}
+
+	// sort the keys by traversal order
+	TArray<FRigElementKey> Keys = Hierarchy->SortKeys(InKeys);
+
+	FRigHierarchyCopyPasteContent Data;
+	for (const FRigElementKey& Key : Keys)
+	{
+		FRigBaseElement* Element = Hierarchy->Find(Key);
+		if(Element == nullptr)
+		{
+			continue;
+		}
+
+		FRigHierarchyCopyPasteContentPerElement PerElementData;
+		PerElementData.Key = Key;
+		PerElementData.Parents = Hierarchy->GetParents(Key);
+
+		if(FRigMultiParentElement* MultiParentElement = Cast<FRigMultiParentElement>(Element))
+		{
+			ensure(PerElementData.Parents.Num() == MultiParentElement->ParentWeights.Num());
+			PerElementData.ParentWeights = MultiParentElement->ParentWeights;
+		}
+		else
+		{
+			PerElementData.ParentWeights.SetNumZeroed(PerElementData.Parents.Num());
+			if(PerElementData.ParentWeights.Num() > 0)
+			{
+				PerElementData.ParentWeights[0] = 1.f;
+			}
+		}
+
+		if(FRigTransformElement* TransformElement = Cast<FRigTransformElement>(Element))
+		{
+			PerElementData.Pose.Initial.Local.Set(Hierarchy->GetTransform(TransformElement, ERigTransformType::InitialLocal));
+			PerElementData.Pose.Initial.Global.Set(Hierarchy->GetTransform(TransformElement, ERigTransformType::InitialGlobal));
+			PerElementData.Pose.Current.Local.Set(Hierarchy->GetTransform(TransformElement, ERigTransformType::CurrentLocal));
+			PerElementData.Pose.Current.Global.Set(Hierarchy->GetTransform(TransformElement, ERigTransformType::CurrentGlobal));
+		}
+
+		switch (Key.Type)
+		{
+			case ERigElementType::Bone:
+			{
+				FRigBoneElement DefaultElement;
+				FRigBoneElement::StaticStruct()->ExportText(PerElementData.Content, Element, &DefaultElement, nullptr, PPF_None, nullptr);
+				break;
+			}
+			case ERigElementType::Control:
+			{
+				FRigControlElement DefaultElement;
+				FRigControlElement::StaticStruct()->ExportText(PerElementData.Content, Element, &DefaultElement, nullptr, PPF_None, nullptr);
+				break;
+			}
+			case ERigElementType::Space:
+			{
+				FRigSpaceElement DefaultElement;
+				FRigSpaceElement::StaticStruct()->ExportText(PerElementData.Content, Element, &DefaultElement, nullptr, PPF_None, nullptr);
+				break;
+			}
+			case ERigElementType::Curve:
+			{
+				FRigCurveElement DefaultElement;
+				FRigCurveElement::StaticStruct()->ExportText(PerElementData.Content, Element, &DefaultElement, nullptr, PPF_None, nullptr);
+				break;
+			}
+			default:
+			{
+				ensure(false);
+				break;
+			}
+		}
+
+		Data.Elements.Add(PerElementData);
+	}
+
+	FString ExportedText;
+	FRigHierarchyCopyPasteContent DefaultContent;
+	FRigHierarchyCopyPasteContent::StaticStruct()->ExportText(ExportedText, &Data, &DefaultContent, nullptr, PPF_None, nullptr);
+	return ExportedText;
+}
+
+class FRigHierarchyImportErrorContext : public FOutputDevice
+{
+public:
+
+	int32 NumErrors;
+
+	FRigHierarchyImportErrorContext()
+        : FOutputDevice()
+        , NumErrors(0)
+	{
+	}
+
+	virtual void Serialize(const TCHAR* V, ELogVerbosity::Type Verbosity, const class FName& Category) override
+	{
+		UE_LOG(LogControlRig, Error, TEXT("Error Importing To Hierarchy: %s"), V);
+		NumErrors++;
+	}
+};
+
+TArray<FRigElementKey> URigHierarchyController::ImportFromText(FString InContent, bool bReplaceExistingElements, bool bSelectNewElements, bool bSetupUndo)
+{
+	TArray<FRigElementKey> PastedKeys;
+	if(!IsValid())
+	{
+		return PastedKeys;
+	}
+
+	FRigHierarchyCopyPasteContent Data;
+	FRigHierarchyImportErrorContext ErrorPipe;
+	FRigHierarchyCopyPasteContent::StaticStruct()->ImportText(*InContent, &Data, nullptr, EPropertyPortFlags::PPF_None, &ErrorPipe, FRigHierarchyCopyPasteContent::StaticStruct()->GetName(), true);
+	if (ErrorPipe.NumErrors > 0)
+	{
+		return PastedKeys;
+	}
+
+	if (Data.Elements.Num() == 0)
+	{
+		return PastedKeys;
+	}
+
+#if WITH_EDITOR
+	TSharedPtr<FScopedTransaction> TransactionPtr;
+	if(bSetupUndo)
+	{
+		TransactionPtr = MakeShared<FScopedTransaction>(NSLOCTEXT("RigHierarchyController", "Add Bone", "Add Bone"));
+		Hierarchy->Modify();
+	}
+#endif
+
+	TMap<FRigElementKey, FRigElementKey> KeyMap;
+	for(FRigBaseElement* Element : *Hierarchy)
+	{
+		KeyMap.Add(Element->GetKey(), Element->GetKey());
+	}
+	
+	for(const FRigHierarchyCopyPasteContentPerElement& PerElementData : Data.Elements)
+	{
+		ErrorPipe.NumErrors = 0;
+
+		FRigBaseElement* NewElement = nullptr;
+
+		switch (PerElementData.Key.Type)
+		{
+			case ERigElementType::Bone:
+			{
+				NewElement = new FRigBoneElement();
+				FRigBoneElement::StaticStruct()->ImportText(*PerElementData.Content, NewElement, nullptr, EPropertyPortFlags::PPF_None, &ErrorPipe, FRigBoneElement::StaticStruct()->GetName(), true);
+				break;
+			}
+			case ERigElementType::Space:
+			{
+				NewElement = new FRigSpaceElement();
+				FRigSpaceElement::StaticStruct()->ImportText(*PerElementData.Content, NewElement, nullptr, EPropertyPortFlags::PPF_None, &ErrorPipe, FRigSpaceElement::StaticStruct()->GetName(), true);
+				break;
+			}
+			case ERigElementType::Control:
+			{
+				NewElement = new FRigControlElement();
+				FRigControlElement::StaticStruct()->ImportText(*PerElementData.Content, NewElement, nullptr, EPropertyPortFlags::PPF_None, &ErrorPipe, FRigControlElement::StaticStruct()->GetName(), true);
+				break;
+			}
+			case ERigElementType::Curve:
+			{
+				NewElement = new FRigCurveElement();
+				FRigCurveElement::StaticStruct()->ImportText(*PerElementData.Content, NewElement, nullptr, EPropertyPortFlags::PPF_None, &ErrorPipe, FRigCurveElement::StaticStruct()->GetName(), true);
+				break;
+			}
+			default:
+			{
+				ensure(false);
+				break;
+			}
+		}
+
+		check(NewElement);
+		ensure(NewElement->GetKey() == PerElementData.Key);
+
+		if(bReplaceExistingElements)
+		{
+			if(FRigBaseElement* ExistingElement = Hierarchy->Find(NewElement->GetKey()))
+			{
+				ExistingElement->CopyPose(NewElement, true, true);
+
+				TArray<FRigElementKey> CurrentParents = Hierarchy->GetParents(NewElement->GetKey());
+
+				bool bUpdateParents = CurrentParents.Num() != PerElementData.Parents.Num();
+				if(!bUpdateParents)
+				{
+					for(const FRigElementKey& CurrentParent : CurrentParents)
+					{
+						if(!PerElementData.Parents.Contains(CurrentParent))
+						{
+							bUpdateParents = true;
+							break;
+						}
+					}
+				}
+
+				if(bUpdateParents)
+				{
+					RemoveAllParents(ExistingElement->GetKey(), true, bSetupUndo);
+
+					for(const FRigElementKey& NewParent : PerElementData.Parents)
+					{
+						AddParent(ExistingElement->GetKey(), NewParent, true, bSetupUndo);
+					}
+				}
+				
+				for(int32 ParentIndex = 0; ParentIndex < PerElementData.ParentWeights.Num(); ParentIndex++)
+				{
+					Hierarchy->SetParentWeight(ExistingElement, ParentIndex, PerElementData.ParentWeights[ParentIndex], true, true);
+					Hierarchy->SetParentWeight(ExistingElement, ParentIndex, PerElementData.ParentWeights[ParentIndex], false, true);
+				}
+
+				PastedKeys.Add(ExistingElement->GetKey());
+				delete NewElement;
+				continue;
+			}
+		}
+		
+		NewElement->Key.Name = Hierarchy->GetSafeNewName(NewElement->Key.Name.ToString(), NewElement->Key.Type);
+		AddElement(NewElement, nullptr, true);
+
+		KeyMap.FindOrAdd(PerElementData.Key) = NewElement->Key;
+
+		for(const FRigElementKey& OriginalParent : PerElementData.Parents)
+		{
+			FRigElementKey Parent = OriginalParent;
+			if(const FRigElementKey* RemappedParent = KeyMap.Find(Parent))
+			{
+				Parent = *RemappedParent;
+			}
+
+			AddParent(NewElement->GetKey(), Parent, true, bSetupUndo);
+		}
+		
+		for(int32 ParentIndex = 0; ParentIndex < PerElementData.ParentWeights.Num(); ParentIndex++)
+		{
+			Hierarchy->SetParentWeight(NewElement, ParentIndex, PerElementData.ParentWeights[ParentIndex], true, true);
+			Hierarchy->SetParentWeight(NewElement, ParentIndex, PerElementData.ParentWeights[ParentIndex], false, true);
+		}
+
+		PastedKeys.Add(NewElement->GetKey());
+	}
+
+	for(const FRigHierarchyCopyPasteContentPerElement& PerElementData : Data.Elements)
+	{
+		FRigElementKey MappedKey = KeyMap.FindChecked(PerElementData.Key);
+		FRigBaseElement* Element = Hierarchy->FindChecked(MappedKey);
+
+		if(FRigTransformElement* TransformElement = Cast<FRigTransformElement>(Element))
+		{
+			Hierarchy->SetTransform(TransformElement, PerElementData.Pose.Initial.Local.Transform, ERigTransformType::InitialLocal, true, true);
+			Hierarchy->SetTransform(TransformElement, PerElementData.Pose.Current.Local.Transform, ERigTransformType::CurrentLocal, true, true);
+		}
+	}
+	
+#if WITH_EDITOR
+	TransactionPtr.Reset();
+#endif
+
+	return PastedKeys;
+}
+
+void URigHierarchyController::Notify(ERigHierarchyNotification InNotifType, const FRigBaseElement* InElement)
+{
+	if(!IsValid())
+	{
+		return;
+	}
+	Hierarchy->Notify(InNotifType, InElement);
+}
+
+void URigHierarchyController::HandleHierarchyModified(ERigHierarchyNotification InNotifType, URigHierarchy* InHierarchy, const FRigBaseElement* InElement) const
+{
+	ensure(IsValid());
+	ensure(InHierarchy == Hierarchy);
+	ModifiedEvent.Broadcast(InNotifType, InHierarchy, InElement);
+}
+
+int32 URigHierarchyController::AddElement(FRigBaseElement* InElementToAdd, FRigBaseElement* InFirstParent, bool bMaintainGlobalTransform)
+{
+	ensure(IsValid());
+
+	InElementToAdd->SubIndex = Hierarchy->Num(InElementToAdd->Key.Type);
+	InElementToAdd->Index = Hierarchy->Elements.Add(InElementToAdd);
+
+	Hierarchy->IndexLookup.Add(InElementToAdd->Key, InElementToAdd->Index);
+	Hierarchy->TopologyVersion++;
+
+	Notify(ERigHierarchyNotification::ElementAdded, InElementToAdd);
+
+	SetParent(InElementToAdd, InFirstParent, bMaintainGlobalTransform);
+
+	return InElementToAdd->Index;
+}
+
+bool URigHierarchyController::RemoveElement(FRigElementKey InElement, bool bSetupUndo)
+{
+	if(!IsValid())
+	{
+		return false;
+	}
+
+	FRigBaseElement* Element = Hierarchy->Find(InElement);
+	if(Element == nullptr)
+	{
+		ReportWarningf(TEXT("Cannot Remove Element: '%s' not found."), *InElement.ToString());
+		return false;
+	}
+
+#if WITH_EDITOR
+	TSharedPtr<FScopedTransaction> TransactionPtr;
+	if(bSetupUndo)
+	{
+		TransactionPtr = MakeShared<FScopedTransaction>(NSLOCTEXT("RigHierarchyController", "Remove Element", "Remove Element"));
+		Hierarchy->Modify();
+	}
+#endif
+
+	const bool bRemoved = RemoveElement(Element);
+	
+#if WITH_EDITOR
+	if(!bRemoved && TransactionPtr.IsValid())
+	{
+		TransactionPtr->Cancel();
+	}
+	TransactionPtr.Reset();
+#endif
+	
+	return bRemoved;
+}
+
+bool URigHierarchyController::RemoveElement(FRigBaseElement* InElement)
+{
+	if(InElement == nullptr)
+	{
+		return false;
+	}
+
+	// make sure this element is part of this hierarchy
+	ensure(Hierarchy->FindChecked(InElement->Key) == InElement);
+
+	// deselect if needed
+	if(InElement->IsSelected())
+	{
+		SelectElement(InElement->GetKey(), false);
+	}
+
+	// if this is a transform element - make sure to allow dependents to store their global transforms
+	if(FRigTransformElement* TransformElement = Cast<FRigTransformElement>(InElement))
+	{
+		TArray<FRigTransformElement::FElementToDirty> PreviousElementsToDirty = TransformElement->ElementsToDirty; 
+		for(const FRigTransformElement::FElementToDirty& ElementToDirty : PreviousElementsToDirty)
+		{
+			if(FRigSingleParentElement* SingleParentElement = Cast<FRigSingleParentElement>(ElementToDirty.Element))
+			{
+				if(SingleParentElement->ParentElement == InElement)
+				{
+					RemoveParent(SingleParentElement, InElement, true);
+				}
+			}
+			else if(FRigMultiParentElement* MultiParentElement = Cast<FRigMultiParentElement>(ElementToDirty.Element))
+			{
+				if(MultiParentElement->ParentElements.Contains(InElement))
+				{
+					RemoveParent(MultiParentElement, InElement, true);
+				}
+			}
+		}
+	}
+
+	const int32 NumElementsRemoved = Hierarchy->Elements.Remove(InElement);
+	ensure(NumElementsRemoved == 1);
+
+	const int32 NumLookupsRemoved = Hierarchy->IndexLookup.Remove(InElement->Key);
+	ensure(NumLookupsRemoved == 1);
+	for(TPair<FRigElementKey, int32>& Pair : Hierarchy->IndexLookup)
+	{
+		if(Pair.Value > InElement->Index)
+		{
+			Pair.Value--;
+		}
+	}
+
+	// update the indices of all other elements
+	for (FRigBaseElement* RemainingElement : Hierarchy->Elements)
+	{
+		if(RemainingElement->Index > InElement->Index)
+		{
+			RemainingElement->Index--;
+		}
+	}
+
+	if(FRigSingleParentElement* SingleParentElement = Cast<FRigSingleParentElement>(InElement))
+	{
+		RemoveElementToDirty(SingleParentElement->ParentElement, InElement);
+	}
+	else if(FRigMultiParentElement* MultiParentElement = Cast<FRigMultiParentElement>(InElement))
+	{
+		for(FRigTransformElement* ParentElement : MultiParentElement->ParentElements)
+		{
+			RemoveElementToDirty(ParentElement, InElement);
+		}
+	}
+
+	if(InElement->SubIndex != INDEX_NONE)
+	{
+		for(FRigBaseElement* Element : *Hierarchy)
+		{
+			if(Element->SubIndex > InElement->SubIndex && Element->GetType() == InElement->GetType())
+			{
+				Element->SubIndex--;
+			}
+		}
+	}
+
+	Hierarchy->TopologyVersion++;
+
+	Notify(ERigHierarchyNotification::ElementRemoved, InElement);
+	if(Hierarchy->Num() == 0)
+	{
+		Notify(ERigHierarchyNotification::HierarchyReset, nullptr);
+	}
+
+	delete InElement;
+
+	return NumElementsRemoved == 1;
+}
+
+FRigElementKey URigHierarchyController::RenameElement(FRigElementKey InElement, FName InName, bool bSetupUndo)
+{
+	if(!IsValid())
+	{
+		return FRigElementKey();
+	}
+
+	FRigBaseElement* Element = Hierarchy->Find(InElement);
+	if(Element == nullptr)
+	{
+		ReportWarningf(TEXT("Cannot Rename Element: '%s' not found."), *InElement.ToString());
+		return FRigElementKey();
+	}
+
+#if WITH_EDITOR
+	TSharedPtr<FScopedTransaction> TransactionPtr;
+	if(bSetupUndo)
+	{
+		TransactionPtr = MakeShared<FScopedTransaction>(NSLOCTEXT("RigHierarchyController", "Rename Element", "Rename Element"));
+		Hierarchy->Modify();
+	}
+#endif
+
+	const bool bRenamed = RenameElement(Element, InName);
+
+#if WITH_EDITOR
+	if(!bRenamed && TransactionPtr.IsValid())
+	{
+		TransactionPtr->Cancel();
+	}
+	TransactionPtr.Reset();
+#endif
+
+	return bRenamed ? Element->GetKey() : FRigElementKey();
+}
+
+bool URigHierarchyController::RenameElement(FRigBaseElement* InElement, const FName &InName)
+{
+	if(InElement == nullptr)
+	{
+		return false;
+	}
+
+	if(InElement->GetName() == InName)
+	{
+		return false;
+	}
+
+	const FRigElementKey OldKey = InElement->GetKey();
+	InElement->Key.Name = Hierarchy->GetSafeNewName(InName.ToString(), InElement->GetType());
+	const FRigElementKey NewKey = InElement->GetKey();
+
+	Hierarchy->IndexLookup.Remove(OldKey);
+	Hierarchy->IndexLookup.Add(NewKey, InElement->Index);
+
+	// update all multi parent elements' index lookups
+	for (FRigBaseElement* Element : Hierarchy->Elements)
+	{
+		if(FRigMultiParentElement* MultiParentElement = Cast<FRigMultiParentElement>(Element))
+		{
+			const int32* ExistingIndexPtr = MultiParentElement->IndexLookup.Find(OldKey);
+			if(ExistingIndexPtr)
+			{
+				const int32 ExistingIndex = *ExistingIndexPtr;
+				MultiParentElement->IndexLookup.Remove(OldKey);
+				MultiParentElement->IndexLookup.Add(NewKey, ExistingIndex);
+			}
+		}
+	}
+	
+	Hierarchy->PreviousNameMap.FindOrAdd(NewKey) = OldKey;
+	Notify(ERigHierarchyNotification::ElementRenamed, InElement);
+
+	return true;
+}
+
+bool URigHierarchyController::AddParent(FRigElementKey InChild, FRigElementKey InParent, bool bMaintainGlobalTransform, bool bSetupUndo)
+{
+	if(!IsValid())
+	{
+		return false;
+	}
+
+	FRigBaseElement* Child = Hierarchy->Find(InChild);
+	if(Child == nullptr)
+	{
+		ReportWarningf(TEXT("Cannot Add Parent, Child '%s' not found."), *InChild.ToString());
+		return false;
+	}
+
+	FRigBaseElement* Parent = Hierarchy->Find(InParent);
+	if(Parent == nullptr)
+	{
+		ReportWarningf(TEXT("Cannot Add Parent, Parent '%s' not found."), *InParent.ToString());
+		return false;
+	}
+
+#if WITH_EDITOR
+	TSharedPtr<FScopedTransaction> TransactionPtr;
+	if(bSetupUndo)
+	{
+		TransactionPtr = MakeShared<FScopedTransaction>(NSLOCTEXT("RigHierarchyController", "Add Parent", "Add Parent"));
+		Hierarchy->Modify();
+	}
+#endif
+
+	const bool bAdded = AddParent(Child, Parent, bMaintainGlobalTransform);
+
+#if WITH_EDITOR
+	if(!bAdded && TransactionPtr.IsValid())
+	{
+		TransactionPtr->Cancel();
+	}
+	TransactionPtr.Reset();
+#endif
+	
+	return bAdded;
+}
+
+bool URigHierarchyController::AddParent(FRigBaseElement* InChild, FRigBaseElement* InParent, bool bMaintainGlobalTransform, bool bRemoveAllParents)
+{
+	if(InChild == nullptr || InParent == nullptr)
+	{
+		return false;
+	}
+
+	// single parent children can't be parented multiple times
+	if(FRigSingleParentElement* SingleParentElement = Cast<FRigSingleParentElement>(InChild))
+	{
+		if(SingleParentElement->ParentElement == InParent)
+		{
+			return false;
+		}
+		bRemoveAllParents = true;
+	}
+
+	else if(FRigMultiParentElement* MultiParentElement = Cast<FRigMultiParentElement>(InChild))
+	{
+		if(MultiParentElement->ParentElements.Contains(InParent))
+		{
+			return false;
+		}
+	}
+
+	if(Hierarchy->IsParentedTo(InParent, InChild))
+	{
+		ReportAndNotifyErrorf(TEXT("Cannot parent '%s' to '%s' - would cause a cycle."), *InChild->Key.ToString(), *InParent->Key.ToString());
+		return false;
+	}
+
+	if(bRemoveAllParents)
+	{
+		RemoveAllParents(InChild, bMaintainGlobalTransform);		
+	}
+
+	if(FRigTransformElement* TransformElement = Cast<FRigTransformElement>(InChild))
+	{
+		if(bMaintainGlobalTransform)
+		{
+			Hierarchy->GetTransform(TransformElement, ERigTransformType::CurrentGlobal);
+			Hierarchy->GetTransform(TransformElement, ERigTransformType::InitialGlobal);
+			TransformElement->Pose.MarkDirty(ERigTransformType::CurrentLocal);
+			TransformElement->Pose.MarkDirty(ERigTransformType::InitialLocal);
+		}
+		else
+		{
+			Hierarchy->GetTransform(TransformElement, ERigTransformType::CurrentLocal);
+			Hierarchy->GetTransform(TransformElement, ERigTransformType::InitialLocal);
+			TransformElement->Pose.MarkDirty(ERigTransformType::CurrentGlobal);
+			TransformElement->Pose.MarkDirty(ERigTransformType::InitialGlobal);
+		}
+	}
+
+	if(FRigSingleParentElement* SingleParentElement = Cast<FRigSingleParentElement>(InChild))
+	{
+		if(FRigTransformElement* NewTransformParent = Cast<FRigTransformElement>(InParent))
+		{
+			AddElementToDirty(NewTransformParent, SingleParentElement);
+			SingleParentElement->ParentElement = NewTransformParent;
+
+			Hierarchy->TopologyVersion++;
+
+			if(!bMaintainGlobalTransform)
+			{
+				Hierarchy->PropagateDirtyFlags(SingleParentElement, ERigTransformType::CurrentGlobal, true);
+				Hierarchy->PropagateDirtyFlags(SingleParentElement, ERigTransformType::InitialGlobal, true);
+			}
+
+			Notify(ERigHierarchyNotification::ParentChanged, SingleParentElement);
+			return true;
+		}
+	}
+	else if(FRigMultiParentElement* MultiParentElement = Cast<FRigMultiParentElement>(InChild))
+	{
+		if(FRigTransformElement* NewTransformParent = Cast<FRigTransformElement>(InParent))
+		{
+			AddElementToDirty(NewTransformParent, MultiParentElement);
+
+			const int32 ParentIndex = MultiParentElement->ParentElements.Add(NewTransformParent);
+			MultiParentElement->ParentWeights.Add(1.f);
+			MultiParentElement->ParentWeightsInitial.Add(1.f);
+			MultiParentElement->IndexLookup.Add(NewTransformParent->GetKey(), ParentIndex);
+
+			MultiParentElement->Parent.MarkDirty(ERigTransformType::CurrentGlobal);  
+			MultiParentElement->Parent.MarkDirty(ERigTransformType::InitialGlobal);
+
+			if(FRigControlElement* ControlElement = Cast<FRigControlElement>(MultiParentElement))
+			{
+				ControlElement->Offset.MarkDirty(ERigTransformType::CurrentGlobal);  
+				ControlElement->Offset.MarkDirty(ERigTransformType::InitialGlobal);
+				ControlElement->Gizmo.MarkDirty(ERigTransformType::CurrentGlobal);  
+				ControlElement->Gizmo.MarkDirty(ERigTransformType::InitialGlobal);
+			}
+
+			Hierarchy->TopologyVersion++;
+
+			if(!bMaintainGlobalTransform)
+			{
+				Hierarchy->PropagateDirtyFlags(MultiParentElement, ERigTransformType::CurrentGlobal, true);
+				Hierarchy->PropagateDirtyFlags(MultiParentElement, ERigTransformType::InitialGlobal, true);
+			}
+
+			Notify(ERigHierarchyNotification::ParentChanged, MultiParentElement);
+			return true;
+		}
+	}
+	
+	return false;
+}
+
+bool URigHierarchyController::RemoveParent(FRigElementKey InChild, FRigElementKey InParent, bool bMaintainGlobalTransform, bool bSetupUndo)
+{
+	if(!IsValid())
+	{
+		return false;
+	}
+
+	FRigBaseElement* Child = Hierarchy->Find(InChild);
+	if(Child == nullptr)
+	{
+		ReportWarningf(TEXT("Cannot Remove Parent, Child '%s' not found."), *InChild.ToString());
+		return false;
+	}
+
+	FRigBaseElement* Parent = Hierarchy->Find(InParent);
+	if(Parent == nullptr)
+	{
+		ReportWarningf(TEXT("Cannot Remove Parent, Parent '%s' not found."), *InParent.ToString());
+		return false;
+	}
+
+#if WITH_EDITOR
+	TSharedPtr<FScopedTransaction> TransactionPtr;
+	if(bSetupUndo)
+	{
+		TransactionPtr = MakeShared<FScopedTransaction>(NSLOCTEXT("RigHierarchyController", "Remove Parent", "Remove Parent"));
+		Hierarchy->Modify();
+	}
+#endif
+
+	const bool bRemoved = RemoveParent(Child, Parent, bMaintainGlobalTransform);
+
+#if WITH_EDITOR
+	if(!bRemoved && TransactionPtr.IsValid())
+	{
+		TransactionPtr->Cancel();
+	}
+	TransactionPtr.Reset();
+#endif
+	
+	return bRemoved;
+}
+
+bool URigHierarchyController::RemoveParent(FRigBaseElement* InChild, FRigBaseElement* InParent, bool bMaintainGlobalTransform)
+{
+	if(InChild == nullptr || InParent == nullptr)
+	{
+		return false;
+	}
+
+	FRigTransformElement* ParentTransformElement = Cast<FRigTransformElement>(InParent);
+	if(ParentTransformElement == nullptr)
+	{
+		return false;
+	}
+
+	// single parent children can't be parented multiple times
+	if(FRigSingleParentElement* SingleParentElement = Cast<FRigSingleParentElement>(InChild))
+	{
+		if(SingleParentElement->ParentElement == ParentTransformElement)
+		{
+			if(bMaintainGlobalTransform)
+			{
+				Hierarchy->GetTransform(SingleParentElement, ERigTransformType::CurrentGlobal);
+				Hierarchy->GetTransform(SingleParentElement, ERigTransformType::InitialGlobal);
+				SingleParentElement->Pose.MarkDirty(ERigTransformType::CurrentLocal);
+				SingleParentElement->Pose.MarkDirty(ERigTransformType::InitialLocal);
+			}
+			else
+			{
+				Hierarchy->GetTransform(SingleParentElement, ERigTransformType::CurrentLocal);
+				Hierarchy->GetTransform(SingleParentElement, ERigTransformType::InitialLocal);
+				SingleParentElement->Pose.MarkDirty(ERigTransformType::CurrentGlobal);
+				SingleParentElement->Pose.MarkDirty(ERigTransformType::InitialGlobal);
+			}
+
+			const FRigElementKey PreviousParentKey = SingleParentElement->ParentElement->GetKey();
+			Hierarchy->PreviousParentMap.FindOrAdd(SingleParentElement->GetKey()) = PreviousParentKey;
+			
+			// remove the previous parent
+			SingleParentElement->ParentElement = nullptr;
+			RemoveElementToDirty(InParent, SingleParentElement); 
+			Hierarchy->TopologyVersion++;
+
+			if(!bMaintainGlobalTransform)
+			{
+				Hierarchy->PropagateDirtyFlags(SingleParentElement, ERigTransformType::CurrentGlobal, true);
+				Hierarchy->PropagateDirtyFlags(SingleParentElement, ERigTransformType::InitialGlobal, true);
+			}
+
+			Notify(ERigHierarchyNotification::ParentChanged, SingleParentElement);
+
+			return true;
+		}
+	}
+
+	// single parent children can't be parented multiple times
+	else if(FRigMultiParentElement* MultiParentElement = Cast<FRigMultiParentElement>(InChild))
+	{
+		const int32 ParentIndex = MultiParentElement->ParentElements.Find(ParentTransformElement);
+		if(MultiParentElement->ParentElements.IsValidIndex(ParentIndex))
+		{
+			if(bMaintainGlobalTransform)
+			{
+				Hierarchy->GetTransform(MultiParentElement, ERigTransformType::CurrentGlobal);
+				Hierarchy->GetTransform(MultiParentElement, ERigTransformType::InitialGlobal);
+				MultiParentElement->Pose.MarkDirty(ERigTransformType::CurrentLocal);
+				MultiParentElement->Pose.MarkDirty(ERigTransformType::InitialLocal);
+			}
+			else
+			{
+				Hierarchy->GetTransform(MultiParentElement, ERigTransformType::CurrentLocal);
+				Hierarchy->GetTransform(MultiParentElement, ERigTransformType::InitialLocal);
+				MultiParentElement->Pose.MarkDirty(ERigTransformType::CurrentGlobal);
+				MultiParentElement->Pose.MarkDirty(ERigTransformType::InitialGlobal);
+			}
+
+			// remove the previous parent
+			RemoveElementToDirty(InParent, MultiParentElement); 
+
+			const FRigElementKey PreviousParentKey = MultiParentElement->ParentElements[ParentIndex]->GetKey();
+			Hierarchy->PreviousParentMap.FindOrAdd(MultiParentElement->GetKey()) = PreviousParentKey;
+
+			MultiParentElement->ParentElements.RemoveAt(ParentIndex);
+			MultiParentElement->ParentWeights.RemoveAt(ParentIndex);
+			MultiParentElement->ParentWeightsInitial.RemoveAt(ParentIndex);
+			MultiParentElement->IndexLookup.Remove(ParentTransformElement->GetKey());
+			for(TPair<FRigElementKey, int32>& Pair : MultiParentElement->IndexLookup)
+			{
+				if(Pair.Value > ParentIndex)
+				{
+					Pair.Value--;
+				}
+			}
+
+			MultiParentElement->Parent.MarkDirty(ERigTransformType::CurrentGlobal);  
+			MultiParentElement->Parent.MarkDirty(ERigTransformType::InitialGlobal);  
+
+			if(FRigControlElement* ControlElement = Cast<FRigControlElement>(MultiParentElement))
+			{
+				ControlElement->Offset.MarkDirty(ERigTransformType::CurrentGlobal);  
+				ControlElement->Offset.MarkDirty(ERigTransformType::InitialGlobal);
+				ControlElement->Gizmo.MarkDirty(ERigTransformType::CurrentGlobal);  
+				ControlElement->Gizmo.MarkDirty(ERigTransformType::InitialGlobal);
+			}
+
+			Hierarchy->TopologyVersion++;
+
+			if(!bMaintainGlobalTransform)
+			{
+				Hierarchy->PropagateDirtyFlags(MultiParentElement, ERigTransformType::CurrentGlobal, true);
+				Hierarchy->PropagateDirtyFlags(MultiParentElement, ERigTransformType::InitialGlobal, true);
+			}
+
+			Notify(ERigHierarchyNotification::ParentChanged, MultiParentElement);
+
+			return true;
+		}
+	}
+
+	return false;
+}
+
+bool URigHierarchyController::RemoveAllParents(FRigElementKey InChild, bool bMaintainGlobalTransform, bool bSetupUndo)
+{
+	if(!IsValid())
+	{
+		return false;
+	}
+
+	FRigBaseElement* Child = Hierarchy->Find(InChild);
+	if(Child == nullptr)
+	{
+		ReportWarningf(TEXT("Cannot Remove All Parents, Child '%s' not found."), *InChild.ToString());
+		return false;
+	}
+
+#if WITH_EDITOR
+	TSharedPtr<FScopedTransaction> TransactionPtr;
+	if(bSetupUndo)
+	{
+		TransactionPtr = MakeShared<FScopedTransaction>(NSLOCTEXT("RigHierarchyController", "Remove Parent", "Remove Parent"));
+		Hierarchy->Modify();
+	}
+#endif
+
+	const bool bRemoved = RemoveAllParents(Child, bMaintainGlobalTransform);
+
+#if WITH_EDITOR
+	if(!bRemoved && TransactionPtr.IsValid())
+	{
+		TransactionPtr->Cancel();
+	}
+	TransactionPtr.Reset();
+#endif
+
+	return bRemoved;
+}
+
+bool URigHierarchyController::RemoveAllParents(FRigBaseElement* InChild, bool bMaintainGlobalTransform)
+{
+	if(FRigSingleParentElement* SingleParentElement = Cast<FRigSingleParentElement>(InChild))
+	{
+		return RemoveParent(SingleParentElement, SingleParentElement->ParentElement, bMaintainGlobalTransform);
+	}
+	else if(FRigMultiParentElement* MultiParentElement = Cast<FRigMultiParentElement>(InChild))
+	{
+		bool bSuccess = true;
+
+		TArray<FRigTransformElement*> Parents = MultiParentElement->ParentElements;
+		for(FRigTransformElement* Parent : Parents)
+		{
+			if(!RemoveParent(MultiParentElement, Parent, bMaintainGlobalTransform))
+			{
+				bSuccess = false;
+			}
+		}
+
+		return bSuccess;
+	}
+	return false;
+}
+
+bool URigHierarchyController::SetParent(FRigElementKey InChild, FRigElementKey InParent, bool bMaintainGlobalTransform, bool bSetupUndo)
+{
+	if(!IsValid())
+	{
+		return false;
+	}
+
+	FRigBaseElement* Child = Hierarchy->Find(InChild);
+	if(Child == nullptr)
+	{
+		ReportWarningf(TEXT("Cannot Set Parent, Child '%s' not found."), *InChild.ToString());
+		return false;
+	}
+
+	FRigBaseElement* Parent = Hierarchy->Find(InParent);
+	if(Parent == nullptr)
+	{
+		ReportWarningf(TEXT("Cannot Set Parent, Parent '%s' not found."), *InParent.ToString());
+		return false;
+	}
+
+#if WITH_EDITOR
+	TSharedPtr<FScopedTransaction> TransactionPtr;
+	if(bSetupUndo)
+	{
+		TransactionPtr = MakeShared<FScopedTransaction>(NSLOCTEXT("RigHierarchyController", "Set Parent", "Set Parent"));
+		Hierarchy->Modify();
+	}
+#endif
+
+	const bool bParentSet = SetParent(Child, Parent, bMaintainGlobalTransform);
+
+#if WITH_EDITOR
+	if(!bParentSet && TransactionPtr.IsValid())
+	{
+		TransactionPtr->Cancel();
+	}
+	TransactionPtr.Reset();
+#endif
+
+	return bParentSet;
+}
+
+TArray<FRigElementKey> URigHierarchyController::DuplicateElements(TArray<FRigElementKey> InKeys, bool bSelectNewElements, bool bSetupUndo)
+{
+	const FString Content = ExportToText(InKeys);
+	return ImportFromText(Content, false, bSelectNewElements, bSetupUndo);
+}
+
+TArray<FRigElementKey> URigHierarchyController::MirrorElements(TArray<FRigElementKey> InKeys, FRigMirrorSettings InSettings, bool bSelectNewElements, bool bSetupUndo)
+{
+	TArray<FRigElementKey> OriginalKeys = Hierarchy->SortKeys(InKeys);
+	TArray<FRigElementKey> DuplicatedKeys = DuplicateElements(OriginalKeys, bSelectNewElements, bSetupUndo);
+
+	if (DuplicatedKeys.Num() != OriginalKeys.Num())
+	{
+		return DuplicatedKeys;
+	}
+
+	for (int32 Index = 0; Index < OriginalKeys.Num(); Index++)
+	{
+		if (DuplicatedKeys[Index].Type != OriginalKeys[Index].Type)
+		{
+			return DuplicatedKeys;
+		}
+	}
+
+	// mirror the transforms
+	for (int32 Index = 0; Index < OriginalKeys.Num(); Index++)
+	{
+		FTransform GlobalTransform = Hierarchy->GetGlobalTransform(OriginalKeys[Index]);
+		FTransform InitialTransform = Hierarchy->GetInitialGlobalTransform(OriginalKeys[Index]);
+
+		// also mirror the offset, limits and gizmo transform
+		if (OriginalKeys[Index].Type == ERigElementType::Control)
+		{
+			if(FRigControlElement* DuplicatedControlElement = Hierarchy->Find<FRigControlElement>(DuplicatedKeys[Index]))
+			{
+				TGuardValue<bool> DisableTranslationLimit(DuplicatedControlElement->Settings.bLimitTranslation, false);
+				TGuardValue<bool> DisableRotationLimit(DuplicatedControlElement->Settings.bLimitRotation, false);
+				TGuardValue<bool> DisableScaleLimit(DuplicatedControlElement->Settings.bLimitScale, false);
+
+				// mirror offset
+				FTransform OriginalGlobalOffsetTransform = Hierarchy->GetGlobalControlOffsetTransform(OriginalKeys[Index]);
+				FTransform ParentTransform = Hierarchy->GetParentTransform(DuplicatedKeys[Index]);
+				FTransform OffsetTransform = InSettings.MirrorTransform(OriginalGlobalOffsetTransform).GetRelativeTransform(ParentTransform);
+				Hierarchy->SetControlOffsetTransform(DuplicatedKeys[Index], OffsetTransform, true, false, true);
+				Hierarchy->SetControlOffsetTransform(DuplicatedKeys[Index], OffsetTransform, false, false, true);
+
+				// mirror limits
+				FTransform DuplicatedGlobalOffsetTransform = Hierarchy->GetGlobalControlOffsetTransform(DuplicatedKeys[Index]);
+
+				for (ERigControlValueType ValueType = ERigControlValueType::Minimum;
+                    ValueType <= ERigControlValueType::Maximum;
+                    ValueType = (ERigControlValueType)(uint8(ValueType) + 1)
+                    )
+				{
+					const FRigControlValue LimitValue = Hierarchy->GetControlValue(DuplicatedKeys[Index], ValueType);
+					const FTransform LocalLimitTransform = LimitValue.GetAsTransform(DuplicatedControlElement->Settings.ControlType, DuplicatedControlElement->Settings.PrimaryAxis);
+					FTransform GlobalLimitTransform = LocalLimitTransform * OriginalGlobalOffsetTransform;
+					FTransform DuplicatedLimitTransform = InSettings.MirrorTransform(GlobalLimitTransform).GetRelativeTransform(DuplicatedGlobalOffsetTransform);
+					FRigControlValue DuplicatedValue;
+					DuplicatedValue.SetFromTransform(DuplicatedLimitTransform, DuplicatedControlElement->Settings.ControlType, DuplicatedControlElement->Settings.PrimaryAxis);
+					Hierarchy->SetControlValue(DuplicatedControlElement, DuplicatedValue, ValueType, false);
+				}
+
+				// we need to do this here to make sure that the limits don't apply ( the tguardvalue is still active within here )
+				Hierarchy->SetGlobalTransform(DuplicatedKeys[Index], InSettings.MirrorTransform(GlobalTransform), true, false, true);
+				Hierarchy->SetGlobalTransform(DuplicatedKeys[Index], InSettings.MirrorTransform(GlobalTransform), false, false, true);
+
+				// mirror gizmo transform
+				FTransform GlobalGizmoTransform = Hierarchy->GetControlGizmoTransform(DuplicatedControlElement, ERigTransformType::InitialLocal) * OriginalGlobalOffsetTransform;
+				Hierarchy->SetControlGizmoTransform(DuplicatedControlElement, InSettings.MirrorTransform(GlobalGizmoTransform).GetRelativeTransform(DuplicatedGlobalOffsetTransform), ERigTransformType::InitialLocal, true);
+				Hierarchy->SetControlGizmoTransform(DuplicatedControlElement, InSettings.MirrorTransform(GlobalGizmoTransform).GetRelativeTransform(DuplicatedGlobalOffsetTransform), ERigTransformType::CurrentLocal, true);
+			}
+		}
+		else
+		{
+			Hierarchy->SetGlobalTransform(DuplicatedKeys[Index], InSettings.MirrorTransform(GlobalTransform), true, false, true);
+			Hierarchy->SetGlobalTransform(DuplicatedKeys[Index], InSettings.MirrorTransform(GlobalTransform), false, false, true);
+		}
+	}
+
+	// correct the names
+	if (!InSettings.OldName.IsEmpty() && !InSettings.NewName.IsEmpty())
+	{
+		URigHierarchyController* Controller = Hierarchy->GetController(true);
+		check(Controller);
+		
+		for (int32 Index = 0; Index < DuplicatedKeys.Num(); Index++)
+		{
+			FName OldName = OriginalKeys[Index].Name;
+			FString OldNameStr = OldName.ToString();
+			FString NewNameStr = OldNameStr.Replace(*InSettings.OldName, *InSettings.NewName, ESearchCase::CaseSensitive);
+			if (NewNameStr != OldNameStr)
+			{
+				Controller->RenameElement(DuplicatedKeys[Index], *NewNameStr, true);
+			}
+		}
+	}
+
+	return DuplicatedKeys;
+}
+
+bool URigHierarchyController::SetParent(FRigBaseElement* InChild, FRigBaseElement* InParent, bool bMaintainGlobalTransform)
+{
+	if(InChild == nullptr || InParent == nullptr)
+	{
+		return false;
+	}
+	return AddParent(InChild, InParent, bMaintainGlobalTransform, true);
+}
+
+void URigHierarchyController::AddElementToDirty(FRigBaseElement* InParent, FRigBaseElement* InElementToAdd, int32 InHierarchyDistance) const
+{
+	if(InParent == nullptr)
+	{
+		return;
+	} 
+
+	FRigTransformElement* ElementToAdd = Cast<FRigTransformElement>(InElementToAdd);
+	if(ElementToAdd == nullptr)
+	{
+		return;
+	}
+
+	if(FRigTransformElement* TransformParent = Cast<FRigTransformElement>(InParent))
+	{
+		const FRigTransformElement::FElementToDirty ElementToDirty(ElementToAdd, InHierarchyDistance);
+		TransformParent->ElementsToDirty.AddUnique(ElementToDirty);
+	}
+
+#if URIGHIERARCHY_RECURSIVE_DIRTY_PROPAGATION
+	// nothing to do 
+#else
+	if(FRigSingleParentElement* SingleParentElement = Cast<FRigSingleParentElement>(InParent))
+	{
+		AddElementToDirty(SingleParentElement->ParentElement, InElementToAdd, InHierarchyDistance + 1);
+	}
+	else if(FRigMultiParentElement* MultiParentElement = Cast<FRigMultiParentElement>(InParent))
+	{
+		for(FRigTransformElement* ParentElement : MultiParentElement->ParentElements)
+		{
+			AddElementToDirty(ParentElement, InElementToAdd, InHierarchyDistance + 1);
+		}
+	}
+#endif
+}
+
+void URigHierarchyController::RemoveElementToDirty(FRigBaseElement* InParent, FRigBaseElement* InElementToRemove) const
+{
+	if(InParent == nullptr)
+	{
+		return;
+	}
+
+	FRigTransformElement* ElementToRemove = Cast<FRigTransformElement>(InElementToRemove);
+	if(ElementToRemove == nullptr)
+	{
+		return;
+	}
+
+	if(FRigTransformElement* TransformParent = Cast<FRigTransformElement>(InParent))
+	{
+		TransformParent->ElementsToDirty.Remove(ElementToRemove);
+	}
+
+#if URIGHIERARCHY_RECURSIVE_DIRTY_PROPAGATION
+	// nothing to do 
+#else
+	if(FRigSingleParentElement* SingleParentElement = Cast<FRigSingleParentElement>(InParent))
+	{
+		RemoveElementToDirty(SingleParentElement->ParentElement, InElementToRemove);
+	}
+	else if(FRigMultiParentElement* MultiParentElement = Cast<FRigMultiParentElement>(InParent))
+	{
+		for(FRigTransformElement* ParentElement : MultiParentElement->ParentElements)
+		{
+			RemoveElementToDirty(ParentElement, InElementToRemove);
+		}
+	}
+#endif
+}
+
+void URigHierarchyController::ReportWarning(const FString& InMessage) const
+{
+	if(!bReportWarningsAndErrors)
+	{
+		return;
+	}
+
+	FString Message = InMessage;
+	if (Hierarchy)
+	{
+		if (UPackage* Package = Cast<UPackage>(Hierarchy->GetOutermost()))
+		{
+			Message = FString::Printf(TEXT("%s : %s"), *Package->GetPathName(), *InMessage);
+		}
+	}
+
+	FScriptExceptionHandler::Get().HandleException(ELogVerbosity::Warning, *Message, *FString());
+}
+
+void URigHierarchyController::ReportError(const FString& InMessage) const
+{
+	if(!bReportWarningsAndErrors)
+	{
+		return;
+	}
+
+	FString Message = InMessage;
+	if (Hierarchy)
+	{
+		if (UPackage* Package = Cast<UPackage>(Hierarchy->GetOutermost()))
+		{
+			Message = FString::Printf(TEXT("%s : %s"), *Package->GetPathName(), *InMessage);
+		}
+	}
+
+	FScriptExceptionHandler::Get().HandleException(ELogVerbosity::Error, *Message, *FString());
+}
+
+void URigHierarchyController::ReportAndNotifyError(const FString& InMessage) const
+{
+	if (!bReportWarningsAndErrors)
+	{
+		return;
+	}
+
+	ReportError(InMessage);
+
+#if WITH_EDITOR
+	FNotificationInfo Info(FText::FromString(InMessage));
+	Info.bUseSuccessFailIcons = true;
+	Info.Image = FEditorStyle::GetBrush(TEXT("MessageLog.Warning"));
+	Info.bFireAndForget = true;
+	Info.bUseThrobber = true;
+	// longer message needs more time to read
+	Info.FadeOutDuration = FMath::Clamp(0.1f * InMessage.Len(), 5.0f, 20.0f);
+	Info.ExpireDuration = Info.FadeOutDuration;
+	TSharedPtr<SNotificationItem> NotificationPtr = FSlateNotificationManager::Get().AddNotification(Info);
+	if (NotificationPtr)
+	{
+		NotificationPtr->SetCompletionState(SNotificationItem::CS_Fail);
+	}
+#endif
+}
diff --git a/Engine/Plugins/Experimental/ControlRig/Source/ControlRig/Private/Rigs/RigHierarchyDefines.cpp b/Engine/Plugins/Experimental/ControlRig/Source/ControlRig/Private/Rigs/RigHierarchyDefines.cpp
--- a/Engine/Plugins/Experimental/ControlRig/Source/ControlRig/Private/Rigs/RigHierarchyDefines.cpp
+++ b/Engine/Plugins/Experimental/ControlRig/Source/ControlRig/Private/Rigs/RigHierarchyDefines.cpp
@@ -1,24 +1,66 @@
 // Copyright Epic Games, Inc. All Rights Reserved.
 
 #include "Rigs/RigHierarchyDefines.h"
-#include "Rigs/RigHierarchyContainer.h"
+#include "Rigs/RigHierarchy.h"
+
+////////////////////////////////////////////////////////////////////////////////
+// FRigElementKey
+////////////////////////////////////////////////////////////////////////////////
+
+void FRigElementKey::Serialize(FArchive& Ar)
+{
+	if (Ar.IsSaving() || Ar.IsObjectReferenceCollector())
+	{
+		Save(Ar);
+	}
+	else if (Ar.IsLoading())
+	{
+		Load(Ar);
+	}
+	else
+	{
+		checkNoEntry();
+	}
+}
+
+void FRigElementKey::Save(FArchive& Ar)
+{
+	static const UEnum* ElementTypeEnum = StaticEnum<ERigElementType>();
+
+	FName TypeName = ElementTypeEnum->GetNameByValue((int64)Type);
+	Ar << TypeName;
+	Ar << Name;
+}
+
+void FRigElementKey::Load(FArchive& Ar)
+{
+	static const UEnum* ElementTypeEnum = StaticEnum<ERigElementType>();
+
+	FName TypeName;
+	Ar << TypeName;
+
+	const int64 TypeValue = ElementTypeEnum->GetValueByName(TypeName);
+	Type = (ERigElementType)TypeValue;
+
+	Ar << Name;
+}
 
 ////////////////////////////////////////////////////////////////////////////////
 // FRigElementKeyCollection
 ////////////////////////////////////////////////////////////////////////////////
 
 FRigElementKeyCollection FRigElementKeyCollection::MakeFromChildren(
-	const FRigHierarchyContainer* InContainer,
+	URigHierarchy* InHierarchy,
 	const FRigElementKey& InParentKey,
 	bool bRecursive,
 	bool bIncludeParent,
 	uint8 InElementTypes)
 {
-	check(InContainer);
+	check(InHierarchy);
 
 	FRigElementKeyCollection Collection;
 
-	int32 Index = InContainer->GetIndex(InParentKey);
+	int32 Index = InHierarchy->GetIndex(InParentKey);
 	if (Index == INDEX_NONE)
 	{
 		return Collection;
@@ -40,76 +82,16 @@ FRigElementKeyCollection FRigElementKeyCollection::MakeFromChildren(
 	for (int32 ParentIndex = 0; ParentIndex < ParentKeys.Num(); ParentIndex++)
 	{
 		const FRigElementKey ParentKey = ParentKeys[ParentIndex];
-		switch (ParentKey.Type)
+		TArray<FRigElementKey> Children = InHierarchy->GetChildren(ParentKey);
+		for(const FRigElementKey& Child : Children)
 		{
-			case ERigElementType::Bone:
-			{
-				const FRigBoneHierarchy& Bones = InContainer->BoneHierarchy;
-				const FRigBone& Bone = Bones[ParentKey.Name];
-
-				if (bAddBones)
-				{
-					for (int32 Dependent : Bone.Dependents)
-					{
-						FRigElementKey DependentKey = Bones[Dependent].GetElementKey();
-						Collection.AddUnique(DependentKey);
-
-						if (bRecursive)
-						{
-							ParentKeys.AddUnique(DependentKey);
-						}
-					}
-				}
-
-				// fall through so that we get the spaces under bones as well
-			}
-			case ERigElementType::Space:
-			case ERigElementType::Control:
+			if((InElementTypes & (uint8)Child.Type) == (uint8)Child.Type)
 			{
-				if (bAddSpaces)
+				const int32 PreviousSize = Collection.Num();
+				if(PreviousSize == Collection.AddUnique(Child))
 				{
-					const FRigSpaceHierarchy& Spaces = InContainer->SpaceHierarchy;
-					for (const FRigSpace& Space : Spaces)
-					{
-						if (Space.GetParentElementKey() == ParentKey)
-						{
-							FRigElementKey DependentKey = Space.GetElementKey();
-							Collection.AddUnique(DependentKey);
-
-							if (bRecursive)
-							{
-								ParentKeys.AddUnique(DependentKey);
-							}
-						}
-					}
+					ParentKeys.Add(Child);
 				}
-
-				if (bAddControls && ParentKey.Type != ERigElementType::Bone)
-				{
-					const FRigControlHierarchy& Controls = InContainer->ControlHierarchy;
-					for (const FRigControl& Control : Controls)
-					{
-						if (Control.GetSpaceElementKey() == ParentKey || 
-							Control.GetParentElementKey() == ParentKey)
-						{
-							FRigElementKey DependentKey = Control.GetElementKey();
-							Collection.AddUnique(DependentKey);
-
-							if (bRecursive)
-							{
-								ParentKeys.AddUnique(DependentKey);
-							}
-						}
-					}
-				}
-				break;
-			}
-			case ERigElementType::Curve:
-			case ERigElementType::All:
-			case ERigElementType::None:
-			default:
-			{
-				break;
 			}
 		}
 	}
@@ -118,37 +100,37 @@ FRigElementKeyCollection FRigElementKeyCollection::MakeFromChildren(
 }
 
 FRigElementKeyCollection FRigElementKeyCollection::MakeFromName(
-	const FRigHierarchyContainer* InContainer,
+	URigHierarchy* InHierarchy,
 	const FName& InPartialName,
 	uint8 InElementTypes
 )
 {
 	if (InPartialName.IsNone())
 	{
-		return MakeFromCompleteHierarchy(InContainer, InElementTypes);
+		return MakeFromCompleteHierarchy(InHierarchy, InElementTypes);
 	}
 
-	check(InContainer);
+	check(InHierarchy);
 
-	FRigElementKeyCollection Collection(InContainer->GetAllItems(true));
+	FRigElementKeyCollection Collection(InHierarchy->GetAllKeys(true));
 	Collection = Collection.FilterByType(InElementTypes);
 	Collection = Collection.FilterByName(InPartialName);
 	return Collection;
 }
 
 FRigElementKeyCollection FRigElementKeyCollection::MakeFromChain(
-	const FRigHierarchyContainer* InContainer,
+	URigHierarchy* InHierarchy,
 	const FRigElementKey& InFirstItem,
 	const FRigElementKey& InLastItem,
 	bool bReverse
 )
 {
-	check(InContainer);
+	check(InHierarchy);
 
 	FRigElementKeyCollection Collection;
 
-	int32 FirstIndex = InContainer->GetIndex(InFirstItem);
-	int32 LastIndex = InContainer->GetIndex(InLastItem);
+	int32 FirstIndex = InHierarchy->GetIndex(InFirstItem);
+	int32 LastIndex = InHierarchy->GetIndex(InLastItem);
 
 	if (FirstIndex == INDEX_NONE || LastIndex == INDEX_NONE)
 	{
@@ -159,7 +141,7 @@ FRigElementKeyCollection FRigElementKeyCollection::MakeFromChain(
 	while (LastKey.IsValid() && LastKey != InFirstItem)
 	{
 		Collection.Keys.Add(LastKey);
-		LastKey = InContainer->GetParentKey(LastKey);
+		LastKey = InHierarchy->GetFirstParent(LastKey);
 	}
 
 	if (LastKey != InFirstItem)
@@ -180,13 +162,13 @@ FRigElementKeyCollection FRigElementKeyCollection::MakeFromChain(
 }
 
 FRigElementKeyCollection FRigElementKeyCollection::MakeFromCompleteHierarchy(
-	const FRigHierarchyContainer* InContainer,
+	URigHierarchy* InHierarchy,
 	uint8 InElementTypes
 )
 {
-	check(InContainer);
+	check(InHierarchy);
 
-	FRigElementKeyCollection Collection(InContainer->GetAllItems(true));
+	FRigElementKeyCollection Collection(InHierarchy->GetAllKeys(true));
 	return Collection.FilterByType(InElementTypes);
 }
 
@@ -268,3 +250,52 @@ FRigElementKeyCollection FRigElementKeyCollection::FilterByName(const FName& InP
 	}
 	return Collection;
 }
+
+////////////////////////////////////////////////////////////////////////////////
+// FRigMirrorSettings
+////////////////////////////////////////////////////////////////////////////////
+
+FTransform FRigMirrorSettings::MirrorTransform(const FTransform& InTransform) const
+{
+	FTransform Transform = InTransform;
+	FQuat Quat = Transform.GetRotation();
+
+	Transform.SetLocation(MirrorVector(Transform.GetLocation()));
+
+	switch (AxisToFlip)
+	{
+	case EAxis::X:
+		{
+			FVector Y = MirrorVector(Quat.GetAxisY());
+			FVector Z = MirrorVector(Quat.GetAxisZ());
+			FMatrix Rotation = FRotationMatrix::MakeFromYZ(Y, Z);
+			Transform.SetRotation(FQuat(Rotation));
+			break;
+		}
+	case EAxis::Y:
+		{
+			FVector X = MirrorVector(Quat.GetAxisX());
+			FVector Z = MirrorVector(Quat.GetAxisZ());
+			FMatrix Rotation = FRotationMatrix::MakeFromXZ(X, Z);
+			Transform.SetRotation(FQuat(Rotation));
+			break;
+		}
+	default:
+		{
+			FVector X = MirrorVector(Quat.GetAxisX());
+			FVector Y = MirrorVector(Quat.GetAxisY());
+			FMatrix Rotation = FRotationMatrix::MakeFromXY(X, Y);
+			Transform.SetRotation(FQuat(Rotation));
+			break;
+		}
+	}
+
+	return Transform;
+}
+
+FVector FRigMirrorSettings::MirrorVector(const FVector& InVector) const
+{
+	FVector Axis = FVector::ZeroVector;
+	Axis.SetComponentForAxis(MirrorAxis, 1.f);
+	return InVector.MirrorByVector(Axis);
+}
diff --git a/Engine/Plugins/Experimental/ControlRig/Source/ControlRig/Private/Rigs/RigHierarchyElements.cpp b/Engine/Plugins/Experimental/ControlRig/Source/ControlRig/Private/Rigs/RigHierarchyElements.cpp
new file mode 100644
--- /dev/null
+++ b/Engine/Plugins/Experimental/ControlRig/Source/ControlRig/Private/Rigs/RigHierarchyElements.cpp
@@ -0,0 +1,456 @@
+// Copyright Epic Games, Inc. All Rights Reserved.
+
+#include "Rigs/RigHierarchyElements.h"
+#include "Rigs/RigHierarchy.h"
+#include "UObject/AnimObjectVersion.h"
+
+////////////////////////////////////////////////////////////////////////////////
+// FRigBaseElement
+////////////////////////////////////////////////////////////////////////////////
+
+void FRigBaseElement::Serialize(FArchive& Ar, URigHierarchy* Hierarchy, ESerializationPhase SerializationPhase)
+{
+	Ar.UsingCustomVersion(FAnimObjectVersion::GUID);
+
+	if (Ar.IsSaving() || Ar.IsObjectReferenceCollector())
+	{
+		Save(Ar, Hierarchy, SerializationPhase);
+	}
+	else if (Ar.IsLoading())
+	{
+		Load(Ar, Hierarchy, SerializationPhase);
+	}
+	else
+	{
+		checkNoEntry();
+	}
+}
+
+void FRigBaseElement::Save(FArchive& Ar, URigHierarchy* Hierarchy, ESerializationPhase SerializationPhase)
+{
+	if(SerializationPhase == ESerializationPhase::StaticData)
+	{
+		Ar << Key;
+	}
+}
+
+void FRigBaseElement::Load(FArchive& Ar, URigHierarchy* Hierarchy, ESerializationPhase SerializationPhase)
+{
+	if(SerializationPhase == ESerializationPhase::StaticData)
+	{
+		FRigElementKey LoadedKey;
+	
+		Ar << LoadedKey;
+
+		ensure(LoadedKey.Type == Key.Type);
+		Key = LoadedKey;
+	}
+}
+
+////////////////////////////////////////////////////////////////////////////////
+// FRigComputedTransform
+////////////////////////////////////////////////////////////////////////////////
+
+void FRigComputedTransform::Save(FArchive& Ar)
+{
+	Ar << Transform;
+	Ar << bDirty;
+}
+
+void FRigComputedTransform::Load(FArchive& Ar)
+{
+	// load and save are identical
+	Save(Ar);
+}
+
+////////////////////////////////////////////////////////////////////////////////
+// FRigLocalAndGlobalTransform
+////////////////////////////////////////////////////////////////////////////////
+
+void FRigLocalAndGlobalTransform::Save(FArchive& Ar)
+{
+	Local.Save(Ar);
+	Global.Save(Ar);
+}
+
+void FRigLocalAndGlobalTransform::Load(FArchive& Ar)
+{
+	Local.Load(Ar);
+	Global.Load(Ar);
+}
+
+////////////////////////////////////////////////////////////////////////////////
+// FRigCurrentAndInitialTransform
+////////////////////////////////////////////////////////////////////////////////
+
+void FRigCurrentAndInitialTransform::Save(FArchive& Ar)
+{
+	Current.Save(Ar);
+	Initial.Save(Ar);
+}
+
+void FRigCurrentAndInitialTransform::Load(FArchive& Ar)
+{
+	Current.Load(Ar);
+	Initial.Load(Ar);
+}
+
+////////////////////////////////////////////////////////////////////////////////
+// FRigTransformElement
+////////////////////////////////////////////////////////////////////////////////
+
+void FRigTransformElement::Save(FArchive& Ar, URigHierarchy* Hierarchy, ESerializationPhase SerializationPhase)
+{
+	Super::Save(Ar, Hierarchy, SerializationPhase);
+
+	if(SerializationPhase == ESerializationPhase::StaticData)
+	{
+		Pose.Save(Ar);
+	}
+}
+
+void FRigTransformElement::Load(FArchive& Ar, URigHierarchy* Hierarchy, ESerializationPhase SerializationPhase)
+{
+	Super::Load(Ar, Hierarchy, SerializationPhase);
+
+	if(SerializationPhase == ESerializationPhase::StaticData)
+	{
+		Pose.Load(Ar);
+	}
+}
+
+void FRigTransformElement::CopyPose(FRigBaseElement* InOther, bool bCurrent, bool bInitial)
+{
+	if(FRigTransformElement* Other = Cast<FRigTransformElement>(InOther))
+	{
+		if(bCurrent)
+		{
+			Pose.Current = Other->Pose.Current;
+		}
+		if(bInitial)
+		{
+			Pose.Initial = Other->Pose.Initial;
+		}
+	}
+}
+
+////////////////////////////////////////////////////////////////////////////////
+// FRigSingleParentElement
+////////////////////////////////////////////////////////////////////////////////
+
+void FRigSingleParentElement::Save(FArchive& Ar, URigHierarchy* Hierarchy, ESerializationPhase SerializationPhase)
+{
+	Super::Save(Ar, Hierarchy, SerializationPhase);
+
+	if(SerializationPhase == ESerializationPhase::InterElementData)
+	{
+		FRigElementKey ParentKey;
+		if(ParentElement)
+		{
+			ParentKey = ParentElement->GetKey();
+		}
+		Ar << ParentKey;
+	}
+}
+
+void FRigSingleParentElement::Load(FArchive& Ar, URigHierarchy* Hierarchy, ESerializationPhase SerializationPhase)
+{
+	Super::Load(Ar, Hierarchy, SerializationPhase);
+
+	if(SerializationPhase == ESerializationPhase::InterElementData)
+	{
+		FRigElementKey ParentKey;
+		Ar << ParentKey;
+
+		if(ParentKey.IsValid())
+		{
+			ParentElement = Hierarchy->FindChecked<FRigTransformElement>(ParentKey);
+		}
+	}
+}
+
+////////////////////////////////////////////////////////////////////////////////
+// FRigMultiParentElement
+////////////////////////////////////////////////////////////////////////////////
+
+void FRigMultiParentElement::Save(FArchive& Ar, URigHierarchy* Hierarchy, ESerializationPhase SerializationPhase)
+{
+	Super::Save(Ar, Hierarchy, SerializationPhase);
+
+	if(SerializationPhase == ESerializationPhase::StaticData)
+	{
+		Parent.Save(Ar);
+
+		int32 NumParents = ParentElements.Num();
+		Ar << NumParents;
+	}
+	else if(SerializationPhase == ESerializationPhase::InterElementData)
+	{
+		for(int32 ParentIndex = 0; ParentIndex < ParentElements.Num(); ParentIndex++)
+		{
+			FRigElementKey ParentKey;
+			if(ParentElements[ParentIndex])
+			{
+				ParentKey = ParentElements[ParentIndex]->GetKey();
+			}
+
+			Ar << ParentKey;
+			Ar << ParentWeightsInitial[ParentIndex];
+			Ar << ParentWeights[ParentIndex];
+		}
+	}
+}
+
+void FRigMultiParentElement::Load(FArchive& Ar, URigHierarchy* Hierarchy, ESerializationPhase SerializationPhase)
+{
+	Super::Load(Ar, Hierarchy, SerializationPhase);
+
+	if(SerializationPhase == ESerializationPhase::StaticData)
+	{
+		Parent.Load(Ar);
+
+		int32 NumParents = 0;
+		Ar << NumParents;
+
+		ParentElements.SetNumZeroed(NumParents);
+		ParentWeights.SetNumZeroed(NumParents);
+		ParentWeightsInitial.SetNumZeroed(NumParents);
+	}
+	else if(SerializationPhase == ESerializationPhase::InterElementData)
+	{
+		for(int32 ParentIndex = 0; ParentIndex < ParentElements.Num(); ParentIndex++)
+		{
+			FRigElementKey ParentKey;
+			Ar << ParentKey;
+			ensure(ParentKey.IsValid());
+
+			ParentElements[ParentIndex] = Hierarchy->FindChecked<FRigTransformElement>(ParentKey);
+
+			Ar << ParentWeightsInitial[ParentIndex];
+			Ar << ParentWeights[ParentIndex];
+
+			IndexLookup.Add(ParentKey, ParentIndex);
+		}
+	}
+}
+
+void FRigMultiParentElement::CopyPose(FRigBaseElement* InOther, bool bCurrent, bool bInitial)
+{
+	if(FRigMultiParentElement* Other = Cast<FRigMultiParentElement>(InOther))
+	{
+		if(bCurrent)
+		{
+			Parent.Current = Other->Parent.Current;
+		}
+		if(bInitial)
+		{
+			Parent.Initial = Other->Parent.Initial;
+		}
+	}
+}
+
+////////////////////////////////////////////////////////////////////////////////
+// FRigBoneElement
+////////////////////////////////////////////////////////////////////////////////
+
+void FRigBoneElement::Save(FArchive& Ar, URigHierarchy* Hierarchy, ESerializationPhase SerializationPhase)
+{
+	Super::Save(Ar, Hierarchy, SerializationPhase);
+
+	if(SerializationPhase == ESerializationPhase::StaticData)
+	{
+		static const UEnum* BoneTypeEnum = StaticEnum<ERigBoneType>();
+		FName TypeName = BoneTypeEnum->GetNameByValue((int64)BoneType);
+		Ar << TypeName;
+	}
+}
+
+void FRigBoneElement::Load(FArchive& Ar, URigHierarchy* Hierarchy, ESerializationPhase SerializationPhase)
+{
+	Super::Load(Ar, Hierarchy, SerializationPhase);
+
+	if(SerializationPhase == ESerializationPhase::StaticData)
+	{
+		static const UEnum* BoneTypeEnum = StaticEnum<ERigBoneType>();
+		FName TypeName;
+		Ar << TypeName;
+		BoneType = (ERigBoneType)BoneTypeEnum->GetValueByName(TypeName);
+	}
+}
+
+////////////////////////////////////////////////////////////////////////////////
+// FRigControlSettings
+////////////////////////////////////////////////////////////////////////////////
+
+FRigControlSettings::FRigControlSettings()
+: ControlType(ERigControlType::Transform)
+, DisplayName(NAME_None)
+, PrimaryAxis(ERigControlAxis::X)
+, bIsCurve(false)
+, bAnimatable(true)
+, bLimitTranslation(false)
+, bLimitRotation(false)
+, bLimitScale(false)
+, bDrawLimits(true)
+, MinimumValue()
+, MaximumValue()
+, bGizmoEnabled(true)
+, bGizmoVisible(true)
+, GizmoName(TEXT("Gizmo"))
+, GizmoColor(FLinearColor::Red)
+, bIsTransientControl(false)
+, ControlEnum(nullptr)
+{
+}
+
+void FRigControlSettings::Save(FArchive& Ar)
+{
+	static const UEnum* ControlTypeEnum = StaticEnum<ERigControlType>();
+	static const UEnum* ControlAxisEnum = StaticEnum<ERigControlAxis>();
+
+	FName ControlTypeName = ControlTypeEnum->GetNameByValue((int64)ControlType);
+	FName PrimaryAxisName = ControlAxisEnum->GetNameByValue((int64)PrimaryAxis);
+	FTransform MinimumTransform = MinimumValue.GetAsTransform(ControlType, PrimaryAxis);
+	FTransform MaximumTransform = MaximumValue.GetAsTransform(ControlType, PrimaryAxis);
+	FString ControlEnumPathName;
+	if(ControlEnum)
+	{
+		ControlEnumPathName = ControlEnum->GetPathName();
+	}
+
+	Ar << ControlTypeName;
+	Ar << DisplayName;
+	Ar << PrimaryAxisName;
+	Ar << bIsCurve;
+	Ar << bAnimatable;
+	Ar << bLimitTranslation;
+	Ar << bLimitRotation;
+	Ar << bLimitScale;
+	Ar << bDrawLimits;
+	Ar << MinimumTransform;
+	Ar << MaximumTransform;
+	Ar << bGizmoEnabled;
+	Ar << bGizmoVisible;
+	Ar << GizmoName;
+	Ar << GizmoColor;
+	Ar << bIsTransientControl;
+	Ar << ControlEnumPathName;
+}
+
+void FRigControlSettings::Load(FArchive& Ar)
+{
+	static const UEnum* ControlTypeEnum = StaticEnum<ERigControlType>();
+	static const UEnum* ControlAxisEnum = StaticEnum<ERigControlAxis>();
+
+	FName ControlTypeName, PrimaryAxisName;
+	FTransform MinimumTransform, MaximumTransform;
+	FString ControlEnumPathName;
+
+	Ar << ControlTypeName;
+	Ar << DisplayName;
+	Ar << PrimaryAxisName;
+	Ar << bIsCurve;
+	Ar << bAnimatable;
+	Ar << bLimitTranslation;
+	Ar << bLimitRotation;
+	Ar << bLimitScale;
+	Ar << bDrawLimits;
+	Ar << MinimumTransform;
+	Ar << MaximumTransform;
+	Ar << bGizmoEnabled;
+	Ar << bGizmoVisible;
+	Ar << GizmoName;
+	Ar << GizmoColor;
+	Ar << bIsTransientControl;
+	Ar << ControlEnumPathName;
+
+	ControlType = (ERigControlType)ControlTypeEnum->GetValueByName(ControlTypeName);
+	PrimaryAxis = (ERigControlAxis)ControlAxisEnum->GetValueByName(PrimaryAxisName);
+	MinimumValue.SetFromTransform(MinimumTransform, ControlType, PrimaryAxis);
+	MaximumValue.SetFromTransform(MaximumTransform, ControlType, PrimaryAxis);
+
+	ControlEnum = nullptr;
+	if(!ControlEnumPathName.IsEmpty())
+	{
+		ControlEnum = FindObject<UEnum>(ANY_PACKAGE, *ControlEnumPathName);
+	}
+}
+
+////////////////////////////////////////////////////////////////////////////////
+// FRigControlElement
+////////////////////////////////////////////////////////////////////////////////
+
+void FRigControlElement::Save(FArchive& Ar, URigHierarchy* Hierarchy, ESerializationPhase SerializationPhase)
+{
+	Super::Save(Ar, Hierarchy, SerializationPhase);
+
+	if(SerializationPhase == ESerializationPhase::StaticData)
+	{
+		Settings.Save(Ar);
+		Offset.Save(Ar);
+		Gizmo.Save(Ar);
+	}
+}
+
+void FRigControlElement::Load(FArchive& Ar, URigHierarchy* Hierarchy, ESerializationPhase SerializationPhase)
+{
+	Super::Load(Ar, Hierarchy, SerializationPhase);
+
+	if(SerializationPhase == ESerializationPhase::StaticData)
+	{
+		Settings.Load(Ar);
+		Offset.Load(Ar);
+		Gizmo.Load(Ar);
+	}
+}
+
+void FRigControlElement::CopyPose(FRigBaseElement* InOther, bool bCurrent, bool bInitial)
+{
+	if(FRigControlElement* Other = Cast<FRigControlElement>(InOther))
+	{
+		if(bCurrent)
+		{
+			Offset.Current = Other->Offset.Current;
+			Gizmo.Current = Other->Gizmo.Current;
+		}
+		if(bInitial)
+		{
+			Offset.Initial = Other->Offset.Initial;
+			Gizmo.Initial = Other->Gizmo.Initial;
+		}
+	}
+}
+
+////////////////////////////////////////////////////////////////////////////////
+// FRigCurveElement
+////////////////////////////////////////////////////////////////////////////////
+
+void FRigCurveElement::Save(FArchive& Ar, URigHierarchy* Hierarchy, ESerializationPhase SerializationPhase)
+{
+	Super::Save(Ar, Hierarchy, SerializationPhase);
+
+	if(SerializationPhase == ESerializationPhase::StaticData)
+	{
+		Ar << Value;
+	}
+}
+
+void FRigCurveElement::Load(FArchive& Ar, URigHierarchy* Hierarchy, ESerializationPhase SerializationPhase)
+{
+	Super::Load(Ar, Hierarchy, SerializationPhase);
+
+	if(SerializationPhase == ESerializationPhase::StaticData)
+	{
+		Ar << Value;
+	}
+}
+
+void FRigCurveElement::CopyPose(FRigBaseElement* InOther, bool bCurrent, bool bInitial)
+{
+	Super::CopyPose(InOther, bCurrent, bInitial);
+	
+	if(FRigCurveElement* Other = Cast<FRigCurveElement>(InOther))
+	{
+		Value = Other->Value;
+	}
+}
diff --git a/Engine/Plugins/Experimental/ControlRig/Source/ControlRig/Private/Rigs/RigSpaceHierarchy.cpp b/Engine/Plugins/Experimental/ControlRig/Source/ControlRig/Private/Rigs/RigSpaceHierarchy.cpp
--- a/Engine/Plugins/Experimental/ControlRig/Source/ControlRig/Private/Rigs/RigSpaceHierarchy.cpp
+++ b/Engine/Plugins/Experimental/ControlRig/Source/ControlRig/Private/Rigs/RigSpaceHierarchy.cpp
@@ -10,717 +10,5 @@
 ////////////////////////////////////////////////////////////////////////////////
 
 FRigSpaceHierarchy::FRigSpaceHierarchy()
-	:Container(nullptr)
 {
 }
-
-FRigSpaceHierarchy& FRigSpaceHierarchy::operator= (const FRigSpaceHierarchy &InOther)
-{
-#if WITH_EDITOR
-	for (int32 Index = Num() - 1; Index >= 0; Index--)
-	{
-		FRigSpace SpaceToRemove = Spaces[Index];
-		OnSpaceRemoved.Broadcast(Container, FRigElementKey(SpaceToRemove.Name, ERigElementType::Space));
-	}
-#endif
-
-	Spaces.Reset();
-	Spaces.Append(InOther.Spaces);
-	NameToIndexMapping.Reset();
-	RefreshMapping();
-
-#if WITH_EDITOR
-	for (const FRigSpace& SpaceAdded : Spaces)
-	{
-		OnSpaceAdded.Broadcast(Container, FRigElementKey(SpaceAdded.Name, ERigElementType::Space));
-	}
-#endif
-
-	return *this;
-}
-
-bool FRigSpaceHierarchy::IsNameAvailable(const FString& InPotentialNewName, FString* OutErrorMessage) const
-{
-	if (Container)
-	{
-		return Container->IsNameAvailable(InPotentialNewName, ERigElementType::Space, OutErrorMessage);
-	}
-
-	// if we are not part of a container, fall back to index based validation
-	if (InPotentialNewName.Len() > FRigHierarchyContainer::GetMaxNameLength())
-	{
-		return false;
-	}
-	return GetIndex(*InPotentialNewName) == INDEX_NONE;
-}
-
-FName FRigSpaceHierarchy::GetSafeNewName(const FString& InPotentialNewName) const
-{
-	return FRigHierarchyContainer::GetSafeNewName(InPotentialNewName, [this](const FString& InName) -> bool {
-		return IsNameAvailable(InName);
-	});
-}
-
-FRigSpace& FRigSpaceHierarchy::Add(const FName& InNewName, ERigSpaceType InSpaceType, const FName& InParentName, const FTransform& InTransform)
-{
-	FRigSpace NewSpace;
-	NewSpace.Name = GetSafeNewName(InNewName.ToString());
-	NewSpace.ParentIndex = GetParentIndex(InSpaceType, InParentName);
-	NewSpace.SpaceType = NewSpace.ParentIndex == INDEX_NONE ? ERigSpaceType::Global : InSpaceType;
-	NewSpace.ParentName = NewSpace.ParentIndex == INDEX_NONE ? NAME_None : InParentName;
-	NewSpace.InitialTransform = InTransform;
-	NewSpace.LocalTransform = InTransform;
-	FName NewSpaceName = NewSpace.Name;
-	Spaces.Add(NewSpace);
-	RefreshMapping();
-
-#if WITH_EDITOR
-	OnSpaceAdded.Broadcast(Container, NewSpace.GetElementKey());
-#endif
-
-	int32 Index = GetIndex(NewSpaceName);
-	return Spaces[Index];
-}
-
-bool FRigSpaceHierarchy::Reparent(const FName& InName, ERigSpaceType InSpaceType, const FName& InNewParentName)
-{
-	int32 Index = GetIndex(InName);
-	if (Index != INDEX_NONE)
-	{
-		FRigSpace& Space = Spaces[Index];
-
-#if WITH_EDITOR
-		FName OldParentName = Space.ParentName;
-#endif
-
-		int32 ParentIndex = GetParentIndex(InSpaceType, InNewParentName);
-		if (ParentIndex != INDEX_NONE)
-		{
-			if(ParentIndex == Index)
-			{
-				ParentIndex = INDEX_NONE;
-			}
-			else if (Container != nullptr)
-			{
-				switch (InSpaceType)
-				{
-					case ERigSpaceType::Global:
-					{
-						ParentIndex = INDEX_NONE;
-						break;
-					}
-					case ERigSpaceType::Bone:
-					{
-						break;
-					}
-					case ERigSpaceType::Space:
-					{
-						if (Container->IsParentedTo(ERigElementType::Space, ParentIndex, ERigElementType::Space, Index))
-						{
-							ParentIndex = INDEX_NONE;
-						}
-						break;
-					}
-					case ERigSpaceType::Control:
-					{
-						if (Container->IsParentedTo(ERigElementType::Control, ParentIndex, ERigElementType::Space, Index))
-						{
-							ParentIndex = INDEX_NONE;
-						}
-						break;
-					}
-				}
-			}
-		}
-
-		Space.ParentIndex = ParentIndex;
-		Space.SpaceType = Space.ParentIndex == INDEX_NONE ? ERigSpaceType::Global : InSpaceType;
-		Space.ParentName = Space.ParentIndex == INDEX_NONE ? NAME_None : InNewParentName;
-
-#if WITH_EDITOR
-		FName NewParentName = Space.ParentName;
-#endif
-
-		RefreshMapping();
-
-#if WITH_EDITOR
-		if (OldParentName != NewParentName)
-		{
-			OnSpaceReparented.Broadcast(Container, Space.GetElementKey(), OldParentName, NewParentName);
-		}
-#endif		
-		return Spaces[GetIndex(InName)].ParentName == InNewParentName;
-	}
-	return false;
-}
-
-FRigSpace FRigSpaceHierarchy::Remove(const FName& InNameToRemove)
-{
-	int32 IndexToDelete = GetIndex(InNameToRemove);
-	Select(InNameToRemove, false);
-	FRigSpace RemovedSpace = Spaces[IndexToDelete];
-	Spaces.RemoveAt(IndexToDelete);
-
-#if WITH_EDITOR
-	TArray<FName> SpacesReparented;
-#endif
-	for (FRigSpace& ChildSpace : Spaces)
-	{
-		if (ChildSpace.SpaceType == ERigSpaceType::Space && ChildSpace.ParentName == InNameToRemove)
-		{
-			ChildSpace.SpaceType = ERigSpaceType::Global;
-			ChildSpace.ParentIndex = INDEX_NONE;
-			ChildSpace.ParentName = NAME_None;
-#if WITH_EDITOR
-			SpacesReparented.Add(ChildSpace.Name);
-#endif
-		}
-	}
-
-	RefreshMapping();
-
-#if WITH_EDITOR
-	for (const FName& ReparentedSpace : SpacesReparented)
-	{
-		OnSpaceReparented.Broadcast(Container, FRigElementKey(ReparentedSpace, RigElementType()), InNameToRemove, NAME_None);
-	}
-	OnSpaceRemoved.Broadcast(Container, RemovedSpace.GetElementKey());
-#endif
-
-	return RemovedSpace;
-}
-
-FName FRigSpaceHierarchy::GetName(int32 InIndex) const
-{
-	if (Spaces.IsValidIndex(InIndex))
-	{
-		return Spaces[InIndex].Name;
-	}
-
-	return NAME_None;
-}
-
-int32 FRigSpaceHierarchy::GetIndexSlow(const FName& InName) const
-{
-	for (int32 Index = 0; Index < Spaces.Num(); ++Index)
-	{
-		if (Spaces[Index].Name == InName)
-		{
-			return Index;
-		}
-	}
-
-	return INDEX_NONE;
-}
-
-void FRigSpaceHierarchy::SetGlobalTransform(const FName& InName, const FTransform& InTransform)
-{
-	SetGlobalTransform(GetIndex(InName), InTransform);
-}
-
-void FRigSpaceHierarchy::SetGlobalTransform(int32 InIndex, const FTransform& InTransform)
-{
-	if (Container == nullptr)
-	{
-		SetLocalTransform(InIndex, InTransform);
-		return;
-	}
-
-	if (Spaces.IsValidIndex(InIndex))
-	{
-		FRigSpace& Space = Spaces[InIndex];
-		if (Space.SpaceType == ERigSpaceType::Global)
-		{
-			SetLocalTransform(InIndex, InTransform);
-			return;
-		}
-
-		FTransform ParentTransform = Container->GetGlobalTransform(Space.GetParentElementKey());
-		SetLocalTransform(InIndex, InTransform.GetRelativeTransform(ParentTransform));
-	}
-}
-
-FTransform FRigSpaceHierarchy::GetGlobalTransform(const FName& InName) const
-{
-	return GetGlobalTransform(GetIndex(InName));
-}
-
-FTransform FRigSpaceHierarchy::GetGlobalTransform(int32 InIndex) const
-{
-	if(Container == nullptr)
-	{
-		return GetLocalTransform(InIndex);
-	}
-
-	if (Spaces.IsValidIndex(InIndex))
-	{
-#if WITH_EDITOR
-		// this is not threadsafe,
-		// but currently only one thread is run per character.
-		if (RecursionGuard.Num() < Spaces.Num())
-		{
-			RecursionGuard.AddZeroed(Spaces.Num() - RecursionGuard.Num());
-		}
-
-		if (RecursionGuard[InIndex])
-		{
-			TArray<FString> SpacesHitElements;
-			TArray<FString> SpaceParentsElements;
-			for (int32 SpaceIndex = 0; SpaceIndex < RecursionGuard.Num(); SpaceIndex++)
-			{
-				if (RecursionGuard[SpaceIndex])
-				{
-					SpacesHitElements.Add(Spaces[SpaceIndex].Name.ToString());
-				}
-				if (!Spaces[SpaceIndex].ParentName.IsNone())
-				{
-					SpaceParentsElements.Add(FString::Printf(TEXT("%s->%s"), *Spaces[SpaceIndex].ParentName.ToString(), *Spaces[SpaceIndex].Name.ToString()));
-				}
-			}
-			FString SpacesHit = FString::Join(SpacesHitElements, TEXT(", "));
-			FString SpaceParents = FString::Join(SpaceParentsElements, TEXT(", "));
-
-			checkf(!RecursionGuard[InIndex], TEXT("Control Rig Recursion (JIRA UE-94987)\n%s\n%s"), *SpacesHit, *SpaceParents);
-			return GetLocalTransform(InIndex);
-		}
-		TGuardValue<bool> Guard(RecursionGuard[InIndex], true);
-#endif
-
-		const FRigSpace& Space = Spaces[InIndex];
-		switch (Space.SpaceType)
-		{
-			case ERigSpaceType::Global:
-			{
-				return Space.LocalTransform;
-			}
-			case ERigSpaceType::Bone:
-			{
-				FTransform Transform = Space.LocalTransform * Container->GetGlobalTransform(ERigElementType::Bone, Space.ParentIndex);
-				Transform.NormalizeRotation();
-				return Transform;
-			}
-			case ERigSpaceType::Space:
-			{
-				FTransform Transform = Space.LocalTransform * Container->GetGlobalTransform(ERigElementType::Space, Space.ParentIndex);
-				Transform.NormalizeRotation();
-				return Transform;
-			}
-			case ERigSpaceType::Control:
-			{
-				FTransform Transform = Space.LocalTransform * Container->GetGlobalTransform(ERigElementType::Control, Space.ParentIndex);
-				Transform.NormalizeRotation();
-				return Transform;
-			}
-		}
-	}
-
-	return FTransform::Identity;
-}
-
-void FRigSpaceHierarchy::SetLocalTransform(const FName& InName, const FTransform& InTransform)
-{
-	SetLocalTransform(GetIndex(InName), InTransform);
-}
-
-void FRigSpaceHierarchy::SetLocalTransform(int32 InIndex, const FTransform& InTransform)
-{
-	if (Spaces.IsValidIndex(InIndex))
-	{
-		FRigSpace& Space = Spaces[InIndex];
-		Space.LocalTransform = InTransform;
-		Space.LocalTransform.NormalizeRotation();
-	}
-}
-
-FTransform FRigSpaceHierarchy::GetLocalTransform(const FName& InName) const
-{
-	return GetLocalTransform(GetIndex(InName));
-}
-
-FTransform FRigSpaceHierarchy::GetLocalTransform(int32 InIndex) const
-{
-	if (Spaces.IsValidIndex(InIndex))
-	{
-		return Spaces[InIndex].LocalTransform;
-	}
-
-	return FTransform::Identity;
-}
-
-void FRigSpaceHierarchy::SetInitialTransform(const FName& InName, const FTransform& InTransform)
-{
-	SetInitialTransform(GetIndex(InName), InTransform);
-}
-
-void FRigSpaceHierarchy::SetInitialTransform(int32 InIndex, const FTransform& InTransform)
-{
-	if (Spaces.IsValidIndex(InIndex))
-	{
-		FRigSpace& Space = Spaces[InIndex];
-		Space.InitialTransform = InTransform;
-		Space.InitialTransform.NormalizeRotation();
-	}
-}
-
-FTransform FRigSpaceHierarchy::GetInitialTransform(const FName& InName) const
-{
-	return GetInitialTransform(GetIndex(InName));
-}
-
-FTransform FRigSpaceHierarchy::GetInitialTransform(int32 InIndex) const
-{
-	if (Spaces.IsValidIndex(InIndex))
-	{
-		return Spaces[InIndex].InitialTransform;
-	}
-
-	return FTransform::Identity;
-}
-
-void FRigSpaceHierarchy::SetInitialGlobalTransform(const FName& InName, const FTransform& GlobalTransform)
-{
-	SetInitialGlobalTransform(GetIndex(InName), GlobalTransform);
-}
-
-// wip - should support all types that can provide space transform data (pos/rotation)
-void FRigSpaceHierarchy::SetInitialGlobalTransform(int32 InIndex, const FTransform& GlobalTransform)
-{
-	if (Spaces.IsValidIndex(InIndex))
-	{
-		FRigSpace& Space = Spaces[InIndex];
-		FTransform ParentTransform = FTransform::Identity;
-		if (Container)
-		{
-			ParentTransform = Container->GetInitialGlobalTransform(Space.GetParentElementKey());
-		}
-		Space.InitialTransform = GlobalTransform.GetRelativeTransform(ParentTransform);
-	}
-}
-
-FTransform FRigSpaceHierarchy::GetInitialGlobalTransform(const FName& InName) const
-{
-	return GetInitialGlobalTransform(GetIndex(InName));
-}
-
-FTransform FRigSpaceHierarchy::GetInitialGlobalTransform(int32 InIndex) const
-{
-	// @todo: Templatize
-	if (Spaces.IsValidIndex(InIndex))
-	{
-		const FRigSpace& Space = Spaces[InIndex];
-		FTransform ParentTransform = FTransform::Identity;
-		if (Container)
-		{
-			ParentTransform = Container->GetInitialGlobalTransform(Space.GetParentElementKey());
-		}
-		FTransform Transform = Space.InitialTransform * ParentTransform;
-		Transform.NormalizeRotation();
-		return Transform;
-	}
-	return FTransform::Identity;
-}
-
-FName FRigSpaceHierarchy::Rename(const FName& InOldName, const FName& InNewName)
-{
-	if (InOldName != InNewName)
-	{
-		const int32 Found = GetIndex(InOldName);
-		if (Found != INDEX_NONE)
-		{
-			FName NewName = GetSafeNewName(InNewName.ToString());
-
-			bool bWasSelected = IsSelected(InOldName);
-			if(bWasSelected)
-			{
-				Select(InOldName, false);
-			}
-
-			Spaces[Found].Name = NewName;
-
-			// go through find all children and rename them
-			for (int32 Index = 0; Index < Spaces.Num(); ++Index)
-			{
-				if (Spaces[Index].SpaceType == ERigSpaceType::Space && Spaces[Index].ParentName == InOldName)
-				{
-					Spaces[Index].ParentName = NewName;
-				}
-			}
-
-			RefreshMapping();
-
-#if WITH_EDITOR
-			OnSpaceRenamed.Broadcast(Container, RigElementType(), InOldName, NewName);
-#endif
-			if(bWasSelected)
-			{
-				Select(NewName, true);
-			}
-			return NewName;
-		}
-	}
-
-	return NAME_None;
-}
-
-void FRigSpaceHierarchy::RefreshMapping()
-{
-	NameToIndexMapping.Empty();
-	for (int32 Index = 0; Index < Spaces.Num(); ++Index)
-	{
-		Spaces[Index].Index = Index;
-		NameToIndexMapping.Add(Spaces[Index].Name, Index);
-	}
-}
-
-void FRigSpaceHierarchy::Initialize(bool bResetTransforms)
-{
-	RefreshMapping();
-
-	// initialize transform
-	for (int32 Index = 0; Index < Spaces.Num(); ++Index)
-	{
-		if (bResetTransforms)
-		{
-			Spaces[Index].LocalTransform = Spaces[Index].InitialTransform;
-		}
-		if (Container)
-		{
-			Spaces[Index].ParentIndex = Container->GetIndex(Spaces[Index].GetParentElementKey(true /* force */));
-		}
-	}
-
-#if WITH_EDITOR
-	RecursionGuard.Reset();
-	RecursionGuard.AddZeroed(Spaces.Num());
-#endif
-}
-
-void FRigSpaceHierarchy::Reset()
-{
-	Spaces.Reset();
-}
-
-void FRigSpaceHierarchy::ResetTransforms()
-{
-	// initialize transform
-	for (int32 Index = 0; Index < Spaces.Num(); ++Index)
-	{
-		Spaces[Index].LocalTransform = Spaces[Index].InitialTransform;
-	}
-}
-
-void FRigSpaceHierarchy::CopyInitialTransforms(const FRigSpaceHierarchy& InOther)
-{
-	ensure(InOther.Num() == Num());
-
-	for (int32 Index = 0; Index < Spaces.Num(); ++Index)
-	{
-		Spaces[Index].InitialTransform = InOther.Spaces[Index].InitialTransform;
-	}
-}
-
-int32 FRigSpaceHierarchy::GetParentIndex(ERigSpaceType InSpaceType, const FName& InName) const
-{
-	if (Container != nullptr)
-	{
-		switch (InSpaceType)
-		{
-			case ERigSpaceType::Global:
-			{
-				return INDEX_NONE;
-			}
-			case ERigSpaceType::Bone:
-			{
-				return Container->GetIndex(FRigElementKey(InName, ERigElementType::Bone));
-			}
-			case ERigSpaceType::Space:
-			{
-				return Container->GetIndex(FRigElementKey(InName, ERigElementType::Space));
-			}
-			case ERigSpaceType::Control:
-			{
-				return Container->GetIndex(FRigElementKey(InName, ERigElementType::Control));
-			}
-		}
-	}
-
-	return INDEX_NONE;
-}
-
-bool FRigSpaceHierarchy::Select(const FName& InName, bool bSelect)
-{
-	if(GetIndex(InName) == INDEX_NONE)
-	{
-		return false;
-	}
-
-	if(bSelect == IsSelected(InName))
-	{
-		return false;
-	}
-
-	if(bSelect)
-	{
-		Selection.Add(InName);
-	}
-	else
-	{
-		Selection.Remove(InName);
-	}
-
-	OnSpaceSelected.Broadcast(Container, FRigElementKey(InName, RigElementType()), bSelect);
-
-	return true;
-}
-
-bool FRigSpaceHierarchy::ClearSelection()
-{
-	TArray<FName> TempSelection;
-	TempSelection.Append(Selection);
-	for(const FName& SelectedName : TempSelection)
-	{
-		Select(SelectedName, false);
-	}
-	return TempSelection.Num() > 0;
-}
-
-TArray<FName> FRigSpaceHierarchy::CurrentSelection() const
-{
-	TArray<FName> TempSelection;
-	TempSelection.Append(Selection);
-	return TempSelection;
-}
-
-bool FRigSpaceHierarchy::IsSelected(const FName& InName) const
-{
-	return Selection.Contains(InName);
-}
-
-#if WITH_EDITOR
-
-void FRigSpaceHierarchy::HandleOnElementRemoved(FRigHierarchyContainer* InContainer, const FRigElementKey& InKey)
-{
-	if (Container == nullptr)
-	{
-		return;
-	}
-
-	switch (InKey.Type)
-	{
-		case ERigElementType::Bone:
-		{
-			for (FRigSpace& Space : Spaces)
-			{
-				if (Space.SpaceType == ERigSpaceType::Bone && Space.ParentName == InKey.Name)
-				{
-					Space.ParentIndex = Container->BoneHierarchy.GetIndex(InKey.Name);
-					Space.ParentName = Space.ParentIndex == INDEX_NONE ? NAME_None : InKey.Name;
-					Space.SpaceType = Space.ParentIndex == INDEX_NONE ? ERigSpaceType::Global : Space.SpaceType;
-					OnSpaceReparented.Broadcast(Container, Space.GetElementKey(), InKey.Name, Space.ParentName);
-				}
-			}
-			break;
-		}
-		case ERigElementType::Control:
-		{
-			for (FRigSpace& Space : Spaces)
-			{
-				if (Space.SpaceType == ERigSpaceType::Control && Space.ParentName == InKey.Name)
-				{
-					Space.ParentIndex = Container->ControlHierarchy.GetIndex(InKey.Name);
-					Space.ParentName = Space.ParentIndex == INDEX_NONE ? NAME_None : InKey.Name;
-					Space.SpaceType = Space.ParentIndex == INDEX_NONE ? ERigSpaceType::Global : Space.SpaceType;
-					OnSpaceReparented.Broadcast(Container, Space.GetElementKey(), InKey.Name, Space.ParentName);
-				}
-			}
-			break;
-		}
-		case ERigElementType::Space:
-		case ERigElementType::Curve:
-		{
-			break;
-		}
-	}
-}
-
-void FRigSpaceHierarchy::HandleOnElementRenamed(FRigHierarchyContainer* InContainer, ERigElementType InElementType, const FName& InOldName, const FName& InNewName)
-{
-	if (Container == nullptr)
-	{
-		return;
-	}
-
-	switch (InElementType)
-	{
-		case ERigElementType::Bone:
-		{
-			for (FRigSpace& Space : Spaces)
-			{
-				if (Space.SpaceType == ERigSpaceType::Bone && Space.ParentName == InOldName)
-				{
-					Space.ParentIndex = Container->BoneHierarchy.GetIndex(InNewName);
-					Space.ParentName = Space.ParentIndex == INDEX_NONE ? NAME_None : InNewName;
-					Space.SpaceType = Space.ParentIndex == INDEX_NONE ? ERigSpaceType::Global : Space.SpaceType;
-					OnSpaceReparented.Broadcast(Container, Space.GetElementKey(), InOldName, Space.ParentName);
-				}
-			}
-			break;
-		}
-		case ERigElementType::Control:
-		{
-			for (FRigSpace& Space : Spaces)
-			{
-				if (Space.SpaceType == ERigSpaceType::Control && Space.ParentName == InOldName)
-				{
-					Space.ParentIndex = Container->ControlHierarchy.GetIndex(InNewName);
-					Space.ParentName = Space.ParentIndex == INDEX_NONE ? NAME_None : InNewName;
-					Space.SpaceType = Space.ParentIndex == INDEX_NONE ? ERigSpaceType::Global : Space.SpaceType;
-					OnSpaceReparented.Broadcast(Container, Space.GetElementKey(), InOldName, Space.ParentName);
-				}
-			}
-			break;
-		}
-		case ERigElementType::Space:
-		case ERigElementType::Curve:
-		{
-			break;
-		}
-	}
-}
-
-#endif
-
-FRigPose FRigSpaceHierarchy::GetPose() const
-{
-	FRigPose Pose;
-	AppendToPose(Pose);
-	return Pose;
-}
-
-void FRigSpaceHierarchy::SetPose(FRigPose& InPose)
-{
-	for(FRigPoseElement& Element : InPose)
-	{
-		if(Element.Index.GetKey().Type == ERigElementType::Space)
-		{
-			if(Element.Index.UpdateCache(Container))
-			{
-				SetLocalTransform(Element.Index.GetIndex(), Element.LocalTransform);
-			}
-		}
-	}
-}
-
-void FRigSpaceHierarchy::AppendToPose(FRigPose& InOutPose) const
-{
-	for(const FRigSpace& Space : Spaces)
-	{
-		FRigPoseElement Element;
-		if(Element.Index.UpdateCache(Space.GetElementKey(), Container))
-		{
-			Element.GlobalTransform = GetGlobalTransform(Element.Index.GetIndex());
-			Element.LocalTransform = GetLocalTransform(Element.Index.GetIndex());
-			InOutPose.Elements.Add(Element);
-		}
-	}
-}
\ No newline at end of file
