Control Rig: New rig hierarchy data structure refactoring

diff --git a/Engine/Plugins/Experimental/ControlRig/Source/ControlRig/Private/Sequencer/MovieSceneControlRigParameterSection.cpp b/Engine/Plugins/Experimental/ControlRig/Source/ControlRig/Private/Sequencer/MovieSceneControlRigParameterSection.cpp
--- a/Engine/Plugins/Experimental/ControlRig/Source/ControlRig/Private/Sequencer/MovieSceneControlRigParameterSection.cpp
+++ b/Engine/Plugins/Experimental/ControlRig/Source/ControlRig/Private/Sequencer/MovieSceneControlRigParameterSection.cpp
@@ -44,12 +44,10 @@ struct FParameterFloatChannelEditorData
 	{
 		if (ControlRig)
 		{
-			FRigControl* RigControl = ControlRig->FindControl(ParameterName);
-			if (RigControl)
+			FRigControlElement* ControlElement = ControlRig->FindControl(ParameterName);
+			if (ControlElement)
 			{
-				float Val = RigControl->Value.Get<float>();
-				return Val;
-		
+				return ControlRig->GetHierarchy()->GetControlValue(ControlElement, ERigControlValueType::Current).Get<float>();
 			}
 		}
 		return TOptional<float>();
@@ -160,23 +158,23 @@ struct FParameterVectorChannelEditorData
 	{
 		if (ControlRig)
 		{
-			FRigControl* RigControl = ControlRig->FindControl(ParameterName);
-			if (RigControl)
+			FRigControlElement* ControlElement = ControlRig->FindControl(ParameterName);
+			if (ControlElement)
 			{
 		
 				if (NumChannels == 2)
 				{
-					FVector2D Vector = RigControl->Value.Get<FVector2D>();
+					const FVector2D Vector = ControlRig->GetHierarchy()->GetControlValue(ControlElement, ERigControlValueType::Current).Get<FVector2D>();
 					return FVector4(Vector.X, Vector.Y, 0.f, 0.f);
 				}
 				else if (NumChannels == 3)
 				{
-					FVector Vector = RigControl->Value.Get<FVector>();
+					const FVector Vector = ControlRig->GetHierarchy()->GetControlValue(ControlElement, ERigControlValueType::Current).Get<FVector>();
 					return FVector4(Vector.X, Vector.Y, Vector.Z, 0.f);
 				}
 				else
 				{
-					return RigControl->Value.Get<FVector4>();
+					return ControlRig->GetHierarchy()->GetControlValue(ControlElement, ERigControlValueType::Current).Get<FVector4>();
 				}
 			}
 		}
@@ -524,23 +522,23 @@ struct FParameterTransformChannelEditorData
 	{
 		if (ControlRig)
 		{
-			FRigControl* RigControl = ControlRig->FindControl(ParameterName);
-			if (RigControl)
+			FRigControlElement* ControlElement = ControlRig->FindControl(ParameterName);
+			if (ControlElement)
 			{
-				if (RigControl->ControlType == ERigControlType::Transform)
+				if (ControlElement->Settings.ControlType == ERigControlType::Transform)
 				{
-					FTransform Transform = RigControl->Value.Get<FTransform>();
+					const FTransform Transform = ControlRig->GetHierarchy()->GetControlValue(ControlElement, ERigControlValueType::Current).Get<FTransform>();
 					return Transform.GetTranslation();
 				}
-				else if  (RigControl->ControlType == ERigControlType::TransformNoScale)
+				else if  (ControlElement->Settings.ControlType == ERigControlType::TransformNoScale)
 				{
-					FTransformNoScale NoScale = RigControl->Value.Get<FTransformNoScale>();
+					const FTransformNoScale NoScale = ControlRig->GetHierarchy()->GetControlValue(ControlElement, ERigControlValueType::Current).Get<FTransformNoScale>();
 					FTransform Transform = NoScale;
 					return Transform.GetTranslation();
 				}
-				else if (RigControl->ControlType == ERigControlType::EulerTransform)
+				else if (ControlElement->Settings.ControlType == ERigControlType::EulerTransform)
 				{
-					FEulerTransform Euler = RigControl->Value.Get<FEulerTransform>();
+					const FEulerTransform Euler = ControlRig->GetHierarchy()->GetControlValue(ControlElement, ERigControlValueType::Current).Get<FEulerTransform>();
 					return Euler.Location;
 				}
 			}
@@ -553,23 +551,23 @@ struct FParameterTransformChannelEditorData
 
 		if (ControlRig)
 		{
-			FRigControl* RigControl = ControlRig->FindControl(ParameterName);
-			if (RigControl)
+			FRigControlElement* ControlElement = ControlRig->FindControl(ParameterName);
+			if (ControlElement)
 			{
-				if (RigControl->ControlType == ERigControlType::Transform)
+				if (ControlElement->Settings.ControlType == ERigControlType::Transform)
 				{
-					FTransform Transform = RigControl->Value.Get<FTransform>();
+					FTransform Transform = ControlRig->GetHierarchy()->GetControlValue(ControlElement, ERigControlValueType::Current).Get<FTransform>();
 					return Transform.GetRotation().Rotator();
 				}
-				else if (RigControl->ControlType == ERigControlType::TransformNoScale)
+				else if (ControlElement->Settings.ControlType == ERigControlType::TransformNoScale)
 				{
-					FTransformNoScale NoScale = RigControl->Value.Get<FTransformNoScale>();
+					FTransformNoScale NoScale = ControlRig->GetHierarchy()->GetControlValue(ControlElement, ERigControlValueType::Current).Get<FTransformNoScale>();
 					FTransform Transform = NoScale;
 					return Transform.GetRotation().Rotator();
 				}
-				else if (RigControl->ControlType == ERigControlType::EulerTransform)
+				else if (ControlElement->Settings.ControlType == ERigControlType::EulerTransform)
 				{
-					FEulerTransform Euler = RigControl->Value.Get<FEulerTransform>();
+					FEulerTransform Euler = ControlRig->GetHierarchy()->GetControlValue(ControlElement, ERigControlValueType::Current).Get<FEulerTransform>();
 					return Euler.Rotation;
 			}
 		}
@@ -581,17 +579,17 @@ struct FParameterTransformChannelEditorData
 	{
 		if (ControlRig)
 		{
-			FRigControl* RigControl = ControlRig->FindControl(ParameterName);
-			if (RigControl)
+			FRigControlElement* ControlElement = ControlRig->FindControl(ParameterName);
+			if (ControlElement)
 			{
-				if (RigControl->ControlType == ERigControlType::Transform)
+				if (ControlElement->Settings.ControlType == ERigControlType::Transform)
 				{
-					FTransform Transform = RigControl->Value.Get<FTransform>();
+					FTransform Transform = ControlRig->GetHierarchy()->GetControlValue(ControlElement, ERigControlValueType::Current).Get<FTransform>();
 					return Transform.GetScale3D();
 				}
-				else if (RigControl->ControlType == ERigControlType::EulerTransform)
+				else if (ControlElement->Settings.ControlType == ERigControlType::EulerTransform)
 				{
-					FEulerTransform Euler = RigControl->Value.Get<FEulerTransform>();
+					FEulerTransform Euler = ControlRig->GetHierarchy()->GetControlValue(ControlElement, ERigControlValueType::Current).Get<FEulerTransform>();
 					return Euler.Scale;
 				}
 			}
@@ -700,16 +698,16 @@ void UMovieSceneControlRigParameterSection::SetBlendType(EMovieSceneBlendType In
 
 			// Set Defaults based upon Type
 			TArrayView<FMovieSceneFloatChannel*> FloatChannels = ChannelProxy->GetChannels<FMovieSceneFloatChannel>();
-			const TArray<FRigControl>& Controls = ControlRig->AvailableControls();
+			TArray<FRigControlElement*> Controls = ControlRig->AvailableControls();
 
-			for (const FRigControl& RigControl : Controls)
+			for (FRigControlElement* ControlElement : Controls)
 			{
-				switch (RigControl.ControlType)
+				switch (ControlElement->Settings.ControlType)
 				{
 
 				case ERigControlType::Scale:
 				{
-					ChannelInfo = ControlChannelMap.Find(RigControl.Name);
+					ChannelInfo = ControlChannelMap.Find(ControlElement->GetName());
 					if (ChannelInfo)
 					{
 						if (InBlendType == EMovieSceneBlendType::Absolute)
@@ -730,7 +728,7 @@ void UMovieSceneControlRigParameterSection::SetBlendType(EMovieSceneBlendType In
 				case ERigControlType::Transform:
 				case ERigControlType::EulerTransform:
 				{
-					ChannelInfo = ControlChannelMap.Find(RigControl.Name);
+					ChannelInfo = ControlChannelMap.Find(ControlElement->GetName());
 					if (ChannelInfo)
 					{
 						if (InBlendType == EMovieSceneBlendType::Absolute)
@@ -1103,8 +1101,9 @@ void UMovieSceneControlRigParameterSection::ReconstructChannelProxy(bool bForce)
 		// Need to create the channels in sorted orders
 		if (ControlRig)
 		{
-			TArray<FRigControl> SortedControls;
+			TArray<FRigControlElement*> SortedControls;
 			ControlRig->GetControlsInOrder(SortedControls);
+			
 			int32 ControlIndex = 0; 
 			int32 TotalIndex = 0; 
 			int32 FloatChannelIndex = 0;
@@ -1114,29 +1113,36 @@ void UMovieSceneControlRigParameterSection::ReconstructChannelProxy(bool bForce)
 			const FName BoolChannelTypeName = FMovieSceneBoolChannel::StaticStruct()->GetFName();
 			const FName EnumChannelTypeName = FMovieSceneByteChannel::StaticStruct()->GetFName();
 			const FName IntegerChannelTypeName = FMovieSceneIntegerChannel::StaticStruct()->GetFName();
-			for (const FRigControl& RigControl : SortedControls)
+			
+			for (FRigControlElement* ControlElement : SortedControls)
 			{
-
-				if (!RigControl.bAnimatable)
+				if (!ControlElement->Settings.bAnimatable)
 				{
 					continue;
 				}
+
+				FRigControlElement* ParentControlElement = Cast<FRigControlElement>(ControlRig->GetHierarchy()->GetFirstParent(ControlElement));
+				FName ParentControlName = NAME_None;
+				if(ParentControlElement)
+				{
+					ParentControlName = ParentControlElement->GetName();
+				}
 #if WITH_EDITOR
-				switch (RigControl.ControlType)
+
+				switch (ControlElement->Settings.ControlType)
 				{
 				case ERigControlType::Float:
 				{
 					for (FScalarParameterNameAndCurve& Scalar : GetScalarParameterNamesAndCurves())
 					{
-						if (RigControl.Name == Scalar.ParameterName)
+						if (ControlElement->GetName() == Scalar.ParameterName)
 						{
 							bool bEnabled = ControlsMask[ControlIndex];
 
 							FText Group;
-							if(!RigControl.ParentName.IsNone())
+							if(ParentControlElement)
 							{
-								FRigControl& ParentControl = ControlRig->GetControlHierarchy()[RigControl.ParentName];
-								switch(ParentControl.ControlType)
+								switch(ParentControlElement->Settings.ControlType)
 								{
 									case ERigControlType::Position:
 									case ERigControlType::Scale:
@@ -1145,7 +1151,7 @@ void UMovieSceneControlRigParameterSection::ReconstructChannelProxy(bool bForce)
 									case ERigControlType::EulerTransform:
 									case ERigControlType::TransformNoScale:
 									{
-										Group = FText::FromName(ParentControl.GetDisplayName());
+										Group = FText::FromName(ParentControlElement->GetDisplayName());
 										break;
 									}
 								}
@@ -1154,12 +1160,12 @@ void UMovieSceneControlRigParameterSection::ReconstructChannelProxy(bool bForce)
 							if (Group.IsEmpty())
 							{
 								ControlChannelMap.Add(Scalar.ParameterName, FChannelMapInfo(ControlIndex, TotalIndex, FloatChannelIndex));
-								Group = FText::FromName(RigControl.GetDisplayName());
+								Group = FText::FromName(ControlElement->GetDisplayName());
 								ControlIndex++;  //up the index only if no parent
 							}
 							else
 							{
-								FChannelMapInfo* pChannelIndex = ControlChannelMap.Find(RigControl.ParentName);
+								FChannelMapInfo* pChannelIndex = ControlChannelMap.Find(ParentControlName);
 								if (pChannelIndex)
 								{
 									ControlChannelMap.Add(Scalar.ParameterName, FChannelMapInfo(ControlIndex, TotalIndex, FloatChannelIndex,pChannelIndex->ControlIndex));
@@ -1172,7 +1178,7 @@ void UMovieSceneControlRigParameterSection::ReconstructChannelProxy(bool bForce)
 							}
 
 							FParameterFloatChannelEditorData EditorData(ControlRig, Scalar.ParameterName, bEnabled, Group, TotalIndex);
-							EditorData.MetaData.DisplayText = FText::FromName(RigControl.GetDisplayName());
+							EditorData.MetaData.DisplayText = FText::FromName(ControlElement->GetDisplayName());
 							Channels.Add(Scalar.ParameterCurve, EditorData.MetaData, EditorData.ExternalValues);
 							FloatChannelIndex += 1;
 							TotalIndex += 1;
@@ -1185,15 +1191,14 @@ void UMovieSceneControlRigParameterSection::ReconstructChannelProxy(bool bForce)
 				{
 					for (FBoolParameterNameAndCurve& Bool : GetBoolParameterNamesAndCurves())
 					{
-						if (RigControl.Name == Bool.ParameterName)
+						if (ControlElement->GetName() == Bool.ParameterName)
 						{
 							bool bEnabled = ControlsMask[ControlIndex];
 
 							FText Group;
-							if(!RigControl.ParentName.IsNone())
+							if(ParentControlElement)
 							{
-								FRigControl& ParentControl = ControlRig->GetControlHierarchy()[RigControl.ParentName];
-								switch(ParentControl.ControlType)
+								switch(ParentControlElement->Settings.ControlType)
 								{
 									case ERigControlType::Position:
 									case ERigControlType::Scale:
@@ -1202,7 +1207,7 @@ void UMovieSceneControlRigParameterSection::ReconstructChannelProxy(bool bForce)
 									case ERigControlType::EulerTransform:
 									case ERigControlType::TransformNoScale:
 									{
-										Group = FText::FromName(ParentControl.GetDisplayName());
+										Group = FText::FromName(ParentControlElement->GetDisplayName());
 										break;
 									}
 								}
@@ -1211,12 +1216,12 @@ void UMovieSceneControlRigParameterSection::ReconstructChannelProxy(bool bForce)
 							if (Group.IsEmpty())
 							{
 								ControlChannelMap.Add(Bool.ParameterName, FChannelMapInfo(ControlIndex, TotalIndex, BoolChannelIndex, INDEX_NONE,BoolChannelTypeName));
-								Group = FText::FromName(RigControl.GetDisplayName());
+								Group = FText::FromName(ControlElement->GetDisplayName());
 								ControlIndex++;  //up the index only if no parent
 							}
 							else
 							{
-								FChannelMapInfo* pChannelIndex = ControlChannelMap.Find(RigControl.ParentName);
+								FChannelMapInfo* pChannelIndex = ControlChannelMap.Find(ParentControlName);
 								if (pChannelIndex)
 								{
 									ControlChannelMap.Add(Bool.ParameterName, FChannelMapInfo(ControlIndex, TotalIndex, BoolChannelIndex,pChannelIndex->ControlIndex,BoolChannelTypeName));
@@ -1228,7 +1233,7 @@ void UMovieSceneControlRigParameterSection::ReconstructChannelProxy(bool bForce)
 							}
 
 							FMovieSceneChannelMetaData MetaData(Bool.ParameterName, Group, Group, bEnabled);
-							MetaData.DisplayText = FText::FromName(RigControl.GetDisplayName());
+							MetaData.DisplayText = FText::FromName(ControlElement->GetDisplayName());
 							MetaData.SortOrder = TotalIndex++;
 							BoolChannelIndex += 1;
 							// Prevent single channels from collapsing to the track node
@@ -1241,19 +1246,18 @@ void UMovieSceneControlRigParameterSection::ReconstructChannelProxy(bool bForce)
 				}
 				case ERigControlType::Integer:
 				{
-					if (RigControl.ControlEnum)
+					if (ControlElement->Settings.ControlEnum)
 					{
 						for (FEnumParameterNameAndCurve& Enum : GetEnumParameterNamesAndCurves())
 						{
-							if (RigControl.Name == Enum.ParameterName)
+							if (ControlElement->GetName() == Enum.ParameterName)
 							{
 								bool bEnabled = ControlsMask[ControlIndex];
 
 								FText Group;
-								if (!RigControl.ParentName.IsNone())
+								if (ParentControlElement)
 								{
-									FRigControl& ParentControl = ControlRig->GetControlHierarchy()[RigControl.ParentName];
-									switch (ParentControl.ControlType)
+									switch (ParentControlElement->Settings.ControlType)
 									{
 									case ERigControlType::Position:
 									case ERigControlType::Scale:
@@ -1262,7 +1266,7 @@ void UMovieSceneControlRigParameterSection::ReconstructChannelProxy(bool bForce)
 									case ERigControlType::EulerTransform:
 									case ERigControlType::TransformNoScale:
 									{
-										Group = FText::FromName(ParentControl.GetDisplayName());
+										Group = FText::FromName(ParentControlElement->GetDisplayName());
 										break;
 									}
 									}
@@ -1271,12 +1275,12 @@ void UMovieSceneControlRigParameterSection::ReconstructChannelProxy(bool bForce)
 								if (Group.IsEmpty())
 								{
 									ControlChannelMap.Add(Enum.ParameterName, FChannelMapInfo(ControlIndex, TotalIndex, EnumChannelIndex,INDEX_NONE, EnumChannelTypeName));
-									Group = FText::FromName(RigControl.GetDisplayName());
+									Group = FText::FromName(ControlElement->GetDisplayName());
 									ControlIndex++;  //up the index only if no parent
 								}
 								else
 								{
-									FChannelMapInfo* pChannelIndex = ControlChannelMap.Find(RigControl.ParentName);
+									FChannelMapInfo* pChannelIndex = ControlChannelMap.Find(ParentControlName);
 									if (pChannelIndex)
 									{
 										ControlChannelMap.Add(Enum.ParameterName, FChannelMapInfo(ControlIndex, TotalIndex, EnumChannelIndex, pChannelIndex->ControlIndex, EnumChannelTypeName));
@@ -1288,7 +1292,7 @@ void UMovieSceneControlRigParameterSection::ReconstructChannelProxy(bool bForce)
 								}
 
 								FMovieSceneChannelMetaData MetaData(Enum.ParameterName, Group, Group, bEnabled);
-								MetaData.DisplayText = FText::FromName(RigControl.GetDisplayName());
+								MetaData.DisplayText = FText::FromName(ControlElement->GetDisplayName());
 								EnumChannelIndex += 1;
 								MetaData.SortOrder = TotalIndex++;
 								// Prevent single channels from collapsing to the track node
@@ -1302,15 +1306,14 @@ void UMovieSceneControlRigParameterSection::ReconstructChannelProxy(bool bForce)
 					{
 						for (FIntegerParameterNameAndCurve& Integer : GetIntegerParameterNamesAndCurves())
 						{
-							if (RigControl.Name == Integer.ParameterName)
+							if (ControlElement->GetName() == Integer.ParameterName)
 							{
 								bool bEnabled = ControlsMask[ControlIndex];
 
 								FText Group;
-								if (!RigControl.ParentName.IsNone())
+								if (ParentControlElement)
 								{
-									FRigControl& ParentControl = ControlRig->GetControlHierarchy()[RigControl.ParentName];
-									switch (ParentControl.ControlType)
+									switch (ParentControlElement->Settings.ControlType)
 									{
 									case ERigControlType::Position:
 									case ERigControlType::Scale:
@@ -1319,7 +1322,7 @@ void UMovieSceneControlRigParameterSection::ReconstructChannelProxy(bool bForce)
 									case ERigControlType::EulerTransform:
 									case ERigControlType::TransformNoScale:
 									{
-										Group = FText::FromName(ParentControl.GetDisplayName());
+										Group = FText::FromName(ParentControlElement->GetDisplayName());
 										break;
 									}
 									}
@@ -1327,12 +1330,12 @@ void UMovieSceneControlRigParameterSection::ReconstructChannelProxy(bool bForce)
 								if (Group.IsEmpty())
 								{
 									ControlChannelMap.Add(Integer.ParameterName, FChannelMapInfo(ControlIndex, TotalIndex, IntegerChannelIndex,INDEX_NONE,IntegerChannelTypeName));
-									Group = FText::FromName(RigControl.GetDisplayName());
+									Group = FText::FromName(ControlElement->GetDisplayName());
 									ControlIndex++;  //up the index only if no parent
 								}
 								else
 								{
-									FChannelMapInfo* pChannelIndex = ControlChannelMap.Find(RigControl.ParentName);
+									FChannelMapInfo* pChannelIndex = ControlChannelMap.Find(ParentControlName);
 									if (pChannelIndex)
 									{
 										ControlChannelMap.Add(Integer.ParameterName, FChannelMapInfo(ControlIndex, TotalIndex, IntegerChannelIndex, pChannelIndex->ControlIndex,IntegerChannelTypeName));
@@ -1344,7 +1347,7 @@ void UMovieSceneControlRigParameterSection::ReconstructChannelProxy(bool bForce)
 								}
 
 								FMovieSceneChannelMetaData MetaData(Integer.ParameterName, Group, Group, bEnabled);
-								MetaData.DisplayText = FText::FromName(RigControl.GetDisplayName());
+								MetaData.DisplayText = FText::FromName(ControlElement->GetDisplayName());
 								IntegerChannelIndex += 1;
 								MetaData.SortOrder = TotalIndex++;
 								// Prevent single channels from collapsing to the track node
@@ -1361,11 +1364,11 @@ void UMovieSceneControlRigParameterSection::ReconstructChannelProxy(bool bForce)
 				{
 					for (FVector2DParameterNameAndCurves& Vector2D : GetVector2DParameterNamesAndCurves())
 					{
-						if (RigControl.Name == Vector2D.ParameterName)
+						if (ControlElement->GetName() == Vector2D.ParameterName)
 						{
 							ControlChannelMap.Add(Vector2D.ParameterName, FChannelMapInfo(ControlIndex, TotalIndex, FloatChannelIndex));
 							bool bEnabled = ControlsMask[ControlIndex++];
-							FText Group = FText::FromName(RigControl.GetDisplayName());
+							FText Group = FText::FromName(ControlElement->GetDisplayName());
 							FParameterVectorChannelEditorData EditorData(ControlRig, Vector2D.ParameterName, bEnabled, Group, TotalIndex, 2);
 							Channels.Add(Vector2D.XCurve, EditorData.MetaData[0], EditorData.ExternalValues[0]);
 							Channels.Add(Vector2D.YCurve, EditorData.MetaData[1], EditorData.ExternalValues[1]);
@@ -1382,9 +1385,9 @@ void UMovieSceneControlRigParameterSection::ReconstructChannelProxy(bool bForce)
 				{
 					for (FVectorParameterNameAndCurves& Vector : GetVectorParameterNamesAndCurves())
 					{
-						if (RigControl.Name == Vector.ParameterName)
+						if (ControlElement->GetName() == Vector.ParameterName)
 						{
-							if (RigControl.ControlType == ERigControlType::Scale)
+							if (ControlElement->Settings.ControlType == ERigControlType::Scale)
 							{
 								if (BlendType == EMovieSceneBlendType::Additive)
 								{
@@ -1401,7 +1404,7 @@ void UMovieSceneControlRigParameterSection::ReconstructChannelProxy(bool bForce)
 							}
 							ControlChannelMap.Add(Vector.ParameterName, FChannelMapInfo(ControlIndex, TotalIndex, FloatChannelIndex));
 							bool bEnabled = ControlsMask[ControlIndex++];
-							FText Group = FText::FromName(RigControl.GetDisplayName());
+							FText Group = FText::FromName(ControlElement->GetDisplayName());
 							FParameterVectorChannelEditorData EditorData(ControlRig, Vector.ParameterName, bEnabled, Group, TotalIndex, 3);
 							Channels.Add(Vector.XCurve, EditorData.MetaData[0], EditorData.ExternalValues[0]);
 							Channels.Add(Vector.YCurve, EditorData.MetaData[1], EditorData.ExternalValues[1]);
@@ -1420,11 +1423,11 @@ void UMovieSceneControlRigParameterSection::ReconstructChannelProxy(bool bForce)
 				{
 					for (FTransformParameterNameAndCurves& Transform : GetTransformParameterNamesAndCurves())
 					{
-						if (RigControl.Name == Transform.ParameterName)
+						if (ControlElement->GetName() == Transform.ParameterName)
 						{
 							ControlChannelMap.Add(Transform.ParameterName, FChannelMapInfo(ControlIndex, TotalIndex, FloatChannelIndex));
 							bool bEnabled = ControlsMask[ControlIndex++];
-							FText Group = FText::FromName(RigControl.GetDisplayName());
+							FText Group = FText::FromName(ControlElement->GetDisplayName());
 
 							FParameterTransformChannelEditorData EditorData(ControlRig, Transform.ParameterName, bEnabled, TransformMask.GetChannels(), Group, 
 								TotalIndex);
@@ -1437,7 +1440,8 @@ void UMovieSceneControlRigParameterSection::ReconstructChannelProxy(bool bForce)
 							Channels.Add(Transform.Rotation[1], EditorData.MetaData[4], EditorData.ExternalValues[4]);
 							Channels.Add(Transform.Rotation[2], EditorData.MetaData[5], EditorData.ExternalValues[5]);
 
-							if (RigControl.ControlType == ERigControlType::Transform || RigControl.ControlType == ERigControlType::EulerTransform)
+							if (ControlElement->Settings.ControlType == ERigControlType::Transform ||
+								ControlElement->Settings.ControlType == ERigControlType::EulerTransform)
 							{
 								if (BlendType == EMovieSceneBlendType::Additive)
 								{
@@ -1472,13 +1476,13 @@ void UMovieSceneControlRigParameterSection::ReconstructChannelProxy(bool bForce)
 					break;
 				}
 #else
-				switch (RigControl.ControlType)
+				switch (ControlElement->Settings.ControlType)
 				{
 				case ERigControlType::Float:
 				{
 					for (FScalarParameterNameAndCurve& Scalar : GetScalarParameterNamesAndCurves())
 					{
-						if (RigControl.Name == Scalar.ParameterName)
+						if (ControlElement->GetName() == Scalar.ParameterName)
 						{
 							ControlChannelMap.Add(Scalar.ParameterName, FChannelMapInfo(ControlIndex, TotalIndex,FloatChannelIndex));
 							Channels.Add(Scalar.ParameterCurve);
@@ -1494,7 +1498,7 @@ void UMovieSceneControlRigParameterSection::ReconstructChannelProxy(bool bForce)
 				{
 					for (FBoolParameterNameAndCurve& Bool : GetBoolParameterNamesAndCurves())
 					{
-						if (RigControl.Name == Bool.ParameterName)
+						if (ControlElement->GetName() == Bool.ParameterName)
 						{
 							ControlChannelMap.Add(Bool.ParameterName, FChannelMapInfo(ControlIndex, TotalIndex, BoolChannelIndex));
 							Channels.Add(Bool.ParameterCurve);
@@ -1508,11 +1512,11 @@ void UMovieSceneControlRigParameterSection::ReconstructChannelProxy(bool bForce)
 				}
 				case ERigControlType::Integer:
 				{
-					if (RigControl.ControlEnum)
+					if (ControlElement->Settings.ControlEnum)
 					{
 						for (FEnumParameterNameAndCurve& Enum : GetEnumParameterNamesAndCurves())
 						{
-							if (RigControl.Name == Enum.ParameterName)
+							if (ControlElement->GetName() == Enum.ParameterName)
 							{
 								ControlChannelMap.Add(Enum.ParameterName, FChannelMapInfo(ControlIndex, TotalIndex, EnumChannelIndex));
 								Channels.Add(Enum.ParameterCurve);
@@ -1527,7 +1531,7 @@ void UMovieSceneControlRigParameterSection::ReconstructChannelProxy(bool bForce)
 					{
 						for (FIntegerParameterNameAndCurve& Integer : GetIntegerParameterNamesAndCurves())
 						{
-							if (RigControl.Name == Integer.ParameterName)
+							if (ControlElement->GetName() == Integer.ParameterName)
 							{
 								ControlChannelMap.Add(Integer.ParameterName, FChannelMapInfo(ControlIndex, TotalIndex, IntegerChannelIndex));
 								Channels.Add(Integer.ParameterCurve);
@@ -1544,7 +1548,7 @@ void UMovieSceneControlRigParameterSection::ReconstructChannelProxy(bool bForce)
 				{
 					for (FVector2DParameterNameAndCurves& Vector2D : GetVector2DParameterNamesAndCurves())
 					{
-						if (RigControl.Name == Vector2D.ParameterName)
+						if (ControlElement->GetName() == Vector2D.ParameterName)
 						{
 							ControlChannelMap.Add(Vector2D.ParameterName, FChannelMapInfo(ControlIndex, TotalIndex, FloatChannelIndex));
 							Channels.Add(Vector2D.XCurve);
@@ -1563,7 +1567,7 @@ void UMovieSceneControlRigParameterSection::ReconstructChannelProxy(bool bForce)
 				{
 					for (FVectorParameterNameAndCurves& Vector : GetVectorParameterNamesAndCurves())
 					{
-						if (RigControl.Name == Vector.ParameterName)
+						if (ControlElement->GetName() == Vector.ParameterName)
 						{
 							ControlChannelMap.Add(Vector.ParameterName, FChannelMapInfo(ControlIndex, TotalIndex, FloatChannelIndex));
 							Channels.Add(Vector.XCurve);
@@ -1593,7 +1597,7 @@ void UMovieSceneControlRigParameterSection::ReconstructChannelProxy(bool bForce)
 				{
 					for (FTransformParameterNameAndCurves& Transform : GetTransformParameterNamesAndCurves())
 					{
-						if (RigControl.Name == Transform.ParameterName)
+						if (ControlElement->GetName() == Transform.ParameterName)
 						{
 							ControlChannelMap.Add(Transform.ParameterName, FChannelMapInfo(ControlIndex, TotalIndex,FloatChannelIndex));
 							Channels.Add(Transform.Translation[0]);
@@ -1604,7 +1608,7 @@ void UMovieSceneControlRigParameterSection::ReconstructChannelProxy(bool bForce)
 							Channels.Add(Transform.Rotation[1]);
 							Channels.Add(Transform.Rotation[2]);
 
-							if (RigControl.ControlType == ERigControlType::Transform || RigControl.ControlType == ERigControlType::EulerTransform)
+							if (ControlElement->Settings.ControlType == ERigControlType::Transform || ControlElement->Settings.ControlType == ERigControlType::EulerTransform)
 							{
 								Channels.Add(Transform.Scale[0]);
 								Channels.Add(Transform.Scale[1]);
@@ -1704,21 +1708,20 @@ void UMovieSceneControlRigParameterSection::RecreateWithThisControlRig(UControlR
 	TransformParameterNamesAndCurves.Empty();
 	*/
 	TArray<bool> OnArray;
-	const TArray<FRigControl>& Controls = ControlRig->AvailableControls();
-	OnArray.Init(true, Controls.Num());
+	OnArray.Init(true, ControlRig->AvailableControls().Num());
 	SetControlsMask(OnArray);
 
-	TArray<FRigControl> SortedControls;
+	TArray<FRigControlElement*> SortedControls;
 	ControlRig->GetControlsInOrder(SortedControls);
 
-	for (const FRigControl& RigControl : SortedControls)
+	for (FRigControlElement* ControlElement : SortedControls)
 	{
-		if (!RigControl.bAnimatable)
+		if (!ControlElement->Settings.bAnimatable)
 		{
 			continue;
 		}
 
-		switch (RigControl.ControlType)
+		switch (ControlElement->Settings.ControlType)
 		{
 		case ERigControlType::Float:
 		{
@@ -1726,9 +1729,9 @@ void UMovieSceneControlRigParameterSection::RecreateWithThisControlRig(UControlR
 			if (bSetDefault)
 			{
 				//or use IntialValue?
-				DefaultValue = RigControl.Value.Get<float>();
+				DefaultValue = ControlRig->GetHierarchy()->GetControlValue(ControlElement, ERigControlValueType::Current).Get<float>();
 			}
-			AddScalarParameter(RigControl.Name, DefaultValue, false);
+			AddScalarParameter(ControlElement->GetName(), DefaultValue, false);
 			break;
 		}
 		case ERigControlType::Bool:
@@ -1737,22 +1740,22 @@ void UMovieSceneControlRigParameterSection::RecreateWithThisControlRig(UControlR
 			if (bSetDefault)
 			{
 				//or use IntialValue?
-				DefaultValue = RigControl.Value.Get<bool>();
+				DefaultValue = ControlRig->GetHierarchy()->GetControlValue(ControlElement, ERigControlValueType::Current).Get<bool>();
 			}
-			AddBoolParameter(RigControl.Name, DefaultValue, false);
+			AddBoolParameter(ControlElement->GetName(), DefaultValue, false);
 			break;
 		}
 		case ERigControlType::Integer:
 		{
-			if (RigControl.ControlEnum)
+			if (ControlElement->Settings.ControlEnum)
 			{
 				TOptional<uint8> DefaultValue;
 				if (bSetDefault)
 				{
 					//or use IntialValue?
-					DefaultValue = (uint8)RigControl.Value.Get<int32>();
+					DefaultValue = (uint8)ControlRig->GetHierarchy()->GetControlValue(ControlElement, ERigControlValueType::Current).Get<int32>();
 				}
-				AddEnumParameter(RigControl.Name, RigControl.ControlEnum, DefaultValue, false);
+				AddEnumParameter(ControlElement->GetName(), ControlElement->Settings.ControlEnum, DefaultValue, false);
 			}
 			else
 			{
@@ -1760,9 +1763,9 @@ void UMovieSceneControlRigParameterSection::RecreateWithThisControlRig(UControlR
 				if (bSetDefault)
 				{
 					//or use IntialValue?
-					DefaultValue = RigControl.Value.Get<int32>();
+					DefaultValue = ControlRig->GetHierarchy()->GetControlValue(ControlElement, ERigControlValueType::Current).Get<int32>();
 				}
-				AddIntegerParameter(RigControl.Name, DefaultValue, false);
+				AddIntegerParameter(ControlElement->GetName(), DefaultValue, false);
 			}
 			break;
 		}
@@ -1772,9 +1775,9 @@ void UMovieSceneControlRigParameterSection::RecreateWithThisControlRig(UControlR
 			if (bSetDefault)
 			{
 				//or use IntialValue?
-				DefaultValue = RigControl.Value.Get<FVector2D>();
+				DefaultValue = ControlRig->GetHierarchy()->GetControlValue(ControlElement, ERigControlValueType::Current).Get<FVector2D>();
 			}
-			AddVector2DParameter(RigControl.Name, DefaultValue, false);
+			AddVector2DParameter(ControlElement->GetName(), DefaultValue, false);
 			break;
 		}
 
@@ -1786,9 +1789,9 @@ void UMovieSceneControlRigParameterSection::RecreateWithThisControlRig(UControlR
 			if (bSetDefault)
 			{
 				//or use IntialValue?
-				DefaultValue = RigControl.Value.Get<FVector>();
+				DefaultValue = ControlRig->GetHierarchy()->GetControlValue(ControlElement, ERigControlValueType::Current).Get<FVector>();
 			}
-			AddVectorParameter(RigControl.Name, DefaultValue, false);
+			AddVectorParameter(ControlElement->GetName(), DefaultValue, false);
 			//mz todo specify rotator special so we can do quat interps
 			break;
 		}
@@ -1799,23 +1802,23 @@ void UMovieSceneControlRigParameterSection::RecreateWithThisControlRig(UControlR
 			TOptional<FTransform> DefaultValue;
 			if (bSetDefault)
 			{
-				if (RigControl.ControlType == ERigControlType::Transform)
+				if (ControlElement->Settings.ControlType == ERigControlType::Transform)
 				{
-					DefaultValue = RigControl.Value.Get<FTransform>();
+					DefaultValue = ControlRig->GetHierarchy()->GetControlValue(ControlElement, ERigControlValueType::Current).Get<FTransform>();
 				}
-				else if (RigControl.ControlType == ERigControlType::EulerTransform)
+				else if (ControlElement->Settings.ControlType == ERigControlType::EulerTransform)
 
 				{
-					FEulerTransform Euler = RigControl.Value.Get<FEulerTransform>();
+					FEulerTransform Euler = ControlRig->GetHierarchy()->GetControlValue(ControlElement, ERigControlValueType::Current).Get<FEulerTransform>();
 					DefaultValue = Euler;
 				}
 				else
 				{
-					FTransformNoScale NoScale = RigControl.Value.Get<FTransformNoScale>();
+					FTransformNoScale NoScale = ControlRig->GetHierarchy()->GetControlValue(ControlElement, ERigControlValueType::Current).Get<FTransformNoScale>();
 					DefaultValue = NoScale;
 				}
 			}
-			AddTransformParameter(RigControl.Name, DefaultValue, false);
+			AddTransformParameter(ControlElement->GetName(), DefaultValue, false);
 			break;
 		}
 
@@ -1837,16 +1840,16 @@ void UMovieSceneControlRigParameterSection::RecordControlRigKey(FFrameNumber Fra
 		TArrayView<FMovieSceneIntegerChannel*> IntChannels = GetChannelProxy().GetChannels<FMovieSceneIntegerChannel>();
 		TArrayView<FMovieSceneByteChannel*> EnumChannels = GetChannelProxy().GetChannels<FMovieSceneByteChannel>();
 
-		TArray<FRigControl> Controls;
+		TArray<FRigControlElement*> Controls;
 		ControlRig->GetControlsInOrder(Controls);
 		
-		for (const FRigControl& RigControl : Controls)
+		for (FRigControlElement* ControlElement : Controls)
 		{
-			if (!RigControl.bAnimatable)
+			if (!ControlElement->Settings.bAnimatable)
 			{
 				continue;
 			}
-			FChannelMapInfo* pChannelIndex = ControlChannelMap.Find(RigControl.Name);
+			FChannelMapInfo* pChannelIndex = ControlChannelMap.Find(ControlElement->GetName());
 			if (pChannelIndex == nullptr)
 			{
 				continue;
@@ -1854,11 +1857,11 @@ void UMovieSceneControlRigParameterSection::RecordControlRigKey(FFrameNumber Fra
 			int32 ChannelIndex = pChannelIndex->ChannelIndex;
 
 	
-			switch (RigControl.ControlType)
+			switch (ControlElement->Settings.ControlType)
 			{
 				case ERigControlType::Bool:
 				{
-					bool Val = RigControl.Value.Get<bool>();
+					bool Val = ControlRig->GetHierarchy()->GetControlValue(ControlElement, ERigControlValueType::Current).Get<bool>();
 					if (bSetDefault)
 					{
 						BoolChannels[ChannelIndex]->SetDefault(Val);
@@ -1868,9 +1871,9 @@ void UMovieSceneControlRigParameterSection::RecordControlRigKey(FFrameNumber Fra
 				}
 				case ERigControlType::Integer:
 				{
-					if (RigControl.ControlEnum)
+					if (ControlElement->Settings.ControlEnum)
 					{
-						uint8 Val = (uint8)RigControl.Value.Get<uint8>();
+						uint8 Val = (uint8)ControlRig->GetHierarchy()->GetControlValue(ControlElement, ERigControlValueType::Current).Get<uint8>();
 						if (bSetDefault)
 						{
 							EnumChannels[ChannelIndex]->SetDefault(Val);
@@ -1879,7 +1882,7 @@ void UMovieSceneControlRigParameterSection::RecordControlRigKey(FFrameNumber Fra
 					}
 					else
 					{
-						int32 Val = RigControl.Value.Get<int32>();
+						int32 Val = ControlRig->GetHierarchy()->GetControlValue(ControlElement, ERigControlValueType::Current).Get<int32>();
 						if (bSetDefault)
 						{
 							IntChannels[ChannelIndex]->SetDefault(Val);
@@ -1890,7 +1893,7 @@ void UMovieSceneControlRigParameterSection::RecordControlRigKey(FFrameNumber Fra
 				}
 				case ERigControlType::Float:
 				{
-					float Val = RigControl.Value.Get<float>();
+					float Val = ControlRig->GetHierarchy()->GetControlValue(ControlElement, ERigControlValueType::Current).Get<float>();
 					if (bSetDefault)
 					{
 						FloatChannels[ChannelIndex]->SetDefault(Val);
@@ -1907,7 +1910,7 @@ void UMovieSceneControlRigParameterSection::RecordControlRigKey(FFrameNumber Fra
 				}
 				case ERigControlType::Vector2D:
 				{
-					FVector2D Val = RigControl.Value.Get<FVector2D>();
+					FVector2D Val = ControlRig->GetHierarchy()->GetControlValue(ControlElement, ERigControlValueType::Current).Get<FVector2D>();
 					if (bSetDefault)
 					{
 						FloatChannels[ChannelIndex]->SetDefault(Val.X);
@@ -1930,8 +1933,8 @@ void UMovieSceneControlRigParameterSection::RecordControlRigKey(FFrameNumber Fra
 				case ERigControlType::Scale:
 				case ERigControlType::Rotator:
 				{
-					FVector Val = RigControl.Value.Get<FVector>();
-					if (RigControl.ControlType == ERigControlType::Rotator &&
+					FVector Val = ControlRig->GetHierarchy()->GetControlValue(ControlElement, ERigControlValueType::Current).Get<FVector>();
+					if (ControlElement->Settings.ControlType == ERigControlType::Rotator &&
 						FloatChannels[ChannelIndex]->GetNumKeys() > 0)
 					{
 						float LastVal = FloatChannels[ChannelIndex]->GetValues()[FloatChannels[ChannelIndex]->GetNumKeys() - 1].Value;
@@ -1967,19 +1970,19 @@ void UMovieSceneControlRigParameterSection::RecordControlRigKey(FFrameNumber Fra
 				case ERigControlType::EulerTransform:
 				{
 					FTransform Val = FTransform::Identity;
-					if (RigControl.ControlType == ERigControlType::TransformNoScale)
+					if (ControlElement->Settings.ControlType == ERigControlType::TransformNoScale)
 					{
-						FTransformNoScale NoScale = RigControl.Value.Get<FTransformNoScale>();
+						FTransformNoScale NoScale = ControlRig->GetHierarchy()->GetControlValue(ControlElement, ERigControlValueType::Current).Get<FTransformNoScale>();
 						Val = NoScale;
 					}
-					else if (RigControl.ControlType == ERigControlType::EulerTransform)
+					else if (ControlElement->Settings.ControlType == ERigControlType::EulerTransform)
 					{
-						FEulerTransform Euler = RigControl.Value.Get < FEulerTransform >();
+						FEulerTransform Euler = ControlRig->GetHierarchy()->GetControlValue(ControlElement, ERigControlValueType::Current).Get < FEulerTransform >();
 						Val = Euler.ToFTransform();
 					}
 					else
 					{
-						Val = RigControl.Value.Get<FTransform>();
+						Val = ControlRig->GetHierarchy()->GetControlValue(ControlElement, ERigControlValueType::Current).Get<FTransform>();
 					}
 					FVector CurrentVector = Val.GetTranslation();
 					if (bSetDefault)
@@ -2029,7 +2032,8 @@ void UMovieSceneControlRigParameterSection::RecordControlRigKey(FFrameNumber Fra
 						FloatChannels[ChannelIndex++]->AddLinearKey(FrameNumber, CurrentVector.Z);
 					}
 
-					if (RigControl.ControlType == ERigControlType::Transform || RigControl.ControlType == ERigControlType::EulerTransform)
+					if (ControlElement->Settings.ControlType == ERigControlType::Transform ||
+						ControlElement->Settings.ControlType == ERigControlType::EulerTransform)
 					{
 						CurrentVector = Val.GetScale3D();
 						if (bSetDefault)
@@ -2074,8 +2078,7 @@ bool UMovieSceneControlRigParameterSection::LoadAnimSequenceIntoThisSection(UAni
 		return false;
 	}
 
-	FRigBoneHierarchy& SourceBones = ControlRig->GetBoneHierarchy();
-	FRigCurveContainer& SourceCurves = ControlRig->GetCurveContainer();
+	URigHierarchy* SourceHierarchy = ControlRig->GetHierarchy();
 	
 	FFrameRate TickResolution = MovieScene->GetTickResolution();
 	float Length = AnimSequence->GetPlayLength();
@@ -2111,7 +2114,7 @@ bool UMovieSceneControlRigParameterSection::LoadAnimSequenceIntoThisSection(UAni
 		for (const FFloatCurve& Curve : CurveData.FloatCurves)
 		{
 			const float Val = Curve.FloatCurve.Eval(SequenceSecond);
-			SourceCurves.SetValue(Curve.Name.DisplayName,Val);
+			SourceHierarchy->SetCurveValue(FRigElementKey(Curve.Name.DisplayName, ERigElementType::Curve), Val);
 		}
 
 		for(int32 TrackIndex = 0; TrackIndex < BoneAnimationTracks.Num(); ++TrackIndex)
@@ -2121,7 +2124,7 @@ bool UMovieSceneControlRigParameterSection::LoadAnimSequenceIntoThisSection(UAni
 			FName BoneName = Skeleton->GetReferenceSkeleton().GetBoneName(AnimationTrack.BoneTreeIndex);
 			FTransform BoneTransform;
 			UE::Anim::GetBoneTransformFromModel(DataModel, BoneTransform, TrackIndex, SequenceSecond, EAnimInterpolationType::Linear);
-			SourceBones.SetLocalTransform(BoneName, BoneTransform);
+			SourceHierarchy->SetLocalTransform(FRigElementKey(BoneName, ERigElementType::Bone), BoneTransform);
 		}
 		if (Index == 0)
 		{
diff --git a/Engine/Plugins/Experimental/ControlRig/Source/ControlRig/Private/Sequencer/MovieSceneControlRigParameterTemplate.cpp b/Engine/Plugins/Experimental/ControlRig/Source/ControlRig/Private/Sequencer/MovieSceneControlRigParameterTemplate.cpp
--- a/Engine/Plugins/Experimental/ControlRig/Source/ControlRig/Private/Sequencer/MovieSceneControlRigParameterTemplate.cpp
+++ b/Engine/Plugins/Experimental/ControlRig/Source/ControlRig/Private/Sequencer/MovieSceneControlRigParameterTemplate.cpp
@@ -691,9 +691,9 @@ struct FControlRigParameterPreAnimatedTokenProducer : IMovieScenePreAnimatedToke
 
 						for (TNameAndValue<FTransform>& Value : TransformValues)
 						{
-							if (FRigControl* RigControl = ControlRig->FindControl(Value.Name))
+							if (FRigControlElement* ControlElement = ControlRig->FindControl(Value.Name))
 							{
-								switch (RigControl->ControlType)
+								switch (ControlElement->Settings.ControlType)
 								{
 								case ERigControlType::Transform:
 								{
@@ -760,63 +760,64 @@ struct FControlRigParameterPreAnimatedTokenProducer : IMovieScenePreAnimatedToke
 				}
 			}
 
-			const TArray<FRigControl>& Controls = ControlRig->AvailableControls();
+			TArray<FRigControlElement*> Controls = ControlRig->AvailableControls();
 			FRigControlValue Value;
-			for (const FRigControl& RigControl : Controls)
+			
+			for (FRigControlElement* ControlElement : Controls)
 			{
-				switch (RigControl.ControlType)
+				switch (ControlElement->Settings.ControlType)
 				{
 				case ERigControlType::Bool:
 				{
-					bool Val = RigControl.Value.Get<bool>();
-					Token.BoolValues.Add(TNameAndValue<bool>{ RigControl.Name, Val });
+					const bool Val = ControlRig->GetHierarchy()->GetControlValue(ControlElement, ERigControlValueType::Current).Get<bool>();
+					Token.BoolValues.Add(TNameAndValue<bool>{ ControlElement->GetName(), Val });
 					break;
 				}
 				case ERigControlType::Integer:
 				{
-					int32 Val = RigControl.Value.Get<int32>();
-					Token.IntegerValues.Add(TNameAndValue<int32>{ RigControl.Name, Val });
+					const int32 Val = ControlRig->GetHierarchy()->GetControlValue(ControlElement, ERigControlValueType::Current).Get<int32>();
+					Token.IntegerValues.Add(TNameAndValue<int32>{ ControlElement->GetName(), Val });
 					break;
 				}
 				case ERigControlType::Float:
 				{
-					float Val = RigControl.Value.Get<float>();
-					Token.ScalarValues.Add(TNameAndValue<float>{ RigControl.Name, Val });
+					const float Val = ControlRig->GetHierarchy()->GetControlValue(ControlElement, ERigControlValueType::Current).Get<float>();
+					Token.ScalarValues.Add(TNameAndValue<float>{ ControlElement->GetName(), Val });
 					break;
 				}
 				case ERigControlType::Vector2D:
 				{
-					FVector2D Val = RigControl.Value.Get<FVector2D>();
-					Token.Vector2DValues.Add(TNameAndValue<FVector2D>{ RigControl.Name, Val });
+					const FVector2D Val = ControlRig->GetHierarchy()->GetControlValue(ControlElement, ERigControlValueType::Current).Get<FVector2D>();
+					Token.Vector2DValues.Add(TNameAndValue<FVector2D>{ ControlElement->GetName(), Val });
 					break;
 				}
 				case ERigControlType::Position:
 				case ERigControlType::Scale:
 				case ERigControlType::Rotator:
 				{
-					FVector Val = RigControl.Value.Get<FVector>();
-					Token.VectorValues.Add(TNameAndValue<FVector>{ RigControl.Name, Val });
+					const FVector Val = ControlRig->GetHierarchy()->GetControlValue(ControlElement, ERigControlValueType::Current).Get<FVector>();
+					Token.VectorValues.Add(TNameAndValue<FVector>{ ControlElement->GetName(), Val });
 					//mz todo specify rotator special so we can do quat interps
 					break;
 				}
 				case ERigControlType::Transform:
 				{
-					FTransform Val = RigControl.Value.Get<FTransform>();
-					Token.TransformValues.Add(TNameAndValue<FTransform>{ RigControl.Name, Val });
+					const FTransform Val = ControlRig->GetHierarchy()->GetControlValue(ControlElement, ERigControlValueType::Current).Get<FTransform>();
+					Token.TransformValues.Add(TNameAndValue<FTransform>{ ControlElement->GetName(), Val });
 					break;
 				}
 				case ERigControlType::TransformNoScale:
 				{
-					FTransformNoScale NoScale = RigControl.Value.Get<FTransformNoScale>();
+					const FTransformNoScale NoScale = ControlRig->GetHierarchy()->GetControlValue(ControlElement, ERigControlValueType::Current).Get<FTransformNoScale>();
 					FTransform Val = NoScale;
-					Token.TransformValues.Add(TNameAndValue<FTransform>{ RigControl.Name, Val });
+					Token.TransformValues.Add(TNameAndValue<FTransform>{ ControlElement->GetName(), Val });
 					break;
 				}
 				case ERigControlType::EulerTransform:
 				{
-					FEulerTransform Euler = RigControl.Value.Get<FEulerTransform>();
+					const FEulerTransform Euler = ControlRig->GetHierarchy()->GetControlValue(ControlElement, ERigControlValueType::Current).Get<FEulerTransform>();
 					FTransform Val = Euler.ToFTransform();
-					Token.TransformValues.Add(TNameAndValue<FTransform>{ RigControl.Name, Val });
+					Token.TransformValues.Add(TNameAndValue<FTransform>{ ControlElement->GetName(), Val });
 					break;
 				}
 				}
@@ -973,10 +974,10 @@ struct TControlRigParameterActuatorFloat : TMovieSceneBlendingActuator<FControlR
 
 		if (Section && Section->ControlRig)
 		{
-			FRigControl* RigControl = Section->ControlRig->FindControl(ParameterName);
-			if (RigControl && RigControl->ControlType == ERigControlType::Float)
+			FRigControlElement* ControlElement = Section->ControlRig->FindControl(ParameterName);
+			if (ControlElement && ControlElement->Settings.ControlType == ERigControlType::Float)
 			{
-				float Val = RigControl->Value.Get<float>();
+				const float Val = Section->ControlRig->GetHierarchy()->GetControlValue(ControlElement, ERigControlValueType::Current).Get<float>();
 				return FControlRigTrackTokenFloat(Val);
 			}
 
@@ -996,8 +997,8 @@ struct TControlRigParameterActuatorFloat : TMovieSceneBlendingActuator<FControlR
 
 			if (Section->ControlRig && (Section->ControlsToSet.Num() == 0 || Section->ControlsToSet.Contains(ParameterName)))
 			{
-				FRigControl* RigControl = Section->ControlRig->FindControl(ParameterName);
-				if (RigControl && RigControl->ControlType == ERigControlType::Float)
+				FRigControlElement* ControlElement = Section->ControlRig->FindControl(ParameterName);
+				if (ControlElement && ControlElement->Settings.ControlType == ERigControlType::Float)
 				{
 					Section->ControlRig->SetControlValue<float>(ParameterName, InFinalValue.Value, true, EControlRigSetKey::Never);
 				}
@@ -1037,10 +1038,10 @@ struct TControlRigParameterActuatorVector2D : TMovieSceneBlendingActuator<FContr
 		const UMovieSceneControlRigParameterSection* Section = SectionData.Get();
 
 		{
-			FRigControl* RigControl = Section->ControlRig->FindControl(ParameterName);
-			if (RigControl && (RigControl->ControlType == ERigControlType::Vector2D))
+			FRigControlElement* ControlElement = Section->ControlRig->FindControl(ParameterName);
+			if (ControlElement && (ControlElement->Settings.ControlType == ERigControlType::Vector2D))
 			{
-				FVector2D Val = RigControl->Value.Get<FVector2D>();
+				FVector2D Val = Section->ControlRig->GetHierarchy()->GetControlValue(ControlElement, ERigControlValueType::Current).Get<FVector2D>();
 				return FControlRigTrackTokenVector2D(Val);
 			}
 		}
@@ -1059,8 +1060,8 @@ struct TControlRigParameterActuatorVector2D : TMovieSceneBlendingActuator<FContr
 
 			if (Section->ControlRig && (Section->ControlsToSet.Num() == 0 || Section->ControlsToSet.Contains(ParameterName)))
 			{
-				FRigControl* RigControl = Section->ControlRig->FindControl(ParameterName);
-				if (RigControl && (RigControl->ControlType == ERigControlType::Vector2D))
+				FRigControlElement* ControlElement = Section->ControlRig->FindControl(ParameterName);
+				if (ControlElement && (ControlElement->Settings.ControlType == ERigControlType::Vector2D))
 				{
 					Section->ControlRig->SetControlValue<FVector2D>(ParameterName, InFinalValue.Value, true, EControlRigSetKey::Never);
 				}
@@ -1099,10 +1100,10 @@ struct TControlRigParameterActuatorVector : TMovieSceneBlendingActuator<FControl
 
 		if (Section->ControlRig)
 		{
-			FRigControl* RigControl = Section->ControlRig->FindControl(ParameterName);
-			if (RigControl && (RigControl->ControlType == ERigControlType::Position || RigControl->ControlType == ERigControlType::Scale || RigControl->ControlType == ERigControlType::Rotator))
+			FRigControlElement* ControlElement = Section->ControlRig->FindControl(ParameterName);
+			if (ControlElement && (ControlElement->Settings.ControlType == ERigControlType::Position || ControlElement->Settings.ControlType == ERigControlType::Scale || ControlElement->Settings.ControlType == ERigControlType::Rotator))
 			{
-				FVector Val = RigControl->Value.Get<FVector>();
+				FVector Val = Section->ControlRig->GetHierarchy()->GetControlValue(ControlElement, ERigControlValueType::Current).Get<FVector>();
 				return FControlRigTrackTokenVector(Val);
 			}
 		}
@@ -1120,8 +1121,8 @@ struct TControlRigParameterActuatorVector : TMovieSceneBlendingActuator<FControl
 			Section->SetDoNotKey(true);
 			if (Section->ControlRig && (Section->ControlsToSet.Num() == 0 || Section->ControlsToSet.Contains(ParameterName)))
 			{
-				FRigControl* RigControl = Section->ControlRig->FindControl(ParameterName);
-				if (RigControl && (RigControl->ControlType == ERigControlType::Position || RigControl->ControlType == ERigControlType::Scale || RigControl->ControlType == ERigControlType::Rotator))
+				FRigControlElement* ControlElement = Section->ControlRig->FindControl(ParameterName);
+				if (ControlElement && (ControlElement->Settings.ControlType == ERigControlType::Position || ControlElement->Settings.ControlType == ERigControlType::Scale || ControlElement->Settings.ControlType == ERigControlType::Rotator))
 				{
 						Section->ControlRig->SetControlValue<FVector>(ParameterName, InFinalValue.Value, true, EControlRigSetKey::Never);
 				}
@@ -1160,21 +1161,21 @@ struct TControlRigParameterActuatorTransform : TMovieSceneBlendingActuator<FCont
 
 		if (Section->ControlRig && (Section->ControlsToSet.Num() == 0 || Section->ControlsToSet.Contains(ParameterName)))
 		{
-			FRigControl* RigControl = Section->ControlRig->FindControl(ParameterName);
-			if (RigControl && RigControl->ControlType == ERigControlType::Transform)
+			FRigControlElement* ControlElement = Section->ControlRig->FindControl(ParameterName);
+			if (ControlElement && ControlElement->Settings.ControlType == ERigControlType::Transform)
 			{
-				FTransform Val = RigControl->Value.Get<FTransform>();
+				FTransform Val = Section->ControlRig->GetHierarchy()->GetControlValue(ControlElement, ERigControlValueType::Current).Get<FTransform>();
 				return FControlRigTrackTokenTransform(Val);
 			}
-			else if(RigControl && RigControl->ControlType == ERigControlType::TransformNoScale)
+			else if(ControlElement && ControlElement->Settings.ControlType == ERigControlType::TransformNoScale)
 			{
-				FTransformNoScale ValNoScale = RigControl->Value.Get<FTransformNoScale>();
+				FTransformNoScale ValNoScale = Section->ControlRig->GetHierarchy()->GetControlValue(ControlElement, ERigControlValueType::Current).Get<FTransformNoScale>();
 				FTransform Val = ValNoScale;
 				return FControlRigTrackTokenTransform(Val);
 			}
-			else if (RigControl && RigControl->ControlType == ERigControlType::EulerTransform)
+			else if (ControlElement && ControlElement->Settings.ControlType == ERigControlType::EulerTransform)
 			{
-				FEulerTransform Euler = RigControl->Value.Get<FEulerTransform>();
+				FEulerTransform Euler = Section->ControlRig->GetHierarchy()->GetControlValue(ControlElement, ERigControlValueType::Current).Get<FEulerTransform>();
 				FTransform Val = Euler.ToFTransform();
 				return FControlRigTrackTokenTransform(Val);
 		}
@@ -1202,17 +1203,17 @@ struct TControlRigParameterActuatorTransform : TMovieSceneBlendingActuator<FCont
 		
 			if (Section->ControlRig && (Section->ControlsToSet.Num() == 0 || Section->ControlsToSet.Contains(ParameterName)))
 			{
-				FRigControl* RigControl = Section->ControlRig->FindControl(ParameterName);
-				if (RigControl && RigControl->ControlType == ERigControlType::Transform)
+				FRigControlElement* ControlElement = Section->ControlRig->FindControl(ParameterName);
+				if (ControlElement && ControlElement->Settings.ControlType == ERigControlType::Transform)
 				{
 						Section->ControlRig->SetControlValue<FTransform>(ParameterName, InFinalValue.Value,true, EControlRigSetKey::Never);
 				}
-				else if (RigControl && RigControl->ControlType == ERigControlType::TransformNoScale)
+				else if (ControlElement && ControlElement->Settings.ControlType == ERigControlType::TransformNoScale)
 				{
 					FTransformNoScale NoScale = InFinalValue.Value;
 						Section->ControlRig->SetControlValue<FTransformNoScale>(ParameterName, NoScale, true, EControlRigSetKey::Never);
 				}
-				else if (RigControl && RigControl->ControlType == ERigControlType::EulerTransform)
+				else if (ControlElement && ControlElement->Settings.ControlType == ERigControlType::EulerTransform)
 				{
 					FEulerTransform Euler = InFinalValue.Value;
 					Section->ControlRig->SetControlValue<FEulerTransform>(ParameterName, Euler, true, EControlRigSetKey::Never);
@@ -1289,8 +1290,8 @@ void FMovieSceneControlRigParameterTemplate::Evaluate(const FMovieSceneEvaluatio
 			{
 				if (Section->ControlsToSet.Num() == 0 || Section->ControlsToSet.Contains(BoolNameAndValue.ParameterName))
 				{
-					FRigControl* RigControl = Section->ControlRig->FindControl(BoolNameAndValue.ParameterName);
-					if (RigControl && RigControl->ControlType == ERigControlType::Bool)
+					FRigControlElement* ControlElement = Section->ControlRig->FindControl(BoolNameAndValue.ParameterName);
+					if (ControlElement && ControlElement->Settings.ControlType == ERigControlType::Bool)
 					{
 						Section->ControlRig->SetControlValue<bool>(BoolNameAndValue.ParameterName, BoolNameAndValue.Value, true, EControlRigSetKey::Never);
 					}
@@ -1301,8 +1302,8 @@ void FMovieSceneControlRigParameterTemplate::Evaluate(const FMovieSceneEvaluatio
 			{
 				if (Section->ControlsToSet.Num() == 0 || Section->ControlsToSet.Contains(IntegerNameAndValue.ParameterName))
 				{
-					FRigControl* RigControl = Section->ControlRig->FindControl(IntegerNameAndValue.ParameterName);
-					if (RigControl && RigControl->ControlType == ERigControlType::Integer)
+					FRigControlElement* ControlElement = Section->ControlRig->FindControl(IntegerNameAndValue.ParameterName);
+					if (ControlElement && ControlElement->Settings.ControlType == ERigControlType::Integer)
 					{
 						Section->ControlRig->SetControlValue<int32>(IntegerNameAndValue.ParameterName, IntegerNameAndValue.Value, true, EControlRigSetKey::Never);
 					}
diff --git a/Engine/Plugins/Experimental/ControlRig/Source/ControlRig/Private/Sequencer/MovieSceneControlRigParameterTrack.cpp b/Engine/Plugins/Experimental/ControlRig/Source/ControlRig/Private/Sequencer/MovieSceneControlRigParameterTrack.cpp
--- a/Engine/Plugins/Experimental/ControlRig/Source/ControlRig/Private/Sequencer/MovieSceneControlRigParameterTrack.cpp
+++ b/Engine/Plugins/Experimental/ControlRig/Source/ControlRig/Private/Sequencer/MovieSceneControlRigParameterTrack.cpp
@@ -417,15 +417,15 @@ TArray<FFBXNodeAndChannels>* UMovieSceneControlRigParameterTrack::GetNodeAndChan
 			if (StringArray.Num() > 0)
 			{
 				FString NodeName = StringArray[0];
-				FRigControl* Control = GetControlRig() ? GetControlRig()->FindControl(FName(*StringArray[0])) : nullptr;
-				if (Control)
+				FRigControlElement* ControlElement = GetControlRig() ? GetControlRig()->FindControl(FName(*StringArray[0])) : nullptr;
+				if (ControlElement)
 				{
 					NodeName = NodeName.ToUpper();
 					if (NodeAndChannels->Num() == 0 || (*NodeAndChannels)[NodeAndChannels->Num() - 1].NodeName != NodeName)
 					{
 						FFBXNodeAndChannels NodeAndChannel;
 						NodeAndChannel.MovieSceneTrack = this;
-						NodeAndChannel.ControlType = (FFBXControlRigTypeProxyEnum)(uint8)Control->ControlType;
+						NodeAndChannel.ControlType = (FFBXControlRigTypeProxyEnum)(uint8)ControlElement->Settings.ControlType;
 						NodeAndChannel.NodeName = NodeName;
 						NodeAndChannels->Add(NodeAndChannel);
 					}
diff --git a/Engine/Plugins/Experimental/ControlRig/Source/ControlRig/Private/Tools/ControlRigPose.cpp b/Engine/Plugins/Experimental/ControlRig/Source/ControlRig/Private/Tools/ControlRigPose.cpp
--- a/Engine/Plugins/Experimental/ControlRig/Source/ControlRig/Private/Tools/ControlRigPose.cpp
+++ b/Engine/Plugins/Experimental/ControlRig/Source/ControlRig/Private/Tools/ControlRigPose.cpp
@@ -10,15 +10,15 @@
 
 void FControlRigControlPose::SavePose(UControlRig* ControlRig, bool bUseAll)
 {
-	TArray<FRigControl> CurrentControls;
+	TArray<FRigControlElement*> CurrentControls;
 	ControlRig->GetControlsInOrder(CurrentControls);
 	CopyOfControls.SetNum(0);
-	FRigControlHierarchy& Hierarchy = ControlRig->GetControlHierarchy();
-	for (const FRigControl& RigControl : CurrentControls)
+	
+	for (FRigControlElement* ControlElement : CurrentControls)
 	{
-		if (RigControl.bAnimatable && (bUseAll || ControlRig->IsControlSelected(RigControl.Name)))
+		if (ControlElement->Settings.bAnimatable && (bUseAll || ControlRig->IsControlSelected(ControlElement->GetName())))
 		{
-			FRigControlCopy Copy(RigControl, Hierarchy);
+			FRigControlCopy Copy(ControlElement, ControlRig->GetHierarchy());
 			CopyOfControls.Add(Copy);
 		}
 	}
@@ -58,21 +58,23 @@ void FControlRigControlPose::PastePoseInternal(UControlRig* ControlRig, bool bDo
 		MirrorTable.SetUpMirrorTable(ControlRig);
 	}
 
-	TArray<FRigControl> SortedControls;
+	TArray<FRigControlElement*> SortedControls;
 	ControlRig->GetControlsInOrder(SortedControls);
+	
 	bool bDoLocal = true;
-	for (const FRigControl& RigControl : SortedControls)
+	for (FRigControlElement* ControlElement : SortedControls)
 	{
-		if (!ControlRig->IsControlSelected(RigControl.Name))
+		if (!ControlRig->IsControlSelected(ControlElement->GetName()))
 		{
 			continue;
 		}
-		FRigControlCopy* CopyRigControl = MirrorTable.GetControl(*this, RigControl.Name);
+		
+		FRigControlCopy* CopyRigControl = MirrorTable.GetControl(*this, ControlElement->GetName());
 		if (CopyRigControl)
 		{
-			if (CopyRigControl->ControlType == RigControl.ControlType)
+			if (CopyRigControl->ControlType == ControlElement->Settings.ControlType)
 			{
-				switch (RigControl.ControlType)
+				switch (ControlElement->Settings.ControlType)
 				{
 				case ERigControlType::Transform:
 				case ERigControlType::TransformNoScale:
@@ -85,11 +87,11 @@ void FControlRigControlPose::PastePoseInternal(UControlRig* ControlRig, bool bDo
 					{
 						if (bDoLocal) // -V547  
 						{
-							ControlRig->SetControlLocalTransform(RigControl.Name, CopyRigControl->LocalTransform, true, Context);
+							ControlRig->SetControlLocalTransform(ControlElement->GetName(), CopyRigControl->LocalTransform, true, Context);
 						}
 						else
 						{
-							ControlRig->SetControlGlobalTransform(RigControl.Name, CopyRigControl->GlobalTransform, true, Context);
+							ControlRig->SetControlGlobalTransform(ControlElement->GetName(), CopyRigControl->GlobalTransform, true, Context);
 						}
 					}
 					else
@@ -100,32 +102,32 @@ void FControlRigControlPose::PastePoseInternal(UControlRig* ControlRig, bool bDo
 						FQuat LocalRotation;
 						bool bIsMatched = MirrorTable.IsMatched(CopyRigControl->Name);
 						MirrorTable.GetMirrorTransform(*CopyRigControl, bDoLocal,bIsMatched, GlobalTranslation, GlobalRotation, LocalTranslation, LocalRotation);
-						SetControlMirrorTransform(bDoLocal,ControlRig, RigControl.Name, bIsMatched, GlobalTranslation, GlobalRotation, LocalTranslation,LocalRotation, true, Context);
+						SetControlMirrorTransform(bDoLocal,ControlRig, ControlElement->GetName(), bIsMatched, GlobalTranslation, GlobalRotation, LocalTranslation,LocalRotation, true, Context);
 					}				
 					break;
 				}
 				case ERigControlType::Float:
 				{
 					float Val = CopyRigControl->Value.Get<float>();
-					ControlRig->SetControlValue<float>(RigControl.Name, Val, true, Context);
+					ControlRig->SetControlValue<float>(ControlElement->GetName(), Val, true, Context);
 					break;
 				}
 				case ERigControlType::Bool:
 				{
 					bool Val = CopyRigControl->Value.Get<bool>();
-					ControlRig->SetControlValue<bool>(RigControl.Name, Val, true, Context);
+					ControlRig->SetControlValue<bool>(ControlElement->GetName(), Val, true, Context);
 					break;
 				}
 				case ERigControlType::Integer:
 				{
 					int32 Val = CopyRigControl->Value.Get<int32>();
-					ControlRig->SetControlValue<int32>(RigControl.Name, Val, true, Context);
+					ControlRig->SetControlValue<int32>(ControlElement->GetName(), Val, true, Context);
 					break;
 				}
 				case ERigControlType::Vector2D:
 				{
 					FVector2D Val = CopyRigControl->Value.Get<FVector2D>();
-					ControlRig->SetControlValue<FVector2D>(RigControl.Name, Val, true, Context);
+					ControlRig->SetControlValue<FVector2D>(ControlElement->GetName(), Val, true, Context);
 					break;
 				}
 				default:
@@ -154,19 +156,20 @@ void FControlRigControlPose::BlendWithInitialPoses(FControlRigControlPose& Initi
 		MirrorTable.SetUpMirrorTable(ControlRig);
 	}
 
-	TArray<FRigControl> SortedControls;
+	TArray<FRigControlElement*> SortedControls;
 	ControlRig->GetControlsInOrder(SortedControls);
+	
 	bool bDoLocal = true;
-	for (const FRigControl& RigControl : SortedControls)
+	for (FRigControlElement* ControlElement : SortedControls)
 	{
-		if (!ControlRig->IsControlSelected(RigControl.Name))
+		if (!ControlRig->IsControlSelected(ControlElement->GetName()))
 		{
 			continue;
 		}
-		FRigControlCopy* CopyRigControl = MirrorTable.GetControl(*this, RigControl.Name);
+		FRigControlCopy* CopyRigControl = MirrorTable.GetControl(*this, ControlElement->GetName());
 		if (CopyRigControl)
 		{
-			if (CopyRigControl->ControlType == RigControl.ControlType)
+			if (CopyRigControl->ControlType == ControlElement->Settings.ControlType)
 			{
 
 				FRigControlCopy* InitialFound = nullptr;
@@ -194,7 +197,7 @@ void FControlRigControlPose::BlendWithInitialPoses(FControlRigControlPose& Initi
 								Rotation = FQuat::Slerp(InitialVal.GetRotation(), Val.GetRotation(), BlendValue); //doing slerp here not fast lerp, can be slow this is for content creation
 								Scale = FMath::Lerp(InitialVal.GetScale3D(), Val.GetScale3D(), BlendValue);
 								Val = FTransform(Rotation, Translation, Scale);
-								ControlRig->SetControlLocalTransform(RigControl.Name, Val, bDoKey, Context);
+								ControlRig->SetControlLocalTransform(ControlElement->GetName(), Val, bDoKey, Context);
 							}
 							else
 							{
@@ -206,7 +209,7 @@ void FControlRigControlPose::BlendWithInitialPoses(FControlRigControlPose& Initi
 								Rotation = FQuat::Slerp(InitialVal.GetRotation(), Val.GetRotation(), BlendValue); //doing slerp here not fast lerp, can be slow this is for content creation
 								Scale = FMath::Lerp(InitialVal.GetScale3D(), Val.GetScale3D(), BlendValue);
 								Val = FTransform(Rotation, Translation, Scale);
-								ControlRig->SetControlGlobalTransform(RigControl.Name, Val, bDoKey, Context);
+								ControlRig->SetControlGlobalTransform(ControlElement->GetName(), Val, bDoKey, Context);
 							}
 						}
 						else
@@ -225,7 +228,7 @@ void FControlRigControlPose::BlendWithInitialPoses(FControlRigControlPose& Initi
 							GlobalRotation = FQuat::Slerp(InitialGlobalRotation, GlobalRotation, BlendValue); //doing slerp here not fast lerp, can be slow this is for content creation
 							LocalTranslation = FMath::Lerp(InitialLocalTranslation, LocalTranslation, BlendValue);
 							LocalRotation = FQuat::Slerp(InitialLocationRotation, LocalRotation, BlendValue); //doing slerp here not fast lerp, can be slow this is for content creation
-							SetControlMirrorTransform(bDoLocal,ControlRig, RigControl.Name, bIsMatched, GlobalTranslation, GlobalRotation,LocalTranslation, LocalRotation, bDoKey, Context);							
+							SetControlMirrorTransform(bDoLocal,ControlRig, ControlElement->GetName(), bIsMatched, GlobalTranslation, GlobalRotation,LocalTranslation, LocalRotation, bDoKey, Context);							
 						}
 					}
 				}
diff --git a/Engine/Plugins/Experimental/ControlRig/Source/ControlRig/Private/Tools/ControlRigPoseMirrorTable.cpp b/Engine/Plugins/Experimental/ControlRig/Source/ControlRig/Private/Tools/ControlRigPoseMirrorTable.cpp
--- a/Engine/Plugins/Experimental/ControlRig/Source/ControlRig/Private/Tools/ControlRigPoseMirrorTable.cpp
+++ b/Engine/Plugins/Experimental/ControlRig/Source/ControlRig/Private/Tools/ControlRigPoseMirrorTable.cpp
@@ -11,10 +11,10 @@ void FControlRigPoseMirrorTable::SetUpMirrorTable(const UControlRig* ControlRig)
 	MatchedControls.Reset();
 	if (Settings && ControlRig)
 	{
-		const TArray<FRigControl>& CurrentControls = ControlRig->AvailableControls();
-		for (const FRigControl& RigControl : CurrentControls)
+		TArray<FRigControlElement*> CurrentControls = ControlRig->AvailableControls();
+		for (FRigControlElement* ControlElement : CurrentControls)
 		{
-			FString CurrentString = RigControl.Name.ToString();
+			FString CurrentString = ControlElement->GetName().ToString();
 			if (CurrentString.Contains(Settings->RightSide))
 			{
 				FString NewString = CurrentString.Replace(*Settings->RightSide, *Settings->LeftSide);
diff --git a/Engine/Plugins/Experimental/ControlRig/Source/ControlRig/Private/Units/Debug/RigUnit_DebugHierarchy.cpp b/Engine/Plugins/Experimental/ControlRig/Source/ControlRig/Private/Units/Debug/RigUnit_DebugHierarchy.cpp
--- a/Engine/Plugins/Experimental/ControlRig/Source/ControlRig/Private/Units/Debug/RigUnit_DebugHierarchy.cpp
+++ b/Engine/Plugins/Experimental/ControlRig/Source/ControlRig/Private/Units/Debug/RigUnit_DebugHierarchy.cpp
@@ -16,9 +16,9 @@ FRigUnit_DebugHierarchy_Execute()
 		return;
 	}
 
-	const FRigBoneHierarchy* Hierarchy = Context.GetBones();
+	URigHierarchy* Hierarchy = Context.Hierarchy;
 	if (Hierarchy)
 	{
-		Context.DrawInterface->DrawHierarchy(WorldOffset, *Hierarchy, EControlRigDrawHierarchyMode::Axes, Scale, Color, Thickness);
+		Context.DrawInterface->DrawHierarchy(WorldOffset, Hierarchy, EControlRigDrawHierarchyMode::Axes, Scale, Color, Thickness);
 	}
 }
\ No newline at end of file
diff --git a/Engine/Plugins/Experimental/ControlRig/Source/ControlRig/Private/Units/Debug/RigUnit_DebugPoint.cpp b/Engine/Plugins/Experimental/ControlRig/Source/ControlRig/Private/Units/Debug/RigUnit_DebugPoint.cpp
--- a/Engine/Plugins/Experimental/ControlRig/Source/ControlRig/Private/Units/Debug/RigUnit_DebugPoint.cpp
+++ b/Engine/Plugins/Experimental/ControlRig/Source/ControlRig/Private/Units/Debug/RigUnit_DebugPoint.cpp
@@ -18,9 +18,9 @@ FRigUnit_DebugPoint_Execute()
 
 	FVector Center = FVector::ZeroVector;
 	FVector DrawVector = Vector;
-	if (Space != NAME_None && Context.GetBones() != nullptr)
+	if (Space != NAME_None && Context.Hierarchy != nullptr)
 	{
-		FTransform Transform = Context.GetBones()->GetGlobalTransform(Space);
+		const FTransform Transform = Context.Hierarchy->GetGlobalTransform(FRigElementKey(Space, ERigElementType::Bone));
 		Center = Transform.GetLocation();
 		DrawVector = Transform.TransformPosition(DrawVector);
 	}
@@ -55,9 +55,9 @@ FRigUnit_DebugPointMutable_Execute()
 
 	FVector Center = FVector::ZeroVector;
 	FVector DrawVector = Vector;
-	if (Space != NAME_None && Context.GetBones() != nullptr)
+	if (Space != NAME_None && Context.Hierarchy != nullptr)
 	{
-		FTransform Transform = Context.GetBones()->GetGlobalTransform(Space);
+		const FTransform Transform = Context.Hierarchy->GetGlobalTransform(FRigElementKey(Space, ERigElementType::Bone));
 		Center = Transform.GetLocation();
 		DrawVector = Transform.TransformPosition(DrawVector);
 	}
diff --git a/Engine/Plugins/Experimental/ControlRig/Source/ControlRig/Private/Units/Debug/RigUnit_DebugTransform.cpp b/Engine/Plugins/Experimental/ControlRig/Source/ControlRig/Private/Units/Debug/RigUnit_DebugTransform.cpp
--- a/Engine/Plugins/Experimental/ControlRig/Source/ControlRig/Private/Units/Debug/RigUnit_DebugTransform.cpp
+++ b/Engine/Plugins/Experimental/ControlRig/Source/ControlRig/Private/Units/Debug/RigUnit_DebugTransform.cpp
@@ -17,9 +17,9 @@ FRigUnit_DebugTransform_Execute()
 	}
 
 	FTransform DrawTransform = Transform;
-	if (Space != NAME_None && Context.GetBones() != nullptr)
+	if (Space != NAME_None && Context.Hierarchy != nullptr)
 	{
-		DrawTransform = Transform * Context.GetBones()->GetGlobalTransform(Space);
+		DrawTransform = Transform * Context.Hierarchy->GetGlobalTransform(FRigElementKey(Space, ERigElementType::Bone));
 	}
 
 	switch (Mode)
@@ -73,9 +73,9 @@ FRigUnit_DebugTransformMutableItemSpace_Execute()
 	}
 
 	FTransform DrawTransform = Transform;
-	if (Space != NAME_None && Context.GetBones() != nullptr)
+	if (Space.IsValid() && Context.Hierarchy != nullptr)
 	{
-		DrawTransform = Transform * Context.GetBones()->GetGlobalTransform(Space);
+		DrawTransform = Transform * Context.Hierarchy->GetGlobalTransform(Space);
 	}
 
 	switch (Mode)
@@ -113,9 +113,9 @@ FRigUnit_DebugTransformArrayMutable_Execute()
 	}
 
 	WorkData.DrawTransforms.SetNumUninitialized(Transforms.Num());
-	if (Space != NAME_None && Context.GetBones() != nullptr)
+	if (Space != NAME_None && Context.Hierarchy)
 	{
-		FTransform SpaceTransform = Context.GetBones()->GetGlobalTransform(Space);
+		FTransform SpaceTransform = Context.Hierarchy->GetGlobalTransform(FRigElementKey(Space, ERigElementType::Bone));
 		for(int32 Index=0;Index<Transforms.Num();Index++)
 		{
 			WorkData.DrawTransforms[Index] = Transforms[Index] * SpaceTransform;
diff --git a/Engine/Plugins/Experimental/ControlRig/Source/ControlRig/Private/Units/Deprecated/RigUnit_AimConstraint.cpp b/Engine/Plugins/Experimental/ControlRig/Source/ControlRig/Private/Units/Deprecated/RigUnit_AimConstraint.cpp
--- a/Engine/Plugins/Experimental/ControlRig/Source/ControlRig/Private/Units/Deprecated/RigUnit_AimConstraint.cpp
+++ b/Engine/Plugins/Experimental/ControlRig/Source/ControlRig/Private/Units/Deprecated/RigUnit_AimConstraint.cpp
@@ -15,10 +15,11 @@ FRigUnit_AimConstraint_Execute()
 	{
 		ConstraintData.Reset();
 
-		FRigBoneHierarchy* Hierarchy = ExecuteContext.GetBones();
+		URigHierarchy* Hierarchy = ExecuteContext.Hierarchy;
 		if (Hierarchy)
 		{
-			int32 BoneIndex = Hierarchy->GetIndex(Joint);
+			const FRigElementKey Key(Joint, ERigElementType::Bone); 
+			int32 BoneIndex = Hierarchy->GetIndex(Key);
 			if (BoneIndex != INDEX_NONE)
 			{
 				const int32 TargetNum = AimTargets.Num();
@@ -42,10 +43,11 @@ FRigUnit_AimConstraint_Execute()
 	}
 	else if (Context.State == EControlRigState::Update)
 	{
-		FRigBoneHierarchy* Hierarchy = ExecuteContext.GetBones();
+		URigHierarchy* Hierarchy = ExecuteContext.Hierarchy;
 		if (Hierarchy)
 		{
-			int32 BoneIndex = Hierarchy->GetIndex(Joint);
+			const FRigElementKey Key(Joint, ERigElementType::Bone); 
+			int32 BoneIndex = Hierarchy->GetIndex(Key);
 			if (BoneIndex != INDEX_NONE)
 			{
 				const int32 TargetNum = AimTargets.Num();
@@ -68,7 +70,7 @@ FRigUnit_AimConstraint_Execute()
 					}
 
 					FTransform BaseTransform = FTransform::Identity;
-					int32 ParentIndex = (*Hierarchy)[BoneIndex].ParentIndex;
+					int32 ParentIndex = Hierarchy->GetIndex(Hierarchy->GetFirstParent(Key));
 					if (ParentIndex != INDEX_NONE)
 					{
 						BaseTransform = Hierarchy->GetGlobalTransform(ParentIndex);
diff --git a/Engine/Plugins/Experimental/ControlRig/Source/ControlRig/Private/Units/Deprecated/RigUnit_ApplyFK.cpp b/Engine/Plugins/Experimental/ControlRig/Source/ControlRig/Private/Units/Deprecated/RigUnit_ApplyFK.cpp
--- a/Engine/Plugins/Experimental/ControlRig/Source/ControlRig/Private/Units/Deprecated/RigUnit_ApplyFK.cpp
+++ b/Engine/Plugins/Experimental/ControlRig/Source/ControlRig/Private/Units/Deprecated/RigUnit_ApplyFK.cpp
@@ -14,10 +14,11 @@ FRigUnit_ApplyFK_Execute()
 	}
 	else if (Context.State == EControlRigState::Update)
 	{
-		FRigBoneHierarchy* Hierarchy = ExecuteContext.GetBones();
+		URigHierarchy* Hierarchy = ExecuteContext.Hierarchy;
 		if (Hierarchy)
 		{
-			int32 Index = Hierarchy->GetIndex(Joint);
+			const FRigElementKey Key(Joint, ERigElementType::Bone);
+			int32 Index = Hierarchy->GetIndex(Key);
 			if (Index != INDEX_NONE)
 			{
 				// first filter input transform
@@ -26,8 +27,8 @@ FRigUnit_ApplyFK_Execute()
 
 				FTransform InputBaseTransform = UtilityHelpers::GetBaseTransformByMode(
 					ApplyTransformSpace,
-					[Hierarchy](const FRigElementKey& BoneKey) { return Hierarchy->GetGlobalTransform(BoneKey.Name); },
-					(*Hierarchy)[Index].GetParentElementKey(),
+					[Hierarchy](const FRigElementKey& BoneKey) { return Hierarchy->GetGlobalTransform(BoneKey); },
+					Hierarchy->GetFirstParent(Key),
 					FRigElementKey(BaseJoint, ERigElementType::Bone),
 					BaseTransform
 				);
@@ -58,8 +59,8 @@ FRigUnit_ApplyFK_Execute()
 
 IMPLEMENT_RIGUNIT_AUTOMATION_TEST(FRigUnit_ApplyFK)
 {
-	BoneHierarchy.Add(TEXT("Root"), NAME_None, ERigBoneType::User, FTransform(FVector(1.f, 0.f, 0.f)));
-	BoneHierarchy.Add(TEXT("BoneA"), TEXT("Root"), ERigBoneType::User, FTransform(FVector(1.f, 2.f, 3.f)));
+	const FRigElementKey Root = Controller->AddBone(TEXT("Root"), FRigElementKey(), FTransform(FVector(1.f, 0.f, 0.f)), true, ERigBoneType::User);
+	const FRigElementKey BoneA = Controller->AddBone(TEXT("BoneA"), Root, FTransform(FVector(1.f, 2.f, 3.f)), true, ERigBoneType::User);
 
 	Unit.ExecuteContext = ExecuteContext;
 	Unit.Joint = TEXT("BoneA");
@@ -67,25 +68,25 @@ IMPLEMENT_RIGUNIT_AUTOMATION_TEST(FRigUnit_ApplyFK)
 	Unit.ApplyTransformSpace = ETransformSpaceMode::GlobalSpace;
 	Unit.Transform = FTransform(FVector(0.f, 5.f, 0.f));
 
-	BoneHierarchy.Initialize();
+	Hierarchy->ResetPoseToInitial(ERigElementType::Bone);
 	Execute();
-	AddErrorIfFalse(BoneHierarchy.GetGlobalTransform(1).GetTranslation().Equals(FVector(0.f, 5.f, 0.f)), TEXT("unexpected global transform"));
-	AddErrorIfFalse(BoneHierarchy.GetLocalTransform(1).GetTranslation().Equals(FVector(-1.f, 5.f, 0.f)), TEXT("unexpected local transform"));
+	AddErrorIfFalse(Hierarchy->GetGlobalTransform(1).GetTranslation().Equals(FVector(0.f, 5.f, 0.f)), TEXT("unexpected global transform"));
+	AddErrorIfFalse(Hierarchy->GetLocalTransform(1).GetTranslation().Equals(FVector(-1.f, 5.f, 0.f)), TEXT("unexpected local transform"));
 
 	Unit.ApplyTransformMode = EApplyTransformMode::Override;
 	Unit.ApplyTransformSpace = ETransformSpaceMode::LocalSpace;
 
-	BoneHierarchy.Initialize();
+	Hierarchy->ResetPoseToInitial(ERigElementType::Bone);
 	Execute();
-	AddErrorIfFalse(BoneHierarchy.GetGlobalTransform(1).GetTranslation().Equals(FVector(1.f, 5.f, 0.f)), TEXT("unexpected global transform"));
-	AddErrorIfFalse(BoneHierarchy.GetLocalTransform(1).GetTranslation().Equals(FVector(0.f, 5.f, 0.f)), TEXT("unexpected local transform"));
+	AddErrorIfFalse(Hierarchy->GetGlobalTransform(1).GetTranslation().Equals(FVector(1.f, 5.f, 0.f)), TEXT("unexpected global transform"));
+	AddErrorIfFalse(Hierarchy->GetLocalTransform(1).GetTranslation().Equals(FVector(0.f, 5.f, 0.f)), TEXT("unexpected local transform"));
 
 	Unit.ApplyTransformMode = EApplyTransformMode::Additive;
 
-	BoneHierarchy.Initialize();
+	Hierarchy->ResetPoseToInitial(ERigElementType::Bone);
 	Execute();
-	AddErrorIfFalse(BoneHierarchy.GetGlobalTransform(1).GetTranslation().Equals(FVector(1.f, 7.f, 3.f)), TEXT("unexpected global transform"));
-	AddErrorIfFalse(BoneHierarchy.GetLocalTransform(1).GetTranslation().Equals(FVector(0.f, 7.f, 3.f)), TEXT("unexpected local transform"));
+	AddErrorIfFalse(Hierarchy->GetGlobalTransform(1).GetTranslation().Equals(FVector(1.f, 7.f, 3.f)), TEXT("unexpected global transform"));
+	AddErrorIfFalse(Hierarchy->GetLocalTransform(1).GetTranslation().Equals(FVector(0.f, 7.f, 3.f)), TEXT("unexpected local transform"));
 	return true;
 }
 #endif
\ No newline at end of file
diff --git a/Engine/Plugins/Experimental/ControlRig/Source/ControlRig/Private/Units/Deprecated/RigUnit_GetJointTransform.cpp b/Engine/Plugins/Experimental/ControlRig/Source/ControlRig/Private/Units/Deprecated/RigUnit_GetJointTransform.cpp
--- a/Engine/Plugins/Experimental/ControlRig/Source/ControlRig/Private/Units/Deprecated/RigUnit_GetJointTransform.cpp
+++ b/Engine/Plugins/Experimental/ControlRig/Source/ControlRig/Private/Units/Deprecated/RigUnit_GetJointTransform.cpp
@@ -7,34 +7,4 @@
 FRigUnit_GetJointTransform_Execute()
 {
     DECLARE_SCOPE_HIERARCHICAL_COUNTER_RIGUNIT()
-
-	FRigBoneHierarchy* Hierarchy = ExecuteContext.GetBones();
-	if (Hierarchy)
-	{
-		int32 Index = Hierarchy->GetIndex(Joint);
-		if (Index != INDEX_NONE)
-		{
-			FTransform GlobalTransform;
-			switch (Type)
-			{
-			case ETransformGetterType::Current:
-				{
-					FTransform ComputedBaseTransform = UtilityHelpers::GetBaseTransformByMode(TransformSpace, [Hierarchy](const FRigElementKey& JointKey) { return Hierarchy->GetGlobalTransform(JointKey.Name); },
-					(*Hierarchy)[Index].GetParentElementKey(), FRigElementKey(BaseJoint, ERigElementType::Bone), BaseTransform);
-
-					Output = Hierarchy->GetInitialGlobalTransform(Index).GetRelativeTransform(ComputedBaseTransform);
-					break;
-				}
-			case ETransformGetterType::Initial:
-			default:
-				{
-					FTransform ComputedBaseTransform = UtilityHelpers::GetBaseTransformByMode(TransformSpace, [Hierarchy](const FRigElementKey& JointKey) { return Hierarchy->GetInitialGlobalTransform(JointKey.Name); },
-					(*Hierarchy)[Index].GetParentElementKey(), FRigElementKey(BaseJoint, ERigElementType::Bone), BaseTransform);
-
-					Output = Hierarchy->GetInitialGlobalTransform(Index).GetRelativeTransform(ComputedBaseTransform);
-					break;
-				}
-			}
-		}
-	}
 }
diff --git a/Engine/Plugins/Experimental/ControlRig/Source/ControlRig/Private/Units/Deprecated/RigUnit_TwoBoneIKFK.cpp b/Engine/Plugins/Experimental/ControlRig/Source/ControlRig/Private/Units/Deprecated/RigUnit_TwoBoneIKFK.cpp
--- a/Engine/Plugins/Experimental/ControlRig/Source/ControlRig/Private/Units/Deprecated/RigUnit_TwoBoneIKFK.cpp
+++ b/Engine/Plugins/Experimental/ControlRig/Source/ControlRig/Private/Units/Deprecated/RigUnit_TwoBoneIKFK.cpp
@@ -11,23 +11,26 @@ FRigUnit_TwoBoneIKFK_Execute()
 
 	if (Context.State == EControlRigState::Init)
 	{
-		const FRigBoneHierarchy* Hierarchy = Context.GetBones();
+		URigHierarchy* Hierarchy = ExecuteContext.Hierarchy;
 		if (Hierarchy)
 		{
 			// reset
 			StartJointIndex = MidJointIndex = EndJointIndex = INDEX_NONE;
 			UpperLimbLength = LowerLimbLength = 0.f;
 
+			const FRigElementKey StartJointKey(StartJoint, ERigElementType::Bone);
+			const FRigElementKey EndJointKey(EndJoint, ERigElementType::Bone);
+			
 			// verify the chain
-			int32 StartIndex = Hierarchy->GetIndex(StartJoint);
-			int32 EndIndex = Hierarchy->GetIndex(EndJoint);
+			int32 StartIndex = Hierarchy->GetIndex(StartJointKey);
+			int32 EndIndex = Hierarchy->GetIndex(EndJointKey);
 			if (StartIndex != INDEX_NONE && EndIndex != INDEX_NONE)
 			{
 				// ensure the chain
-				int32 EndParentIndex = (*Hierarchy)[EndIndex].ParentIndex;
+				int32 EndParentIndex = Hierarchy->GetFirstParent(EndIndex);
 				if (EndParentIndex != INDEX_NONE)
 				{
-					int32 MidParentIndex = (*Hierarchy)[EndParentIndex].ParentIndex;
+					int32 MidParentIndex = Hierarchy->GetFirstParent(EndParentIndex);
 					if (MidParentIndex == StartIndex)
 					{
 						StartJointIndex = StartIndex;
@@ -71,7 +74,7 @@ FRigUnit_TwoBoneIKFK_Execute()
 			else if (FMath::IsNearlyEqual(IKBlend, 1.f))
 			{
 				// update transform before going through IK
-				const FRigBoneHierarchy* Hierarchy = Context.GetBones();
+				const URigHierarchy* Hierarchy = Context.Hierarchy;
 				check(Hierarchy);
 
 				StartJointIKTransform = Hierarchy->GetGlobalTransform(StartJointIndex);
@@ -88,7 +91,7 @@ FRigUnit_TwoBoneIKFK_Execute()
 			else
 			{
 				// update transform before going through IK
-				const FRigBoneHierarchy* Hierarchy = Context.GetBones();
+				const URigHierarchy* Hierarchy = Context.Hierarchy;
 				check(Hierarchy);
 
 				StartJointIKTransform = Hierarchy->GetGlobalTransform(StartJointIndex);
@@ -103,7 +106,7 @@ FRigUnit_TwoBoneIKFK_Execute()
 				EndJointTransform.Blend(MidJointFKTransform, EndJointIKTransform, IKBlend);
 			}
 
-			FRigBoneHierarchy* Hierarchy = ExecuteContext.GetBones();
+			URigHierarchy* Hierarchy = ExecuteContext.Hierarchy;
 			check(Hierarchy);
 			Hierarchy->SetGlobalTransform(StartJointIndex, StartJointTransform);
 			Hierarchy->SetGlobalTransform(MidJointIndex, MidJointTransform);
diff --git a/Engine/Plugins/Experimental/ControlRig/Source/ControlRig/Private/Units/Execution/RigUnit_BeginExecution.cpp b/Engine/Plugins/Experimental/ControlRig/Source/ControlRig/Private/Units/Execution/RigUnit_BeginExecution.cpp
--- a/Engine/Plugins/Experimental/ControlRig/Source/ControlRig/Private/Units/Execution/RigUnit_BeginExecution.cpp
+++ b/Engine/Plugins/Experimental/ControlRig/Source/ControlRig/Private/Units/Execution/RigUnit_BeginExecution.cpp
@@ -8,6 +8,6 @@ FName FRigUnit_BeginExecution::EventName = TEXT("Update");
 FRigUnit_BeginExecution_Execute()
 {
     DECLARE_SCOPE_HIERARCHICAL_COUNTER_RIGUNIT()
-	ExecuteContext.Hierarchy = (FRigHierarchyContainer*)Context.Hierarchy;
+	ExecuteContext.Hierarchy = Context.Hierarchy;
 	ExecuteContext.EventName = FRigUnit_BeginExecution::EventName;
 }
diff --git a/Engine/Plugins/Experimental/ControlRig/Source/ControlRig/Private/Units/Execution/RigUnit_Collection.cpp b/Engine/Plugins/Experimental/ControlRig/Source/ControlRig/Private/Units/Execution/RigUnit_Collection.cpp
--- a/Engine/Plugins/Experimental/ControlRig/Source/ControlRig/Private/Units/Execution/RigUnit_Collection.cpp
+++ b/Engine/Plugins/Experimental/ControlRig/Source/ControlRig/Private/Units/Execution/RigUnit_Collection.cpp
@@ -13,7 +13,7 @@ FRigUnit_CollectionChain_Execute()
 		CachedHierarchyHash = INDEX_NONE;
 	}
 
-	int32 CurrentHierarchyHash = Context.Hierarchy->Version * 17;
+	int32 CurrentHierarchyHash = Context.Hierarchy->GetTopologyVersion() * 17;
 	CurrentHierarchyHash += GetTypeHash(FirstItem);
 	CurrentHierarchyHash += GetTypeHash(LastItem);
 	CurrentHierarchyHash += Reverse ? 1 : 0;
@@ -54,7 +54,7 @@ FRigUnit_CollectionNameSearch_Execute()
 		CachedHierarchyHash = INDEX_NONE;
 	}
 
-	int32 CurrentHierarchyHash = Context.Hierarchy->Version * 17;
+	int32 CurrentHierarchyHash = Context.Hierarchy->GetTopologyVersion() * 17;
 	CurrentHierarchyHash += GetTypeHash(PartialName);
 	CurrentHierarchyHash += (int32)TypeToSearch * 8;
 
@@ -76,7 +76,7 @@ FRigUnit_CollectionChildren_Execute()
 		CachedHierarchyHash = INDEX_NONE;
 	}
 
-	int32 CurrentHierarchyHash = Context.Hierarchy->Version * 17;
+	int32 CurrentHierarchyHash = Context.Hierarchy->GetTopologyVersion() * 17;
 	CurrentHierarchyHash += GetTypeHash(Parent);
 	CurrentHierarchyHash += bRecursive ? 2 : 0;
 	CurrentHierarchyHash += bIncludeParent ? 1 : 0;
@@ -107,7 +107,7 @@ FRigUnit_CollectionReplaceItems_Execute()
 		CachedHierarchyHash = INDEX_NONE;
 	}
 
-	int32 CurrentHierarchyHash = Context.Hierarchy->Version * 17;
+	int32 CurrentHierarchyHash = Context.Hierarchy->GetTopologyVersion() * 17;
 	CurrentHierarchyHash += GetTypeHash(Items);
 	CurrentHierarchyHash += 12 * GetTypeHash(Old);
 	CurrentHierarchyHash += 13 * GetTypeHash(New);
diff --git a/Engine/Plugins/Experimental/ControlRig/Source/ControlRig/Private/Units/Execution/RigUnit_Hierarchy.cpp b/Engine/Plugins/Experimental/ControlRig/Source/ControlRig/Private/Units/Execution/RigUnit_Hierarchy.cpp
--- a/Engine/Plugins/Experimental/ControlRig/Source/ControlRig/Private/Units/Execution/RigUnit_Hierarchy.cpp
+++ b/Engine/Plugins/Experimental/ControlRig/Source/ControlRig/Private/Units/Execution/RigUnit_Hierarchy.cpp
@@ -25,7 +25,7 @@ FRigUnit_HierarchyGetParent_Execute()
 
 		if(CachedChild.UpdateCache(Child, Context.Hierarchy))
 		{
-			Parent = Context.Hierarchy->GetParentKey(Child);
+			Parent = Context.Hierarchy->GetFirstParent(Child);
 			if(Parent.IsValid())
 			{
 				CachedParent.UpdateCache(Parent, Context.Hierarchy);
@@ -58,7 +58,7 @@ FRigUnit_HierarchyGetParents_Execute()
 				{
 					Keys.Add(Parent);
 				}
-				Parent = Context.Hierarchy->GetParentKey(Parent);
+				Parent = Context.Hierarchy->GetFirstParent(Parent);
 			}
 			while(Parent.IsValid());
 
@@ -95,7 +95,9 @@ FRigUnit_HierarchyGetChildren_Execute()
 			{
 				Keys.Add(Parent);
 			}
-			Keys.Append(Context.Hierarchy->GetChildKeys(Parent, bRecursive));
+
+			
+			Keys.Append(Context.Hierarchy->GetChildren(Parent, bRecursive));
 
 			CachedChildren = FRigElementKeyCollection(Keys);
 		}
@@ -122,10 +124,10 @@ FRigUnit_HierarchyGetSiblings_Execute()
 		{
 			TArray<FRigElementKey> Keys;
 
-			FRigElementKey Parent = Context.Hierarchy->GetParentKey(Item);
+			FRigElementKey Parent = Context.Hierarchy->GetFirstParent(Item);
 			if(Parent.IsValid())
 			{
-				TArray<FRigElementKey> Children = Context.Hierarchy->GetChildKeys(Parent, false);
+				TArray<FRigElementKey> Children = Context.Hierarchy->GetChildren(Parent, false);
 				for(FRigElementKey Child : Children)
 				{
 					if(bIncludeItem || Child != Item)
diff --git a/Engine/Plugins/Experimental/ControlRig/Source/ControlRig/Private/Units/Execution/RigUnit_InverseExecution.cpp b/Engine/Plugins/Experimental/ControlRig/Source/ControlRig/Private/Units/Execution/RigUnit_InverseExecution.cpp
--- a/Engine/Plugins/Experimental/ControlRig/Source/ControlRig/Private/Units/Execution/RigUnit_InverseExecution.cpp
+++ b/Engine/Plugins/Experimental/ControlRig/Source/ControlRig/Private/Units/Execution/RigUnit_InverseExecution.cpp
@@ -8,6 +8,6 @@ FName FRigUnit_InverseExecution::EventName = TEXT("Inverse");
 FRigUnit_InverseExecution_Execute()
 {
     DECLARE_SCOPE_HIERARCHICAL_COUNTER_RIGUNIT()
-	ExecuteContext.Hierarchy = (FRigHierarchyContainer*)Context.Hierarchy;
+	ExecuteContext.Hierarchy = Context.Hierarchy;
 	ExecuteContext.EventName = FRigUnit_InverseExecution::EventName;
 }
diff --git a/Engine/Plugins/Experimental/ControlRig/Source/ControlRig/Private/Units/Hierarchy/RigUnit_AddBoneTransform.cpp b/Engine/Plugins/Experimental/ControlRig/Source/ControlRig/Private/Units/Hierarchy/RigUnit_AddBoneTransform.cpp
--- a/Engine/Plugins/Experimental/ControlRig/Source/ControlRig/Private/Units/Hierarchy/RigUnit_AddBoneTransform.cpp
+++ b/Engine/Plugins/Experimental/ControlRig/Source/ControlRig/Private/Units/Hierarchy/RigUnit_AddBoneTransform.cpp
@@ -6,7 +6,7 @@
 FRigUnit_AddBoneTransform_Execute()
 {
     DECLARE_SCOPE_HIERARCHICAL_COUNTER_RIGUNIT()
-	FRigBoneHierarchy* Hierarchy = ExecuteContext.GetBones();
+	URigHierarchy* Hierarchy = ExecuteContext.Hierarchy;
 	if (Hierarchy)
 	{
 		switch (Context.State)
@@ -17,7 +17,8 @@ FRigUnit_AddBoneTransform_Execute()
 			}
 			case EControlRigState::Update:
 			{
-				if (!CachedBone.UpdateCache(Bone, Hierarchy))
+				const FRigElementKey Key(Bone, ERigElementType::Bone); 
+				if (!CachedBone.UpdateCache(Key, Hierarchy))
 				{
 					UE_CONTROLRIG_RIGUNIT_REPORT_WARNING(TEXT("Bone '%s' is not valid."), *Bone.ToString());
 				}
diff --git a/Engine/Plugins/Experimental/ControlRig/Source/ControlRig/Private/Units/Hierarchy/RigUnit_GetBoneTransform.cpp b/Engine/Plugins/Experimental/ControlRig/Source/ControlRig/Private/Units/Hierarchy/RigUnit_GetBoneTransform.cpp
--- a/Engine/Plugins/Experimental/ControlRig/Source/ControlRig/Private/Units/Hierarchy/RigUnit_GetBoneTransform.cpp
+++ b/Engine/Plugins/Experimental/ControlRig/Source/ControlRig/Private/Units/Hierarchy/RigUnit_GetBoneTransform.cpp
@@ -7,7 +7,7 @@
 FRigUnit_GetBoneTransform_Execute()
 {
     DECLARE_SCOPE_HIERARCHICAL_COUNTER_RIGUNIT()
-	const FRigBoneHierarchy* Hierarchy = Context.GetBones();
+	const URigHierarchy* Hierarchy = Context.Hierarchy;
 	if (Hierarchy)
 	{
 		switch (Context.State)
@@ -19,7 +19,7 @@ FRigUnit_GetBoneTransform_Execute()
 			}
 			case EControlRigState::Update:
 			{
-				if (!CachedBone.UpdateCache(Bone, Hierarchy))
+				if (!CachedBone.UpdateCache(FRigElementKey(Bone, ERigElementType::Bone), Hierarchy))
 				{
 					UE_CONTROLRIG_RIGUNIT_REPORT_WARNING(TEXT("Bone '%s' is not valid."), *Bone.ToString());
 				}
@@ -57,9 +57,8 @@ FRigUnit_GetBoneTransform_Execute()
 
 IMPLEMENT_RIGUNIT_AUTOMATION_TEST(FRigUnit_GetBoneTransform)
 {
-	BoneHierarchy.Add(TEXT("Root"), NAME_None, ERigBoneType::User, FTransform(FVector(1.f, 0.f, 0.f)));
-	BoneHierarchy.Add(TEXT("BoneA"), TEXT("Root"), ERigBoneType::User, FTransform(FVector(1.f, 2.f, 3.f)));
-	BoneHierarchy.Initialize();
+	const FRigElementKey Root = Controller->AddBone(TEXT("Root"), FRigElementKey(), FTransform(FVector(1.f, 0.f, 0.f)), true, ERigBoneType::User);
+	const FRigElementKey BoneA = Controller->AddBone(TEXT("BoneA"), Root, FTransform(FVector(1.f, 2.f, 3.f)), true, ERigBoneType::User);
 
 	Unit.Bone = TEXT("Unknown");
 	Unit.Space = EBoneGetterSetterMode::GlobalSpace;
diff --git a/Engine/Plugins/Experimental/ControlRig/Source/ControlRig/Private/Units/Hierarchy/RigUnit_GetBoneTransform.h b/Engine/Plugins/Experimental/ControlRig/Source/ControlRig/Private/Units/Hierarchy/RigUnit_GetBoneTransform.h
--- a/Engine/Plugins/Experimental/ControlRig/Source/ControlRig/Private/Units/Hierarchy/RigUnit_GetBoneTransform.h
+++ b/Engine/Plugins/Experimental/ControlRig/Source/ControlRig/Private/Units/Hierarchy/RigUnit_GetBoneTransform.h
@@ -22,14 +22,9 @@ struct CONTROLRIG_API FRigUnit_GetBoneTransform : public FRigUnit
 	{
 		if (InPinPath.StartsWith(TEXT("Transform")) && Space == EBoneGetterSetterMode::LocalSpace)
 		{
-			if (const FRigHierarchyContainer* Container = (const FRigHierarchyContainer*)InUserContext)
+			if (const URigHierarchy* Hierarchy = (const URigHierarchy*)InUserContext)
 			{
-				int32 BoneIndex = Container->BoneHierarchy.GetIndex(Bone);
-				if (BoneIndex != INDEX_NONE)
-				{
-					return Container->BoneHierarchy[BoneIndex].GetParentElementKey();
-				}
-
+				return Hierarchy->GetFirstParent(FRigElementKey(Bone, ERigElementType::Bone));
 			}
 		}
 		return FRigElementKey();
diff --git a/Engine/Plugins/Experimental/ControlRig/Source/ControlRig/Private/Units/Hierarchy/RigUnit_GetControlInitialTransform.cpp b/Engine/Plugins/Experimental/ControlRig/Source/ControlRig/Private/Units/Hierarchy/RigUnit_GetControlInitialTransform.cpp
--- a/Engine/Plugins/Experimental/ControlRig/Source/ControlRig/Private/Units/Hierarchy/RigUnit_GetControlInitialTransform.cpp
+++ b/Engine/Plugins/Experimental/ControlRig/Source/ControlRig/Private/Units/Hierarchy/RigUnit_GetControlInitialTransform.cpp
@@ -6,7 +6,7 @@
 FRigUnit_GetControlInitialTransform_Execute()
 {
     DECLARE_SCOPE_HIERARCHICAL_COUNTER_RIGUNIT()
-	const FRigControlHierarchy* Hierarchy = Context.GetControls();
+	const URigHierarchy* Hierarchy = Context.Hierarchy;
 	if (Hierarchy)
 	{
 		switch (Context.State)
@@ -17,7 +17,8 @@ FRigUnit_GetControlInitialTransform_Execute()
 			}
 			case EControlRigState::Update:
 			{
-				if (!CachedControlIndex.UpdateCache(Control, Hierarchy))
+				const FRigElementKey Key(Control, ERigElementType::Control); 
+				if (!CachedControlIndex.UpdateCache(Key, Hierarchy))
 				{
 					UE_CONTROLRIG_RIGUNIT_REPORT_WARNING(TEXT("Control '%s' is not valid."), *Control.ToString());
 				}
@@ -32,7 +33,7 @@ FRigUnit_GetControlInitialTransform_Execute()
 						}
 						case EBoneGetterSetterMode::LocalSpace:
 						{
-							Transform = Hierarchy->GetInitialValue<FTransform>(CachedControlIndex);
+							Transform = Hierarchy->GetInitialLocalTransform(CachedControlIndex);
 							break;
 						}
 						default:
