Control Rig: New rig hierarchy data structure refactoring

diff --git a/Engine/Plugins/Experimental/ControlRig/Source/ControlRig/Private/Units/Hierarchy/RigUnit_GetControlTransform.cpp b/Engine/Plugins/Experimental/ControlRig/Source/ControlRig/Private/Units/Hierarchy/RigUnit_GetControlTransform.cpp
--- a/Engine/Plugins/Experimental/ControlRig/Source/ControlRig/Private/Units/Hierarchy/RigUnit_GetControlTransform.cpp
+++ b/Engine/Plugins/Experimental/ControlRig/Source/ControlRig/Private/Units/Hierarchy/RigUnit_GetControlTransform.cpp
@@ -6,7 +6,7 @@
 FRigUnit_GetControlBool_Execute()
 {
     DECLARE_SCOPE_HIERARCHICAL_COUNTER_RIGUNIT()
-	const FRigControlHierarchy* Hierarchy = Context.GetControls();
+	const URigHierarchy* Hierarchy = Context.Hierarchy;
 	if (Hierarchy)
 	{
 		switch (Context.State)
@@ -17,13 +17,13 @@ FRigUnit_GetControlBool_Execute()
 			}
 			case EControlRigState::Update:
 			{
-				if (!CachedControlIndex.UpdateCache(Control, Hierarchy))
+				if (!CachedControlIndex.UpdateCache(FRigElementKey(Control, ERigElementType::Control), Hierarchy))
 				{
 					UE_CONTROLRIG_RIGUNIT_REPORT_WARNING(TEXT("Control '%s' is not valid."), *Control.ToString());
 				}
 				else
 				{
-					BoolValue = Hierarchy->GetValue(CachedControlIndex).Get<bool>();
+					BoolValue = Hierarchy->GetControlValue(CachedControlIndex).Get<bool>();
 				}
 			}
 			default:
@@ -37,7 +37,7 @@ FRigUnit_GetControlBool_Execute()
 FRigUnit_GetControlFloat_Execute()
 {
     DECLARE_SCOPE_HIERARCHICAL_COUNTER_RIGUNIT()
-	const FRigControlHierarchy* Hierarchy = Context.GetControls();
+	const URigHierarchy* Hierarchy = Context.Hierarchy;
 	if (Hierarchy)
 	{
 		switch (Context.State)
@@ -48,15 +48,15 @@ FRigUnit_GetControlFloat_Execute()
 			}
 			case EControlRigState::Update:
 			{
-				if (!CachedControlIndex.UpdateCache(Control, Hierarchy))
+				if (!CachedControlIndex.UpdateCache(FRigElementKey(Control, ERigElementType::Control), Hierarchy))
 				{
 					UE_CONTROLRIG_RIGUNIT_REPORT_WARNING(TEXT("Control '%s' is not valid."), *Control.ToString());
 				}
 				else
 				{
-					FloatValue = Hierarchy->GetValue(CachedControlIndex).Get<float>();
-					Minimum = Hierarchy->GetValue(CachedControlIndex, ERigControlValueType::Minimum).Get<float>();
-					Maximum = Hierarchy->GetValue(CachedControlIndex, ERigControlValueType::Maximum).Get<float>();
+					FloatValue = Hierarchy->GetControlValue(CachedControlIndex).Get<float>();
+					Minimum = Hierarchy->GetControlValue(CachedControlIndex, ERigControlValueType::Minimum).Get<float>();
+					Maximum = Hierarchy->GetControlValue(CachedControlIndex, ERigControlValueType::Maximum).Get<float>();
 				}
 			}
 			default:
@@ -71,7 +71,7 @@ FRigUnit_GetControlFloat_Execute()
 FRigUnit_GetControlInteger_Execute()
 {
     DECLARE_SCOPE_HIERARCHICAL_COUNTER_RIGUNIT()
-	const FRigControlHierarchy* Hierarchy = Context.GetControls();
+	const URigHierarchy* Hierarchy = Context.Hierarchy;
 	if (Hierarchy)
 	{
 		switch (Context.State)
@@ -82,15 +82,15 @@ FRigUnit_GetControlInteger_Execute()
 			}
 			case EControlRigState::Update:
 			{
-				if (!CachedControlIndex.UpdateCache(Control, Hierarchy))
+				if (!CachedControlIndex.UpdateCache(FRigElementKey(Control, ERigElementType::Control), Hierarchy))
 				{
 					UE_CONTROLRIG_RIGUNIT_REPORT_WARNING(TEXT("Control '%s' is not valid."), *Control.ToString());
 				}
 				else
 				{
-					IntegerValue = Hierarchy->GetValue(CachedControlIndex).Get<int32>();
-					Minimum = Hierarchy->GetValue(CachedControlIndex, ERigControlValueType::Minimum).Get<int32>();
-					Maximum = Hierarchy->GetValue(CachedControlIndex, ERigControlValueType::Maximum).Get<int32>();
+					IntegerValue = Hierarchy->GetControlValue(CachedControlIndex).Get<int32>();
+					Minimum = Hierarchy->GetControlValue(CachedControlIndex, ERigControlValueType::Minimum).Get<int32>();
+					Maximum = Hierarchy->GetControlValue(CachedControlIndex, ERigControlValueType::Maximum).Get<int32>();
 				}
 			}
 			default:
@@ -104,7 +104,7 @@ FRigUnit_GetControlInteger_Execute()
 FRigUnit_GetControlVector2D_Execute()
 {
     DECLARE_SCOPE_HIERARCHICAL_COUNTER_RIGUNIT()
-	const FRigControlHierarchy* Hierarchy = Context.GetControls();
+	const URigHierarchy* Hierarchy = Context.Hierarchy;
 	if (Hierarchy)
 	{
 		switch (Context.State)
@@ -115,15 +115,15 @@ FRigUnit_GetControlVector2D_Execute()
 			}
 			case EControlRigState::Update:
 			{
-				if (!CachedControlIndex.UpdateCache(Control, Hierarchy))
+				if (!CachedControlIndex.UpdateCache(FRigElementKey(Control, ERigElementType::Control), Hierarchy))
 				{
 					UE_CONTROLRIG_RIGUNIT_REPORT_WARNING(TEXT("Control '%s' is not valid."), *Control.ToString());
 				}
 				else
 				{
-					Vector = Hierarchy->GetValue(CachedControlIndex).Get<FVector2D>();
-					Minimum = Hierarchy->GetValue(CachedControlIndex, ERigControlValueType::Minimum).Get<FVector2D>();
-					Maximum = Hierarchy->GetValue(CachedControlIndex, ERigControlValueType::Maximum).Get<FVector2D>();
+					Vector = Hierarchy->GetControlValue(CachedControlIndex).Get<FVector2D>();
+					Minimum = Hierarchy->GetControlValue(CachedControlIndex, ERigControlValueType::Minimum).Get<FVector2D>();
+					Maximum = Hierarchy->GetControlValue(CachedControlIndex, ERigControlValueType::Maximum).Get<FVector2D>();
 				}
 			}
 			default:
@@ -137,7 +137,7 @@ FRigUnit_GetControlVector2D_Execute()
 FRigUnit_GetControlVector_Execute()
 {
     DECLARE_SCOPE_HIERARCHICAL_COUNTER_RIGUNIT()
-	const FRigControlHierarchy* Hierarchy = Context.GetControls();
+	const URigHierarchy* Hierarchy = Context.Hierarchy;
 	if (Hierarchy)
 	{
 		switch (Context.State)
@@ -148,7 +148,7 @@ FRigUnit_GetControlVector_Execute()
 			}
 			case EControlRigState::Update:
 			{
-				if (!CachedControlIndex.UpdateCache(Control, Hierarchy))
+				if (!CachedControlIndex.UpdateCache(FRigElementKey(Control, ERigElementType::Control), Hierarchy))
 				{
 					UE_CONTROLRIG_RIGUNIT_REPORT_WARNING(TEXT("Control '%s' is not valid."), *Control.ToString());
 				}
@@ -173,17 +173,19 @@ FRigUnit_GetControlVector_Execute()
 						}
 					}
 
-					if((*Hierarchy)[CachedControlIndex].ControlType == ERigControlType::Position)
+					const ERigControlType ControlType = Hierarchy->GetChecked<FRigControlElement>(CachedControlIndex)->Settings.ControlType;
+					
+					if(ControlType == ERigControlType::Position)
 					{
 						Vector = Transform.GetLocation();
 					}
-					else if((*Hierarchy)[CachedControlIndex].ControlType == ERigControlType::Scale)
+					else if(ControlType == ERigControlType::Scale)
 					{
 						Vector = Transform.GetScale3D();
 					}
 
-					Minimum = Hierarchy->GetValue(CachedControlIndex, ERigControlValueType::Minimum).Get<FVector>();
-					Maximum = Hierarchy->GetValue(CachedControlIndex, ERigControlValueType::Maximum).Get<FVector>();
+					Minimum = Hierarchy->GetControlValue(CachedControlIndex, ERigControlValueType::Minimum).Get<FVector>();
+					Maximum = Hierarchy->GetControlValue(CachedControlIndex, ERigControlValueType::Maximum).Get<FVector>();
 				}
 			}
 			default:
@@ -197,7 +199,7 @@ FRigUnit_GetControlVector_Execute()
 FRigUnit_GetControlRotator_Execute()
 {
     DECLARE_SCOPE_HIERARCHICAL_COUNTER_RIGUNIT()
-	const FRigControlHierarchy* Hierarchy = Context.GetControls();
+	const URigHierarchy* Hierarchy = Context.Hierarchy;
 	if (Hierarchy)
 	{
 		switch (Context.State)
@@ -208,7 +210,7 @@ FRigUnit_GetControlRotator_Execute()
 			}
 			case EControlRigState::Update:
 			{
-				if (!CachedControlIndex.UpdateCache(Control, Hierarchy))
+				if (!CachedControlIndex.UpdateCache(FRigElementKey(Control, ERigElementType::Control), Hierarchy))
 				{
 					UE_CONTROLRIG_RIGUNIT_REPORT_WARNING(TEXT("Control '%s' is not valid."), *Control.ToString());
 				}
@@ -234,8 +236,8 @@ FRigUnit_GetControlRotator_Execute()
 					}
 
 					Rotator = Transform.GetRotation().Rotator();
-					Minimum = Hierarchy->GetValue(CachedControlIndex, ERigControlValueType::Minimum).Get<FRotator>();
-					Maximum = Hierarchy->GetValue(CachedControlIndex, ERigControlValueType::Maximum).Get<FRotator>();
+					Minimum = Hierarchy->GetControlValue(CachedControlIndex, ERigControlValueType::Minimum).Get<FRotator>();
+					Maximum = Hierarchy->GetControlValue(CachedControlIndex, ERigControlValueType::Maximum).Get<FRotator>();
 				}
 			}
 			default:
@@ -249,7 +251,7 @@ FRigUnit_GetControlRotator_Execute()
 FRigUnit_GetControlTransform_Execute()
 {
     DECLARE_SCOPE_HIERARCHICAL_COUNTER_RIGUNIT()
-	const FRigControlHierarchy* Hierarchy = Context.GetControls();
+	const URigHierarchy* Hierarchy = Context.Hierarchy;
 	if (Hierarchy)
 	{
 		switch (Context.State)
@@ -260,7 +262,7 @@ FRigUnit_GetControlTransform_Execute()
 			}
 			case EControlRigState::Update:
 			{
-				if (!CachedControlIndex.UpdateCache(Control, Hierarchy))
+				if (!CachedControlIndex.UpdateCache(FRigElementKey(Control, ERigElementType::Control), Hierarchy))
 				{
 					UE_CONTROLRIG_RIGUNIT_REPORT_WARNING(TEXT("Control '%s' is not valid."), *Control.ToString());
 				}
@@ -283,8 +285,8 @@ FRigUnit_GetControlTransform_Execute()
 							break;
 						}
 					}
-					Minimum = Hierarchy->GetValue(CachedControlIndex, ERigControlValueType::Minimum).Get<FTransform>();
-					Maximum = Hierarchy->GetValue(CachedControlIndex, ERigControlValueType::Maximum).Get<FTransform>();
+					Minimum = Hierarchy->GetControlValue(CachedControlIndex, ERigControlValueType::Minimum).Get<FTransform>();
+					Maximum = Hierarchy->GetControlValue(CachedControlIndex, ERigControlValueType::Maximum).Get<FTransform>();
 				}
 			}
 			default:
@@ -300,9 +302,21 @@ FRigUnit_GetControlTransform_Execute()
 
 IMPLEMENT_RIGUNIT_AUTOMATION_TEST(FRigUnit_GetControlTransform)
 {
-	ControlHierarchy.Add(TEXT("Root"), ERigControlType::Transform, NAME_None, NAME_None, FTransform(FVector(1.f, 0.f, 0.f)));
-	ControlHierarchy.Add(TEXT("ControlA"), ERigControlType::Transform, TEXT("Root"), NAME_None, FTransform(FVector(1.f, 2.f, 3.f)));
-	ControlHierarchy.Initialize();
+	const FRigElementKey Root = Controller->AddControl(
+		TEXT("Root"),
+		FRigElementKey(),
+		FRigControlSettings(),
+		FRigControlValue::Make(FTransform(FVector(1.f, 0.f, 0.f))),
+		FTransform::Identity,
+		FTransform::Identity);
+
+	const FRigElementKey ControlA = Controller->AddControl(
+	    TEXT("ControlA"),
+	    Root,
+	    FRigControlSettings(),
+	    FRigControlValue::Make(FTransform(FVector(1.f, 2.f, 3.f))),
+	    FTransform::Identity,
+	    FTransform::Identity);
 
 	Unit.Control = TEXT("Unknown");
 	Unit.Space = EBoneGetterSetterMode::GlobalSpace;
@@ -318,7 +332,7 @@ IMPLEMENT_RIGUNIT_AUTOMATION_TEST(FRigUnit_GetControlTransform)
 	AddErrorIfFalse(Unit.Transform.GetTranslation().Equals(FVector(1.f, 0.f, 0.f)), TEXT("unexpected global transform (1)"));
 	Unit.Space = EBoneGetterSetterMode::LocalSpace;
 	InitAndExecute();
-	AddErrorIfFalse(Unit.Transform.GetTranslation().Equals(FVector(0.f, 0.f, 0.f)), TEXT("unexpected local transform (1)"));
+	AddErrorIfFalse(Unit.Transform.GetTranslation().Equals(FVector(1.f, 0.f, 0.f)), TEXT("unexpected local transform (1)"));
 
 	Unit.Control = TEXT("ControlA");
 	Unit.Space = EBoneGetterSetterMode::GlobalSpace;
@@ -326,7 +340,7 @@ IMPLEMENT_RIGUNIT_AUTOMATION_TEST(FRigUnit_GetControlTransform)
 	AddErrorIfFalse(Unit.Transform.GetTranslation().Equals(FVector(2.f, 2.f, 3.f)), TEXT("unexpected global transform (2)"));
 	Unit.Space = EBoneGetterSetterMode::LocalSpace;
 	InitAndExecute();
-	AddErrorIfFalse(Unit.Transform.GetTranslation().Equals(FVector(0.f, 0.f, 0.f)), TEXT("unexpected local transform (2)"));
+	AddErrorIfFalse(Unit.Transform.GetTranslation().Equals(FVector(1.f, 2.f, 3.f)), TEXT("unexpected local transform (2)"));
 
 	return true;
 }
diff --git a/Engine/Plugins/Experimental/ControlRig/Source/ControlRig/Private/Units/Hierarchy/RigUnit_GetCurveValue.cpp b/Engine/Plugins/Experimental/ControlRig/Source/ControlRig/Private/Units/Hierarchy/RigUnit_GetCurveValue.cpp
--- a/Engine/Plugins/Experimental/ControlRig/Source/ControlRig/Private/Units/Hierarchy/RigUnit_GetCurveValue.cpp
+++ b/Engine/Plugins/Experimental/ControlRig/Source/ControlRig/Private/Units/Hierarchy/RigUnit_GetCurveValue.cpp
@@ -6,8 +6,8 @@
 FRigUnit_GetCurveValue_Execute()
 {
     DECLARE_SCOPE_HIERARCHICAL_COUNTER_RIGUNIT()
-	const FRigCurveContainer* CurveContainer = Context.GetCurves();
-	if (CurveContainer)
+	URigHierarchy* Hierarchy = Context.Hierarchy;
+	if (Hierarchy)
 	{
 		switch (Context.State)
 		{
@@ -17,9 +17,9 @@ FRigUnit_GetCurveValue_Execute()
 			}
 			case EControlRigState::Update:
 			{
-				if(CachedCurveIndex.UpdateCache(Curve, CurveContainer))
+				if(CachedCurveIndex.UpdateCache(FRigElementKey(Curve, ERigElementType::Curve), Hierarchy))
 				{
-					Value = CurveContainer->GetValue(CachedCurveIndex);
+					Value = Hierarchy->GetCurveValue(CachedCurveIndex);
 				}
 			}
 			default:
diff --git a/Engine/Plugins/Experimental/ControlRig/Source/ControlRig/Private/Units/Hierarchy/RigUnit_GetInitialBoneTransform.cpp b/Engine/Plugins/Experimental/ControlRig/Source/ControlRig/Private/Units/Hierarchy/RigUnit_GetInitialBoneTransform.cpp
--- a/Engine/Plugins/Experimental/ControlRig/Source/ControlRig/Private/Units/Hierarchy/RigUnit_GetInitialBoneTransform.cpp
+++ b/Engine/Plugins/Experimental/ControlRig/Source/ControlRig/Private/Units/Hierarchy/RigUnit_GetInitialBoneTransform.cpp
@@ -6,7 +6,7 @@
 FRigUnit_GetInitialBoneTransform_Execute()
 {
     DECLARE_SCOPE_HIERARCHICAL_COUNTER_RIGUNIT()
-	const FRigBoneHierarchy* Hierarchy = Context.GetBones();
+	const URigHierarchy* Hierarchy = Context.Hierarchy;
 	if (Hierarchy)
 	{
 		switch (Context.State)
@@ -17,7 +17,7 @@ FRigUnit_GetInitialBoneTransform_Execute()
 			}
 			case EControlRigState::Update:
 			{
-				if (!CachedBone.UpdateCache(Bone, Hierarchy))
+				if (!CachedBone.UpdateCache(FRigElementKey(Bone, ERigElementType::Bone), Hierarchy))
 				{
 					UE_CONTROLRIG_RIGUNIT_REPORT_WARNING(TEXT("Bone '%s' is not valid."), *Bone.ToString());
 				}
@@ -55,9 +55,8 @@ FRigUnit_GetInitialBoneTransform_Execute()
 
 IMPLEMENT_RIGUNIT_AUTOMATION_TEST(FRigUnit_GetInitialBoneTransform)
 {
-	BoneHierarchy.Add(TEXT("Root"), NAME_None, ERigBoneType::User, FTransform(FVector(1.f, 0.f, 0.f)));
-	BoneHierarchy.Add(TEXT("BoneA"), TEXT("Root"), ERigBoneType::User, FTransform(FVector(1.f, 2.f, 3.f)));
-	BoneHierarchy.Initialize();
+	const FRigElementKey Root = Controller->AddBone(TEXT("Root"), FRigElementKey(), FTransform(FVector(1.f, 0.f, 0.f)), true, ERigBoneType::User);
+	const FRigElementKey BoneA = Controller->AddBone(TEXT("BoneA"), Root, FTransform(FVector(1.f, 2.f, 3.f)), true, ERigBoneType::User);
 
 	Unit.Bone = TEXT("Root");
 	Unit.Space = EBoneGetterSetterMode::GlobalSpace;
diff --git a/Engine/Plugins/Experimental/ControlRig/Source/ControlRig/Private/Units/Hierarchy/RigUnit_GetInitialBoneTransform.h b/Engine/Plugins/Experimental/ControlRig/Source/ControlRig/Private/Units/Hierarchy/RigUnit_GetInitialBoneTransform.h
--- a/Engine/Plugins/Experimental/ControlRig/Source/ControlRig/Private/Units/Hierarchy/RigUnit_GetInitialBoneTransform.h
+++ b/Engine/Plugins/Experimental/ControlRig/Source/ControlRig/Private/Units/Hierarchy/RigUnit_GetInitialBoneTransform.h
@@ -22,14 +22,9 @@ struct CONTROLRIG_API FRigUnit_GetInitialBoneTransform : public FRigUnit
 	{
 		if (InPinPath.StartsWith(TEXT("Transform")) && Space == EBoneGetterSetterMode::LocalSpace)
 		{
-			if (const FRigHierarchyContainer* Container = (const FRigHierarchyContainer*)InUserContext)
+			if (const URigHierarchy* Hierarchy = (const URigHierarchy*)InUserContext)
 			{
-				int32 BoneIndex = Container->BoneHierarchy.GetIndex(Bone);
-				if (BoneIndex != INDEX_NONE)
-				{
-					return Container->BoneHierarchy[BoneIndex].GetParentElementKey();
-				}
-
+				return Hierarchy->GetFirstParent(FRigElementKey(Bone, ERigElementType::Bone));
 			}
 		}
 		return FRigElementKey();
diff --git a/Engine/Plugins/Experimental/ControlRig/Source/ControlRig/Private/Units/Hierarchy/RigUnit_GetRelativeBoneTransform.cpp b/Engine/Plugins/Experimental/ControlRig/Source/ControlRig/Private/Units/Hierarchy/RigUnit_GetRelativeBoneTransform.cpp
--- a/Engine/Plugins/Experimental/ControlRig/Source/ControlRig/Private/Units/Hierarchy/RigUnit_GetRelativeBoneTransform.cpp
+++ b/Engine/Plugins/Experimental/ControlRig/Source/ControlRig/Private/Units/Hierarchy/RigUnit_GetRelativeBoneTransform.cpp
@@ -6,7 +6,7 @@
 FRigUnit_GetRelativeBoneTransform_Execute()
 {
     DECLARE_SCOPE_HIERARCHICAL_COUNTER_RIGUNIT()
-	const FRigBoneHierarchy* Hierarchy = Context.GetBones();
+	const URigHierarchy* Hierarchy = Context.Hierarchy;
 	if (Hierarchy)
 	{
 		switch (Context.State)
@@ -18,11 +18,11 @@ FRigUnit_GetRelativeBoneTransform_Execute()
 			}
 			case EControlRigState::Update:
 			{
-				if (!CachedBone.UpdateCache(Bone, Hierarchy))
+				if (!CachedBone.UpdateCache(FRigElementKey(Bone, ERigElementType::Bone), Hierarchy))
 				{
 					UE_CONTROLRIG_RIGUNIT_REPORT_WARNING(TEXT("Bone '%s' is not valid."), *Bone.ToString());
 				}
-				else if (!CachedSpace.UpdateCache(Space, Hierarchy))
+				else if (!CachedSpace.UpdateCache(FRigElementKey(Space, ERigElementType::Bone), Hierarchy))
 				{
 					UE_CONTROLRIG_RIGUNIT_REPORT_WARNING(TEXT("Space '%s' is not valid."), *Space.ToString());
 				}
@@ -46,10 +46,9 @@ FRigUnit_GetRelativeBoneTransform_Execute()
 
 IMPLEMENT_RIGUNIT_AUTOMATION_TEST(FRigUnit_GetRelativeBoneTransform)
 {
-	BoneHierarchy.Add(TEXT("Root"), NAME_None, ERigBoneType::User, FTransform(FVector(1.f, 0.f, 0.f)));
-	BoneHierarchy.Add(TEXT("BoneA"), TEXT("Root"), ERigBoneType::User, FTransform(FVector(1.f, 2.f, 3.f)));
-	BoneHierarchy.Add(TEXT("BoneB"), TEXT("Root"), ERigBoneType::User, FTransform(FVector(-4.f, 0.f, 0.f)));
-	BoneHierarchy.Initialize();
+	const FRigElementKey Root = Controller->AddBone(TEXT("Root"), FRigElementKey(), FTransform(FVector(1.f, 0.f, 0.f)), true, ERigBoneType::User);
+	const FRigElementKey BoneA = Controller->AddBone(TEXT("BoneA"), Root, FTransform(FVector(1.f, 2.f, 3.f)), true, ERigBoneType::User);
+	const FRigElementKey BoneB = Controller->AddBone(TEXT("BoneB"), Root, FTransform(FVector(-4.f, 0.f, 0.f)), true, ERigBoneType::User);
 
 	Unit.Bone = TEXT("Unknown");
 	Unit.Space = TEXT("Root");
diff --git a/Engine/Plugins/Experimental/ControlRig/Source/ControlRig/Private/Units/Hierarchy/RigUnit_GetSpaceTransform.cpp b/Engine/Plugins/Experimental/ControlRig/Source/ControlRig/Private/Units/Hierarchy/RigUnit_GetSpaceTransform.cpp
--- a/Engine/Plugins/Experimental/ControlRig/Source/ControlRig/Private/Units/Hierarchy/RigUnit_GetSpaceTransform.cpp
+++ b/Engine/Plugins/Experimental/ControlRig/Source/ControlRig/Private/Units/Hierarchy/RigUnit_GetSpaceTransform.cpp
@@ -6,7 +6,7 @@
 FRigUnit_GetSpaceTransform_Execute()
 {
     DECLARE_SCOPE_HIERARCHICAL_COUNTER_RIGUNIT()
-	const FRigSpaceHierarchy* Hierarchy = Context.GetSpaces();
+	const URigHierarchy* Hierarchy = Context.Hierarchy;
 	if (Hierarchy)
 	{
 		switch (Context.State)
@@ -17,7 +17,7 @@ FRigUnit_GetSpaceTransform_Execute()
 			}
 			case EControlRigState::Update:
 			{
-				if (CachedSpaceIndex.UpdateCache(Space, Hierarchy))
+				if (CachedSpaceIndex.UpdateCache(FRigElementKey(Space, ERigElementType::Space), Hierarchy))
 				{
 					switch (SpaceType)
 					{
diff --git a/Engine/Plugins/Experimental/ControlRig/Source/ControlRig/Private/Units/Hierarchy/RigUnit_GetTransform.cpp b/Engine/Plugins/Experimental/ControlRig/Source/ControlRig/Private/Units/Hierarchy/RigUnit_GetTransform.cpp
--- a/Engine/Plugins/Experimental/ControlRig/Source/ControlRig/Private/Units/Hierarchy/RigUnit_GetTransform.cpp
+++ b/Engine/Plugins/Experimental/ControlRig/Source/ControlRig/Private/Units/Hierarchy/RigUnit_GetTransform.cpp
@@ -14,7 +14,7 @@ FRigUnit_GetTransform_Execute()
 {
     DECLARE_SCOPE_HIERARCHICAL_COUNTER_RIGUNIT()
 
-	if (const FRigHierarchyContainer* Hierarchy = Context.Hierarchy)
+	if (URigHierarchy* Hierarchy = Context.Hierarchy)
 	{
 		switch (Context.State)
 		{
@@ -41,7 +41,7 @@ FRigUnit_GetTransform_Execute()
 							}
 							case EBoneGetterSetterMode::LocalSpace:
 							{
-								Transform = Hierarchy->GetInitialTransform(CachedIndex);
+								Transform = Hierarchy->GetInitialLocalTransform(CachedIndex);
 								break;
 							}
 							default:
diff --git a/Engine/Plugins/Experimental/ControlRig/Source/ControlRig/Private/Units/Hierarchy/RigUnit_GetTransform.h b/Engine/Plugins/Experimental/ControlRig/Source/ControlRig/Private/Units/Hierarchy/RigUnit_GetTransform.h
--- a/Engine/Plugins/Experimental/ControlRig/Source/ControlRig/Private/Units/Hierarchy/RigUnit_GetTransform.h
+++ b/Engine/Plugins/Experimental/ControlRig/Source/ControlRig/Private/Units/Hierarchy/RigUnit_GetTransform.h
@@ -27,9 +27,9 @@ struct CONTROLRIG_API FRigUnit_GetTransform : public FRigUnit
 	{
 		if(Space == EBoneGetterSetterMode::LocalSpace)
 		{
-			if (const FRigHierarchyContainer* Container = (const FRigHierarchyContainer*)InUserContext)
+			if (const URigHierarchy* Hierarchy = (const URigHierarchy*)InUserContext)
 			{
-				return Container->GetParentKey(Item);
+				return Hierarchy->GetFirstParent(Item);
 			}
 		}
 		return FRigElementKey();
diff --git a/Engine/Plugins/Experimental/ControlRig/Source/ControlRig/Private/Units/Hierarchy/RigUnit_ParentSwitchConstraint.cpp b/Engine/Plugins/Experimental/ControlRig/Source/ControlRig/Private/Units/Hierarchy/RigUnit_ParentSwitchConstraint.cpp
--- a/Engine/Plugins/Experimental/ControlRig/Source/ControlRig/Private/Units/Hierarchy/RigUnit_ParentSwitchConstraint.cpp
+++ b/Engine/Plugins/Experimental/ControlRig/Source/ControlRig/Private/Units/Hierarchy/RigUnit_ParentSwitchConstraint.cpp
@@ -10,7 +10,7 @@ FRigUnit_ParentSwitchConstraint_Execute()
 	Switched = false;
 	Transform = FTransform::Identity;
 
-	if (FRigHierarchyContainer* Hierarchy = ExecuteContext.Hierarchy)
+	if (URigHierarchy* Hierarchy = ExecuteContext.Hierarchy)
 	{
 		switch (Context.State)
 		{
diff --git a/Engine/Plugins/Experimental/ControlRig/Source/ControlRig/Private/Units/Hierarchy/RigUnit_PropagateTransform.cpp b/Engine/Plugins/Experimental/ControlRig/Source/ControlRig/Private/Units/Hierarchy/RigUnit_PropagateTransform.cpp
--- a/Engine/Plugins/Experimental/ControlRig/Source/ControlRig/Private/Units/Hierarchy/RigUnit_PropagateTransform.cpp
+++ b/Engine/Plugins/Experimental/ControlRig/Source/ControlRig/Private/Units/Hierarchy/RigUnit_PropagateTransform.cpp
@@ -15,8 +15,11 @@ FRigUnit_PropagateTransform_Execute()
     	return;
     }
 
-    if (FRigHierarchyContainer* Hierarchy = ExecuteContext.Hierarchy)
+    if (URigHierarchy* Hierarchy = ExecuteContext.Hierarchy)
 	{
+    	/*
+    	 * This node doesn't do anything anymore now that the hierarchy is lazy
+    	 *
 		switch (Context.State)
 		{
 			case EControlRigState::Init:
@@ -59,6 +62,6 @@ FRigUnit_PropagateTransform_Execute()
 				break;
 			}
 		}
+		*/
 	}
-
 }
\ No newline at end of file
diff --git a/Engine/Plugins/Experimental/ControlRig/Source/ControlRig/Private/Units/Hierarchy/RigUnit_PropagateTransform.h b/Engine/Plugins/Experimental/ControlRig/Source/ControlRig/Private/Units/Hierarchy/RigUnit_PropagateTransform.h
--- a/Engine/Plugins/Experimental/ControlRig/Source/ControlRig/Private/Units/Hierarchy/RigUnit_PropagateTransform.h
+++ b/Engine/Plugins/Experimental/ControlRig/Source/ControlRig/Private/Units/Hierarchy/RigUnit_PropagateTransform.h
@@ -9,7 +9,7 @@
  * Propagate Transform can be used to force a recalculation of a bone's global transform
  * from its local - as well as propagating that change onto the children.
  */
-USTRUCT(meta=(DisplayName="Propagate Transform", Category="Hierarchy", DocumentationPolicy = "Strict", Keywords = "PropagateToChildren,RecomputeGlobal,RecalculateGlobal", Varying))
+USTRUCT(meta=(DisplayName="Propagate Transform", Category="Hierarchy", DocumentationPolicy = "Strict", Keywords = "PropagateToChildren,RecomputeGlobal,RecalculateGlobal", Varying, Deprecated = "5.0"))
 struct CONTROLRIG_API FRigUnit_PropagateTransform : public FRigUnitMutable
 {
 	GENERATED_BODY()
diff --git a/Engine/Plugins/Experimental/ControlRig/Source/ControlRig/Private/Units/Hierarchy/RigUnit_SendEvent.cpp b/Engine/Plugins/Experimental/ControlRig/Source/ControlRig/Private/Units/Hierarchy/RigUnit_SendEvent.cpp
--- a/Engine/Plugins/Experimental/ControlRig/Source/ControlRig/Private/Units/Hierarchy/RigUnit_SendEvent.cpp
+++ b/Engine/Plugins/Experimental/ControlRig/Source/ControlRig/Private/Units/Hierarchy/RigUnit_SendEvent.cpp
@@ -17,7 +17,7 @@ FRigUnit_SendEvent_Execute()
 		return;
 	}
 
-	FRigHierarchyContainer* Hierarchy = ExecuteContext.Hierarchy;
+	URigHierarchy* Hierarchy = ExecuteContext.Hierarchy;
 	if (Hierarchy)
 	{
 		switch (Context.State)
diff --git a/Engine/Plugins/Experimental/ControlRig/Source/ControlRig/Private/Units/Hierarchy/RigUnit_SetBoneInitialTransform.cpp b/Engine/Plugins/Experimental/ControlRig/Source/ControlRig/Private/Units/Hierarchy/RigUnit_SetBoneInitialTransform.cpp
--- a/Engine/Plugins/Experimental/ControlRig/Source/ControlRig/Private/Units/Hierarchy/RigUnit_SetBoneInitialTransform.cpp
+++ b/Engine/Plugins/Experimental/ControlRig/Source/ControlRig/Private/Units/Hierarchy/RigUnit_SetBoneInitialTransform.cpp
@@ -8,11 +8,9 @@ FRigUnit_SetBoneInitialTransform_Execute()
 {
     DECLARE_SCOPE_HIERARCHICAL_COUNTER_RIGUNIT()
 
-	FRigBoneHierarchy* Hierarchy = ExecuteContext.GetBones();
+	URigHierarchy* Hierarchy = ExecuteContext.Hierarchy;
 	if (Hierarchy)
 	{
-		FRigBoneHierarchy& HierarchyRef = *Hierarchy;
-
 		switch (Context.State)
 		{
 			case EControlRigState::Init:
@@ -22,7 +20,8 @@ FRigUnit_SetBoneInitialTransform_Execute()
 			}
 			case EControlRigState::Update:
 			{
-				if (!CachedBone.UpdateCache(Bone, Hierarchy))
+				const FRigElementKey Key(Bone, ERigElementType::Bone);
+				if (!CachedBone.UpdateCache(Key, Hierarchy))
 				{
 					UE_CONTROLRIG_RIGUNIT_REPORT_WARNING(TEXT("Bone '%s' is not valid."), *Bone.ToString());
 					return;
@@ -30,11 +29,11 @@ FRigUnit_SetBoneInitialTransform_Execute()
 
 				if (Space == EBoneGetterSetterMode::LocalSpace)
 				{
-					HierarchyRef.SetInitialLocalTransform(CachedBone, Transform);
+					Hierarchy->SetInitialLocalTransform(CachedBone, Transform);
 				}
 				else
 				{
-					HierarchyRef.SetInitialGlobalTransform(CachedBone, Transform);
+					Hierarchy->SetInitialGlobalTransform(CachedBone, Transform);
 				}
 			}
 			default:
diff --git a/Engine/Plugins/Experimental/ControlRig/Source/ControlRig/Private/Units/Hierarchy/RigUnit_SetBoneRotation.cpp b/Engine/Plugins/Experimental/ControlRig/Source/ControlRig/Private/Units/Hierarchy/RigUnit_SetBoneRotation.cpp
--- a/Engine/Plugins/Experimental/ControlRig/Source/ControlRig/Private/Units/Hierarchy/RigUnit_SetBoneRotation.cpp
+++ b/Engine/Plugins/Experimental/ControlRig/Source/ControlRig/Private/Units/Hierarchy/RigUnit_SetBoneRotation.cpp
@@ -6,7 +6,7 @@
 FRigUnit_SetBoneRotation_Execute()
 {
     DECLARE_SCOPE_HIERARCHICAL_COUNTER_RIGUNIT()
-	FRigBoneHierarchy* Hierarchy = ExecuteContext.GetBones();
+	URigHierarchy* Hierarchy = ExecuteContext.Hierarchy;
 	if (Hierarchy)
 	{
 		switch (Context.State)
@@ -18,7 +18,8 @@ FRigUnit_SetBoneRotation_Execute()
 			}
 			case EControlRigState::Update:
 			{
-				if (!CachedBone.UpdateCache(Bone, Hierarchy))
+				const FRigElementKey Key(Bone, ERigElementType::Bone);
+				if (!CachedBone.UpdateCache(Key, Hierarchy))
 				{
 					UE_CONTROLRIG_RIGUNIT_REPORT_WARNING(TEXT("Bone '%s' is not valid."), *Bone.ToString());
 				}
@@ -80,59 +81,58 @@ FRigUnit_SetBoneRotation_Execute()
 
 IMPLEMENT_RIGUNIT_AUTOMATION_TEST(FRigUnit_SetBoneRotation)
 {
-	BoneHierarchy.Add(TEXT("Root"), NAME_None, ERigBoneType::User, FTransform(FQuat(FVector(-1.f, 0.f, 0.f), 0.1f)));
-	BoneHierarchy.Add(TEXT("BoneA"), TEXT("Root"), ERigBoneType::User, FTransform(FQuat(FVector(-1.f, 0.f, 0.f), 0.5f)));
-	BoneHierarchy.Add(TEXT("BoneB"), TEXT("BoneA"), ERigBoneType::User, FTransform(FQuat(FVector(-1.f, 0.f, 0.f), 0.7f)));
-	BoneHierarchy.Initialize();
+	const FRigElementKey Root = Controller->AddBone(TEXT("Root"), FRigElementKey(), FTransform(FQuat(FVector(-1.f, 0.f, 0.f), 0.1f)), true, ERigBoneType::User);
+	const FRigElementKey BoneA = Controller->AddBone(TEXT("BoneA"), Root, FTransform(FQuat(FVector(-1.f, 0.f, 0.f), 0.5f)), true, ERigBoneType::User);
+	const FRigElementKey BoneB = Controller->AddBone(TEXT("BoneB"), BoneA, FTransform(FQuat(FVector(-1.f, 0.f, 0.f), 0.7f)), true, ERigBoneType::User);
 
-	Unit.ExecuteContext.Hierarchy = &HierarchyContainer;
+	Unit.ExecuteContext.Hierarchy = Hierarchy;
 
-	BoneHierarchy.ResetTransforms();
+	Hierarchy->ResetPoseToInitial(ERigElementType::Bone);
 	Unit.Bone = TEXT("Root");
 	Unit.Space = EBoneGetterSetterMode::GlobalSpace;
 	Unit.Rotation = FQuat(FVector(-1.f, 0.f, 0.f), 0.25f);
 	Unit.bPropagateToChildren = false;
 	InitAndExecute();
-	AddErrorIfFalse(FMath::IsNearlyEqual(BoneHierarchy.GetGlobalTransform(0).GetRotation().GetAngle(), 0.25f, 0.001f), FString::Printf(TEXT("unexpected angle %.04f"), BoneHierarchy.GetGlobalTransform(0).GetRotation().GetAngle()));
-	AddErrorIfFalse(FMath::IsNearlyEqual(BoneHierarchy.GetGlobalTransform(1).GetRotation().GetAngle(), 0.5f, 0.001f), FString::Printf(TEXT("unexpected angle %.04f"), BoneHierarchy.GetGlobalTransform(1).GetRotation().GetAngle()));
-	AddErrorIfFalse(FMath::IsNearlyEqual(BoneHierarchy.GetGlobalTransform(2).GetRotation().GetAngle(), 0.7f, 0.001f), FString::Printf(TEXT("unexpected angle %.04f"), BoneHierarchy.GetGlobalTransform(2).GetRotation().GetAngle()));
+	AddErrorIfFalse(FMath::IsNearlyEqual(Hierarchy->GetGlobalTransform(0).GetRotation().GetAngle(), 0.25f, 0.001f), FString::Printf(TEXT("unexpected angle %.04f"), Hierarchy->GetGlobalTransform(0).GetRotation().GetAngle()));
+	AddErrorIfFalse(FMath::IsNearlyEqual(Hierarchy->GetGlobalTransform(1).GetRotation().GetAngle(), 0.5f, 0.001f), FString::Printf(TEXT("unexpected angle %.04f"), Hierarchy->GetGlobalTransform(1).GetRotation().GetAngle()));
+	AddErrorIfFalse(FMath::IsNearlyEqual(Hierarchy->GetGlobalTransform(2).GetRotation().GetAngle(), 0.7f, 0.001f), FString::Printf(TEXT("unexpected angle %.04f"), Hierarchy->GetGlobalTransform(2).GetRotation().GetAngle()));
 
-	BoneHierarchy.ResetTransforms();
+	Hierarchy->ResetPoseToInitial(ERigElementType::Bone);
 	Unit.Space = EBoneGetterSetterMode::LocalSpace;
 	InitAndExecute();
-	AddErrorIfFalse(FMath::IsNearlyEqual(BoneHierarchy.GetGlobalTransform(0).GetRotation().GetAngle(), 0.25f, 0.001f), FString::Printf(TEXT("unexpected angle %.04f"), BoneHierarchy.GetGlobalTransform(0).GetRotation().GetAngle()));
-	AddErrorIfFalse(FMath::IsNearlyEqual(BoneHierarchy.GetGlobalTransform(1).GetRotation().GetAngle(), 0.5f, 0.001f), FString::Printf(TEXT("unexpected angle %.04f"), BoneHierarchy.GetGlobalTransform(1).GetRotation().GetAngle()));
-	AddErrorIfFalse(FMath::IsNearlyEqual(BoneHierarchy.GetGlobalTransform(2).GetRotation().GetAngle(), 0.7f, 0.001f), FString::Printf(TEXT("unexpected angle %.04f"), BoneHierarchy.GetGlobalTransform(2).GetRotation().GetAngle()));
+	AddErrorIfFalse(FMath::IsNearlyEqual(Hierarchy->GetGlobalTransform(0).GetRotation().GetAngle(), 0.25f, 0.001f), FString::Printf(TEXT("unexpected angle %.04f"), Hierarchy->GetGlobalTransform(0).GetRotation().GetAngle()));
+	AddErrorIfFalse(FMath::IsNearlyEqual(Hierarchy->GetGlobalTransform(1).GetRotation().GetAngle(), 0.5f, 0.001f), FString::Printf(TEXT("unexpected angle %.04f"), Hierarchy->GetGlobalTransform(1).GetRotation().GetAngle()));
+	AddErrorIfFalse(FMath::IsNearlyEqual(Hierarchy->GetGlobalTransform(2).GetRotation().GetAngle(), 0.7f, 0.001f), FString::Printf(TEXT("unexpected angle %.04f"), Hierarchy->GetGlobalTransform(2).GetRotation().GetAngle()));
 
-	BoneHierarchy.ResetTransforms();
+	Hierarchy->ResetPoseToInitial(ERigElementType::Bone);
 	Unit.bPropagateToChildren = true;
 	InitAndExecute();
-	AddErrorIfFalse(FMath::IsNearlyEqual(BoneHierarchy.GetGlobalTransform(0).GetRotation().GetAngle(), 0.25f, 0.001f), FString::Printf(TEXT("unexpected angle %.04f"), BoneHierarchy.GetGlobalTransform(0).GetRotation().GetAngle()));
-	AddErrorIfFalse(FMath::IsNearlyEqual(BoneHierarchy.GetGlobalTransform(1).GetRotation().GetAngle(), 0.65f, 0.001f), FString::Printf(TEXT("unexpected angle %.04f"), BoneHierarchy.GetGlobalTransform(1).GetRotation().GetAngle()));
-	AddErrorIfFalse(FMath::IsNearlyEqual(BoneHierarchy.GetGlobalTransform(2).GetRotation().GetAngle(), 0.85f, 0.001f), FString::Printf(TEXT("unexpected angle %.04f"), BoneHierarchy.GetGlobalTransform(2).GetRotation().GetAngle()));
+	AddErrorIfFalse(FMath::IsNearlyEqual(Hierarchy->GetGlobalTransform(0).GetRotation().GetAngle(), 0.25f, 0.001f), FString::Printf(TEXT("unexpected angle %.04f"), Hierarchy->GetGlobalTransform(0).GetRotation().GetAngle()));
+	AddErrorIfFalse(FMath::IsNearlyEqual(Hierarchy->GetGlobalTransform(1).GetRotation().GetAngle(), 0.65f, 0.001f), FString::Printf(TEXT("unexpected angle %.04f"), Hierarchy->GetGlobalTransform(1).GetRotation().GetAngle()));
+	AddErrorIfFalse(FMath::IsNearlyEqual(Hierarchy->GetGlobalTransform(2).GetRotation().GetAngle(), 0.85f, 0.001f), FString::Printf(TEXT("unexpected angle %.04f"), Hierarchy->GetGlobalTransform(2).GetRotation().GetAngle()));
 
-	BoneHierarchy.ResetTransforms();
+	Hierarchy->ResetPoseToInitial(ERigElementType::Bone);
 	Unit.Bone = TEXT("BoneA");
 	Unit.Space = EBoneGetterSetterMode::GlobalSpace;
 	Unit.bPropagateToChildren = false;
 	InitAndExecute();
-	AddErrorIfFalse(FMath::IsNearlyEqual(BoneHierarchy.GetGlobalTransform(0).GetRotation().GetAngle(), 0.1f, 0.001f), FString::Printf(TEXT("unexpected angle %.04f"), BoneHierarchy.GetGlobalTransform(0).GetRotation().GetAngle()));
-	AddErrorIfFalse(FMath::IsNearlyEqual(BoneHierarchy.GetGlobalTransform(1).GetRotation().GetAngle(), 0.25f, 0.001f), FString::Printf(TEXT("unexpected angle %.04f"), BoneHierarchy.GetGlobalTransform(1).GetRotation().GetAngle()));
-	AddErrorIfFalse(FMath::IsNearlyEqual(BoneHierarchy.GetGlobalTransform(2).GetRotation().GetAngle(), 0.7f, 0.001f), FString::Printf(TEXT("unexpected angle %.04f"), BoneHierarchy.GetGlobalTransform(2).GetRotation().GetAngle()));
+	AddErrorIfFalse(FMath::IsNearlyEqual(Hierarchy->GetGlobalTransform(0).GetRotation().GetAngle(), 0.1f, 0.001f), FString::Printf(TEXT("unexpected angle %.04f"), Hierarchy->GetGlobalTransform(0).GetRotation().GetAngle()));
+	AddErrorIfFalse(FMath::IsNearlyEqual(Hierarchy->GetGlobalTransform(1).GetRotation().GetAngle(), 0.25f, 0.001f), FString::Printf(TEXT("unexpected angle %.04f"), Hierarchy->GetGlobalTransform(1).GetRotation().GetAngle()));
+	AddErrorIfFalse(FMath::IsNearlyEqual(Hierarchy->GetGlobalTransform(2).GetRotation().GetAngle(), 0.7f, 0.001f), FString::Printf(TEXT("unexpected angle %.04f"), Hierarchy->GetGlobalTransform(2).GetRotation().GetAngle()));
 
-	BoneHierarchy.ResetTransforms();
+	Hierarchy->ResetPoseToInitial(ERigElementType::Bone);
 	Unit.Space = EBoneGetterSetterMode::LocalSpace;
 	InitAndExecute();
-	AddErrorIfFalse(FMath::IsNearlyEqual(BoneHierarchy.GetGlobalTransform(0).GetRotation().GetAngle(), 0.1f, 0.001f), FString::Printf(TEXT("unexpected angle %.04f"), BoneHierarchy.GetGlobalTransform(0).GetRotation().GetAngle()));
-	AddErrorIfFalse(FMath::IsNearlyEqual(BoneHierarchy.GetGlobalTransform(1).GetRotation().GetAngle(), 0.35f, 0.001f), FString::Printf(TEXT("unexpected angle %.04f"), BoneHierarchy.GetGlobalTransform(1).GetRotation().GetAngle()));
-	AddErrorIfFalse(FMath::IsNearlyEqual(BoneHierarchy.GetGlobalTransform(2).GetRotation().GetAngle(), 0.7f, 0.001f), FString::Printf(TEXT("unexpected angle %.04f"), BoneHierarchy.GetGlobalTransform(2).GetRotation().GetAngle()));
+	AddErrorIfFalse(FMath::IsNearlyEqual(Hierarchy->GetGlobalTransform(0).GetRotation().GetAngle(), 0.1f, 0.001f), FString::Printf(TEXT("unexpected angle %.04f"), Hierarchy->GetGlobalTransform(0).GetRotation().GetAngle()));
+	AddErrorIfFalse(FMath::IsNearlyEqual(Hierarchy->GetGlobalTransform(1).GetRotation().GetAngle(), 0.35f, 0.001f), FString::Printf(TEXT("unexpected angle %.04f"), Hierarchy->GetGlobalTransform(1).GetRotation().GetAngle()));
+	AddErrorIfFalse(FMath::IsNearlyEqual(Hierarchy->GetGlobalTransform(2).GetRotation().GetAngle(), 0.7f, 0.001f), FString::Printf(TEXT("unexpected angle %.04f"), Hierarchy->GetGlobalTransform(2).GetRotation().GetAngle()));
 
-	BoneHierarchy.ResetTransforms();
+	Hierarchy->ResetPoseToInitial(ERigElementType::Bone);
 	Unit.bPropagateToChildren = true;
 	InitAndExecute();
-	AddErrorIfFalse(FMath::IsNearlyEqual(BoneHierarchy.GetGlobalTransform(0).GetRotation().GetAngle(), 0.1f, 0.001f), FString::Printf(TEXT("unexpected angle %.04f"), BoneHierarchy.GetGlobalTransform(0).GetRotation().GetAngle()));
-	AddErrorIfFalse(FMath::IsNearlyEqual(BoneHierarchy.GetGlobalTransform(1).GetRotation().GetAngle(), 0.35f, 0.001f), FString::Printf(TEXT("unexpected angle %.04f"), BoneHierarchy.GetGlobalTransform(1).GetRotation().GetAngle()));
-	AddErrorIfFalse(FMath::IsNearlyEqual(BoneHierarchy.GetGlobalTransform(2).GetRotation().GetAngle(), 0.55f, 0.001f), FString::Printf(TEXT("unexpected angle %.04f"), BoneHierarchy.GetGlobalTransform(2).GetRotation().GetAngle()));
+	AddErrorIfFalse(FMath::IsNearlyEqual(Hierarchy->GetGlobalTransform(0).GetRotation().GetAngle(), 0.1f, 0.001f), FString::Printf(TEXT("unexpected angle %.04f"), Hierarchy->GetGlobalTransform(0).GetRotation().GetAngle()));
+	AddErrorIfFalse(FMath::IsNearlyEqual(Hierarchy->GetGlobalTransform(1).GetRotation().GetAngle(), 0.35f, 0.001f), FString::Printf(TEXT("unexpected angle %.04f"), Hierarchy->GetGlobalTransform(1).GetRotation().GetAngle()));
+	AddErrorIfFalse(FMath::IsNearlyEqual(Hierarchy->GetGlobalTransform(2).GetRotation().GetAngle(), 0.55f, 0.001f), FString::Printf(TEXT("unexpected angle %.04f"), Hierarchy->GetGlobalTransform(2).GetRotation().GetAngle()));
 
 	return true;
 }
diff --git a/Engine/Plugins/Experimental/ControlRig/Source/ControlRig/Private/Units/Hierarchy/RigUnit_SetBoneTransform.cpp b/Engine/Plugins/Experimental/ControlRig/Source/ControlRig/Private/Units/Hierarchy/RigUnit_SetBoneTransform.cpp
--- a/Engine/Plugins/Experimental/ControlRig/Source/ControlRig/Private/Units/Hierarchy/RigUnit_SetBoneTransform.cpp
+++ b/Engine/Plugins/Experimental/ControlRig/Source/ControlRig/Private/Units/Hierarchy/RigUnit_SetBoneTransform.cpp
@@ -7,7 +7,7 @@ FRigUnit_SetBoneTransform_Execute()
 {
     DECLARE_SCOPE_HIERARCHICAL_COUNTER_RIGUNIT()
 
-	FRigBoneHierarchy* Hierarchy = ExecuteContext.GetBones();
+	URigHierarchy* Hierarchy = ExecuteContext.Hierarchy;
 	if (Hierarchy)
 	{
 		switch (Context.State)
@@ -19,7 +19,8 @@ FRigUnit_SetBoneTransform_Execute()
 			}
 			case EControlRigState::Update:
 			{
-				if (!CachedBone.UpdateCache(Bone, Hierarchy))
+				const FRigElementKey Key(Bone, ERigElementType::Bone);
+				if (!CachedBone.UpdateCache(Key, Hierarchy))
 				{
 					UE_CONTROLRIG_RIGUNIT_REPORT_WARNING(TEXT("Bone '%s' is not valid."), *Bone.ToString());
 				}
@@ -77,58 +78,57 @@ FRigUnit_SetBoneTransform_Execute()
 
 IMPLEMENT_RIGUNIT_AUTOMATION_TEST(FRigUnit_SetBoneTransform)
 {
-	BoneHierarchy.Add(TEXT("Root"), NAME_None, ERigBoneType::User, FTransform(FVector(1.f, 0.f, 0.f)));
-	BoneHierarchy.Add(TEXT("BoneA"), TEXT("Root"), ERigBoneType::User, FTransform(FVector(1.f, 2.f, 3.f)));
-	BoneHierarchy.Add(TEXT("BoneB"), TEXT("BoneA"), ERigBoneType::User, FTransform(FVector(1.f, 5.f, 3.f)));
-	BoneHierarchy.Initialize();
-	Unit.ExecuteContext.Hierarchy = &HierarchyContainer;
+	const FRigElementKey Root = Controller->AddBone(TEXT("Root"), FRigElementKey(), FTransform(FVector(1.f, 0.f, 0.f)), true, ERigBoneType::User);
+	const FRigElementKey BoneA = Controller->AddBone(TEXT("BoneA"), Root, FTransform(FVector(1.f, 2.f, 3.f)), true, ERigBoneType::User);
+	const FRigElementKey BoneB = Controller->AddBone(TEXT("BoneB"), BoneA, FTransform(FVector(1.f, 5.f, 3.f)), true, ERigBoneType::User);
+	Unit.ExecuteContext.Hierarchy = Hierarchy;
 
-	BoneHierarchy.ResetTransforms();
+	Hierarchy->ResetPoseToInitial(ERigElementType::Bone);
 	Unit.Bone = TEXT("Root");
 	Unit.Space = EBoneGetterSetterMode::GlobalSpace;
 	Unit.Transform = FTransform(FVector(0.f, 0.f, 7.f));
 	Unit.bPropagateToChildren = false;
 	InitAndExecute();
-	AddErrorIfFalse(BoneHierarchy.GetGlobalTransform(0).GetTranslation().Equals(FVector(0.f, 0.f, 7.f)), TEXT("unexpected transform"));
-	AddErrorIfFalse(BoneHierarchy.GetGlobalTransform(1).GetTranslation().Equals(FVector(1.f, 2.f, 3.f)), TEXT("unexpected transform"));
-	AddErrorIfFalse(BoneHierarchy.GetGlobalTransform(2).GetTranslation().Equals(FVector(1.f, 5.f, 3.f)), TEXT("unexpected transform"));
+	AddErrorIfFalse(Hierarchy->GetGlobalTransform(0).GetTranslation().Equals(FVector(0.f, 0.f, 7.f)), TEXT("unexpected transform"));
+	AddErrorIfFalse(Hierarchy->GetGlobalTransform(1).GetTranslation().Equals(FVector(1.f, 2.f, 3.f)), TEXT("unexpected transform"));
+	AddErrorIfFalse(Hierarchy->GetGlobalTransform(2).GetTranslation().Equals(FVector(1.f, 5.f, 3.f)), TEXT("unexpected transform"));
 
-	BoneHierarchy.ResetTransforms();
+	Hierarchy->ResetPoseToInitial(ERigElementType::Bone);
 	Unit.Space = EBoneGetterSetterMode::LocalSpace;
 	InitAndExecute();
-	AddErrorIfFalse(BoneHierarchy.GetGlobalTransform(0).GetTranslation().Equals(FVector(0.f, 0.f, 7.f)), TEXT("unexpected transform"));
-	AddErrorIfFalse(BoneHierarchy.GetGlobalTransform(1).GetTranslation().Equals(FVector(1.f, 2.f, 3.f)), TEXT("unexpected transform"));
-	AddErrorIfFalse(BoneHierarchy.GetGlobalTransform(2).GetTranslation().Equals(FVector(1.f, 5.f, 3.f)), TEXT("unexpected transform"));
+	AddErrorIfFalse(Hierarchy->GetGlobalTransform(0).GetTranslation().Equals(FVector(0.f, 0.f, 7.f)), TEXT("unexpected transform"));
+	AddErrorIfFalse(Hierarchy->GetGlobalTransform(1).GetTranslation().Equals(FVector(1.f, 2.f, 3.f)), TEXT("unexpected transform"));
+	AddErrorIfFalse(Hierarchy->GetGlobalTransform(2).GetTranslation().Equals(FVector(1.f, 5.f, 3.f)), TEXT("unexpected transform"));
 
-	BoneHierarchy.ResetTransforms();
+	Hierarchy->ResetPoseToInitial(ERigElementType::Bone);
 	Unit.bPropagateToChildren = true;
 	InitAndExecute();
-	AddErrorIfFalse(BoneHierarchy.GetGlobalTransform(0).GetTranslation().Equals(FVector(0.f, 0.f, 7.f)), TEXT("unexpected transform"));
-	AddErrorIfFalse(BoneHierarchy.GetGlobalTransform(1).GetTranslation().Equals(FVector(0.f, 2.f, 10.f)), TEXT("unexpected transform"));
-	AddErrorIfFalse(BoneHierarchy.GetGlobalTransform(2).GetTranslation().Equals(FVector(0.f, 5.f, 10.f)), TEXT("unexpected transform"));
+	AddErrorIfFalse(Hierarchy->GetGlobalTransform(0).GetTranslation().Equals(FVector(0.f, 0.f, 7.f)), TEXT("unexpected transform"));
+	AddErrorIfFalse(Hierarchy->GetGlobalTransform(1).GetTranslation().Equals(FVector(0.f, 2.f, 10.f)), TEXT("unexpected transform"));
+	AddErrorIfFalse(Hierarchy->GetGlobalTransform(2).GetTranslation().Equals(FVector(0.f, 5.f, 10.f)), TEXT("unexpected transform"));
 
-	BoneHierarchy.ResetTransforms();
+	Hierarchy->ResetPoseToInitial(ERigElementType::Bone);
 	Unit.Bone = TEXT("BoneA");
 	Unit.Space = EBoneGetterSetterMode::GlobalSpace;
 	Unit.bPropagateToChildren = false;
 	InitAndExecute();
-	AddErrorIfFalse(BoneHierarchy.GetGlobalTransform(0).GetTranslation().Equals(FVector(1.f, 0.f, 0.f)), TEXT("unexpected transform"));
-	AddErrorIfFalse(BoneHierarchy.GetGlobalTransform(1).GetTranslation().Equals(FVector(0.f, 0.f, 7.f)), TEXT("unexpected transform"));
-	AddErrorIfFalse(BoneHierarchy.GetGlobalTransform(2).GetTranslation().Equals(FVector(1.f, 5.f, 3.f)), TEXT("unexpected transform"));
+	AddErrorIfFalse(Hierarchy->GetGlobalTransform(0).GetTranslation().Equals(FVector(1.f, 0.f, 0.f)), TEXT("unexpected transform"));
+	AddErrorIfFalse(Hierarchy->GetGlobalTransform(1).GetTranslation().Equals(FVector(0.f, 0.f, 7.f)), TEXT("unexpected transform"));
+	AddErrorIfFalse(Hierarchy->GetGlobalTransform(2).GetTranslation().Equals(FVector(1.f, 5.f, 3.f)), TEXT("unexpected transform"));
 
-	BoneHierarchy.ResetTransforms();
+	Hierarchy->ResetPoseToInitial(ERigElementType::Bone);
 	Unit.Space = EBoneGetterSetterMode::LocalSpace;
 	InitAndExecute();
-	AddErrorIfFalse(BoneHierarchy.GetGlobalTransform(0).GetTranslation().Equals(FVector(1.f, 0.f, 0.f)), TEXT("unexpected transform"));
-	AddErrorIfFalse(BoneHierarchy.GetGlobalTransform(1).GetTranslation().Equals(FVector(1.f, 0.f, 7.f)), TEXT("unexpected transform"));
-	AddErrorIfFalse(BoneHierarchy.GetGlobalTransform(2).GetTranslation().Equals(FVector(1.f, 5.f, 3.f)), TEXT("unexpected transform"));
+	AddErrorIfFalse(Hierarchy->GetGlobalTransform(0).GetTranslation().Equals(FVector(1.f, 0.f, 0.f)), TEXT("unexpected transform"));
+	AddErrorIfFalse(Hierarchy->GetGlobalTransform(1).GetTranslation().Equals(FVector(1.f, 0.f, 7.f)), TEXT("unexpected transform"));
+	AddErrorIfFalse(Hierarchy->GetGlobalTransform(2).GetTranslation().Equals(FVector(1.f, 5.f, 3.f)), TEXT("unexpected transform"));
 
-	BoneHierarchy.ResetTransforms();
+	Hierarchy->ResetPoseToInitial(ERigElementType::Bone);
 	Unit.bPropagateToChildren = true;
 	InitAndExecute();
-	AddErrorIfFalse(BoneHierarchy.GetGlobalTransform(0).GetTranslation().Equals(FVector(1.f, 0.f, 0.f)), TEXT("unexpected transform"));
-	AddErrorIfFalse(BoneHierarchy.GetGlobalTransform(1).GetTranslation().Equals(FVector(1.f, 0.f, 7.f)), TEXT("unexpected transform"));
-	AddErrorIfFalse(BoneHierarchy.GetGlobalTransform(2).GetTranslation().Equals(FVector(1.f, 3.f, 7.f)), TEXT("unexpected transform"));
+	AddErrorIfFalse(Hierarchy->GetGlobalTransform(0).GetTranslation().Equals(FVector(1.f, 0.f, 0.f)), TEXT("unexpected transform"));
+	AddErrorIfFalse(Hierarchy->GetGlobalTransform(1).GetTranslation().Equals(FVector(1.f, 0.f, 7.f)), TEXT("unexpected transform"));
+	AddErrorIfFalse(Hierarchy->GetGlobalTransform(2).GetTranslation().Equals(FVector(1.f, 3.f, 7.f)), TEXT("unexpected transform"));
 
 	return true;
 }
diff --git a/Engine/Plugins/Experimental/ControlRig/Source/ControlRig/Private/Units/Hierarchy/RigUnit_SetBoneTranslation.cpp b/Engine/Plugins/Experimental/ControlRig/Source/ControlRig/Private/Units/Hierarchy/RigUnit_SetBoneTranslation.cpp
--- a/Engine/Plugins/Experimental/ControlRig/Source/ControlRig/Private/Units/Hierarchy/RigUnit_SetBoneTranslation.cpp
+++ b/Engine/Plugins/Experimental/ControlRig/Source/ControlRig/Private/Units/Hierarchy/RigUnit_SetBoneTranslation.cpp
@@ -6,7 +6,7 @@
 FRigUnit_SetBoneTranslation_Execute()
 {
     DECLARE_SCOPE_HIERARCHICAL_COUNTER_RIGUNIT()
-	FRigBoneHierarchy* Hierarchy = ExecuteContext.GetBones();
+	URigHierarchy* Hierarchy = ExecuteContext.Hierarchy;
 	if (Hierarchy)
 	{
 		switch (Context.State)
@@ -18,7 +18,8 @@ FRigUnit_SetBoneTranslation_Execute()
 			}
 			case EControlRigState::Update:
 			{
-				if (!CachedBone.UpdateCache(Bone, Hierarchy))
+				const FRigElementKey Key(Bone, ERigElementType::Bone); 
+				if (!CachedBone.UpdateCache(Key, Hierarchy))
 				{
 					UE_CONTROLRIG_RIGUNIT_REPORT_WARNING(TEXT("Bone '%s' is not valid."), *Bone.ToString());
 				}
diff --git a/Engine/Plugins/Experimental/ControlRig/Source/ControlRig/Private/Units/Hierarchy/RigUnit_SetControlColor.cpp b/Engine/Plugins/Experimental/ControlRig/Source/ControlRig/Private/Units/Hierarchy/RigUnit_SetControlColor.cpp
--- a/Engine/Plugins/Experimental/ControlRig/Source/ControlRig/Private/Units/Hierarchy/RigUnit_SetControlColor.cpp
+++ b/Engine/Plugins/Experimental/ControlRig/Source/ControlRig/Private/Units/Hierarchy/RigUnit_SetControlColor.cpp
@@ -6,7 +6,7 @@
 FRigUnit_SetControlColor_Execute()
 {
     DECLARE_SCOPE_HIERARCHICAL_COUNTER_RIGUNIT()
-	FRigControlHierarchy* Hierarchy = ExecuteContext.GetControls();
+	URigHierarchy* Hierarchy = ExecuteContext.Hierarchy;
 	if (Hierarchy)
 	{
 		switch (Context.State)
@@ -18,9 +18,11 @@ FRigUnit_SetControlColor_Execute()
 			}
 			case EControlRigState::Update:
 			{
-				if (CachedControlIndex.UpdateCache(Control, Hierarchy))
+				if (CachedControlIndex.UpdateCache(FRigElementKey(Control, ERigElementType::Control), Hierarchy))
 				{
-					(*Hierarchy)[CachedControlIndex].GizmoColor = Color;
+					FRigControlElement* ControlElement = Hierarchy->GetChecked<FRigControlElement>(CachedControlIndex);
+					ControlElement->Settings.GizmoColor = Color;
+					Hierarchy->Notify(ERigHierarchyNotification::ControlSettingChanged, ControlElement);
 				}
 				break;
 			}
diff --git a/Engine/Plugins/Experimental/ControlRig/Source/ControlRig/Private/Units/Hierarchy/RigUnit_SetControlOffset.cpp b/Engine/Plugins/Experimental/ControlRig/Source/ControlRig/Private/Units/Hierarchy/RigUnit_SetControlOffset.cpp
--- a/Engine/Plugins/Experimental/ControlRig/Source/ControlRig/Private/Units/Hierarchy/RigUnit_SetControlOffset.cpp
+++ b/Engine/Plugins/Experimental/ControlRig/Source/ControlRig/Private/Units/Hierarchy/RigUnit_SetControlOffset.cpp
@@ -1,6 +1,8 @@
 // Copyright Epic Games, Inc. All Rights Reserved.
 
 #include "RigUnit_SetControlOffset.h"
+
+#include "GeometryCollection/GeometryCollectionAlgo.h"
 #include "Units/RigUnitContext.h"
 #include "Math/ControlRigMathLibrary.h"
 #include "Units/Execution/RigUnit_PrepareForExecution.h"
@@ -9,7 +11,7 @@ FRigUnit_SetControlOffset_Execute()
 {
     DECLARE_SCOPE_HIERARCHICAL_COUNTER_RIGUNIT()
 
-	FRigControlHierarchy* Hierarchy = ExecuteContext.GetControls();
+	URigHierarchy* Hierarchy = ExecuteContext.Hierarchy;
 	if (Hierarchy)
 	{
 		switch (Context.State)
@@ -21,18 +23,23 @@ FRigUnit_SetControlOffset_Execute()
 			}
 			case EControlRigState::Update:
 			{
-				if (!CachedControlIndex.UpdateCache(Control, Hierarchy))
+				if (!CachedControlIndex.UpdateCache(FRigElementKey(Control, ERigElementType::Control), Hierarchy))
 				{
 					UE_CONTROLRIG_RIGUNIT_REPORT_WARNING(TEXT("Control '%s' is not valid."), *Control.ToString());
 					return;
 				}
 
-				FTransform OffsetToSet = Offset;
+				FRigControlElement* ControlElement = Hierarchy->Get<FRigControlElement>(CachedControlIndex);
 				if (Space == EBoneGetterSetterMode::GlobalSpace)
 				{
-					OffsetToSet = OffsetToSet.GetRelativeTransform(Hierarchy->GetParentInitialTransform(CachedControlIndex, false));
+					Hierarchy->SetControlOffsetTransform(ControlElement, Offset, ERigTransformType::CurrentGlobal, true, false);
+					Hierarchy->SetControlOffsetTransform(ControlElement, Offset, ERigTransformType::InitialGlobal, true, false);
+				}
+				else
+				{
+					Hierarchy->SetControlOffsetTransform(ControlElement, Offset, ERigTransformType::CurrentLocal, true, false);
+					Hierarchy->SetControlOffsetTransform(ControlElement, Offset, ERigTransformType::InitialLocal, true, false);
 				}
-				Hierarchy->SetControlOffset(CachedControlIndex, OffsetToSet);
 			}
 			default:
 			{
diff --git a/Engine/Plugins/Experimental/ControlRig/Source/ControlRig/Private/Units/Hierarchy/RigUnit_SetControlTransform.cpp b/Engine/Plugins/Experimental/ControlRig/Source/ControlRig/Private/Units/Hierarchy/RigUnit_SetControlTransform.cpp
--- a/Engine/Plugins/Experimental/ControlRig/Source/ControlRig/Private/Units/Hierarchy/RigUnit_SetControlTransform.cpp
+++ b/Engine/Plugins/Experimental/ControlRig/Source/ControlRig/Private/Units/Hierarchy/RigUnit_SetControlTransform.cpp
@@ -7,7 +7,7 @@
 FRigUnit_SetControlBool_Execute()
 {
     DECLARE_SCOPE_HIERARCHICAL_COUNTER_RIGUNIT()
-	FRigControlHierarchy* Hierarchy = ExecuteContext.GetControls();
+	URigHierarchy* Hierarchy = ExecuteContext.Hierarchy;
 	if (Hierarchy)
 	{
 		switch (Context.State)
@@ -19,9 +19,10 @@ FRigUnit_SetControlBool_Execute()
 			}
 			case EControlRigState::Update:
 			{
-				if (CachedControlIndex.UpdateCache(Control, Hierarchy))
+				const FRigElementKey Key(Control, ERigElementType::Control);
+				if (CachedControlIndex.UpdateCache(Key, Hierarchy))
 				{
-					Hierarchy->SetValue(CachedControlIndex, FRigControlValue::Make<bool>(BoolValue));
+					Hierarchy->SetControlValue(CachedControlIndex, FRigControlValue::Make<bool>(BoolValue));
 				}
 				break;
 			}
@@ -36,7 +37,7 @@ FRigUnit_SetControlBool_Execute()
 FRigUnit_SetMultiControlBool_Execute()
 {
 	DECLARE_SCOPE_HIERARCHICAL_COUNTER_RIGUNIT()
-	FRigControlHierarchy* Hierarchy = ExecuteContext.GetControls();
+	URigHierarchy* Hierarchy = ExecuteContext.Hierarchy;
 	if (Hierarchy)
 	{
 		CachedControlIndices.SetNum(Entries.Num());
@@ -64,7 +65,7 @@ FRigUnit_SetMultiControlBool_Execute()
 FRigUnit_SetControlFloat_Execute()
 {
     DECLARE_SCOPE_HIERARCHICAL_COUNTER_RIGUNIT()
-	FRigControlHierarchy* Hierarchy = ExecuteContext.GetControls();
+	URigHierarchy* Hierarchy = ExecuteContext.Hierarchy;
 	if (Hierarchy)
 	{
 		switch (Context.State)
@@ -76,16 +77,17 @@ FRigUnit_SetControlFloat_Execute()
 			}
 			case EControlRigState::Update:
 			{
-				if (CachedControlIndex.UpdateCache(Control, Hierarchy))
+				const FRigElementKey Key(Control, ERigElementType::Control);
+				if (CachedControlIndex.UpdateCache(Key, Hierarchy))
 				{
 					if(FMath::IsNearlyEqual(Weight, 1.f))
 					{
-						Hierarchy->SetValue(CachedControlIndex, FRigControlValue::Make<float>(FloatValue));
+						Hierarchy->SetControlValue(CachedControlIndex, FRigControlValue::Make<float>(FloatValue));
 					}
 					else
 					{
-						float PreviousValue = Hierarchy->GetValue(CachedControlIndex).Get<float>();
-						Hierarchy->SetValue(CachedControlIndex, FRigControlValue::Make<float>(FMath::Lerp<float>(PreviousValue, FloatValue, FMath::Clamp<float>(Weight, 0.f, 1.f))));
+						float PreviousValue = Hierarchy->GetControlValue(CachedControlIndex).Get<float>();
+						Hierarchy->SetControlValue(CachedControlIndex, FRigControlValue::Make<float>(FMath::Lerp<float>(PreviousValue, FloatValue, FMath::Clamp<float>(Weight, 0.f, 1.f))));
 					}
 				}
 				break;
@@ -101,7 +103,7 @@ FRigUnit_SetControlFloat_Execute()
 FRigUnit_SetMultiControlFloat_Execute()
 {
 	DECLARE_SCOPE_HIERARCHICAL_COUNTER_RIGUNIT()
-	FRigControlHierarchy* Hierarchy = ExecuteContext.GetControls();
+	URigHierarchy* Hierarchy = ExecuteContext.Hierarchy;
 	if (Hierarchy)
 	{ 
 		CachedControlIndices.SetNum(Entries.Num());
@@ -129,7 +131,7 @@ FRigUnit_SetMultiControlFloat_Execute()
 FRigUnit_SetControlInteger_Execute()
 {
     DECLARE_SCOPE_HIERARCHICAL_COUNTER_RIGUNIT()
-	FRigControlHierarchy* Hierarchy = ExecuteContext.GetControls();
+	URigHierarchy* Hierarchy = ExecuteContext.Hierarchy;
 	if (Hierarchy)
 	{
 		switch (Context.State)
@@ -141,16 +143,17 @@ FRigUnit_SetControlInteger_Execute()
 			}
 			case EControlRigState::Update:
 			{
-				if (CachedControlIndex.UpdateCache(Control, Hierarchy))
+				const FRigElementKey Key(Control, ERigElementType::Control);
+				if (CachedControlIndex.UpdateCache(Key, Hierarchy))
 				{
 					if(FMath::IsNearlyEqual(Weight, 1.f))
 					{
-						Hierarchy->SetValue(CachedControlIndex, FRigControlValue::Make<int32>(IntegerValue));
+						Hierarchy->SetControlValue(CachedControlIndex, FRigControlValue::Make<int32>(IntegerValue));
 					}
 					else
 					{
-						int32 PreviousValue = Hierarchy->GetValue(CachedControlIndex).Get<int32>();
-						Hierarchy->SetValue(CachedControlIndex, FRigControlValue::Make<int32>((int32)FMath::Lerp<float>((float)PreviousValue, (float)IntegerValue, FMath::Clamp<float>(Weight, 0.f, 1.f))));
+						int32 PreviousValue = Hierarchy->GetControlValue(CachedControlIndex).Get<int32>();
+						Hierarchy->SetControlValue(CachedControlIndex, FRigControlValue::Make<int32>((int32)FMath::Lerp<float>((float)PreviousValue, (float)IntegerValue, FMath::Clamp<float>(Weight, 0.f, 1.f))));
 					}
 				}
 				break;
@@ -166,7 +169,7 @@ FRigUnit_SetControlInteger_Execute()
 FRigUnit_SetMultiControlInteger_Execute()
 {
 	DECLARE_SCOPE_HIERARCHICAL_COUNTER_RIGUNIT()
-	FRigControlHierarchy* Hierarchy = ExecuteContext.GetControls();
+	URigHierarchy* Hierarchy = ExecuteContext.Hierarchy;
 	if (Hierarchy)
 	{
 		CachedControlIndices.SetNum(Entries.Num());
@@ -194,7 +197,7 @@ FRigUnit_SetMultiControlInteger_Execute()
 FRigUnit_SetControlVector2D_Execute()
 {
     DECLARE_SCOPE_HIERARCHICAL_COUNTER_RIGUNIT()
-	FRigControlHierarchy* Hierarchy = ExecuteContext.GetControls();
+	URigHierarchy* Hierarchy = ExecuteContext.Hierarchy;
 	if (Hierarchy)
 	{
 		switch (Context.State)
@@ -206,16 +209,17 @@ FRigUnit_SetControlVector2D_Execute()
 			}
 			case EControlRigState::Update:
 			{
-				if (CachedControlIndex.UpdateCache(Control, Hierarchy))
+				const FRigElementKey Key(Control, ERigElementType::Control);
+				if (CachedControlIndex.UpdateCache(Key, Hierarchy))
 				{
 					if(FMath::IsNearlyEqual(Weight, 1.f))
 					{
-						Hierarchy->SetValue(CachedControlIndex, FRigControlValue::Make<FVector2D>(Vector));
+						Hierarchy->SetControlValue(CachedControlIndex, FRigControlValue::Make<FVector2D>(Vector));
 					}
 					else
 					{
-						FVector2D PreviousValue = Hierarchy->GetValue(CachedControlIndex).Get<FVector2D>();
-						Hierarchy->SetValue(CachedControlIndex, FRigControlValue::Make<FVector2D>(FMath::Lerp<FVector2D>(PreviousValue, Vector, FMath::Clamp<float>(Weight, 0.f, 1.f))));
+						FVector2D PreviousValue = Hierarchy->GetControlValue(CachedControlIndex).Get<FVector2D>();
+						Hierarchy->SetControlValue(CachedControlIndex, FRigControlValue::Make<FVector2D>(FMath::Lerp<FVector2D>(PreviousValue, Vector, FMath::Clamp<float>(Weight, 0.f, 1.f))));
 					}
 				}
 				break;
@@ -231,7 +235,7 @@ FRigUnit_SetControlVector2D_Execute()
 FRigUnit_SetMultiControlVector2D_Execute()
 {
 	DECLARE_SCOPE_HIERARCHICAL_COUNTER_RIGUNIT()
-	FRigControlHierarchy* Hierarchy = ExecuteContext.GetControls();
+	URigHierarchy* Hierarchy = ExecuteContext.Hierarchy;
 	if (Hierarchy)
 	{
 		CachedControlIndices.SetNum(Entries.Num());
@@ -259,7 +263,7 @@ FRigUnit_SetMultiControlVector2D_Execute()
 FRigUnit_SetControlVector_Execute()
 {
     DECLARE_SCOPE_HIERARCHICAL_COUNTER_RIGUNIT()
-	FRigControlHierarchy* Hierarchy = ExecuteContext.GetControls();
+	URigHierarchy* Hierarchy = ExecuteContext.Hierarchy;
 	if (Hierarchy)
 	{
 		switch (Context.State)
@@ -271,15 +275,18 @@ FRigUnit_SetControlVector_Execute()
 			}
 			case EControlRigState::Update:
 			{
-				if (CachedControlIndex.UpdateCache(Control, Hierarchy))
+				const FRigElementKey Key(Control, ERigElementType::Control);
+				if (CachedControlIndex.UpdateCache(Key, Hierarchy))
 				{
+					const ERigControlType ControlType = Hierarchy->GetChecked<FRigControlElement>(CachedControlIndex)->Settings.ControlType;
+					
 					FTransform Transform = FTransform::Identity;
 					if (Space == EBoneGetterSetterMode::GlobalSpace)
 					{
 						Transform = Hierarchy->GetGlobalTransform(CachedControlIndex);
 					}
 
-					if ((*Hierarchy)[CachedControlIndex].ControlType == ERigControlType::Position)
+					if (ControlType == ERigControlType::Position)
 					{
 						if(FMath::IsNearlyEqual(Weight, 1.f))
 						{
@@ -291,7 +298,7 @@ FRigUnit_SetControlVector_Execute()
 							Transform.SetLocation(FMath::Lerp<FVector>(PreviousValue, Vector, FMath::Clamp<float>(Weight, 0.f, 1.f)));
 						}
 					}
-					else if ((*Hierarchy)[CachedControlIndex].ControlType == ERigControlType::Scale)
+					else if (ControlType == ERigControlType::Scale)
 					{
 						if(FMath::IsNearlyEqual(Weight, 1.f))
 						{
@@ -334,7 +341,7 @@ FRigUnit_SetControlVector_Execute()
 FRigUnit_SetControlRotator_Execute()
 {
     DECLARE_SCOPE_HIERARCHICAL_COUNTER_RIGUNIT()
-	FRigControlHierarchy* Hierarchy = ExecuteContext.GetControls();
+	URigHierarchy* Hierarchy = ExecuteContext.Hierarchy;
 	if (Hierarchy)
 	{
 		switch (Context.State)
@@ -346,7 +353,8 @@ FRigUnit_SetControlRotator_Execute()
 			}
 			case EControlRigState::Update:
 			{
-				if (CachedControlIndex.UpdateCache(Control, Hierarchy))
+				const FRigElementKey Key(Control, ERigElementType::Control);
+				if (CachedControlIndex.UpdateCache(Key, Hierarchy))
 				{
 					FTransform Transform = FTransform::Identity;
 					if (Space == EBoneGetterSetterMode::GlobalSpace)
@@ -396,7 +404,7 @@ FRigUnit_SetControlRotator_Execute()
 FRigUnit_SetMultiControlRotator_Execute()
 {
 	DECLARE_SCOPE_HIERARCHICAL_COUNTER_RIGUNIT()
-	FRigControlHierarchy* Hierarchy = ExecuteContext.GetControls();
+	URigHierarchy* Hierarchy = ExecuteContext.Hierarchy;
 	if (Hierarchy)
 	{
 		CachedControlIndices.SetNum(Entries.Num());
@@ -424,7 +432,7 @@ FRigUnit_SetMultiControlRotator_Execute()
 FRigUnit_SetControlTransform_Execute()
 {
     DECLARE_SCOPE_HIERARCHICAL_COUNTER_RIGUNIT()
-	FRigControlHierarchy* Hierarchy = ExecuteContext.GetControls();
+	URigHierarchy* Hierarchy = ExecuteContext.Hierarchy;
 	if (Hierarchy)
 	{
 		switch (Context.State)
@@ -436,7 +444,8 @@ FRigUnit_SetControlTransform_Execute()
 			}
 			case EControlRigState::Update:
 			{
-				if (CachedControlIndex.UpdateCache(Control, Hierarchy))
+				const FRigElementKey Key(Control, ERigElementType::Control);
+				if (CachedControlIndex.UpdateCache(Key, Hierarchy))
 				{
 					switch (Space)
 					{
@@ -487,10 +496,12 @@ FRigUnit_SetControlTransform_Execute()
 
 IMPLEMENT_RIGUNIT_AUTOMATION_TEST(FRigUnit_SetMultiControlBool)
 {
-	ControlHierarchy.Add(TEXT("Control1"), ERigControlType::Bool, NAME_None, NAME_None, FTransform(), FRigControlValue(), NAME_None, FTransform(), FLinearColor());
-	ControlHierarchy.Add(TEXT("Control2"), ERigControlType::Bool, NAME_None, NAME_None, FTransform(), FRigControlValue(), NAME_None, FTransform(), FLinearColor());
-	ControlHierarchy.Initialize();
-	Unit.ExecuteContext.Hierarchy = &HierarchyContainer;
+	FRigControlSettings Settings;
+	Settings.ControlType = ERigControlType::Bool;
+	
+	Controller->AddControl(TEXT("Control1"), FRigElementKey(), Settings, FRigControlValue(), FTransform::Identity, FTransform::Identity);
+	Controller->AddControl(TEXT("Control2"), FRigElementKey(), Settings, FRigControlValue(), FTransform::Identity, FTransform::Identity);
+	Unit.ExecuteContext.Hierarchy = Hierarchy;
 
 	// this unit has an empty entry by default, clear it for testing purpose
 	Unit.Entries.Reset();
@@ -510,18 +521,20 @@ IMPLEMENT_RIGUNIT_AUTOMATION_TEST(FRigUnit_SetMultiControlBool)
 	Init();
 	Execute();
 
-	AddErrorIfFalse(ControlHierarchy[TEXT("Control1")].GetValue(ERigControlValueType::Current).Get<bool>() == true, TEXT("unexpected control value"));
-	AddErrorIfFalse(ControlHierarchy[TEXT("Control2")].GetValue(ERigControlValueType::Current).Get<bool>() == true, TEXT("unexpected control value")); 
+	AddErrorIfFalse(Hierarchy->GetControlValue(FRigElementKey(TEXT("Control1"), ERigElementType::Control), ERigControlValueType::Current).Get<bool>() == true, TEXT("unexpected control value"));
+	AddErrorIfFalse(Hierarchy->GetControlValue(FRigElementKey(TEXT("Control2"), ERigElementType::Control), ERigControlValueType::Current).Get<bool>() == true, TEXT("unexpected control value")); 
 
 	return true;
 }
 
 IMPLEMENT_RIGUNIT_AUTOMATION_TEST(FRigUnit_SetMultiControlFloat)
 {
-	ControlHierarchy.Add(TEXT("Control1"), ERigControlType::Float, NAME_None, NAME_None, FTransform(), FRigControlValue(), NAME_None, FTransform(), FLinearColor());
-	ControlHierarchy.Add(TEXT("Control2"), ERigControlType::Float, NAME_None, NAME_None, FTransform(), FRigControlValue(), NAME_None, FTransform(), FLinearColor());
-	ControlHierarchy.Initialize();
-	Unit.ExecuteContext.Hierarchy = &HierarchyContainer;
+	FRigControlSettings Settings;
+	Settings.ControlType = ERigControlType::Float;
+	
+	Controller->AddControl(TEXT("Control1"), FRigElementKey(), Settings, FRigControlValue(), FTransform::Identity, FTransform::Identity);
+	Controller->AddControl(TEXT("Control2"), FRigElementKey(), Settings, FRigControlValue(), FTransform::Identity, FTransform::Identity);
+	Unit.ExecuteContext.Hierarchy = Hierarchy;
 
 	// this unit has an empty entry by default, clear it for testing purpose
 	Unit.Entries.Reset();
@@ -541,8 +554,8 @@ IMPLEMENT_RIGUNIT_AUTOMATION_TEST(FRigUnit_SetMultiControlFloat)
 	Init();
 	Execute();
 
-	AddErrorIfFalse(ControlHierarchy[TEXT("Control1")].GetValue(ERigControlValueType::Current).Get<float>() == 10.0f, TEXT("unexpected control value"));
-	AddErrorIfFalse(ControlHierarchy[TEXT("Control2")].GetValue(ERigControlValueType::Current).Get<float>() == 20.0f, TEXT("unexpected control value"));
+	AddErrorIfFalse(Hierarchy->GetControlValue(FRigElementKey(TEXT("Control1"), ERigElementType::Control), ERigControlValueType::Current).Get<float>() == 10.0f, TEXT("unexpected control value"));
+	AddErrorIfFalse(Hierarchy->GetControlValue(FRigElementKey(TEXT("Control2"), ERigElementType::Control), ERigControlValueType::Current).Get<float>() == 20.0f, TEXT("unexpected control value"));
 
 	// test if the caches were updated 
 	AddErrorIfFalse(Unit.CachedControlIndices.Num() == 2, TEXT("unexpected number of cache entries"));
@@ -567,7 +580,7 @@ IMPLEMENT_RIGUNIT_AUTOMATION_TEST(FRigUnit_SetMultiControlFloat)
 
 	Execute();
 
-	AddErrorIfFalse(ControlHierarchy[TEXT("Control2")].GetValue(ERigControlValueType::Current).Get<float>() == 30.0f, TEXT("unexpected control value"));
+	AddErrorIfFalse(Hierarchy->GetControlValue(FRigElementKey(TEXT("Control2"), ERigElementType::Control), ERigControlValueType::Current).Get<float>() == 30.0f, TEXT("unexpected control value"));
 
 	AddErrorIfFalse(Unit.CachedControlIndices.Num() == 1, TEXT("unexpected number of cache entries"));
 	AddErrorIfFalse(Unit.CachedControlIndices[0].GetKey().Name == Unit.Entries[0].Control, TEXT("unexpected cached control"));
@@ -577,10 +590,12 @@ IMPLEMENT_RIGUNIT_AUTOMATION_TEST(FRigUnit_SetMultiControlFloat)
 
 IMPLEMENT_RIGUNIT_AUTOMATION_TEST(FRigUnit_SetMultiControlInteger)
 {
-	ControlHierarchy.Add(TEXT("Control1"), ERigControlType::Integer, NAME_None, NAME_None, FTransform(), FRigControlValue(), NAME_None, FTransform(), FLinearColor());
-	ControlHierarchy.Add(TEXT("Control2"), ERigControlType::Integer, NAME_None, NAME_None, FTransform(), FRigControlValue(), NAME_None, FTransform(), FLinearColor());
-	ControlHierarchy.Initialize();
-	Unit.ExecuteContext.Hierarchy = &HierarchyContainer;
+	FRigControlSettings Settings;
+	Settings.ControlType = ERigControlType::Integer;
+	
+	Controller->AddControl(TEXT("Control1"), FRigElementKey(), Settings, FRigControlValue(), FTransform::Identity, FTransform::Identity);
+	Controller->AddControl(TEXT("Control2"), FRigElementKey(), Settings, FRigControlValue(), FTransform::Identity, FTransform::Identity);
+	Unit.ExecuteContext.Hierarchy = Hierarchy;
 
 	// this unit has an empty entry by default, clear it for testing purpose
 	Unit.Entries.Reset();
@@ -600,18 +615,20 @@ IMPLEMENT_RIGUNIT_AUTOMATION_TEST(FRigUnit_SetMultiControlInteger)
 	Init();
 	Execute();
 
-	AddErrorIfFalse(ControlHierarchy[TEXT("Control1")].GetValue(ERigControlValueType::Current).Get<int32>() == 10, TEXT("unexpected control value"));
-	AddErrorIfFalse(ControlHierarchy[TEXT("Control2")].GetValue(ERigControlValueType::Current).Get<int32>() == 20, TEXT("unexpected control value"));
+	AddErrorIfFalse(Hierarchy->GetControlValue(FRigElementKey(TEXT("Control1"), ERigElementType::Control), ERigControlValueType::Current).Get<int32>() == 10, TEXT("unexpected control value"));
+	AddErrorIfFalse(Hierarchy->GetControlValue(FRigElementKey(TEXT("Control2"), ERigElementType::Control), ERigControlValueType::Current).Get<int32>() == 20, TEXT("unexpected control value"));
 
 	return true;
 }
 
 IMPLEMENT_RIGUNIT_AUTOMATION_TEST(FRigUnit_SetMultiControlVector2D)
 {
-	ControlHierarchy.Add(TEXT("Control1"), ERigControlType::Vector2D, NAME_None, NAME_None, FTransform(), FRigControlValue(), NAME_None, FTransform(), FLinearColor());
-	ControlHierarchy.Add(TEXT("Control2"), ERigControlType::Vector2D, NAME_None, NAME_None, FTransform(), FRigControlValue(), NAME_None, FTransform(), FLinearColor());
-	ControlHierarchy.Initialize();
-	Unit.ExecuteContext.Hierarchy = &HierarchyContainer;
+	FRigControlSettings Settings;
+	Settings.ControlType = ERigControlType::Vector2D;
+	
+	Controller->AddControl(TEXT("Control1"), FRigElementKey(), Settings, FRigControlValue(), FTransform::Identity, FTransform::Identity);
+	Controller->AddControl(TEXT("Control2"), FRigElementKey(), Settings, FRigControlValue(), FTransform::Identity, FTransform::Identity);
+	Unit.ExecuteContext.Hierarchy = Hierarchy;
 
 	// this unit has an empty entry by default, clear it for testing purpose
 	Unit.Entries.Reset();
@@ -631,18 +648,20 @@ IMPLEMENT_RIGUNIT_AUTOMATION_TEST(FRigUnit_SetMultiControlVector2D)
 	Init();
 	Execute();
 
-	AddErrorIfFalse(ControlHierarchy[TEXT("Control1")].GetValue(ERigControlValueType::Current).Get<FVector2D>() == FVector2D(1.f, 1.f), TEXT("unexpected control value"));
-	AddErrorIfFalse(ControlHierarchy[TEXT("Control2")].GetValue(ERigControlValueType::Current).Get<FVector2D>() == FVector2D(2.f, 2.f), TEXT("unexpected control value"));
+	AddErrorIfFalse(Hierarchy->GetControlValue(FRigElementKey(TEXT("Control1"), ERigElementType::Control), ERigControlValueType::Current).Get<FVector2D>() == FVector2D(1.f, 1.f), TEXT("unexpected control value"));
+	AddErrorIfFalse(Hierarchy->GetControlValue(FRigElementKey(TEXT("Control2"), ERigElementType::Control), ERigControlValueType::Current).Get<FVector2D>() == FVector2D(2.f, 2.f), TEXT("unexpected control value"));
 
 	return true;
 }
 
 IMPLEMENT_RIGUNIT_AUTOMATION_TEST(FRigUnit_SetMultiControlRotator)
 {
-	ControlHierarchy.Add(TEXT("Control1"), ERigControlType::Rotator, NAME_None, NAME_None, FTransform(), FRigControlValue(), NAME_None, FTransform(), FLinearColor());
-	ControlHierarchy.Add(TEXT("Control2"), ERigControlType::Rotator, NAME_None, NAME_None, FTransform(), FRigControlValue(), NAME_None, FTransform(), FLinearColor());
-	ControlHierarchy.Initialize();
-	Unit.ExecuteContext.Hierarchy = &HierarchyContainer;
+	FRigControlSettings Settings;
+	Settings.ControlType = ERigControlType::Rotator;
+	
+	Controller->AddControl(TEXT("Control1"), FRigElementKey(), Settings, FRigControlValue(), FTransform::Identity, FTransform::Identity);
+	Controller->AddControl(TEXT("Control2"), FRigElementKey(), Settings, FRigControlValue(), FTransform::Identity, FTransform::Identity);
+	Unit.ExecuteContext.Hierarchy = Hierarchy;
 
 	// this unit has an empty entry by default, clear it for testing purpose
 	Unit.Entries.Reset();
@@ -664,8 +683,8 @@ IMPLEMENT_RIGUNIT_AUTOMATION_TEST(FRigUnit_SetMultiControlRotator)
 	Init();
 	Execute(); 
 	
-	AddErrorIfFalse(ControlHierarchy[TEXT("Control1")].GetValue(ERigControlValueType::Current).Get<FRotator>() != FRotator(0.f, 0.f, 0.f), TEXT("unexpected control value"));
-	AddErrorIfFalse(ControlHierarchy[TEXT("Control2")].GetValue(ERigControlValueType::Current).Get<FRotator>() != FRotator(0.f, 0.f, 0.f), TEXT("unexpected control value"));
+	AddErrorIfFalse(Hierarchy->GetControlValue(FRigElementKey(TEXT("Control1"), ERigElementType::Control), ERigControlValueType::Current).Get<FRotator>() != FRotator(0.f, 0.f, 0.f), TEXT("unexpected control value"));
+	AddErrorIfFalse(Hierarchy->GetControlValue(FRigElementKey(TEXT("Control2"), ERigElementType::Control), ERigControlValueType::Current).Get<FRotator>() != FRotator(0.f, 0.f, 0.f), TEXT("unexpected control value"));
 
 	return true;
 }
diff --git a/Engine/Plugins/Experimental/ControlRig/Source/ControlRig/Private/Units/Hierarchy/RigUnit_SetControlVisibility.cpp b/Engine/Plugins/Experimental/ControlRig/Source/ControlRig/Private/Units/Hierarchy/RigUnit_SetControlVisibility.cpp
--- a/Engine/Plugins/Experimental/ControlRig/Source/ControlRig/Private/Units/Hierarchy/RigUnit_SetControlVisibility.cpp
+++ b/Engine/Plugins/Experimental/ControlRig/Source/ControlRig/Private/Units/Hierarchy/RigUnit_SetControlVisibility.cpp
@@ -7,8 +7,8 @@
 FRigUnit_SetControlVisibility_Execute()
 {
     DECLARE_SCOPE_HIERARCHICAL_COUNTER_RIGUNIT()
-	FRigControlHierarchy* Controls = ExecuteContext.GetControls();
-	if (Controls)
+	URigHierarchy* Hierarchy = ExecuteContext.Hierarchy;
+	if (Hierarchy)
 	{
 		switch (Context.State)
 		{
@@ -31,14 +31,14 @@ FRigUnit_SetControlVisibility_Execute()
 				}
 				else if (!Pattern.IsEmpty())
 				{
-					for (int32 Index = 0; Index < Controls->Num(); Index++)
+					Hierarchy->ForEach<FRigControlElement>([&Keys, Pattern](FRigControlElement* ControlElement) -> bool
 					{
-						const FRigControl& Control = (*Controls)[Index];
-						if (Control.Name.ToString().Contains(Pattern, ESearchCase::CaseSensitive))
+						if (ControlElement->GetName().ToString().Contains(Pattern, ESearchCase::CaseSensitive))
 						{
-							Keys.Add(Control.GetElementKey());
+							Keys.Add(ControlElement->GetKey());
 						}
-					}
+						return true;
+					});
 				}
 
 				if (CachedControlIndices.Num() != Keys.Num())
@@ -56,7 +56,7 @@ FRigUnit_SetControlVisibility_Execute()
 				{
 					if (CachedControlIndex.IsValid())
 					{
-						(*Controls)[CachedControlIndex].bGizmoVisible = bVisible;
+						Hierarchy->SetControlVisibility(CachedControlIndex, bVisible);
 					}
 				}
 				break;
diff --git a/Engine/Plugins/Experimental/ControlRig/Source/ControlRig/Private/Units/Hierarchy/RigUnit_SetCurveValue.cpp b/Engine/Plugins/Experimental/ControlRig/Source/ControlRig/Private/Units/Hierarchy/RigUnit_SetCurveValue.cpp
--- a/Engine/Plugins/Experimental/ControlRig/Source/ControlRig/Private/Units/Hierarchy/RigUnit_SetCurveValue.cpp
+++ b/Engine/Plugins/Experimental/ControlRig/Source/ControlRig/Private/Units/Hierarchy/RigUnit_SetCurveValue.cpp
@@ -6,8 +6,8 @@
 FRigUnit_SetCurveValue_Execute()
 {
     DECLARE_SCOPE_HIERARCHICAL_COUNTER_RIGUNIT()
-	FRigCurveContainer* CurveContainer = ExecuteContext.GetCurves();
-	if (CurveContainer)
+	URigHierarchy* Hierarchy = ExecuteContext.Hierarchy;
+	if (Hierarchy)
 	{
 		switch (Context.State)
 		{
@@ -18,9 +18,10 @@ FRigUnit_SetCurveValue_Execute()
 			}
 			case EControlRigState::Update:
 			{
-				if (CachedCurveIndex.UpdateCache(Curve, CurveContainer))
+				const FRigElementKey Key(Curve, ERigElementType::Curve);
+				if (CachedCurveIndex.UpdateCache(Key, Hierarchy))
 				{
-					CurveContainer->SetValue(CachedCurveIndex, Value);
+					Hierarchy->SetCurveValue(CachedCurveIndex, Value);
 				}
 			}
 			default:
@@ -36,24 +37,23 @@ FRigUnit_SetCurveValue_Execute()
 
 IMPLEMENT_RIGUNIT_AUTOMATION_TEST(FRigUnit_SetCurveValue)
 {
-	CurveContainer.Add(TEXT("CurveA"));
-	CurveContainer.Add(TEXT("CurveB"));
-	CurveContainer.Initialize();
-	Unit.ExecuteContext.Hierarchy = &HierarchyContainer;
+	const FRigElementKey CurveA = Controller->AddCurve(TEXT("CurveA"), 0.f);
+	const FRigElementKey CurveB = Controller->AddCurve(TEXT("CurveB"), 0.f);
+	Unit.ExecuteContext.Hierarchy = Hierarchy;
 	
-	CurveContainer.ResetValues();
+	Hierarchy->ResetCurveValues();
 	Unit.Curve = TEXT("CurveA");
 	Unit.Value = 3.0f;
 	InitAndExecute();
 
-	AddErrorIfFalse(CurveContainer.GetValue(FName(TEXT("CurveA"))) == 3.f, TEXT("unexpected value"));
+	AddErrorIfFalse(Hierarchy->GetCurveValue(CurveA) == 3.f, TEXT("unexpected value"));
 
-	CurveContainer.ResetValues();
+	Hierarchy->ResetCurveValues();
 	Unit.Curve = TEXT("CurveB");
 	Unit.Value = 13.0f;
 	InitAndExecute();
 
-	AddErrorIfFalse(CurveContainer.GetValue(FName(TEXT("CurveB"))) == 13.f, TEXT("unexpected value"));
+	AddErrorIfFalse(Hierarchy->GetCurveValue(CurveB) == 13.f, TEXT("unexpected value"));
 
 	return true;
 }
diff --git a/Engine/Plugins/Experimental/ControlRig/Source/ControlRig/Private/Units/Hierarchy/RigUnit_SetRelativeBoneTransform.cpp b/Engine/Plugins/Experimental/ControlRig/Source/ControlRig/Private/Units/Hierarchy/RigUnit_SetRelativeBoneTransform.cpp
--- a/Engine/Plugins/Experimental/ControlRig/Source/ControlRig/Private/Units/Hierarchy/RigUnit_SetRelativeBoneTransform.cpp
+++ b/Engine/Plugins/Experimental/ControlRig/Source/ControlRig/Private/Units/Hierarchy/RigUnit_SetRelativeBoneTransform.cpp
@@ -6,7 +6,7 @@
 FRigUnit_SetRelativeBoneTransform_Execute()
 {
     DECLARE_SCOPE_HIERARCHICAL_COUNTER_RIGUNIT()
-	FRigBoneHierarchy* Hierarchy = ExecuteContext.GetBones();
+	URigHierarchy* Hierarchy = ExecuteContext.Hierarchy;
 	if (Hierarchy)
 	{
 		switch (Context.State)
@@ -18,11 +18,13 @@ FRigUnit_SetRelativeBoneTransform_Execute()
 			}
 			case EControlRigState::Update:
 			{
-				if (!CachedBone.UpdateCache(Bone, Hierarchy))
+				const FRigElementKey BoneKey(Bone, ERigElementType::Bone);
+				const FRigElementKey SpaceKey(Space, ERigElementType::Bone);
+				if (!CachedBone.UpdateCache(BoneKey, Hierarchy))
 				{
 					UE_CONTROLRIG_RIGUNIT_REPORT_WARNING(TEXT("Bone '%s' is not valid."), *Bone.ToString());
 				}
-				else if (!CachedSpaceIndex.UpdateCache(Space, Hierarchy))
+				else if (!CachedSpaceIndex.UpdateCache(SpaceKey, Hierarchy))
 				{
 					UE_CONTROLRIG_RIGUNIT_REPORT_WARNING(TEXT("Bone '%s' is not valid."), *Bone.ToString());
 					UE_CONTROLRIG_RIGUNIT_REPORT_WARNING(TEXT("Space '%s' is not valid."), *Space.ToString());
@@ -55,48 +57,47 @@ FRigUnit_SetRelativeBoneTransform_Execute()
 
 IMPLEMENT_RIGUNIT_AUTOMATION_TEST(FRigUnit_SetRelativeBoneTransform)
 {
-	BoneHierarchy.Add(TEXT("Root"), NAME_None, ERigBoneType::User, FTransform(FVector(1.f, 0.f, 0.f)));
-	BoneHierarchy.Add(TEXT("BoneA"), TEXT("Root"), ERigBoneType::User, FTransform(FVector(1.f, 2.f, 3.f)));
-	BoneHierarchy.Add(TEXT("BoneB"), TEXT("BoneA"), ERigBoneType::User, FTransform(FVector(1.f, 5.f, 3.f)));
-	BoneHierarchy.Add(TEXT("BoneC"), TEXT("Root"), ERigBoneType::User, FTransform(FVector(-4.f, 0.f, 0.f)));
-	BoneHierarchy.Initialize();
-	Unit.ExecuteContext.Hierarchy = &HierarchyContainer;
+	const FRigElementKey Root = Controller->AddBone(TEXT("Root"), FRigElementKey(), FTransform(FVector(1.f, 0.f, 0.f)), true, ERigBoneType::User);
+	const FRigElementKey BoneA = Controller->AddBone(TEXT("BoneA"), Root, FTransform(FVector(1.f, 2.f, 3.f)), true, ERigBoneType::User);
+	const FRigElementKey BoneB = Controller->AddBone(TEXT("BoneB"), BoneA, FTransform(FVector(1.f, 5.f, 3.f)), true, ERigBoneType::User);
+	const FRigElementKey BoneC = Controller->AddBone(TEXT("BoneC"), Root, FTransform(FVector(-4.f, 0.f, 0.f)), true, ERigBoneType::User);
+	Unit.ExecuteContext.Hierarchy = Hierarchy;
 
-	BoneHierarchy.ResetTransforms();
+	Hierarchy->ResetPoseToInitial(ERigElementType::Bone);
 	Unit.Bone = TEXT("BoneA");
 	Unit.Space = TEXT("Root");
 	Unit.Transform = FTransform(FVector(0.f, 0.f, 7.f));
 	Unit.bPropagateToChildren = false;
 	InitAndExecute();
-	AddErrorIfFalse(BoneHierarchy.GetGlobalTransform(0).GetTranslation().Equals(FVector(1.f, 0.f, 0.f)), TEXT("unexpected transform"));
-	AddErrorIfFalse(BoneHierarchy.GetGlobalTransform(1).GetTranslation().Equals(FVector(1.f, 0.f, 7.f)), TEXT("unexpected transform"));
-	AddErrorIfFalse(BoneHierarchy.GetGlobalTransform(3).GetTranslation().Equals(FVector(1.f, 5.f, 3.f)), TEXT("unexpected transform"));
-	AddErrorIfFalse(BoneHierarchy.GetGlobalTransform(2).GetTranslation().Equals(FVector(-4.f, 0.f, 0.f)), TEXT("unexpected transform"));
+	AddErrorIfFalse(Hierarchy->GetGlobalTransform(0).GetTranslation().Equals(FVector(1.f, 0.f, 0.f)), TEXT("unexpected transform"));
+	AddErrorIfFalse(Hierarchy->GetGlobalTransform(1).GetTranslation().Equals(FVector(1.f, 0.f, 7.f)), TEXT("unexpected transform"));
+	AddErrorIfFalse(Hierarchy->GetGlobalTransform(2).GetTranslation().Equals(FVector(1.f, 5.f, 3.f)), TEXT("unexpected transform"));
+	AddErrorIfFalse(Hierarchy->GetGlobalTransform(3).GetTranslation().Equals(FVector(-4.f, 0.f, 0.f)), TEXT("unexpected transform"));
 
-	BoneHierarchy.ResetTransforms();
+	Hierarchy->ResetPoseToInitial(ERigElementType::Bone);
 	Unit.bPropagateToChildren = true;
 	InitAndExecute();
-	AddErrorIfFalse(BoneHierarchy.GetGlobalTransform(0).GetTranslation().Equals(FVector(1.f, 0.f, 0.f)), TEXT("unexpected transform"));
-	AddErrorIfFalse(BoneHierarchy.GetGlobalTransform(1).GetTranslation().Equals(FVector(1.f, 0.f, 7.f)), TEXT("unexpected transform"));
-	AddErrorIfFalse(BoneHierarchy.GetGlobalTransform(3).GetTranslation().Equals(FVector(1.f, 3.f, 7.f)), TEXT("unexpected transform"));
-	AddErrorIfFalse(BoneHierarchy.GetGlobalTransform(2).GetTranslation().Equals(FVector(-4.f, 0.f, 0.f)), TEXT("unexpected transform"));
+	AddErrorIfFalse(Hierarchy->GetGlobalTransform(0).GetTranslation().Equals(FVector(1.f, 0.f, 0.f)), TEXT("unexpected transform"));
+	AddErrorIfFalse(Hierarchy->GetGlobalTransform(1).GetTranslation().Equals(FVector(1.f, 0.f, 7.f)), TEXT("unexpected transform"));
+	AddErrorIfFalse(Hierarchy->GetGlobalTransform(2).GetTranslation().Equals(FVector(1.f, 3.f, 7.f)), TEXT("unexpected transform"));
+	AddErrorIfFalse(Hierarchy->GetGlobalTransform(3).GetTranslation().Equals(FVector(-4.f, 0.f, 0.f)), TEXT("unexpected transform"));
 
-	BoneHierarchy.ResetTransforms();
+	Hierarchy->ResetPoseToInitial(ERigElementType::Bone);
 	Unit.Space = TEXT("BoneC");
 	Unit.bPropagateToChildren = false;
 	InitAndExecute();
-	AddErrorIfFalse(BoneHierarchy.GetGlobalTransform(0).GetTranslation().Equals(FVector(1.f, 0.f, 0.f)), TEXT("unexpected transform"));
-	AddErrorIfFalse(BoneHierarchy.GetGlobalTransform(1).GetTranslation().Equals(FVector(-4.f, 0.f, 7.f)), TEXT("unexpected transform"));
-	AddErrorIfFalse(BoneHierarchy.GetGlobalTransform(3).GetTranslation().Equals(FVector(1.f, 5.f, 3.f)), TEXT("unexpected transform"));
-	AddErrorIfFalse(BoneHierarchy.GetGlobalTransform(2).GetTranslation().Equals(FVector(-4.f, 0.f, 0.f)), TEXT("unexpected transform"));
+	AddErrorIfFalse(Hierarchy->GetGlobalTransform(0).GetTranslation().Equals(FVector(1.f, 0.f, 0.f)), TEXT("unexpected transform"));
+	AddErrorIfFalse(Hierarchy->GetGlobalTransform(1).GetTranslation().Equals(FVector(-4.f, 0.f, 7.f)), TEXT("unexpected transform"));
+	AddErrorIfFalse(Hierarchy->GetGlobalTransform(2).GetTranslation().Equals(FVector(1.f, 5.f, 3.f)), TEXT("unexpected transform"));
+	AddErrorIfFalse(Hierarchy->GetGlobalTransform(3).GetTranslation().Equals(FVector(-4.f, 0.f, 0.f)), TEXT("unexpected transform"));
 
-	BoneHierarchy.ResetTransforms();
+	Hierarchy->ResetPoseToInitial(ERigElementType::Bone);
 	Unit.bPropagateToChildren = true;
 	InitAndExecute();
-	AddErrorIfFalse(BoneHierarchy.GetGlobalTransform(0).GetTranslation().Equals(FVector(1.f, 0.f, 0.f)), TEXT("unexpected transform"));
-	AddErrorIfFalse(BoneHierarchy.GetGlobalTransform(1).GetTranslation().Equals(FVector(-4.f, 0.f, 7.f)), TEXT("unexpected transform"));
-	AddErrorIfFalse(BoneHierarchy.GetGlobalTransform(3).GetTranslation().Equals(FVector(-4.f, 3.f, 7.f)), TEXT("unexpected transform"));
-	AddErrorIfFalse(BoneHierarchy.GetGlobalTransform(2).GetTranslation().Equals(FVector(-4.f, 0.f, 0.f)), TEXT("unexpected transform"));
+	AddErrorIfFalse(Hierarchy->GetGlobalTransform(0).GetTranslation().Equals(FVector(1.f, 0.f, 0.f)), TEXT("unexpected transform"));
+	AddErrorIfFalse(Hierarchy->GetGlobalTransform(1).GetTranslation().Equals(FVector(-4.f, 0.f, 7.f)), TEXT("unexpected transform"));
+	AddErrorIfFalse(Hierarchy->GetGlobalTransform(2).GetTranslation().Equals(FVector(-4.f, 3.f, 7.f)), TEXT("unexpected transform"));
+	AddErrorIfFalse(Hierarchy->GetGlobalTransform(3).GetTranslation().Equals(FVector(-4.f, 0.f, 0.f)), TEXT("unexpected transform"));
 
 	return true;
 }
diff --git a/Engine/Plugins/Experimental/ControlRig/Source/ControlRig/Private/Units/Hierarchy/RigUnit_SetSpaceInitialTransform.cpp b/Engine/Plugins/Experimental/ControlRig/Source/ControlRig/Private/Units/Hierarchy/RigUnit_SetSpaceInitialTransform.cpp
--- a/Engine/Plugins/Experimental/ControlRig/Source/ControlRig/Private/Units/Hierarchy/RigUnit_SetSpaceInitialTransform.cpp
+++ b/Engine/Plugins/Experimental/ControlRig/Source/ControlRig/Private/Units/Hierarchy/RigUnit_SetSpaceInitialTransform.cpp
@@ -8,11 +8,9 @@ FRigUnit_SetSpaceInitialTransform_Execute()
 {
     DECLARE_SCOPE_HIERARCHICAL_COUNTER_RIGUNIT()
 
-	FRigSpaceHierarchy* Hierarchy = ExecuteContext.GetSpaces();
+	URigHierarchy* Hierarchy = ExecuteContext.Hierarchy;
 	if (Hierarchy)
 	{
-		FRigSpaceHierarchy& HierarchyRef = *Hierarchy;
-
 		switch (Context.State)
 		{
 			case EControlRigState::Init:
@@ -22,7 +20,8 @@ FRigUnit_SetSpaceInitialTransform_Execute()
 			}
 			case EControlRigState::Update:
 			{
-				if (!CachedSpaceIndex.UpdateCache(SpaceName, Hierarchy))
+				const FRigElementKey SpaceKey(SpaceName, ERigElementType::Space);
+				if (!CachedSpaceIndex.UpdateCache(SpaceKey, Hierarchy))
 				{
 					UE_CONTROLRIG_RIGUNIT_REPORT_WARNING(TEXT("Space '%s' is not valid."), *SpaceName.ToString());
 					return;
@@ -31,12 +30,11 @@ FRigUnit_SetSpaceInitialTransform_Execute()
 				FTransform InitialTransform = Transform;
 				if (Space == EBoneGetterSetterMode::GlobalSpace)
 				{
-					FRigSpace& SpaceRef = HierarchyRef[CachedSpaceIndex];
-					FTransform ParentTransform = ExecuteContext.Hierarchy->GetInitialGlobalTransform(SpaceRef.GetParentElementKey());
+					const FTransform ParentTransform = Hierarchy->GetParentTransformByIndex(CachedSpaceIndex, true);
 					InitialTransform = InitialTransform.GetRelativeTransform(ParentTransform);
 				}
 
-				HierarchyRef.SetInitialTransform(CachedSpaceIndex, InitialTransform);
+				Hierarchy->SetInitialLocalTransform(CachedSpaceIndex, InitialTransform);
 			}
 			default:
 			{
diff --git a/Engine/Plugins/Experimental/ControlRig/Source/ControlRig/Private/Units/Hierarchy/RigUnit_SetSpaceTransform.cpp b/Engine/Plugins/Experimental/ControlRig/Source/ControlRig/Private/Units/Hierarchy/RigUnit_SetSpaceTransform.cpp
--- a/Engine/Plugins/Experimental/ControlRig/Source/ControlRig/Private/Units/Hierarchy/RigUnit_SetSpaceTransform.cpp
+++ b/Engine/Plugins/Experimental/ControlRig/Source/ControlRig/Private/Units/Hierarchy/RigUnit_SetSpaceTransform.cpp
@@ -7,7 +7,7 @@
 FRigUnit_SetSpaceTransform_Execute()
 {
     DECLARE_SCOPE_HIERARCHICAL_COUNTER_RIGUNIT()
-	FRigSpaceHierarchy* Hierarchy = ExecuteContext.GetSpaces();
+	URigHierarchy* Hierarchy = ExecuteContext.Hierarchy;
 	if (Hierarchy)
 	{
 		switch (Context.State)
@@ -19,7 +19,9 @@ FRigUnit_SetSpaceTransform_Execute()
 			}
 			case EControlRigState::Update:
 			{
-				if (CachedSpaceIndex.UpdateCache(Space, Hierarchy))
+				const FRigElementKey SpaceKey(Space, ERigElementType::Space);
+
+				if (CachedSpaceIndex.UpdateCache(SpaceKey, Hierarchy))
 				{
 					switch (SpaceType)
 					{
@@ -27,12 +29,12 @@ FRigUnit_SetSpaceTransform_Execute()
 						{
 							if(FMath::IsNearlyEqual(Weight, 1.f))
 							{
-								Hierarchy->SetGlobalTransform(CachedSpaceIndex, Transform);
+								Hierarchy->SetGlobalTransform(CachedSpaceIndex, Transform, true);
 							}
 							else
 							{
-								FTransform PreviousTransform = Hierarchy->GetGlobalTransform(CachedSpaceIndex);
-								Hierarchy->SetGlobalTransform(CachedSpaceIndex, FControlRigMathLibrary::LerpTransform(PreviousTransform, Transform, FMath::Clamp<float>(Weight, 0.f, 1.f)));
+								const FTransform PreviousTransform = Hierarchy->GetGlobalTransform(CachedSpaceIndex);
+								Hierarchy->SetGlobalTransform(CachedSpaceIndex, FControlRigMathLibrary::LerpTransform(PreviousTransform, Transform, FMath::Clamp<float>(Weight, 0.f, 1.f)), true);
 							}
 							break;
 						}
@@ -40,12 +42,12 @@ FRigUnit_SetSpaceTransform_Execute()
 						{
 							if(FMath::IsNearlyEqual(Weight, 1.f))
 							{
-								Hierarchy->SetLocalTransform(CachedSpaceIndex, Transform);
+								Hierarchy->SetLocalTransform(CachedSpaceIndex, Transform, true);
 							}
 							else
 							{
-								FTransform PreviousTransform = Hierarchy->GetLocalTransform(CachedSpaceIndex);
-								Hierarchy->SetLocalTransform(CachedSpaceIndex, FControlRigMathLibrary::LerpTransform(PreviousTransform, Transform, FMath::Clamp<float>(Weight, 0.f, 1.f)));
+								const FTransform PreviousTransform = Hierarchy->GetLocalTransform(CachedSpaceIndex);
+								Hierarchy->SetLocalTransform(CachedSpaceIndex, FControlRigMathLibrary::LerpTransform(PreviousTransform, Transform, FMath::Clamp<float>(Weight, 0.f, 1.f)), true);
 							}
 							break;
 						}
diff --git a/Engine/Plugins/Experimental/ControlRig/Source/ControlRig/Private/Units/Hierarchy/RigUnit_SetTransform.cpp b/Engine/Plugins/Experimental/ControlRig/Source/ControlRig/Private/Units/Hierarchy/RigUnit_SetTransform.cpp
--- a/Engine/Plugins/Experimental/ControlRig/Source/ControlRig/Private/Units/Hierarchy/RigUnit_SetTransform.cpp
+++ b/Engine/Plugins/Experimental/ControlRig/Source/ControlRig/Private/Units/Hierarchy/RigUnit_SetTransform.cpp
@@ -22,7 +22,7 @@ FRigUnit_SetTransform_Execute()
 		return;
 	}
 
-	if (FRigHierarchyContainer* Hierarchy = ExecuteContext.Hierarchy)
+	if (URigHierarchy* Hierarchy = ExecuteContext.Hierarchy)
 	{
 		switch (Context.State)
 		{
@@ -43,59 +43,20 @@ FRigUnit_SetTransform_Execute()
 				}
 				else
 				{
-					if(bInitial)
+					FTransform WeightedTransform = Transform;
+					if (Weight < 1.f - SMALL_NUMBER)
 					{
-						// special case controls - since we want their offset to behave differently
-						if (CachedIndex.GetKey().Type == ERigElementType::Control)
-						{
-							FTransform OffsetTransform = Transform;
-
-							if (Space == EBoneGetterSetterMode::GlobalSpace)
-							{
-								FRigElementKey ParentKey = Hierarchy->GetParentKey(CachedIndex);
-								if (ParentKey.IsValid())
-								{
-									FTransform ParentTransform = Hierarchy->GetInitialGlobalTransform(ParentKey);
-									OffsetTransform = OffsetTransform.GetRelativeTransform(ParentTransform);
-								}
-							}
-
-							FRigControl& Control = (*ExecuteContext.GetControls())[CachedIndex];
-							Control.OffsetTransform = OffsetTransform;
-
-							if (Control.ControlType == ERigControlType::Transform ||
-								Control.ControlType == ERigControlType::TransformNoScale || 
-								Control.ControlType == ERigControlType::Position || 
-								Control.ControlType == ERigControlType::Rotator || 
-								Control.ControlType == ERigControlType::Scale)
-							{
-								Control.SetValueFromTransform(FTransform::Identity, ERigControlValueType::Initial);
-								Control.SetValueFromTransform(FTransform::Identity, ERigControlValueType::Current);
-							}
-
-							return;
-						}
-
+						FTransform PreviousTransform = WeightedTransform;
 						switch (Space)
 						{
 							case EBoneGetterSetterMode::GlobalSpace:
 							{
-								Hierarchy->SetInitialGlobalTransform(CachedIndex, Transform);
-
-								if (ExecuteContext.EventName == FRigUnit_PrepareForExecution::EventName)
-								{
-									Hierarchy->SetGlobalTransform(CachedIndex, Transform);
-								}
+								PreviousTransform = Hierarchy->GetGlobalTransformByIndex(CachedIndex, bInitial);
 								break;
 							}
 							case EBoneGetterSetterMode::LocalSpace:
 							{
-								Hierarchy->SetInitialTransform(CachedIndex, Transform);
-
-								if (ExecuteContext.EventName == FRigUnit_PrepareForExecution::EventName)
-								{
-									Hierarchy->SetLocalTransform(CachedIndex, Transform);
-								}
+								PreviousTransform = Hierarchy->GetLocalTransformByIndex(CachedIndex, bInitial);
 								break;
 							}
 							default:
@@ -103,49 +64,34 @@ FRigUnit_SetTransform_Execute()
 								break;
 							}
 						}
+						WeightedTransform = FControlRigMathLibrary::LerpTransform(PreviousTransform, WeightedTransform, Weight);
 					}
-					else
+
+					switch (Space)
 					{
-						FTransform WeightedTransform = Transform;
-						if (Weight < 1.f - SMALL_NUMBER)
+						case EBoneGetterSetterMode::GlobalSpace:
 						{
-							FTransform PreviousTransform = WeightedTransform;
-							switch (Space)
+							Hierarchy->SetGlobalTransformByIndex(CachedIndex, WeightedTransform, bInitial, bPropagateToChildren);
+
+							if (bInitial && ExecuteContext.EventName == FRigUnit_PrepareForExecution::EventName)
 							{
-								case EBoneGetterSetterMode::GlobalSpace:
-								{
-									PreviousTransform = Hierarchy->GetGlobalTransform(CachedIndex);
-									break;
-								}
-								case EBoneGetterSetterMode::LocalSpace:
-								{
-									PreviousTransform = Hierarchy->GetLocalTransform(CachedIndex);
-									break;
-								}
-								default:
-								{
-									break;
-								}
+								Hierarchy->SetGlobalTransformByIndex(CachedIndex, WeightedTransform, false, bPropagateToChildren);
 							}
-							WeightedTransform = FControlRigMathLibrary::LerpTransform(PreviousTransform, WeightedTransform, Weight);
+							break;
 						}
-
-						switch (Space)
+						case EBoneGetterSetterMode::LocalSpace:
 						{
-							case EBoneGetterSetterMode::GlobalSpace:
-							{
-								Hierarchy->SetGlobalTransform(CachedIndex, WeightedTransform, bPropagateToChildren);
-								break;
-							}
-							case EBoneGetterSetterMode::LocalSpace:
-							{
-								Hierarchy->SetLocalTransform(CachedIndex, WeightedTransform, bPropagateToChildren);
-								break;
-							}
-							default:
+							Hierarchy->SetLocalTransformByIndex(CachedIndex, WeightedTransform, bInitial, bPropagateToChildren);
+
+							if (bInitial && ExecuteContext.EventName == FRigUnit_PrepareForExecution::EventName)
 							{
-								break;
+								Hierarchy->SetLocalTransformByIndex(CachedIndex, WeightedTransform, false, bPropagateToChildren);
 							}
+							break;
+						}
+						default:
+						{
+							break;
 						}
 					}
 				}
diff --git a/Engine/Plugins/Experimental/ControlRig/Source/ControlRig/Private/Units/Highlevel/Harmonics/RigUnit_BoneHarmonics.cpp b/Engine/Plugins/Experimental/ControlRig/Source/ControlRig/Private/Units/Highlevel/Harmonics/RigUnit_BoneHarmonics.cpp
--- a/Engine/Plugins/Experimental/ControlRig/Source/ControlRig/Private/Units/Highlevel/Harmonics/RigUnit_BoneHarmonics.cpp
+++ b/Engine/Plugins/Experimental/ControlRig/Source/ControlRig/Private/Units/Highlevel/Harmonics/RigUnit_BoneHarmonics.cpp
@@ -35,7 +35,7 @@ FRigUnit_BoneHarmonics_Execute()
 FRigUnit_ItemHarmonics_Execute()
 {
     DECLARE_SCOPE_HIERARCHICAL_COUNTER_RIGUNIT()
-	FRigHierarchyContainer* Hierarchy = ExecuteContext.Hierarchy;
+	URigHierarchy* Hierarchy = ExecuteContext.Hierarchy;
 	if (Hierarchy == nullptr)
 	{
 		return;
diff --git a/Engine/Plugins/Experimental/ControlRig/Source/ControlRig/Private/Units/Highlevel/Harmonics/RigUnit_ChainHarmonics.cpp b/Engine/Plugins/Experimental/ControlRig/Source/ControlRig/Private/Units/Highlevel/Harmonics/RigUnit_ChainHarmonics.cpp
--- a/Engine/Plugins/Experimental/ControlRig/Source/ControlRig/Private/Units/Highlevel/Harmonics/RigUnit_ChainHarmonics.cpp
+++ b/Engine/Plugins/Experimental/ControlRig/Source/ControlRig/Private/Units/Highlevel/Harmonics/RigUnit_ChainHarmonics.cpp
@@ -23,7 +23,7 @@ FRigUnit_ChainHarmonics_Execute()
 FRigUnit_ChainHarmonicsPerItem_Execute()
 {
     DECLARE_SCOPE_HIERARCHICAL_COUNTER_RIGUNIT()
-	FRigHierarchyContainer* Hierarchy = ExecuteContext.Hierarchy;
+	URigHierarchy* Hierarchy = ExecuteContext.Hierarchy;
 	if (Hierarchy == nullptr)
 	{
 		return;
@@ -55,11 +55,11 @@ FRigUnit_ChainHarmonicsPerItem_Execute()
 
 		Items.Add(FCachedRigElement(ChainRoot, Hierarchy));
 
-		TArray<FRigElementKey> Children = Hierarchy->GetChildKeys(Items.Last(), false);
+		TArray<FRigElementKey> Children = Hierarchy->GetChildren(Items.Last().GetKey());
 		while (Children.Num() > 0)
 		{
 			Items.Add(FCachedRigElement(Children[0], Hierarchy));
-			Children = Hierarchy->GetChildKeys(Children[0], false);
+			Children = Hierarchy->GetChildren(Children[0]);
 		}
 
 		if (Items.Num() < 2)
@@ -100,7 +100,7 @@ FRigUnit_ChainHarmonicsPerItem_Execute()
 	}
 
 	FTransform ParentTransform = FTransform::Identity;
-	FRigElementKey ParentKey = Hierarchy->GetParentKey(Items[0]);
+	FRigElementKey ParentKey = Hierarchy->GetFirstParent(Items[0].GetKey());
 	if (ParentKey.IsValid())
 	{
 		ParentTransform = Hierarchy->GetGlobalTransform(ParentKey);
diff --git a/Engine/Plugins/Experimental/ControlRig/Source/ControlRig/Private/Units/Highlevel/Hierarchy/RigUnit_AimBone.cpp b/Engine/Plugins/Experimental/ControlRig/Source/ControlRig/Private/Units/Highlevel/Hierarchy/RigUnit_AimBone.cpp
--- a/Engine/Plugins/Experimental/ControlRig/Source/ControlRig/Private/Units/Highlevel/Hierarchy/RigUnit_AimBone.cpp
+++ b/Engine/Plugins/Experimental/ControlRig/Source/ControlRig/Private/Units/Highlevel/Hierarchy/RigUnit_AimBone.cpp
@@ -9,7 +9,7 @@ FRigUnit_AimBoneMath_Execute()
 
 	Result = InputTransform;
 
-	const FRigHierarchyContainer* Hierarchy = Context.Hierarchy;
+	URigHierarchy* Hierarchy = Context.Hierarchy;
 	if (Hierarchy == nullptr)
 	{
 		return;
@@ -176,7 +176,7 @@ FRigUnit_AimBone_Execute()
 FRigUnit_AimItem_Execute()
 {
     DECLARE_SCOPE_HIERARCHICAL_COUNTER_RIGUNIT()
-	FRigHierarchyContainer* Hierarchy = ExecuteContext.Hierarchy;
+	URigHierarchy* Hierarchy = ExecuteContext.Hierarchy;
 	if (Hierarchy == nullptr)
 	{
 		return;
diff --git a/Engine/Plugins/Experimental/ControlRig/Source/ControlRig/Private/Units/Highlevel/Hierarchy/RigUnit_CCDIK.cpp b/Engine/Plugins/Experimental/ControlRig/Source/ControlRig/Private/Units/Highlevel/Hierarchy/RigUnit_CCDIK.cpp
--- a/Engine/Plugins/Experimental/ControlRig/Source/ControlRig/Private/Units/Highlevel/Hierarchy/RigUnit_CCDIK.cpp
+++ b/Engine/Plugins/Experimental/ControlRig/Source/ControlRig/Private/Units/Highlevel/Hierarchy/RigUnit_CCDIK.cpp
@@ -51,7 +51,7 @@ FRigUnit_CCDIK_Execute()
 FRigUnit_CCDIKPerItem_Execute()
 {
     DECLARE_SCOPE_HIERARCHICAL_COUNTER_RIGUNIT()
-	FRigHierarchyContainer* Hierarchy = ExecuteContext.Hierarchy;
+	URigHierarchy* Hierarchy = ExecuteContext.Hierarchy;
 	if (Hierarchy == nullptr)
 	{
 		return;
@@ -77,7 +77,7 @@ FRigUnit_CCDIKPerItem_Execute()
 	{
 		if (CachedItems.Num() == 0 && Items.Num() > 0)
 		{
-			CachedItems.Add(FCachedRigElement(Hierarchy->GetParentKey(Items[0]), Hierarchy));
+			CachedItems.Add(FCachedRigElement(Hierarchy->GetFirstParent(Items[0]), Hierarchy));
 
 			for (FRigElementKey Item : Items)
 			{
