Control Rig: New rig hierarchy data structure refactoring

diff --git a/Engine/Plugins/Experimental/ControlRig/Source/ControlRig/Private/Units/Highlevel/Hierarchy/RigUnit_DistributeRotation.cpp b/Engine/Plugins/Experimental/ControlRig/Source/ControlRig/Private/Units/Highlevel/Hierarchy/RigUnit_DistributeRotation.cpp
--- a/Engine/Plugins/Experimental/ControlRig/Source/ControlRig/Private/Units/Highlevel/Hierarchy/RigUnit_DistributeRotation.cpp
+++ b/Engine/Plugins/Experimental/ControlRig/Source/ControlRig/Private/Units/Highlevel/Hierarchy/RigUnit_DistributeRotation.cpp
@@ -9,16 +9,16 @@ FRigUnit_DistributeRotation_Execute()
 
 	if(WorkData.CachedItems.Num() == 0)
 	{
-		FRigBoneHierarchy* Hierarchy = ExecuteContext.GetBones();
+		URigHierarchy* Hierarchy = ExecuteContext.Hierarchy;
 		if (Hierarchy == nullptr)
 		{
 			return;
 		}
 
-		int32 EndBoneIndex = Hierarchy->GetIndex(EndBone);
+		int32 EndBoneIndex = Hierarchy->GetIndex(FRigElementKey(EndBone, ERigElementType::Bone));
 		if (EndBoneIndex != INDEX_NONE)
 		{
-			int32 StartBoneIndex = Hierarchy->GetIndex(StartBone);
+			int32 StartBoneIndex = Hierarchy->GetIndex(FRigElementKey(StartBone, ERigElementType::Bone));
 			if (StartBoneIndex == EndBoneIndex)
 			{
 				return;
@@ -26,12 +26,12 @@ FRigUnit_DistributeRotation_Execute()
 
 			while (EndBoneIndex != INDEX_NONE)
 			{
-				Items.Add((*Hierarchy)[EndBoneIndex].GetElementKey());
+				Items.Add(Hierarchy->GetKey(EndBoneIndex));
 				if (EndBoneIndex == StartBoneIndex)
 				{
 					break;
 				}
-				EndBoneIndex = (*Hierarchy)[EndBoneIndex].ParentIndex;
+				EndBoneIndex = Hierarchy->GetFirstParent(EndBoneIndex);
 			}
 		}
 
@@ -59,7 +59,7 @@ FRigUnit_DistributeRotation_Execute()
 FRigUnit_DistributeRotationForCollection_Execute()
 {
     DECLARE_SCOPE_HIERARCHICAL_COUNTER_RIGUNIT()
-	FRigHierarchyContainer* Hierarchy = ExecuteContext.Hierarchy;
+	URigHierarchy* Hierarchy = ExecuteContext.Hierarchy;
 	if (Hierarchy == nullptr)
 	{
 		return;
@@ -251,13 +251,12 @@ FRigUnit_DistributeRotationForCollection_Execute()
 
 IMPLEMENT_RIGUNIT_AUTOMATION_TEST(FRigUnit_DistributeRotation)
 {
-	BoneHierarchy.Add(TEXT("Root"), NAME_None, ERigBoneType::User, FTransform(FVector(1.f, 0.f, 0.f)));
-	BoneHierarchy.Add(TEXT("BoneA"), TEXT("Root"), ERigBoneType::User, FTransform(FVector(2.f, 0.f, 0.f)));
-	BoneHierarchy.Add(TEXT("BoneB"), TEXT("BoneA"), ERigBoneType::User, FTransform(FVector(2.f, 0.f, 0.f)));
-	BoneHierarchy.Add(TEXT("BoneC"), TEXT("BoneB"), ERigBoneType::User, FTransform(FVector(2.f, 0.f, 0.f)));
-	BoneHierarchy.Add(TEXT("BoneD"), TEXT("BoneC"), ERigBoneType::User, FTransform(FVector(2.f, 0.f, 0.f)));
-	BoneHierarchy.Initialize();
-	Unit.ExecuteContext.Hierarchy = &HierarchyContainer;
+	const FRigElementKey Root = Controller->AddBone(TEXT("Root"), FRigElementKey(), FTransform(FVector(1.f, 0.f, 0.f)), true, ERigBoneType::User);
+	const FRigElementKey BoneA = Controller->AddBone(TEXT("BoneA"), Root, FTransform(FVector(2.f, 0.f, 0.f)), true, ERigBoneType::User);
+	const FRigElementKey BoneB = Controller->AddBone(TEXT("BoneB"), BoneA, FTransform(FVector(2.f, 0.f, 0.f)), true, ERigBoneType::User);
+	const FRigElementKey BoneC = Controller->AddBone(TEXT("BoneC"), BoneB, FTransform(FVector(2.f, 0.f, 0.f)), true, ERigBoneType::User);
+	const FRigElementKey BoneD = Controller->AddBone(TEXT("BoneD"), BoneC, FTransform(FVector(2.f, 0.f, 0.f)), true, ERigBoneType::User);
+	Unit.ExecuteContext.Hierarchy = Hierarchy;
 
 	Unit.StartBone = TEXT("Root");
 	Unit.EndBone = TEXT("BoneD");
@@ -293,15 +292,15 @@ IMPLEMENT_RIGUNIT_AUTOMATION_TEST(FRigUnit_DistributeRotation)
 	AddErrorIfFalse(FMath::IsNearlyEqual(Unit.WorkData.ItemRotationT[4], 0.0f, 0.001f), TEXT("unexpected bone t"));
 
 	FVector Euler = FVector::ZeroVector;
-	Euler = FControlRigMathLibrary::EulerFromQuat(BoneHierarchy.GetLocalTransform(0).GetRotation(), EControlRigRotationOrder::XYZ);
+	Euler = FControlRigMathLibrary::EulerFromQuat(Hierarchy->GetLocalTransform(0).GetRotation(), EControlRigRotationOrder::XYZ);
 	AddErrorIfFalse(FMath::IsNearlyEqual(Euler.Y, 0.f, 0.1f), TEXT("unexpected rotation Y"));
-	Euler = FControlRigMathLibrary::EulerFromQuat(BoneHierarchy.GetLocalTransform(1).GetRotation(), EControlRigRotationOrder::XYZ);
+	Euler = FControlRigMathLibrary::EulerFromQuat(Hierarchy->GetLocalTransform(1).GetRotation(), EControlRigRotationOrder::XYZ);
 	AddErrorIfFalse(FMath::IsNearlyEqual(Euler.Y, 45.f, 0.1f), TEXT("unexpected rotation Y"));
-	Euler = FControlRigMathLibrary::EulerFromQuat(BoneHierarchy.GetLocalTransform(2).GetRotation(), EControlRigRotationOrder::XYZ);
+	Euler = FControlRigMathLibrary::EulerFromQuat(Hierarchy->GetLocalTransform(2).GetRotation(), EControlRigRotationOrder::XYZ);
 	AddErrorIfFalse(FMath::IsNearlyEqual(Euler.Y, 90.f, 0.1f), TEXT("unexpected rotation Y"));
-	Euler = FControlRigMathLibrary::EulerFromQuat(BoneHierarchy.GetLocalTransform(3).GetRotation(), EControlRigRotationOrder::XYZ);
+	Euler = FControlRigMathLibrary::EulerFromQuat(Hierarchy->GetLocalTransform(3).GetRotation(), EControlRigRotationOrder::XYZ);
 	AddErrorIfFalse(FMath::IsNearlyEqual(Euler.Y, 45.f, 0.1f), TEXT("unexpected rotation Y"));
-	Euler = FControlRigMathLibrary::EulerFromQuat(BoneHierarchy.GetLocalTransform(4).GetRotation(), EControlRigRotationOrder::XYZ);
+	Euler = FControlRigMathLibrary::EulerFromQuat(Hierarchy->GetLocalTransform(4).GetRotation(), EControlRigRotationOrder::XYZ);
 	AddErrorIfFalse(FMath::IsNearlyEqual(Euler.Y, 0.f, 0.1f), TEXT("unexpected rotation Y"));
 
 	return true;
diff --git a/Engine/Plugins/Experimental/ControlRig/Source/ControlRig/Private/Units/Highlevel/Hierarchy/RigUnit_FitChainToCurve.cpp b/Engine/Plugins/Experimental/ControlRig/Source/ControlRig/Private/Units/Highlevel/Hierarchy/RigUnit_FitChainToCurve.cpp
--- a/Engine/Plugins/Experimental/ControlRig/Source/ControlRig/Private/Units/Highlevel/Hierarchy/RigUnit_FitChainToCurve.cpp
+++ b/Engine/Plugins/Experimental/ControlRig/Source/ControlRig/Private/Units/Highlevel/Hierarchy/RigUnit_FitChainToCurve.cpp
@@ -47,7 +47,7 @@ FRigUnit_FitChainToCurvePerItem_Execute()
 {
     DECLARE_SCOPE_HIERARCHICAL_COUNTER_RIGUNIT()
 
-	FRigHierarchyContainer* Hierarchy = ExecuteContext.Hierarchy;
+	URigHierarchy* Hierarchy = ExecuteContext.Hierarchy;
 	if (Hierarchy == nullptr)
 	{
 		return;
@@ -379,7 +379,7 @@ FRigUnit_FitChainToCurvePerItem_Execute()
 	if (Rotations.Num() > 0)
 	{
 		FTransform BaseTransform = FTransform::Identity;
-		FRigElementKey ParentKey = Hierarchy->GetParentKey(CachedItems[0]);
+		FRigElementKey ParentKey = Hierarchy->GetFirstParent(CachedItems[0].GetKey());
 		if (ParentKey.IsValid())
 		{
 			BaseTransform = Hierarchy->GetGlobalTransform(ParentKey);
diff --git a/Engine/Plugins/Experimental/ControlRig/Source/ControlRig/Private/Units/Highlevel/Hierarchy/RigUnit_ModifyBoneTransforms.cpp b/Engine/Plugins/Experimental/ControlRig/Source/ControlRig/Private/Units/Highlevel/Hierarchy/RigUnit_ModifyBoneTransforms.cpp
--- a/Engine/Plugins/Experimental/ControlRig/Source/ControlRig/Private/Units/Highlevel/Hierarchy/RigUnit_ModifyBoneTransforms.cpp
+++ b/Engine/Plugins/Experimental/ControlRig/Source/ControlRig/Private/Units/Highlevel/Hierarchy/RigUnit_ModifyBoneTransforms.cpp
@@ -31,52 +31,51 @@ FRigUnit_ModifyBoneTransforms_Execute()
 
 IMPLEMENT_RIGUNIT_AUTOMATION_TEST(FRigUnit_ModifyBoneTransforms)
 {
-	BoneHierarchy.Add(TEXT("Root"), NAME_None, ERigBoneType::User, FTransform(FVector(1.f, 0.f, 0.f)));
-	BoneHierarchy.Add(TEXT("BoneA"), TEXT("Root"), ERigBoneType::User, FTransform(FVector(1.f, 2.f, 3.f)));
-	BoneHierarchy.Add(TEXT("BoneB"), TEXT("Root"), ERigBoneType::User, FTransform(FVector(5.f, 6.f, 7.f)));
-	BoneHierarchy.Initialize();
-	Unit.ExecuteContext.Hierarchy = &HierarchyContainer;
+	const FRigElementKey Root = Controller->AddBone(TEXT("Root"), FRigElementKey(), FTransform(FVector(1.f, 0.f, 0.f)), true, ERigBoneType::User);
+	const FRigElementKey BoneA = Controller->AddBone(TEXT("BoneA"), Root, FTransform(FVector(1.f, 2.f, 3.f)), true, ERigBoneType::User);
+	const FRigElementKey BoneB = Controller->AddBone(TEXT("BoneB"), Root, FTransform(FVector(5.f, 6.f, 7.f)), true, ERigBoneType::User);
+	Unit.ExecuteContext.Hierarchy = Hierarchy;
 
 	Unit.BoneToModify.SetNumZeroed(2);
 	Unit.BoneToModify[0].Bone = TEXT("BoneA");
 	Unit.BoneToModify[1].Bone = TEXT("BoneB");
 	Unit.BoneToModify[0].Transform = Unit.BoneToModify[1].Transform = FTransform(FVector(10.f, 11.f, 12.f));
 
-	BoneHierarchy.ResetTransforms();
+	Hierarchy->ResetPoseToInitial(ERigElementType::Bone);
 	Unit.Mode = EControlRigModifyBoneMode::AdditiveLocal;
 	InitAndExecute();
-	AddErrorIfFalse((BoneHierarchy.GetGlobalTransform(0).GetTranslation() - FVector(1.f, 0.f, 0.f)).IsNearlyZero(), TEXT("unexpected transform"));
-	AddErrorIfFalse((BoneHierarchy.GetGlobalTransform(1).GetTranslation() - FVector(11.f, 13.f, 15.f)).IsNearlyZero(), TEXT("unexpected transform"));
-	AddErrorIfFalse((BoneHierarchy.GetGlobalTransform(2).GetTranslation() - FVector(15.f, 17.f, 19.f)).IsNearlyZero(), TEXT("unexpected transform"));
+	AddErrorIfFalse((Hierarchy->GetGlobalTransform(0).GetTranslation() - FVector(1.f, 0.f, 0.f)).IsNearlyZero(), TEXT("unexpected transform"));
+	AddErrorIfFalse((Hierarchy->GetGlobalTransform(1).GetTranslation() - FVector(11.f, 13.f, 15.f)).IsNearlyZero(), TEXT("unexpected transform"));
+	AddErrorIfFalse((Hierarchy->GetGlobalTransform(2).GetTranslation() - FVector(15.f, 17.f, 19.f)).IsNearlyZero(), TEXT("unexpected transform"));
 
-	BoneHierarchy.ResetTransforms();
+	Hierarchy->ResetPoseToInitial(ERigElementType::Bone);
 	Unit.Mode = EControlRigModifyBoneMode::AdditiveGlobal;
 	InitAndExecute();
-	AddErrorIfFalse((BoneHierarchy.GetGlobalTransform(0).GetTranslation() - FVector(1.f, 0.f, 0.f)).IsNearlyZero(), TEXT("unexpected transform"));
-	AddErrorIfFalse((BoneHierarchy.GetGlobalTransform(1).GetTranslation() - FVector(11.f, 13.f, 15.f)).IsNearlyZero(), TEXT("unexpected transform"));
-	AddErrorIfFalse((BoneHierarchy.GetGlobalTransform(2).GetTranslation() - FVector(15.f, 17.f, 19.f)).IsNearlyZero(), TEXT("unexpected transform"));
+	AddErrorIfFalse((Hierarchy->GetGlobalTransform(0).GetTranslation() - FVector(1.f, 0.f, 0.f)).IsNearlyZero(), TEXT("unexpected transform"));
+	AddErrorIfFalse((Hierarchy->GetGlobalTransform(1).GetTranslation() - FVector(11.f, 13.f, 15.f)).IsNearlyZero(), TEXT("unexpected transform"));
+	AddErrorIfFalse((Hierarchy->GetGlobalTransform(2).GetTranslation() - FVector(15.f, 17.f, 19.f)).IsNearlyZero(), TEXT("unexpected transform"));
 
-	BoneHierarchy.ResetTransforms();
+	Hierarchy->ResetPoseToInitial(ERigElementType::Bone);
 	Unit.Mode = EControlRigModifyBoneMode::OverrideLocal;
 	InitAndExecute();
-	AddErrorIfFalse((BoneHierarchy.GetGlobalTransform(0).GetTranslation() - FVector(1.f, 0.f, 0.f)).IsNearlyZero(), TEXT("unexpected transform"));
-	AddErrorIfFalse((BoneHierarchy.GetGlobalTransform(1).GetTranslation() - FVector(11.f, 11.f, 12.f)).IsNearlyZero(), TEXT("unexpected transform"));
-	AddErrorIfFalse((BoneHierarchy.GetGlobalTransform(2).GetTranslation() - FVector(11.f, 11.f, 12.f)).IsNearlyZero(), TEXT("unexpected transform"));
+	AddErrorIfFalse((Hierarchy->GetGlobalTransform(0).GetTranslation() - FVector(1.f, 0.f, 0.f)).IsNearlyZero(), TEXT("unexpected transform"));
+	AddErrorIfFalse((Hierarchy->GetGlobalTransform(1).GetTranslation() - FVector(11.f, 11.f, 12.f)).IsNearlyZero(), TEXT("unexpected transform"));
+	AddErrorIfFalse((Hierarchy->GetGlobalTransform(2).GetTranslation() - FVector(11.f, 11.f, 12.f)).IsNearlyZero(), TEXT("unexpected transform"));
 
-	BoneHierarchy.ResetTransforms();
+	Hierarchy->ResetPoseToInitial(ERigElementType::Bone);
 	Unit.Mode = EControlRigModifyBoneMode::OverrideGlobal;
 	InitAndExecute();
-	AddErrorIfFalse((BoneHierarchy.GetGlobalTransform(0).GetTranslation() - FVector(1.f, 0.f, 0.f)).IsNearlyZero(), TEXT("unexpected transform"));
-	AddErrorIfFalse((BoneHierarchy.GetGlobalTransform(1).GetTranslation() - FVector(10.f, 11.f, 12.f)).IsNearlyZero(), TEXT("unexpected transform"));
-	AddErrorIfFalse((BoneHierarchy.GetGlobalTransform(2).GetTranslation() - FVector(10.f, 11.f, 12.f)).IsNearlyZero(), TEXT("unexpected transform"));
+	AddErrorIfFalse((Hierarchy->GetGlobalTransform(0).GetTranslation() - FVector(1.f, 0.f, 0.f)).IsNearlyZero(), TEXT("unexpected transform"));
+	AddErrorIfFalse((Hierarchy->GetGlobalTransform(1).GetTranslation() - FVector(10.f, 11.f, 12.f)).IsNearlyZero(), TEXT("unexpected transform"));
+	AddErrorIfFalse((Hierarchy->GetGlobalTransform(2).GetTranslation() - FVector(10.f, 11.f, 12.f)).IsNearlyZero(), TEXT("unexpected transform"));
 
-	BoneHierarchy.ResetTransforms();
+	Hierarchy->ResetPoseToInitial(ERigElementType::Bone);
 	Unit.Mode = EControlRigModifyBoneMode::AdditiveLocal;
 	Unit.Weight = 0.5f;
 	InitAndExecute();
-	AddErrorIfFalse((BoneHierarchy.GetGlobalTransform(0).GetTranslation() - FVector(1.f, 0.f, 0.f)).IsNearlyZero(), TEXT("unexpected transform"));
-	AddErrorIfFalse((BoneHierarchy.GetGlobalTransform(1).GetTranslation() - FVector(6.f, 7.5f, 9.f)).IsNearlyZero(), TEXT("unexpected transform"));
-	AddErrorIfFalse((BoneHierarchy.GetGlobalTransform(2).GetTranslation() - FVector(10.f, 11.5f, 13.f)).IsNearlyZero(), TEXT("unexpected transform"));
+	AddErrorIfFalse((Hierarchy->GetGlobalTransform(0).GetTranslation() - FVector(1.f, 0.f, 0.f)).IsNearlyZero(), TEXT("unexpected transform"));
+	AddErrorIfFalse((Hierarchy->GetGlobalTransform(1).GetTranslation() - FVector(6.f, 7.5f, 9.f)).IsNearlyZero(), TEXT("unexpected transform"));
+	AddErrorIfFalse((Hierarchy->GetGlobalTransform(2).GetTranslation() - FVector(10.f, 11.5f, 13.f)).IsNearlyZero(), TEXT("unexpected transform"));
 
 
 	return true;
diff --git a/Engine/Plugins/Experimental/ControlRig/Source/ControlRig/Private/Units/Highlevel/Hierarchy/RigUnit_ModifyBoneTransforms.h b/Engine/Plugins/Experimental/ControlRig/Source/ControlRig/Private/Units/Highlevel/Hierarchy/RigUnit_ModifyBoneTransforms.h
--- a/Engine/Plugins/Experimental/ControlRig/Source/ControlRig/Private/Units/Highlevel/Hierarchy/RigUnit_ModifyBoneTransforms.h
+++ b/Engine/Plugins/Experimental/ControlRig/Source/ControlRig/Private/Units/Highlevel/Hierarchy/RigUnit_ModifyBoneTransforms.h
@@ -76,13 +76,9 @@ struct CONTROLRIG_API FRigUnit_ModifyBoneTransforms : public FRigUnit_HighlevelB
 
 				if (Mode == EControlRigModifyBoneMode::OverrideLocal)
 				{
-					if (const FRigHierarchyContainer* Container = (const FRigHierarchyContainer*)InUserContext)
+					if (const URigHierarchy* Hierarchy = (const URigHierarchy*)InUserContext)
 					{
-						int32 BoneIndex = Container->BoneHierarchy.GetIndex(BoneToModify[Index].Bone);
-						if (BoneIndex != INDEX_NONE)
-						{
-							return Container->BoneHierarchy[BoneIndex].GetParentElementKey();
-						}
+						return Hierarchy->GetFirstParent(FRigElementKey(BoneToModify[Index].Bone, ERigElementType::Bone));
 					}
 				}
 			}
diff --git a/Engine/Plugins/Experimental/ControlRig/Source/ControlRig/Private/Units/Highlevel/Hierarchy/RigUnit_MultiFABRIK.cpp b/Engine/Plugins/Experimental/ControlRig/Source/ControlRig/Private/Units/Highlevel/Hierarchy/RigUnit_MultiFABRIK.cpp
--- a/Engine/Plugins/Experimental/ControlRig/Source/ControlRig/Private/Units/Highlevel/Hierarchy/RigUnit_MultiFABRIK.cpp
+++ b/Engine/Plugins/Experimental/ControlRig/Source/ControlRig/Private/Units/Highlevel/Hierarchy/RigUnit_MultiFABRIK.cpp
@@ -215,7 +215,7 @@ namespace MultiFABRIK
 FRigUnit_MultiFABRIK_Execute()
 {
     DECLARE_SCOPE_HIERARCHICAL_COUNTER_RIGUNIT()
-	FRigBoneHierarchy* Hierarchy = ExecuteContext.GetBones();
+	URigHierarchy* Hierarchy = ExecuteContext.Hierarchy;
 	if (Hierarchy == nullptr)
 	{
 		return;
@@ -239,15 +239,15 @@ FRigUnit_MultiFABRIK_Execute()
 		if(BoneTree.Num() == 0)
 		{
 			// verify the chain
-			const int32 RootBoneIndex = Hierarchy->GetIndex(RootBone);
+			const FRigElementKey RootKey(RootBone, ERigElementType::Bone);
+			const int32 RootBoneIndex = Hierarchy->GetIndex(RootKey);
 			if (RootBoneIndex != INDEX_NONE)
 			{
-				const FRigBone& RootRigBone = (*Hierarchy)[RootBoneIndex];
-
 				// fill up all effector indices
 				for (int32 Index = 0; Index < Effectors.Num(); ++Index)
 				{
-					EffectorIndices.Add(FCachedRigElement(Effectors[Index].Bone, Hierarchy));
+					const FRigElementKey EffectorKey(Effectors[Index].Bone, ERigElementType::Bone);
+					EffectorIndices.Add(FCachedRigElement(EffectorKey, Hierarchy));
 				}
 
 				// go up to until you meet root
@@ -265,13 +265,23 @@ FRigUnit_MultiFABRIK_Execute()
 						, ParentIndex(InParentIndex)
 					{}
 
-					FLocalBoneData(const FRigBone& InRigBone)
-						: BoneName(InRigBone.Name)
-						, BoneIndex(InRigBone.Index)
-						, ParentBoneName(InRigBone.ParentName)
-						, ParentIndex(InRigBone.ParentIndex)
-					{}
-					
+					FLocalBoneData(URigHierarchy* InHierarchy, int32 InBoneIndex)
+					: BoneName(NAME_None)
+					, BoneIndex(InBoneIndex)
+					, ParentBoneName(NAME_None)
+				    , ParentIndex(INDEX_NONE)
+					{
+						FRigBaseElement* BoneElement = InHierarchy->Get(InBoneIndex);
+						BoneName = BoneElement->GetName();
+						BoneIndex = BoneElement->GetIndex();
+
+						if(FRigBaseElement* ParentBoneElement = InHierarchy->GetFirstParent(BoneElement))
+						{
+							ParentBoneName = ParentBoneElement->GetName();
+							ParentIndex = ParentBoneElement->GetIndex();
+						}
+					}
+
 				};
 
 				TMap<FName, int32> ExistingBoneList;
@@ -287,16 +297,15 @@ FRigUnit_MultiFABRIK_Execute()
 						int32 CurrentIndex = EffectorIndex;
 						do 
 						{
-							const FRigBone& RigBone = (*Hierarchy)[CurrentIndex];
-							BoneList.Insert(FLocalBoneData(RigBone), 0);
-							CurrentIndex = RigBone.ParentIndex;
-						} while (CurrentIndex != RootRigBone.Index && CurrentIndex != INDEX_NONE);
+							BoneList.Insert(FLocalBoneData(Hierarchy, CurrentIndex), 0);
+							CurrentIndex = Hierarchy->GetIndex(Hierarchy->GetFirstParent(Hierarchy->GetKey(CurrentIndex)));
+						} while (CurrentIndex != RootBoneIndex && CurrentIndex != INDEX_NONE);
 
 						// if we haven't got to root, this is not valid chain
-						if (CurrentIndex == RootRigBone.Index)
+						if (CurrentIndex == RootBoneIndex)
 						{
 							// add root
-							BoneList.Insert(FLocalBoneData(RootRigBone), 0);
+							BoneList.Insert(FLocalBoneData(Hierarchy, RootBoneIndex), 0);
 
 							// go through all bone list collected
 							for (const FLocalBoneData& Bone : BoneList)
@@ -311,8 +320,8 @@ FRigUnit_MultiFABRIK_Execute()
 									Data = &BoneTree[TreeIndex];
 
 									Data->BoneName = Bone.BoneName;
-									Data->BoneIndex = FCachedRigElement(Bone.BoneName, Hierarchy);
-									Data->ParentIndex = FCachedRigElement(Bone.ParentBoneName, Hierarchy);
+									Data->BoneIndex = FCachedRigElement(FRigElementKey(Bone.BoneName, ERigElementType::Bone), Hierarchy);
+									Data->ParentIndex = FCachedRigElement(FRigElementKey(Bone.ParentBoneName, ERigElementType::Bone), Hierarchy);
 
 									// set bone length 
 									if (Data->ParentIndex.IsValid())
@@ -464,16 +473,12 @@ FRigUnit_MultiFABRIK_Execute()
 
 IMPLEMENT_RIGUNIT_AUTOMATION_TEST(FRigUnit_MultiFABRIK)
 {
-	BoneHierarchy.Add(TEXT("Root"), NAME_None, ERigBoneType::User, FTransform(FVector(1.f, 0.f, 0.f)));
-	// add  2 chains of length 2,2 
-	BoneHierarchy.Add(TEXT("Chain1_0"), TEXT("Root"), ERigBoneType::User, FTransform(FVector(1.f, 2.f, 0.f)));
-	BoneHierarchy.Add(TEXT("Chain1_1"), TEXT("Chain1_0"), ERigBoneType::User, FTransform(FVector(3.f, 2.f, 0.f)));
-	//length 3,3
-	BoneHierarchy.Add(TEXT("Chain2_0"), TEXT("Root"), ERigBoneType::User, FTransform(FVector(-2.f, 0.f, 0.f)));
-	BoneHierarchy.Add(TEXT("Chain2_1"), TEXT("Chain2_0"), ERigBoneType::User, FTransform(FVector(-2.f, 3.f, 0.f)));
-
-	BoneHierarchy.Initialize();
-	Unit.ExecuteContext.Hierarchy = &HierarchyContainer;
+	const FRigElementKey Root = Controller->AddBone(TEXT("Root"), FRigElementKey(), FTransform(FVector(1.f, 0.f, 0.f)), true, ERigBoneType::User);
+	const FRigElementKey Chain1_0 = Controller->AddBone(TEXT("Chain1_0"), Root, FTransform(FVector(1.f, 2.f, 0.f)), true, ERigBoneType::User);
+	const FRigElementKey Chain1_1 = Controller->AddBone(TEXT("Chain1_1"), Chain1_0, FTransform(FVector(3.f, 2.f, 0.f)), true, ERigBoneType::User);
+	const FRigElementKey Chain2_0 = Controller->AddBone(TEXT("Chain2_0"), Root, FTransform(FVector(-2.f, 0.f, 0.f)), true, ERigBoneType::User);
+	const FRigElementKey Chain2_1 = Controller->AddBone(TEXT("Chain2_1"), Chain2_0, FTransform(FVector(-2.f, 3.f, 0.f)), true, ERigBoneType::User);
+	Unit.ExecuteContext.Hierarchy = Hierarchy;
 
 	// first validation test
 	// make sure this doesn't crash
@@ -493,8 +498,8 @@ IMPLEMENT_RIGUNIT_AUTOMATION_TEST(FRigUnit_MultiFABRIK)
 	Unit.bPropagateToChildren = true;
 
 	InitAndExecute();
-	AddErrorIfFalse(BoneHierarchy.GetGlobalTransform(TEXT("Chain1_1")).GetTranslation().Equals(FVector(3.f, 2.f, 0.f)), TEXT("unexpected transform"));
-	AddErrorIfFalse(BoneHierarchy.GetGlobalTransform(TEXT("Chain2_1")).GetTranslation().Equals(FVector(-2.f, 3.f, 0.f)), TEXT("unexpected transform"));
+	AddErrorIfFalse(Hierarchy->GetGlobalTransform(Chain1_1).GetTranslation().Equals(FVector(3.f, 2.f, 0.f)), TEXT("unexpected transform"));
+	AddErrorIfFalse(Hierarchy->GetGlobalTransform(Chain2_1).GetTranslation().Equals(FVector(-2.f, 3.f, 0.f)), TEXT("unexpected transform"));
 
 	// root is (1, 0, 0)
 	Unit.Effectors[0].Bone = TEXT("Chain1_1");
@@ -504,8 +509,8 @@ IMPLEMENT_RIGUNIT_AUTOMATION_TEST(FRigUnit_MultiFABRIK)
 	Unit.bPropagateToChildren = true;
 
 	InitAndExecute();
-	AddErrorIfFalse(BoneHierarchy.GetGlobalTransform(TEXT("Chain1_1")).GetTranslation().Equals(FVector(4.f, 0.f, 0.f), 0.1f), TEXT("unexpected transform"));
-	AddErrorIfFalse(BoneHierarchy.GetGlobalTransform(TEXT("Chain2_1")).GetTranslation().Equals(FVector(0.f, -5.f, 0.f), 0.1f), TEXT("unexpected transform"));
+	AddErrorIfFalse(Hierarchy->GetGlobalTransform(Chain1_1).GetTranslation().Equals(FVector(4.f, 0.f, 0.f), 0.1f), TEXT("unexpected transform"));
+	AddErrorIfFalse(Hierarchy->GetGlobalTransform(Chain2_1).GetTranslation().Equals(FVector(0.f, -5.f, 0.f), 0.1f), TEXT("unexpected transform"));
 	return true;
 }
 #endif
\ No newline at end of file
diff --git a/Engine/Plugins/Experimental/ControlRig/Source/ControlRig/Private/Units/Highlevel/Hierarchy/RigUnit_SpringIK.cpp b/Engine/Plugins/Experimental/ControlRig/Source/ControlRig/Private/Units/Highlevel/Hierarchy/RigUnit_SpringIK.cpp
--- a/Engine/Plugins/Experimental/ControlRig/Source/ControlRig/Private/Units/Highlevel/Hierarchy/RigUnit_SpringIK.cpp
+++ b/Engine/Plugins/Experimental/ControlRig/Source/ControlRig/Private/Units/Highlevel/Hierarchy/RigUnit_SpringIK.cpp
@@ -6,7 +6,7 @@
 FRigUnit_SpringIK_Execute()
 {
     DECLARE_SCOPE_HIERARCHICAL_COUNTER_RIGUNIT()
-	FRigBoneHierarchy* Hierarchy = ExecuteContext.GetBones();
+	URigHierarchy* Hierarchy = ExecuteContext.Hierarchy;
 	if (Hierarchy == nullptr)
 	{
 		return;
@@ -29,10 +29,10 @@ FRigUnit_SpringIK_Execute()
 		Simulation.Reset();
 		Simulation.TimeStep = 1.f / 60.f;
 
-		int32 EndBoneIndex = Hierarchy->GetIndex(EndBone);
+		int32 EndBoneIndex = Hierarchy->GetIndex(FRigElementKey(EndBone, ERigElementType::Bone));
 		if (EndBoneIndex != INDEX_NONE)
 		{
-			int32 StartBoneIndex = Hierarchy->GetIndex(StartBone);
+			int32 StartBoneIndex = Hierarchy->GetIndex(FRigElementKey(StartBone, ERigElementType::Bone));
 			if (StartBoneIndex == EndBoneIndex)
 			{
 				return;
@@ -40,12 +40,13 @@ FRigUnit_SpringIK_Execute()
 
 			while (EndBoneIndex != INDEX_NONE)
 			{
-				CachedBones.Add(FCachedRigElement((*Hierarchy)[EndBoneIndex].Name, Hierarchy));
+				const FRigElementKey Key(EndBone, ERigElementType::Bone);
+				CachedBones.Add(FCachedRigElement(Key, Hierarchy));
 				if (EndBoneIndex == StartBoneIndex)
 				{
 					break;
 				}
-				EndBoneIndex = (*Hierarchy)[EndBoneIndex].ParentIndex;
+				EndBoneIndex = Hierarchy->GetIndex(Hierarchy->GetFirstParent(Key));
 			}
 		}
 
@@ -114,7 +115,8 @@ FRigUnit_SpringIK_Execute()
 			}
 		}
 
-		CachedPoleVector = FCachedRigElement(PoleVectorSpace, Hierarchy);
+		const FRigElementKey CachedPoleVectorKey(PoleVectorSpace, ERigElementType::Bone);
+		CachedPoleVector = FCachedRigElement(CachedPoleVectorKey, Hierarchy);
 	}
 
 	if (CachedBones.Num() < 3)
@@ -261,7 +263,7 @@ FRigUnit_SpringIK_Execute()
 
 		if(bLimitLocalPosition)
 		{
-			int32 ParentIndex = (*Hierarchy)[CachedBones[PointIndex]].ParentIndex;
+			int32 ParentIndex = Hierarchy->GetIndex(Hierarchy->GetFirstParent(CachedBones[PointIndex].GetKey()));
 			if (ParentIndex != INDEX_NONE)
 			{
 				FTransform InitialTransform = Hierarchy->GetInitialGlobalTransform(CachedBones[PointIndex]);
diff --git a/Engine/Plugins/Experimental/ControlRig/Source/ControlRig/Private/Units/Highlevel/Hierarchy/RigUnit_TransformConstraint.cpp b/Engine/Plugins/Experimental/ControlRig/Source/ControlRig/Private/Units/Highlevel/Hierarchy/RigUnit_TransformConstraint.cpp
--- a/Engine/Plugins/Experimental/ControlRig/Source/ControlRig/Private/Units/Highlevel/Hierarchy/RigUnit_TransformConstraint.cpp
+++ b/Engine/Plugins/Experimental/ControlRig/Source/ControlRig/Private/Units/Highlevel/Hierarchy/RigUnit_TransformConstraint.cpp
@@ -32,7 +32,7 @@ FRigUnit_TransformConstraintPerItem_Execute()
 		ConstraintData.Reset();
 		ConstraintDataToTargets.Reset();
 
-		FRigHierarchyContainer* Hierarchy = ExecuteContext.Hierarchy;
+		URigHierarchy* Hierarchy = ExecuteContext.Hierarchy;
 		if(Hierarchy)
 		{
 			if (Item.IsValid())
@@ -40,20 +40,20 @@ FRigUnit_TransformConstraintPerItem_Execute()
 				const int32 TargetNum = Targets.Num();
 				if (TargetNum > 0)
 				{
-					const FTransform SourceTransform = bUseInitialTransforms ? Hierarchy->GetInitialGlobalTransform(Item) : Hierarchy->GetGlobalTransform(Item);
+					const FTransform SourceTransform = bUseInitialTransforms ? Hierarchy->GetGlobalTransform(Item, true) : Hierarchy->GetGlobalTransform(Item, false);
 					FTransform InputBaseTransform =
 						bUseInitialTransforms ?
 						UtilityHelpers::GetBaseTransformByMode(
 							BaseTransformSpace,
-							[Hierarchy](const FRigElementKey& Item) { return Hierarchy->GetInitialGlobalTransform(Item); },
-							Hierarchy->GetParentKey(Item),
+							[Hierarchy](const FRigElementKey& Item) { return Hierarchy->GetGlobalTransform(Item, true); },
+							Hierarchy->GetFirstParent(Item),
 							BaseItem,
 							BaseTransform
 						) :
 						UtilityHelpers::GetBaseTransformByMode(
 							BaseTransformSpace,
-							[Hierarchy](const FRigElementKey& Item) { return Hierarchy->GetGlobalTransform(Item); },
-							Hierarchy->GetParentKey(Item),
+							[Hierarchy](const FRigElementKey& Item) { return Hierarchy->GetGlobalTransform(Item, false); },
+							Hierarchy->GetFirstParent(Item),
 							BaseItem,
 							BaseTransform
 						);
@@ -100,7 +100,7 @@ FRigUnit_TransformConstraintPerItem_Execute()
 	}
 	else if (Context.State == EControlRigState::Update)
 	{
-		FRigHierarchyContainer* Hierarchy = ExecuteContext.Hierarchy;
+		URigHierarchy* Hierarchy = ExecuteContext.Hierarchy;
 		if (Hierarchy)
 		{
 			if ((ConstraintData.Num() != Targets.Num()))
@@ -126,7 +126,7 @@ FRigUnit_TransformConstraintPerItem_Execute()
 					}
 
 					FTransform InputBaseTransform = UtilityHelpers::GetBaseTransformByMode(BaseTransformSpace, [Hierarchy](const FRigElementKey& Item) { return Hierarchy->GetGlobalTransform(Item); },
-							Hierarchy->GetParentKey(Item), BaseItem, BaseTransform);
+							Hierarchy->GetFirstParent(Item), BaseItem, BaseTransform);
 
 					FTransform SourceTransform = Hierarchy->GetGlobalTransform(Item);
 
diff --git a/Engine/Plugins/Experimental/ControlRig/Source/ControlRig/Private/Units/Highlevel/Hierarchy/RigUnit_TransformConstraint.h b/Engine/Plugins/Experimental/ControlRig/Source/ControlRig/Private/Units/Highlevel/Hierarchy/RigUnit_TransformConstraint.h
--- a/Engine/Plugins/Experimental/ControlRig/Source/ControlRig/Private/Units/Highlevel/Hierarchy/RigUnit_TransformConstraint.h
+++ b/Engine/Plugins/Experimental/ControlRig/Source/ControlRig/Private/Units/Highlevel/Hierarchy/RigUnit_TransformConstraint.h
@@ -70,13 +70,9 @@ struct CONTROLRIG_API FRigUnit_TransformConstraint : public FRigUnit_HighlevelBa
 
 			if (BaseTransformSpace == ETransformSpaceMode::LocalSpace)
 			{
-				if (const FRigHierarchyContainer* Container = (const FRigHierarchyContainer*)InUserContext)
+				if (const URigHierarchy* Hierarchy = (const URigHierarchy*)InUserContext)
 				{
-					int32 BoneIndex = Container->BoneHierarchy.GetIndex(Bone);
-					if (BoneIndex != INDEX_NONE)
-					{
-						return Container->BoneHierarchy[BoneIndex].GetParentElementKey();
-					}
+					return Hierarchy->GetFirstParent(FRigElementKey(Bone, ERigElementType::Bone));
 				}
 			}
 		}
@@ -139,9 +135,9 @@ struct CONTROLRIG_API FRigUnit_TransformConstraintPerItem : public FRigUnit_High
 
 			if (BaseTransformSpace == ETransformSpaceMode::LocalSpace)
 			{
-				if (const FRigHierarchyContainer* Container = (const FRigHierarchyContainer*)InUserContext)
+				if (const URigHierarchy* Hierarchy = (const URigHierarchy*)InUserContext)
 				{
-					return Container->GetParentKey(Item);
+					return Hierarchy->GetFirstParent(Item);
 				}
 			}
 		}
diff --git a/Engine/Plugins/Experimental/ControlRig/Source/ControlRig/Private/Units/Highlevel/Hierarchy/RigUnit_TwistBones.cpp b/Engine/Plugins/Experimental/ControlRig/Source/ControlRig/Private/Units/Highlevel/Hierarchy/RigUnit_TwistBones.cpp
--- a/Engine/Plugins/Experimental/ControlRig/Source/ControlRig/Private/Units/Highlevel/Hierarchy/RigUnit_TwistBones.cpp
+++ b/Engine/Plugins/Experimental/ControlRig/Source/ControlRig/Private/Units/Highlevel/Hierarchy/RigUnit_TwistBones.cpp
@@ -38,7 +38,7 @@ FRigUnit_TwistBones_Execute()
 FRigUnit_TwistBonesPerItem_Execute()
 {
     DECLARE_SCOPE_HIERARCHICAL_COUNTER_RIGUNIT()
-	FRigHierarchyContainer* Hierarchy = ExecuteContext.Hierarchy;
+	URigHierarchy* Hierarchy = ExecuteContext.Hierarchy;
 	if (Hierarchy == nullptr)
 	{
 		return;
@@ -134,23 +134,6 @@ FRigUnit_TwistBonesPerItem_Execute()
 
 		OffsetQuat = FQuat::Slerp(ItemTransforms[Index].GetRotation(), OffsetQuat * ItemTransforms[Index].GetRotation(), FMath::Clamp<float>(Weight, 0.f, 1.f));
 		ItemTransforms[Index].SetRotation(OffsetQuat);
-		Hierarchy->SetGlobalTransform(CachedItems[Index], ItemTransforms[Index], false);
-
-		if (bPropagateToChildren)
-		{
-			if (CachedItems[Index].GetKey().Type == ERigElementType::Bone)
-			{
-				const TArray<int32> Dependents = Hierarchy->BoneHierarchy[CachedItems[Index]].Dependents;
-				for (int32 DependentIndex = 0; DependentIndex < Dependents.Num(); ++DependentIndex)
-				{
-					const FRigBone& Dependent = Hierarchy->BoneHierarchy[Dependents[DependentIndex]];
-					if (!CachedItems.Contains(FCachedRigElement(Dependent.GetElementKey(), Hierarchy)))
-					{
-						Hierarchy->BoneHierarchy.RecalculateGlobalTransform(Dependent.Index);
-						Hierarchy->BoneHierarchy.PropagateTransform(Dependent.Index);
-					}
-				}
-			}
-		}
+		Hierarchy->SetGlobalTransform(CachedItems[Index], ItemTransforms[Index], bPropagateToChildren);
 	}
 }
diff --git a/Engine/Plugins/Experimental/ControlRig/Source/ControlRig/Private/Units/Simulation/RigUnit_PointSimulation.cpp b/Engine/Plugins/Experimental/ControlRig/Source/ControlRig/Private/Units/Simulation/RigUnit_PointSimulation.cpp
--- a/Engine/Plugins/Experimental/ControlRig/Source/ControlRig/Private/Units/Simulation/RigUnit_PointSimulation.cpp
+++ b/Engine/Plugins/Experimental/ControlRig/Source/ControlRig/Private/Units/Simulation/RigUnit_PointSimulation.cpp
@@ -6,7 +6,7 @@
 FRigUnit_PointSimulation_Execute()
 {
     DECLARE_SCOPE_HIERARCHICAL_COUNTER_RIGUNIT()
-	FRigBoneHierarchy* Hierarchy = ExecuteContext.GetBones();
+	URigHierarchy* Hierarchy = ExecuteContext.Hierarchy;
 
 	float DeltaTime = Context.DeltaTime;
 
@@ -158,7 +158,8 @@ FRigUnit_PointSimulation_Execute()
 
 		for (int32 TargetIndex = 0; TargetIndex < BoneTargets.Num(); TargetIndex++)
 		{
-			if (BoneIndices[TargetIndex].UpdateCache(BoneTargets[TargetIndex].Bone, Hierarchy))
+			const FRigElementKey Key(BoneTargets[TargetIndex].Bone, ERigElementType::Bone);
+			if (BoneIndices[TargetIndex].UpdateCache(Key, Hierarchy))
 			{
 				const FRigUnit_PointSimulation_BoneTarget& BoneTarget = BoneTargets[TargetIndex];
 				FTransform Transform = Hierarchy->GetGlobalTransform(BoneIndices[TargetIndex]);
@@ -170,12 +171,12 @@ FRigUnit_PointSimulation_Execute()
 
 				if (bLimitLocalPosition)
 				{
-					int32 ParentIndex = (*Hierarchy)[BoneIndices[TargetIndex]].ParentIndex;
-					if (ParentIndex != INDEX_NONE)
+					FRigElementKey ParentKey = Hierarchy->GetFirstParent(BoneIndices[TargetIndex].GetKey());
+					if (ParentKey.IsValid())
 					{
 						FTransform InitialTransform = Hierarchy->GetInitialGlobalTransform(BoneIndices[TargetIndex]);
-						FTransform ParentInitialTransform = Hierarchy->GetInitialGlobalTransform(ParentIndex);
-						FTransform ParentTransform = Hierarchy->GetGlobalTransform(ParentIndex);
+						FTransform ParentInitialTransform = Hierarchy->GetGlobalTransform(ParentKey, true);
+						FTransform ParentTransform = Hierarchy->GetGlobalTransform(ParentKey);
 						float ExpectedDistance = (InitialTransform.GetLocation() - ParentInitialTransform.GetLocation()).Size();
 						if (ExpectedDistance > SMALL_NUMBER)
 						{
diff --git a/Engine/Plugins/Experimental/ControlRig/Source/ControlRig/Private/Validation/ControlRigNumericalValidationPass.cpp b/Engine/Plugins/Experimental/ControlRig/Source/ControlRig/Private/Validation/ControlRigNumericalValidationPass.cpp
--- a/Engine/Plugins/Experimental/ControlRig/Source/ControlRig/Private/Validation/ControlRigNumericalValidationPass.cpp
+++ b/Engine/Plugins/Experimental/ControlRig/Source/ControlRig/Private/Validation/ControlRigNumericalValidationPass.cpp
@@ -79,15 +79,14 @@ void UControlRigNumericalValidationPass::OnEvent(UControlRig* InControlRig, cons
 			if(Key.Type == ERigElementType::Curve)
 			{
 				float A = PoseElement.CurveValue;
-				const FRigCurve& Curve = InControlRig->GetCurveContainer()[ElementIndex];
-				float B = Curve.Value;
+				float B = InControlRig->GetHierarchy()->GetCurveValue(ElementIndex);
 
 				if (FMath::Abs(A - B) > CurvePrecision + SMALL_NUMBER)
 				{
 					FString EventNameADisplayString = InContext->GetDisplayNameForEvent(EventNameA);
 					FString EventNameBDisplayString = InContext->GetDisplayNameForEvent(EventNameB);
 					FString Message = FString::Printf(TEXT("Values don't match between %s and %s."), *EventNameADisplayString, *EventNameBDisplayString);
-					InContext->Report(EMessageSeverity::Warning, Curve.GetElementKey(), Message);
+					InContext->Report(EMessageSeverity::Warning, Key, Message);
 				}
 				continue;
 			}
diff --git a/Engine/Plugins/Experimental/ControlRig/Source/ControlRig/Public/ControlRig.h b/Engine/Plugins/Experimental/ControlRig/Source/ControlRig/Public/ControlRig.h
--- a/Engine/Plugins/Experimental/ControlRig/Source/ControlRig/Public/ControlRig.h
+++ b/Engine/Plugins/Experimental/ControlRig/Source/ControlRig/Public/ControlRig.h
@@ -9,7 +9,7 @@
 #include "Templates/SubclassOf.h"
 #include "ControlRigDefines.h"
 #include "ControlRigGizmoLibrary.h"
-#include "Rigs/RigHierarchyContainer.h"
+#include "Rigs/RigHierarchy.h"
 #include "Units/RigUnitContext.h"
 #include "Animation/NodeMappingProviderInterface.h"
 #include "Units/RigUnit.h"
@@ -51,13 +51,13 @@ class CONTROLRIG_API UControlRig : public UObject, public INodeMappingProviderIn
 public:
 
 	/** Bindable event for external objects to contribute to / filter a control value */
-	DECLARE_EVENT_ThreeParams(UControlRig, FFilterControlEvent, UControlRig*, const FRigControl&, FRigControlValue&);
+	DECLARE_EVENT_ThreeParams(UControlRig, FFilterControlEvent, UControlRig*, FRigControlElement*, FRigControlValue&);
 
 	/** Bindable event for external objects to be notified of Control changes */
-	DECLARE_EVENT_ThreeParams(UControlRig, FControlModifiedEvent, UControlRig*, const FRigControl&, const FRigControlModifiedContext&);
+	DECLARE_EVENT_ThreeParams(UControlRig, FControlModifiedEvent, UControlRig*, FRigControlElement*, const FRigControlModifiedContext&);
 
 	/** Bindable event for external objects to be notified that a Control is Selected */
-	DECLARE_EVENT_ThreeParams(UControlRig, FControlSelectedEvent, UControlRig*, const FRigControl&, bool);
+	DECLARE_EVENT_ThreeParams(UControlRig, FControlSelectedEvent, UControlRig*, FRigControlElement*, bool);
 
 	static const FName OwnerComponent;
 
@@ -184,72 +184,28 @@ public:
 		return ObjectName;
 	}
 
-	FRigHierarchyContainer* GetHierarchy()
-	{
-		return &Hierarchy;
-	}
-
-	FRigBoneHierarchy& GetBoneHierarchy()
-	{
-		return Hierarchy.BoneHierarchy;
-	}
-
-	FRigSpaceHierarchy& GetSpaceHierarchy()
-	{
-		return Hierarchy.SpaceHierarchy;
-	}
-
-	FRigControlHierarchy& GetControlHierarchy()
-	{
-		return Hierarchy.ControlHierarchy;
-	}
-
-	FRigCurveContainer& GetCurveContainer()
+	UFUNCTION(BlueprintCallable, Category = "Control Rig")
+	FORCEINLINE_DEBUGGABLE URigHierarchy* GetHierarchy()
 	{
-		return Hierarchy.CurveContainer;
+		if(DynamicHierarchy != nullptr && DynamicHierarchy->GetOuter() != this)
+		{
+			DynamicHierarchy = nullptr;
+		}
+		
+		if(DynamicHierarchy == nullptr)
+		{
+			DynamicHierarchy = NewObject<URigHierarchy>(this);
+		}
+		return DynamicHierarchy;
 	}
 
-	/** Evaluate another animation ControlRig */
-	UFUNCTION(BlueprintCallable, Category = "Control Rig")
-	FTransform GetGlobalTransform(const FRigElementKey& InKey) const;
-
-	/** Evaluate another animation ControlRig */
-	UFUNCTION(BlueprintCallable, Category = "Control Rig")
-	void SetGlobalTransform(const FRigElementKey& InKey, const FTransform& InTransform, bool bPropagateTransform = true);
-
-	/** Evaluate another animation ControlRig */
-	FTransform GetGlobalTransform(const FName& BoneName) const;
-
-	/** Evaluate another animation ControlRig */
-	void SetGlobalTransform(const FName& BoneName, const FTransform& InTransform, bool bPropagateTransform = true) ;
-
-	/** Evaluate another animation ControlRig */
-	FTransform GetGlobalTransform(const int32 BoneIndex) const;
-
-	/** Evaluate another animation ControlRig */
-	void SetGlobalTransform(const int32 BoneIndex, const FTransform& InTransform, bool bPropagateTransform = true) ;
-
-	/** Evaluate another animation ControlRig */
-	UFUNCTION(BlueprintCallable, Category = "Control Rig")
-	float GetCurveValue(const FName& CurveName) const;
-
-	/** Evaluate another animation ControlRig */
-	UFUNCTION(BlueprintCallable, Category = "Control Rig")
-	void SetCurveValue(const FName& CurveName, const float CurveValue);
-
-	/** Evaluate another animation ControlRig */
-	float GetCurveValue(const int32 CurveIndex) const;
-
-	/** Evaluate another animation ControlRig */
-	void SetCurveValue(const int32 CurveIndex, const float CurveValue);
-
 #if WITH_EDITOR
 
 	// called after post reinstance when compilng blueprint by Sequencer
 	void PostReinstanceCallback(const UControlRig* Old);
 
-
 #endif // WITH_EDITOR
+	
 	// BEGIN UObject interface
 	static void AddReferencedObjects(UObject* InThis, FReferenceCollector& Collector);
 	virtual void BeginDestroy() override;
@@ -287,12 +243,8 @@ public:
 	/** Data Source Registry Getter */
 	UAnimationDataSourceRegistry* GetDataSourceRegistry();
 
-	virtual const TArray<FRigSpace>& AvailableSpaces() const;
-	virtual FRigSpace* FindSpace(const FName& InSpaceName);
-	virtual FTransform GetSpaceGlobalTransform(const FName& InSpaceName);
-	virtual bool SetSpaceGlobalTransform(const FName& InSpaceName, const FTransform& InTransform);
-	virtual const TArray<FRigControl>& AvailableControls() const;
-	virtual FRigControl* FindControl(const FName& InControlName);
+	virtual TArray<FRigControlElement*> AvailableControls() const;
+	virtual FRigControlElement* FindControl(const FName& InControlName) const;
 	virtual bool ShouldApplyLimits() const { return !bSetupModeEnabled; }
 	virtual bool IsSetupModeEnabled() const { return bSetupModeEnabled; }
 	virtual FTransform SetupControlFromGlobalTransform(const FName& InControlName, const FTransform& InGlobalTransform);
@@ -307,30 +259,25 @@ public:
 	}
 
 	// Returns the value of a Control
-	FORCEINLINE_DEBUGGABLE const FRigControlValue& GetControlValue(const FName& InControlName)
+	FORCEINLINE_DEBUGGABLE FRigControlValue GetControlValue(const FName& InControlName)
 	{
-		const FRigControl* Control = FindControl(InControlName);
-		check(Control);
-		return Control->Value;
+		const FRigElementKey Key(InControlName, ERigElementType::Control);
+		return DynamicHierarchy->GetControlValue(Key);
 	}
 
 	// Sets the relative value of a Control
 	FORCEINLINE_DEBUGGABLE virtual void SetControlValueImpl(const FName& InControlName, const FRigControlValue& InValue, bool bNotify = true,
 		const FRigControlModifiedContext& Context = FRigControlModifiedContext())
 	{
-		FRigControl* Control = FindControl(InControlName);
-		check(Control);
-
-		Control->Value = InValue;
-
-		if (ShouldApplyLimits())
-		{
-			Control->ApplyLimits(Control->Value);
-		}
+		const FRigElementKey Key(InControlName, ERigElementType::Control);
+		DynamicHierarchy->SetControlValue(Key, InValue);
 
 		if (bNotify && OnControlModified.IsBound())
 		{
-			OnControlModified.Broadcast(this, *Control, Context);
+			if(FRigControlElement* ControlElement = FindControl(InControlName))
+			{
+				OnControlModified.Broadcast(this, ControlElement, Context);
+			}
 		}
 	}
 
@@ -355,10 +302,9 @@ public:
 	virtual void SetControlLocalTransform(const FName& InControlName, const FTransform& InLocalTransform, bool bNotify = true, const FRigControlModifiedContext& Context = FRigControlModifiedContext());
 	virtual FTransform GetControlLocalTransform(const FName& InControlName) ;
 
-	virtual bool SetControlSpace(const FName& InControlName, const FName& InSpaceName);
 	virtual UControlRigGizmoLibrary* GetGizmoLibrary() const;
 	virtual void CreateRigControlsForCurveContainer();
-	virtual void GetControlsInOrder(TArray<FRigControl>& SortedControls) const;
+	virtual void GetControlsInOrder(TArray<FRigControlElement*>& SortedControls) const;
 
 	virtual void SelectControl(const FName& InControlName, bool bSelect = true);
 	virtual bool ClearControlSelection();
@@ -395,7 +341,7 @@ public:
 	// selection changes coming from the manipulated subject.
 	FControlSelectedEvent& ControlSelected() { return OnControlSelected; }
 
-	bool IsCurveControl(const FRigControl* InRigControl) const;
+	bool IsCurveControl(const FRigControlElement* InControlElement) const;
 
 	DECLARE_EVENT_ThreeParams(UControlRig, FControlRigExecuteEvent, class UControlRig*, const EControlRigState, const FName&);
 	FControlRigExecuteEvent& OnInitialized_AnyThread() { return InitializedEvent; }
@@ -422,7 +368,7 @@ private:
 	URigVM* VM;
 
 	UPROPERTY()
-	FRigHierarchyContainer Hierarchy;
+	URigHierarchy *DynamicHierarchy;
 
 	UPROPERTY()
 	TSoftObjectPtr<UControlRigGizmoLibrary> GizmoLibrary;
@@ -445,8 +391,8 @@ private:
 	FRigNameCache NameCache;
 
 private:
-	// Controls for the container
-	void HandleOnControlModified(UControlRig* Subject, const FRigControl& Control, const FRigControlModifiedContext& Context);
+	
+	void HandleOnControlModified(UControlRig* Subject, FRigControlElement* Control, const FRigControlModifiedContext& Context);
 
 	TArray<FRigVMExternalVariable> GetExternalVariablesImpl(bool bFallbackToBlueprint) const;
 
@@ -496,16 +442,10 @@ private:
 	/** Broadcasts a notification whenever the controlrig is executed / updated. */
 	FControlRigExecuteEvent ExecutedEvent;
 
-	/** Handle a Control UI Settting Changed */
-	void HandleOnControlUISettingChanged(FRigHierarchyContainer* InContainer, const FRigElementKey& InKey);
-
-	/** Handle a Control Being Selected */
-	void HandleOnControlSelected(FRigHierarchyContainer* InContainer, const FRigElementKey& InKey, bool bSelected);
+	/** Handle changes within the hierarchy */
+	void HandleHierarchyModified(ERigHierarchyNotification InNotification, URigHierarchy* InHierarchy, const FRigBaseElement* InElement);
 
 #if WITH_EDITOR
-	/** Update the available controls within the control rig editor */
-	void UpdateAvailableControls();
-
 	/** Remove a transient / temporary control used to interact with a pin */
 	FName AddTransientControl(URigVMPin* InPin, FRigElementKey SpaceKey = FRigElementKey());
 
@@ -524,17 +464,18 @@ private:
 	FName RemoveTransientControl(const FRigElementKey& InElement);
 
 	static FName GetNameForTransientControl(const FRigElementKey& InElement);
+	FName GetNameForTransientControl(URigVMPin* InPin) const;
+	static FString GetPinNameFromTransientControl(const FRigElementKey& InKey);
+	static FRigElementKey GetElementKeyFromTransientControl(const FRigElementKey& InKey);
 
 	/** Removes all  transient / temporary control used to interact with pins */
 	void ClearTransientControls();
 
-	TArray<FRigControl> AvailableControlsOverride;
-	TArray<FRigControl> TransientControls;
 	UAnimPreviewInstance* PreviewInstance;
 
 #endif
 
-	void HandleOnRigEvent(FRigHierarchyContainer* InContainer, const FRigEventContext& InEvent);
+	void HandleHierarchyEvent(URigHierarchy* InHierarchy, const FRigEventContext& InEvent);
 	FRigEventDelegate RigEventDelegate;
 
 	void InitializeFromCDO();
@@ -585,10 +526,10 @@ protected:
 private:
 
 	void CopyPoseFromOtherRig(UControlRig* Subject);
-	void HandleInteractionRigControlModified(UControlRig* Subject, const FRigControl& Control, const FRigControlModifiedContext& Context);
+	void HandleInteractionRigControlModified(UControlRig* Subject, FRigControlElement* Control, const FRigControlModifiedContext& Context);
 	void HandleInteractionRigInitialized(UControlRig* Subject, EControlRigState State, const FName& EventName);
 	void HandleInteractionRigExecuted(UControlRig* Subject, EControlRigState State, const FName& EventName);
-	void HandleInteractionRigControlSelected(UControlRig* Subject, const FRigControl& InControl, bool bSelected, bool bInverted);
+	void HandleInteractionRigControlSelected(UControlRig* Subject, FRigControlElement* InControl, bool bSelected, bool bInverted);
 
 #if WITH_EDITOR
 
@@ -651,7 +592,7 @@ protected:
 	FControlModifiedEvent OnControlModified;
 	FControlSelectedEvent OnControlSelected;
 
-	TArray<FRigControl> QueuedModifiedControls;
+	TArray<FRigElementKey> QueuedModifiedControls;
 
 	friend class FControlRigBlueprintCompilerContext;
 	friend struct FRigHierarchyRef;
diff --git a/Engine/Plugins/Experimental/ControlRig/Source/ControlRig/Public/ControlRigDefines.h b/Engine/Plugins/Experimental/ControlRig/Source/ControlRig/Public/ControlRigDefines.h
--- a/Engine/Plugins/Experimental/ControlRig/Source/ControlRig/Public/ControlRigDefines.h
+++ b/Engine/Plugins/Experimental/ControlRig/Source/ControlRig/Public/ControlRigDefines.h
@@ -3,8 +3,7 @@
 #pragma once
 
 #include "PropertyPathHelpers.h"
-#include "Rigs/RigHierarchyContainer.h"
-#include "Rigs/RigCurveContainer.h"
+#include "Rigs/RigHierarchy.h"
 #include "Rigs/RigHierarchyCache.h"
 #include "Stats/StatsHierarchical.h"
 #include "RigVMCore/RigVMExecuteContext.h"
@@ -24,43 +23,7 @@ struct FControlRigExecuteContext : public FRigVMExecuteContext
 	{
 	}
 		
-	FRigHierarchyContainer* Hierarchy;
-
-	FRigBoneHierarchy* GetBones()
-	{
-		if (Hierarchy != nullptr)
-		{
-			return &Hierarchy->BoneHierarchy;
-		}
-		return nullptr;
-	}
-
-	FRigSpaceHierarchy* GetSpaces()
-	{
-		if (Hierarchy != nullptr)
-		{
-			return &Hierarchy->SpaceHierarchy;
-		}
-		return nullptr;
-	}
-
-	FRigControlHierarchy* GetControls()
-	{
-		if (Hierarchy != nullptr)
-		{
-			return &Hierarchy->ControlHierarchy;
-		}
-		return nullptr;
-	}
-
-	FRigCurveContainer* GetCurves()
-	{
-		if (Hierarchy != nullptr)
-		{
-			return &Hierarchy->CurveContainer;
-		}
-		return nullptr;
-	}
+	URigHierarchy* Hierarchy;
 };
 
 UENUM()
diff --git a/Engine/Plugins/Experimental/ControlRig/Source/ControlRig/Public/ControlRigObjectVersion.h b/Engine/Plugins/Experimental/ControlRig/Source/ControlRig/Public/ControlRigObjectVersion.h
--- a/Engine/Plugins/Experimental/ControlRig/Source/ControlRig/Public/ControlRigObjectVersion.h
+++ b/Engine/Plugins/Experimental/ControlRig/Source/ControlRig/Public/ControlRigObjectVersion.h
@@ -30,6 +30,9 @@ struct CONTROLRIG_API FControlRigObjectVersion
 		// Full variable support
 		BlueprintVariableSupport,
 
+		// Hierarchy V2.0
+		RigHierarchyV2,
+
 		// -----<new versions can be added above this line>-------------------------------------------------
 		VersionPlusOne,
 		LatestVersion = VersionPlusOne - 1
diff --git a/Engine/Plugins/Experimental/ControlRig/Source/ControlRig/Public/Drawing/ControlRigDrawInterface.h b/Engine/Plugins/Experimental/ControlRig/Source/ControlRig/Public/Drawing/ControlRigDrawInterface.h
--- a/Engine/Plugins/Experimental/ControlRig/Source/ControlRig/Public/Drawing/ControlRigDrawInterface.h
+++ b/Engine/Plugins/Experimental/ControlRig/Source/ControlRig/Public/Drawing/ControlRigDrawInterface.h
@@ -3,7 +3,7 @@
 #pragma once
 
 #include "CoreMinimal.h"
-#include "Rigs/RigHierarchyContainer.h"
+#include "Rigs/RigHierarchy.h"
 #include "Math/Simulation/CRSimPointContainer.h"
 #include "Drawing/ControlRigDrawInstruction.h"
 #include "Drawing/ControlRigDrawContainer.h"
@@ -40,7 +40,7 @@ public:
 	void DrawRectangle(const FTransform& WorldOffset, const FTransform& Transform, float Size, const FLinearColor& Color, float Thickness);
 	void DrawArc(const FTransform& WorldOffset, const FTransform& Transform, float Radius, float MinimumAngle, float MaximumAngle, const FLinearColor& Color, float Thickness, int32 Detail);
 	void DrawBezier(const FTransform& WorldOffset, const FCRFourPointBezier& InBezier, float MinimumU, float MaximumU, const FLinearColor& Color, float Thickness, int32 Detail);
-	void DrawHierarchy(const FTransform& WorldOffset, const FRigBoneHierarchy& Hierarchy, EControlRigDrawHierarchyMode::Type Mode, float Scale, const FLinearColor& Color, float Thickness);
+	void DrawHierarchy(const FTransform& WorldOffset, URigHierarchy* Hierarchy, EControlRigDrawHierarchyMode::Type Mode, float Scale, const FLinearColor& Color, float Thickness);
 	void DrawPointSimulation(const FTransform& WorldOffset, const FCRSimPointContainer& Simulation, const FLinearColor& Color, float Thickness, float PrimitiveSize = 0.f, bool bDrawPointsAsSphere = false);
 	void DrawCone(const FTransform& WorldOffset, const FTransform& ConeOffset, float Angle1, float Angle2, uint32 NumSides, bool bDrawSideLines, const FLinearColor& SideLineColor, FMaterialRenderProxy* const MaterialRenderProxy);
 
diff --git a/Engine/Plugins/Experimental/ControlRig/Source/ControlRig/Public/Rigs/AdditiveControlRig.h b/Engine/Plugins/Experimental/ControlRig/Source/ControlRig/Public/Rigs/AdditiveControlRig.h
--- a/Engine/Plugins/Experimental/ControlRig/Source/ControlRig/Public/Rigs/AdditiveControlRig.h
+++ b/Engine/Plugins/Experimental/ControlRig/Source/ControlRig/Public/Rigs/AdditiveControlRig.h
@@ -29,6 +29,9 @@ private:
 	// custom units that are used to execute this rig
 	TArray<FRigUnit_AddBoneTransform> AddBoneRigUnits;
 
+	UPROPERTY(transient)
+	URigHierarchyController* Controller;
+
 	/** Create RigElements - bone hierarchy and curves - from incoming skeleton */
 	void CreateRigElements(const USkeletalMesh* InReferenceMesh);
 	void CreateRigElements(const FReferenceSkeleton& InReferenceSkeleton, const FSmartNameMapping* InSmartNameMapping);
diff --git a/Engine/Plugins/Experimental/ControlRig/Source/ControlRig/Public/Rigs/FKControlRig.h b/Engine/Plugins/Experimental/ControlRig/Source/ControlRig/Public/Rigs/FKControlRig.h
--- a/Engine/Plugins/Experimental/ControlRig/Source/ControlRig/Public/Rigs/FKControlRig.h
+++ b/Engine/Plugins/Experimental/ControlRig/Source/ControlRig/Public/Rigs/FKControlRig.h
@@ -2,6 +2,7 @@
 
 #pragma once
 #include "ControlRig.h"
+#include "Rigs/RigHierarchyController.h"
 #include "Units/Hierarchy/RigUnit_SetTransform.h"
 #include "Units/Hierarchy/RigUnit_SetCurveValue.h"
 #include "FKControlRig.generated.h"
@@ -69,5 +70,8 @@ private:
 	UPROPERTY()
 	EControlRigFKRigExecuteMode ApplyMode;
 
+	UPROPERTY(transient)
+	URigHierarchyController* Controller;
+
 	friend class FControlRigInteractionTest;
 };
\ No newline at end of file
diff --git a/Engine/Plugins/Experimental/ControlRig/Source/ControlRig/Public/Rigs/RigBoneHierarchy.h b/Engine/Plugins/Experimental/ControlRig/Source/ControlRig/Public/Rigs/RigBoneHierarchy.h
--- a/Engine/Plugins/Experimental/ControlRig/Source/ControlRig/Public/Rigs/RigBoneHierarchy.h
+++ b/Engine/Plugins/Experimental/ControlRig/Source/ControlRig/Public/Rigs/RigBoneHierarchy.h
@@ -9,15 +9,9 @@
 #include "RigBoneHierarchy.generated.h"
 
 class UControlRig;
+struct FRigHierarchyContainer;
 
-UENUM(BlueprintType)
-enum class ERigBoneType : uint8
-{
-	Imported,
-	User
-};
-
-USTRUCT(BlueprintType)
+USTRUCT()
 struct CONTROLRIG_API FRigBone: public FRigElement
 {
 	GENERATED_BODY()
@@ -74,180 +68,26 @@ struct CONTROLRIG_API FRigBone: public FRigElement
 	}
 };
 
-USTRUCT(BlueprintType)
+USTRUCT()
 struct CONTROLRIG_API FRigBoneHierarchy
 {
 	GENERATED_BODY()
 
 	FRigBoneHierarchy();
-	FRigBoneHierarchy& operator= (const FRigBoneHierarchy &InOther);
-
-	FORCEINLINE ERigElementType RigElementType() const { return ERigElementType::Bone; }
-
-	FORCEINLINE int32 Num() const { return Bones.Num(); }
-	FORCEINLINE const FRigBone& operator[](int32 InIndex) const { return Bones[InIndex]; }
-	FORCEINLINE FRigBone& operator[](int32 InIndex) { return Bones[InIndex]; }
-	FORCEINLINE const FRigBone& operator[](const FName& InName) const { return Bones[GetIndex(InName)]; }
-	FORCEINLINE FRigBone& operator[](const FName& InName) { return Bones[GetIndex(InName)]; }
 
 	FORCEINLINE TArray<FRigBone>::RangedForIteratorType      begin()       { return Bones.begin(); }
 	FORCEINLINE TArray<FRigBone>::RangedForConstIteratorType begin() const { return Bones.begin(); }
 	FORCEINLINE TArray<FRigBone>::RangedForIteratorType      end()         { return Bones.end();   }
 	FORCEINLINE TArray<FRigBone>::RangedForConstIteratorType end() const   { return Bones.end();   }
 
-	bool IsNameAvailable(const FString& InPotentialNewName, FString* OutErrorMessage = nullptr) const;
-
-	FName GetSafeNewName(const FString& InPotentialNewName) const;
-
-	FRigBone& Add(const FName& InNewName, const FName& InParentName = NAME_None, ERigBoneType InType = ERigBoneType::Imported, const FTransform& InInitTransform = FTransform::Identity);
-
-	FRigBone& Add(const FName& InNewName, const FName& InParentName, ERigBoneType InType, const FTransform& InInitTransform, const FTransform& InLocalTransform, const FTransform& InGlobalTransform);
-
-	FRigBone Remove(const FName& InNameToRemove);
-
-	FName Rename(const FName& InOldName, const FName& InNewName);
-
-	bool Reparent(const FName& InName, const FName& InNewParentName);
-
-	FName GetName(int32 InIndex) const;
-
-	FORCEINLINE int32 GetIndex(const FName& InName) const
-	{
-		if(NameToIndexMapping.Num() != Bones.Num())
-		{
-			return GetIndexSlow(InName);
-		}
-
-		const int32* Index = NameToIndexMapping.Find(InName);
-		if (Index)
-		{
-			return *Index;
-		}
-
-		return INDEX_NONE;
-	}
-
-	// list of names of children - this is not cheap, and is supposed to be used only for one time set up
-	int32 GetChildren(const FName& InName, TArray<int32>& OutChildren, bool bRecursively) const;
-
-	int32 GetChildren(const int32 InIndex, TArray<int32>& OutChildren, bool bRecursively) const;
-
-	void SetGlobalTransform(const FName& InName, const FTransform& InTransform, bool bPropagateTransform = true);
-
-	void SetGlobalTransform(int32 InIndex, const FTransform& InTransform, bool bPropagateTransform = true);
-
-	FTransform GetGlobalTransform(const FName& InName) const;
-
-	FTransform GetGlobalTransform(int32 InIndex) const;
-
-	void SetLocalTransform(const FName& InName, const FTransform& InTransform, bool bPropagateTransform = true);
-
-	void SetLocalTransform(int32 InIndex, const FTransform& InTransform, bool bPropagateTransform = true);
-
-	FTransform GetLocalTransform(const FName& InName) const;
-
-	FTransform GetLocalTransform(int32 InIndex) const;
-
-	void SetInitialGlobalTransform(const FName& InName, const FTransform& InTransform, bool bPropagateTransform = false);
-
-	void SetInitialGlobalTransform(int32 InIndex, const FTransform& InTransform, bool bPropagateTransform = false);
-
-	void SetInitialLocalTransform(const FName& InName, const FTransform& InTransform, bool bPropagateTransform = false);
-
-	void SetInitialLocalTransform(int32 InIndex, const FTransform& InTransform, bool bPropagateTransform = false);
-
-	FTransform GetInitialGlobalTransform(const FName& InName) const;
-
-	FTransform GetInitialGlobalTransform(int32 InIndex) const;
-
-	FTransform GetInitialLocalTransform(const FName& InName) const;
-
-	FTransform GetInitialLocalTransform(int32 InIndex) const;
-
-	// updates all of the internal caches
-	void Initialize(bool bResetTransforms = true);
-
-	// clears the hierarchy and removes all content
-	void Reset();
-
-	// returns the current pose
-	FRigPose GetPose() const;
-
-	// sets the current transforms from the given pose
-	void SetPose(FRigPose& InPose);
-
-	// resets all of the transforms back to the initial transform
-	void ResetTransforms();
-
-	// copies all initial transforms from another hierarchy
-	void CopyInitialTransforms(const FRigBoneHierarchy& InOther);
-
-	// recomputes all of the global transforms from local
-	void RecomputeGlobalTransforms();
-
-	// recomputes the local transform of a single bone
-	void RecalculateLocalTransform(int32 InIndex);
-
-	// recomputes the global transform of a single bone
-	void RecalculateGlobalTransform(int32 InIndex);
-
-	// propagates the transform change for a single bone
-	void PropagateTransform(int32 InIndex);
-
-	// copies the global transform to the initial transform
-	void CopyGlobalToInitial(int32 InIndex, bool bPropagate = true);
-
-	// import skeleton
-	TArray<FRigElementKey> ImportSkeleton(const FReferenceSkeleton& InSkeleton, const FName& InNameSpace, bool bReplaceExistingBones, bool bRemoveObsoleteBones, bool bSelectBones, bool bNotify);
-
-	bool Select(const FName& InName, bool bSelect = true);
-	bool ClearSelection();
-	TArray<FName> CurrentSelection() const;
-	bool IsSelected(const FName& InName) const;
-
-	FRigElementSelected OnBoneSelected;
-
-#if WITH_EDITOR
-
-	FRigElementAdded OnBoneAdded;
-	FRigElementRemoved OnBoneRemoved;
-	FRigElementRenamed OnBoneRenamed;
-	FRigElementReparented OnBoneReparented;
-
-#endif
-
 private:
 
 	// disable copy constructor
 	FRigBoneHierarchy(const FRigBoneHierarchy& InOther) {}
 
-	FRigHierarchyContainer* Container;
-
 	UPROPERTY(EditAnywhere, Category = FRigBoneHierarchy)
 	TArray<FRigBone> Bones;
 
-	// can serialize fine? 
-	UPROPERTY()
-	TMap<FName, int32> NameToIndexMapping;
-
-	UPROPERTY(transient)
-	TArray<FName> Selection;
-
-	int32 GetIndexSlow(const FName& InName) const;
-
-	void RecalculateLocalTransform(FRigBone& InOutBone);
-	void RecalculateGlobalTransform(FRigBone& InOutBone);
-
-	void RefreshParentNames();
-	void RefreshMapping();
-	void AppendToPose(FRigPose& InOutPose) const;
-	void Sort();
-
-	// list of names of children - this is not cheap, and is supposed to be used only for one time set up
-	int32 GetChildrenRecursive(const int32 InIndex, TArray<int32>& OutChildren, bool bRecursively) const;
-
-	bool bSuspendNotifications;
-
 	friend struct FRigHierarchyContainer;
 	friend struct FCachedRigElement;
 	friend class UControlRigHierarchyModifier;
diff --git a/Engine/Plugins/Experimental/ControlRig/Source/ControlRig/Public/Rigs/RigControlHierarchy.h b/Engine/Plugins/Experimental/ControlRig/Source/ControlRig/Public/Rigs/RigControlHierarchy.h
--- a/Engine/Plugins/Experimental/ControlRig/Source/ControlRig/Public/Rigs/RigControlHierarchy.h
+++ b/Engine/Plugins/Experimental/ControlRig/Source/ControlRig/Public/Rigs/RigControlHierarchy.h
@@ -10,195 +10,10 @@
 #include "RigControlHierarchy.generated.h"
 
 class UControlRig;
-
-UENUM(BlueprintType)
-enum class ERigControlType : uint8
-{
-	Bool,
-	Float,
-	Integer,
-	Vector2D,
-	Position,
-	Scale,
-	Rotator,
-	Transform,
-	TransformNoScale,
-	EulerTransform,
-};
-
-UENUM(BlueprintType)
-enum class ERigControlValueType : uint8
-{
-	Initial,
-	Current,
-	Minimum,
-	Maximum
-};
-
-UENUM(BlueprintType)
-enum class ERigControlAxis : uint8
-{
-	X,
-	Y,
-	Z
-};
-
-USTRUCT(BlueprintType)
-struct CONTROLRIG_API FRigControlValueStorage
-{
-public:
-
-	GENERATED_BODY()
-
-	FRigControlValueStorage()
-	{
-		FMemory::Memzero(this, sizeof(FRigControlValueStorage));
-	}
-
-	UPROPERTY()
-	float Float00;
-
-	UPROPERTY()
-	float Float01;
-
-	UPROPERTY()
-	float Float02;
-
-	UPROPERTY()
-	float Float03;
-
-	UPROPERTY()
-	float Float10;
-
-	UPROPERTY()
-	float Float11;
-
-	UPROPERTY()
-	float Float12;
-
-	UPROPERTY()
-	float Float13;
-
-	UPROPERTY()
-	float Float20;
-
-	UPROPERTY()
-	float Float21;
-
-	UPROPERTY()
-	float Float22;
-
-	UPROPERTY()
-	float Float23;
-
-	UPROPERTY()
-	float Float30;
-
-	UPROPERTY()
-	float Float31;
-
-	UPROPERTY()
-	float Float32;
-
-	UPROPERTY()
-	float Float33;
-
-	UPROPERTY()
-	bool bValid;
-};
-
-USTRUCT(BlueprintType)
-struct CONTROLRIG_API FRigControlValue
-{
-	GENERATED_BODY()
-
-public:
-
-	FRigControlValue()
-		: FloatStorage()
-		, Storage_DEPRECATED(FTransform::Identity)
-	{
-	}
-
-	FORCEINLINE_DEBUGGABLE bool IsValid() const
-	{
-		return FloatStorage.bValid;
-	}
-
-	template<class T>
-	FORCEINLINE_DEBUGGABLE T Get() const
-	{
-		return GetRef<T>();
-	}
-
-	template<class T>
-	FORCEINLINE_DEBUGGABLE T& GetRef()
-	{
-		FloatStorage.bValid = true;
-		return *(T*)&FloatStorage;
-	}
-
-	template<class T>
-	FORCEINLINE_DEBUGGABLE const T& GetRef() const
-	{
-		return *(T*)&FloatStorage;
-	}
-
-	template<class T>
-	FORCEINLINE_DEBUGGABLE void Set(T InValue)
-	{
-		GetRef<T>() = InValue;
-	}
-
-	template<class T>
-	FORCEINLINE_DEBUGGABLE FString ToString() const
-	{
-		FString Result;
-		TBaseStructure<T>::Get()->ExportText(Result, &GetRef<T>(), nullptr, nullptr, PPF_None, nullptr);
-		return Result;
-	}
-
-	template<class T>
-	FORCEINLINE_DEBUGGABLE T SetFromString(const FString& InString)
-	{
-		T Value;
-		TBaseStructure<T>::Get()->ImportText(*InString, &Value, nullptr, PPF_None, nullptr, TBaseStructure<T>::Get()->GetName());
-		Set<T>(Value);
-		return Value;
-	}
-
-	template<class T>
-	FORCEINLINE_DEBUGGABLE static FRigControlValue Make(T InValue)
-	{
-		FRigControlValue Value;
-		Value.Set<T>(InValue);
-		return Value;
-	}
-
-private:
-
-	UPROPERTY()
-	FRigControlValueStorage FloatStorage;
-	
-	UPROPERTY()
-	FTransform Storage_DEPRECATED;
-
-	friend struct FRigControlHierarchy;
-};
-
-template<>
-FORCEINLINE_DEBUGGABLE FQuat FRigControlValue::SetFromString<FQuat>(const FString& InString)
-{
-	FQuat Value;
-	TBaseStructure<FQuat>::Get()->ImportText(*InString, &Value, nullptr, PPF_None, nullptr, TBaseStructure<FQuat>::Get()->GetName());
-	Set<FRotator>(Value.Rotator());
-	return Value;
-}
-
-
 class UStaticMesh;
+struct FRigHierarchyContainer;
 
-USTRUCT(BlueprintType)
+USTRUCT()
 struct CONTROLRIG_API FRigControl : public FRigElement
 {
 	GENERATED_BODY()
@@ -415,15 +230,7 @@ struct CONTROLRIG_API FRigControl : public FRigElement
 		return Value;
 	}
 
-	void ApplyLimits(FRigControlValue& InOutValue);
-	FORCEINLINE_DEBUGGABLE static float Clamp(float Value, float Minimum, float Maximum)
-	{
-		if (Minimum < Maximum)
-		{
-			return FMath::Clamp<float>(Value, Minimum, Maximum);
-		}
-		return FMath::Clamp<float>(Value, Maximum, Minimum);
-	}
+	void ApplyLimits(FRigControlValue& InOutValue) const;
 
 	FORCEINLINE_DEBUGGABLE static FProperty* FindPropertyForValueType(ERigControlValueType InValueType)
 	{
@@ -454,335 +261,29 @@ struct CONTROLRIG_API FRigControl : public FRigElement
 
 };
 
-USTRUCT(BlueprintType)
+USTRUCT()
 struct CONTROLRIG_API FRigControlHierarchy
 {
 	GENERATED_BODY()
 
 	FRigControlHierarchy();
-	FRigControlHierarchy& operator= (const FRigControlHierarchy &InOther);
-
-	FORCEINLINE_DEBUGGABLE ERigElementType RigElementType() const { return ERigElementType::Control; }
 
 	FORCEINLINE_DEBUGGABLE int32 Num() const { return Controls.Num(); }
-	FORCEINLINE_DEBUGGABLE const FRigControl& operator[](int32 InIndex) const { return Controls[InIndex]; }
-	FORCEINLINE_DEBUGGABLE FRigControl& operator[](int32 InIndex) { return Controls[InIndex]; }
-	FORCEINLINE_DEBUGGABLE const FRigControl& operator[](const FName& InName) const { return Controls[GetIndex(InName)]; }
-	FORCEINLINE_DEBUGGABLE FRigControl& operator[](const FName& InName) { return Controls[GetIndex(InName)]; }
-
-	FORCEINLINE_DEBUGGABLE const TArray<FRigControl>& GetControls() const { return Controls; }
-
 	FORCEINLINE_DEBUGGABLE TArray<FRigControl>::RangedForIteratorType      begin()       { return Controls.begin(); }
 	FORCEINLINE_DEBUGGABLE TArray<FRigControl>::RangedForConstIteratorType begin() const { return Controls.begin(); }
 	FORCEINLINE_DEBUGGABLE TArray<FRigControl>::RangedForIteratorType      end()         { return Controls.end();   }
 	FORCEINLINE_DEBUGGABLE TArray<FRigControl>::RangedForConstIteratorType end() const   { return Controls.end();   }
 
-	bool IsNameAvailable(const FString& InPotentialNewName, FString* OutErrorMessage = nullptr) const;
-
-	FName GetSafeNewName(const FString& InPotentialNewName) const;
-
-	FRigControl& Add(
-		const FName& InNewName,
-		ERigControlType InControlType = ERigControlType::Transform,
-		const FName& InParentName = NAME_None,
-		const FName& InSpaceName = NAME_None,
-		const FTransform& InOffsetTransform = FTransform::Identity,
-		const FRigControlValue& InValue = FRigControlValue(),
-		const FName& InGizmoName = TEXT("Gizmo"),
-		const FTransform& InGizmoTransform = FTransform::Identity,
-		const FLinearColor& InGizmoColor = FLinearColor::Red
-	);
-
-	FRigControl Remove(const FName& InNameToRemove);
-
-	FName Rename(const FName& InOldName, const FName& InNewName);
-
-	bool Reparent(const FName& InName, const FName& InNewParentName);
-
-	void SetSpace(const FName& InName, const FName& InNewSpaceName);
-
-	FName GetName(int32 InIndex) const;
-
-	FORCEINLINE_DEBUGGABLE int32 GetIndex(const FName& InName) const
-	{
-		if(NameToIndexMapping.Num() != Controls.Num())
-		{
-			return GetIndexSlow(InName);
-		}
-
-		const int32* Index = NameToIndexMapping.Find(InName);
-		if (Index)
-		{
-			return *Index;
-		}
-
-		return INDEX_NONE;
-	}
-
-	// list of names of children - this is not cheap, and is supposed to be used only for one time set up
-	int32 GetChildren(const FName& InName, TArray<int32>& OutChildren, bool bRecursively) const;
-
-	int32 GetChildren(const int32 InIndex, TArray<int32>& OutChildren, bool bRecursively) const;
-
-	void SetGlobalTransform(const FName& InName, const FTransform& InTransform);
-
-	void SetGlobalTransform(int32 InIndex, const FTransform& InTransform);
-
-	FTransform GetGlobalTransform(const FName& InName) const;
-
-	FTransform GetGlobalTransform(int32 InIndex) const;
-
-	void SetLocalTransform(const FName& InName, const FTransform& InTransform, ERigControlValueType InValueType = ERigControlValueType::Current);
-
-	void SetLocalTransform(int32 InIndex, const FTransform& InTransform, ERigControlValueType InValueType = ERigControlValueType::Current);
-
-	FTransform GetLocalTransform(const FName& InName, ERigControlValueType InValueType = ERigControlValueType::Current) const;
-
-	FTransform GetLocalTransform(int32 InIndex, ERigControlValueType InValueType = ERigControlValueType::Current) const;
-
-	FTransform GetParentTransform(int32 InIndex, bool bIncludeOffsetTransform = true) const;
-
-	FTransform GetParentInitialTransform(int32 InIndex, bool bIncludeOffsetTransform = true) const;
-
-	// this is only valid if you have transform as type
-	// setting initial transform from global transform (control transform is based on parent)
-	void SetInitialGlobalTransform(const FName& InName, const FTransform& GlobalTransform);
-	
-	void SetInitialGlobalTransform(int32 InIndex, const FTransform& GlobalTransform);
-	FTransform GetInitialGlobalTransform(const FName& InName) const;
-	
-	FTransform GetInitialGlobalTransform(int32 InIndex) const;
-
-	void SetControlOffset(int32 InIndex, const FTransform& InOffsetTransform);
+	void PostLoad();
 	
-	void SetValue(const FName& InName, const FRigControlValue& InValue, ERigControlValueType InValueType = ERigControlValueType::Current);
-
-	void SetValue(int32 InIndex, const FRigControlValue& InValue, ERigControlValueType InValueType = ERigControlValueType::Current);
-
-	template<class T>
-	FORCEINLINE_DEBUGGABLE void SetValue(const FName& InName, T InValue, ERigControlValueType InValueType = ERigControlValueType::Current)
-	{
-		SetValue(InName, FRigControlValue::Make<T>(InValue), InValueType);
-	}
-
-	template<class T>
-	FORCEINLINE_DEBUGGABLE void SetValue(int32 InIndex, T InValue, ERigControlValueType InValueType = ERigControlValueType::Current)
-	{
-		SetValue(InIndex, FRigControlValue::Make<T>(InValue), InValueType);
-	}
-
-	FRigControlValue GetValue(const FName& InName, ERigControlValueType InValueType = ERigControlValueType::Current) const;
-
-	FRigControlValue GetValue(int32 InIndex, ERigControlValueType InValueType = ERigControlValueType::Current) const;
-
-	template<class T>
-	FORCEINLINE_DEBUGGABLE T GetValue(const FName& InName, ERigControlValueType InValueType = ERigControlValueType::Current) const
-	{
-		return GetValue(InName, InValueType).Get<T>();
-	}
-
-	template<class T>
-	FORCEINLINE_DEBUGGABLE T GetValue(int32 InIndex, ERigControlValueType InValueType = ERigControlValueType::Current) const
-	{
-		return GetValue(InIndex, InValueType).Get<T>();
-	}
-
-	FORCEINLINE_DEBUGGABLE void SetInitialValue(const FName& InName, const FRigControlValue& InValue)
-	{
-		SetValue(InName, InValue, ERigControlValueType::Initial);
-	}
-
-	FORCEINLINE_DEBUGGABLE void SetInitialValue(int32 InIndex, const FRigControlValue& InValue)
-	{
-		SetValue(InIndex, InValue, ERigControlValueType::Initial);
-	}
-
-	template<class T>
-	FORCEINLINE_DEBUGGABLE void SetInitialValue(const FName& InName, T InValue)
-	{
-		SetInitialValue(InName, FRigControlValue::Make<T>(InValue));
-	}
-
-	template<class T>
-	FORCEINLINE_DEBUGGABLE void SetInitialValue(int32 InIndex, T InValue)
-	{
-		SetInitialValue(InIndex, FRigControlValue::Make<T>(InValue));
-	}
-
-	FORCEINLINE_DEBUGGABLE FRigControlValue GetInitialValue(const FName& InName) const
-	{
-		return GetValue(InName, ERigControlValueType::Initial);
-	}
-
-	FORCEINLINE_DEBUGGABLE FRigControlValue GetInitialValue(int32 InIndex) const
-	{
-		return GetValue(InIndex, ERigControlValueType::Initial);
-	}
-
-	template<class T>
-	FORCEINLINE_DEBUGGABLE T GetInitialValue(const FName& InName) const
-	{
-		return GetInitialValue(InName).Get<T>();
-	}
-
-	template<class T>
-	FORCEINLINE_DEBUGGABLE T GetInitialValue(int32 InIndex) const
-	{
-		return GetInitialValue(InIndex).Get<T>();
-	}
-
-	FORCEINLINE_DEBUGGABLE void SetMinimumValue(const FName& InName, const FRigControlValue& InValue)
-	{
-		SetValue(InName, InValue, ERigControlValueType::Minimum);
-	}
-
-	FORCEINLINE_DEBUGGABLE void SetMinimumValue(int32 InIndex, const FRigControlValue& InValue)
-	{
-		SetValue(InIndex, InValue, ERigControlValueType::Minimum);
-	}
-
-	template<class T>
-	FORCEINLINE_DEBUGGABLE void SetMinimumValue(const FName& InName, T InValue)
-	{
-		SetMinimumValue(InName, FRigControlValue::Make<T>(InValue));
-	}
-
-	template<class T>
-	FORCEINLINE_DEBUGGABLE void SetMinimumValue(int32 InIndex, T InValue)
-	{
-		SetMinimumValue(InIndex, FRigControlValue::Make<T>(InValue));
-	}
-
-	FORCEINLINE_DEBUGGABLE FRigControlValue GetMinimumValue(const FName& InName) const
-	{
-		return GetValue(InName, ERigControlValueType::Minimum);
-	}
-
-	FORCEINLINE_DEBUGGABLE FRigControlValue GetMinimumValue(int32 InIndex) const
-	{
-		return GetValue(InIndex, ERigControlValueType::Minimum);
-	}
-
-	template<class T>
-	FORCEINLINE_DEBUGGABLE T GetMinimumValue(const FName& InName) const
-	{
-		return GetMinimumValue(InName).Get<T>();
-	}
-
-	template<class T>
-	FORCEINLINE_DEBUGGABLE T GetMinimumValue(int32 InIndex) const
-	{
-		return GetMinimumValue(InIndex).Get<T>();
-	}
-
-	FORCEINLINE_DEBUGGABLE void SetMaximumValue(const FName& InName, const FRigControlValue& InValue)
-	{
-		return SetValue(InName, InValue, ERigControlValueType::Maximum);
-	}
-
-	FORCEINLINE_DEBUGGABLE void SetMaximumValue(int32 InIndex, const FRigControlValue& InValue)
-	{
-		return SetValue(InIndex, InValue, ERigControlValueType::Maximum);
-	}
-
-	template<class T>
-	FORCEINLINE_DEBUGGABLE void SetMaximumValue(const FName& InName, T InValue)
-	{
-		SetMaximumValue(InName, FRigControlValue::Make<T>(InValue));
-	}
-
-	template<class T>
-	FORCEINLINE_DEBUGGABLE void SetMaximumValue(int32 InIndex, T InValue)
-	{
-		SetMaximumValue(InIndex, FRigControlValue::Make<T>(InValue));
-	}
-
-	FORCEINLINE_DEBUGGABLE FRigControlValue GetMaximumValue(const FName& InName) const
-	{
-		return GetValue(InName, ERigControlValueType::Maximum);
-	}
-
-	FORCEINLINE_DEBUGGABLE FRigControlValue GetMaximumValue(int32 InIndex) const
-	{
-		return GetValue(InIndex, ERigControlValueType::Maximum);
-	}
-
-	template<class T>
-	FORCEINLINE_DEBUGGABLE T GetMaximumValue(const FName& InName) const
-	{
-		return GetMaximumValue(InName).Get<T>();
-	}
-
-	template<class T>
-	FORCEINLINE_DEBUGGABLE T GetMaximumValue(int32 InIndex) const
-	{
-		return GetMaximumValue(InIndex).Get<T>();
-	}
-
-	// updates all of the internal caches
-	void Initialize(bool bResetTransforms = true);
-
-	// clears the hierarchy and removes all content
-	void Reset();
-
-	// returns the current pose
-	FRigPose GetPose() const;
-
-	// sets the current transforms from the given pose
-	void SetPose(FRigPose& InPose);
-
-	// resets all of the values back to the initial values
-	void ResetValues();
-
-	// copies all offset transforms from another hierarchy
-	void CopyOffsetTransforms(const FRigControlHierarchy& InOther);
-
-	bool Select(const FName& InName, bool bSelect = true);
-	bool ClearSelection();
-	TArray<FName> CurrentSelection() const;
-	bool IsSelected(const FName& InName) const;
-
-	FRigElementAdded OnControlAdded;
-	FRigElementRemoved OnControlRemoved;
-	FRigElementRenamed OnControlRenamed;
-	FRigElementReparented OnControlReparented;
-	FRigElementSelected OnControlSelected;
-
-	FRigElementChanged OnControlUISettingsChanged;
-
-	void HandleOnElementRemoved(FRigHierarchyContainer* InContainer, const FRigElementKey& InKey);
-	void HandleOnElementRenamed(FRigHierarchyContainer* InContainer, ERigElementType InElementType, const FName& InOldName, const FName& InNewName);
-
 private:
 
 	// disable copy constructor
 	FRigControlHierarchy(const FRigControlHierarchy& InOther) {}
 
-	FRigHierarchyContainer* Container;
-
 	UPROPERTY(EditAnywhere, Category = FRigControlHierarchy)
 	TArray<FRigControl> Controls;
 
-	// can serialize fine? 
-	UPROPERTY()
-	TMap<FName, int32> NameToIndexMapping;
-
-	UPROPERTY(transient)
-	TArray<FName> Selection;
-
-	int32 GetSpaceIndex(const FName& InName) const;
-
-	int32 GetIndexSlow(const FName& InName) const;
-
-	void RefreshMapping();
-
-	void AppendToPose(FRigPose& InOutPose) const;
-
-	// list of names of children - this is not cheap, and is supposed to be used only for one time set up
-	int32 GetChildrenRecursive(const int32 InIndex, TArray<int32>& OutChildren, bool bRecursively) const;
-
-	void PostLoad();
-
 	friend struct FRigHierarchyContainer;
 	friend struct FCachedRigElement;
 	friend class UControlRigHierarchyModifier;
diff --git a/Engine/Plugins/Experimental/ControlRig/Source/ControlRig/Public/Rigs/RigCurveContainer.h b/Engine/Plugins/Experimental/ControlRig/Source/ControlRig/Public/Rigs/RigCurveContainer.h
--- a/Engine/Plugins/Experimental/ControlRig/Source/ControlRig/Public/Rigs/RigCurveContainer.h
+++ b/Engine/Plugins/Experimental/ControlRig/Source/ControlRig/Public/Rigs/RigCurveContainer.h
@@ -10,8 +10,9 @@
 
 class UControlRig;
 class USkeleton;
+struct FRigHierarchyContainer;
 
-USTRUCT(BlueprintType)
+USTRUCT()
 struct FRigCurve : public FRigElement
 {
 	GENERATED_BODY()
@@ -32,7 +33,7 @@ struct FRigCurve : public FRigElement
 	}
 };
 
-USTRUCT(BlueprintType)
+USTRUCT()
 struct CONTROLRIG_API FRigCurveContainer
 {
 	GENERATED_BODY()
@@ -40,91 +41,12 @@ struct CONTROLRIG_API FRigCurveContainer
 public:
 
 	FRigCurveContainer();
-	FRigCurveContainer& operator= (const FRigCurveContainer &InOther);
-
-	FORCEINLINE ERigElementType RigElementType() const { return ERigElementType::Curve; }
-
-	FORCEINLINE int32 Num() const { return Curves.Num(); }
-	FORCEINLINE const FRigCurve& operator[](int32 InIndex) const { return Curves[InIndex]; }
-	FORCEINLINE FRigCurve& operator[](int32 InIndex) { return Curves[InIndex]; }
-	FORCEINLINE const FRigCurve& operator[](const FName& InName) const { return Curves[GetIndex(InName)]; }
-	FORCEINLINE FRigCurve& operator[](const FName& InName) { return Curves[GetIndex(InName)]; }
 
 	FORCEINLINE TArray<FRigCurve>::RangedForIteratorType      begin()       { return Curves.begin(); }
 	FORCEINLINE TArray<FRigCurve>::RangedForConstIteratorType begin() const { return Curves.begin(); }
 	FORCEINLINE TArray<FRigCurve>::RangedForIteratorType      end()         { return Curves.end();   }
 	FORCEINLINE TArray<FRigCurve>::RangedForConstIteratorType end() const   { return Curves.end();   }
-
-	bool IsNameAvailable(const FString& InPotentialNewName, FString* OutErrorMessage = nullptr) const;
-
-	FName GetSafeNewName(const FString& InPotentialNewName) const;
-
-	FRigCurve& Add(const FName& InNewName, float InValue = 0.f);
-
-	FRigCurve Remove(const FName& InName);
-
-	FName GetName(int32 InIndex) const;
-
-	FORCEINLINE int32 GetIndex(const FName& InName) const
-	{
-		if(NameToIndexMapping.Num() != Curves.Num())
-		{
-			return GetIndexSlow(InName);
-		}
-
-		const int32* Index = NameToIndexMapping.Find(InName);
-		if (Index)
-		{
-			return *Index;
-		}
-
-		return INDEX_NONE;
-	}
-
-	void SetValue(const FName& InName, const float InValue);
-
-	void SetValue(int32 InIndex, const float InValue);
-
-	float GetValue(const FName& InName) const;
-
-	float GetValue(int32 InIndex) const;
-
-	FName Rename(const FName& InOldName, const FName& InNewName);
-
-	// updates all of the internal caches
-	void Initialize();
-
-	// clears the CurveContainer and removes all content
-	void Reset();
-
-	// resets all of the transforms back to the initial transform
-	void ResetValues();
-
-	// returns the current pose
-	FRigPose GetPose() const;
-
-	// sets the current values from the given pose
-	void SetPose(FRigPose& InPose);
-
-	FRigHierarchyContainer* Container;
-
-	bool Select(const FName& InName, bool bSelect = true);
-	bool ClearSelection();
-	TArray<FName> CurrentSelection() const;
-	bool IsSelected(const FName& InName) const;
-
-	FRigElementSelected OnCurveSelected;
-
-#if WITH_EDITOR
-
-	FRigElementAdded OnCurveAdded;
-	FRigElementRemoved OnCurveRemoved;
-	FRigElementRenamed OnCurveRenamed;
-
-	TArray<FRigElementKey> ImportCurvesFromSkeleton(const USkeleton* InSkeleton, const FName& InNameSpace, bool bRemoveObsoleteCurves, bool bSelectCurves, bool bNotify);
-
-#endif
-
+	
 private:
 
 	// disable copy constructor
@@ -133,20 +55,6 @@ private:
 	UPROPERTY(EditAnywhere, Category = FRigCurveContainer)
 	TArray<FRigCurve> Curves;
 
-	UPROPERTY()
-	TMap<FName, int32> NameToIndexMapping;
-
-	UPROPERTY(transient)
-	TArray<FName> Selection;
-
-	int32 GetIndexSlow(const FName& InName) const;
-
-	void RefreshMapping();
-
-	void AppendToPose(FRigPose& InOutPose) const;
-
-	bool bSuspendNotifications;
-
 	friend struct FRigHierarchyContainer;
 	friend struct FCachedRigElement;
 	friend class UControlRigHierarchyModifier;
