Control Rig: New rig hierarchy data structure refactoring

diff --git a/Engine/Plugins/Experimental/ControlRig/Source/ControlRig/Public/Rigs/RigHierarchy.h b/Engine/Plugins/Experimental/ControlRig/Source/ControlRig/Public/Rigs/RigHierarchy.h
new file mode 100644
--- /dev/null
+++ b/Engine/Plugins/Experimental/ControlRig/Source/ControlRig/Public/Rigs/RigHierarchy.h
@@ -0,0 +1,2103 @@
+// Copyright Epic Games, Inc. All Rights Reserved.
+
+#pragma once
+
+#include "CoreMinimal.h"
+#include "RigHierarchyElements.h"
+#include "RigHierarchyPose.h"
+#include "UObject/WeakObjectPtrTemplates.h"
+#include "RigHierarchy.generated.h"
+
+class URigHierarchy;
+class URigHierarchyController;
+
+#define URIGHIERARCHY_RECURSIVE_DIRTY_PROPAGATION 0
+
+DECLARE_MULTICAST_DELEGATE_ThreeParams(FRigHierarchyModifiedEvent, ERigHierarchyNotification /* type */, URigHierarchy* /* hierarchy */, const FRigBaseElement* /* element */);
+
+UCLASS(BlueprintType)
+class CONTROLRIG_API URigHierarchy : public UObject
+{
+	GENERATED_BODY()
+
+public:
+
+	URigHierarchy()
+	: TopologyVersion(0)
+	, bEnableDirtyPropagation(true)
+	, Elements()
+	, IndexLookup()
+	, TransformStackIndex(0)
+	, bTransactingForTransformChange(false)
+	{}
+
+	virtual ~URigHierarchy();
+
+	// UObject interface
+	virtual void Serialize(FArchive& Ar) override;
+	void Save(FArchive& Ar);
+	void Load(FArchive& Ar);
+#if WITH_EDITOR
+	virtual void PostTransacted(const FTransactionObjectEvent& TransactionEvent) override;
+#endif
+
+	/**
+	 * Clears the whole hierarchy and removes all elements.
+	 */
+	UFUNCTION(BlueprintCallable, Category = URigHierarchy)
+	void Reset();
+
+	/**
+	 * Copies the contents of a hierarchy onto this one
+	 */
+	UFUNCTION(BlueprintCallable, Category = URigHierarchy)
+	void CopyHierarchy(URigHierarchy* InHierarchy);
+
+	/**
+	 * Copies the contents of a hierarchy onto this one
+	 */
+	UFUNCTION(BlueprintCallable, Category = URigHierarchy)
+	void CopyPose(URigHierarchy* InHierarchy, bool bCurrent, bool bInitial);
+
+	/**
+	 * Resets the current pose of a filtered lost if elements to the initial / ref pose.
+	 */
+	UFUNCTION(BlueprintCallable, Category = URigHierarchy)
+	void ResetPoseToInitial(ERigElementType InTypeFilter);
+
+	/**
+	 * Resets the current pose of all elements to the initial / ref pose.
+	 */
+	void ResetPoseToInitial()
+	{
+		ResetPoseToInitial(ERigElementType::All);
+	}
+
+	/**
+	 * Resets all curves to 0.0
+	 */
+	UFUNCTION(BlueprintCallable, Category = URigHierarchy)
+    void ResetCurveValues();
+
+	/**
+	 * Returns the number of elements in the Hierarchy.
+	 * @return The number of elements in the Hierarchy
+	 */
+	UFUNCTION(BlueprintCallable, Category = URigHierarchy)
+	FORCEINLINE int32 Num() const
+	{
+		return Elements.Num();
+	}
+
+	/**
+	 * Returns the number of elements in the Hierarchy.
+	 * @param InElementType The type filter to apply
+	 * @return The number of elements in the Hierarchy
+	 */
+    int32 Num(ERigElementType InElementType) const;
+
+	// iterators
+	FORCEINLINE TArray<FRigBaseElement*>::RangedForIteratorType      begin() { return Elements.begin(); }
+	FORCEINLINE TArray<FRigBaseElement*>::RangedForIteratorType      end() { return Elements.end(); }
+
+	/**
+	 * Iterator function to invoke a lambda / TFunction for each element
+	 * @param PerElementFunction The function to invoke for each element
+	 */
+	FORCEINLINE_DEBUGGABLE void ForEach(TFunction<bool(FRigBaseElement*)> PerElementFunction) const
+	{
+		for (int32 ElementIndex = 0; ElementIndex < Elements.Num(); ElementIndex++)
+		{
+			if(!PerElementFunction(Elements[ElementIndex]))
+			{
+				return;
+			}
+		}
+	}
+
+	/**
+	 * Filtered template Iterator function to invoke a lambda / TFunction for each element of a given type.
+	 * @param PerElementFunction The function to invoke for each element of a given type
+	 */
+	template<typename T>
+	FORCEINLINE_DEBUGGABLE void ForEach(TFunction<bool(T*)> PerElementFunction) const
+	{
+		for (int32 ElementIndex = 0; ElementIndex < Elements.Num(); ElementIndex++)
+		{
+			if(T* CastElement = Cast<T>(Elements[ElementIndex]))
+			{
+				if(!PerElementFunction(CastElement))
+				{
+					return;
+				}
+			}
+		}
+	}
+
+	/**
+	 * Returns true if the provided element index is valid
+	 * @param InElementIndex The index to validate
+	 * @return Returns true if the provided element index is valid
+	 */
+	UFUNCTION(BlueprintCallable, Category = URigHierarchy)
+	FORCEINLINE bool IsValidIndex(int32 InElementIndex) const
+	{
+		return Elements.IsValidIndex(InElementIndex);
+	}
+
+	/**
+	 * Returns true if the provided element key is valid
+	 * @param InKey The key to validate
+	 * @return Returns true if the provided element key is valid
+	 */
+	UFUNCTION(BlueprintCallable, Category = URigHierarchy, meta = (DisplayName = "Contains", ScriptName = "Contains"))
+	FORCEINLINE bool Contains_ForBlueprint(FRigElementKey InKey) const
+	{
+		return Contains(InKey);
+	}
+
+	/**
+	 * Returns true if the provided element key is valid
+	 * @param InKey The key to validate
+	 * @return Returns true if the provided element key is valid
+	 */
+	FORCEINLINE bool Contains(const FRigElementKey& InKey) const
+	{
+		return GetIndex(InKey) != INDEX_NONE;
+	}
+
+	/**
+	 * Returns the index of an element given its key
+	 * @param InKey The key of the element to retrieve the index for
+	 * @return The index of the element or INDEX_NONE
+	 */
+	UFUNCTION(BlueprintCallable, Category = URigHierarchy, meta = (DisplayName = "Get Index", ScriptName = "GetIndex"))
+	FORCEINLINE int32 GetIndex_ForBlueprint(FRigElementKey InKey) const
+	{
+		return GetIndex(InKey);
+	}
+
+	/**
+	 * Returns the index of an element given its key
+	 * @param InKey The key of the element to retrieve the index for
+	 * @return The index of the element or INDEX_NONE
+	 */
+	FORCEINLINE int32 GetIndex(const FRigElementKey& InKey) const
+	{
+		if(const int32* Index = IndexLookup.Find(InKey))
+		{
+			return *Index;
+		}
+		return INDEX_NONE;
+	}
+
+	/**
+	 * Returns the key of an element given its index
+	 * @param InElementIndex The index of the element to retrieve the key for
+	 * @return The key of an element given its index
+	 */
+	UFUNCTION(BlueprintCallable, Category = URigHierarchy)
+    FORCEINLINE FRigElementKey GetKey(int32 InElementIndex) const
+	{
+		if(Elements.IsValidIndex(InElementIndex))
+		{
+			return Elements[InElementIndex]->Key;
+		}
+		return FRigElementKey();
+	}
+
+	/**
+	 * Returns an element at a given index or nullptr.
+	 * @param InIndex The index of the element to retrieve.
+	 * @return The retrieved element or nullptr.
+	 */
+	FORCEINLINE const FRigBaseElement* Get(int32 InIndex) const
+	{
+		if(Elements.IsValidIndex(InIndex))
+		{
+			return Elements[InIndex];
+		}
+		return nullptr;
+	}
+
+	/**
+	 * Returns an element at a given index or nullptr.
+	 * @param InIndex The index of the element to retrieve.
+	 * @return The retrieved element or nullptr.
+	 */
+	FRigBaseElement* Get(int32 InIndex)
+	{
+		if(Elements.IsValidIndex(InIndex))
+		{
+			return Elements[InIndex];
+		}
+		return nullptr;
+	}
+
+	/**
+	 * Returns an element at a given index or nullptr.
+	 * This templated method also casts to the chosen
+	 * element type but does not guarantee a valid result.
+	 * @param InIndex The index of the element to retrieve.
+	 * @return The retrieved element or nullptr.
+	 */
+	template<typename T>
+	FORCEINLINE const T* Get(int32 InIndex) const
+	{
+		return Cast<T>(Get(InIndex));
+	}
+
+	/**
+	 * Returns an element at a given index or nullptr.
+	 * This templated method also casts to the chosen
+	 * element type but does not guarantee a valid result.
+	 * @param InIndex The index of the element to retrieve.
+	 * @return The retrieved element or nullptr.
+	 */
+	template<typename T>
+	FORCEINLINE T* Get(int32 InIndex)
+	{
+		return Cast<T>(Get(InIndex));
+	}
+
+	/**
+	 * Returns an element at a given index.
+	 * This templated method also casts to the chosen
+	 * element type and checks for a the valid result.
+	 * @param InIndex The index of the element to retrieve.
+	 * @return The retrieved element or nullptr.
+	 */
+	template<typename T>
+	FORCEINLINE const T* GetChecked(int32 InIndex) const
+	{
+		return CastChecked<T>(Get(InIndex));
+	}
+
+	/**
+	 * Returns an element at a given index.
+	 * This templated method also casts to the chosen
+	 * element type and checks for a the valid result.
+	 * @param InIndex The index of the element to retrieve.
+	 * @return The retrieved element or nullptr.
+	 */
+	template<typename T>
+    FORCEINLINE T* GetChecked(int32 InIndex)
+	{
+		return CastChecked<T>(Get(InIndex));
+	}
+
+	/**
+	 * Returns an element for a given key or nullptr.
+	 * @param InKey The key of the element to retrieve.
+	 * @return The retrieved element or nullptr.
+	 */
+	FORCEINLINE const FRigBaseElement* Find(const FRigElementKey& InKey) const
+	{
+		return Get(GetIndex(InKey));
+	}
+
+	/**
+	 * Returns an element for a given key or nullptr.
+	 * @param InKey The key of the element to retrieve.
+	 * @return The retrieved element or nullptr.
+	 */
+	FORCEINLINE FRigBaseElement* Find(const FRigElementKey& InKey)
+	{
+		return Get(GetIndex(InKey));
+	}
+
+	/**
+	 * Returns an element for a given key and raises for invalid results.
+	 * @param InKey The key of the element to retrieve.
+	 * @return The retrieved element or nullptr.
+	 */
+	FORCEINLINE const FRigBaseElement* FindChecked(const FRigElementKey& InKey) const
+	{
+		const FRigBaseElement* Element = Get(GetIndex(InKey));
+		check(Element);
+		return Element;
+	}
+
+	/**
+	 * Returns an element for a given key and raises for invalid results.
+	 * @param InKey The key of the element to retrieve.
+	 * @return The retrieved element or nullptr.
+	 */
+	FORCEINLINE FRigBaseElement* FindChecked(const FRigElementKey& InKey)
+	{
+		FRigBaseElement* Element = Get(GetIndex(InKey));
+		check(Element);
+		return Element;
+	}
+
+	/**
+	 * Returns an element for a given key or nullptr.
+	 * This templated method also casts to the chosen
+	 * element type but does not guarantee a valid result.
+	 * @param InKey The key of the element to retrieve.
+	 * @return The retrieved element or nullptr.
+	 */
+	template<typename T>
+	FORCEINLINE const T* Find(const FRigElementKey& InKey) const
+	{
+		return Get<T>(GetIndex(InKey));
+	}
+
+	/**
+	 * Returns an element for a given key or nullptr.
+	 * This templated method also casts to the chosen
+	 * element type but does not guarantee a valid result.
+	 * @param InKey The key of the element to retrieve.
+	 * @return The retrieved element or nullptr.
+	 */
+	template<typename T>
+	FORCEINLINE T* Find(const FRigElementKey& InKey)
+	{
+		return Get<T>(GetIndex(InKey));
+	}
+
+	/**
+	 * Returns an element for a given key.
+	 * This templated method also casts to the chosen
+	 * element type and checks for a the valid result.
+	 * @param InKey The key of the element to retrieve.
+	 * @return The retrieved element or nullptr.
+	 */
+	template<typename T>
+	FORCEINLINE const T* FindChecked(const FRigElementKey& InKey) const
+	{
+		return GetChecked<T>(GetIndex(InKey));
+	}
+
+	/**
+	 * Returns an element for a given key.
+	 * This templated method also casts to the chosen
+	 * element type and checks for a the valid result.
+	 * @param InKey The key of the element to retrieve.
+	 * @return The retrieved element or nullptr.
+	 */
+	template<typename T>
+    FORCEINLINE T* FindChecked(const FRigElementKey& InKey)
+	{
+		return GetChecked<T>(GetIndex(InKey));
+	}
+
+	/**
+	 * Filtered accessor to retrieve all elements of a given type
+	 * @param bTraverse Returns the elements in order of a depth first traversal
+	 */
+	template<typename T>
+	FORCEINLINE TArray<T*> GetElementsOfType(bool bTraverse = false) const
+	{
+		TArray<T*> Results;
+
+		if(bTraverse)
+		{
+			TArray<bool> ElementVisited;
+			ElementVisited.AddZeroed(Elements.Num());
+		
+			for (int32 ElementIndex = 0; ElementIndex < Elements.Num(); ElementIndex++)
+			{
+				FRigBaseElement* Element = Elements[ElementIndex];
+				Traverse(Element, true, [&ElementVisited, &Results](FRigBaseElement* InElement, bool& bContinue)
+                {
+                    bContinue = !ElementVisited[InElement->GetIndex()];
+
+                    if(bContinue)
+                    {
+                        if(T* CastElement = Cast<T>(InElement))
+                        {
+                            Results.Add(CastElement);
+                        }
+                        ElementVisited[InElement->GetIndex()] = true;
+                    }
+                });
+			}
+		}
+		else
+		{
+			for (int32 ElementIndex = 0; ElementIndex < Elements.Num(); ElementIndex++)
+			{
+				if(T* CastElement = Cast<T>(Elements[ElementIndex]))
+				{
+					Results.Add(CastElement);
+				}
+			}
+		}
+		return Results;
+	}
+
+	/**
+	 * Filtered accessor to retrieve all element keys of a given type
+	 * @param bTraverse Returns the element keys in order of a depth first traversal
+	 */
+	template<typename T>
+    FORCEINLINE TArray<FRigElementKey> GetKeysOfType(bool bTraverse = false) const
+	{
+		TArray<FRigElementKey> Keys;
+		TArray<T*> Results = GetElementsOfType<T>(bTraverse);
+		for(T* Element : Results)
+		{
+			Keys.Add(Element->GetKey());
+		}
+		return Keys;
+	}
+
+	/**
+	 * Filtered accessor to retrieve all elements of a given type
+	 * @param InKeepElementFunction A function to return true if an element is to be keep
+	 * @param bTraverse Returns the elements in order of a depth first traversal
+	 */
+	template<typename T>
+    FORCEINLINE TArray<T*> GetFilteredElements(TFunction<bool(T*)> InKeepElementFunction, bool bTraverse = false) const
+	{
+		TArray<T*> Results;
+
+		if(bTraverse)
+		{
+			TArray<bool> ElementVisited;
+			ElementVisited.AddZeroed(Elements.Num());
+		
+			for (int32 ElementIndex = 0; ElementIndex < Elements.Num(); ElementIndex++)
+			{
+				FRigBaseElement* Element = Elements[ElementIndex];
+				Traverse(Element, true, [&ElementVisited, &Results, InKeepElementFunction](FRigBaseElement* InElement, bool& bContinue)
+                {
+                    bContinue = !ElementVisited[InElement->GetIndex()];
+
+                    if(bContinue)
+                    {
+                        if(T* CastElement = Cast<T>(InElement))
+                        {
+							if(InKeepElementFunction(CastElement))
+							{
+								Results.Add(CastElement);
+							}
+                        }
+                        ElementVisited[InElement->GetIndex()] = true;
+                    }
+                });
+			}
+		}
+		else
+		{
+			for (int32 ElementIndex = 0; ElementIndex < Elements.Num(); ElementIndex++)
+			{
+				if(T* CastElement = Cast<T>(Elements[ElementIndex]))
+				{
+					if(InKeepElementFunction(CastElement))
+					{
+						Results.Add(CastElement);
+					}
+				}
+			}
+		}
+		return Results;
+	}
+
+	/**
+	 * Returns all Bone elements
+	 * @param bTraverse Returns the elements in order of a depth first traversal
+	 */
+	FORCEINLINE TArray<FRigBoneElement*> GetBones(bool bTraverse = false) const
+	{
+		return GetElementsOfType<FRigBoneElement>(bTraverse);
+	}
+
+	/**
+	 * Returns all Bone elements
+	 * @param bTraverse Returns the elements in order of a depth first traversal
+	 */
+	UFUNCTION(BlueprintCallable, Category = URigHierarchy, meta = (DisplayName = "Get Bones", ScriptName = "GetBones"))
+	FORCEINLINE TArray<FRigElementKey> GetBoneKeys(bool bTraverse = true) const
+	{
+		return GetKeysOfType<FRigBoneElement>(bTraverse);
+	}
+
+	/**
+	 * Returns all Space elements
+	 * @param bTraverse Returns the elements in order of a depth first traversal
+	 */
+	FORCEINLINE TArray<FRigSpaceElement*> GetSpaces(bool bTraverse = false) const
+	{
+		return GetElementsOfType<FRigSpaceElement>(bTraverse);
+	}
+
+	/**
+	 * Returns all Space elements
+	 * @param bTraverse Returns the elements in order of a depth first traversal
+	 */
+	UFUNCTION(BlueprintCallable, Category = URigHierarchy, meta = (DisplayName = "Get Spaces", ScriptName = "GetSpaces"))
+	FORCEINLINE TArray<FRigElementKey> GetSpaceKeys(bool bTraverse = true) const
+	{
+		return GetKeysOfType<FRigSpaceElement>(bTraverse);
+	}
+
+	/**
+	 * Returns all Control elements
+	 * @param bTraverse Returns the elements in order of a depth first traversal
+	 */
+	FORCEINLINE TArray<FRigControlElement*> GetControls(bool bTraverse = false) const
+	{
+		return GetElementsOfType<FRigControlElement>(bTraverse);
+	}
+
+	/**
+	 * Returns all Control elements
+	 * @param bTraverse Returns the elements in order of a depth first traversal
+	 */
+	UFUNCTION(BlueprintCallable, Category = URigHierarchy, meta = (DisplayName = "Get Controls", ScriptName = "GetControls"))
+	FORCEINLINE TArray<FRigElementKey> GetControlKeys(bool bTraverse = true) const
+	{
+		return GetKeysOfType<FRigControlElement>(bTraverse);
+	}
+
+	/**
+	 * Returns all transient Control elements
+	 */
+	FORCEINLINE TArray<FRigControlElement*> GetTransientControls() const
+	{
+		return GetFilteredElements<FRigControlElement>([](FRigControlElement* ControlElement) -> bool
+		{
+			return ControlElement->Settings.bIsTransientControl;
+		});
+	}
+
+	/**
+	 * Returns all Curve elements
+	 */
+	FORCEINLINE TArray<FRigCurveElement*> GetCurves() const
+	{
+		return GetElementsOfType<FRigCurveElement>();
+	}
+
+	/**
+	 * Returns all Curve elements
+	 */
+	UFUNCTION(BlueprintCallable, Category = URigHierarchy, meta = (DisplayName = "Get Curves", ScriptName = "GetCurves"))
+	FORCEINLINE TArray<FRigElementKey> GetCurveKeys() const
+	{
+		return GetKeysOfType<FRigCurveElement>(false);
+	}
+
+
+	/**
+	 * Returns the selected elements
+	 * @InTypeFilter The types to retrieve the selection for
+	 * @return An array of the currently selected elements
+	 */
+	TArray<FRigBaseElement*> GetSelectedElements(ERigElementType InTypeFilter = ERigElementType::All) const;
+
+	/**
+	 * Returns the keys of selected elements
+	 * @InTypeFilter The types to retrieve the selection for
+	 * @return An array of the currently selected elements
+	 */
+	TArray<FRigElementKey> GetSelectedKeys(ERigElementType InTypeFilter = ERigElementType::All) const;
+
+	/**
+	 * Returns true if a given element is selected
+	 * @param InKey The key to check
+	 * @return true if a given element is selected
+	 */
+	UFUNCTION(BlueprintCallable, Category = URigHierarchy)
+	FORCEINLINE bool IsSelected(FRigElementKey InKey) const
+	{
+		return IsSelected(Find(InKey));
+	}
+
+	/**
+	 * Returns true if a given element is selected
+	 * @param InIndex The index to check
+	 * @return true if a given element is selected
+	 */
+	UFUNCTION(BlueprintCallable, Category = URigHierarchy)
+    FORCEINLINE bool IsSelectedByIndex(int32 InIndex) const
+	{
+		return IsSelected(Get(InIndex));
+	}
+
+	FORCEINLINE bool IsSelected(int32 InIndex) const
+	{
+		return IsSelectedByIndex(InIndex);
+	}
+
+	/**
+	 * Sorts the input key list by traversing the hierarchy
+	 * @param InKeys The keys to sort
+	 * @return The sorted keys
+	 */
+	FORCEINLINE TArray<FRigElementKey> SortKeys(const TArray<FRigElementKey>& InKeys) const
+	{
+		TArray<FRigElementKey> Result;
+		Traverse([InKeys, &Result](FRigBaseElement* Element, bool& bContinue)
+        {
+            const FRigElementKey& Key = Element->GetKey();
+            if(InKeys.Contains(Key))
+            {
+                Result.AddUnique(Key);
+            }
+        });
+		return Result;
+	}
+
+	/**
+	 * Returns the max allowed length for a name within the hierarchy.
+	 * @return Returns the max allowed length for a name within the hierarchy.
+	 */
+	static int32 GetMaxNameLength() { return 100; }
+
+	/**
+	 * Sanitizes a name by removing invalid characters.
+	 * @param InOutName The name to sanitize in place.
+	 */
+	static void SanitizeName(FString& InOutName);
+
+	/**
+	 * Sanitizes a name by removing invalid characters.
+	 * @param InName The name to sanitize.
+	 * @return The sanitized name.
+ 	 */
+	static FName GetSanitizedName(const FString& InName);
+
+	/**
+	 * Returns true if a given name is available.
+	 * @param InPotentialNewName The name to test for availability
+	 * @param InType The type of the to-be-added element
+	 * @param OutErrorMessage An optional pointer to return a potential error message 
+	 * @return Returns true if the name is available.
+	 */
+	bool IsNameAvailable(const FString& InPotentialNewName, ERigElementType InType, FString* OutErrorMessage = nullptr) const;
+
+	/**
+	 * Returns a valid new name for a to-be-added element.
+	 * @param InPotentialNewName The name to be sanitized and adjusted for availability
+	 * @param InType The type of the to-be-added element
+	 * @return Returns the name to use for the to-be-added element.
+	 */
+	FName GetSafeNewName(const FString& InPotentialNewName, ERigElementType InType) const;
+
+	/**
+	 * Returns the modified event, which can be used to 
+	 * subscribe to topological changes happening within the hierarchy.
+	 * @return The event used for subscription.
+	 */
+	FRigHierarchyModifiedEvent& OnModified() { return ModifiedEvent; }
+
+	/**
+	 * Returns the local current or initial value for a given key.
+	 * If the key is invalid FTransform::Identity will be returned.
+	 * @param InKey The key to retrieve the transform for
+	 * @param bInitial If true the initial transform will be used
+	 * @return The local current or initial transform's value.
+	 */
+	UFUNCTION(BlueprintCallable, Category = URigHierarchy)
+    FORCEINLINE_DEBUGGABLE FTransform GetLocalTransform(FRigElementKey InKey, bool bInitial = false) const
+	{
+		return GetLocalTransformByIndex(GetIndex(InKey), bInitial);
+	}
+
+	/**
+	 * Returns the local current or initial value for a element index.
+	 * If the index is invalid FTransform::Identity will be returned.
+	 * @param InElementIndex The index to retrieve the transform for
+	 * @param bInitial If true the initial transform will be used
+	 * @return The local current or initial transform's value.
+	 */
+	UFUNCTION(BlueprintCallable, Category = URigHierarchy)
+	FORCEINLINE_DEBUGGABLE FTransform GetLocalTransformByIndex(int32 InElementIndex, bool bInitial = false) const
+	{
+		if(Elements.IsValidIndex(InElementIndex))
+		{
+			if(FRigTransformElement* TransformElement = Cast<FRigTransformElement>(Elements[InElementIndex]))
+			{
+				return GetTransform(TransformElement, bInitial ? ERigTransformType::InitialLocal : ERigTransformType::CurrentLocal);
+			}
+		}
+		return FTransform::Identity;
+	}
+
+	FORCEINLINE_DEBUGGABLE FTransform GetLocalTransform(int32 InElementIndex) const
+	{
+		return GetLocalTransformByIndex(InElementIndex, false);
+	}
+	FORCEINLINE_DEBUGGABLE FTransform GetInitialLocalTransform(int32 InElementIndex) const
+	{
+		return GetLocalTransformByIndex(InElementIndex, true);
+	}
+
+	FORCEINLINE_DEBUGGABLE FTransform GetInitialLocalTransform(const FRigElementKey &InKey) const
+	{
+		return GetLocalTransform(InKey, true);
+	}
+
+	/**
+	 * Sets the local current or initial transform for a given key.
+	 * @param InKey The key to set the transform for
+	 * @param InTransform The new transform value to set
+	 * @param bInitial If true the initial transform will be used
+	 * @param bSetupUndo If true the transform stack will be setup for undo / redo
+	 * @param bAffectChildren If set to false children will not move (maintain global).
+	 */
+	UFUNCTION(BlueprintCallable, Category = URigHierarchy)
+    FORCEINLINE_DEBUGGABLE void SetLocalTransform(FRigElementKey InKey, FTransform InTransform, bool bInitial = false, bool bAffectChildren = true, bool bSetupUndo = false)
+	{
+		SetLocalTransformByIndex(GetIndex(InKey), InTransform, bInitial, bAffectChildren, bSetupUndo);
+	}
+
+	/**
+	 * Sets the local current or initial transform for a given element index.
+	 * @param InElementIndex The index of the element to set the transform for
+	 * @param InTransform The new transform value to set
+	 * @param bInitial If true the initial transform will be used
+	 * @param bSetupUndo If true the transform stack will be setup for undo / redo
+	 * @param bAffectChildren If set to false children will not move (maintain global).
+	 */
+	UFUNCTION(BlueprintCallable, Category = URigHierarchy)
+    FORCEINLINE_DEBUGGABLE void SetLocalTransformByIndex(int32 InElementIndex, FTransform InTransform, bool bInitial = false, bool bAffectChildren = true, bool bSetupUndo = false)
+	{
+		if(Elements.IsValidIndex(InElementIndex))
+		{
+			if(FRigTransformElement* TransformElement = Cast<FRigTransformElement>(Elements[InElementIndex]))
+			{
+				SetTransform(TransformElement, InTransform, bInitial ? ERigTransformType::InitialLocal : ERigTransformType::CurrentLocal, bAffectChildren, bSetupUndo);
+			}
+		}
+	}
+
+	FORCEINLINE_DEBUGGABLE void SetLocalTransform(int32 InElementIndex, const FTransform& InTransform, bool bAffectChildren = true, bool bSetupUndo = false)
+	{
+		SetLocalTransformByIndex(InElementIndex, InTransform, false, bAffectChildren, bSetupUndo);
+	}
+
+	FORCEINLINE_DEBUGGABLE void SetInitialLocalTransform(int32 InElementIndex, const FTransform& InTransform, bool bAffectChildren = true, bool bSetupUndo = false)
+	{
+		SetLocalTransformByIndex(InElementIndex, InTransform, true, bAffectChildren, bSetupUndo);
+    }
+
+	FORCEINLINE_DEBUGGABLE void SetInitialLocalTransform(const FRigElementKey& InKey, const FTransform& InTransform, bool bAffectChildren = true, bool bSetupUndo = false)
+	{
+		SetLocalTransform(InKey, InTransform, true, bAffectChildren, bSetupUndo);
+	}
+
+	/**
+	 * Returns the global current or initial value for a given key.
+	 * If the key is invalid FTransform::Identity will be returned.
+	 * @param InKey The key to retrieve the transform for
+	 * @param bInitial If true the initial transform will be used
+	 * @return The global current or initial transform's value.
+	 */
+	UFUNCTION(BlueprintCallable, Category = URigHierarchy)
+    FORCEINLINE_DEBUGGABLE FTransform GetGlobalTransform(FRigElementKey InKey, bool bInitial = false) const
+	{
+		return GetGlobalTransformByIndex(GetIndex(InKey), bInitial);
+	}
+
+	/**
+	 * Returns the global current or initial value for a element index.
+	 * If the index is invalid FTransform::Identity will be returned.
+	 * @param InElementIndex The index to retrieve the transform for
+	 * @param bInitial If true the initial transform will be used
+	 * @return The global current or initial transform's value.
+	 */
+	UFUNCTION(BlueprintCallable, Category = URigHierarchy)
+    FORCEINLINE_DEBUGGABLE FTransform GetGlobalTransformByIndex(int32 InElementIndex, bool bInitial = false) const
+	{
+		if(Elements.IsValidIndex(InElementIndex))
+		{
+			if(FRigTransformElement* TransformElement = Cast<FRigTransformElement>(Elements[InElementIndex]))
+			{
+				return GetTransform(TransformElement, bInitial ? ERigTransformType::InitialGlobal : ERigTransformType::CurrentGlobal);
+			}
+		}
+		return FTransform::Identity;
+	}
+
+	FORCEINLINE_DEBUGGABLE FTransform GetGlobalTransform(int32 InElementIndex) const
+	{
+		return GetGlobalTransformByIndex(InElementIndex, false);
+	}
+	FORCEINLINE_DEBUGGABLE FTransform GetInitialGlobalTransform(int32 InElementIndex) const
+	{
+		return GetGlobalTransformByIndex(InElementIndex, true);
+	}
+
+	FORCEINLINE_DEBUGGABLE FTransform GetInitialGlobalTransform(const FRigElementKey &InKey) const
+	{
+		return GetGlobalTransform(InKey, true);
+	}
+
+	/**
+	 * Sets the global current or initial transform for a given key.
+	 * @param InKey The key to set the transform for
+	 * @param InTransform The new transform value to set
+	 * @param bInitial If true the initial transform will be used
+	 * @param bAffectChildren If set to false children will not move (maintain global).
+	 * @param bSetupUndo If true the transform stack will be setup for undo / redo
+	 */
+	UFUNCTION(BlueprintCallable, Category = URigHierarchy)
+    FORCEINLINE_DEBUGGABLE void SetGlobalTransform(FRigElementKey InKey, FTransform InTransform, bool bInitial = false, bool bAffectChildren = true, bool bSetupUndo = false)
+	{
+		SetGlobalTransformByIndex(GetIndex(InKey), InTransform, bInitial, bAffectChildren, bSetupUndo);
+	}
+
+	/**
+	 * Sets the global current or initial transform for a given element index.
+	 * @param InElementIndex The index of the element to set the transform for
+	 * @param InTransform The new transform value to set
+	 * @param bInitial If true the initial transform will be used
+	 * @param bAffectChildren If set to false children will not move (maintain global).
+	 * @param bSetupUndo If true the transform stack will be setup for undo / redo
+	 */
+	UFUNCTION(BlueprintCallable, Category = URigHierarchy)
+    FORCEINLINE_DEBUGGABLE void SetGlobalTransformByIndex(int32 InElementIndex, FTransform InTransform, bool bInitial = false, bool bAffectChildren = true, bool bSetupUndo = false)
+	{
+		if(Elements.IsValidIndex(InElementIndex))
+		{
+			if(FRigTransformElement* TransformElement = Cast<FRigTransformElement>(Elements[InElementIndex]))
+			{
+				SetTransform(TransformElement, InTransform, bInitial ? ERigTransformType::InitialGlobal : ERigTransformType::CurrentGlobal, bAffectChildren, bSetupUndo);
+			}
+		}
+	}
+
+	FORCEINLINE_DEBUGGABLE void SetGlobalTransform(int32 InElementIndex, const FTransform& InTransform, bool bAffectChildren = true, bool bSetupUndo = false)
+	{
+		SetGlobalTransformByIndex(InElementIndex, InTransform, false, bAffectChildren, bSetupUndo);
+	}
+
+	FORCEINLINE_DEBUGGABLE void SetInitialGlobalTransform(int32 InElementIndex, const FTransform& InTransform, bool bAffectChildren = true, bool bSetupUndo = false)
+	{
+		SetGlobalTransformByIndex(InElementIndex, InTransform, true, bAffectChildren, bSetupUndo);
+	}
+
+	FORCEINLINE_DEBUGGABLE void SetInitialGlobalTransform(const FRigElementKey& InKey, const FTransform& InTransform, bool bAffectChildren = true, bool bSetupUndo = false)
+	{
+		SetGlobalTransform(InKey, InTransform, true, bAffectChildren, bSetupUndo);
+	}
+
+	/**
+	 * Returns the global offset transform for a given control element.
+	 * @param InKey The key of the control to retrieve the transform for
+	 * @param bInitial If true the initial transform will be used
+	 * @return The global offset transform
+	 */
+	UFUNCTION(BlueprintCallable, Category = URigHierarchy)
+    FORCEINLINE_DEBUGGABLE FTransform GetGlobalControlOffsetTransform(FRigElementKey InKey, bool bInitial = false) const
+	{
+		return GetGlobalControlOffsetTransformByIndex(GetIndex(InKey), bInitial);
+	}
+
+	/**
+	 * Returns the global offset transform for a given control element.
+	 * @param InElementIndex The index of the control to retrieve the transform for
+	 * @param bInitial If true the initial transform will be used
+	 * @return The global offset transform
+	 */
+	UFUNCTION(BlueprintCallable, Category = URigHierarchy)
+    FORCEINLINE_DEBUGGABLE FTransform GetGlobalControlOffsetTransformByIndex(int32 InElementIndex, bool bInitial = false) const
+	{
+		if(Elements.IsValidIndex(InElementIndex))
+		{
+			if(FRigControlElement* ControlElement = Cast<FRigControlElement>(Elements[InElementIndex]))
+			{
+				return GetControlOffsetTransform(ControlElement, bInitial ? ERigTransformType::InitialGlobal : ERigTransformType::CurrentGlobal);
+			}
+		}
+		return FTransform::Identity;
+	}
+
+	/**
+	 * Returns the global gizmo transform for a given control element.
+	 * @param InKey The key of the control to retrieve the transform for
+	 * @param bInitial If true the initial transform will be used
+	 * @return The global gizmo transform
+	 */
+	UFUNCTION(BlueprintCallable, Category = URigHierarchy)
+    FORCEINLINE_DEBUGGABLE FTransform GetGlobalControlGizmoTransform(FRigElementKey InKey, bool bInitial = false) const
+	{
+		return GetGlobalControlGizmoTransformByIndex(GetIndex(InKey), bInitial);
+	}
+
+	/**
+	 * Returns the global gizmo transform for a given control element.
+	 * @param InElementIndex The index of the control to retrieve the transform for
+	 * @param bInitial If true the initial transform will be used
+	 * @return The global gizmo transform
+	 */
+	UFUNCTION(BlueprintCallable, Category = URigHierarchy)
+    FORCEINLINE_DEBUGGABLE FTransform GetGlobalControlGizmoTransformByIndex(int32 InElementIndex, bool bInitial = false) const
+	{
+		if(Elements.IsValidIndex(InElementIndex))
+		{
+			if(FRigControlElement* ControlElement = Cast<FRigControlElement>(Elements[InElementIndex]))
+			{
+				return GetControlGizmoTransform(ControlElement, bInitial ? ERigTransformType::InitialGlobal : ERigTransformType::CurrentGlobal);
+			}
+		}
+		return FTransform::Identity;
+	}
+
+	/**
+	 * Returns a control's current value given its key
+	 * @param InKey The key of the element to retrieve the current value for
+	 * @param InValueType The type of value to return
+	 * @return Returns the current value of the control
+	 */
+	UFUNCTION(BlueprintCallable, Category = URigHierarchy)
+    FORCEINLINE_DEBUGGABLE FRigControlValue GetControlValue(FRigElementKey InKey, ERigControlValueType InValueType = ERigControlValueType::Current) const
+	{
+		return GetControlValueByIndex(GetIndex(InKey), InValueType);
+	}
+
+	/**
+	 * Returns a control's current value given its key
+	 * @param InKey The key of the element to retrieve the current value for
+	 * @param InValueType The type of value to return
+	 * @return Returns the current value of the control
+	 */
+	template<typename T>
+	FORCEINLINE_DEBUGGABLE T GetControlValue(FRigElementKey InKey, ERigControlValueType InValueType = ERigControlValueType::Current) const
+	{
+		return GetControlValueByIndex(GetIndex(InKey), InValueType).Get<T>();
+	}
+
+	/**
+	 * Returns a control's current value given its index
+	 * @param InElementIndex The index of the element to retrieve the current value for
+	 * @param InValueType The type of value to return
+	 * @return Returns the current value of the control
+	 */
+	UFUNCTION(BlueprintCallable, Category = URigHierarchy)
+    FORCEINLINE_DEBUGGABLE FRigControlValue GetControlValueByIndex(int32 InElementIndex, ERigControlValueType InValueType = ERigControlValueType::Current) const
+	{
+		if(Elements.IsValidIndex(InElementIndex))
+		{
+			if(FRigControlElement* ControlElement = Cast<FRigControlElement>(Elements[InElementIndex]))
+			{
+				return GetControlValue(ControlElement, InValueType);
+			}
+		}
+		return FRigControlValue();
+	}
+
+	/**
+	 * Returns a control's current value given its index
+	 * @param InElementIndex The index of the element to retrieve the current value for
+	 * @param InValueType The type of value to return
+	 * @return Returns the current value of the control
+	 */
+	FORCEINLINE_DEBUGGABLE FRigControlValue GetControlValue(int32 InElementIndex, ERigControlValueType InValueType = ERigControlValueType::Current) const
+	{
+		return GetControlValueByIndex(InElementIndex, InValueType);
+	}
+
+	/**
+	 * Returns a control's current value given its index
+	 * @param InElementIndex The index of the element to retrieve the current value for
+	 * @param InValueType The type of value to return
+	 * @return Returns the current value of the control
+	 */
+	template<typename T>
+	FORCEINLINE_DEBUGGABLE T GetControlValue(int32 InElementIndex, ERigControlValueType InValueType = ERigControlValueType::Current) const
+	{
+		return GetControlValueByIndex(InElementIndex, InValueType).Get<T>();
+	}
+
+	/**
+	 * Returns a control's initial value given its index
+	 * @param InElementIndex The index of the element to retrieve the initial value for
+	 * @return Returns the current value of the control
+	 */
+	FORCEINLINE_DEBUGGABLE FRigControlValue GetInitialControlValue(int32 InElementIndex) const
+	{
+		return GetControlValueByIndex(InElementIndex, ERigControlValueType::Initial);
+	}
+
+	/**
+	 * Returns a control's initial value given its index
+	 * @param InElementIndex The index of the element to retrieve the current value for
+	 * @return Returns the current value of the control
+	 */
+	template<typename T>
+	FORCEINLINE_DEBUGGABLE T GetInitialControlValue(int32 InElementIndex) const
+	{
+		return GetInitialControlValue(InElementIndex).Get<T>();
+	}
+
+	/**
+	 * Sets a control's current value given its key
+	 * @param InKey The key of the element to set the current value for
+	 * @param InValue The value to set on the control
+	 * @param InValueType The type of value to set
+	 * @param bSetupUndo If true the transform stack will be setup for undo / redo
+	 */
+	UFUNCTION(BlueprintCallable, Category = URigHierarchy)
+    FORCEINLINE_DEBUGGABLE void SetControlValue(FRigElementKey InKey, FRigControlValue InValue, ERigControlValueType InValueType = ERigControlValueType::Current, bool bSetupUndo = false)
+	{
+		SetControlValueByIndex(GetIndex(InKey), InValue, InValueType, bSetupUndo);
+	}
+
+	/**
+	 * Sets a control's current value given its key
+	 * @param InKey The key of the element to set the current value for
+	 * @param InValue The value to set on the control
+	 * @param InValueType The type of value to set
+	 * @param bSetupUndo If true the transform stack will be setup for undo / redo
+	 */
+	template<typename T>
+	FORCEINLINE_DEBUGGABLE void SetControlValue(FRigElementKey InKey, const T& InValue, ERigControlValueType InValueType = ERigControlValueType::Current, bool bSetupUndo = false) const
+	{
+		return SetControlValue(InKey, FRigControlValue::Make<T>(InValue), InValueType, bSetupUndo);
+	}
+
+	/**
+	 * Sets a control's current value given its index
+	 * @param InElementIndex The index of the element to set the current value for
+	 * @param InValue The value to set on the control
+	 * @param InValueType The type of value to set
+  	 * @param bSetupUndo If true the transform stack will be setup for undo / redo
+	 */
+	UFUNCTION(BlueprintCallable, Category = URigHierarchy)
+	FORCEINLINE_DEBUGGABLE void SetControlValueByIndex(int32 InElementIndex, FRigControlValue InValue, ERigControlValueType InValueType = ERigControlValueType::Current, bool bSetupUndo = false)
+	{
+		if(Elements.IsValidIndex(InElementIndex))
+		{
+			if(FRigControlElement* ControlElement = Cast<FRigControlElement>(Elements[InElementIndex]))
+			{
+				SetControlValue(ControlElement, InValue, InValueType, bSetupUndo);
+			}
+		}
+	}
+
+	/**
+	 * Sets a control's current value given its index
+	 * @param InElementIndex The index of the element to set the current value for
+	 * @param InValue The value to set on the control
+	 * @param InValueType The type of value to set
+	 * @param bSetupUndo If true the transform stack will be setup for undo / redo
+	 */
+	FORCEINLINE_DEBUGGABLE void SetControlValue(int32 InElementIndex, const FRigControlValue& InValue, ERigControlValueType InValueType = ERigControlValueType::Current, bool bSetupUndo = false)
+	{
+		SetControlValueByIndex(InElementIndex, InValue, InValueType, bSetupUndo);
+	}
+
+	/**
+	 * Sets a control's current value given its index
+	 * @param InElementIndex The index of the element to set the current value for
+	 * @param InValue The value to set on the control
+	 * @param InValueType The type of value to set
+	 * @param bSetupUndo If true the transform stack will be setup for undo / redo
+	 */
+	template<typename T>
+	FORCEINLINE_DEBUGGABLE void SetControlValue(int32 InElementIndex, const T& InValue, ERigControlValueType InValueType = ERigControlValueType::Current, bool bSetupUndo = false) const
+	{
+		return SetControlValue(InElementIndex, FRigControlValue::Make<T>(InValue), InValueType, bSetupUndo);
+	}
+
+	/**
+	 * Sets a control's initial value given its index
+	 * @param InElementIndex The index of the element to set the initial value for
+	 * @param InValue The value to set on the control
+	 * @param bSetupUndo If true the transform stack will be setup for undo / redo
+	 */
+	FORCEINLINE_DEBUGGABLE void SetInitialControlValue(int32 InElementIndex, const FRigControlValue& InValue, bool bSetupUndo = false)
+	{
+		SetControlValueByIndex(InElementIndex, InValue, ERigControlValueType::Initial, bSetupUndo);
+	}
+
+	/**
+	 * Sets a control's initial value given its index
+	 * @param InElementIndex The index of the element to set the initial value for
+	 * @param InValue The value to set on the control
+	 * @param bSetupUndo If true the transform stack will be setup for undo / redo
+	 */
+	template<typename T>
+	FORCEINLINE_DEBUGGABLE void SetInitialControlValue(int32 InElementIndex, const T& InValue, bool bSetupUndo = false) const
+	{
+		return SetInitialControlValue(InElementIndex, FRigControlValue::Make<T>(InValue), bSetupUndo);
+	}
+
+	/**
+	 * Sets a control's current visibility based on a key
+	 * @param InKey The key of the element to set the visibility for
+	 * @param bVisibility The visibility to set on the control
+	 */
+	UFUNCTION(BlueprintCallable, Category = URigHierarchy)
+    FORCEINLINE_DEBUGGABLE void SetControlVisibility(FRigElementKey InKey, bool bVisibility)
+	{
+		SetControlVisibilityByIndex(GetIndex(InKey), bVisibility);
+	}
+
+	/**
+	 * Sets a control's current visibility based on a key
+	 * @param InElementIndex The index of the element to set the visibility for
+	 * @param bVisibility The visibility to set on the control
+	 */
+	UFUNCTION(BlueprintCallable, Category = URigHierarchy)
+    FORCEINLINE_DEBUGGABLE void SetControlVisibilityByIndex(int32 InElementIndex, bool bVisibility)
+	{
+		if(Elements.IsValidIndex(InElementIndex))
+		{
+			if(FRigControlElement* ControlElement = Cast<FRigControlElement>(Elements[InElementIndex]))
+			{
+				SetControlVisibility(ControlElement, bVisibility);
+			}
+		}
+	}
+
+	FORCEINLINE_DEBUGGABLE void SetControlVisibility(int32 InElementIndex, bool bVisibility)
+	{
+		SetControlVisibilityByIndex(InElementIndex, bVisibility);
+	}
+
+	/**
+	 * Returns a curve's value given its key
+	 * @param InKey The key of the element to retrieve the value for
+	 * @return Returns the value of the curve
+	 */
+	UFUNCTION(BlueprintCallable, Category = URigHierarchy)
+    FORCEINLINE_DEBUGGABLE float GetCurveValue(FRigElementKey InKey) const
+	{
+		return GetCurveValueByIndex(GetIndex(InKey));
+	}
+
+	/**
+	 * Returns a curve's value given its index
+	 * @param InElementIndex The index of the element to retrieve the value for
+	 * @return Returns the value of the curve
+	 */
+	UFUNCTION(BlueprintCallable, Category = URigHierarchy)
+    FORCEINLINE_DEBUGGABLE float GetCurveValueByIndex(int32 InElementIndex) const
+	{
+		if(Elements.IsValidIndex(InElementIndex))
+		{
+			if(FRigCurveElement* CurveElement = Cast<FRigCurveElement>(Elements[InElementIndex]))
+			{
+				return GetCurveValue(CurveElement);
+			}
+		}
+		return 0.f;
+	}
+
+	FORCEINLINE_DEBUGGABLE float GetCurveValue(int32 InElementIndex) const
+	{
+		return GetCurveValueByIndex(InElementIndex);
+	}
+
+	/**
+	 * Sets a curve's value given its key
+	 * @param InKey The key of the element to set the value for
+	 * @param InValue The value to set on the curve
+	 * @param bSetupUndo If true the transform stack will be setup for undo / redo
+	 */
+	UFUNCTION(BlueprintCallable, Category = URigHierarchy)
+    FORCEINLINE_DEBUGGABLE void SetCurveValue(FRigElementKey InKey, float InValue, bool bSetupUndo = false)
+	{
+		SetCurveValueByIndex(GetIndex(InKey), InValue, bSetupUndo);
+	}
+
+	/**
+	 * Sets a curve's value given its index
+	 * @param InElementIndex The index of the element to set the value for
+	 * @param InValue The value to set on the curve
+	 * @param bSetupUndo If true the transform stack will be setup for undo / redo
+	 */
+	UFUNCTION(BlueprintCallable, Category = URigHierarchy)
+	FORCEINLINE_DEBUGGABLE void SetCurveValueByIndex(int32 InElementIndex, float InValue, bool bSetupUndo = false)
+	{
+		if(Elements.IsValidIndex(InElementIndex))
+		{
+			if(FRigCurveElement* CurveElement = Cast<FRigCurveElement>(Elements[InElementIndex]))
+			{
+				SetCurveValue(CurveElement, InValue, bSetupUndo);
+			}
+		}
+	}
+
+	FORCEINLINE_DEBUGGABLE void SetCurveValue(int32 InElementIndex, float InValue, bool bSetupUndo = false)
+	{
+		SetCurveValueByIndex(InElementIndex, InValue, bSetupUndo);
+	}
+
+	/**
+	 * Sets the offset transform for a given control element by key
+	 * @param InKey The key of the control element to set the offset transform for
+	 * @param InTransform The new offset transform value to set
+	 * @param bInitial If true the initial value will be used
+	 * @param bAffectChildren If set to false children will not move (maintain global).
+	 * @param bSetupUndo If true the transform stack will be setup for undo / redo
+	 */
+	UFUNCTION(BlueprintCallable, Category = URigHierarchy)
+    FORCEINLINE_DEBUGGABLE void SetControlOffsetTransform(FRigElementKey InKey, FTransform InTransform, bool bInitial = false, bool bAffectChildren = true, bool bSetupUndo = false)
+	{
+		return SetControlOffsetTransformByIndex(GetIndex(InKey), InTransform, bInitial, bAffectChildren, bSetupUndo);
+	}
+
+	/**
+	 * Sets the local offset transform for a given control element by index
+	 * @param InElementIndex The index of the control element to set the offset transform for
+ 	 * @param InTransform The new local offset transform value to set
+	 * @param bInitial If true the initial value will be used
+	 * @param bAffectChildren If set to false children will not move (maintain global).
+	 * @param bSetupUndo If true the transform stack will be setup for undo / redo
+	 */
+	UFUNCTION(BlueprintCallable, Category = URigHierarchy)
+    FORCEINLINE_DEBUGGABLE void SetControlOffsetTransformByIndex(int32 InElementIndex, FTransform InTransform, bool bInitial = false, bool bAffectChildren = true, bool bSetupUndo = false)
+	{
+		if(Elements.IsValidIndex(InElementIndex))
+		{
+			if(FRigControlElement* ControlElement = Cast<FRigControlElement>(Elements[InElementIndex]))
+			{
+				SetControlOffsetTransform(ControlElement, InTransform, bInitial ? ERigTransformType::InitialLocal : ERigTransformType::CurrentLocal, bAffectChildren, bSetupUndo);
+			}
+		}
+	}
+
+	/**
+	 * Sets the gizmo transform for a given control element by key
+	 * @param InKey The key of the control element to set the gizmo transform for
+	 * @param InTransform The new gizmo transform value to set
+	 * @param bInitial If true the initial value will be used
+	 * @param bSetupUndo If true the transform stack will be setup for undo / redo
+	 */
+	UFUNCTION(BlueprintCallable, Category = URigHierarchy)
+    FORCEINLINE_DEBUGGABLE void SetControlGizmoTransform(FRigElementKey InKey, FTransform InTransform, bool bInitial = false, bool bSetupUndo = false)
+	{
+		return SetControlGizmoTransformByIndex(GetIndex(InKey), InTransform, bInitial, bSetupUndo);
+	}
+
+	/**
+	 * Sets the local gizmo transform for a given control element by index
+	 * @param InElementIndex The index of the control element to set the gizmo transform for
+	 * @param InTransform The new local gizmo transform value to set
+	 * @param bInitial If true the initial value will be used
+	 * @param bSetupUndo If true the transform stack will be setup for undo / redo
+	 */
+	UFUNCTION(BlueprintCallable, Category = URigHierarchy)
+    FORCEINLINE_DEBUGGABLE void SetControlGizmoTransformByIndex(int32 InElementIndex, FTransform InTransform, bool bInitial = false, bool bSetupUndo = false)
+	{
+		if(Elements.IsValidIndex(InElementIndex))
+		{
+			if(FRigControlElement* ControlElement = Cast<FRigControlElement>(Elements[InElementIndex]))
+			{
+				SetControlGizmoTransform(ControlElement, InTransform, bInitial ? ERigTransformType::InitialLocal : ERigTransformType::CurrentLocal, bSetupUndo);
+			}
+		}
+	}
+
+	/**
+	 * Returns the global current or initial value for a given key.
+	 * If the element does not have a parent FTransform::Identity will be returned.
+	 * @param InKey The key of the element to retrieve the transform for
+	 * @param bInitial If true the initial transform will be used
+	 * @return The element's parent's global current or initial transform's value.
+	 */
+	UFUNCTION(BlueprintCallable, Category = URigHierarchy)
+    FORCEINLINE_DEBUGGABLE FTransform GetParentTransform(FRigElementKey InKey, bool bInitial = false) const
+	{
+		return GetParentTransformByIndex(GetIndex(InKey), bInitial);
+	}
+
+	/**
+	 * Returns the global current or initial value for a given element index.
+	 * If the element does not have a parent FTransform::Identity will be returned.
+	 * @param InElementIndex The index of the element to retrieve the transform for
+	 * @param bInitial If true the initial transform will be used
+	 * @return The element's parent's global current or initial transform's value.
+	 */
+	UFUNCTION(BlueprintCallable, Category = URigHierarchy)
+    FORCEINLINE_DEBUGGABLE FTransform GetParentTransformByIndex(int32 InElementIndex, bool bInitial = false) const
+	{
+		if(Elements.IsValidIndex(InElementIndex))
+		{
+			return GetParentTransform(Elements[InElementIndex], bInitial ? ERigTransformType::InitialGlobal : ERigTransformType::CurrentGlobal);
+		}
+		return FTransform::Identity;
+	}
+
+	/**
+	 * Returns the child elements of a given element key
+	 * @param InKey The key of the element to retrieve the children for
+	 * @param bRecursive If set to true grand-children will also be returned etc
+	 * @return Returns the child elements
+	 */
+	UFUNCTION(BlueprintCallable, Category = URigHierarchy)
+	TArray<FRigElementKey> GetChildren(FRigElementKey InKey, bool bRecursive = false) const;
+
+	/**
+	 * Returns the child elements of a given element index
+	 * @param InIndex The index of the element to retrieve the children for
+	 * @param bRecursive If set to true grand-children will also be returned etc
+	 * @return Returns the child elements' indices
+	 */
+    TArray<int32> GetChildren(int32 InIndex, bool bRecursive = false) const;
+
+	/**
+	 * Returns the child elements of a given element
+	 * @param InElement The element to retrieve the children for
+	 * @return Returns the child elements
+	 */
+	const TArray<FRigBaseElement*>& GetChildren(const FRigBaseElement* InElement) const;
+
+	/**
+	 * Returns the child elements of a given element
+	 * @param InElement The element to retrieve the children for
+	 * @param bRecursive If set to true grand-children will also be returned etc
+	 * @return Returns the child elements
+	 */
+	TArray<FRigBaseElement*> GetChildren(const FRigBaseElement* InElement, bool bRecursive) const;
+
+	/**
+	 * Returns the parent elements of a given element key
+	 * @param InKey The key of the element to retrieve the parents for
+	 * @param bRecursive If set to true parents of parents will also be returned
+	 * @return Returns the parent elements
+	 */
+	UFUNCTION(BlueprintCallable, Category = URigHierarchy)
+    TArray<FRigElementKey> GetParents(FRigElementKey InKey, bool bRecursive = false) const;
+
+	/**
+	 * Returns the parent elements of a given element index
+	 * @param InIndex The index of the element to retrieve the parents for
+	 * @param bRecursive If set to true parents of parents will also be returned
+	 * @return Returns the parent elements' indices
+	 */
+    TArray<int32> GetParents(int32 InIndex, bool bRecursive = false) const;
+
+	/**
+	 * Returns the parent elements of a given element
+	 * @param InElement The element to retrieve the parents for
+	 * @param bRecursive If set to true parents of parents will also be returned
+	 * @return Returns the parent elements
+	 */
+	TArray<FRigBaseElement*> GetParents(const FRigBaseElement* InElement, bool bRecursive = false) const;
+
+	/**
+	 * Returns the first parent element of a given element key
+	 * @param InKey The key of the element to retrieve the parents for
+	 * @return Returns the first parent element
+	 */
+	UFUNCTION(BlueprintCallable, Category = URigHierarchy)
+    FRigElementKey GetFirstParent(FRigElementKey InKey) const;
+
+	/**
+	 * Returns the first parent element of a given element index
+	 * @param InIndex The index of the element to retrieve the parent for
+	 * @return Returns the first parent index (or INDEX_NONE)
+	 */
+    int32 GetFirstParent(int32 InIndex) const;
+
+	/**
+	 * Returns the first parent element of a given element
+	 * @param InElement The element to retrieve the parents for
+	 * @return Returns the first parent element
+	 */
+	FRigBaseElement* GetFirstParent(const FRigBaseElement* InElement) const;
+
+	/**
+	 * Returns the number of parents of an element
+	 * @param InKey The key of the element to retrieve the number of parents for
+	 * @return Returns the number of parents of an element
+	 */
+	UFUNCTION(BlueprintCallable, Category = URigHierarchy)
+    int32 GetNumberOfParents(FRigElementKey InKey) const;
+
+	/**
+	 * Returns the number of parents of an element
+	 * @param InIndex The index of the element to retrieve the number of parents for
+	 * @return Returns the number of parents of an element
+	 */
+    int32 GetNumberOfParents(int32 InIndex) const;
+
+	/**
+	 * Returns the number of parents of an element
+	 * @param InElement The element to retrieve the number of parents for
+	 * @return Returns the number of parents of an element
+	 */
+	int32 GetNumberOfParents(const FRigBaseElement* InElement) const;
+
+	/**
+	 * Returns the weight of a parent below a multi parent element
+	 * @param InChild The key of the multi parented element
+	 * @param InParent The key of the parent to look up the weight for
+	 * @param bInitial If true the initial weights will be used
+	 * @return Returns the weight of a parent below a multi parent element, or FLT_MAX if the parent is invalid
+	 */
+	UFUNCTION(BlueprintCallable, Category = URigHierarchy)
+    float GetParentWeight(FRigElementKey InChild, FRigElementKey InParent, bool bInitial = false) const;
+
+	/**
+	 * Returns the weight of a parent below a multi parent element
+	 * @param InChild The multi parented element
+	 * @param InParent The parent to look up the weight for
+	 * @param bInitial If true the initial weights will be used
+	 * @return Returns the weight of a parent below a multi parent element, or FLT_MAX if the parent is invalid
+	 */
+    float GetParentWeight(const FRigBaseElement* InChild, const FRigBaseElement* InParent, bool bInitial = false) const;
+
+	/**
+	 * Returns the weight of a parent below a multi parent element
+	 * @param InChild The multi parented element
+	 * @param InParentIndex The index of the parent inside of the multi parent element
+	 * @param bInitial If true the initial weights will be used
+	 * @return Returns the weight of a parent below a multi parent element, or FLT_MAX if the parent is invalid
+	 */
+	float GetParentWeight(const FRigBaseElement* InChild, int32 InParentIndex, bool bInitial = false) const;
+
+	/**
+	 * Sets the weight of a parent below a multi parent element
+	 * @param InChild The key of the multi parented element
+	 * @param InParent The key of the parent to look up the weight for
+	 * @param InWeight The new weight to set for the parent
+	 * @param bInitial If true the initial weights will be used
+	 * @param bAffectChildren If set to false children will not move (maintain global).
+	 * @return Returns true if changing the weight was successful
+	 */
+	UFUNCTION(BlueprintCallable, Category = URigHierarchy)
+    bool SetParentWeight(FRigElementKey InChild, FRigElementKey InParent, float InWeight, bool bInitial = false, bool bAffectChildren = true);
+
+	/**
+	 * Sets the weight of a parent below a multi parent element
+	 * @param InChild The multi parented element
+	 * @param InParent The parent to look up the weight for
+	 * @param InWeight The new weight to set for the parent
+	 * @param bInitial If true the initial weights will be used
+	 * @param bAffectChildren If set to false children will not move (maintain global).
+	 * @return Returns true if changing the weight was successful
+	 */
+	bool SetParentWeight(FRigBaseElement* InChild, const FRigBaseElement* InParent, float InWeight, bool bInitial = false, bool bAffectChildren = true);
+
+	/**
+	 * Sets the weight of a parent below a multi parent element
+	 * @param InChild The multi parented element
+	 * @param InParentIndex The index of the parent inside of the multi parent element
+	 * @param InWeight The new weight to set for the parent
+	 * @param bInitial If true the initial weights will be used
+	 * @param bAffectChildren If set to false children will not move (maintain global).
+	 * @return Returns true if changing the weight was successful
+	 */
+	bool SetParentWeight(FRigBaseElement* InChild, int32 InParentIndex, float InWeight, bool bInitial = false, bool bAffectChildren = true);
+
+	/**
+	 * Returns true if an element is parented to another element
+	 * @param InChild The key of the child element to check for a parent
+	 * @param InParent The key of the parent element to check for
+	 * @return True if the child is parented to the parent
+	 */
+	UFUNCTION(BlueprintCallable, Category = URigHierarchy)
+	FORCEINLINE bool IsParentedTo(FRigElementKey InChild, FRigElementKey InParent) const
+	{
+		return IsParentedTo(GetIndex(InChild), GetIndex(InParent));
+	}
+
+	/**
+	 * Returns true if an element is parented to another element
+	 * @param InChildIndex The index of the child element to check for a parent
+	 * @param InParentIndex The index of the parent element to check for
+	 * @return True if the child is parented to the parent
+	 */
+    FORCEINLINE bool IsParentedTo(int32 InChildIndex, int32 InParentIndex) const
+	{
+		if(Elements.IsValidIndex(InChildIndex) && Elements.IsValidIndex(InParentIndex))
+		{
+			return IsParentedTo(Elements[InChildIndex], Elements[InParentIndex]);
+		}
+		return false;
+	}
+
+	/**
+	 * Returns all element keys of this hierarchy
+	 * @param bTraverse If set to true the keys will be returned by depth first traversal
+	 * @param InElementType The type filter to apply
+	 * @return The keys of all elements
+	 */
+	TArray<FRigElementKey> GetAllKeys(bool bTraverse = false, ERigElementType InElementType = ERigElementType::All) const;
+
+	/**
+	 * Returns all element keys of this hierarchy
+	 * @param bTraverse If set to true the keys will be returned by depth first traversal
+	 * @return The keys of all elements
+	 */
+	UFUNCTION(BlueprintCallable, Category = URigHierarchy, meta = (DisplayName = "Get All Keys", ScriptName = "GetAllKeys"))
+	FORCEINLINE TArray<FRigElementKey> GetAllKeys_ForBlueprint(bool bTraverse = true) const
+	{
+		return GetAllKeys(bTraverse, ERigElementType::All);
+	}
+
+	/**
+	 * Helper function to traverse the hierarchy
+	 * @param InElement The element to start the traversal at
+	 * @param bTowardsChildren If set to true the traverser walks downwards (towards the children), otherwise upwards (towards the parents)
+	 * @param PerElementFunction The function to call for each visited element
+	 */
+	void Traverse(FRigBaseElement* InElement, bool bTowardsChildren, TFunction<void(FRigBaseElement*, bool& /* continue */)> PerElementFunction) const;
+
+	/**
+	 * Helper function to traverse the hierarchy from the root
+	 * @param PerElementFunction The function to call for each visited element
+	 * @param bTowardsChildren If set to true the traverser walks downwards (towards the children), otherwise upwards (towards the parents)
+	 */
+	void Traverse(TFunction<void(FRigBaseElement*, bool& /* continue */)> PerElementFunction, bool bTowardsChildren = true) const;
+
+	/**
+	 * Performs undo for one transform change
+	 */
+	bool Undo();
+
+	/**
+	 * Performs redo for one transform change
+	 */
+	bool Redo();
+
+	/**
+	 * Returns the transform stack index
+	 */
+	int32 GetTransformStackIndex() const { return TransformStackIndex; }
+
+	/**
+	 * Sends an event from the hierarchy to the world
+	 * @param InEvent The event to send
+	 * @param bAsynchronous If set to true the event will go on a thread safe queue
+	 */
+	void SendEvent(const FRigEventContext& InEvent, bool bAsynchronous = true);
+
+	/**
+	 * Returns the delegate to listen to for events coming from this hierarchy
+	 * @return The delegate to listen to for events coming from this hierarchy
+	 */
+	FORCEINLINE FRigEventDelegate& OnEventReceived() { return EventDelegate; }
+
+	/**
+	 * Returns a controller for this hierarchy
+	 * @param bCreateIfNeeded Creates a controller if needed
+	 * @return The Controller for this hierarchy
+	 */
+	UFUNCTION(BlueprintCallable, Category = URigHierarchy)
+	URigHierarchyController* GetController(bool bCreateIfNeeded = false);
+
+	/**
+	 * Returns the topology version of this hierarchy
+	 */
+	uint16 GetTopologyVersion() const { return TopologyVersion; }
+
+	/**
+	 * Returns the current / initial pose of the hierarchy
+	 * @param bInitial If set to true the initial pose will be returned
+	 * @return The pose of the hierarchy
+	 */
+	UFUNCTION(BlueprintCallable, Category = URigHierarchy)
+	FRigPose GetPose(bool bInitial = false) const;
+
+	/**
+	 * Sets the current / initial pose of the hierarchy
+	 * @param InPose The pose to set on the hierarchy
+	 * @param InTransformType The transform type to set
+	 */
+	void SetPose(const FRigPose& InPose, ERigTransformType::Type InTransformType = ERigTransformType::CurrentLocal);
+
+	/**
+	 * Sets the current / initial pose of the hierarchy
+	 * @param InPose The pose to set on the hierarchy
+	 */
+	UFUNCTION(BlueprintCallable, Category = URigHierarchy, meta = (DisplayName = "Set Pose", ScriptName = "SetPose"))
+	FORCEINLINE void SetPose_ForBlueprint(FRigPose InPose)
+	{
+		return SetPose(InPose);
+	}
+
+	/**
+	 * Creates a rig control value from a bool value
+	 * @param InValue The value to create the rig control value from
+	 * @return The converted control rig val ue
+	 */
+	UFUNCTION(BlueprintCallable, Category = URigHierarchy)
+	static FORCEINLINE FRigControlValue MakeControlValueFromBool(bool InValue)
+	{
+		return FRigControlValue::Make<bool>(InValue);
+	}
+
+	/**
+	 * Creates a rig control value from a float value
+	 * @param InValue The value to create the rig control value from
+	 * @return The converted control rig val ue
+	 */
+	UFUNCTION(BlueprintCallable, Category = URigHierarchy)
+    static FORCEINLINE FRigControlValue MakeControlValueFromFloat(float InValue)
+	{
+		return FRigControlValue::Make<float>(InValue);
+	}
+
+	/**
+	 * Returns the contained float value from a a Rig Control Value
+	 * @param InValue The Rig Control value to convert from
+	 * @return The converted float value
+	 */
+	UFUNCTION(BlueprintCallable, Category = URigHierarchy)
+	static FORCEINLINE float GetFloatFromControlValue(FRigControlValue InValue)
+	{
+		return InValue.Get<float>();
+	}
+
+	/**
+	 * Creates a rig control value from a int32 value
+	 * @param InValue The value to create the rig control value from
+	 * @return The converted control rig val ue
+	 */
+	UFUNCTION(BlueprintCallable, Category = URigHierarchy)
+    static FORCEINLINE FRigControlValue MakeControlValueFromInt(int32 InValue)
+	{
+		return FRigControlValue::Make<int32>(InValue);
+	}
+
+	/**
+	 * Returns the contained int32 value from a a Rig Control Value
+	 * @param InValue The Rig Control value to convert from
+	 * @return The converted int32 value
+	 */
+	UFUNCTION(BlueprintCallable, Category = URigHierarchy)
+	static FORCEINLINE int32 GetIntFromControlValue(FRigControlValue InValue)
+	{
+		return InValue.Get<int32>();
+	}
+
+	/**
+	 * Creates a rig control value from a FVector2D value
+	 * @param InValue The value to create the rig control value from
+	 * @return The converted control rig val ue
+	 */
+	UFUNCTION(BlueprintCallable, Category = URigHierarchy)
+    static FORCEINLINE FRigControlValue MakeControlValueFromVector2D(FVector2D InValue)
+	{
+		return FRigControlValue::Make<FVector2D>(InValue);
+	}
+
+	/**
+	 * Returns the contained FVector2D value from a a Rig Control Value
+	 * @param InValue The Rig Control value to convert from
+	 * @return The converted FVector2D value
+	 */
+	UFUNCTION(BlueprintCallable, Category = URigHierarchy)
+	static FORCEINLINE FVector2D GetVector2DFromControlValue(FRigControlValue InValue)
+	{
+		return InValue.Get<FVector2D>();
+	}
+
+	/**
+	 * Creates a rig control value from a FVector value
+	 * @param InValue The value to create the rig control value from
+	 * @return The converted control rig val ue
+	 */
+	UFUNCTION(BlueprintCallable, Category = URigHierarchy)
+    static FORCEINLINE FRigControlValue MakeControlValueFromVector(FVector InValue)
+	{
+		return FRigControlValue::Make<FVector>(InValue);
+	}
+
+	/**
+	 * Returns the contained FVector value from a a Rig Control Value
+	 * @param InValue The Rig Control value to convert from
+	 * @return The converted FVector value
+	 */
+	UFUNCTION(BlueprintCallable, Category = URigHierarchy)
+	static FORCEINLINE FVector GetVectorFromControlValue(FRigControlValue InValue)
+	{
+		return InValue.Get<FVector>();
+	}
+
+	/**
+	 * Creates a rig control value from a FRotator value
+	 * @param InValue The value to create the rig control value from
+	 * @return The converted control rig val ue
+	 */
+	UFUNCTION(BlueprintCallable, Category = URigHierarchy)
+    static FORCEINLINE FRigControlValue MakeControlValueFromRotator(FRotator InValue)
+	{
+		return FRigControlValue::Make<FRotator>(InValue);
+	}
+
+	/**
+	 * Returns the contained FRotator value from a a Rig Control Value
+	 * @param InValue The Rig Control value to convert from
+	 * @return The converted FRotator value
+	 */
+	UFUNCTION(BlueprintCallable, Category = URigHierarchy)
+	static FORCEINLINE FRotator GetRotatorFromControlValue(FRigControlValue InValue)
+	{
+		return InValue.Get<FRotator>();
+	}
+
+	/**
+	 * Creates a rig control value from a FTransform value
+	 * @param InValue The value to create the rig control value from
+	 * @return The converted control rig val ue
+	 */
+	UFUNCTION(BlueprintCallable, Category = URigHierarchy)
+    static FORCEINLINE FRigControlValue MakeControlValueFromTransform(FTransform InValue)
+	{
+		return FRigControlValue::Make<FTransform>(InValue);
+	}
+
+	/**
+	 * Returns the contained FTransform value from a a Rig Control Value
+	 * @param InValue The Rig Control value to convert from
+	 * @return The converted FTransform value
+	 */
+	UFUNCTION(BlueprintCallable, Category = URigHierarchy)
+	static FORCEINLINE FTransform GetTransformFromControlValue(FRigControlValue InValue)
+	{
+		return InValue.Get<FTransform>();
+	}
+
+	/**
+	 * Creates a rig control value from a FEulerTransform value
+	 * @param InValue The value to create the rig control value from
+	 * @return The converted control rig val ue
+	 */
+	UFUNCTION(BlueprintCallable, Category = URigHierarchy)
+    static FORCEINLINE FRigControlValue MakeControlValueFromEulerTransform(FEulerTransform InValue)
+	{
+		return FRigControlValue::Make<FEulerTransform>(InValue);
+	}
+
+	/**
+	 * Returns the contained FEulerTransform value from a a Rig Control Value
+	 * @param InValue The Rig Control value to convert from
+	 * @return The converted FEulerTransform value
+	 */
+	UFUNCTION(BlueprintCallable, Category = URigHierarchy)
+	static FORCEINLINE FEulerTransform GetEulerTransformFromControlValue(FRigControlValue InValue)
+	{
+		return InValue.Get<FEulerTransform>();
+	}
+
+	/**
+	 * Creates a rig control value from a FTransformNoScale value
+	 * @param InValue The value to create the rig control value from
+	 * @return The converted control rig val ue
+	 */
+	UFUNCTION(BlueprintCallable, Category = URigHierarchy)
+    static FORCEINLINE FRigControlValue MakeControlValueFromTransformNoScale(FTransformNoScale InValue)
+	{
+		return FRigControlValue::Make<FTransformNoScale>(InValue);
+	}
+
+	/**
+	 * Returns the contained FTransformNoScale value from a a Rig Control Value
+	 * @param InValue The Rig Control value to convert from
+	 * @return The converted FTransformNoScale value
+	 */
+	UFUNCTION(BlueprintCallable, Category = URigHierarchy)
+	static FORCEINLINE FTransformNoScale GetTransformNoScaleFromControlValue(FRigControlValue InValue)
+	{
+		return InValue.Get<FTransformNoScale>();
+	}
+
+private:
+
+	FRigHierarchyModifiedEvent ModifiedEvent;
+	FRigEventDelegate EventDelegate;
+
+public:
+
+	void Notify(ERigHierarchyNotification InNotifType, const FRigBaseElement* InElement);
+
+	/**
+	 * Returns a transform based on a given transform type
+	 * @param InTransformElement The element to retrieve the transform for
+	 * @param InTransformType The type of transform to retrieve
+	 * @return The local current or initial transform's value.
+	 */
+	FTransform GetTransform(FRigTransformElement* InTransformElement, const ERigTransformType::Type InTransformType) const;
+
+	/**
+	 * Returns a transform for a given element's parent based on the transform type
+	 * If the element does not have a parent FTransform::Identity will be returned.
+	 * @param InElement The element to retrieve the transform for
+	 * @param InTransformType The type of transform to retrieve
+	 * @return The element's parent's transform
+	 */
+	FTransform GetParentTransform(FRigBaseElement* InElement, const ERigTransformType::Type InTransformType) const;
+
+	/**
+	 * Sets a transform for a given element based on the transform type
+	 * @param InTransformElement The element to set the transform for
+	 * @param InTransform The type of transform to set
+	 * @param InTransformType The type of transform to set
+	 * @param bAffectChildren If set to false children will not move (maintain global).
+	 * @param bSetupUndo If true the transform stack will be setup for undo / redo
+	 * @param bForce Set the transform even if it is the same as the previously set one
+	 */
+	void SetTransform(FRigTransformElement* InTransformElement, const FTransform& InTransform, const ERigTransformType::Type InTransformType, bool bAffectChildren, bool bSetupUndo = false, bool bForce = false);
+
+	/**
+	 * Returns the global offset transform for a given control element.
+	 * @param InControlElement The control element to retrieve the offset transform for
+	 * @param InTransformType The type of transform to set
+	 * @return The global offset transform
+	 */
+	FTransform GetControlOffsetTransform(FRigControlElement* InControlElement, const ERigTransformType::Type InTransformType) const;
+
+	/**
+	 * Sets the offset transform for a given control element
+	 * @param InControlElement The element to set the transform for
+	 * @param InTransform The offset transform to set
+	 * @param InTransformType The type of transform to set
+	 * @param bAffectChildren If set to false children will not move (maintain global).
+	 * @param bSetupUndo If true the transform stack will be setup for undo / redo
+	 * @param bForce Set the transform even if it is the same as the previously set one
+	 */
+	void SetControlOffsetTransform(FRigControlElement* InControlElement, const FTransform& InTransform, const ERigTransformType::Type InTransformType, bool bAffectChildren, bool bSetupUndo = false, bool bForce = false);
+
+	/**
+	 * Returns the global gizmo transform for a given control element.
+	 * @param InControlElement The control element to retrieve the gizmo transform for
+	 * @param InTransformType The type of transform to set
+	 * @return The global gizmo transform
+	 */
+	FTransform GetControlGizmoTransform(FRigControlElement* InControlElement, const ERigTransformType::Type InTransformType) const;
+
+	/**
+	 * Sets the gizmo transform for a given control element
+	 * @param InControlElement The element to set the transform for
+	 * @param InTransform The gizmo transform to set
+	 * @param InTransformType The type of transform to set
+	 * @param bSetupUndo If true the transform stack will be setup for undo / redo
+	 * @param bForce Set the transform even if it is the same as the previously set one
+	 */
+	void SetControlGizmoTransform(FRigControlElement* InControlElement, const FTransform& InTransform, const ERigTransformType::Type InTransformType, bool bSetupUndo = false, bool bForce = false);
+
+	/**
+	 * Returns a control's current value
+	 * @param InControlElement The element to retrieve the current value for
+	 * @param InValueType The type of value to return
+	 * @return Returns the current value of the control
+	 */
+	FRigControlValue GetControlValue(FRigControlElement* InControlElement, ERigControlValueType InValueType) const;
+
+	template<typename T>
+	FORCEINLINE_DEBUGGABLE T GetControlValue(FRigControlElement* InControlElement, ERigControlValueType InValueType) const
+	{
+		return GetControlValue(InControlElement, InValueType).Get<T>();
+	}
+
+	/**
+	 * Sets a control's current value
+	 * @param InControlElement The element to set the current value for
+	 * @param InValueType The type of value to set
+	 * @param InValue The value to set on the control
+	 * @param bSetupUndo If true the transform stack will be setup for undo / redo
+	 * @param bForce Set the transform even if it is the same as the previously set one
+	 */
+	void SetControlValue(FRigControlElement* InControlElement, const FRigControlValue& InValue, ERigControlValueType InValueType, bool bSetupUndo = false, bool bForce = false);
+
+	template<typename T>
+	FORCEINLINE_DEBUGGABLE void SetControlValue(FRigControlElement* InControlElement, const T& InValue, ERigControlValueType InValueType, bool bSetupUndo = false, bool bForce = false) const
+	{
+		SetControlValue(InControlElement, FRigControlValue::Make<T>(InValue), InValueType, bSetupUndo, bForce);
+	}
+
+	/**
+	 * Sets a control's current visibility
+	 * @param InControlElement The element to set the visibility for
+	 * @param bVisibility The new visibility for the control
+	 */
+	void SetControlVisibility(FRigControlElement* InControlElement, bool bVisibility);
+
+	/**
+	 * Returns a curve's value
+	 * @param InCurveElement The element to retrieve the value for
+	 * @return Returns the value of the curve
+	 */
+	float GetCurveValue(FRigCurveElement* InCurveElement) const;
+
+	/**
+	 * Sets a curve's value
+	 * @param InCurveElement The element to set the value for
+	 * @param InValue The value to set on the control
+	 * @param bSetupUndo If true the transform stack will be setup for undo / redo
+	 * @param bForce Set the transform even if it is the same as the previously set one
+	 */
+	void SetCurveValue(FRigCurveElement* InCurveElement, float InValue, bool bSetupUndo = false, bool bForce = false);
+
+	/**
+	 * Returns the previous name of an element prior to a rename operation
+	 * @param InKey The key of the element to request the old name for
+	 */
+	UFUNCTION(BlueprintCallable, Category = URigHierarchy)
+	FName GetPreviousName(const FRigElementKey& InKey) const;
+
+	/**
+	 * Returns the previous parent of an element prior to a reparent operation
+	 * @param InKey The key of the element to request the old parent  for
+	 */
+	UFUNCTION(BlueprintCallable, Category = URigHierarchy)
+	FRigElementKey GetPreviousParent(const FRigElementKey& InKey) const;
+
+	/**
+	 * Returns true if an element is parented to another element
+	 * @param InChild The child element to check for a parent
+	 * @param InParent The parent element to check for
+	 * @return True if the child is parented to the parent
+	 */
+	bool IsParentedTo(FRigBaseElement* InChild, FRigBaseElement* InParent) const;
+
+private:
+
+
+	/**
+	 * Returns true if a given element is selected
+	 * @param InElement The element to check
+	 * @return true if a given element is selected
+	 */
+    bool IsSelected(const FRigBaseElement* InElement) const;
+
+	/**
+	 * Removes the transient cached children table for all elements.
+	 */
+	void ResetCachedChildren();
+
+	/**
+	 * Updates the transient cached children table for a given element if needed (or if bForce == true).
+	 * @param InElement The element to update the children table for
+	 * @param bForce If set to true the table will always be updated
+	 */
+	void UpdateCachedChildren(const FRigBaseElement* InElement, bool bForce = false) const;
+
+	/**
+	 * Marks all affected elements of a given element as dirty
+	 * @param InTransformElement The element that has changed
+	 * @param InTransformType The type of transform to propagate for
+	 * @param bAffectChildren If set to false children will not move (maintain global).
+	 */
+	void PropagateDirtyFlags(FRigTransformElement* InTransformElement, const ERigTransformType::Type InTransformType, bool bAffectChildren) const;
+
+	/**
+	 * The topology version of the hierarchy changes when elements are
+	 * added, removed, re-parented or renamed.
+	 */
+	UPROPERTY(transient)
+	uint16 TopologyVersion;
+
+	/**
+	 * If set to false the dirty flag propagation will be disabled
+	 */
+	UPROPERTY(transient)
+	bool bEnableDirtyPropagation;
+
+	// Storage for the elements
+	mutable TArray<FRigBaseElement*> Elements;
+
+	// Managed lookup from Key to Index
+	TMap<FRigElementKey, int32> IndexLookup;
+
+	// Static empty element array used for ref returns
+	static const TArray<FRigBaseElement*> EmptyElementArray;
+
+	///////////////////////////////////////////////
+	/// Undo redo related
+	///////////////////////////////////////////////
+
+	enum ETransformStackEntryType
+	{
+		TransformPose,
+		ControlOffset,
+		ControlGizmo,
+		CurveValue
+	};
+
+	struct FTransformStackEntry
+	{
+		FORCEINLINE FTransformStackEntry()
+			: Key()
+			, EntryType(ETransformStackEntryType::TransformPose)
+			, TransformType(ERigTransformType::CurrentLocal)
+			, OldTransform(FTransform::Identity)
+			, NewTransform(FTransform::Identity)
+			, bAffectChildren(true)
+		{}
+
+		FORCEINLINE FTransformStackEntry(
+			const FRigElementKey& InKey,
+			ETransformStackEntryType InEntryType,
+			ERigTransformType::Type InTransformType,
+			const FTransform& InOldTransform,
+			const FTransform& InNewTransform,
+			bool bInAffectChildren)
+			: Key(InKey)
+			, EntryType(InEntryType)
+			, TransformType(InTransformType)
+			, OldTransform(InOldTransform)
+			, NewTransform(InNewTransform)
+			, bAffectChildren(bInAffectChildren)
+		{}
+
+		FRigElementKey Key;
+		ETransformStackEntryType EntryType;
+		ERigTransformType::Type TransformType;
+		FTransform OldTransform;
+		FTransform NewTransform;
+		bool bAffectChildren;
+	};
+
+	/**
+	 * The index identifying where we stand with the stack
+	 */
+	UPROPERTY()
+	int32 TransformStackIndex;
+
+	/**
+	 * A flag to indicate if the next serialize should contain only transform changes
+	 */
+	bool bTransactingForTransformChange;
+	
+	/**
+	 * The stack of actions to undo
+	 */
+	TArray<FTransformStackEntry> TransformUndoStack;
+
+	/**
+	 * The stack of actions to undo
+	 */
+	TArray<FTransformStackEntry> TransformRedoStack;
+
+	/**
+	 * Sets the transform stack index - which in turns performs a series of undos / redos
+	 * @param InTransformStackIndex The new index for the transform stack
+	 */
+	bool SetTransformStackIndex(int32 InTransformStackIndex);
+
+	/**
+	 * Stores a transform on the stack
+	 */
+	void PushTransformToStack(
+			const FRigElementKey& InKey,
+            ETransformStackEntryType InEntryType,
+            ERigTransformType::Type InTransformType,
+            const FTransform& InOldTransform,
+            const FTransform& InNewTransform,
+            bool bAffectChildren);
+
+	/**
+	 * Stores a curve value on the stack
+	 */
+	void PushCurveToStack(
+            const FRigElementKey& InKey,
+            float InOldCurveValue,
+            float InNewCurveValue);
+
+	/**
+	 * Restores a transform on the stack
+	 */
+	bool ApplyTransformFromStack(const FTransformStackEntry& InEntry, bool bUndo);
+
+	TWeakObjectPtr<URigHierarchy> HierarchyForSelectionPtr;
+	TWeakObjectPtr<UObject> LastControllerPtr;
+
+	TMap<FRigElementKey, FRigElementKey> PreviousParentMap;
+	TMap<FRigElementKey, FRigElementKey> PreviousNameMap;
+	
+	friend class URigHierarchyController;
+	friend class UControlRig;
+	friend class FControlRigEditor;
+};
diff --git a/Engine/Plugins/Experimental/ControlRig/Source/ControlRig/Public/Rigs/RigHierarchyCache.h b/Engine/Plugins/Experimental/ControlRig/Source/ControlRig/Public/Rigs/RigHierarchyCache.h
--- a/Engine/Plugins/Experimental/ControlRig/Source/ControlRig/Public/Rigs/RigHierarchyCache.h
+++ b/Engine/Plugins/Experimental/ControlRig/Source/ControlRig/Public/Rigs/RigHierarchyCache.h
@@ -4,13 +4,11 @@
 
 #include "CoreMinimal.h"
 #include "RigHierarchyDefines.h"
+#include "RigHierarchyElements.h"
+
 #include "RigHierarchyCache.generated.h"
 
-struct FRigHierarchyContainer;
-struct FRigBoneHierarchy;
-struct FRigSpaceHierarchy;
-struct FRigControlHierarchy;
-struct FRigCurveContainer;
+class URigHierarchy;
 
 USTRUCT(BlueprintType)
 struct CONTROLRIG_API FCachedRigElement
@@ -25,44 +23,12 @@ public:
 		, ContainerVersion(INDEX_NONE)
 	{}
 
-	FORCEINLINE FCachedRigElement(const FRigElementKey& InKey, const FRigHierarchyContainer* InContainer)
-		: Key()
-		, Index(UINT16_MAX)
-		, ContainerVersion(INDEX_NONE)
-	{
-		UpdateCache(InKey, InContainer);
-	}
-
-	FORCEINLINE FCachedRigElement(const FName& InName, const FRigBoneHierarchy* InHierarchy)
-		: Key()
-		, Index(UINT16_MAX)
-		, ContainerVersion(INDEX_NONE)
-	{
-		UpdateCache(InName, InHierarchy);
-	}
-
-	FORCEINLINE FCachedRigElement(const FName& InName, const FRigSpaceHierarchy* InHierarchy)
-		: Key()
-		, Index(UINT16_MAX)
-		, ContainerVersion(INDEX_NONE)
-	{
-		UpdateCache(InName, InHierarchy);
-	}
-
-	FORCEINLINE FCachedRigElement(const FName& InName, const FRigControlHierarchy* InHierarchy)
-		: Key()
-		, Index(UINT16_MAX)
-		, ContainerVersion(INDEX_NONE)
-	{
-		UpdateCache(InName, InHierarchy);
-	}
-
-	FORCEINLINE FCachedRigElement(const FName& InName, const FRigCurveContainer* InHierarchy)
+	FORCEINLINE FCachedRigElement(const FRigElementKey& InKey, const URigHierarchy* InHierarchy)
 		: Key()
 		, Index(UINT16_MAX)
 		, ContainerVersion(INDEX_NONE)
 	{
-		UpdateCache(InName, InHierarchy);
+		UpdateCache(InKey, InHierarchy);
 	}
 
 	FORCEINLINE bool IsValid() const
@@ -75,13 +41,24 @@ public:
 		Key = FRigElementKey();
 		Index = UINT16_MAX;
 		ContainerVersion = INDEX_NONE;
+		Element = nullptr;
 	}
 
-	FORCEINLINE operator bool() const
+	FORCEINLINE explicit operator bool() const
 	{
 		return IsValid();
 	}
 
+	FORCEINLINE operator int32() const
+	{
+		return GetIndex();
+	}
+
+	FORCEINLINE explicit operator FRigElementKey() const
+	{
+		return Key;
+	}
+
 	FORCEINLINE int32 GetIndex() const
 	{
 		if(Index == UINT16_MAX)
@@ -96,24 +73,21 @@ public:
 		return Key;
 	}
 
-	bool UpdateCache(const FRigHierarchyContainer* InContainer);
-
-	bool UpdateCache(const FRigElementKey& InKey, const FRigHierarchyContainer* InContainer);
-
-	bool UpdateCache(const FName& InName, const FRigBoneHierarchy* InHierarchy);
-
-	bool UpdateCache(const FName& InName, const FRigSpaceHierarchy* InHierarchy);
+	FORCEINLINE const FRigBaseElement* GetElement() const
+	{
+		return Element;
+	}
 
-	bool UpdateCache(const FName& InName, const FRigControlHierarchy* InHierarchy);
+	bool UpdateCache(const URigHierarchy* InHierarchy);
 
-	bool UpdateCache(const FName& InName, const FRigCurveContainer* InHierarchy);
+	bool UpdateCache(const FRigElementKey& InKey, const URigHierarchy* InHierarchy);
 
 	friend FORCEINLINE uint32 GetTypeHash(const FCachedRigElement& Cache)
 	{
 		return GetTypeHash(Cache.Key) * 13 + (uint32)Cache.Index;
 	}
 
-	bool IsIdentical(const FRigElementKey& InKey, const FRigHierarchyContainer* InContainer);
+	bool IsIdentical(const FRigElementKey& InKey, const URigHierarchy* InHierarchy);
 
 	FORCEINLINE bool operator ==(const FCachedRigElement& Other) const
 	{
@@ -163,16 +137,6 @@ public:
 		return Index > Other.Index;
 	}
 
-	FORCEINLINE operator int32() const
-	{
-		return GetIndex();
-	}
-
-	FORCEINLINE operator FRigElementKey() const
-	{
-		return Key;
-	}
-
 private:
 
 	UPROPERTY()
@@ -183,4 +147,6 @@ private:
 
 	UPROPERTY()
 	int32 ContainerVersion;
+
+	const FRigBaseElement* Element;
 };
diff --git a/Engine/Plugins/Experimental/ControlRig/Source/ControlRig/Public/Rigs/RigHierarchyContainer.h b/Engine/Plugins/Experimental/ControlRig/Source/ControlRig/Public/Rigs/RigHierarchyContainer.h
--- a/Engine/Plugins/Experimental/ControlRig/Source/ControlRig/Public/Rigs/RigHierarchyContainer.h
+++ b/Engine/Plugins/Experimental/ControlRig/Source/ControlRig/Public/Rigs/RigHierarchyContainer.h
@@ -16,75 +16,7 @@
 class UControlRig;
 struct FRigHierarchyContainer;
 
-UENUM(BlueprintType)
-enum class ERigHierarchyImportMode : uint8
-{
-	Append,
-
-	Replace,
-
-	ReplaceInitialLocalTransform,
-
-	ReplaceInitialGlobalTransform,
-
-	ReplaceCurrentLocalTransform,
-
-	ReplaceCurrentGlobalTransform,
-
-	/** MAX - invalid */
-	Max UMETA(Hidden),
-};
-
 USTRUCT()
-struct CONTROLRIG_API FRigHierarchyCopyPasteContent
-{
-	GENERATED_USTRUCT_BODY()
-
-	UPROPERTY()
-	TArray<ERigElementType> Types;
-
-	UPROPERTY()
-	TArray<FString> Contents;
-
-	UPROPERTY()
-	TArray<FTransform> LocalTransforms;
-
-	UPROPERTY()
-	TArray<FTransform> GlobalTransforms;
-};
-
-USTRUCT()
-struct CONTROLRIG_API FRigMirrorSettings
-{
-	GENERATED_USTRUCT_BODY()
-
-	FRigMirrorSettings()
-	: MirrorAxis(EAxis::X)
-	, AxisToFlip(EAxis::Z)
-	{
-	}
-
-	// the axis to mirror against
-	UPROPERTY(EditAnywhere, Category = Settings)
-	TEnumAsByte<EAxis::Type> MirrorAxis;
-
-	// the axis to flip for rotations
-	UPROPERTY(EditAnywhere, Category = Settings)
-	TEnumAsByte<EAxis::Type> AxisToFlip;
-
-	// the string to replace all occurences of with New Name
-	UPROPERTY(EditAnywhere, Category = Settings)
-	FString OldName;
-
-	// the string to replace all occurences of Old Name with
-	UPROPERTY(EditAnywhere, Category = Settings)
-	FString NewName;
-
-	FTransform MirrorTransform(const FTransform& InTransform) const;
-	FVector MirrorVector(const FVector& InVector) const;
-};
-
-USTRUCT(BlueprintType)
 struct CONTROLRIG_API FRigHierarchyContainer
 {
 public:
@@ -107,280 +39,8 @@ public:
 	UPROPERTY()
 	FRigCurveContainer CurveContainer;
 
-	UPROPERTY(transient)
-	int32 Version;
-
-	void Initialize(bool bResetTransforms = true);
-	void Reset();
-	void ResetTransforms();
-
-	// copies all initial transforms from another hierarchy
-	void CopyInitialTransforms(const FRigHierarchyContainer& InOther);
-
-	// returns this hierarchy's current pose
-	FRigPose GetPose() const;
-
-	// sets the current hierarchy from the given pose
-	void SetPose(FRigPose& InPose);
-
-	FORCEINLINE int32 GetIndex(const FRigElementKey& InKey) const
-	{
-		switch(InKey.Type)
-		{
-			case ERigElementType::Bone:
-			{
-				return BoneHierarchy.GetIndex(InKey.Name);
-			}
-			case ERigElementType::Space:
-			{
-				return SpaceHierarchy.GetIndex(InKey.Name);
-			}
-			case ERigElementType::Control:
-			{
-				return ControlHierarchy.GetIndex(InKey.Name);
-			}
-			case ERigElementType::Curve:
-			{
-				return CurveContainer.GetIndex(InKey.Name);
-			}
-		}
-		return INDEX_NONE;
-	}
-
-	FORCEINLINE FTransform GetInitialTransform(const FRigElementKey& InKey) const
-	{
-		return GetInitialTransform(InKey.Type, GetIndex(InKey));
-	}
-	FORCEINLINE FTransform GetInitialTransform(const FCachedRigElement& InCachedElement) const
-	{
-		if(InCachedElement.IsValid())
-		{
-			return GetInitialTransform(InCachedElement.GetKey().Type, InCachedElement.GetIndex());
-		}
-		return FTransform::Identity;
-	}
-	FORCEINLINE FTransform GetInitialTransform(FCachedRigElement& InCachedElement) const
-	{
-		if(InCachedElement.UpdateCache(this))
-		{
-			return GetInitialTransform(InCachedElement.GetKey().Type, InCachedElement.GetIndex());
-		}
-		return FTransform::Identity;
-	}
-	FTransform GetInitialTransform(ERigElementType InElementType, int32 InIndex) const;
-
-	FORCEINLINE void SetInitialTransform(const FRigElementKey& InKey, const FTransform& InTransform)
-	{
-		SetInitialTransform(InKey.Type, GetIndex(InKey), InTransform);
-	}
-	FORCEINLINE void SetInitialTransform(const FCachedRigElement& InCachedElement, const FTransform& InTransform)
-	{
-		if(InCachedElement.IsValid())
-		{
-			SetInitialTransform(InCachedElement.GetKey().Type, InCachedElement.GetIndex(), InTransform);
-		}
-	}
-	FORCEINLINE void SetInitialTransform(FCachedRigElement& InCachedElement, const FTransform& InTransform)
-	{
-		if(InCachedElement.UpdateCache(this))
-		{
-			SetInitialTransform(InCachedElement.GetKey().Type, InCachedElement.GetIndex(), InTransform);
-		}
-	}
-	void SetInitialTransform(ERigElementType InElementType, int32 InIndex, const FTransform& InTransform);
-
-	FORCEINLINE FTransform GetInitialGlobalTransform(const FRigElementKey& InKey) const
-	{
-		return GetInitialGlobalTransform(InKey.Type, GetIndex(InKey));
-	}
-	FORCEINLINE FTransform GetInitialGlobalTransform(const FCachedRigElement& InCachedElement) const
-	{
-		if(InCachedElement.IsValid())
-		{
-			return GetInitialGlobalTransform(InCachedElement.GetKey().Type, InCachedElement.GetIndex());
-		}
-		return FTransform::Identity;
-	}
-	FORCEINLINE FTransform GetInitialGlobalTransform(FCachedRigElement& InCachedElement) const
-	{
-		if(InCachedElement.UpdateCache(this))
-		{
-			return GetInitialGlobalTransform(InCachedElement.GetKey().Type, InCachedElement.GetIndex());
-		}
-		return FTransform::Identity;
-	}
-	FTransform GetInitialGlobalTransform(ERigElementType InElementType, int32 InIndex) const;
-
-	FORCEINLINE void SetInitialGlobalTransform(const FRigElementKey& InKey, const FTransform& InTransform)
-	{
-		SetInitialGlobalTransform(InKey.Type, GetIndex(InKey), InTransform);
-	}
-	FORCEINLINE void SetInitialGlobalTransform(const FCachedRigElement& InCachedElement, const FTransform& InTransform)
-	{
-		if(InCachedElement.IsValid())
-		{
-			SetInitialGlobalTransform(InCachedElement.GetKey().Type, InCachedElement.GetIndex(), InTransform);
-		}
-	}
-	FORCEINLINE void SetInitialGlobalTransform(FCachedRigElement& InCachedElement, const FTransform& InTransform)
-	{
-		if(InCachedElement.UpdateCache(this))
-		{
-			SetInitialGlobalTransform(InCachedElement.GetKey().Type, InCachedElement.GetIndex(), InTransform);
-		}
-	}
-	void SetInitialGlobalTransform(ERigElementType InElementType, int32 InIndex, const FTransform& InTransform);
-
-	FORCEINLINE FTransform GetLocalTransform(const FRigElementKey& InKey) const
-	{
-		return GetLocalTransform(InKey.Type, GetIndex(InKey));
-	}
-	FORCEINLINE FTransform GetLocalTransform(const FCachedRigElement& InCachedElement) const
-	{
-		if(InCachedElement.IsValid())
-		{
-			return GetLocalTransform(InCachedElement.GetKey().Type, InCachedElement.GetIndex());
-		}
-		return FTransform::Identity;
-	}
-	FORCEINLINE FTransform GetLocalTransform(FCachedRigElement& InCachedElement) const
-	{
-		if(InCachedElement.UpdateCache(this))
-		{
-			return GetLocalTransform(InCachedElement.GetKey().Type, InCachedElement.GetIndex());
-		}
-		return FTransform::Identity;
-	}
-	FTransform GetLocalTransform(ERigElementType InElementType, int32 InIndex) const;
-
-	FORCEINLINE void SetLocalTransform(const FRigElementKey& InKey, const FTransform& InTransform, bool bPropagateToChildren = true)
-	{
-		SetLocalTransform(InKey.Type, GetIndex(InKey), InTransform, bPropagateToChildren);
-	}
-	FORCEINLINE void SetLocalTransform(const FCachedRigElement& InCachedElement, const FTransform& InTransform, bool bPropagateToChildren = true)
-	{
-		if(InCachedElement.IsValid())
-		{
-			SetLocalTransform(InCachedElement.GetKey().Type, InCachedElement.GetIndex(), InTransform, bPropagateToChildren);
-		}
-	}
-	FORCEINLINE void SetLocalTransform(FCachedRigElement& InCachedElement, const FTransform& InTransform, bool bPropagateToChildren = true)
-	{
-		if(InCachedElement.UpdateCache(this))
-		{
-			SetLocalTransform(InCachedElement.GetKey().Type, InCachedElement.GetIndex(), InTransform, bPropagateToChildren);
-		}
-	}
-	void SetLocalTransform(ERigElementType InElementType, int32 InIndex, const FTransform& InTransform, bool bPropagateToChildren = true);
-
-	FORCEINLINE FTransform GetGlobalTransform(const FRigElementKey& InKey) const
-	{
-		return GetGlobalTransform(InKey.Type, GetIndex(InKey));
-	}
-	FORCEINLINE FTransform GetGlobalTransform(const FCachedRigElement& InCachedElement) const
-	{
-		if(InCachedElement.IsValid())
-		{
-			return GetGlobalTransform(InCachedElement.GetKey().Type, InCachedElement.GetIndex());
-		}
-		return FTransform::Identity;
-	}
-	FORCEINLINE FTransform GetGlobalTransform(FCachedRigElement& InCachedElement) const
-	{
-		if(InCachedElement.UpdateCache(this))
-		{
-			return GetGlobalTransform(InCachedElement.GetKey().Type, InCachedElement.GetIndex());
-		}
-		return FTransform::Identity;
-	}
-	FTransform GetGlobalTransform(ERigElementType InElementType, int32 InIndex) const;
-
-	FORCEINLINE void SetGlobalTransform(const FRigElementKey& InKey, const FTransform& InTransform, bool bPropagateToChildren = true)
-	{
-		SetGlobalTransform(InKey.Type, GetIndex(InKey), InTransform, bPropagateToChildren);
-	}
-	FORCEINLINE void SetGlobalTransform(const FCachedRigElement& InCachedElement, const FTransform& InTransform, bool bPropagateToChildren = true)
-	{
-		if(InCachedElement.IsValid())
-		{
-			SetGlobalTransform(InCachedElement.GetKey().Type, InCachedElement.GetIndex(), InTransform, bPropagateToChildren);
-		}
-	}
-	FORCEINLINE void SetGlobalTransform(FCachedRigElement& InCachedElement, const FTransform& InTransform, bool bPropagateToChildren = true)
-	{
-		if(InCachedElement.UpdateCache(this))
-		{
-			SetGlobalTransform(InCachedElement.GetKey().Type, InCachedElement.GetIndex(), InTransform, bPropagateToChildren);
-		}
-	}
-	void SetGlobalTransform(ERigElementType InElementType, int32 InIndex, const FTransform& InTransform, bool bPropagateToChildren = true);
-
-	FRigElementKey GetParentKey(const FRigElementKey& InKey) const;
-
-	TArray<FRigElementKey> GetChildKeys(const FRigElementKey& InKey, bool bRecursive = false) const;
-
-	bool IsParentedTo(ERigElementType InChildType, int32 InChildIndex, ERigElementType InParentType, int32 InParentIndex) const;
-
-	void SortKeyArray(TArray<FRigElementKey>& InOutKeysToSort) const;
-
-#if WITH_EDITOR
-
-	FString ExportSelectionToText() const;
-	FString ExportToText(const TArray<FRigElementKey>& InSelection) const;
-	TArray<FRigElementKey> ImportFromText(const FString& InContent, ERigHierarchyImportMode InImportMode = ERigHierarchyImportMode::Append, bool bSelectNewElements = true, const TArray<FRigElementKey>& InSelection = TArray<FRigElementKey>());
-
-	TArray<FRigElementKey> DuplicateItems(const TArray<FRigElementKey>& InSelection, bool bSelectNewElements = true);
-	TArray<FRigElementKey> MirrorItems(const TArray<FRigElementKey>& InSelection, const FRigMirrorSettings& InSettings, bool bSelectNewElements = true);
-
-
-#endif
-
-	TArray<FRigElementKey> CurrentSelection() const;
-	TArray<FRigElementKey> GetAllItems(bool bSort = false) const;
-
-
-	bool Select(const FRigElementKey& InKey, bool bSelect = true);
-	bool ClearSelection();
-	bool ClearSelection(ERigElementType InElementType);
-	bool IsSelected(const FRigElementKey& InKey) const;
-
-	FRigElementSelected OnElementSelected;
-	FRigElementChanged OnElementChanged;
-
-#if WITH_EDITOR
-	FRigElementAdded OnElementAdded;
-	FRigElementRemoved OnElementRemoved;
-	FRigElementRenamed OnElementRenamed;
-	FRigElementReparented OnElementReparented;
-
-#endif
-
-	void SendEvent(const FRigEventContext& InEvent, bool bAsyncronous = true);
-	FRigEventDelegate OnEventReceived;
-
-	static int32 GetMaxNameLength() { return 100; }
-	static void SanitizeName(FString& InOutName);
-	static FName GetSanitizedName(const FString& InName);
-	static FName GetSafeNewName(const FString& InName, TFunction<bool(const FString&)> IsNameAvailableFunction);
-	bool IsNameAvailable(const FString& InName, ERigElementType InElementType, FString* OutErrorMessage = nullptr);
-
-protected:
-
-	void HandleOnElementSelected(FRigHierarchyContainer* InContainer, const FRigElementKey& InKey, bool bSelected);
-#if WITH_EDITOR
-	void HandleOnElementAdded(FRigHierarchyContainer* InContainer, const FRigElementKey& InKey);
-	void HandleOnElementRemoved(FRigHierarchyContainer* InContainer, const FRigElementKey& InKey);
-	void HandleOnElementRenamed(FRigHierarchyContainer* InContainer, ERigElementType InElementType, const FName& InOldName, const FName& InNewName);
-	void HandleOnElementReparented(FRigHierarchyContainer* InContainer, const FRigElementKey& InKey, const FName& InOldParentName, const FName& InNewParentName);
-#endif
-
 private:
 
-	void UpdateDepthIndexIfRequired() const;
-	mutable TMap<FRigElementKey, int32> DepthIndexByKey;
-
-	void AppendToPose(FRigPose& InOutPose) const;
-
 	friend class SRigHierarchy;
 };
 
