Control Rig: New rig hierarchy data structure refactoring

diff --git a/Engine/Plugins/Experimental/ControlRig/Source/ControlRig/Public/Rigs/RigHierarchyDefines.h b/Engine/Plugins/Experimental/ControlRig/Source/ControlRig/Public/Rigs/RigHierarchyDefines.h
--- a/Engine/Plugins/Experimental/ControlRig/Source/ControlRig/Public/Rigs/RigHierarchyDefines.h
+++ b/Engine/Plugins/Experimental/ControlRig/Source/ControlRig/Public/Rigs/RigHierarchyDefines.h
@@ -3,9 +3,11 @@
 #pragma once
 
 #include "CoreMinimal.h"
+#include "TransformNoScale.h"
+#include "EulerTransform.h"
 #include "RigHierarchyDefines.generated.h"
 
-struct FRigHierarchyContainer;
+class URigHierarchy;
 
 /* 
  * This is rig element types that we support
@@ -22,6 +24,29 @@ enum class ERigElementType : uint8
 	All = Bone | Space | Control | Curve UMETA(Hidden),
 };
 
+UENUM(BlueprintType)
+enum class ERigBoneType : uint8
+{
+	Imported,
+	User
+};
+
+UENUM()
+enum class ERigHierarchyNotification : uint8
+{
+	ElementAdded,
+	ElementRemoved,
+	ElementRenamed,
+	ElementSelected,
+	ElementDeselected,
+	ParentChanged,
+	HierarchyReset,
+	ControlSettingChanged,
+
+	/** MAX - invalid */
+	Max UMETA(Hidden),
+};
+
 UENUM()
 enum class ERigEvent : uint8
 {
@@ -44,6 +69,665 @@ enum class EControlRigSetKey : uint8
 	Never				//Never set a key here.
 };
 
+UENUM(BlueprintType)
+enum class ERigControlType : uint8
+{
+	Bool,
+    Float,
+    Integer,
+    Vector2D,
+    Position,
+    Scale,
+    Rotator,
+    Transform,
+    TransformNoScale,
+    EulerTransform,
+};
+
+UENUM(BlueprintType)
+enum class ERigControlValueType : uint8
+{
+	Initial,
+    Current,
+    Minimum,
+    Maximum
+};
+
+UENUM(BlueprintType)
+enum class ERigControlAxis : uint8
+{
+	X,
+    Y,
+    Z
+};
+
+USTRUCT(BlueprintType)
+struct CONTROLRIG_API FRigControlValueStorage
+{
+public:
+
+	GENERATED_BODY()
+
+	FRigControlValueStorage()
+	{
+		FMemory::Memzero(this, sizeof(FRigControlValueStorage));
+	}
+
+	UPROPERTY()
+	float Float00;
+
+	UPROPERTY()
+	float Float01;
+
+	UPROPERTY()
+	float Float02;
+
+	UPROPERTY()
+	float Float03;
+
+	UPROPERTY()
+	float Float10;
+
+	UPROPERTY()
+	float Float11;
+
+	UPROPERTY()
+	float Float12;
+
+	UPROPERTY()
+	float Float13;
+
+	UPROPERTY()
+	float Float20;
+
+	UPROPERTY()
+	float Float21;
+
+	UPROPERTY()
+	float Float22;
+
+	UPROPERTY()
+	float Float23;
+
+	UPROPERTY()
+	float Float30;
+
+	UPROPERTY()
+	float Float31;
+
+	UPROPERTY()
+	float Float32;
+
+	UPROPERTY()
+	float Float33;
+
+	UPROPERTY()
+	bool bValid;
+};
+
+USTRUCT(BlueprintType)
+struct CONTROLRIG_API FRigControlValue
+{
+	GENERATED_BODY()
+
+public:
+
+	FRigControlValue()
+		: FloatStorage()
+		, Storage_DEPRECATED(FTransform::Identity)
+	{
+	}
+
+	FORCEINLINE_DEBUGGABLE bool IsValid() const
+	{
+		return FloatStorage.bValid;
+	}
+
+	template<class T>
+	FORCEINLINE_DEBUGGABLE T Get() const
+	{
+		return GetRef<T>();
+	}
+
+	template<class T>
+	FORCEINLINE_DEBUGGABLE T& GetRef()
+	{
+		FloatStorage.bValid = true;
+		return *(T*)&FloatStorage;
+	}
+
+	template<class T>
+	FORCEINLINE_DEBUGGABLE const T& GetRef() const
+	{
+		return *(T*)&FloatStorage;
+	}
+
+	template<class T>
+	FORCEINLINE_DEBUGGABLE void Set(T InValue)
+	{
+		GetRef<T>() = InValue;
+	}
+
+	template<class T>
+	FORCEINLINE_DEBUGGABLE FString ToString() const
+	{
+		FString Result;
+		TBaseStructure<T>::Get()->ExportText(Result, &GetRef<T>(), nullptr, nullptr, PPF_None, nullptr);
+		return Result;
+	}
+
+	template<class T>
+	FORCEINLINE_DEBUGGABLE T SetFromString(const FString& InString)
+	{
+		T Value;
+		TBaseStructure<T>::Get()->ImportText(*InString, &Value, nullptr, PPF_None, nullptr, TBaseStructure<T>::Get()->GetName());
+		Set<T>(Value);
+		return Value;
+	}
+
+	template<class T>
+	FORCEINLINE_DEBUGGABLE static FRigControlValue Make(T InValue)
+	{
+		FRigControlValue Value;
+		Value.Set<T>(InValue);
+		return Value;
+	}
+
+	FORCEINLINE FTransform GetAsTransform(ERigControlType InControlType, ERigControlAxis InPrimaryAxis) const
+	{
+		FTransform Transform = FTransform::Identity;
+		switch (InControlType)
+		{
+			case ERigControlType::Bool:
+			{
+				Transform.SetLocation(FVector(Get<bool>() ? 1.f : 0.f, 0.f, 0.f));
+				break;
+			}
+			case ERigControlType::Float:
+			{
+				const float ValueToGet = Get<float>();
+				switch (InPrimaryAxis)
+				{
+					case ERigControlAxis::X:
+					{
+						Transform.SetLocation(FVector(ValueToGet, 0.f, 0.f));
+						break;
+					}
+					case ERigControlAxis::Y:
+					{
+						Transform.SetLocation(FVector(0.f, ValueToGet, 0.f));
+						break;
+					}
+					case ERigControlAxis::Z:
+					{
+						Transform.SetLocation(FVector(0.f, 0.f, ValueToGet));
+						break;
+					}
+				}
+				break;
+			}
+			case ERigControlType::Integer:
+			{
+				const int32 ValueToGet = Get<int32>();
+				switch (InPrimaryAxis)
+				{
+					case ERigControlAxis::X:
+					{
+						Transform.SetLocation(FVector((float)ValueToGet, 0.f, 0.f));
+						break;
+					}
+					case ERigControlAxis::Y:
+					{
+						Transform.SetLocation(FVector(0.f, (float)ValueToGet, 0.f));
+						break;
+					}
+					case ERigControlAxis::Z:
+					{
+						Transform.SetLocation(FVector(0.f, 0.f, (float)ValueToGet));
+						break;
+					}
+				}
+				break;
+			}
+			case ERigControlType::Vector2D:
+			{
+				const FVector2D ValueToGet = Get<FVector2D>();
+				switch (InPrimaryAxis)
+				{
+					case ERigControlAxis::X:
+					{
+						Transform.SetLocation(FVector(0.f, ValueToGet.X, ValueToGet.Y));
+						break;
+					}
+					case ERigControlAxis::Y:
+					{
+						Transform.SetLocation(FVector(ValueToGet.X, 0.f, ValueToGet.Y));
+						break;
+					}
+					case ERigControlAxis::Z:
+					{
+						Transform.SetLocation(FVector(ValueToGet.X, ValueToGet.Y, 0.f));
+						break;
+					}
+				}
+				break;
+			}
+			case ERigControlType::Position:
+			{
+				Transform.SetLocation(Get<FVector>());
+				break;
+			}
+			case ERigControlType::Scale:
+			{
+				Transform.SetScale3D(Get<FVector>());
+				break;
+			}
+			case ERigControlType::Rotator:
+			{
+				Transform.SetRotation(FQuat(Get<FRotator>()));
+				break;
+			}
+			case ERigControlType::Transform:
+			{
+				return Get<FTransform>();
+			}
+			case ERigControlType::TransformNoScale:
+			{
+				const FTransformNoScale TransformNoScale = Get<FTransformNoScale>();
+				Transform = TransformNoScale;
+				Transform.NormalizeRotation();
+				break;
+			}
+			case ERigControlType::EulerTransform:
+			{
+				const FEulerTransform EulerTransform = Get<FEulerTransform>();
+				Transform = FTransform(EulerTransform.ToFTransform());
+				Transform.NormalizeRotation();
+				break;
+			}
+			default:
+			{
+				ensure(false);
+				break;
+			}
+		}
+		return Transform;
+	}
+	
+	FORCEINLINE void SetFromTransform(const FTransform& InTransform, ERigControlType InControlType, ERigControlAxis InPrimaryAxis)
+	{
+		switch (InControlType)
+		{
+			case ERigControlType::Bool:
+			{
+				Set<bool>(InTransform.GetLocation().X > SMALL_NUMBER);
+				break;
+			}
+			case ERigControlType::Float:
+			{
+				switch (InPrimaryAxis)
+				{
+					case ERigControlAxis::X:
+					{
+						Set<float>(InTransform.GetLocation().X);
+						break;
+					}
+					case ERigControlAxis::Y:
+					{
+						Set<float>(InTransform.GetLocation().Y);
+						break;
+					}
+					case ERigControlAxis::Z:
+					{
+						Set<float>(InTransform.GetLocation().Z);
+						break;
+					}
+				}
+				break;
+			}
+			case ERigControlType::Integer:
+			{
+				switch (InPrimaryAxis)
+				{
+					case ERigControlAxis::X:
+					{
+						Set<int32>((int32)InTransform.GetLocation().X);
+						break;
+					}
+					case ERigControlAxis::Y:
+					{
+						Set<int32>((int32)InTransform.GetLocation().Y);
+						break;
+					}
+					case ERigControlAxis::Z:
+					{
+						Set<int32>((int32)InTransform.GetLocation().Z);
+						break;
+					}
+				}
+				break;
+			}
+			case ERigControlType::Vector2D:
+			{
+				const FVector Location = InTransform.GetLocation();
+				switch (InPrimaryAxis)
+				{
+					case ERigControlAxis::X:
+					{
+						Set<FVector2D>(FVector2D(Location.Y, Location.Z));
+						break;
+					}
+					case ERigControlAxis::Y:
+					{
+						Set<FVector2D>(FVector2D(Location.X, Location.Z));
+						break;
+					}
+					case ERigControlAxis::Z:
+					{
+						Set<FVector2D>(FVector2D(Location.X, Location.Y));
+						break;
+					}
+				}
+				break;
+			}
+			case ERigControlType::Position:
+			{
+				Set<FVector>(InTransform.GetLocation());
+				break;
+			}
+			case ERigControlType::Scale:
+			{
+				Set<FVector>(InTransform.GetScale3D());
+				break;
+			}
+			case ERigControlType::Rotator:
+			{
+				//allow for values ><180/-180 by getting diff and adding that back in.
+				FRotator CurrentRotator = Get<FRotator>();
+				FRotator CurrentRotWind, CurrentRotRem;
+				CurrentRotator.GetWindingAndRemainder(CurrentRotWind, CurrentRotRem);
+
+				//Get Diff
+				const FRotator NewRotator = FRotator(InTransform.GetRotation());
+				FRotator DeltaRot = NewRotator - CurrentRotRem;
+				DeltaRot.Normalize();
+
+				//Add Diff
+				CurrentRotator = CurrentRotator + DeltaRot;
+				Set<FRotator>(CurrentRotator);
+				break;
+			}
+			case ERigControlType::Transform:
+			{
+				Set<FTransform>(InTransform);
+				break;
+			}
+			case ERigControlType::TransformNoScale:
+			{
+				const FTransformNoScale NoScale = InTransform;
+				Set<FTransformNoScale>(NoScale);
+				break;
+			}
+			case ERigControlType::EulerTransform:
+			{
+				//Find Diff of the rotation from current and just add that instead of setting so we can go over/under -180
+				FEulerTransform NewTransform(InTransform);
+
+				const FEulerTransform CurrentEulerTransform = Get<FEulerTransform>();
+				FRotator CurrentWinding;
+				FRotator CurrentRotRemainder;
+				CurrentEulerTransform.Rotation.GetWindingAndRemainder(CurrentWinding, CurrentRotRemainder);
+				const FRotator NewRotator = InTransform.GetRotation().Rotator();
+				FRotator DeltaRot = NewRotator - CurrentRotRemainder;
+				DeltaRot.Normalize();
+				const FRotator NewRotation(CurrentEulerTransform.Rotation + DeltaRot);
+				NewTransform.Rotation = NewRotation;
+				Set<FEulerTransform>(NewTransform);
+				break;
+			}
+			default:
+			{
+				ensure(false);
+				break;
+			}
+		}
+	}
+
+	FORCEINLINE void ApplyLimits(
+		bool bLimitTranslation,
+		bool bLimitRotation,
+		bool bLimitScale,
+		ERigControlType InControlType,
+		const FRigControlValue& InMinimumValue,
+		const FRigControlValue& InMaximumValue)
+	{
+		if (!bLimitTranslation && !bLimitRotation && !bLimitScale)
+		{
+			return;
+		}
+
+		struct Local
+		{
+			FORCEINLINE static float Clamp(const float Value, const float Minimum, const float Maximum)
+			{
+				if (Minimum < Maximum)
+				{
+					return FMath::Clamp<float>(Value, Minimum, Maximum);
+				}
+				return FMath::Clamp<float>(Value, Maximum, Minimum);
+			}
+
+			FORCEINLINE static int32 Clamp(const int32 Value, const int32 Minimum, const int32 Maximum)
+			{
+				if (Minimum < Maximum)
+				{
+					return FMath::Clamp<int32>(Value, Minimum, Maximum);
+				}
+				return FMath::Clamp<int32>(Value, Maximum, Minimum);
+			}
+		};
+
+		switch(InControlType)
+		{
+			case ERigControlType::Float:
+			{
+				if (bLimitTranslation)
+				{
+					float& ValueRef = GetRef<float>();
+					ValueRef = Local::Clamp(ValueRef, InMinimumValue.Get<float>(), InMaximumValue.Get<float>());
+				}
+				break;
+			}
+			case ERigControlType::Integer:
+			{
+				if (bLimitTranslation)
+				{
+					int32& ValueRef = GetRef<int32>();
+					ValueRef = Local::Clamp(ValueRef, InMinimumValue.Get<int32>(), InMaximumValue.Get<int32>());
+				}
+				break;
+			}
+			case ERigControlType::Vector2D:
+			{
+				if (bLimitTranslation)
+				{
+					FVector2D& ValueRef = GetRef<FVector2D>();
+					const FVector2D& Min = InMinimumValue.GetRef<FVector2D>();
+					const FVector2D& Max = InMaximumValue.GetRef<FVector2D>();
+					ValueRef.X = Local::Clamp(ValueRef.X, Min.X, Max.X);
+					ValueRef.Y = Local::Clamp(ValueRef.Y, Min.Y, Max.Y);
+				}
+				break;
+			}
+			case ERigControlType::Position:
+			{
+				if (bLimitTranslation)
+				{
+					FVector& ValueRef = GetRef<FVector>();
+					const FVector& Min = InMinimumValue.GetRef<FVector>();
+					const FVector& Max = InMaximumValue.GetRef<FVector>();
+					ValueRef.X = Local::Clamp(ValueRef.X, Min.X, Max.X);
+					ValueRef.Y = Local::Clamp(ValueRef.Y, Min.Y, Max.Y);
+					ValueRef.Z = Local::Clamp(ValueRef.Z, Min.Z, Max.Z);
+				}
+				break;
+			}
+			case ERigControlType::Scale:
+			{
+				if (bLimitScale)
+				{
+					FVector& ValueRef = GetRef<FVector>();
+					const FVector& Min = InMinimumValue.GetRef<FVector>();
+					const FVector& Max = InMaximumValue.GetRef<FVector>();
+					ValueRef.X = Local::Clamp(ValueRef.X, Min.X, Max.X);
+					ValueRef.Y = Local::Clamp(ValueRef.Y, Min.Y, Max.Y);
+					ValueRef.Z = Local::Clamp(ValueRef.Z, Min.Z, Max.Z);
+				}
+				break;
+			}
+			case ERigControlType::Rotator:
+			{
+				if (bLimitRotation)
+				{
+					FRotator& ValueRef = GetRef<FRotator>();
+					const FRotator& Min = InMinimumValue.GetRef<FRotator>();
+					const FRotator& Max = InMaximumValue.GetRef<FRotator>();
+					ValueRef.Pitch = Local::Clamp(ValueRef.Pitch, Min.Pitch, Max.Pitch);
+					ValueRef.Yaw = Local::Clamp(ValueRef.Yaw, Min.Yaw, Max.Yaw);
+					ValueRef.Roll = Local::Clamp(ValueRef.Roll, Min.Roll, Max.Roll);
+				}
+				break;
+			}
+			case ERigControlType::Transform:
+			{
+				FTransform& ValueRef = GetRef<FTransform>();
+				const FTransform& Min = InMinimumValue.GetRef<FTransform>();
+				const FTransform& Max = InMaximumValue.GetRef<FTransform>();
+
+				if (bLimitTranslation)
+				{
+					ValueRef.SetLocation(FVector(
+						Local::Clamp(ValueRef.GetLocation().X, Min.GetLocation().X, Max.GetLocation().X),
+						Local::Clamp(ValueRef.GetLocation().Y, Min.GetLocation().Y, Max.GetLocation().Y),
+						Local::Clamp(ValueRef.GetLocation().Z, Min.GetLocation().Z, Max.GetLocation().Z)
+					));
+				}
+				if (bLimitRotation)
+				{
+					const FRotator Rotator = ValueRef.GetRotation().Rotator();
+					const FRotator MinRotator = Min.GetRotation().Rotator();
+					const FRotator MaxRotator = Max.GetRotation().Rotator();
+
+					ValueRef.SetRotation(FQuat(FRotator(
+						Local::Clamp(Rotator.Pitch, MinRotator.Pitch, MaxRotator.Pitch),
+						Local::Clamp(Rotator.Yaw, MinRotator.Yaw, MaxRotator.Yaw),
+						Local::Clamp(Rotator.Roll, MinRotator.Roll, MaxRotator.Roll)
+					)));
+				}
+				if (bLimitScale)
+				{
+					ValueRef.SetScale3D(FVector(
+						Local::Clamp(ValueRef.GetScale3D().X, Min.GetScale3D().X, Max.GetScale3D().X),
+						Local::Clamp(ValueRef.GetScale3D().Y, Min.GetScale3D().Y, Max.GetScale3D().Y),
+						Local::Clamp(ValueRef.GetScale3D().Z, Min.GetScale3D().Z, Max.GetScale3D().Z)
+					));
+				}
+				break;
+			}
+			case ERigControlType::TransformNoScale:
+			{
+				FTransformNoScale& ValueRef = GetRef<FTransformNoScale>();
+				const FTransformNoScale& Min = InMinimumValue.GetRef<FTransformNoScale>();
+				const FTransformNoScale& Max = InMaximumValue.GetRef<FTransformNoScale>();
+
+				if (bLimitTranslation)
+				{
+					ValueRef.Location = FVector(
+						Local::Clamp(ValueRef.Location.X, Min.Location.X, Max.Location.X),
+						Local::Clamp(ValueRef.Location.Y, Min.Location.Y, Max.Location.Y),
+						Local::Clamp(ValueRef.Location.Z, Min.Location.Z, Max.Location.Z)
+					);
+				}
+				if (bLimitRotation)
+				{
+					const FRotator Rotator = ValueRef.Rotation.Rotator();
+					const FRotator MinRotator = Min.Rotation.Rotator();
+					const FRotator MaxRotator = Max.Rotation.Rotator();
+
+					ValueRef.Rotation = FQuat(FRotator(
+						Local::Clamp(Rotator.Pitch, MinRotator.Pitch, MaxRotator.Pitch),
+						Local::Clamp(Rotator.Yaw, MinRotator.Yaw, MaxRotator.Yaw),
+						Local::Clamp(Rotator.Roll, MinRotator.Roll, MaxRotator.Roll)
+					));
+				}
+				break;
+			}
+
+			case ERigControlType::EulerTransform:
+			{
+				FEulerTransform& ValueRef = GetRef<FEulerTransform>();
+				const FEulerTransform& Min = InMinimumValue.GetRef<FEulerTransform>();
+				const FEulerTransform& Max = InMaximumValue.GetRef<FEulerTransform>();
+
+				if (bLimitTranslation)
+				{
+					ValueRef.Location = FVector(
+						Local::Clamp(ValueRef.Location.X, Min.Location.X, Max.Location.X),
+						Local::Clamp(ValueRef.Location.Y, Min.Location.Y, Max.Location.Y),
+						Local::Clamp(ValueRef.Location.Z, Min.Location.Z, Max.Location.Z)
+					);
+				}
+				if (bLimitRotation)
+				{
+					const FRotator Rotator = ValueRef.Rotation;
+					const FRotator MinRotator = Min.Rotation;
+					const FRotator MaxRotator = Max.Rotation;
+
+					ValueRef.Rotation = FRotator(
+						Local::Clamp(Rotator.Pitch, MinRotator.Pitch, MaxRotator.Pitch),
+						Local::Clamp(Rotator.Yaw, MinRotator.Yaw, MaxRotator.Yaw),
+						Local::Clamp(Rotator.Roll, MinRotator.Roll, MaxRotator.Roll)
+					);
+				}
+				if (bLimitScale)
+				{
+					ValueRef.Location = FVector(
+						Local::Clamp(ValueRef.Scale.X, Min.Scale.X, Max.Scale.X),
+						Local::Clamp(ValueRef.Scale.Y, Min.Scale.Y, Max.Scale.Y),
+						Local::Clamp(ValueRef.Scale.Z, Min.Scale.Z, Max.Scale.Z)
+					);
+				}
+				break;
+			}
+			case ERigControlType::Bool:
+			default:
+			{
+				break;
+			}
+		}
+	}
+
+private:
+
+	UPROPERTY()
+	FRigControlValueStorage FloatStorage;
+	
+	UPROPERTY()
+	FTransform Storage_DEPRECATED;
+
+	friend struct FRigControlHierarchy;
+	friend class UControlRigBlueprint;
+};
+
+template<>
+FORCEINLINE_DEBUGGABLE FQuat FRigControlValue::SetFromString<FQuat>(const FString& InString)
+{
+	FQuat Value;
+	TBaseStructure<FQuat>::Get()->ImportText(*InString, &Value, nullptr, PPF_None, nullptr, TBaseStructure<FQuat>::Get()->GetName());
+	Set<FRotator>(Value.Rotator());
+	return Value;
+}
+
 USTRUCT(BlueprintType)
 struct CONTROLRIG_API FRigControlModifiedContext
 {
@@ -101,6 +785,8 @@ namespace FRigElementTypeHelper
 USTRUCT(BlueprintType)
 struct CONTROLRIG_API FRigElementKey
 {
+public:
+	
 	GENERATED_BODY()
 
 	UPROPERTY(BlueprintReadWrite, EditAnywhere, Category = "Hierarchy")
@@ -119,12 +805,21 @@ struct CONTROLRIG_API FRigElementKey
 		, Name(InName)
 	{}
 
+	void Serialize(FArchive& Ar);
+	void Save(FArchive& Ar);
+	void Load(FArchive& Ar);
+	FORCEINLINE friend FArchive& operator<<(FArchive& Ar, FRigElementKey& P)
+	{
+		P.Serialize(Ar);
+		return Ar;
+	}
+
 	FORCEINLINE bool IsValid() const
 	{
 		return Name != NAME_None && Type != ERigElementType::None;
 	}
 
-	FORCEINLINE operator bool() const
+	FORCEINLINE explicit operator bool() const
 	{
 		return IsValid();
 	}
@@ -296,7 +991,7 @@ struct CONTROLRIG_API FRigElementKeyCollection
 
 	// creates a collection containing all of the children of a given 
 	static FRigElementKeyCollection MakeFromChildren(
-		const FRigHierarchyContainer* InContainer, 
+		URigHierarchy* InHierarchy, 
 		const FRigElementKey& InParentKey,
 		bool bRecursive = true,
 		bool bIncludeParent = false,
@@ -304,20 +999,20 @@ struct CONTROLRIG_API FRigElementKeyCollection
 
 	// creates a collection containing all of the elements with a given name
 	static FRigElementKeyCollection MakeFromName(
-		const FRigHierarchyContainer* InContainer,
+		URigHierarchy* InHierarchy,
 		const FName& InPartialName,
 		uint8 InElementTypes = (uint8)ERigElementType::All);
 
 	// creates a collection containing an item chain
 	static FRigElementKeyCollection MakeFromChain(
-		const FRigHierarchyContainer* InContainer,
+		URigHierarchy* InHierarchy,
 		const FRigElementKey& InFirstItem,
 		const FRigElementKey& InLastItem,
 		bool bReverse = false);
 
 	// creates a collection containing all keys of a hierarchy
 	static FRigElementKeyCollection MakeFromCompleteHierarchy(
-		const FRigHierarchyContainer* InContainer,
+		URigHierarchy* InHierarchy,
 		uint8 InElementTypes = (uint8)ERigElementType::All);
 
 	// returns the union between two collections
@@ -391,10 +1086,35 @@ struct CONTROLRIG_API FRigEventContext
 	void* Payload;
 };
 
-DECLARE_MULTICAST_DELEGATE_TwoParams(FRigElementChanged, FRigHierarchyContainer*, const FRigElementKey&);
-DECLARE_MULTICAST_DELEGATE_TwoParams(FRigElementAdded, FRigHierarchyContainer*, const FRigElementKey&);
-DECLARE_MULTICAST_DELEGATE_TwoParams(FRigElementRemoved, FRigHierarchyContainer*, const FRigElementKey&);
-DECLARE_MULTICAST_DELEGATE_FourParams(FRigElementRenamed, FRigHierarchyContainer*, ERigElementType, const FName&, const FName&);
-DECLARE_MULTICAST_DELEGATE_FourParams(FRigElementReparented, FRigHierarchyContainer*, const FRigElementKey&, const FName&, const FName&);
-DECLARE_MULTICAST_DELEGATE_ThreeParams(FRigElementSelected, FRigHierarchyContainer*, const FRigElementKey&, bool);
-DECLARE_MULTICAST_DELEGATE_TwoParams(FRigEventDelegate, FRigHierarchyContainer*, const FRigEventContext&);
+USTRUCT(BlueprintType)
+struct CONTROLRIG_API FRigMirrorSettings
+{
+	GENERATED_USTRUCT_BODY()
+
+    FRigMirrorSettings()
+    : MirrorAxis(EAxis::X)
+    , AxisToFlip(EAxis::Z)
+	{
+	}
+
+	// the axis to mirror against
+	UPROPERTY(EditAnywhere, Category = Settings)
+	TEnumAsByte<EAxis::Type> MirrorAxis;
+
+	// the axis to flip for rotations
+	UPROPERTY(EditAnywhere, Category = Settings)
+	TEnumAsByte<EAxis::Type> AxisToFlip;
+
+	// the string to replace all occurences of with New Name
+	UPROPERTY(EditAnywhere, Category = Settings)
+	FString OldName;
+
+	// the string to replace all occurences of Old Name with
+	UPROPERTY(EditAnywhere, Category = Settings)
+	FString NewName;
+
+	FTransform MirrorTransform(const FTransform& InTransform) const;
+	FVector MirrorVector(const FVector& InVector) const;
+};
+
+DECLARE_MULTICAST_DELEGATE_TwoParams(FRigEventDelegate, URigHierarchy*, const FRigEventContext&);
diff --git a/Engine/Plugins/Experimental/ControlRig/Source/ControlRig/Public/Rigs/RigHierarchyElements.h b/Engine/Plugins/Experimental/ControlRig/Source/ControlRig/Public/Rigs/RigHierarchyElements.h
new file mode 100644
--- /dev/null
+++ b/Engine/Plugins/Experimental/ControlRig/Source/ControlRig/Public/Rigs/RigHierarchyElements.h
@@ -0,0 +1,920 @@
+// Copyright Epic Games, Inc. All Rights Reserved.
+
+#pragma once
+
+#include "CoreMinimal.h"
+#include "RigHierarchyDefines.h"
+#include "RigHierarchyElements.generated.h"
+
+struct FRigBaseElement;
+class URigHierarchy;
+
+#define DECLARE_RIG_ELEMENT_METHODS(ElementType) \
+template<typename T> \
+friend FORCEINLINE const T* Cast(const ElementType* InElement) \
+{ \
+   return Cast<T>((const FRigBaseElement*) InElement); \
+} \
+template<typename T> \
+friend FORCEINLINE T* Cast(ElementType* InElement) \
+{ \
+   return Cast<T>((FRigBaseElement*) InElement); \
+}
+
+namespace ERigTransformType
+{
+	enum Type
+	{
+		InitialLocal,
+        CurrentLocal,
+        InitialGlobal,
+        CurrentGlobal
+	};
+
+	FORCEINLINE ERigTransformType::Type SwapCurrentAndInitial(const Type InTransformType)
+	{
+		switch(InTransformType)
+		{
+			case CurrentLocal:
+			{
+				return InitialLocal;
+			}
+			case CurrentGlobal:
+			{
+				return InitialGlobal;
+			}
+			case InitialLocal:
+			{
+				return CurrentLocal;
+			}
+			default:
+			{
+				break;
+			}
+		}
+		return CurrentGlobal;
+	}
+
+	FORCEINLINE Type SwapLocalAndGlobal(const Type InTransformType)
+	{
+		switch(InTransformType)
+		{
+			case CurrentLocal:
+			{
+				return CurrentGlobal;
+			}
+			case CurrentGlobal:
+			{
+				return CurrentLocal;
+			}
+			case InitialLocal:
+			{
+				return InitialGlobal;
+			}
+			default:
+			{
+				break;
+			}
+		}
+		return InitialLocal;
+	}
+
+	FORCEINLINE Type MakeLocal(const Type InTransformType)
+	{
+		switch(InTransformType)
+		{
+			case CurrentLocal:
+			case CurrentGlobal:
+			{
+				return CurrentLocal;
+			}
+			default:
+			{
+				break;
+			}
+		}
+		return InitialLocal;
+	}
+
+	FORCEINLINE Type MakeGlobal(const Type InTransformType)
+	{
+		switch(InTransformType)
+		{
+			case CurrentLocal:
+			case CurrentGlobal:
+			{
+				return CurrentGlobal;
+			}
+			default:
+			{
+				break;
+			}
+		}
+		return InitialGlobal;
+	}
+
+	FORCEINLINE Type MakeInitial(const Type InTransformType)
+	{
+		switch(InTransformType)
+		{
+			case CurrentLocal:
+			case InitialLocal:
+			{
+				return InitialLocal;
+			}
+			default:
+			{
+				break;
+			}
+		}
+		return InitialGlobal;
+	}
+
+	FORCEINLINE Type MakeCurrent(const Type InTransformType)
+	{
+		switch(InTransformType)
+		{
+			case CurrentLocal:
+        	case InitialLocal:
+			{
+				return CurrentLocal;
+			}
+			default:
+			{
+				break;
+			}
+		}
+		return CurrentGlobal;
+	}
+
+	FORCEINLINE bool IsLocal(const Type InTransformType)
+	{
+		switch(InTransformType)
+		{
+			case CurrentLocal:
+        	case InitialLocal:
+			{
+				return true;
+			}
+			default:
+			{
+				break;
+			}
+		}
+		return false;
+	}
+
+	FORCEINLINE bool IsGlobal(const Type InTransformType)
+	{
+		switch(InTransformType)
+		{
+			case CurrentGlobal:
+        	case InitialGlobal:
+			{
+				return true;;
+			}
+			default:
+			{
+				break;
+			}
+		}
+		return false;
+	}
+
+	FORCEINLINE bool IsInitial(const Type InTransformType)
+	{
+		switch(InTransformType)
+		{
+			case InitialLocal:
+        	case InitialGlobal:
+			{
+				return true;
+			}
+			default:
+			{
+				break;
+			}
+		}
+		return false;
+	}
+
+	FORCEINLINE bool IsCurrent(const Type InTransformType)
+	{
+		switch(InTransformType)
+		{
+			case CurrentLocal:
+        	case CurrentGlobal:
+			{
+				return true;
+			}
+			default:
+			{
+				break;
+			}
+		}
+		return false;
+	}
+}
+
+USTRUCT(BlueprintType)
+struct CONTROLRIG_API FRigComputedTransform
+{
+	GENERATED_BODY()
+
+	FRigComputedTransform()
+	: Transform(FTransform::Identity)
+	, bDirty(false)
+	{}
+
+	void Save(FArchive& Ar);
+	void Load(FArchive& Ar);
+
+	FORCEINLINE_DEBUGGABLE void Set(const FTransform& InTransform)
+	{
+#if WITH_EDITOR
+		ensure(InTransform.GetRotation().IsNormalized());
+#endif
+		// ensure(!FMath::IsNearlyZero(InTransform.GetScale3D().X));
+		// ensure(!FMath::IsNearlyZero(InTransform.GetScale3D().Y));
+		// ensure(!FMath::IsNearlyZero(InTransform.GetScale3D().Z));
+		Transform = InTransform;
+		bDirty = false;
+	}
+
+	FORCEINLINE static bool Equals(const FTransform& A, const FTransform& B, const float InTolerance = 0.0001f)
+	{
+		return (A.GetTranslation() - B.GetTranslation()).IsNearlyZero(InTolerance) &&
+			A.GetRotation().Equals(B.GetRotation(), InTolerance) &&
+			(A.GetScale3D() - B.GetScale3D()).IsNearlyZero(InTolerance);
+	}
+
+	FORCEINLINE bool operator == (const FRigComputedTransform& Other) const
+	{
+		return bDirty == Other.bDirty && Equals(Transform, Other.Transform);
+    }
+
+	UPROPERTY(BlueprintReadOnly, EditAnywhere, Category = "Pose")
+	FTransform Transform;
+
+	bool bDirty;
+};
+
+USTRUCT(BlueprintType)
+struct CONTROLRIG_API FRigLocalAndGlobalTransform
+{
+	GENERATED_BODY()
+
+	FRigLocalAndGlobalTransform()
+    : Local()
+    , Global()
+	{}
+
+	void Save(FArchive& Ar);
+	void Load(FArchive& Ar);
+
+	FORCEINLINE bool operator == (const FRigLocalAndGlobalTransform& Other) const
+	{
+		return Local == Other.Local && Global == Other.Global;
+	}
+
+	UPROPERTY(BlueprintReadOnly, EditAnywhere, Category = "Pose")
+	FRigComputedTransform Local;
+
+	UPROPERTY(BlueprintReadOnly, EditAnywhere, Category = "Pose")
+	FRigComputedTransform Global;
+};
+
+USTRUCT(BlueprintType)
+struct CONTROLRIG_API FRigCurrentAndInitialTransform
+{
+	GENERATED_BODY()
+
+	FRigCurrentAndInitialTransform()
+    : Current()
+    , Initial()
+	{}
+
+	FORCEINLINE const FRigComputedTransform& operator[](const ERigTransformType::Type InTransformType) const
+	{
+		switch(InTransformType)
+		{
+			case ERigTransformType::CurrentLocal:
+			{
+				return Current.Local;
+			}
+			case ERigTransformType::CurrentGlobal:
+			{
+				return Current.Global;
+			}
+			case ERigTransformType::InitialLocal:
+			{
+				return Initial.Local;
+			}
+			default:
+			{
+				break;
+			}
+		}
+		return Initial.Global;
+	}
+
+	FORCEINLINE FRigComputedTransform& operator[](const ERigTransformType::Type InTransformType)
+	{
+		switch(InTransformType)
+		{
+			case ERigTransformType::CurrentLocal:
+			{
+				return Current.Local;
+			}
+			case ERigTransformType::CurrentGlobal:
+			{
+				return Current.Global;
+			}
+			case ERigTransformType::InitialLocal:
+			{
+				return Initial.Local;
+			}
+			default:
+			{
+				break;
+			}
+		}
+		return Initial.Global;
+	}
+
+	FORCEINLINE const FTransform& Get(const ERigTransformType::Type InTransformType) const
+	{
+		return operator[](InTransformType).Transform;
+	}
+
+	FORCEINLINE void Set(const ERigTransformType::Type InTransformType, const FTransform& InTransform)
+	{
+		operator[](InTransformType).Set(InTransform);
+	}
+
+	FORCEINLINE bool IsDirty(const ERigTransformType::Type InTransformType) const
+	{
+		return operator[](InTransformType).bDirty;
+	}
+
+	FORCEINLINE void MarkDirty(const ERigTransformType::Type InTransformType)
+	{
+		ensure(!(operator[](ERigTransformType::SwapLocalAndGlobal(InTransformType)).bDirty));
+		operator[](InTransformType).bDirty = true;
+	}
+
+	void Save(FArchive& Ar);
+	void Load(FArchive& Ar);
+
+	FORCEINLINE bool operator == (const FRigCurrentAndInitialTransform& Other) const
+	{
+		return Current == Other.Current && Initial == Other.Initial;
+	}
+
+	UPROPERTY(BlueprintReadOnly, EditAnywhere, Category = "Pose")
+	FRigLocalAndGlobalTransform Current;
+
+	UPROPERTY(BlueprintReadOnly, EditAnywhere, Category = "Pose")
+	FRigLocalAndGlobalTransform Initial;
+};
+
+USTRUCT()
+struct CONTROLRIG_API FRigHierarchyCopyPasteContentPerElement
+{
+	GENERATED_USTRUCT_BODY()
+
+    UPROPERTY()
+	FRigElementKey Key;
+
+	UPROPERTY()
+	FString Content;
+
+	UPROPERTY()
+	TArray<FRigElementKey> Parents;
+
+	UPROPERTY()
+	TArray<float> ParentWeights;
+
+	UPROPERTY()
+	FRigCurrentAndInitialTransform Pose;
+};
+
+USTRUCT()
+struct CONTROLRIG_API FRigHierarchyCopyPasteContent
+{
+	GENERATED_USTRUCT_BODY()
+
+    UPROPERTY()
+	TArray<FRigHierarchyCopyPasteContentPerElement> Elements;
+};
+
+USTRUCT(BlueprintType)
+struct CONTROLRIG_API FRigBaseElement
+{
+	GENERATED_BODY()
+
+public:
+
+	FRigBaseElement()
+    : Key()
+    , Index(INDEX_NONE)
+	, SubIndex(INDEX_NONE)
+	, bSelected(false)
+	, TopologyVersion(0)
+	{}
+
+	virtual ~FRigBaseElement(){}
+
+	enum ESerializationPhase
+	{
+		StaticData,
+		InterElementData
+	};
+
+protected:
+	
+	UPROPERTY(BlueprintReadOnly, VisibleAnywhere, Category = RigElement, meta = (AllowPrivateAccess = "true"))
+	FRigElementKey Key;
+
+	UPROPERTY(BlueprintReadOnly, VisibleAnywhere, Category = RigElement, meta = (AllowPrivateAccess = "true"))
+	int32 Index;
+
+	UPROPERTY(BlueprintReadOnly, VisibleAnywhere, Category = RigElement, meta = (AllowPrivateAccess = "true"))
+	int32 SubIndex;
+
+	UPROPERTY(BlueprintReadOnly, Transient, Category = RigElement, meta = (AllowPrivateAccess = "true"))
+	bool bSelected;
+
+	FORCEINLINE static bool IsClassOf(const FRigBaseElement* InElement)
+	{
+		return true;
+	}
+
+public:
+
+	void Serialize(FArchive& Ar, URigHierarchy* Hierarchy, ESerializationPhase SerializationPhase);
+	virtual void Save(FArchive& Ar, URigHierarchy* Hierarchy, ESerializationPhase SerializationPhase);
+	virtual void Load(FArchive& Ar, URigHierarchy* Hierarchy, ESerializationPhase SerializationPhase);
+
+	FORCEINLINE const FName& GetName() const { return Key.Name; }
+	FORCEINLINE virtual const FName& GetDisplayName() const { return GetName(); }
+	FORCEINLINE ERigElementType GetType() const { return Key.Type; }
+	FORCEINLINE const FRigElementKey& GetKey() const { return Key; }
+	FORCEINLINE int32 GetIndex() const { return Index; }
+	FORCEINLINE int32 GetSubIndex() const { return SubIndex; }
+	FORCEINLINE bool IsSelected() const { return bSelected; }
+
+	template<typename T>
+	FORCEINLINE bool IsA() const { return T::IsClassOf(this); }
+
+	FORCEINLINE bool IsTypeOf(ERigElementType InElementType) const
+	{
+		return ((uint8)InElementType & (uint8)Key.Type) == (uint8)Key.Type;
+	}
+
+	template<typename T>
+    friend FORCEINLINE const T* Cast(const FRigBaseElement* InElement)
+	{
+		if(InElement)
+		{
+			if(InElement->IsA<T>())
+			{
+				return static_cast<const T*>(InElement);
+			}
+		}
+		return nullptr;
+	}
+
+	template<typename T>
+    friend FORCEINLINE T* Cast(FRigBaseElement* InElement)
+	{
+		if(InElement)
+		{
+			if(InElement->IsA<T>())
+			{
+				return static_cast<T*>(InElement);
+			}
+		}
+		return nullptr;
+	}
+
+	template<typename T>
+    friend FORCEINLINE const T* CastChecked(const FRigBaseElement* InElement)
+	{
+		const T* Element = Cast<T>(InElement);
+		check(Element);
+		return Element;
+	}
+
+	template<typename T>
+    friend FORCEINLINE T* CastChecked(FRigBaseElement* InElement)
+	{
+		T* Element = Cast<T>(InElement);
+		check(Element);
+		return Element;
+	}
+
+	virtual void CopyPose(FRigBaseElement* InOther, bool bCurrent, bool bInitial) {}
+
+protected:
+
+	mutable uint16 TopologyVersion;
+	mutable TArray<FRigBaseElement*> CachedChildren;
+
+	friend class URigHierarchy;
+	friend class URigHierarchyController;
+};
+
+USTRUCT(BlueprintType)
+struct CONTROLRIG_API FRigTransformElement : public FRigBaseElement
+{
+public:
+	
+	GENERATED_BODY()
+	DECLARE_RIG_ELEMENT_METHODS(FRigTransformElement)
+
+	virtual ~FRigTransformElement(){}
+
+	virtual void Save(FArchive& A, URigHierarchy* Hierarchy, ESerializationPhase SerializationPhase) override;
+	virtual void Load(FArchive& Ar, URigHierarchy* Hierarchy, ESerializationPhase SerializationPhase) override;
+
+	UPROPERTY(BlueprintReadOnly, EditAnywhere, Category = RigElement, meta = (DisplayAfter = "Index", ShowOnlyInnerProperties))
+	FRigCurrentAndInitialTransform Pose;
+
+protected:
+
+	struct FElementToDirty
+	{
+		FORCEINLINE FElementToDirty()
+			: Element(nullptr)
+			, HierarchyDistance(INDEX_NONE)
+		{}
+
+		FORCEINLINE FElementToDirty(FRigTransformElement* InElement, int32 InHierarchyDistance = INDEX_NONE)
+			: Element(InElement)
+			, HierarchyDistance(InHierarchyDistance)
+		{}
+
+		FORCEINLINE bool operator ==(const FElementToDirty& Other) const
+		{
+			return Element == Other.Element;
+		}
+
+		FORCEINLINE bool operator !=(const FElementToDirty& Other) const
+		{
+			return Element != Other.Element;
+		}
+		
+		FRigTransformElement* Element;
+		int32 HierarchyDistance;
+	};
+	
+	TArray<FElementToDirty> ElementsToDirty;
+
+	FORCEINLINE static bool IsClassOf(const FRigBaseElement* InElement)
+	{
+		return InElement->GetType() == ERigElementType::Bone ||
+			InElement->GetType() == ERigElementType::Space ||
+			InElement->GetType() == ERigElementType::Control;
+	}
+
+public:
+	
+	virtual void CopyPose(FRigBaseElement* InOther, bool bCurrent, bool bInitial) override;
+
+protected:
+	
+	friend struct FRigBaseElement;
+	friend struct FRigSingleParentElement;
+	friend struct FRigMultiParentElement;
+	friend class URigHierarchy;
+	friend class URigHierarchyController;
+};
+
+USTRUCT(BlueprintType)
+struct CONTROLRIG_API FRigSingleParentElement : public FRigTransformElement
+{
+public:
+	
+	GENERATED_BODY()
+	DECLARE_RIG_ELEMENT_METHODS(FRigSingleParentElement)
+
+	FRigSingleParentElement()
+	: FRigTransformElement()
+	, ParentElement(nullptr)
+	{}
+
+	virtual ~FRigSingleParentElement(){}
+
+	virtual void Save(FArchive& A, URigHierarchy* Hierarchy, ESerializationPhase SerializationPhase) override;
+	virtual void Load(FArchive& Ar, URigHierarchy* Hierarchy, ESerializationPhase SerializationPhase) override;
+
+	FRigTransformElement* ParentElement;
+
+protected:
+	
+	FORCEINLINE static bool IsClassOf(const FRigBaseElement* InElement)
+	{
+		return InElement->GetType() == ERigElementType::Bone;
+	}
+
+	friend struct FRigBaseElement;
+};
+
+USTRUCT(BlueprintType)
+struct CONTROLRIG_API FRigMultiParentElement : public FRigTransformElement
+{
+public:
+	
+	GENERATED_BODY()
+	DECLARE_RIG_ELEMENT_METHODS(FRigMultiParentElement)
+
+    FRigMultiParentElement()
+    : FRigTransformElement()
+	{}
+
+	virtual ~FRigMultiParentElement(){}
+
+	virtual void Save(FArchive& A, URigHierarchy* Hierarchy, ESerializationPhase SerializationPhase) override;
+	virtual void Load(FArchive& Ar, URigHierarchy* Hierarchy, ESerializationPhase SerializationPhase) override;
+
+	UPROPERTY(BlueprintReadOnly, Category = RigElement)
+	FRigCurrentAndInitialTransform Parent;
+
+	TArray<FRigTransformElement*> ParentElements;
+	TArray<float> ParentWeights;
+	TArray<float> ParentWeightsInitial;
+	TMap<FRigElementKey, int32> IndexLookup;
+
+protected:
+
+	FORCEINLINE static bool IsClassOf(const FRigBaseElement* InElement)
+	{
+		return InElement->GetType() == ERigElementType::Space ||
+			InElement->GetType() == ERigElementType::Control;
+	}
+
+public:
+	
+	virtual void CopyPose(FRigBaseElement* InOther, bool bCurrent, bool bInitial) override;
+
+protected:
+	
+	friend struct FRigBaseElement;
+};
+
+
+USTRUCT(BlueprintType)
+struct CONTROLRIG_API FRigBoneElement : public FRigSingleParentElement
+{
+public:
+	
+	GENERATED_BODY()
+	DECLARE_RIG_ELEMENT_METHODS(FRigBoneElement)
+
+	FRigBoneElement()
+		: FRigSingleParentElement()
+	{
+		Key.Type = ERigElementType::Bone;
+		BoneType = ERigBoneType::User;
+	}
+	
+	virtual ~FRigBoneElement(){}
+
+	virtual void Save(FArchive& A, URigHierarchy* Hierarchy, ESerializationPhase SerializationPhase) override;
+	virtual void Load(FArchive& Ar, URigHierarchy* Hierarchy, ESerializationPhase SerializationPhase) override;
+
+	UPROPERTY(BlueprintReadOnly, VisibleAnywhere, Category = RigElement)
+	ERigBoneType BoneType;
+
+protected:
+	
+	FORCEINLINE static bool IsClassOf(const FRigBaseElement* InElement)
+	{
+		return InElement->GetType() == ERigElementType::Bone;
+	}
+
+	friend struct FRigBaseElement;
+};
+
+USTRUCT(BlueprintType)
+struct CONTROLRIG_API FRigSpaceElement : public FRigMultiParentElement
+{
+public:
+	
+	GENERATED_BODY()
+	DECLARE_RIG_ELEMENT_METHODS(FRigSpaceElement)
+
+	FRigSpaceElement()
+    : FRigMultiParentElement()
+	{
+		Key.Type = ERigElementType::Space; 
+	}
+
+	virtual ~FRigSpaceElement(){}
+
+protected:
+	
+	FORCEINLINE static bool IsClassOf(const FRigBaseElement* InElement)
+	{
+		return InElement->GetType() == ERigElementType::Space;
+	}
+
+	friend struct FRigBaseElement;
+};
+
+USTRUCT(BlueprintType)
+struct CONTROLRIG_API FRigControlSettings
+{
+	GENERATED_BODY()
+
+	FRigControlSettings();
+
+	void Save(FArchive& Ar);
+	void Load(FArchive& Ar);
+
+	UPROPERTY(BlueprintReadWrite, EditAnywhere, Category = Control)
+	ERigControlType ControlType;
+
+	UPROPERTY(BlueprintReadWrite, EditAnywhere, Category = Control)
+	FName DisplayName;
+
+	/** the primary axis to use for float controls */
+	UPROPERTY(BlueprintReadWrite, EditAnywhere, Category = Control)
+	ERigControlAxis PrimaryAxis;
+
+	/** If Created from a Curve  Container*/
+	UPROPERTY(transient)
+	bool bIsCurve;
+
+	/** If the control is animatable in sequencer */
+	UPROPERTY(BlueprintReadWrite, EditAnywhere, Category = Control)
+	bool bAnimatable;
+
+	/** True if the control has to obey translation limits. */
+	UPROPERTY(BlueprintReadWrite, EditAnywhere, Category = Limits)
+	bool bLimitTranslation;
+
+	/** True if the control has to obey rotation limits. */
+	UPROPERTY(BlueprintReadWrite, EditAnywhere, Category = Limits)
+	bool bLimitRotation;
+
+	/** True if the control has to obey scale limits. */
+	UPROPERTY(BlueprintReadWrite, EditAnywhere, Category = Limits)
+	bool bLimitScale;
+
+	/** True if the limits should be drawn in debug. */
+	UPROPERTY(BlueprintReadWrite, EditAnywhere, Category = Limits, meta = (EditCondition = "bLimitTranslation || bLimitRotation || bLimitScale"))
+	bool bDrawLimits;
+
+	/** The minimum limit of the control's value */
+	UPROPERTY(BlueprintReadWrite, Category = Limits, meta = (EditCondition = "bLimitTranslation || bLimitRotation || bLimitScale"))
+	FRigControlValue MinimumValue;
+
+	/** The maximum limit of the control's value */
+	UPROPERTY(BlueprintReadWrite, Category = Limits, meta = (EditCondition = "bLimitTranslation || bLimitRotation || bLimitScale"))
+	FRigControlValue MaximumValue;
+
+	/** Set to true if the gizmo is enabled in 3d */
+	UPROPERTY(BlueprintReadWrite, EditAnywhere, Category = Gizmo)
+	bool bGizmoEnabled;
+
+	/** Set to true if the gizmo is currently visible in 3d */
+	UPROPERTY(BlueprintReadWrite, EditAnywhere, Category = Gizmo, meta = (EditCondition = "bGizmoEnabled"))
+	bool bGizmoVisible;
+
+	/* This is optional UI setting - this doesn't mean this is always used, but it is optional for manipulation layer to use this*/
+	UPROPERTY(BlueprintReadWrite, EditAnywhere, Category = Gizmo, meta = (EditCondition = "bGizmoEnabled"))
+	FName GizmoName;
+
+	UPROPERTY(BlueprintReadWrite, EditAnywhere, Category = Gizmo, meta = (EditCondition = "bGizmoEnabled"))
+	FLinearColor GizmoColor;
+
+	/** If the control is transient and only visible in the control rig editor */
+	UPROPERTY(BlueprintReadWrite, Category = Control)
+	bool bIsTransientControl;
+
+	/** If the control is 4transient and only visible in the control rig editor */
+	UPROPERTY(BlueprintReadOnly, EditAnywhere, Category = Control)
+	UEnum* ControlEnum;
+
+	/** Applies the limits expressed by these settings to a value */
+	FORCEINLINE void ApplyLimits(FRigControlValue& InOutValue) const
+	{
+		InOutValue.ApplyLimits(bLimitTranslation, bLimitRotation, bLimitScale, ControlType, MinimumValue, MaximumValue);
+	}
+
+	/** Applies the limits expressed by these settings to a transform */
+	FORCEINLINE void ApplyLimits(FTransform& InOutValue) const
+	{
+		FRigControlValue Value;
+		Value.SetFromTransform(InOutValue, ControlType, PrimaryAxis);
+		ApplyLimits(Value);
+		InOutValue = Value.GetAsTransform(ControlType, PrimaryAxis);
+	}
+
+	FORCEINLINE bool IsValueTypeEnabled(ERigControlValueType InValueType) const
+	{
+		if(InValueType == ERigControlValueType::Minimum || InValueType == ERigControlValueType::Maximum)
+		{
+			return bLimitTranslation || bLimitRotation || bLimitScale;
+		}
+		return true;
+	}
+
+	FORCEINLINE FRigControlValue GetIdentityValue() const
+	{
+		FRigControlValue Value;
+		Value.SetFromTransform(FTransform::Identity, ControlType, PrimaryAxis);
+		return Value;
+	}
+};
+
+USTRUCT(BlueprintType)
+struct CONTROLRIG_API FRigControlElement : public FRigMultiParentElement
+{
+	public:
+	
+	GENERATED_BODY()
+	DECLARE_RIG_ELEMENT_METHODS(FRigControlElement)
+
+	FRigControlElement()
+		: FRigMultiParentElement()
+	{
+		Key.Type = ERigElementType::Control; 
+	}
+
+	virtual ~FRigControlElement(){}
+	
+	FORCEINLINE virtual const FName& GetDisplayName() const override
+	{
+		if(!Settings.DisplayName.IsNone())
+		{
+			return Settings.DisplayName;
+		}
+		return FRigMultiParentElement::GetDisplayName();
+	}
+
+	virtual void Save(FArchive& A, URigHierarchy* Hierarchy, ESerializationPhase SerializationPhase) override;
+	virtual void Load(FArchive& Ar, URigHierarchy* Hierarchy, ESerializationPhase SerializationPhase) override;
+
+	UPROPERTY(BlueprintReadOnly, EditAnywhere, Category = Control, meta=(ShowOnlyInnerProperties))
+	FRigControlSettings Settings;
+
+	UPROPERTY(BlueprintReadOnly, EditAnywhere, Category = RigElement)
+	FRigCurrentAndInitialTransform Offset;
+
+	UPROPERTY(BlueprintReadOnly, EditAnywhere, Category = RigElement)
+	FRigCurrentAndInitialTransform Gizmo;
+
+protected:
+	
+	FORCEINLINE static bool IsClassOf(const FRigBaseElement* InElement)
+	{
+		return InElement->GetType() == ERigElementType::Control;
+	}
+
+public:
+	
+	virtual void CopyPose(FRigBaseElement* InOther, bool bCurrent, bool bInitial) override;
+
+protected:
+
+	friend struct FRigBaseElement;
+};
+
+USTRUCT(BlueprintType)
+struct CONTROLRIG_API FRigCurveElement : public FRigBaseElement
+{
+public:
+	
+	GENERATED_BODY()
+	DECLARE_RIG_ELEMENT_METHODS(FRigCurveElement)
+
+	FRigCurveElement()
+		: FRigBaseElement()
+		, Value(0.f)
+	{
+		Key.Type = ERigElementType::Curve;
+	}
+
+	virtual ~FRigCurveElement(){}
+
+	virtual void Save(FArchive& A, URigHierarchy* Hierarchy, ESerializationPhase SerializationPhase) override;
+	virtual void Load(FArchive& Ar, URigHierarchy* Hierarchy, ESerializationPhase SerializationPhase) override;
+
+	float Value;
+	
+	FORCEINLINE static bool IsClassOf(const FRigBaseElement* InElement)
+	{
+		return InElement->GetType() == ERigElementType::Curve;
+	}
+
+public:
+	
+	virtual void CopyPose(FRigBaseElement* InOther, bool bCurrent, bool bInitial) override;
+
+protected:
+	
+	friend struct FRigBaseElement;
+};
\ No newline at end of file
diff --git a/Engine/Plugins/Experimental/ControlRig/Source/ControlRig/Public/Rigs/RigSpaceHierarchy.h b/Engine/Plugins/Experimental/ControlRig/Source/ControlRig/Public/Rigs/RigSpaceHierarchy.h
--- a/Engine/Plugins/Experimental/ControlRig/Source/ControlRig/Public/Rigs/RigSpaceHierarchy.h
+++ b/Engine/Plugins/Experimental/ControlRig/Source/ControlRig/Public/Rigs/RigSpaceHierarchy.h
@@ -8,6 +8,7 @@
 #include "RigSpaceHierarchy.generated.h"
 
 class UControlRig;
+struct FRigHierarchyContainer;
 
 UENUM(BlueprintType)
 enum class ERigSpaceType : uint8
@@ -25,7 +26,7 @@ enum class ERigSpaceType : uint8
 	Space
 };
 
-USTRUCT(BlueprintType)
+USTRUCT()
 struct CONTROLRIG_API FRigSpace : public FRigElement
 {
 	GENERATED_BODY()
@@ -89,157 +90,26 @@ struct CONTROLRIG_API FRigSpace : public FRigElement
 	}
 };
 
-USTRUCT(BlueprintType)
+USTRUCT()
 struct CONTROLRIG_API FRigSpaceHierarchy
 {
 	GENERATED_BODY()
 
 	FRigSpaceHierarchy();
-	FRigSpaceHierarchy& operator= (const FRigSpaceHierarchy &InOther);
-
-	FORCEINLINE ERigElementType RigElementType() const { return ERigElementType::Space; }
-
-	FORCEINLINE int32 Num() const { return Spaces.Num(); }
-	FORCEINLINE const FRigSpace& operator[](int32 InIndex) const { return Spaces[InIndex]; }
-	FORCEINLINE FRigSpace& operator[](int32 InIndex) { return Spaces[InIndex]; }
-	FORCEINLINE const FRigSpace& operator[](const FName& InName) const { return Spaces[GetIndex(InName)]; }
-	FORCEINLINE FRigSpace& operator[](const FName& InName) { return Spaces[GetIndex(InName)]; }
-
-	FORCEINLINE const TArray<FRigSpace>& GetSpaces() const { return Spaces; }
 
 	FORCEINLINE TArray<FRigSpace>::RangedForIteratorType      begin()       { return Spaces.begin(); }
 	FORCEINLINE TArray<FRigSpace>::RangedForConstIteratorType begin() const { return Spaces.begin(); }
 	FORCEINLINE TArray<FRigSpace>::RangedForIteratorType      end()         { return Spaces.end();   }
 	FORCEINLINE TArray<FRigSpace>::RangedForConstIteratorType end() const   { return Spaces.end();   }
 
-	bool IsNameAvailable(const FString& InPotentialNewName, FString* OutErrorMessage = nullptr) const;
-
-	FName GetSafeNewName(const FString& InPotentialNewName) const;
-
-	FRigSpace& Add(const FName& InNewName, ERigSpaceType InSpaceType = ERigSpaceType::Global, const FName& InParentName = NAME_None, const FTransform& InTransform = FTransform::Identity);
-
-	FRigSpace Remove(const FName& InNameToRemove);
-
-	FName Rename(const FName& InOldName, const FName& InNewName);
-
-	bool Reparent(const FName& InName, ERigSpaceType InSpaceType, const FName& InNewParentName);
-
-	FName GetName(int32 InIndex) const;
-
-	FORCEINLINE int32 GetIndex(const FName& InName) const
-	{
-		if(NameToIndexMapping.Num() != Spaces.Num())
-		{
-			return GetIndexSlow(InName);
-		}
-
-		const int32* Index = NameToIndexMapping.Find(InName);
-		if (Index)
-		{
-			return *Index;
-		}
-
-		return INDEX_NONE;
-	}
-
-	void SetGlobalTransform(const FName& InName, const FTransform& InTransform);
-
-	void SetGlobalTransform(int32 InIndex, const FTransform& InTransform);
-
-	FTransform GetGlobalTransform(const FName& InName) const;
-
-	FTransform GetGlobalTransform(int32 InIndex) const;
-
-	void SetLocalTransform(const FName& InName, const FTransform& InTransform);
-
-	void SetLocalTransform(int32 InIndex, const FTransform& InTransform);
-
-	FTransform GetLocalTransform(const FName& InName) const;
-
-	FTransform GetLocalTransform(int32 InIndex) const;
-
-	void SetInitialTransform(const FName& InName, const FTransform& InTransform);
-
-	void SetInitialTransform(int32 InIndex, const FTransform& InTransform);
-
-	FTransform GetInitialTransform(const FName& InName) const;
-
-	FTransform GetInitialTransform(int32 InIndex) const;
-
-	void SetInitialGlobalTransform(const FName& InName, const FTransform& GlobalTransform);
-	
-	void SetInitialGlobalTransform(int32 InIndex, const FTransform& GlobalTransform);
-	
-	FTransform GetInitialGlobalTransform(const FName& InName) const;
-	
-	FTransform GetInitialGlobalTransform(int32 InIndex) const;
-
-	// updates all of the internal caches
-	void Initialize(bool bResetTransforms = true);
-
-	// clears the hierarchy and removes all content
-	void Reset();
-
-	// returns the current pose
-	FRigPose GetPose() const;
-
-	// sets the current transforms from the given pose
-	void SetPose(FRigPose& InPose);
-
-	// resets all of the transforms back to the initial transform
-	void ResetTransforms();
-
-	// copies all initial transforms from another hierarchy
-	void CopyInitialTransforms(const FRigSpaceHierarchy& InOther);
-
-
-	bool Select(const FName& InName, bool bSelect = true);
-	bool ClearSelection();
-	TArray<FName> CurrentSelection() const;
-	bool IsSelected(const FName& InName) const;
-
-	FRigElementSelected OnSpaceSelected;
-
-#if WITH_EDITOR
-	FRigElementAdded OnSpaceAdded;
-	FRigElementRemoved OnSpaceRemoved;
-	FRigElementRenamed OnSpaceRenamed;
-	FRigElementReparented OnSpaceReparented;
-
-	void HandleOnElementRemoved(FRigHierarchyContainer* InContainer, const FRigElementKey& InKey);
-	void HandleOnElementRenamed(FRigHierarchyContainer* InContainer, ERigElementType InElementType, const FName& InOldName, const FName& InNewName);
-
-#endif
-
 private:
 
 	// disable copy constructor
 	FRigSpaceHierarchy(const FRigSpaceHierarchy& InOther) {}
 
-	FRigHierarchyContainer* Container;
-
-	int32 GetParentIndex(ERigSpaceType InSpaceType, const FName& InName) const;
-
 	UPROPERTY(EditAnywhere, Category = FRigSpaceHierarchy)
 	TArray<FRigSpace> Spaces;
 
-	// can serialize fine? 
-	UPROPERTY()
-	TMap<FName, int32> NameToIndexMapping;
-
-	UPROPERTY(transient)
-	TArray<FName> Selection;
-
-	int32 GetIndexSlow(const FName& InName) const;
-
-	void RefreshMapping();
-
-	void AppendToPose(FRigPose& InOutPose) const;
-
-#if WITH_EDITOR
-	mutable TArray<bool> RecursionGuard;
-#endif
-
 	friend struct FRigHierarchyContainer;
 	friend struct FCachedRigElement;
 	friend class UControlRigHierarchyModifier;
diff --git a/Engine/Plugins/Experimental/ControlRig/Source/ControlRig/Public/Tools/ControlRigPose.h b/Engine/Plugins/Experimental/ControlRig/Source/ControlRig/Public/Tools/ControlRigPose.h
--- a/Engine/Plugins/Experimental/ControlRig/Source/ControlRig/Public/Tools/ControlRigPose.h
+++ b/Engine/Plugins/Experimental/ControlRig/Source/ControlRig/Public/Tools/ControlRigPose.h
@@ -26,8 +26,7 @@ struct CONTROLRIG_API FRigControlCopy
 		: Name(NAME_None)
 		, ControlType(ERigControlType::Transform)
 		, Value()
-		, ParentName(NAME_None)
-		, SpaceName(NAME_None)
+		, ParentKey()
 		, OffsetTransform(FTransform::Identity)
 		, ParentTransform(FTransform::Identity)
 		, LocalTransform(FTransform::Identity)
@@ -36,19 +35,17 @@ struct CONTROLRIG_API FRigControlCopy
 	{
 	}
 
-	FRigControlCopy(const FRigControl& InControl, const FRigControlHierarchy& Hierarchy)
+	FRigControlCopy(FRigControlElement* InControlElement, URigHierarchy* InHierarchy)
 	{
-		Name = InControl.Name;
-		ControlType = InControl.ControlType;
-		Value = InControl.Value;
-		ParentName = InControl.ParentName;
-		SpaceName = InControl.SpaceName;
-		OffsetTransform = InControl.OffsetTransform;
-
-		int32 Index = Hierarchy.GetIndex(Name);
-		ParentTransform = Index != INDEX_NONE ? Hierarchy.GetParentTransform(Index) : FTransform::Identity;
-		LocalTransform = Hierarchy.GetLocalTransform(Name);
-		GlobalTransform = Hierarchy.GetGlobalTransform(Name);
+		Name = InControlElement->GetName();
+		ControlType = InControlElement->Settings.ControlType;
+		Value = InHierarchy->GetControlValue(InControlElement, ERigControlValueType::Current);
+		ParentKey = InHierarchy->GetFirstParent(InControlElement->GetKey());
+		OffsetTransform = InHierarchy->GetControlOffsetTransform(InControlElement, ERigTransformType::CurrentLocal);
+
+		ParentTransform = InHierarchy->GetParentTransform(InControlElement, ERigTransformType::CurrentGlobal);
+		LocalTransform = InHierarchy->GetTransform(InControlElement, ERigTransformType::CurrentLocal);
+		GlobalTransform = InHierarchy->GetTransform(InControlElement, ERigTransformType::CurrentGlobal);
 	}
 	virtual ~FRigControlCopy() {}
 
@@ -62,10 +59,7 @@ struct CONTROLRIG_API FRigControlCopy
 	FRigControlValue Value;
 
 	UPROPERTY(BlueprintReadWrite, EditAnywhere, Category = "Names")
-	FName ParentName;
-
-	UPROPERTY(BlueprintReadWrite, EditAnywhere, Category = "Names")
-	FName SpaceName;
+	FRigElementKey ParentKey;
 
 	UPROPERTY(BlueprintReadWrite, EditAnywhere, Category = "Transforms")
 	FTransform OffsetTransform;
diff --git a/Engine/Plugins/Experimental/ControlRig/Source/ControlRig/Public/Units/Hierarchy/RigUnit_SetTransform.h b/Engine/Plugins/Experimental/ControlRig/Source/ControlRig/Public/Units/Hierarchy/RigUnit_SetTransform.h
--- a/Engine/Plugins/Experimental/ControlRig/Source/ControlRig/Public/Units/Hierarchy/RigUnit_SetTransform.h
+++ b/Engine/Plugins/Experimental/ControlRig/Source/ControlRig/Public/Units/Hierarchy/RigUnit_SetTransform.h
@@ -29,9 +29,9 @@ struct CONTROLRIG_API FRigUnit_SetTransform : public FRigUnitMutable
 	{
 		if(Space == EBoneGetterSetterMode::LocalSpace)
 		{
-			if (const FRigHierarchyContainer* Container = (const FRigHierarchyContainer*)InUserContext)
+			if (const URigHierarchy* Hierarchy = (const URigHierarchy*)InUserContext)
 			{
-				return Container->GetParentKey(Item);
+				return Hierarchy->GetFirstParent(Item);
 			}
 		}
 		return FRigElementKey();
@@ -101,9 +101,9 @@ struct CONTROLRIG_API FRigUnit_SetTranslation : public FRigUnitMutable
 	{
 		if(Space == EBoneGetterSetterMode::LocalSpace)
 		{
-			if (const FRigHierarchyContainer* Container = (const FRigHierarchyContainer*)InUserContext)
+			if (const URigHierarchy* Hierarchy = (const URigHierarchy*)InUserContext)
 			{
-				return Container->GetParentKey(Item);
+				return Hierarchy->GetFirstParent(Item);
 			}
 		}
 		return FRigElementKey();
@@ -165,9 +165,9 @@ struct CONTROLRIG_API FRigUnit_SetRotation : public FRigUnitMutable
 	{
 		if(Space == EBoneGetterSetterMode::LocalSpace)
 		{
-			if (const FRigHierarchyContainer* Container = (const FRigHierarchyContainer*)InUserContext)
+			if (const URigHierarchy* Hierarchy = (const URigHierarchy*)InUserContext)
 			{
-				return Container->GetParentKey(Item);
+				return Hierarchy->GetFirstParent(Item);
 			}
 		}
 		return FRigElementKey();
diff --git a/Engine/Plugins/Experimental/ControlRig/Source/ControlRig/Public/Units/RigUnitContext.h b/Engine/Plugins/Experimental/ControlRig/Source/ControlRig/Public/Units/RigUnitContext.h
--- a/Engine/Plugins/Experimental/ControlRig/Source/ControlRig/Public/Units/RigUnitContext.h
+++ b/Engine/Plugins/Experimental/ControlRig/Source/ControlRig/Public/Units/RigUnitContext.h
@@ -69,7 +69,7 @@ struct FRigUnitContext
 	EControlRigState State;
 
 	/** The current hierarchy being executed */
-	const FRigHierarchyContainer* Hierarchy;
+	URigHierarchy* Hierarchy;
 
 	/** True if the rig is executing during an interaction */
 	bool bDuringInteraction;
@@ -94,42 +94,6 @@ struct FRigUnitContext
 	/** A container to store all names */
 	FRigNameCache* NameCache;
 
-	FORCEINLINE const FRigBoneHierarchy* GetBones() const
-	{
-		if (Hierarchy != nullptr)
-		{
-			return &Hierarchy->BoneHierarchy;
-		}
-		return nullptr;
-	}
-
-	FORCEINLINE const FRigSpaceHierarchy* GetSpaces() const
-	{
-		if (Hierarchy != nullptr)
-		{
-			return &Hierarchy->SpaceHierarchy;
-		}
-		return nullptr;
-	}
-
-	FORCEINLINE const FRigControlHierarchy* GetControls() const
-	{
-		if (Hierarchy != nullptr)
-		{
-			return &Hierarchy->ControlHierarchy;
-		}
-		return nullptr;
-	}
-
-	FORCEINLINE const FRigCurveContainer* GetCurves() const
-	{
-		if (Hierarchy != nullptr)
-		{
-			return &Hierarchy->CurveContainer;
-		}
-		return nullptr;
-	}
-
 	/**
 	 * Returns a given data source and cast it to the expected class.
 	 *
diff --git a/Engine/Plugins/Experimental/ControlRig/Source/ControlRig/Public/Units/RigUnitTest.h b/Engine/Plugins/Experimental/ControlRig/Source/ControlRig/Public/Units/RigUnitTest.h
--- a/Engine/Plugins/Experimental/ControlRig/Source/ControlRig/Public/Units/RigUnitTest.h
+++ b/Engine/Plugins/Experimental/ControlRig/Source/ControlRig/Public/Units/RigUnitTest.h
@@ -5,30 +5,37 @@
 #include "Misc/AutomationTest.h"
 #include "RigUnit.h"
 #include "RigUnitContext.h"
-#include "Rigs/RigHierarchyContainer.h"
-#include "Rigs/RigCurveContainer.h"
+#include "Rigs/RigHierarchy.h"
+#include "Rigs/RigHierarchyController.h"
 
 class FControlRigUnitTestBase : public FAutomationTestBase
 {
 public:
 	FControlRigUnitTestBase(const FString& InName, bool bIsComplex)
 		: FAutomationTestBase(InName, bIsComplex)
-		, HierarchyContainer()
-		, BoneHierarchy(HierarchyContainer.BoneHierarchy)
-		, SpaceHierarchy(HierarchyContainer.SpaceHierarchy)
-		, ControlHierarchy(HierarchyContainer.ControlHierarchy)
-		, CurveContainer(HierarchyContainer.CurveContainer)
+		, Hierarchy(nullptr)
+		, Controller(nullptr)
 	{
-		Context.Hierarchy = &HierarchyContainer;
-		ExecuteContext.Hierarchy= &HierarchyContainer;
+		Hierarchy = NewObject<URigHierarchy>();
+		Controller = NewObject<URigHierarchyController>();
+		Controller->SetHierarchy(Hierarchy);
+
+		Hierarchy->AddToRoot();
+		Controller->AddToRoot();
+		
+		Context.Hierarchy = Hierarchy;;
+		ExecuteContext.Hierarchy= Hierarchy;
 		Context.NameCache = &NameCache;
 	}
 
-	FRigHierarchyContainer HierarchyContainer;
-	FRigBoneHierarchy& BoneHierarchy;
-	FRigSpaceHierarchy& SpaceHierarchy;
-	FRigControlHierarchy& ControlHierarchy;
-	FRigCurveContainer& CurveContainer;
+	~FControlRigUnitTestBase()
+	{
+		Controller->RemoveFromRoot();
+		Hierarchy->RemoveFromRoot();
+	}
+
+	URigHierarchy* Hierarchy;
+	URigHierarchyController* Controller;
 	FControlRigExecuteContext ExecuteContext;
 	FRigNameCache NameCache;
 	FRigUnitContext Context;
@@ -56,8 +63,7 @@ public:
 		TUnitStruct Unit; \
 		virtual bool RunTest(const FString& Parameters) override \
 		{ \
-			HierarchyContainer.Reset(); \
-			CurveContainer.Reset(); \
+			Hierarchy->Reset(); \
 			Unit = TUnitStruct(); \
 			return RunControlRigUnitTest(Parameters); \
 		} \
diff --git a/Engine/Plugins/Experimental/ControlRig/Source/ControlRigDeveloper/Private/AnimGraphNode_ControlRig.cpp b/Engine/Plugins/Experimental/ControlRig/Source/ControlRigDeveloper/Private/AnimGraphNode_ControlRig.cpp
--- a/Engine/Plugins/Experimental/ControlRig/Source/ControlRigDeveloper/Private/AnimGraphNode_ControlRig.cpp
+++ b/Engine/Plugins/Experimental/ControlRig/Source/ControlRigDeveloper/Private/AnimGraphNode_ControlRig.cpp
@@ -149,23 +149,24 @@ void UAnimGraphNode_ControlRig::ValidateAnimNodeDuringCompilation(USkeleton* For
 	{
 		if (UControlRigBlueprint* Blueprint = Cast<UControlRigBlueprint>(TargetClass->ClassGeneratedBy))
 		{
-			const FRigBoneHierarchy& BoneHierarchy = Blueprint->HierarchyContainer.BoneHierarchy;
+			URigHierarchy* Hierarchy = Blueprint->Hierarchy;
 			const FReferenceSkeleton& ReferenceSkeleton = ForSkeleton->GetReferenceSkeleton();
 			const TArray<FMeshBoneInfo>& BoneInfos = ReferenceSkeleton.GetRefBoneInfo();
 
 			for (const FMeshBoneInfo& BoneInfo : BoneInfos)
 			{
-				int32 BoneIndex = BoneHierarchy.GetIndex(BoneInfo.Name);
-				if (BoneIndex != INDEX_NONE)
+				const FRigElementKey BoneKey(BoneInfo.Name, ERigElementType::Bone);
+				if (Hierarchy->Contains(BoneKey))
 				{
+					const FRigElementKey ParentKey = Hierarchy->GetFirstParent(BoneKey);
+
 					FName DesiredParentName = NAME_None;
 					if (BoneInfo.ParentIndex != INDEX_NONE)
 					{
 						DesiredParentName = BoneInfos[BoneInfo.ParentIndex].Name;
 					}
 
-					const FRigBone& Bone = BoneHierarchy[BoneIndex];
-					if (DesiredParentName != Bone.ParentName)
+					if (DesiredParentName != ParentKey.Name)
 					{
 						FString Message = FString::Printf(TEXT("@@ - Hierarchy discrepancy for bone '%s' - different parents on Control Rig vs SkeletalMesh."), *BoneInfo.Name.ToString());
 						MessageLog.Warning(*Message, this);
diff --git a/Engine/Plugins/Experimental/ControlRig/Source/ControlRigDeveloper/Private/ControlRigBlueprint.cpp b/Engine/Plugins/Experimental/ControlRig/Source/ControlRigDeveloper/Private/ControlRigBlueprint.cpp
--- a/Engine/Plugins/Experimental/ControlRig/Source/ControlRigDeveloper/Private/ControlRigBlueprint.cpp
+++ b/Engine/Plugins/Experimental/ControlRig/Source/ControlRigDeveloper/Private/ControlRigBlueprint.cpp
@@ -75,6 +75,11 @@ UControlRigBlueprint::UControlRigBlueprint(const FObjectInitializer& ObjectIniti
 	bExposesAnimatableControls = false;
 
 	VMCompileSettings.ASTSettings.ReportDelegate.BindUObject(this, &UControlRigBlueprint::HandleReportFromCompiler);
+
+	Hierarchy = CreateDefaultSubobject<URigHierarchy>(TEXT("Hierarchy"));
+	HierarchyController = CreateDefaultSubobject<URigHierarchyController>(TEXT("HierarchyController"));
+	HierarchyController->SetHierarchy(Hierarchy);
+	HierarchyController->OnModified().AddUObject(this, &UControlRigBlueprint::HandleHierarchyModified);
 }
 
 void UControlRigBlueprint::InitializeModelIfRequired(bool bRecompileVM)
@@ -102,12 +107,6 @@ void UControlRigBlueprint::InitializeModelIfRequired(bool bRecompileVM)
 		}
 
 		FunctionLibraryEdGraph->Initialize(this);
-
-		HierarchyContainer.OnElementAdded.AddUObject(this, &UControlRigBlueprint::HandleOnElementAdded);
-		HierarchyContainer.OnElementRemoved.AddUObject(this, &UControlRigBlueprint::HandleOnElementRemoved);
-		HierarchyContainer.OnElementRenamed.AddUObject(this, &UControlRigBlueprint::HandleOnElementRenamed);
-		HierarchyContainer.OnElementReparented.AddUObject(this, &UControlRigBlueprint::HandleOnElementReparented);
-		HierarchyContainer.OnElementSelected.AddUObject(this, &UControlRigBlueprint::HandleOnElementSelected);
 	}
 }
 
@@ -158,6 +157,12 @@ bool UControlRigBlueprint::CanImportGraphFromText(const FString& InClipboardText
 	return GetTemplateController()->CanImportNodesFromText(InClipboardText);
 }
 
+void UControlRigBlueprint::PostEditChangeChainProperty(FPropertyChangedChainEvent& PropertyChangedEvent)
+{
+	Super::PostEditChangeChainProperty(PropertyChangedEvent);
+	PostEditChangeChainPropertyEvent.Broadcast(PropertyChangedEvent);
+}
+
 bool UControlRigBlueprint::TryImportGraphFromText(const FString& InClipboardText, UEdGraph** OutGraphPtr)
 {
 	if (OutGraphPtr)
@@ -229,14 +234,15 @@ void UControlRigBlueprint::PreSave(const class ITargetPlatform* TargetPlatform)
 	}
 
 	bExposesAnimatableControls = false;
-	for (FRigControl& RigControl : HierarchyContainer.ControlHierarchy)
-	{
-		if (RigControl.bAnimatable)
+	Hierarchy->ForEach<FRigControlElement>([this](FRigControlElement* ControlElement) -> bool
+    {
+		if (ControlElement->Settings.bAnimatable)
 		{
 			bExposesAnimatableControls = true;
-			break;
+			return false;
 		}
-	}
+		return true;
+	});
 }
 
 void UControlRigBlueprint::PostLoad()
@@ -246,19 +252,18 @@ void UControlRigBlueprint::PostLoad()
 	// temporarily disable default value validation during load time, serialized values should always be accepted
 	TGuardValue<bool> DisablePinDefaultValueValidation(GetOrCreateController()->bValidatePinDefaults, false);
 
-	HierarchyContainer.ControlHierarchy.PostLoad();
-
 	// correct the offset transforms
 	if (GetLinkerCustomVersion(FControlRigObjectVersion::GUID) < FControlRigObjectVersion::ControlOffsetTransform)
 	{
-		if (HierarchyContainer.ControlHierarchy.Num() > 0)
+		HierarchyContainer_DEPRECATED.ControlHierarchy.PostLoad();
+		if (HierarchyContainer_DEPRECATED.ControlHierarchy.Num() > 0)
 		{
 			bDirtyDuringLoad = true;
 		}
 
-		for (FRigControl& Control : HierarchyContainer.ControlHierarchy)
+		for (FRigControl& Control : HierarchyContainer_DEPRECATED.ControlHierarchy)
 		{
-			FTransform PreviousOffsetTransform = HierarchyContainer.ControlHierarchy.GetLocalTransform(Control.Index, ERigControlValueType::Initial);
+			const FTransform PreviousOffsetTransform = Control.GetTransformFromValue(ERigControlValueType::Initial);
 			Control.OffsetTransform = PreviousOffsetTransform;
 			Control.InitialValue = Control.Value;
 
@@ -275,7 +280,83 @@ void UControlRigBlueprint::PostLoad()
 				Control.InitialValue = FRigControlValue::Make<FEulerTransform>(FEulerTransform::Identity);
 			}
 		}
-		PropagateHierarchyFromBPToInstances(true, true);
+	}
+
+	// convert the hierarchy from V1 to V2
+	if (GetLinkerCustomVersion(FControlRigObjectVersion::GUID) < FControlRigObjectVersion::RigHierarchyV2)
+	{
+		Modify();
+		
+		Hierarchy->Reset();
+
+		for(const FRigBone& Bone : HierarchyContainer_DEPRECATED.BoneHierarchy)
+		{
+			HierarchyController->AddBone(Bone.Name, Bone.GetParentElementKey(true), Bone.InitialTransform, true, Bone.Type, false);
+		}
+		for(const FRigSpace& Space : HierarchyContainer_DEPRECATED.SpaceHierarchy)
+		{
+			HierarchyController->AddSpace(Space.Name, FRigElementKey(), Space.InitialTransform, false, false);
+		}
+		for(const FRigControl& Control : HierarchyContainer_DEPRECATED.ControlHierarchy)
+		{
+			FRigControlSettings Settings;
+			Settings.ControlType = Control.ControlType;
+			Settings.DisplayName = Control.DisplayName;
+			Settings.PrimaryAxis = Control.PrimaryAxis;
+			Settings.bIsCurve = Control.bIsCurve;
+			Settings.bAnimatable = Control.bAnimatable;
+			Settings.bLimitTranslation = Control.bLimitTranslation;
+			Settings.bLimitRotation = Control.bLimitRotation;
+			Settings.bLimitScale = Control.bLimitScale;
+			Settings.bDrawLimits = Control.bDrawLimits;
+			Settings.MinimumValue = Control.MinimumValue;
+			Settings.MaximumValue = Control.MaximumValue;
+			Settings.bGizmoEnabled = Control.bGizmoEnabled;
+			Settings.bGizmoVisible = Control.bGizmoVisible;
+			Settings.GizmoName = Control.GizmoName;
+			Settings.GizmoColor = Control.GizmoColor;
+			Settings.ControlEnum = Control.ControlEnum;
+
+			FRigControlValue InitialValue = Control.InitialValue;
+			if(!InitialValue.IsValid())
+			{
+				InitialValue.SetFromTransform(InitialValue.Storage_DEPRECATED, Settings.ControlType, Settings.PrimaryAxis);
+			}
+			
+			HierarchyController->AddControl(
+				Control.Name,
+				FRigElementKey(),
+				Settings,
+				InitialValue,
+				Control.OffsetTransform,
+				Control.GizmoTransform,
+				false);
+		}
+		
+		for(const FRigCurve& Curve : HierarchyContainer_DEPRECATED.CurveContainer)
+		{
+			HierarchyController->AddCurve(Curve.Name, Curve.Value, false);
+		}
+
+		for(const FRigSpace& Space : HierarchyContainer_DEPRECATED.SpaceHierarchy)
+		{
+			FRigElementKey ParentKey = Space.GetParentElementKey(true);
+			if(ParentKey.IsValid())
+			{
+				HierarchyController->SetParent(Space.GetElementKey(), ParentKey, false, false);
+			}
+		}
+
+		for(const FRigControl& Control : HierarchyContainer_DEPRECATED.ControlHierarchy)
+		{
+			FRigElementKey ParentKey = Control.GetParentElementKey(true);
+			const FRigElementKey SpaceKey = Control.GetSpaceElementKey(true);
+			ParentKey = SpaceKey.IsValid() ? SpaceKey : ParentKey;
+			if(ParentKey.IsValid())
+			{
+				HierarchyController->SetParent(Control.GetElementKey(), ParentKey, false, false);
+			}
+		}
 	}
 
 	// remove all non-controlrig-graphs
@@ -351,7 +432,7 @@ void UControlRigBlueprint::RecompileVM()
 		TGuardValue<bool> ReentrantGuardSelf(bSuspendModelNotificationsForSelf, true);
 		TGuardValue<bool> ReentrantGuardOthers(bSuspendModelNotificationsForOthers, true);
 
-		CDO->Hierarchy = HierarchyContainer;
+		CDO->GetHierarchy()->CopyHierarchy(Hierarchy);
 
 		if (CDO->VM->GetOuter() != CDO)
 		{
@@ -366,7 +447,7 @@ void UControlRigBlueprint::RecompileVM()
 
 		FRigUnitContext InitContext;
 		InitContext.State = EControlRigState::Init;
-		InitContext.Hierarchy = &CDO->Hierarchy;
+		InitContext.Hierarchy = CDO->DynamicHierarchy;
 
 		FRigUnitContext UpdateContext = InitContext;
 		InitContext.State = EControlRigState::Update;
@@ -396,7 +477,6 @@ void UControlRigBlueprint::RecompileVM()
 		{
 			if (UControlRig* InstanceRig = Cast<UControlRig>(Instance))
 			{
-				InstanceRig->Hierarchy.Initialize(false);
 				InstanceRig->InstantiateVMFromCDO();
 			}
 		}
@@ -858,18 +938,17 @@ void UControlRigBlueprint::PostTransacted(const FTransactionObjectEvent& Transac
 				return;
 			}
 
-			PropagateHierarchyFromBPToInstances(true, true);
-			HierarchyContainer.OnElementChanged.Broadcast(&HierarchyContainer, FRigElementKey());
+			PropagateHierarchyFromBPToInstances();
 
 			// make sure the bone name list is up 2 date for the editor graph
-	for (UEdGraph* Graph : UbergraphPages)
-	{
-		UControlRigGraph* RigGraph = Cast<UControlRigGraph>(Graph);
+			for (UEdGraph* Graph : UbergraphPages)
+			{
+				UControlRigGraph* RigGraph = Cast<UControlRigGraph>(Graph);
 				if (RigGraph == nullptr)
 				{
 					continue;
 				}
-				RigGraph->CacheNameLists(&HierarchyContainer, &DrawContainer);
+				RigGraph->CacheNameLists(Hierarchy, &DrawContainer);
 			}
 
 			RequestAutoVMRecompilation();
@@ -950,16 +1029,6 @@ TArray<UStruct*> UControlRigBlueprint::GetAvailableRigUnits()
 	return Structs;
 }
 
-UControlRigHierarchyModifier* UControlRigBlueprint::GetHierarchyModifier()
-{
-	if (HierarchyModifier == nullptr)
-	{
-		HierarchyModifier = NewObject<UControlRigHierarchyModifier>(this, TEXT("HierarchyModifier"));
-		HierarchyModifier->Container = &HierarchyContainer;
-	}
-	return HierarchyModifier;
-}
-
 #if WITH_EDITOR
 
 FName UControlRigBlueprint::AddMemberVariable(const FName& InName, const FString& InCPPType, bool bIsPublic, bool bIsReadOnly, FString InDefaultValue)
@@ -1055,7 +1124,7 @@ FName UControlRigBlueprint::AddTransientControl(URigVMPin* InPin)
 
 			if (URigVMPin::SplitPinPathAtStart(PinPath, Left, Right))
 		{
-				SpaceKey = DefaultStruct->DetermineSpaceForPin(Right, &HierarchyContainer);
+				SpaceKey = DefaultStruct->DetermineSpaceForPin(Right, Hierarchy);
 			}
 		}
 	}
@@ -1076,12 +1145,6 @@ FName UControlRigBlueprint::AddTransientControl(URigVMPin* InPin)
 		}
 	}
 
-	if (ReturnName != NAME_None)
-	{
-		HierarchyContainer.OnElementAdded.Broadcast(&HierarchyContainer, FRigElementKey(ReturnName, ERigElementType::Control));
-		HierarchyContainer.OnElementSelected.Broadcast(&HierarchyContainer, FRigElementKey(ReturnName, ERigElementType::Control), true);
-	}
-
 	return ReturnName;
 }
 
@@ -1112,11 +1175,6 @@ FName UControlRigBlueprint::RemoveTransientControl(URigVMPin* InPin)
 		}
 	}
 
-	if (RemovedName != NAME_None)
-		{
-		HierarchyContainer.OnElementSelected.Broadcast(&HierarchyContainer, FRigElementKey(RemovedName, ERigElementType::Control), false);
-		HierarchyContainer.OnElementRemoved.Broadcast(&HierarchyContainer, FRigElementKey(RemovedName, ERigElementType::Control));
-	}
 	return RemovedName;
 }
 
@@ -1141,19 +1199,13 @@ FName UControlRigBlueprint::AddTransientControl(const FRigElementKey& InElement)
 	{
 		UControlRig* InstancedControlRig = Cast<UControlRig>(ArchetypeInstance);
 		if (InstancedControlRig)
-			{
+		{
 			FName ControlName = InstancedControlRig->AddTransientControl(InElement);
 			if (ReturnName == NAME_None)
-				{
-				ReturnName = ControlName;
-				}
+			{
+			ReturnName = ControlName;
 			}
 		}
-
-	if (ReturnName != NAME_None)
-	{
-		HierarchyContainer.OnElementAdded.Broadcast(&HierarchyContainer, FRigElementKey(ReturnName, ERigElementType::Control));
-		HierarchyContainer.OnElementSelected.Broadcast(&HierarchyContainer, FRigElementKey(ReturnName, ERigElementType::Control), true);
 	}
 
 	return ReturnName;
@@ -1187,11 +1239,6 @@ FName UControlRigBlueprint::RemoveTransientControl(const FRigElementKey& InEleme
 		}
 	}
 
-	if (RemovedName != NAME_None)
-	{
-		HierarchyContainer.OnElementSelected.Broadcast(&HierarchyContainer, FRigElementKey(RemovedName, ERigElementType::Control), false);
-		HierarchyContainer.OnElementRemoved.Broadcast(&HierarchyContainer, FRigElementKey(RemovedName, ERigElementType::Control));
-	}
 	return RemovedName;
 }
 
@@ -1206,7 +1253,6 @@ void UControlRigBlueprint::ClearTransientControls()
 	UControlRigBlueprintGeneratedClass* RigClass = GetControlRigBlueprintGeneratedClass();
 	UControlRig* CDO = Cast<UControlRig>(RigClass->GetDefaultObject(true /* create if needed */));
 
-	TArray<FRigControl> PreviousControls;
 	TArray<UObject*> ArchetypeInstances;
 	CDO->GetArchetypeInstances(ArchetypeInstances);
 	for (UObject* ArchetypeInstance : ArchetypeInstances)
@@ -1214,19 +1260,9 @@ void UControlRigBlueprint::ClearTransientControls()
 		UControlRig* InstancedControlRig = Cast<UControlRig>(ArchetypeInstance);
 		if (InstancedControlRig)
 		{
-			if (PreviousControls.Num() == 0)
-			{
-				PreviousControls = InstancedControlRig->TransientControls;
-			}
 			InstancedControlRig->ClearTransientControls();
 		}
 	}
-
-	for(const FRigControl& RemovedControl : PreviousControls)
-	{
-		HierarchyContainer.OnElementSelected.Broadcast(&HierarchyContainer, FRigElementKey(RemovedControl.Name, ERigElementType::Control), false);
-		HierarchyContainer.OnElementRemoved.Broadcast(&HierarchyContainer, FRigElementKey(RemovedControl.Name, ERigElementType::Control));
-	}
 }
 
 void UControlRigBlueprint::SetTransientControlValue(const FRigElementKey& InElement)
@@ -1571,7 +1607,7 @@ void UControlRigBlueprint::RebuildGraphFromModel()
 	{
 		if (UControlRigGraph* RigGraph = Cast<UControlRigGraph>(Graph))
 		{
-			RigGraph->CacheNameLists(&HierarchyContainer, &DrawContainer);
+			RigGraph->CacheNameLists(Hierarchy, &DrawContainer);
 		}
 	}
 
@@ -1698,9 +1734,14 @@ void UControlRigBlueprint::HandleModifiedEvent(ERigVMGraphNotifType InNotifType,
 							{
 								if (UControlRig* InstanceRig = Cast<UControlRig>(ArchetypeInstance))
 								{
-									for (const FRigControl& Control : InstanceRig->TransientControls)
-									{
-										if (URigVMPin* ControlledPin = Model->FindPin(Control.Name.ToString()))
+									Hierarchy->ForEach<FRigControlElement>([this, InstanceRig, Pin](FRigControlElement* ControlElement) -> bool
+                                    {
+										if(!ControlElement->Settings.bIsTransientControl)
+										{
+											return true;
+										}
+										
+										if (URigVMPin* ControlledPin = Model->FindPin(ControlElement->GetName().ToString()))
 										{
 											URigVMPin* ControlledPinForLink = ControlledPin->GetPinForLink();
 
@@ -1715,9 +1756,11 @@ void UControlRigBlueprint::HandleModifiedEvent(ERigVMGraphNotifType InNotifType,
 												InstanceRig->ClearTransientControls();
 												InstanceRig->AddTransientControl(ControlledPin);
 											}
-											break;
+											return false;
 										}
-									}
+
+										return true;
+									});
 								}
 							}
 						}
@@ -1883,22 +1926,6 @@ void UControlRigBlueprint::SuspendNotifications(bool bSuspendNotifs)
 	}
 }
 
-void UControlRigBlueprint::CleanupBoneHierarchyDeprecated()
-{
-	if (Hierarchy_DEPRECATED.Num() > 0)
-	{
-		HierarchyContainer.BoneHierarchy = Hierarchy_DEPRECATED;
-		Hierarchy_DEPRECATED.Reset();
-	}
-
-	if (CurveContainer_DEPRECATED.Num() > 0)
-	{
-		HierarchyContainer.CurveContainer = CurveContainer_DEPRECATED;
-		CurveContainer_DEPRECATED.Reset();
-	}
-
-}
-
 void UControlRigBlueprint::CreateMemberVariablesOnLoad()
 {
 #if WITH_EDITOR
@@ -2144,29 +2171,7 @@ void UControlRigBlueprint::PatchVariableNodesOnLoad()
 void UControlRigBlueprint::PropagatePoseFromInstanceToBP(UControlRig* InControlRig)
 {
 	check(InControlRig);
-
-	for (const FRigBone& InputBone : InControlRig->Hierarchy.BoneHierarchy)
-	{
-		FRigBone& OutputBone = HierarchyContainer.BoneHierarchy[InputBone.Name];
-		OutputBone.InitialTransform = InputBone.InitialTransform;
-		OutputBone.LocalTransform = InputBone.LocalTransform;
-		OutputBone.GlobalTransform = InputBone.GlobalTransform;
-	}
-
-	for (const FRigSpace& InputSpace: InControlRig->Hierarchy.SpaceHierarchy)
-	{
-		FRigSpace& OutputSpace = HierarchyContainer.SpaceHierarchy[InputSpace.Name];
-		OutputSpace.InitialTransform = InputSpace.InitialTransform;
-		OutputSpace.LocalTransform = InputSpace.LocalTransform;
-	}
-
-	for (const FRigControl& InputControl : InControlRig->Hierarchy.ControlHierarchy)
-		{
-		FRigControl& OutputControl = HierarchyContainer.ControlHierarchy[InputControl.Name];
-		OutputControl.OffsetTransform = InputControl.OffsetTransform;
-		OutputControl.InitialValue = InputControl.InitialValue;
-		OutputControl.Value = InputControl.Value;
-		}
+	Hierarchy->CopyPose(InControlRig->GetHierarchy(), true, true);
 }
 
 void UControlRigBlueprint::PropagatePoseFromBPToInstances()
@@ -2180,72 +2185,31 @@ void UControlRigBlueprint::PropagatePoseFromBPToInstances()
 			for (UObject* ArchetypeInstance : ArchetypeInstances)
 			{
 				if (UControlRig* InstanceRig = Cast<UControlRig>(ArchetypeInstance))
-		{
-					for (const FRigBone& InputBone : HierarchyContainer.BoneHierarchy)
-					{
-						FRigBone& OutputBone = InstanceRig->Hierarchy.BoneHierarchy[InputBone.Name];
-						OutputBone.InitialTransform = InputBone.InitialTransform;
-						OutputBone.LocalTransform = InputBone.LocalTransform;
-						OutputBone.GlobalTransform = InputBone.GlobalTransform;
-		}
-
-					for (const FRigSpace& InputSpace : HierarchyContainer.SpaceHierarchy)
-					{
-						FRigSpace& OutputSpace = InstanceRig->Hierarchy.SpaceHierarchy[InputSpace.Name];
-						OutputSpace.InitialTransform = InputSpace.InitialTransform;
-						OutputSpace.LocalTransform = InputSpace.LocalTransform;
-					}
-
-					for (const FRigControl& InputControl : HierarchyContainer.ControlHierarchy)
-		{
-						FRigControl& OutputControl = InstanceRig->Hierarchy.ControlHierarchy[InputControl.Name];
-						OutputControl.OffsetTransform = InputControl.OffsetTransform;
-						OutputControl.InitialValue = InputControl.InitialValue;
-						OutputControl.Value = InputControl.Value;
-		}
-		}
-	}
+				{
+					InstanceRig->GetHierarchy()->CopyPose(Hierarchy, true, true);
+				}
+			}
 		}
 	}
 }
 
-void UControlRigBlueprint::PropagateHierarchyFromBPToInstances(bool bInitializeContainer, bool bInitializeRigs)
+void UControlRigBlueprint::PropagateHierarchyFromBPToInstances()
 {
 	if (UClass* MyControlRigClass = GeneratedClass)
 	{
 		if (UControlRig* DefaultObject = Cast<UControlRig>(MyControlRigClass->GetDefaultObject(false)))
 		{
-			if (bInitializeContainer)
-			{
-				HierarchyContainer.Initialize();
-				HierarchyContainer.ResetTransforms();
-			}
+			DefaultObject->GetHierarchy()->CopyHierarchy(Hierarchy);
+			DefaultObject->Initialize(true);
 
-			DefaultObject->Hierarchy = HierarchyContainer;
-			if (bInitializeRigs)
+			TArray<UObject*> ArchetypeInstances;
+			DefaultObject->GetArchetypeInstances(ArchetypeInstances);
+			for (UObject* ArchetypeInstance : ArchetypeInstances)
 			{
-				DefaultObject->Initialize(true);
-			}
-			else
-				{
-				DefaultObject->Hierarchy.Initialize(false);
-			}
-
-						TArray<UObject*> ArchetypeInstances;
-						DefaultObject->GetArchetypeInstances(ArchetypeInstances);
-						for (UObject* ArchetypeInstance : ArchetypeInstances)
-						{
 				if (UControlRig* InstanceRig = Cast<UControlRig>(ArchetypeInstance))
 				{
-					InstanceRig->Hierarchy = HierarchyContainer;
-					if (bInitializeRigs)
-					{
-						InstanceRig->Initialize(true);
-						}
-					else
-					{
-						InstanceRig->Hierarchy.Initialize(false);
-					}
+					InstanceRig->GetHierarchy()->CopyHierarchy(Hierarchy);
+					InstanceRig->Initialize(true);
 				}
 			}
 		}
@@ -2282,13 +2246,13 @@ void UControlRigBlueprint::PropagateDrawInstructionsFromBPToInstances()
 		{
 			continue;
 		}
-		RigGraph->CacheNameLists(&HierarchyContainer, &DrawContainer);
+		RigGraph->CacheNameLists(Hierarchy, &DrawContainer);
 	}
 }
 
 void UControlRigBlueprint::PropagatePropertyFromBPToInstances(FRigElementKey InRigElement, const FProperty* InProperty)
 {
-	int32 ElementIndex = HierarchyContainer.GetIndex(InRigElement);
+	int32 ElementIndex = Hierarchy->GetIndex(InRigElement);
 	ensure(ElementIndex != INDEX_NONE);
 	check(InProperty);
 
@@ -2299,184 +2263,94 @@ void UControlRigBlueprint::PropagatePropertyFromBPToInstances(FRigElementKey InR
 			TArray<UObject*> ArchetypeInstances;
 			DefaultObject->GetArchetypeInstances(ArchetypeInstances);
 
-			int32 PropertyOffset = InProperty->GetOffset_ReplaceWith_ContainerPtrToValuePtr();
-			int32 PropertySize = InProperty->GetSize();
+			const int32 PropertyOffset = InProperty->GetOffset_ReplaceWith_ContainerPtrToValuePtr();
+			const int32 PropertySize = InProperty->GetSize();
 
-			switch (InRigElement.Type)
+			uint8* Source = ((uint8*)Hierarchy->Get(ElementIndex)) + PropertyOffset;
+			for (UObject* ArchetypeInstance : ArchetypeInstances)
 			{
-				case ERigElementType::Bone:
-				{
-					uint8* Source = ((uint8*)&HierarchyContainer.BoneHierarchy[ElementIndex]) + PropertyOffset;
-					for (UObject* ArchetypeInstance : ArchetypeInstances)
-					{
-						if (UControlRig* InstanceRig = Cast<UControlRig>(ArchetypeInstance))
-						{
-							uint8* Dest = ((uint8*)&InstanceRig->Hierarchy.BoneHierarchy[ElementIndex]) + PropertyOffset;
-							FMemory::Memcpy(Dest, Source, PropertySize);
-						}
-					}
-					break;
-				}
-				case ERigElementType::Space:
-				{
-					uint8* Source = ((uint8*)&HierarchyContainer.SpaceHierarchy[ElementIndex]) + PropertyOffset;
-					for (UObject* ArchetypeInstance : ArchetypeInstances)
-					{
-						if (UControlRig* InstanceRig = Cast<UControlRig>(ArchetypeInstance))
-						{
-							uint8* Dest = ((uint8*)&InstanceRig->Hierarchy.SpaceHierarchy[ElementIndex]) + PropertyOffset;
-							FMemory::Memcpy(Dest, Source, PropertySize);
-						}
-					}
-					break;
-				}
-				case ERigElementType::Control:
-				{
-					uint8* Source = ((uint8*)&HierarchyContainer.ControlHierarchy[ElementIndex]) + PropertyOffset;
-					for (UObject* ArchetypeInstance : ArchetypeInstances)
-					{
-						if (UControlRig* InstanceRig = Cast<UControlRig>(ArchetypeInstance))
-						{
-							uint8* Dest = ((uint8*)&InstanceRig->Hierarchy.ControlHierarchy[ElementIndex]) + PropertyOffset;
-							FMemory::Memcpy(Dest, Source, PropertySize);
-						}
-					}
-					break;
-				}
-				case ERigElementType::Curve:
+				if (UControlRig* InstanceRig = Cast<UControlRig>(ArchetypeInstance))
 				{
-					uint8* Source = ((uint8*)&HierarchyContainer.CurveContainer[ElementIndex]) + PropertyOffset;
-					for (UObject* ArchetypeInstance : ArchetypeInstances)
-					{
-						if (UControlRig* InstanceRig = Cast<UControlRig>(ArchetypeInstance))
-						{
-							uint8* Dest = ((uint8*)&InstanceRig->Hierarchy.CurveContainer[ElementIndex]) + PropertyOffset;
-							FMemory::Memcpy(Dest, Source, PropertySize);
-						}
-					}
-					break;
-				}
-				default:
-			{
-					break;
+					uint8* Dest = ((uint8*)InstanceRig->GetHierarchy()->Get(ElementIndex)) + PropertyOffset;
+					FMemory::Memcpy(Dest, Source, PropertySize);
 				}
 			}
 		}
-			}
+	}
 }
 
 void UControlRigBlueprint::PropagatePropertyFromInstanceToBP(FRigElementKey InRigElement, const FProperty* InProperty, UControlRig* InInstance)
 {
-	int32 ElementIndex = HierarchyContainer.GetIndex(InRigElement);
+	const int32 ElementIndex = Hierarchy->GetIndex(InRigElement);
 	ensure(ElementIndex != INDEX_NONE);
 	check(InProperty);
 
-	int32 PropertyOffset = InProperty->GetOffset_ReplaceWith_ContainerPtrToValuePtr();
-	int32 PropertySize = InProperty->GetSize();
-
-	switch (InRigElement.Type)
-	{
-		case ERigElementType::Bone:
-		{
-			uint8* Source = ((uint8*)&InInstance->Hierarchy.BoneHierarchy[ElementIndex]) + PropertyOffset;
-			uint8* Dest = ((uint8*)&HierarchyContainer.BoneHierarchy[ElementIndex]) + PropertyOffset;
-			FMemory::Memcpy(Dest, Source, PropertySize);
-			break;
-		}
-		case ERigElementType::Space:
-			{
-			uint8* Source = ((uint8*)&InInstance->Hierarchy.SpaceHierarchy[ElementIndex]) + PropertyOffset;
-			uint8* Dest = ((uint8*)&HierarchyContainer.SpaceHierarchy[ElementIndex]) + PropertyOffset;
-			FMemory::Memcpy(Dest, Source, PropertySize);
-			break;
-		}
-		case ERigElementType::Control:
-				{
-			uint8* Source = ((uint8*)&InInstance->Hierarchy.ControlHierarchy[ElementIndex]) + PropertyOffset;
-			uint8* Dest = ((uint8*)&HierarchyContainer.ControlHierarchy[ElementIndex]) + PropertyOffset;
-			FMemory::Memcpy(Dest, Source, PropertySize);
-			break;
-				}
-		case ERigElementType::Curve:
-		{
-			uint8* Source = ((uint8*)&InInstance->Hierarchy.CurveContainer[ElementIndex]) + PropertyOffset;
-			uint8* Dest = ((uint8*)&HierarchyContainer.CurveContainer[ElementIndex]) + PropertyOffset;
-			FMemory::Memcpy(Dest, Source, PropertySize);
-			break;
-			}
-		default:
-		{
-			break;
-		}
-	}
-}
-
-#if WITH_EDITOR
-
-void UControlRigBlueprint::HandleOnElementAdded(FRigHierarchyContainer* InContainer, const FRigElementKey& InKey)
-{
-	if (bSuspendAllNotifications)
-	{
-		return;
-	}
-	PropagateHierarchyFromBPToInstances();
+	const int32 PropertyOffset = InProperty->GetOffset_ReplaceWith_ContainerPtrToValuePtr();
+	const int32 PropertySize = InProperty->GetSize();
+	uint8* Source = ((uint8*)InInstance->GetHierarchy()->Get(ElementIndex)) + PropertyOffset;
+	uint8* Dest = ((uint8*)Hierarchy->Get(ElementIndex)) + PropertyOffset;
+	FMemory::Memcpy(Dest, Source, PropertySize);
 }
 
-void UControlRigBlueprint::HandleOnElementRemoved(FRigHierarchyContainer* InContainer, const FRigElementKey& InKey)
-{
-	if (bSuspendAllNotifications)
-	{
-		return;
-	}
-
-	Modify();
-	Influences.OnKeyRemoved(InKey);
-	PropagateHierarchyFromBPToInstances();
-}
 
-void UControlRigBlueprint::HandleOnElementRenamed(FRigHierarchyContainer* InContainer, ERigElementType InElementType, const FName& InOldName, const FName& InNewName)
+void UControlRigBlueprint::HandleHierarchyModified(ERigHierarchyNotification InNotification, URigHierarchy* InHierarchy, const FRigBaseElement* InElement)
 {
-	if (bSuspendAllNotifications)
-	{
-		return;
-	}
-
-	Modify();
-	Influences.OnKeyRenamed(FRigElementKey(InOldName, InElementType), FRigElementKey(InNewName, InElementType));
-	PropagateHierarchyFromBPToInstances();
-}
+#if WITH_EDITOR
 
-void UControlRigBlueprint::HandleOnElementReparented(FRigHierarchyContainer* InContainer, const FRigElementKey& InKey, const FName& InOldParentName, const FName& InNewParentName)
-{
-	if (bSuspendAllNotifications)
+	if(bSuspendAllNotifications)
 	{
 		return;
 	}
-	PropagateHierarchyFromBPToInstances();
-}
 
-void UControlRigBlueprint::HandleOnElementSelected(FRigHierarchyContainer* InContainer, const FRigElementKey& InKey, bool bSelected)
-{
-	if (bSuspendAllNotifications)
-	{
-		return;
-	}
-	if (InKey.Type == ERigElementType::Control)
+	switch(InNotification)
 	{
-		if (UControlRig* RigBeingDebugged = Cast<UControlRig>(GetObjectBeingDebugged()))
+		case ERigHierarchyNotification::ElementRemoved:
+		{
+			Modify();
+			Influences.OnKeyRemoved(InElement->GetKey());
+			PropagateHierarchyFromBPToInstances();
+			break;
+		}
+		case ERigHierarchyNotification::ElementRenamed:
+		{
+			Modify();
+			Influences.OnKeyRenamed(FRigElementKey(InHierarchy->GetPreviousName(InElement->GetKey()), InElement->GetType()), InElement->GetKey());
+			PropagateHierarchyFromBPToInstances();
+			break;
+		}
+		case ERigHierarchyNotification::ElementAdded:
+		case ERigHierarchyNotification::ParentChanged:
+		case ERigHierarchyNotification::HierarchyReset:
+		{
+			Modify();
+			PropagateHierarchyFromBPToInstances();
+			break;
+		}
+		case ERigHierarchyNotification::ElementSelected:
 		{
-			if (FRigControl* Control = RigBeingDebugged->FindControl(InKey.Name))
+			bool bClearTransientControls = true;
+			if (const FRigControlElement* ControlElement = Cast<FRigControlElement>(InElement))
 			{
-				if (!Control->bIsTransientControl)
+				if (ControlElement->Settings.bIsTransientControl)
 				{
-					ClearTransientControls();
+					bClearTransientControls = false;
 				}
 			}
+
+			if(bClearTransientControls)
+			{
+				ClearTransientControls();
+			}
+			break;
+		}
+		default:
+		{
+			break;
 		}
 	}
-}
-
+	
 #endif
+}
 
 UControlRigBlueprint::FControlValueScope::FControlValueScope(UControlRigBlueprint* InBlueprint)
 : Blueprint(InBlueprint)
@@ -2486,10 +2360,10 @@ UControlRigBlueprint::FControlValueScope::FControlValueScope(UControlRigBlueprin
 
 	if (UControlRig* CR = Cast<UControlRig>(Blueprint->GetObjectBeingDebugged()))
 	{
-		const TArray<FRigControl>& Controls = CR->AvailableControls();
-		for (const FRigControl& Control : Controls)
+		TArray<FRigControlElement*> Controls = CR->AvailableControls();
+		for (FRigControlElement* ControlElement : Controls)
 		{
-			ControlValues.Add(Control.Name, CR->GetControlValue(Control.Name));
+			ControlValues.Add(ControlElement->GetName(), CR->GetControlValue(ControlElement->GetName()));
 		}
 	}
 #endif
@@ -2653,6 +2527,11 @@ void UControlRigBlueprint::BroadcastGraphImported(UEdGraph* InGraph)
 	GraphImportedEvent.Broadcast(InGraph);
 }
 
+void UControlRigBlueprint::BroadcastPostEditChangeChainProperty(FPropertyChangedChainEvent& PropertyChangedChainEvent)
+{
+	PostEditChangeChainPropertyEvent.Broadcast(PropertyChangedChainEvent);
+}
+
 #endif
 
 void UControlRigBlueprint::CreateEdGraphForCollapseNodeIfNeeded(URigVMCollapseNode* InNode, bool bForce)
diff --git a/Engine/Plugins/Experimental/ControlRig/Source/ControlRigDeveloper/Private/ControlRigBlueprintCompiler.cpp b/Engine/Plugins/Experimental/ControlRig/Source/ControlRigDeveloper/Private/ControlRigBlueprintCompiler.cpp
--- a/Engine/Plugins/Experimental/ControlRig/Source/ControlRigDeveloper/Private/ControlRigBlueprintCompiler.cpp
+++ b/Engine/Plugins/Experimental/ControlRig/Source/ControlRigDeveloper/Private/ControlRigBlueprintCompiler.cpp
@@ -90,10 +90,8 @@ void FControlRigBlueprintCompilerContext::CopyTermDefaultsToDefaultObject(UObjec
 	UControlRigBlueprint* ControlRigBlueprint = Cast<UControlRigBlueprint>(Blueprint);
 	if (ControlRigBlueprint)
 	{
-		ControlRigBlueprint->CleanupBoneHierarchyDeprecated();
-
 		UControlRig* ControlRig = CastChecked<UControlRig>(DefaultObject);
-		ControlRig->Hierarchy = ControlRigBlueprint->HierarchyContainer;
+		ControlRig->GetHierarchy()->CopyHierarchy(ControlRigBlueprint->Hierarchy);
 		ControlRig->DrawContainer = ControlRigBlueprint->DrawContainer;
 		ControlRig->Influences = ControlRigBlueprint->Influences;
 	}
diff --git a/Engine/Plugins/Experimental/ControlRig/Source/ControlRigDeveloper/Private/Graph/ControlRigGraph.cpp b/Engine/Plugins/Experimental/ControlRig/Source/ControlRigDeveloper/Private/Graph/ControlRigGraph.cpp
--- a/Engine/Plugins/Experimental/ControlRig/Source/ControlRigDeveloper/Private/Graph/ControlRigGraph.cpp
+++ b/Engine/Plugins/Experimental/ControlRig/Source/ControlRigDeveloper/Private/Graph/ControlRigGraph.cpp
@@ -64,20 +64,20 @@ void UControlRigGraph::Serialize(FArchive& Ar)
 
 #if WITH_EDITOR
 
-void UControlRigGraph::CacheNameLists(const FRigHierarchyContainer* HierarchyContainer, const FControlRigDrawContainer* DrawContainer)
+void UControlRigGraph::CacheNameLists(URigHierarchy* InHierarchy, const FControlRigDrawContainer* DrawContainer)
 {
 	if (UControlRigGraph* OuterGraph = Cast<UControlRigGraph>(GetOuter()))
 	{
 		return;
 	}
 
-	check(HierarchyContainer);
+	check(InHierarchy);
 	check(DrawContainer);
 
-	CacheNameList<FRigBoneHierarchy>(HierarchyContainer->BoneHierarchy, BoneNameList);
-	CacheNameList<FRigControlHierarchy>(HierarchyContainer->ControlHierarchy, ControlNameList);
-	CacheNameList<FRigSpaceHierarchy>(HierarchyContainer->SpaceHierarchy, SpaceNameList);
-	CacheNameList<FRigCurveContainer>(HierarchyContainer->CurveContainer, CurveNameList);
+	CacheNameListForHierarchy<FRigBoneElement>(InHierarchy, BoneNameList);
+	CacheNameListForHierarchy<FRigControlElement>(InHierarchy, ControlNameList);
+	CacheNameListForHierarchy<FRigSpaceElement>(InHierarchy, SpaceNameList);
+	CacheNameListForHierarchy<FRigCurveElement>(InHierarchy, CurveNameList);
 	CacheNameList<FControlRigDrawContainer>(*DrawContainer, DrawingNameList);
 }
 
diff --git a/Engine/Plugins/Experimental/ControlRig/Source/ControlRigDeveloper/Public/ControlRigBlueprint.h b/Engine/Plugins/Experimental/ControlRig/Source/ControlRigDeveloper/Public/ControlRigBlueprint.h
--- a/Engine/Plugins/Experimental/ControlRig/Source/ControlRigDeveloper/Public/ControlRigBlueprint.h
+++ b/Engine/Plugins/Experimental/ControlRig/Source/ControlRigDeveloper/Public/ControlRigBlueprint.h
@@ -9,14 +9,14 @@
 #include "Misc/Crc.h"
 #include "ControlRigDefines.h"
 #include "Rigs/RigHierarchyContainer.h"
-#include "Rigs/RigCurveContainer.h"
+#include "Rigs/RigHierarchy.h"
+#include "Rigs/RigHierarchyController.h"
 #include "Interfaces/Interface_PreviewMeshProvider.h"
 #include "ControlRigGizmoLibrary.h"
 #include "RigVMCore/RigVM.h"
 #include "RigVMCore/RigVMStatistics.h"
 #include "RigVMModel/RigVMController.h"
 #include "RigVMCompiler/RigVMCompiler.h"
-#include "ControlRigHierarchyModifier.h"
 #include "ControlRigValidationPass.h"
 #include "Drawing/ControlRigDrawContainer.h"
 
@@ -36,6 +36,7 @@ DECLARE_EVENT_FourParams(UControlRigBlueprint, FOnVariableDroppedEvent, UObject*
 DECLARE_EVENT_OneParam(UControlRigBlueprint, FOnExternalVariablesChanged, const TArray<FRigVMExternalVariable>&);
 DECLARE_EVENT_TwoParams(UControlRigBlueprint, FOnNodeDoubleClicked, UControlRigBlueprint*, URigVMNode*);
 DECLARE_EVENT_OneParam(UControlRigBlueprint, FOnGraphImported, UEdGraph*);
+DECLARE_EVENT_OneParam(UControlRigBlueprint, FOnPostEditChangeChainProperty, FPropertyChangedChainEvent&);
 
 UCLASS(BlueprintType, meta=(IgnoreClassThumbnail))
 class CONTROLRIGDEVELOPER_API UControlRigBlueprint : public UBlueprint, public IInterface_PreviewMeshProvider
@@ -79,6 +80,9 @@ public:
 	virtual bool TryImportGraphFromText(const FString& InClipboardText, UEdGraph** OutGraphPtr = nullptr) override;
 	virtual bool CanImportGraphFromText(const FString& InClipboardText) override;
 
+	// UObject interface
+	virtual void PostEditChangeChainProperty(struct FPropertyChangedChainEvent& PropertyChangedEvent) override;
+
 #endif	// #if WITH_EDITOR
 
 	virtual bool ShouldBeMarkedDirtyUponTransaction() const override { return false; }
@@ -189,9 +193,6 @@ public:
 	UFUNCTION(BlueprintCallable, Category = "VM")
 	static TArray<UStruct*> GetAvailableRigUnits();
 
-	UFUNCTION(BlueprintCallable, Category = "Hierarchy")
-	UControlRigHierarchyModifier* GetHierarchyModifier();
-
 #if WITH_EDITOR
 	UFUNCTION(BlueprintCallable, Category = "Variables")
 	FName AddMemberVariable(const FName& InName, const FString& InCPPType, bool bIsPublic = false, bool bIsReadOnly = false, FString InDefaultValue = TEXT(""));
@@ -233,16 +234,17 @@ public:
 	FRigInfluenceMapPerEvent Influences;
 
 public:
+
 	UPROPERTY()
-	FRigHierarchyContainer HierarchyContainer;
+	FRigHierarchyContainer HierarchyContainer_DEPRECATED;
 
-private:
+	UPROPERTY(BlueprintReadOnly, Category = "Hierarchy")
+	URigHierarchy* Hierarchy;
 
-	UPROPERTY()
-	FRigBoneHierarchy Hierarchy_DEPRECATED;
+	UFUNCTION(BlueprintCallable, Category = "Hierarchy")
+	URigHierarchyController* GetHierarchyController() { return HierarchyController; }
 
-	UPROPERTY()
-	FRigCurveContainer CurveContainer_DEPRECATED;
+private:
 
 	/** Whether or not this rig has an Inversion Event */
 	UPROPERTY(AssetRegistrySearchable)
@@ -282,7 +284,7 @@ private:
 	int32 VMRecompilationBracket;
 
 	UPROPERTY(transient)
-	UControlRigHierarchyModifier* HierarchyModifier;
+	URigHierarchyController* HierarchyController;
 
 	FRigVMGraphModifiedEvent ModifiedEvent;
 	void Notify(ERigVMGraphNotifType InNotifType, UObject* InSubject);
@@ -310,7 +312,6 @@ private:
 
 	static TArray<UControlRigBlueprint*> sCurrentlyOpenedRigBlueprints;
 
-	void CleanupBoneHierarchyDeprecated();
 	void CreateMemberVariablesOnLoad();
 #if WITH_EDITOR
 	static FName FindCRMemberVariableUniqueName(TSharedPtr<FKismetNameValidator> InNameValidator, const FString& InBaseName);
@@ -325,7 +326,7 @@ private:
 public:
 	void PropagatePoseFromInstanceToBP(UControlRig* InControlRig);
 	void PropagatePoseFromBPToInstances();
-	void PropagateHierarchyFromBPToInstances(bool bInitializeContainer = true, bool bInitializeRigs = true);
+	void PropagateHierarchyFromBPToInstances();
 	void PropagateDrawInstructionsFromBPToInstances();
 	void PropagatePropertyFromBPToInstances(FRigElementKey InRigElement, const FProperty* InProperty);
 	void PropagatePropertyFromInstanceToBP(FRigElementKey InRigElement, const FProperty* InProperty, UControlRig* InInstance);
@@ -335,14 +336,7 @@ private:
 	UPROPERTY()
 	UControlRigValidator* Validator;
 
-#if WITH_EDITOR
-
-	void HandleOnElementAdded(FRigHierarchyContainer* InContainer, const FRigElementKey& InKey);
-	void HandleOnElementRemoved(FRigHierarchyContainer* InContainer, const FRigElementKey& InKey);
-	void HandleOnElementRenamed(FRigHierarchyContainer* InContainer, ERigElementType InElementType, const FName& InOldName, const FName& InNewName);
-	void HandleOnElementReparented(FRigHierarchyContainer* InContainer, const FRigElementKey& InKey, const FName& InOldParentName, const FName& InNewParentName);
-	void HandleOnElementSelected(FRigHierarchyContainer* InContainer, const FRigElementKey& InKey, bool bSelected);
-#endif
+	void HandleHierarchyModified(ERigHierarchyNotification InNotification, URigHierarchy* InHierarchy, const FRigBaseElement* InElement);
 
 	// Class used to temporarily cache all 
 	// current control values and reapply them
@@ -378,6 +372,9 @@ public:
 	FOnGraphImported& OnGraphImported() { return GraphImportedEvent; }
 	void BroadcastGraphImported(UEdGraph* InGraph);
 
+	FOnPostEditChangeChainProperty& OnPostEditChangeChainProperty() { return PostEditChangeChainPropertyEvent; }
+	void BroadcastPostEditChangeChainProperty(FPropertyChangedChainEvent& PropertyChangedChainEvent);
+
 private:
 
 	FOnExternalVariablesChanged ExternalVariablesChangedEvent;
@@ -385,6 +382,7 @@ private:
 
 	FOnNodeDoubleClicked NodeDoubleClickedEvent;
 	FOnGraphImported GraphImportedEvent;
+	FOnPostEditChangeChainProperty PostEditChangeChainPropertyEvent;
 
 #endif
 
