Control Rig: New rig hierarchy data structure refactoring

diff --git a/Engine/Plugins/Experimental/ControlRig/Source/ControlRigDeveloper/Public/Graph/ControlRigGraph.h b/Engine/Plugins/Experimental/ControlRig/Source/ControlRigDeveloper/Public/Graph/ControlRigGraph.h
--- a/Engine/Plugins/Experimental/ControlRig/Source/ControlRigDeveloper/Public/Graph/ControlRigGraph.h
+++ b/Engine/Plugins/Experimental/ControlRig/Source/ControlRigDeveloper/Public/Graph/ControlRigGraph.h
@@ -4,7 +4,8 @@
 
 #include "EdGraph/EdGraph.h"
 #include "Graph/ControlRigGraphNode.h"
-#include "Rigs/RigHierarchyContainer.h"
+#include "Rigs/RigHierarchy.h"
+#include "Rigs/RigHierarchyController.h"
 #include "RigVMModel/RigVMGraph.h"
 #include "RigVMCore/RigVM.h"
 #include "Drawing/ControlRigDrawContainer.h"
@@ -39,7 +40,7 @@ public:
 #if WITH_EDITOR
 
 
-	void CacheNameLists(const FRigHierarchyContainer* HierarchyContainer, const FControlRigDrawContainer* DrawContainer);
+	void CacheNameLists(URigHierarchy* InHierarchy, const FControlRigDrawContainer* DrawContainer);
 
 	const TArray<TSharedPtr<FString>>& GetBoneNameList(URigVMPin* InPin = nullptr) const;
 	const TArray<TSharedPtr<FString>>& GetControlNameList(URigVMPin* InPin = nullptr) const;
@@ -71,10 +72,29 @@ public:
 private:
 
 	template<class T>
-	void CacheNameList(const T& ElementList, TArray<TSharedPtr<FString>>& OutNameList)
+	void CacheNameListForHierarchy(URigHierarchy* InHierarchy, TArray<TSharedPtr<FString>>& OutNameList)
 	{
-		DECLARE_SCOPE_HIERARCHICAL_COUNTER_FUNC()
+        TArray<FString> Names;
+		for (auto Element : *InHierarchy)
+		{
+			if(Element->IsA<T>())
+			{
+				Names.Add(Element->GetName().ToString());
+			}
+		}
+		Names.Sort();
+
+		OutNameList.Reset();
+		OutNameList.Add(MakeShared<FString>(FName(NAME_None).ToString()));
+		for (const FString& Name : Names)
+		{
+			OutNameList.Add(MakeShared<FString>(Name));
+		}
+	}
 
+	template<class T>
+	void CacheNameList(const T& ElementList, TArray<TSharedPtr<FString>>& OutNameList)
+	{
 		TArray<FString> Names;
 		for (auto Element : ElementList)
 		{
@@ -90,7 +110,6 @@ private:
 		}
 	}
 
-
 	TArray<TSharedPtr<FString>> BoneNameList;
 	TArray<TSharedPtr<FString>> ControlNameList;
 	TArray<TSharedPtr<FString>> SpaceNameList;
diff --git a/Engine/Plugins/Experimental/ControlRig/Source/ControlRigEditor/Private/ControlRigBlueprintActions.cpp b/Engine/Plugins/Experimental/ControlRig/Source/ControlRigEditor/Private/ControlRigBlueprintActions.cpp
--- a/Engine/Plugins/Experimental/ControlRig/Source/ControlRigEditor/Private/ControlRigBlueprintActions.cpp
+++ b/Engine/Plugins/Experimental/ControlRig/Source/ControlRigEditor/Private/ControlRigBlueprintActions.cpp
@@ -179,11 +179,11 @@ UControlRigBlueprint* FControlRigBlueprintActions::CreateControlRigFromSkeletalM
 		return nullptr;
 	}
 
-	NewControlRigBlueprint->HierarchyContainer.BoneHierarchy.ImportSkeleton(*RefSkeleton, NAME_None, false, false, false, false);
-	NewControlRigBlueprint->HierarchyContainer.CurveContainer.ImportCurvesFromSkeleton(Skeleton, NAME_None, false, false, false);
+	NewControlRigBlueprint->GetHierarchyController()->ImportBones(*RefSkeleton, NAME_None, false, false, false, false);
+	NewControlRigBlueprint->GetHierarchyController()->ImportCurves(Skeleton, NAME_None, false, false);
 	NewControlRigBlueprint->SourceHierarchyImport = Skeleton;
 	NewControlRigBlueprint->SourceCurveImport = Skeleton;
-	NewControlRigBlueprint->PropagateHierarchyFromBPToInstances(true);
+	NewControlRigBlueprint->PropagateHierarchyFromBPToInstances();
 
 	if(SkeletalMesh)
 	{
diff --git a/Engine/Plugins/Experimental/ControlRig/Source/ControlRigEditor/Private/ControlRigBlueprintEditorLibrary.cpp b/Engine/Plugins/Experimental/ControlRig/Source/ControlRigEditor/Private/ControlRigBlueprintEditorLibrary.cpp
--- a/Engine/Plugins/Experimental/ControlRig/Source/ControlRigEditor/Private/ControlRigBlueprintEditorLibrary.cpp
+++ b/Engine/Plugins/Experimental/ControlRig/Source/ControlRigEditor/Private/ControlRigBlueprintEditorLibrary.cpp
@@ -95,11 +95,20 @@ TArray<UStruct*> UControlRigBlueprintEditorLibrary::GetAvailableRigUnits()
 	return UControlRigBlueprint::GetAvailableRigUnits();
 }
 
-UControlRigHierarchyModifier* UControlRigBlueprintEditorLibrary::GetHierarchyModifier(UControlRigBlueprint* InRigBlueprint)
+URigHierarchy* UControlRigBlueprintEditorLibrary::GetHierarchy(UControlRigBlueprint* InRigBlueprint)
 {
 	if(InRigBlueprint == nullptr)
 	{
 		return nullptr;
 	}
-	return InRigBlueprint->GetHierarchyModifier();
+	return InRigBlueprint->Hierarchy;
+}
+
+URigHierarchyController* UControlRigBlueprintEditorLibrary::GetHierarchyController(UControlRigBlueprint* InRigBlueprint)
+{
+	if(InRigBlueprint == nullptr)
+	{
+		return nullptr;
+	}
+	return InRigBlueprint->GetHierarchyController();
 }
diff --git a/Engine/Plugins/Experimental/ControlRig/Source/ControlRigEditor/Private/ControlRigBlueprintFactory.h b/Engine/Plugins/Experimental/ControlRig/Source/ControlRigEditor/Private/ControlRigBlueprintFactory.h
--- a/Engine/Plugins/Experimental/ControlRig/Source/ControlRigEditor/Private/ControlRigBlueprintFactory.h
+++ b/Engine/Plugins/Experimental/ControlRig/Source/ControlRigEditor/Private/ControlRigBlueprintFactory.h
@@ -11,8 +11,8 @@
 #include "ControlRigBlueprint.h"
 #include "ControlRigBlueprintFactory.generated.h"
 
-UCLASS(MinimalAPI, HideCategories=Object)
-class UControlRigBlueprintFactory : public UFactory
+UCLASS(HideCategories=Object)
+class CONTROLRIGEDITOR_API UControlRigBlueprintFactory : public UFactory
 {
 	GENERATED_BODY()
 
diff --git a/Engine/Plugins/Experimental/ControlRig/Source/ControlRigEditor/Private/ControlRigEditorModule.cpp b/Engine/Plugins/Experimental/ControlRig/Source/ControlRigEditor/Private/ControlRigEditorModule.cpp
--- a/Engine/Plugins/Experimental/ControlRig/Source/ControlRigEditor/Private/ControlRigEditorModule.cpp
+++ b/Engine/Plugins/Experimental/ControlRig/Source/ControlRigEditor/Private/ControlRigEditorModule.cpp
@@ -173,14 +173,14 @@ void FControlRigEditorModule::StartupModule()
 	ClassesToUnregisterOnShutdown.Add(UMovieSceneControlRigParameterSection::StaticClass()->GetFName());
 	PropertyEditorModule.RegisterCustomClassLayout(ClassesToUnregisterOnShutdown.Last(), FOnGetDetailCustomizationInstance::CreateStatic(&FMovieSceneControlRigSectionDetailsCustomization::MakeInstance));
 
-	ClassesToUnregisterOnShutdown.Add(FRigBone::StaticStruct()->GetFName());
-	PropertyEditorModule.RegisterCustomClassLayout(ClassesToUnregisterOnShutdown.Last(), FOnGetDetailCustomizationInstance::CreateStatic(&FRigBoneDetails::MakeInstance));
-
-	ClassesToUnregisterOnShutdown.Add(FRigControl::StaticStruct()->GetFName());
-	PropertyEditorModule.RegisterCustomClassLayout(ClassesToUnregisterOnShutdown.Last(), FOnGetDetailCustomizationInstance::CreateStatic(&FRigControlDetails::MakeInstance));
+	ClassesToUnregisterOnShutdown.Add(FRigBoneElement::StaticStruct()->GetFName());
+	PropertyEditorModule.RegisterCustomClassLayout(ClassesToUnregisterOnShutdown.Last(), FOnGetDetailCustomizationInstance::CreateStatic(&FRigBoneElementDetails::MakeInstance));
+	
+	ClassesToUnregisterOnShutdown.Add(FRigControlElement::StaticStruct()->GetFName());
+	PropertyEditorModule.RegisterCustomClassLayout(ClassesToUnregisterOnShutdown.Last(), FOnGetDetailCustomizationInstance::CreateStatic(&FRigControlElementDetails::MakeInstance));
 
-	ClassesToUnregisterOnShutdown.Add(FRigSpace::StaticStruct()->GetFName());
-	PropertyEditorModule.RegisterCustomClassLayout(ClassesToUnregisterOnShutdown.Last(), FOnGetDetailCustomizationInstance::CreateStatic(&FRigSpaceDetails::MakeInstance));
+	ClassesToUnregisterOnShutdown.Add(FRigSpaceElement::StaticStruct()->GetFName());
+	PropertyEditorModule.RegisterCustomClassLayout(ClassesToUnregisterOnShutdown.Last(), FOnGetDetailCustomizationInstance::CreateStatic(&FRigSpaceElementDetails::MakeInstance));
 
 	ClassesToUnregisterOnShutdown.Add(FRigInfluenceMapPerEvent::StaticStruct()->GetFName());
 	PropertyEditorModule.RegisterCustomClassLayout(ClassesToUnregisterOnShutdown.Last(), FOnGetDetailCustomizationInstance::CreateStatic(&FRigInfluenceMapPerEventDetails::MakeInstance));
@@ -202,6 +202,9 @@ void FControlRigEditorModule::StartupModule()
 	PropertiesToUnregisterOnShutdown.Add(FRigElementKey::StaticStruct()->GetFName());
 	PropertyEditorModule.RegisterCustomPropertyTypeLayout(PropertiesToUnregisterOnShutdown.Last(), FOnGetPropertyTypeCustomizationInstance::CreateStatic(&FRigElementKeyDetails::MakeInstance));
 
+	PropertiesToUnregisterOnShutdown.Add(FRigComputedTransform::StaticStruct()->GetFName());
+	PropertyEditorModule.RegisterCustomPropertyTypeLayout(PropertiesToUnregisterOnShutdown.Last(), FOnGetPropertyTypeCustomizationInstance::CreateStatic(&FRigComputedTransformDetails::MakeInstance));
+
 	// Register asset tools
 	auto RegisterAssetTypeAction = [this](const TSharedRef<IAssetTypeActions>& InAssetTypeAction)
 	{
@@ -1469,7 +1472,7 @@ void FControlRigEditorModule::GetContextMenuActions(const UControlRigGraphSchema
 													FRigElementKey SpaceKey;
 													if (URigVMPin::SplitPinPathAtStart(PinPath, Left, Right))
 													{
-														SpaceKey = DefaultStruct->DetermineSpaceForPin(Right, &RigBlueprint->HierarchyContainer);
+														SpaceKey = DefaultStruct->DetermineSpaceForPin(Right, RigBlueprint->Hierarchy);
 													}
 
 													if (SpaceKey.IsValid())
@@ -1539,6 +1542,9 @@ void FControlRigEditorModule::GetContextMenuActions(const UControlRigGraphSchema
 				URigVMGraph* Model = RigBlueprint->GetModel(Context->Node->GetGraph());
 				URigVMController* Controller = RigBlueprint->GetController(Model);
 
+				URigHierarchy* TemporaryHierarchy = NewObject<URigHierarchy>();
+				TemporaryHierarchy->CopyHierarchy(RigBlueprint->Hierarchy);
+
 				TArray<FRigElementKey> RigElementsToSelect;
 				TMap<const URigVMPin*, FRigElementKey> PinToKey;
 				TArray<FName> SelectedNodeNames = Model->GetSelectNodes();
@@ -1549,7 +1555,6 @@ void FControlRigEditorModule::GetContextMenuActions(const UControlRigGraphSchema
 					if (URigVMNode* ModelNode = Model->FindNodeByName(SelectedNodeName))
 					{
 						TSharedPtr<FStructOnScope> StructOnScope;
-						FRigHierarchyContainer TemporaryHierarchy = RigBlueprint->HierarchyContainer;
 						FRigUnit* StructMemory = nullptr;
 						UScriptStruct* ScriptStruct = nullptr;
 						if (URigVMUnitNode* UnitNode = Cast<URigVMUnitNode>(ModelNode))
@@ -1559,7 +1564,7 @@ void FControlRigEditorModule::GetContextMenuActions(const UControlRigGraphSchema
 							StructMemory = (FRigUnit*)StructOnScope->GetStructMemory();
 
 							FRigUnitContext RigUnitContext;
-							RigUnitContext.Hierarchy = &TemporaryHierarchy;
+							RigUnitContext.Hierarchy = TemporaryHierarchy;
 							RigUnitContext.State = EControlRigState::Update;
 							StructMemory->Execute(RigUnitContext);
 						}
@@ -1670,10 +1675,10 @@ void FControlRigEditorModule::GetContextMenuActions(const UControlRigGraphSchema
 						FSlateIcon(),
 						FUIAction(FExecuteAction::CreateLambda([RigBlueprint, RigElementsToSelect]() {
 
-							RigBlueprint->HierarchyContainer.ClearSelection();
+							RigBlueprint->GetHierarchyController()->ClearSelection();
 							for(const FRigElementKey& RigElementToSelect : RigElementsToSelect)
 							{
-								RigBlueprint->HierarchyContainer.Select(RigElementToSelect, true);
+								RigBlueprint->GetHierarchyController()->SelectElement(RigElementToSelect, true);
 							}
 
 						})
@@ -1717,7 +1722,7 @@ void FControlRigEditorModule::GetContextMenuActions(const UControlRigGraphSchema
 								if(NewNameStr != OldNameStr)
 								{
 									Key.Name = *NewNameStr;
-									if(RigBlueprint->HierarchyContainer.GetIndex(Key) != INDEX_NONE)
+									if(RigBlueprint->Hierarchy->GetIndex(Key) != INDEX_NONE)
 									{
 										Controller->SetPinDefaultValue(Pin->GetPinPath(), NewNameStr, false);
 										ReplacedNames++;
