Control Rig: New rig hierarchy data structure refactoring

diff --git a/Engine/Plugins/Experimental/ControlRig/Source/ControlRigEditor/Private/ControlRigElementDetails.cpp b/Engine/Plugins/Experimental/ControlRig/Source/ControlRigEditor/Private/ControlRigElementDetails.cpp
--- a/Engine/Plugins/Experimental/ControlRig/Source/ControlRigEditor/Private/ControlRigElementDetails.cpp
+++ b/Engine/Plugins/Experimental/ControlRig/Source/ControlRigEditor/Private/ControlRigElementDetails.cpp
@@ -23,25 +23,23 @@
 
 #define LOCTEXT_NAMESPACE "ControlRigElementDetails"
 
-void FRigElementKeyDetails::CustomizeHeader(TSharedRef<IPropertyHandle> InStructPropertyHandle, FDetailWidgetRow& HeaderRow, IPropertyTypeCustomizationUtils& StructCustomizationUtils)
+void RigElementDetails_GetCustomizedInfo(TSharedRef<IPropertyHandle> InStructPropertyHandle, UControlRigBlueprint*& OutBlueprint)
 {
-	BlueprintBeingCustomized = nullptr;
-
 	TArray<UObject*> Objects;
 	InStructPropertyHandle->GetOuterObjects(Objects);
 	for (UObject* Object : Objects)
 	{
 		if (Object->IsA<UControlRigBlueprint>())
 		{
-			BlueprintBeingCustomized = Cast<UControlRigBlueprint>(Object);
-			if (BlueprintBeingCustomized)
+			OutBlueprint = Cast<UControlRigBlueprint>(Object);
+			if (OutBlueprint)
 			{
 				break;
 			}
 		}
 	}
 
-	if (BlueprintBeingCustomized == nullptr)
+	if (OutBlueprint == nullptr)
 	{
 		TArray<UPackage*> Packages;
 		InStructPropertyHandle->GetOuterPackages(Packages);
@@ -58,20 +56,26 @@ void FRigElementKeyDetails::CustomizeHeader(TSharedRef<IPropertyHandle> InStruct
 			{
 				if (UControlRig* Rig = Cast<UControlRig>(SubObject))
 				{
-					BlueprintBeingCustomized = Cast<UControlRigBlueprint>(Rig->GetClass()->ClassGeneratedBy);
-					if (BlueprintBeingCustomized)
+					OutBlueprint = Cast<UControlRigBlueprint>(Rig->GetClass()->ClassGeneratedBy);
+					if (OutBlueprint)
 					{
 						break;
 					}
 				}
 			}
 
-			if (BlueprintBeingCustomized)
+			if (OutBlueprint)
 			{
 				break;
 			}
 		}
 	}
+}
+
+void FRigElementKeyDetails::CustomizeHeader(TSharedRef<IPropertyHandle> InStructPropertyHandle, FDetailWidgetRow& HeaderRow, IPropertyTypeCustomizationUtils& StructCustomizationUtils)
+{
+	BlueprintBeingCustomized = nullptr;
+	RigElementDetails_GetCustomizedInfo(InStructPropertyHandle, BlueprintBeingCustomized);
 
 	UControlRigGraph* RigGraph = nullptr;
 	if(BlueprintBeingCustomized)
@@ -431,14 +435,96 @@ void FRigUnitDetails::OnNameListChanged(TSharedPtr<FString> NewSelection, ESelec
 void FRigUnitDetails::OnNameListComboBox(TSharedPtr<FStructOnScope> InStructOnScope, FNameProperty* InProperty, const TArray<TSharedPtr<FString>>* InNameList)
 {
 	TSharedPtr<SControlRigGraphPinNameListValueWidget> Widget = NameListWidgets.FindChecked(InProperty->GetFName());
-	TSharedPtr<FString> CurrentlySelected = GetCurrentlySelectedItem(InStructOnScope, InProperty, InNameList);
+	const TSharedPtr<FString> CurrentlySelected = GetCurrentlySelectedItem(InStructOnScope, InProperty, InNameList);
 	Widget->SetSelectedItem(CurrentlySelected);
 }
 
-void FRigElementDetails::CustomizeDetails(IDetailLayoutBuilder& DetailBuilder)
+void FRigComputedTransformDetails::CustomizeHeader(TSharedRef<IPropertyHandle> InStructPropertyHandle, FDetailWidgetRow& HeaderRow, IPropertyTypeCustomizationUtils& StructCustomizationUtils)
+{
+	BlueprintBeingCustomized = nullptr;
+	RigElementDetails_GetCustomizedInfo(InStructPropertyHandle, BlueprintBeingCustomized);
+}
+
+void FRigComputedTransformDetails::CustomizeChildren(TSharedRef<IPropertyHandle> InStructPropertyHandle, IDetailChildrenBuilder& StructBuilder, IPropertyTypeCustomizationUtils& StructCustomizationUtils)
+{
+	TransformHandle = InStructPropertyHandle->GetChildHandle(TEXT("Transform"));
+
+	StructBuilder
+	.AddProperty(TransformHandle.ToSharedRef())
+	.DisplayName(InStructPropertyHandle->GetPropertyDisplayName());
+
+    FString PropertyPath = TransformHandle->GeneratePathToProperty();
+
+	if(PropertyPath.StartsWith(TEXT("Struct.")))
+	{
+		PropertyPath.RightChopInline(7);
+	}
+
+	if(PropertyPath.StartsWith(TEXT("Pose.")))
+	{
+		PropertyPath.RightChopInline(5);
+		PropertyChain.AddTail(FRigTransformElement::StaticStruct()->FindPropertyByName(TEXT("Pose")));
+	}
+	else if(PropertyPath.StartsWith(TEXT("Offset.")))
+	{
+		PropertyPath.RightChopInline(7);
+		PropertyChain.AddTail(FRigControlElement::StaticStruct()->FindPropertyByName(TEXT("Offset")));
+	}
+	else if(PropertyPath.StartsWith(TEXT("Gizmo.")))
+	{
+		PropertyPath.RightChopInline(6);
+		PropertyChain.AddTail(FRigControlElement::StaticStruct()->FindPropertyByName(TEXT("Gizmo")));
+	}
+
+	if(PropertyPath.StartsWith(TEXT("Current.")))
+	{
+		PropertyPath.RightChopInline(8);
+		PropertyChain.AddTail(FRigCurrentAndInitialTransform::StaticStruct()->FindPropertyByName(TEXT("Current")));
+	}
+	else if(PropertyPath.StartsWith(TEXT("Initial.")))
+	{
+		PropertyPath.RightChopInline(8);
+		PropertyChain.AddTail(FRigCurrentAndInitialTransform::StaticStruct()->FindPropertyByName(TEXT("Initial")));
+	}
+
+	if(PropertyPath.StartsWith(TEXT("Local.")))
+	{
+		PropertyPath.RightChopInline(6);
+		PropertyChain.AddTail(FRigLocalAndGlobalTransform::StaticStruct()->FindPropertyByName(TEXT("Local")));
+	}
+	else if(PropertyPath.StartsWith(TEXT("Global.")))
+	{
+		PropertyPath.RightChopInline(7);
+		PropertyChain.AddTail(FRigLocalAndGlobalTransform::StaticStruct()->FindPropertyByName(TEXT("Global")));
+	}
+
+	PropertyChain.AddTail(TransformHandle->GetProperty());
+	PropertyChain.SetActiveMemberPropertyNode(PropertyChain.GetTail()->GetValue());
+
+	const FSimpleDelegate OnTransformChangedDelegate = FSimpleDelegate::CreateSP(this, &FRigComputedTransformDetails::OnTransformChanged, &PropertyChain);
+	TransformHandle->SetOnPropertyValueChanged(OnTransformChangedDelegate);
+	TransformHandle->SetOnChildPropertyValueChanged(OnTransformChangedDelegate);
+}
+
+void FRigComputedTransformDetails::OnTransformChanged(FEditPropertyChain* InPropertyChain)
+{
+	if(BlueprintBeingCustomized && InPropertyChain)
+	{
+		if(InPropertyChain->Num() > 1)
+		{
+			FPropertyChangedEvent ChangeEvent(InPropertyChain->GetHead()->GetValue(), EPropertyChangeType::ValueSet);
+			ChangeEvent.SetActiveMemberProperty(InPropertyChain->GetTail()->GetValue());
+			FPropertyChangedChainEvent ChainEvent(*InPropertyChain, ChangeEvent);
+			BlueprintBeingCustomized->BroadcastPostEditChangeChainProperty(ChainEvent);
+		}
+	}
+}
+
+void FRigBaseElementDetails::CustomizeDetails(IDetailLayoutBuilder& DetailBuilder)
 {
-	DetailBuilder.HideProperty(TEXT("Index"), FRigElement::StaticStruct());
-	DetailBuilder.HideProperty(TEXT("Name"), FRigElement::StaticStruct());
+	DetailBuilder.HideProperty(TEXT("Key"), FRigBaseElement::StaticStruct());
+	DetailBuilder.HideProperty(TEXT("Index"), FRigBaseElement::StaticStruct());
+	DetailBuilder.HideProperty(TEXT("SubIndex"), FRigBaseElement::StaticStruct());
 
 	TArray<TSharedPtr<FStructOnScope>> StructsBeingCustomized;
 	DetailBuilder.GetStructsBeingCustomized(StructsBeingCustomized);
@@ -456,12 +542,12 @@ void FRigElementDetails::CustomizeDetails(IDetailLayoutBuilder& DetailBuilder)
 					BlueprintBeingCustomized = Cast<UControlRigBlueprint>(Rig->GetClass()->ClassGeneratedBy);
 					if(BlueprintBeingCustomized)
 					{
-						ContainerBeingCustomized = &BlueprintBeingCustomized->HierarchyContainer;
+						HierarchyBeingCustomized = BlueprintBeingCustomized->Hierarchy;
 						if (UControlRig* DebuggedControlRig = Cast<UControlRig>(BlueprintBeingCustomized->GetObjectBeingDebugged()))
 						{
 							if (!DebuggedControlRig->IsSetupModeEnabled())
 							{
-								ContainerBeingCustomized = DebuggedControlRig->GetHierarchy();
+								HierarchyBeingCustomized = DebuggedControlRig->GetHierarchy();
 							}
 						}
 						break;
@@ -469,21 +555,39 @@ void FRigElementDetails::CustomizeDetails(IDetailLayoutBuilder& DetailBuilder)
 				}
 			}
 
-			if (ContainerBeingCustomized)
+			if (HierarchyBeingCustomized)
 			{
-				ElementKeyBeingCustomized = ((const FRigElement*)StructBeingCustomized->GetStructMemory())->GetElementKey();
+				ElementKeyBeingCustomized = ((const FRigBaseElement*)StructBeingCustomized->GetStructMemory())->GetKey();
 				break;
 			}
 		}
 	}
+
+	IDetailCategoryBuilder& Category = DetailBuilder.EditCategory(TEXT("RigElement"));
+	Category.InitiallyCollapsed(false);
+
+	Category.AddCustomRow(FText::FromString(TEXT("Name")))
+    .NameContent()
+    [
+        SNew(STextBlock)
+        .Text(FText::FromString(TEXT("Name")))
+        .Font(IDetailLayoutBuilder::GetDetailFont())
+    ]
+    .ValueContent()
+    [
+        SNew(SEditableTextBox)
+        .Font(IDetailLayoutBuilder::GetDetailFont())
+        .Text(this, &FRigBaseElementDetails::GetName)
+        .OnTextCommitted(this, &FRigBaseElementDetails::SetName)
+    ];
 }
 
-void FRigElementDetails::SetName(const FText& InNewText, ETextCommit::Type InCommitType)
+void FRigBaseElementDetails::SetName(const FText& InNewText, ETextCommit::Type InCommitType)
 {
-	FRigHierarchyContainer* Hierarchy = nullptr;
+	URigHierarchy* Hierarchy = nullptr;
 	if (BlueprintBeingCustomized)
 	{
-		Hierarchy = &BlueprintBeingCustomized->HierarchyContainer;
+		Hierarchy = BlueprintBeingCustomized->Hierarchy;
 	}
 	else
 	{
@@ -492,355 +596,482 @@ void FRigElementDetails::SetName(const FText& InNewText, ETextCommit::Type InCom
 
 	if (Hierarchy)
 	{
-		switch (ElementKeyBeingCustomized.Type)
+		URigHierarchyController* Controller = Hierarchy->GetController(true);
+		check(Controller);
+		Controller->RenameElement(ElementKeyBeingCustomized, *InNewText.ToString(), true);
+	}
+}
+
+void FRigBaseElementDetails::OnStructContentsChanged(FProperty* InProperty, const TSharedRef<IPropertyUtilities> PropertyUtilities)
+{
+	const FPropertyChangedEvent ChangeEvent(InProperty, EPropertyChangeType::ValueSet);
+	PropertyUtilities->NotifyFinishedChangingProperties(ChangeEvent);
+}
+
+void FRigTransformElementDetails::CustomizeDetails(IDetailLayoutBuilder& DetailBuilder)
+{
+	FRigBaseElementDetails::CustomizeDetails(DetailBuilder);
+
+	IDetailCategoryBuilder& Category = DetailBuilder.EditCategory(TEXT("Pose"));
+
+	Category.InitiallyCollapsed(ElementKeyBeingCustomized.Type != ERigElementType::Control);
+}
+
+void FRigBoneElementDetails::CustomizeDetails( IDetailLayoutBuilder& DetailBuilder )
+{
+	FRigTransformElementDetails::CustomizeDetails(DetailBuilder);
+}
+
+void FRigControlElementDetails_SetupBoolValueWidget(IDetailCategoryBuilder& InCategory, ERigControlValueType InValueType, FRigControlElement* InControlElement, URigHierarchy* InHierarchy)
+{
+	UEnum* ControlTypeEnum = StaticEnum<ERigControlType>();
+	UEnum* ValueTypeEnum = StaticEnum<ERigControlValueType>();
+
+	const FString ValueTypeName = ValueTypeEnum->GetDisplayNameTextByValue((int64)InValueType).ToString();
+	const FText PropertyLabel = FText::FromString(FString::Printf(TEXT("%s Value"), *ValueTypeName));
+	TWeakObjectPtr<URigHierarchy> HierarchyPtr = InHierarchy;
+	const FRigElementKey Key = InControlElement->GetKey();
+
+	InCategory.AddCustomRow(PropertyLabel)
+	.NameContent()
+	[
+	    SNew(STextBlock)
+		.Text(PropertyLabel)
+		.Font(IDetailLayoutBuilder::GetDetailFont())
+	]
+	.ValueContent()
+	[
+	    SNew(SVerticalBox)
+		+ SVerticalBox::Slot()
+	    [
+	        SNew(SCheckBox)
+	        .IsChecked_Lambda([HierarchyPtr, Key, InValueType]() -> ECheckBoxState
+	        {
+	        	if(HierarchyPtr.IsValid())
+	        	{
+	        		if(FRigControlElement* ControlElement = HierarchyPtr->Find<FRigControlElement>(Key))
+	        		{
+                        bool Value = HierarchyPtr->GetControlValue(ControlElement, InValueType).Get<bool>();
+						return Value ? ECheckBoxState::Checked : ECheckBoxState::Unchecked;
+                    }
+	        	}
+	        	return ECheckBoxState::Unchecked;
+	        })
+	        .OnCheckStateChanged_Lambda([HierarchyPtr, Key, InValueType](ECheckBoxState NewState)
+	        {
+	        	if(HierarchyPtr.IsValid())
+	        	{
+	        		if(FRigControlElement* ControlElement = HierarchyPtr->Find<FRigControlElement>(Key))
+	        		{
+                        const FRigControlValue Value = FRigControlValue::Make<bool>(NewState == ECheckBoxState::Checked);
+						HierarchyPtr->SetControlValue(ControlElement->GetKey(), Value, InValueType, true);
+                    }
+	        	}
+	        })
+	    ]
+	]
+	.IsEnabled(TAttribute<bool>::Create(TAttribute<bool>::FGetter::CreateLambda([HierarchyPtr, Key, InValueType]()->bool
+    {
+		if(HierarchyPtr.IsValid())
 		{
-			case ERigElementType::Bone:
-			{
-				ElementKeyBeingCustomized.Name = Hierarchy->BoneHierarchy.Rename(ElementKeyBeingCustomized.Name, *InNewText.ToString());
-				break;
-			}
-			case ERigElementType::Control:
-			{
-				ElementKeyBeingCustomized.Name = Hierarchy->ControlHierarchy.Rename(ElementKeyBeingCustomized.Name, *InNewText.ToString());
-				break;
-			}
-			case ERigElementType::Space:
-			{
-				ElementKeyBeingCustomized.Name = Hierarchy->SpaceHierarchy.Rename(ElementKeyBeingCustomized.Name, *InNewText.ToString());
-				break;
-			}
-			case ERigElementType::Curve:
+			if(FRigControlElement* ControlElement = HierarchyPtr->Find<FRigControlElement>(Key))
 			{
-				ElementKeyBeingCustomized.Name = Hierarchy->CurveContainer.Rename(ElementKeyBeingCustomized.Name, *InNewText.ToString());
-				break;
-			}
-			default:
-			{
-				ensure(false);
-				break;
+				return ControlElement->Settings.IsValueTypeEnabled(InValueType);
 			}
 		}
-	}
+		return false;
+    })));
 }
 
-float FRigElementDetails::GetTransformComponent(const FTransform& InTransform, ERigElementDetailsTransformComponent InComponent)
+void FRigControlElementDetails_SetupIntegerValueWidget(IDetailCategoryBuilder& InCategory, ERigControlValueType InValueType, FRigControlElement* InControlElement, URigHierarchy* InHierarchy)
 {
-	switch (InComponent)
-	{
-		case ERigElementDetailsTransformComponent::TranslationX:
-		{
-			return InTransform.GetTranslation().X;
-		}
-		case ERigElementDetailsTransformComponent::TranslationY:
-		{
-			return InTransform.GetTranslation().Y;
-		}
-		case ERigElementDetailsTransformComponent::TranslationZ:
-		{
-			return InTransform.GetTranslation().Z;
-		}
-		case ERigElementDetailsTransformComponent::RotationRoll:
-		{
-			return InTransform.GetRotation().Rotator().Roll;
-		}
-		case ERigElementDetailsTransformComponent::RotationPitch:
-		{
-			return InTransform.GetRotation().Rotator().Pitch;
-		}
-		case ERigElementDetailsTransformComponent::RotationYaw:
-		{
-			return InTransform.GetRotation().Rotator().Yaw;
-		}
-		case ERigElementDetailsTransformComponent::ScaleX:
-		{
-			return InTransform.GetScale3D().X;
-		}
-		case ERigElementDetailsTransformComponent::ScaleY:
-		{
-			return InTransform.GetScale3D().Y;
-		}
-		case ERigElementDetailsTransformComponent::ScaleZ:
+	UEnum* ControlTypeEnum = StaticEnum<ERigControlType>();
+	UEnum* ValueTypeEnum = StaticEnum<ERigControlValueType>();
+
+	const FString ValueTypeName = ValueTypeEnum->GetDisplayNameTextByValue((int64)InValueType).ToString();
+	const FText PropertyLabel = FText::FromString(FString::Printf(TEXT("%s Value"), *ValueTypeName));
+	TWeakObjectPtr<URigHierarchy> HierarchyPtr = InHierarchy;
+	const FRigElementKey Key = InControlElement->GetKey();
+
+	const TAttribute<bool> EnabledAttribute = TAttribute<bool>::Create(TAttribute<bool>::FGetter::CreateLambda([HierarchyPtr, Key, InValueType]()->bool
+    {
+		if(HierarchyPtr.IsValid())
 		{
-			return InTransform.GetScale3D().Z;
-		}
+            if(FRigControlElement* ControlElement = HierarchyPtr->Find<FRigControlElement>(Key))
+            {
+                return ControlElement->Settings.IsValueTypeEnabled(InValueType);
+            }
+        }
+		return false;
+    }));
+
+	const TAttribute<EVisibility> VisibilityAttribute = TAttribute<EVisibility>::Create(TAttribute<EVisibility>::FGetter::CreateLambda([EnabledAttribute]()->EVisibility
+	{
+		return EnabledAttribute.Get() ? EVisibility::Visible : EVisibility::Hidden;
+	}));
+
+	if (InControlElement->Settings.ControlEnum)
+	{
+		InCategory.AddCustomRow(PropertyLabel)
+		.Visibility(VisibilityAttribute)
+        .NameContent()
+        [
+            SNew(STextBlock)
+            .Text(PropertyLabel)
+            .Font(IDetailLayoutBuilder::GetDetailFont())
+        ]
+        .ValueContent()
+        .MinDesiredWidth(125.0f * 3.0f) // copied from FComponentTransformDetails
+        .MaxDesiredWidth(125.0f * 3.0f)
+        [
+            SNew(SVerticalBox)
+            + SVerticalBox::Slot()
+            [
+                SNew(SEnumComboBox, InControlElement->Settings.ControlEnum)
+                .CurrentValue_Lambda([HierarchyPtr, Key, InValueType]() -> int32
+                {
+                	if(HierarchyPtr.IsValid())
+                	{
+                		if(FRigControlElement* ControlElement = HierarchyPtr->Find<FRigControlElement>(Key))
+                		{
+                            return HierarchyPtr->GetControlValue(ControlElement, InValueType).Get<int32>();
+                        }
+                	}
+                	return 0;
+                })
+                .OnEnumSelectionChanged_Lambda([HierarchyPtr, Key, InValueType](int32 NewSelection, ESelectInfo::Type)
+                {
+                	if(HierarchyPtr.IsValid())
+                	{
+                		if(FRigControlElement* ControlElement = HierarchyPtr->Find<FRigControlElement>(Key))
+                		{
+                            const FRigControlValue Value = FRigControlValue::Make<int32>(NewSelection);
+                            HierarchyPtr->SetControlValue(ControlElement->GetKey(), Value, InValueType, true);
+                        }
+                	}
+                })
+                .Font(FEditorStyle::GetFontStyle(TEXT("MenuItem.Font")))
+            ]
+        ]
+		.IsEnabled(EnabledAttribute);
+	}
+	else
+	{
+		InCategory.AddCustomRow(PropertyLabel)
+		.Visibility(VisibilityAttribute)
+        .NameContent()
+        [
+            SNew(STextBlock)
+            .Text(PropertyLabel)
+            .Font(IDetailLayoutBuilder::GetDetailFont())
+        ]
+        .ValueContent()
+        .MinDesiredWidth(125.0f * 3.0f) // copied from FComponentTransformDetails
+        .MaxDesiredWidth(125.0f * 3.0f)
+        [
+            SNew(SVerticalBox)
+            + SVerticalBox::Slot()
+            [
+                SNew(SNumericEntryBox<int32>)
+                .Font(FEditorStyle::GetFontStyle(TEXT("MenuItem.Font")))
+                .AllowSpin(true)
+                .MinSliderValue_Lambda([HierarchyPtr, Key, InValueType]() -> TOptional<int32>
+				{
+				    if(InValueType == ERigControlValueType::Current || InValueType == ERigControlValueType::Initial)
+				    {
+				    	if(HierarchyPtr.IsValid())
+				    	{
+				    		if(FRigControlElement* ControlElement = HierarchyPtr->Find<FRigControlElement>(Key))
+				    		{
+                                return ControlElement->Settings.MinimumValue.Get<int32>();
+                            }
+				    	}
+				    }
+				    return TOptional<int32>();
+				})
+				.MaxSliderValue_Lambda([HierarchyPtr, Key, InValueType]() -> TOptional<int32>
+				{
+				    if(InValueType == ERigControlValueType::Current || InValueType == ERigControlValueType::Initial)
+				    {
+				    	if(HierarchyPtr.IsValid())
+				    	{
+				    		if(FRigControlElement* ControlElement = HierarchyPtr->Find<FRigControlElement>(Key))
+				    		{
+                                return ControlElement->Settings.MaximumValue.Get<int32>();
+                            }
+				    	}
+				    }
+				    return TOptional<int32>();
+				})
+                .Value_Lambda([HierarchyPtr, Key, InValueType]() -> int32
+                {
+                	if(HierarchyPtr.IsValid())
+                	{
+                		if(FRigControlElement* ControlElement = HierarchyPtr->Find<FRigControlElement>(Key))
+                		{
+                            return HierarchyPtr->GetControlValue(ControlElement, InValueType).Get<int32>();
+                        }
+                	}
+                	return 0;
+                })
+                .OnValueChanged_Lambda([HierarchyPtr, Key, InValueType](TOptional<int32> InNewSelection)
+                {
+                	if(InNewSelection.IsSet())
+                	{
+                		if(HierarchyPtr.IsValid())
+                		{
+                			if(FRigControlElement* ControlElement = HierarchyPtr->Find<FRigControlElement>(Key))
+                			{
+                                const FRigControlValue Value = FRigControlValue::Make<int32>(InNewSelection.GetValue());
+								HierarchyPtr->SetControlValue(ControlElement->GetKey(), Value, InValueType, true);
+                            }
+                		}
+                	}
+                })
+            ]
+        ]
+		.IsEnabled(EnabledAttribute);
 	}
-	return 0.f;
 }
 
-void FRigElementDetails::SetTransformComponent(FTransform& OutTransform, ERigElementDetailsTransformComponent InComponent, float InNewValue)
+void FRigControlElementDetails_SetupFloatValueWidget(IDetailCategoryBuilder& InCategory, ERigControlValueType InValueType, FRigControlElement* InControlElement, URigHierarchy* InHierarchy)
 {
-	switch (InComponent)
-	{
-		case ERigElementDetailsTransformComponent::TranslationX:
-		{
-			FVector Translation = OutTransform.GetTranslation();
-			Translation.X = InNewValue;
-			OutTransform.SetTranslation(Translation);
-			break;
-		}
-		case ERigElementDetailsTransformComponent::TranslationY:
-		{
-			FVector Translation = OutTransform.GetTranslation();
-			Translation.Y = InNewValue;
-			OutTransform.SetTranslation(Translation);
-			break;
-		}
-		case ERigElementDetailsTransformComponent::TranslationZ:
-		{
-			FVector Translation = OutTransform.GetTranslation();
-			Translation.Z = InNewValue;
-			OutTransform.SetTranslation(Translation);
-			break;
-		}
-		case ERigElementDetailsTransformComponent::RotationRoll:
-		{
-			FRotator Rotator = OutTransform.Rotator();
-			Rotator.Roll = InNewValue;
-			OutTransform.SetRotation(FQuat(Rotator));
-			break;
-		}
-		case ERigElementDetailsTransformComponent::RotationPitch:
-		{
-			FRotator Rotator = OutTransform.Rotator();
-			Rotator.Pitch = InNewValue;
-			OutTransform.SetRotation(FQuat(Rotator));
-			break;
-		}
-		case ERigElementDetailsTransformComponent::RotationYaw:
-		{
-			FRotator Rotator = OutTransform.Rotator();
-			Rotator.Yaw = InNewValue;
-			OutTransform.SetRotation(FQuat(Rotator));
-			break;
-		}
-		case ERigElementDetailsTransformComponent::ScaleX:
-		{
-			FVector Scale = OutTransform.GetScale3D();
-			Scale.X = InNewValue;
-			OutTransform.SetScale3D(Scale);
-			break;
-		}
-		case ERigElementDetailsTransformComponent::ScaleY:
-		{
-			FVector Scale = OutTransform.GetScale3D();
-			Scale.Y = InNewValue;
-			OutTransform.SetScale3D(Scale);
-			break;
-		}
-		case ERigElementDetailsTransformComponent::ScaleZ:
-		{
-			FVector Scale = OutTransform.GetScale3D();
-			Scale.Z = InNewValue;
-			OutTransform.SetScale3D(Scale);
-			break;
-		}
-	}
+	UEnum* ControlTypeEnum = StaticEnum<ERigControlType>();
+	UEnum* ValueTypeEnum = StaticEnum<ERigControlValueType>();
+
+	const FString ValueTypeName = ValueTypeEnum->GetDisplayNameTextByValue((int64)InValueType).ToString();
+	const FText PropertyLabel = FText::FromString(FString::Printf(TEXT("%s Value"), *ValueTypeName));
+	TWeakObjectPtr<URigHierarchy> HierarchyPtr = InHierarchy;
+	const FRigElementKey Key = InControlElement->GetKey();
+
+	const TAttribute<bool> EnabledAttribute = TAttribute<bool>::Create(TAttribute<bool>::FGetter::CreateLambda([HierarchyPtr, Key, InValueType]()->bool
+    {
+        if(HierarchyPtr.IsValid())
+        {
+            if(FRigControlElement* ControlElement = HierarchyPtr->Find<FRigControlElement>(Key))
+            {
+                return ControlElement->Settings.IsValueTypeEnabled(InValueType);
+            }
+        }
+        return false;
+    }));
+
+	const TAttribute<EVisibility> VisibilityAttribute = TAttribute<EVisibility>::Create(TAttribute<EVisibility>::FGetter::CreateLambda([EnabledAttribute]()->EVisibility
+    {
+        return EnabledAttribute.Get() ? EVisibility::Visible : EVisibility::Hidden;
+    }));
+
+	InCategory.AddCustomRow(PropertyLabel)
+	.Visibility(VisibilityAttribute)
+    .NameContent()
+    [
+        SNew(STextBlock)
+        .Text(PropertyLabel)
+        .Font(IDetailLayoutBuilder::GetDetailFont())
+    ]
+    .ValueContent()
+    .MinDesiredWidth(125.0f * 3.0f) // copied from FComponentTransformDetails
+    .MaxDesiredWidth(125.0f * 3.0f)
+    [
+        SNew(SVerticalBox)
+        + SVerticalBox::Slot()
+        [
+            SNew(SNumericEntryBox<float>)
+            .Font(FEditorStyle::GetFontStyle(TEXT("MenuItem.Font")))
+            .AllowSpin(true)
+            .Value_Lambda([HierarchyPtr, Key, InValueType]() -> float
+            {
+            	if(HierarchyPtr.IsValid())
+            	{
+            		if(FRigControlElement* ControlElement = HierarchyPtr->Find<FRigControlElement>(Key))
+            		{
+                        return HierarchyPtr->GetControlValue(ControlElement, InValueType).Get<float>();
+                    }
+            	}
+            	return 0.f;
+            })
+            .MinSliderValue_Lambda([HierarchyPtr, Key, InValueType]() -> TOptional<float>
+            {
+                if(InValueType == ERigControlValueType::Current || InValueType == ERigControlValueType::Initial)
+                {
+                	if(HierarchyPtr.IsValid())
+                	{
+                		if(FRigControlElement* ControlElement = HierarchyPtr->Find<FRigControlElement>(Key))
+                		{
+                            return ControlElement->Settings.MinimumValue.Get<float>();
+                        }
+                	}
+                }
+                return TOptional<float>();
+            })
+            .MaxSliderValue_Lambda([HierarchyPtr, Key, InValueType]() -> TOptional<float>
+            {
+                if(InValueType == ERigControlValueType::Current || InValueType == ERigControlValueType::Initial)
+                {
+                	if(HierarchyPtr.IsValid())
+                	{
+                		if(FRigControlElement* ControlElement = HierarchyPtr->Find<FRigControlElement>(Key))
+                		{
+                            return ControlElement->Settings.MaximumValue.Get<float>();
+                        }
+                	}
+                }
+                return TOptional<float>();
+            })
+            .OnValueChanged_Lambda([HierarchyPtr, Key, InValueType](TOptional<float> InNewSelection)
+            {
+            	if(InNewSelection.IsSet())
+            	{
+            		if(HierarchyPtr.IsValid())
+            		{
+            			if(FRigControlElement* ControlElement = HierarchyPtr->Find<FRigControlElement>(Key))
+            			{
+                            const FRigControlValue Value = FRigControlValue::Make<float>(InNewSelection.GetValue());
+                            HierarchyPtr->SetControlValue(ControlElement->GetKey(), Value, InValueType, true);
+                        }
+            		}
+            	}
+            })
+        ]
+    ]
+	.IsEnabled(EnabledAttribute);
 }
 
-float FRigElementDetails::GetEulerTransformComponent(const FEulerTransform& InTransform, ERigElementDetailsTransformComponent InComponent)
+template<typename T>
+void FRigControlElementDetails_SetupStructValueWidget(IDetailCategoryBuilder& InCategory, ERigControlValueType InValueType, FRigControlElement* InControlElement, URigHierarchy* InHierarchy)
 {
-	switch (InComponent)
-	{
-		case ERigElementDetailsTransformComponent::TranslationX:
-		{
-			return InTransform.Location.X;
-		}
-		case ERigElementDetailsTransformComponent::TranslationY:
-		{
-			return InTransform.Location.Y;
-		}
-		case ERigElementDetailsTransformComponent::TranslationZ:
-		{
-			return InTransform.Location.Z;
-		}
-		case ERigElementDetailsTransformComponent::RotationRoll:
-		{
-			return InTransform.Rotation.Roll;
-		}
-		case ERigElementDetailsTransformComponent::RotationPitch:
-		{
-			return InTransform.Rotation.Pitch;
-		}
-		case ERigElementDetailsTransformComponent::RotationYaw:
-		{
-			return InTransform.Rotation.Yaw;
-		}
-		case ERigElementDetailsTransformComponent::ScaleX:
-		{
-			return InTransform.Scale.X;
-		}
-		case ERigElementDetailsTransformComponent::ScaleY:
+	UEnum* ControlTypeEnum = StaticEnum<ERigControlType>();
+	UEnum* ValueTypeEnum = StaticEnum<ERigControlValueType>();
+
+	const FString ValueTypeName = ValueTypeEnum->GetDisplayNameTextByValue((int64)InValueType).ToString();
+	const FText PropertyLabel = FText::FromString(FString::Printf(TEXT("%s Value"), *ValueTypeName));
+	const UStruct* ValueStruct = TBaseStructure<T>::Get();
+
+	const TSharedPtr<FStructOnScope> StructToDisplay = MakeShareable(new FStructOnScope(ValueStruct));
+
+	TWeakObjectPtr<URigHierarchy> HierarchyPtr = InHierarchy;
+	const FRigElementKey Key = InControlElement->GetKey();
+
+	const TAttribute<bool> EnabledAttribute = TAttribute<bool>::Create(TAttribute<bool>::FGetter::CreateLambda([HierarchyPtr, Key, InValueType]()->bool
+    {
+        if(HierarchyPtr.IsValid())
+        {
+            if(FRigControlElement* ControlElement = HierarchyPtr->Find<FRigControlElement>(Key))
+            {
+                return ControlElement->Settings.IsValueTypeEnabled(InValueType);
+            }
+        }
+        return false;
+    }));
+
+	const TAttribute<EVisibility> VisibilityAttribute = TAttribute<EVisibility>::Create(TAttribute<EVisibility>::FGetter::CreateLambda([HierarchyPtr, Key, InValueType]()->EVisibility
+    {
+		if(HierarchyPtr.IsValid())
 		{
-			return InTransform.Scale.Y;
-		}
-		case ERigElementDetailsTransformComponent::ScaleZ:
+			if(FRigControlElement* ControlElement = HierarchyPtr->Find<FRigControlElement>(Key))
+			{
+                if(ControlElement->Settings.IsValueTypeEnabled(InValueType))
+                {
+                    return EVisibility::Visible;
+                }
+            }
+		}
+		return EVisibility::Hidden;
+    }));
+
+	IDetailPropertyRow* Row = InCategory.AddExternalStructure(StructToDisplay);
+	Row->DisplayName(PropertyLabel);
+	Row->ShouldAutoExpand(true);
+	Row->IsEnabled(EnabledAttribute);
+	Row->Visibility(VisibilityAttribute);
+
+	TSharedPtr<SWidget> NameWidget, ValueWidget;
+	Row->GetDefaultWidgets(NameWidget, ValueWidget);
+	
+	const FSimpleDelegate OnStructContentsChangedDelegate = FSimpleDelegate::CreateLambda([HierarchyPtr, Key, StructToDisplay, InValueType]()
+	{
+		if(HierarchyPtr.IsValid())
 		{
-			return InTransform.Scale.Z;
+            const FRigControlValue Value = FRigControlValue::Make(*(T*)StructToDisplay->GetStructMemory());
+			HierarchyPtr->SetControlValue(Key, Value, InValueType, true);
 		}
-	}
-	return 0.f;
+	});
+
+	TSharedPtr<IPropertyHandle> Handle = Row->GetPropertyHandle();
+	Handle->SetOnPropertyValueChanged(OnStructContentsChangedDelegate);
+	Handle->SetOnChildPropertyValueChanged(OnStructContentsChangedDelegate);
 }
 
-void FRigElementDetails::SetEulerTransformComponent(FEulerTransform& OutTransform, ERigElementDetailsTransformComponent InComponent, float InNewValue)
+void FRigControlElementDetails_SetupValueWidget(IDetailCategoryBuilder& InCategory, ERigControlValueType InValueType, FRigControlElement* InControlElement, URigHierarchy* InHierarchy)
 {
-	switch (InComponent)
+	switch(InControlElement->Settings.ControlType)
 	{
-		case ERigElementDetailsTransformComponent::TranslationX:
+		case ERigControlType::Bool:
 		{
-			FVector Translation = OutTransform.Location;
-			Translation.X = InNewValue;
-			OutTransform.Location = Translation;
+			if((InValueType == ERigControlValueType::Minimum) || (InValueType == ERigControlValueType::Maximum))
+			{
+				return;
+			}
+			FRigControlElementDetails_SetupBoolValueWidget(InCategory, InValueType, InControlElement, InHierarchy);
+			break;
+		}
+		case ERigControlType::Integer:
+		{
+			FRigControlElementDetails_SetupIntegerValueWidget(InCategory, InValueType, InControlElement, InHierarchy);
 			break;
 		}
-		case ERigElementDetailsTransformComponent::TranslationY:
+		case ERigControlType::Float:
 		{
-			FVector Translation = OutTransform.Location;
-			Translation.Y = InNewValue;
-			OutTransform.Location = Translation;
+			FRigControlElementDetails_SetupFloatValueWidget(InCategory, InValueType, InControlElement, InHierarchy);
 			break;
 		}
-		case ERigElementDetailsTransformComponent::TranslationZ:
+		case ERigControlType::Vector2D:
 		{
-			FVector Translation = OutTransform.Location;
-			Translation.Z = InNewValue;
-			OutTransform.Location = Translation;
+			FRigControlElementDetails_SetupStructValueWidget<FVector2D>(InCategory, InValueType, InControlElement, InHierarchy);
 			break;
 		}
-		case ERigElementDetailsTransformComponent::RotationRoll:
+		case ERigControlType::Position:
+		case ERigControlType::Scale:
 		{
-			FRotator Rotator = OutTransform.Rotation;
-			Rotator.Roll = InNewValue;
-			OutTransform.Rotation = Rotator;
+			FRigControlElementDetails_SetupStructValueWidget<FVector>(InCategory, InValueType, InControlElement, InHierarchy);
 			break;
 		}
-		case ERigElementDetailsTransformComponent::RotationPitch:
+		case ERigControlType::Rotator:
 		{
-			FRotator Rotator = OutTransform.Rotation;
-			Rotator.Pitch = InNewValue;
-			OutTransform.Rotation = Rotator;
+			FRigControlElementDetails_SetupStructValueWidget<FRotator>(InCategory, InValueType, InControlElement, InHierarchy);
 			break;
 		}
-		case ERigElementDetailsTransformComponent::RotationYaw:
+		case ERigControlType::TransformNoScale:
 		{
-			FRotator Rotator = OutTransform.Rotation;
-			Rotator.Yaw = InNewValue;
-			OutTransform.Rotation = Rotator;
+			FRigControlElementDetails_SetupStructValueWidget<FTransformNoScale>(InCategory, InValueType, InControlElement, InHierarchy);
 			break;
 		}
-		case ERigElementDetailsTransformComponent::ScaleX:
+		case ERigControlType::EulerTransform:
 		{
-			FVector Scale = OutTransform.Scale;
-			Scale.X = InNewValue;
-			OutTransform.Scale = Scale;
+			FRigControlElementDetails_SetupStructValueWidget<FEulerTransform>(InCategory, InValueType, InControlElement, InHierarchy);
 			break;
 		}
-		case ERigElementDetailsTransformComponent::ScaleY:
+		case ERigControlType::Transform:
 		{
-			FVector Scale = OutTransform.Scale;
-			Scale.Y = InNewValue;
-			OutTransform.Scale = Scale;
+			FRigControlElementDetails_SetupStructValueWidget<FTransform>(InCategory, InValueType, InControlElement, InHierarchy);
 			break;
 		}
-		case ERigElementDetailsTransformComponent::ScaleZ:
+		default:
 		{
-			FVector Scale = OutTransform.Scale;
-			Scale.Z = InNewValue;
-			OutTransform.Scale = Scale;
+			ensure(false);
 			break;
 		}
 	}
 }
 
+TArray<TSharedPtr<FString>> FRigControlElementDetails::ControlTypeList;
 
-void FRigBoneDetails::CustomizeDetails( IDetailLayoutBuilder& DetailBuilder )
-{
-	FRigElementDetails::CustomizeDetails(DetailBuilder);
-	DetailBuilder.HideProperty(TEXT("ParentName"));
-
-	IDetailCategoryBuilder& Category = DetailBuilder.EditCategory(TEXT("FRigElement"), LOCTEXT("BoneCategory", "Bone"));
-	Category.InitiallyCollapsed(false);
-
-	Category.AddCustomRow(FText::FromString(TEXT("Name")))
-	.NameContent()
-	[
-		SNew(STextBlock)
-		.Text(FText::FromString(TEXT("Name")))
-		.Font(IDetailLayoutBuilder::GetDetailFont())
-	]
-	.ValueContent()
-	[
-		SNew(SEditableTextBox)
-		.Font(IDetailLayoutBuilder::GetDetailFont())
-		.Text(this, &FRigElementDetails::GetName)
-		.OnTextCommitted(this, &FRigElementDetails::SetName)
-	];
-
-	// setup influences modifier
-	{
-		FName InfluencesEventName = FRigUnit_BeginExecution::EventName;
-		FRigInfluenceMap InfluenceMap = BlueprintBeingCustomized->Influences.FindOrAdd(InfluencesEventName).Inverse();
-		FRigInfluenceEntryModifier TempModifier = InfluenceMap.GetEntryModifier(ElementKeyBeingCustomized);
-		InfluenceModifier = MakeShareable(new FRigInfluenceEntryModifier(TempModifier));
-		InfluenceModifierStruct = MakeShareable(new FStructOnScope(FRigInfluenceEntryModifier::StaticStruct(), (uint8*)InfluenceModifier.Get()));
-		InfluenceModifierStruct->SetPackage(BlueprintBeingCustomized->GetOutermost());
-
-		IDetailCategoryBuilder& InversionCategory = DetailBuilder.EditCategory(TEXT("Inversion"), LOCTEXT("InversionCategory", "Inversion"));
-		FProperty* AffectedListProperty = FRigInfluenceEntryModifier::StaticStruct()->FindPropertyByName(TEXT("AffectedList"));
-
-		IDetailPropertyRow* AffectedListRow = InversionCategory.AddExternalStructureProperty(InfluenceModifierStruct, AffectedListProperty->GetFName());
-		AffectedListRow->ShouldAutoExpand(true);
-		AffectedListRow->DisplayName(LOCTEXT("AffectingElements", "Affecting Elements"));
-
-		FSimpleDelegate OnInfluencesChangedDelegate = FSimpleDelegate::CreateSP(this, &FRigBoneDetails::OnAffectedListChanged, AffectedListProperty, DetailBuilder.GetPropertyUtilities());
-
-		TSharedPtr<IPropertyHandle> AffectedListHandle = AffectedListRow->GetPropertyHandle();
-		AffectedListHandle->SetOnPropertyValueChanged(OnInfluencesChangedDelegate);
-		AffectedListHandle->SetOnChildPropertyValueChanged(OnInfluencesChangedDelegate);
-	}
-}
-
-void FRigBoneDetails::OnStructContentsChanged(FProperty* InProperty, const TSharedRef<IPropertyUtilities> PropertyUtilities)
-{
-	FPropertyChangedEvent ChangeEvent(InProperty, EPropertyChangeType::ValueSet);
-	PropertyUtilities->NotifyFinishedChangingProperties(ChangeEvent);
-}
-
-void FRigBoneDetails::OnAffectedListChanged(FProperty* InProperty, const TSharedRef<IPropertyUtilities> PropertyUtilities)
+void FRigControlElementDetails::CustomizeDetails(IDetailLayoutBuilder& DetailBuilder)
 {
-	if (BlueprintBeingCustomized != nullptr && InfluenceModifier.IsValid())
-	{
-		BlueprintBeingCustomized->Modify();
-
-		FName InfluencesEventName = FRigUnit_BeginExecution::EventName;
-		FRigInfluenceMap InverseInfluenceMap = BlueprintBeingCustomized->Influences.FindOrAdd(InfluencesEventName).Inverse();
-		InverseInfluenceMap.SetEntryModifier(ElementKeyBeingCustomized, *InfluenceModifier);
-		BlueprintBeingCustomized->Influences.FindOrAdd(InfluencesEventName) = InverseInfluenceMap.Inverse();
-	}
-
-	OnStructContentsChanged(InProperty, PropertyUtilities);
-}
-
-TArray<TSharedPtr<FString>> FRigControlDetails::ControlTypeList;
+	FRigTransformElementDetails::CustomizeDetails(DetailBuilder);
 
-void FRigControlDetails::CustomizeDetails(IDetailLayoutBuilder& DetailBuilder)
-{
-	FRigElementDetails::CustomizeDetails(DetailBuilder);
-
-	DetailBuilder.HideProperty(TEXT("ControlType"));
-	DetailBuilder.HideProperty(TEXT("DisplayName"));
-	DetailBuilder.HideProperty(TEXT("ParentName"));
-	DetailBuilder.HideProperty(TEXT("SpaceName"));
-	DetailBuilder.HideProperty(TEXT("OffsetTransform"));
-	DetailBuilder.HideProperty(TEXT("InitialValue"));
-	DetailBuilder.HideProperty(TEXT("Value"));
-	DetailBuilder.HideProperty(TEXT("bLimitTranslation"));
-	DetailBuilder.HideProperty(TEXT("bLimitRotation"));
-	DetailBuilder.HideProperty(TEXT("bLimitScale"));
-	DetailBuilder.HideProperty(TEXT("bDrawLimits"));
-	DetailBuilder.HideProperty(TEXT("MinimumValue"));
-	DetailBuilder.HideProperty(TEXT("MaximumValue"));
-	DetailBuilder.HideProperty(TEXT("bIsTransientControl"));
-	DetailBuilder.HideProperty(TEXT("ControlEnum"));
+	//DetailBuilder.HideProperty(TEXT("Offset"));
+	//DetailBuilder.HideProperty(TEXT("Gizmo"));
 
 	GizmoNameList.Reset();
 	if (BlueprintBeingCustomized)
@@ -859,6 +1090,11 @@ void FRigControlDetails::CustomizeDetails(IDetailLayoutBuilder& DetailBuilder)
 		}
 	}
 
+	if (HierarchyBeingCustomized == nullptr || !ElementKeyBeingCustomized)
+	{
+		return;
+	}
+
 	IDetailCategoryBuilder& ControlCategory = DetailBuilder.EditCategory(TEXT("Control"), LOCTEXT("ControlCategory", "Control"));
 	IDetailCategoryBuilder& LimitsCategory = DetailBuilder.EditCategory(TEXT("Limits"), LOCTEXT("LimitsCategory", "Limits"));
 	IDetailCategoryBuilder& GizmoCategory = DetailBuilder.EditCategory(TEXT("Gizmo"), LOCTEXT("GizmoCategory", "Gizmo"));
@@ -867,42 +1103,24 @@ void FRigControlDetails::CustomizeDetails(IDetailLayoutBuilder& DetailBuilder)
 	LimitsCategory.InitiallyCollapsed(false);
 	GizmoCategory.InitiallyCollapsed(false);
 
-	ControlCategory.AddCustomRow(FText::FromString(TEXT("Name")))
-	.NameContent()
-	[
-		SNew(STextBlock)
-		.Text(FText::FromString(TEXT("Name")))
-		.Font(IDetailLayoutBuilder::GetDetailFont())
-	]
-	.ValueContent()
-	[
-		SNew(SEditableTextBox)
-		.Font(IDetailLayoutBuilder::GetDetailFont())
-		.Text(this, &FRigElementDetails::GetName)
-		.OnTextCommitted(this, &FRigElementDetails::SetName)
-	];
-
-	ControlCategory.AddCustomRow(FText::FromString(TEXT("DisplayName")))
+	const TSharedRef<IPropertyHandle> DisplayNameHandle = DetailBuilder.GetProperty(TEXT("Settings.DisplayName"));
+	DetailBuilder.HideProperty(DisplayNameHandle);
+	ControlCategory.AddProperty(DisplayNameHandle).CustomWidget()
 	.NameContent()
-	[
-		SNew(STextBlock)
-		.Text(FText::FromString(TEXT("Display Name")))
-		.Font(IDetailLayoutBuilder::GetDetailFont())
-	]
+    [
+        SNew(STextBlock)
+        .Text(FText::FromString(TEXT("Display Name")))
+        .Font(IDetailLayoutBuilder::GetDetailFont())
+    ]
 	.ValueContent()
-	[
+    [
 		SNew(SEditableTextBox)
 		.Font(IDetailLayoutBuilder::GetDetailFont())
-		.Text(this, &FRigControlDetails::GetDisplayName)
-		.OnTextCommitted(this, &FRigControlDetails::SetDisplayName, DetailBuilder.GetPropertyUtilities())
+		.Text(this, &FRigControlElementDetails::GetDisplayName)
+		.OnTextCommitted(this, &FRigControlElementDetails::SetDisplayName, DetailBuilder.GetPropertyUtilities())
 	];
 
-	if (ContainerBeingCustomized == nullptr || !ElementKeyBeingCustomized)
-	{
-		return;
-	}
-
-	if (ControlTypeList.Num() == 0)
+	if (ControlTypeList.Num() == 0)
 	{
 		UEnum* Enum = StaticEnum<ERigControlType>();
 		for (int64 Index = 0; Index < Enum->GetMaxEnumValue(); Index++)
@@ -911,90 +1129,97 @@ void FRigControlDetails::CustomizeDetails(IDetailLayoutBuilder& DetailBuilder)
 		}
 	}
 
-	UEnum* ControlTypeEnum = StaticEnum<ERigControlType>();
-	UEnum* ValueTypeEnum = StaticEnum<ERigControlValueType>();
-	FRigControl& Control = ContainerBeingCustomized->ControlHierarchy[ElementKeyBeingCustomized.Name];
+	FRigControlElement* ControlElement = HierarchyBeingCustomized->FindChecked<FRigControlElement>(ElementKeyBeingCustomized);
 
 	// when control type changes, we have to refresh detail panel
-	TSharedRef<IPropertyHandle> ControlType = DetailBuilder.GetProperty(TEXT("ControlType"));
-	ControlType->SetOnPropertyValueChanged(FSimpleDelegate::CreateLambda(
+	const TSharedRef<IPropertyHandle> ControlTypeHandle = DetailBuilder.GetProperty(TEXT("Settings.ControlType"));
+	ControlTypeHandle->SetOnPropertyValueChanged(FSimpleDelegate::CreateLambda(
 		[this, &DetailBuilder]()
 		{
 			DetailBuilder.ForceRefreshDetails();
 
-			if (this->ContainerBeingCustomized && this->ElementKeyBeingCustomized.IsValid())
+			if (this->HierarchyBeingCustomized && this->ElementKeyBeingCustomized.IsValid())
 			{
-				FRigControl& Control = this->ContainerBeingCustomized->ControlHierarchy[this->ElementKeyBeingCustomized.Name];
+				FRigControlElement* ControlElement = this->HierarchyBeingCustomized->FindChecked<FRigControlElement>(ElementKeyBeingCustomized);
+				FRigControlValue ValueToSet;
+
+				ControlElement->Settings.bLimitTranslation = false;
+				ControlElement->Settings.bLimitRotation = false;
+				ControlElement->Settings.bLimitScale = false;
 
-				switch (Control.ControlType)
+				switch (ControlElement->Settings.ControlType)
 				{
 					case ERigControlType::Bool:
 					{
-						Control.Value = Control.InitialValue = FRigControlValue::Make<bool>(false);
+						ValueToSet = FRigControlValue::Make<bool>(false);
 						break;
 					}
 					case ERigControlType::Float:
 					{
-						Control.Value = Control.InitialValue = FRigControlValue::Make<float>(0.f);
-						Control.bLimitTranslation = true;
-						Control.MinimumValue = FRigControlValue::Make<float>(0.f);
-						Control.MaximumValue = FRigControlValue::Make<float>(1.f);
-						break;
+						ValueToSet = FRigControlValue::Make<float>(0.f);
+						ControlElement->Settings.bLimitTranslation = true;
+						ControlElement->Settings.MinimumValue = FRigControlValue::Make<float>(0.f);
+						ControlElement->Settings.MaximumValue = FRigControlValue::Make<float>(1.f);
+                        break;
 					}
 					case ERigControlType::Integer:
 					{
-						Control.Value = Control.InitialValue = FRigControlValue::Make<int32>(0);
-						Control.bLimitTranslation = true;
-						Control.MinimumValue = FRigControlValue::Make<int32>(0);
-						Control.MaximumValue = FRigControlValue::Make<int32>(10);
+						ValueToSet = FRigControlValue::Make<int32>(0);
+						ControlElement->Settings.bLimitTranslation = true;
+						ControlElement->Settings.MinimumValue = FRigControlValue::Make<int32>(0);
+						ControlElement->Settings.MaximumValue = FRigControlValue::Make<int32>(10);
 						break;
 					}
 					case ERigControlType::Vector2D:
 					{
-						Control.Value = Control.InitialValue = FRigControlValue::Make<FVector2D>(FVector2D::ZeroVector);
-						Control.bLimitTranslation = true;
-						Control.MinimumValue = FRigControlValue::Make<FVector2D>(FVector2D::ZeroVector);
-						Control.MaximumValue = FRigControlValue::Make<FVector2D>(FVector2D(1.f, 1.f));
+						ValueToSet = FRigControlValue::Make<FVector2D>(FVector2D::ZeroVector);
+						ControlElement->Settings.bLimitTranslation = true;
+						ControlElement->Settings.MinimumValue = FRigControlValue::Make<FVector2D>(FVector2D::ZeroVector);
+						ControlElement->Settings.MaximumValue = FRigControlValue::Make<FVector2D>(FVector2D(1.f, 1.f));
 						break;
 					}
 					case ERigControlType::Position:
 					{
-						Control.Value = Control.InitialValue = FRigControlValue::Make<FVector>(FVector::ZeroVector);
-						Control.MinimumValue = FRigControlValue::Make<FVector>(-FVector::OneVector);
-						Control.MaximumValue = FRigControlValue::Make<FVector>(FVector::OneVector);
+						ValueToSet = FRigControlValue::Make<FVector>(FVector::ZeroVector);
+						ControlElement->Settings.MinimumValue = FRigControlValue::Make<FVector>(-FVector::OneVector);
+						ControlElement->Settings.MaximumValue = FRigControlValue::Make<FVector>(FVector::OneVector);
 						break;
 					}
 					case ERigControlType::Scale:
 					{
-						Control.Value = Control.InitialValue = FRigControlValue::Make<FVector>(FVector::OneVector);
-						Control.MinimumValue = FRigControlValue::Make<FVector>(FVector::ZeroVector);
-						Control.MaximumValue = FRigControlValue::Make<FVector>(FVector::OneVector);
+						ValueToSet = FRigControlValue::Make<FVector>(FVector::OneVector);
+						ControlElement->Settings.MinimumValue = FRigControlValue::Make<FVector>(FVector::ZeroVector);
+						ControlElement->Settings.MaximumValue = FRigControlValue::Make<FVector>(FVector::OneVector);
 						break;
 					}
 					case ERigControlType::Rotator:
 					{
-						Control.Value = Control.InitialValue = FRigControlValue::Make<FRotator>(FRotator::ZeroRotator);
-						Control.MinimumValue = FRigControlValue::Make<FRotator>(FRotator::ZeroRotator);
-						Control.MaximumValue = FRigControlValue::Make<FRotator>(FRotator(180.f, 180.f, 180.f));
+						ValueToSet = FRigControlValue::Make<FRotator>(FRotator::ZeroRotator);
+						ControlElement->Settings.MinimumValue = FRigControlValue::Make<FRotator>(FRotator::ZeroRotator);
+						ControlElement->Settings.MaximumValue = FRigControlValue::Make<FRotator>(FRotator(180.f, 180.f, 180.f));
 						break;
 					}
 					case ERigControlType::Transform:
 					{
-						Control.Value = Control.InitialValue = Control.MinimumValue = Control.MaximumValue = FRigControlValue::Make<FTransform>(FTransform::Identity);
+						ValueToSet = FRigControlValue::Make<FTransform>(FTransform::Identity);
+						ControlElement->Settings.MinimumValue = ValueToSet;
+						ControlElement->Settings.MaximumValue = ValueToSet;
 						break;
 					}
 					case ERigControlType::TransformNoScale:
 					{
 						FTransformNoScale Identity = FTransform::Identity;
-						// Helge todo FTransformNoScale::Identity appears to be no good?
-						Control.Value = Control.InitialValue = Control.MinimumValue = Control.MaximumValue = FRigControlValue::Make<FTransformNoScale>(Identity);
-						break;
+						ValueToSet = FRigControlValue::Make<FTransformNoScale>(Identity);
+						ControlElement->Settings.MinimumValue = ValueToSet;
+                        ControlElement->Settings.MaximumValue = ValueToSet;
+                        break;
 					}
 					case ERigControlType::EulerTransform:
 					{
 						FEulerTransform Identity = FEulerTransform::Identity;
-						// Helge todo FTransformNoScale::Identity appears to be no good?
-						Control.Value = Control.InitialValue = Control.MinimumValue = Control.MaximumValue = FRigControlValue::Make<FEulerTransform>(Identity);
+						ValueToSet = FRigControlValue::Make<FEulerTransform>(Identity);
+						ControlElement->Settings.MinimumValue = ValueToSet;
+						ControlElement->Settings.MaximumValue = ValueToSet;
 						break;
 					}
 					default:
@@ -1004,397 +1229,85 @@ void FRigControlDetails::CustomizeDetails(IDetailLayoutBuilder& DetailBuilder)
 					}
 				}
 
-				if (this->ContainerBeingCustomized != &this->BlueprintBeingCustomized->HierarchyContainer)
+				this->HierarchyBeingCustomized->Notify(ERigHierarchyNotification::ControlSettingChanged, ControlElement);
+				this->HierarchyBeingCustomized->SetControlValue(ControlElement, ValueToSet, ERigControlValueType::Initial, true);
+				this->HierarchyBeingCustomized->SetControlValue(ControlElement, ValueToSet, ERigControlValueType::Current, true);
+
+				if (this->HierarchyBeingCustomized != this->BlueprintBeingCustomized->Hierarchy)
 				{
-					this->BlueprintBeingCustomized->HierarchyContainer.ControlHierarchy[Control.Index] = Control;
+					if(FRigControlElement* OtherControlElement = this->BlueprintBeingCustomized->Hierarchy->Find<FRigControlElement>(ControlElement->GetKey()))
+					{
+						OtherControlElement->Settings = ControlElement->Settings;
+						this->BlueprintBeingCustomized->Hierarchy->Notify(ERigHierarchyNotification::ControlSettingChanged, ControlElement);
+                        this->BlueprintBeingCustomized->Hierarchy->SetControlValue(ControlElement, ValueToSet, ERigControlValueType::Initial, true);
+                        this->BlueprintBeingCustomized->Hierarchy->SetControlValue(ControlElement, ValueToSet, ERigControlValueType::Current, true);
+                   }
 				}
-				this->BlueprintBeingCustomized->PropagateHierarchyFromBPToInstances(false, false);
-				this->ContainerBeingCustomized->ControlHierarchy.OnControlUISettingsChanged.Broadcast(this->ContainerBeingCustomized, this->ElementKeyBeingCustomized);
 			}
 		}
 	));
 
-	ControlCategory.AddCustomRow(FText::FromString(TEXT("ControlType")))
-	.NameContent()
-	[
-		ControlType->CreatePropertyNameWidget()
-	]
-	.ValueContent()
-	[
-		ControlType->CreatePropertyValueWidget()
-	];
-
-	if(Control.ControlType != ERigControlType::Bool)
-	{
-		TSharedRef<IPropertyHandle> OffsetTransform = DetailBuilder.GetProperty(TEXT("OffsetTransform"));
-		OffsetTransform->SetOnPropertyValueChanged(FSimpleDelegate::CreateLambda(
-			[this, OffsetTransform, &DetailBuilder]()
-			{
-				this->BlueprintBeingCustomized->PropagatePropertyFromBPToInstances(this->GetElementKey(), OffsetTransform->GetProperty());
-			}
-		));
-
-		uint8* ValuePtr = (uint8*)&Control.OffsetTransform;
-		const UStruct* ValueStruct = TBaseStructure<FTransform>::Get();
-
-		TSharedPtr<FStructOnScope> StructToDisplay = MakeShareable(new FStructOnScope(ValueStruct, ValuePtr));
-		FProperty* Property = OffsetTransform->GetProperty();
-
-		IDetailPropertyRow* Row = ControlCategory.AddExternalStructure(StructToDisplay);
-		Row->DisplayName(OffsetTransform->GetPropertyDisplayName());
-		Row->ToolTip(Property->GetToolTipText());
-		Row->ShouldAutoExpand(true);
-		Row->IsEnabled(true);
-
-		FSimpleDelegate OnStructContentsChangedDelegate = FSimpleDelegate::CreateSP(this, &FRigControlDetails::OnStructContentsChanged, Property, DetailBuilder.GetPropertyUtilities());
-
-		TSharedPtr<IPropertyHandle> Handle = Row->GetPropertyHandle();
-		Handle->SetOnPropertyValueChanged(OnStructContentsChangedDelegate);
-		Handle->SetOnChildPropertyValueChanged(OnStructContentsChangedDelegate);
-	}
-
-	FString ControlTypeName = ControlTypeEnum->GetDisplayNameTextByValue((int64)Control.ControlType).ToString();
-
-	switch (Control.ControlType)
+	struct Local
 	{
-		case ERigControlType::Bool:
+		static void OnTransformChanged(FEditPropertyChain* InPropertyChain, UControlRigBlueprint* InBlueprint)
 		{
-			for(int64 ValueTypeIndex = 0; ValueTypeIndex < ValueTypeEnum->GetMaxEnumValue(); ValueTypeIndex++)
+			if(InBlueprint && InPropertyChain)
 			{
-				ERigControlValueType ValueType = (ERigControlValueType)ValueTypeIndex;
-
-				if (ValueType == ERigControlValueType::Minimum ||
-					ValueType == ERigControlValueType::Maximum)
+				if(InPropertyChain->Num() > 1)
 				{
-					continue;
+					FPropertyChangedEvent ChangeEvent(InPropertyChain->GetHead()->GetValue(), EPropertyChangeType::ValueSet);
+					ChangeEvent.SetActiveMemberProperty(InPropertyChain->GetTail()->GetValue());
+					FPropertyChangedChainEvent ChainEvent(*InPropertyChain, ChangeEvent);
+					InBlueprint->BroadcastPostEditChangeChainProperty(ChainEvent);
 				}
-
-				FString ValueTypeName = ValueTypeEnum->GetDisplayNameTextByValue(ValueTypeIndex).ToString();
-				FText PropertyLabel = FText::FromString(FString::Printf(TEXT("%s %s"), *ValueTypeName, *ControlTypeName));
-
-				ControlCategory.AddCustomRow(PropertyLabel)
-				.NameContent()
-				.VAlign(VAlign_Top)
-				[
-					SNew(STextBlock)
-					.Text(PropertyLabel)
-					.Font(IDetailLayoutBuilder::GetDetailFont())
-				]
-				.ValueContent()
-				[
-					SNew(SVerticalBox)
-					+ SVerticalBox::Slot()
-					[
-						SNew(SCheckBox)
-						.IsChecked(this, &FRigControlDetails::GetComponentValueBool, ValueType == ERigControlValueType::Initial)
-						.OnCheckStateChanged(this, &FRigControlDetails::SetComponentValueBool, ValueType == ERigControlValueType::Initial, DetailBuilder.GetPropertyUtilities())
-					]
-				];
 			}
-			break;
 		}
-		case ERigControlType::Float:
-		{
-			for(int64 ValueTypeIndex = 0; ValueTypeIndex < ValueTypeEnum->GetMaxEnumValue(); ValueTypeIndex++)
-			{
-				ERigControlValueType ValueType = (ERigControlValueType)ValueTypeIndex;
-				FString ValueTypeName = ValueTypeEnum->GetDisplayNameTextByValue(ValueTypeIndex).ToString();
-				FText PropertyLabel = FText::FromString(FString::Printf(TEXT("%s %s"), *ValueTypeName, *ControlTypeName));
-
-				if (ValueType == ERigControlValueType::Minimum)
-				{
-					LimitsCategory.AddCustomRow(FText::FromString(TEXT("Limit")))
-					.NameContent()
-					[
-						SNew(STextBlock)
-						.Text(FText::FromString(TEXT("Limit")))
-						.Font(IDetailLayoutBuilder::GetDetailFont())
-					]
-					.ValueContent()
-					[
-						DetailBuilder.GetProperty(TEXT("bLimitTranslation"))->CreatePropertyValueWidget()
-					];
-
-					LimitsCategory.AddCustomRow(FText::FromString(TEXT("bDrawLimits")))
-					.NameContent()
-					[
-						DetailBuilder.GetProperty(TEXT("bDrawLimits"))->CreatePropertyNameWidget()
-					]
-					.ValueContent()
-					[
-						DetailBuilder.GetProperty(TEXT("bDrawLimits"))->CreatePropertyValueWidget()
-					];
-				}
+	};
 
-				IDetailCategoryBuilder& Category = (ValueType == ERigControlValueType::Minimum || ValueType == ERigControlValueType::Maximum) ? LimitsCategory : ControlCategory;
-				Category.AddCustomRow(PropertyLabel)
-				.NameContent()
-				.VAlign(VAlign_Top)
-				[
-					SNew(STextBlock)
-					.Text(PropertyLabel)
-					.Font(IDetailLayoutBuilder::GetDetailFont())
-					.IsEnabled(this, &FRigControlDetails::IsEnabled, ValueType)
-				]
-				.ValueContent()
-				.MinDesiredWidth(125.0f * 3.0f) // copied from FComponentTransformDetails
-				.MaxDesiredWidth(125.0f * 3.0f)
-				[
-					SNew(SVerticalBox)
-					+ SVerticalBox::Slot()
-					[
-						SNew(SNumericEntryBox<float>)
-						.Font(FEditorStyle::GetFontStyle(TEXT("MenuItem.Font")))
-						.AllowSpin(true)
-						.IsEnabled(this, &FRigControlDetails::IsEnabled, ValueType)
-						.Value(this, &FRigControlDetails::GetComponentValueFloat, ValueType, ERigElementDetailsTransformComponent::TranslationX)
-						.OnValueChanged(this, &FRigControlDetails::SetComponentValueFloat, ValueType, ERigElementDetailsTransformComponent::TranslationX, DetailBuilder.GetPropertyUtilities())
-					]
-				];
-			}
-			break;
-		}
-		case ERigControlType::Integer:
+	if(ControlElement->Settings.ControlType != ERigControlType::Bool)
+	{
+		// setup offset transform
 		{
-			for(int64 ValueTypeIndex = 0; ValueTypeIndex < ValueTypeEnum->GetMaxEnumValue(); ValueTypeIndex++)
-			{
-				ERigControlValueType ValueType = (ERigControlValueType)ValueTypeIndex;
-				FString ValueTypeName = ValueTypeEnum->GetDisplayNameTextByValue(ValueTypeIndex).ToString();
-				FText PropertyLabel = FText::FromString(FString::Printf(TEXT("%s %s"), *ValueTypeName, *ControlTypeName));
-
-				// skip min and max if this is an enum
-				if (Control.ControlEnum != nullptr &&
-					(
-						ValueType == ERigControlValueType::Minimum ||
-						ValueType == ERigControlValueType::Maximum
-						))
-				{
-					continue;
-				}
+			const TSharedRef<IPropertyHandle> OffsetHandle = DetailBuilder.GetProperty(TEXT("Offset"));
+			const TSharedRef<IPropertyHandle> OffsetInitialLocalTransformHandle = DetailBuilder.GetProperty(TEXT("Offset.Initial.Local.Transform"));
+			ControlCategory.AddProperty(OffsetInitialLocalTransformHandle).DisplayName(FText::FromString(TEXT("Offset Transform")));
+			DetailBuilder.HideProperty(TEXT("Offset"));
 
-				if (ValueType == ERigControlValueType::Minimum)
-				{
-					LimitsCategory.AddCustomRow(FText::FromString(TEXT("Limit")))
-					.NameContent()
-					[
-						SNew(STextBlock)
-						.Text(FText::FromString(TEXT("Limit")))
-						.Font(IDetailLayoutBuilder::GetDetailFont())
-					]
-					.ValueContent()
-					[
-						DetailBuilder.GetProperty(TEXT("bLimitTranslation"))->CreatePropertyValueWidget()
-					];
-
-					LimitsCategory.AddCustomRow(FText::FromString(TEXT("bDrawLimits")))
-					.NameContent()
-					[
-						DetailBuilder.GetProperty(TEXT("bDrawLimits"))->CreatePropertyNameWidget()
-					]
-					.ValueContent()
-					[
-						DetailBuilder.GetProperty(TEXT("bDrawLimits"))->CreatePropertyValueWidget()
-					];
-				}
+			OffsetPropertyChain.AddHead(FRigControlElement::StaticStruct()->FindPropertyByName(TEXT("Offset")));
+			OffsetPropertyChain.AddTail(FRigCurrentAndInitialTransform::StaticStruct()->FindPropertyByName(TEXT("Initial")));
+			OffsetPropertyChain.AddTail(FRigLocalAndGlobalTransform::StaticStruct()->FindPropertyByName(TEXT("Local")));
+			OffsetPropertyChain.AddTail(FRigComputedTransform::StaticStruct()->FindPropertyByName(TEXT("Transform")));
+			OffsetPropertyChain.SetActiveMemberPropertyNode(OffsetPropertyChain.GetTail()->GetValue());
 
-				IDetailCategoryBuilder& Category = (ValueType == ERigControlValueType::Minimum || ValueType == ERigControlValueType::Maximum) ? LimitsCategory : ControlCategory;
-
-				if (Control.ControlEnum)
-				{
-					Category.AddCustomRow(PropertyLabel)
-					.NameContent()
-					.VAlign(VAlign_Top)
-					[
-						SNew(STextBlock)
-						.Text(PropertyLabel)
-						.Font(IDetailLayoutBuilder::GetDetailFont())
-						.IsEnabled(this, &FRigControlDetails::IsEnabled, ValueType)
-					]
-					.ValueContent()
-					.MinDesiredWidth(125.0f * 3.0f) // copied from FComponentTransformDetails
-					.MaxDesiredWidth(125.0f * 3.0f)
-					[
-						SNew(SVerticalBox)
-						+ SVerticalBox::Slot()
-						[
-							SNew(SEnumComboBox, Control.ControlEnum)
-							.IsEnabled(this, &FRigControlDetails::IsEnabled, ValueType)
-							.OnEnumSelectionChanged(SEnumComboBox::FOnEnumSelectionChanged::CreateSP(this, &FRigControlDetails::OnControlEnumChanged, ValueType, DetailBuilder.GetPropertyUtilities()))
-							.CurrentValue(this, &FRigControlDetails::GetControlEnumValue, ValueType)
-							.Font(FEditorStyle::GetFontStyle(TEXT("MenuItem.Font")))
-						]
-					];
-				}
-				else
-				{
-					Category.AddCustomRow(PropertyLabel)
-					.NameContent()
-					.VAlign(VAlign_Top)
-					[
-						SNew(STextBlock)
-						.Text(PropertyLabel)
-						.Font(IDetailLayoutBuilder::GetDetailFont())
-						.IsEnabled(this, &FRigControlDetails::IsEnabled, ValueType)
-					]
-					.ValueContent()
-					.MinDesiredWidth(125.0f * 3.0f) // copied from FComponentTransformDetails
-					.MaxDesiredWidth(125.0f * 3.0f)
-					[
-						SNew(SVerticalBox)
-						+ SVerticalBox::Slot()
-						[
-							SNew(SNumericEntryBox<int32>)
-							.Font(FEditorStyle::GetFontStyle(TEXT("MenuItem.Font")))
-							.AllowSpin(true)
-							.IsEnabled(this, &FRigControlDetails::IsEnabled, ValueType)
-							.Value(this, &FRigControlDetails::GetComponentValueInteger, ValueType)
-							.OnValueChanged(this, &FRigControlDetails::SetComponentValueInteger, ValueType, DetailBuilder.GetPropertyUtilities())
-						]
-					];
-				}
-			}
-			break;
+			const FSimpleDelegate OnTransformChangedDelegate = FSimpleDelegate::CreateStatic(&Local::OnTransformChanged, &OffsetPropertyChain, BlueprintBeingCustomized);
+			OffsetInitialLocalTransformHandle->SetOnPropertyValueChanged(OnTransformChangedDelegate);
+			OffsetInitialLocalTransformHandle->SetOnChildPropertyValueChanged(OnTransformChangedDelegate);
 		}
-		case ERigControlType::Vector2D:
-		case ERigControlType::Position:
-		case ERigControlType::Scale:
-		case ERigControlType::Rotator:
-		case ERigControlType::Transform:
-		case ERigControlType::TransformNoScale:
-		case ERigControlType::EulerTransform:
-		{
-			for(int64 ValueTypeIndex = 0; ValueTypeIndex < ValueTypeEnum->GetMaxEnumValue(); ValueTypeIndex++)
-			{
-				ERigControlValueType ValueType = (ERigControlValueType)ValueTypeIndex;
-				FString ValueTypeName = ValueTypeEnum->GetDisplayNameTextByValue(ValueTypeIndex).ToString();
-				FText PropertyLabel = FText::FromString(FString::Printf(TEXT("%s %s"), *ValueTypeName, *ControlTypeName));
-
-				if (ValueType == ERigControlValueType::Minimum)
-				{
-					FText LimitTranslationLabel = FText::FromString(TEXT("Limit"));
-					FText LimiRotationLabel = LimitTranslationLabel;
-					bool bShowLimitRotation = Control.ControlType == ERigControlType::Rotator;
-
-					if (Control.ControlType == ERigControlType::Transform ||
-						Control.ControlType == ERigControlType::TransformNoScale)
-					{
-						LimitTranslationLabel = FText::FromString(TEXT("Limit Translation"));
-						LimiRotationLabel = FText::FromString(TEXT("Limit Rotation"));
-						bShowLimitRotation = true;
-					}
-
-					if (Control.ControlType != ERigControlType::Rotator)
-					{
-						LimitsCategory.AddCustomRow(LimitTranslationLabel)
-						.NameContent()
-						[
-							SNew(STextBlock)
-							.Text(LimitTranslationLabel)
-						.Font(IDetailLayoutBuilder::GetDetailFont())
-						]
-						.ValueContent()
-						[
-							DetailBuilder.GetProperty(TEXT("bLimitTranslation"))->CreatePropertyValueWidget()
-						];
-					}
-
-
-					if (bShowLimitRotation)
-					{
-						LimitsCategory.AddCustomRow(LimiRotationLabel)
-						.NameContent()
-						[
-							SNew(STextBlock)
-							.Text(LimiRotationLabel)
-							.Font(IDetailLayoutBuilder::GetDetailFont())
-						]
-						.ValueContent()
-						[
-							DetailBuilder.GetProperty(TEXT("bLimitRotation"))->CreatePropertyValueWidget()
-						];
-					}
-
-					LimitsCategory.AddCustomRow(FText::FromString(TEXT("bDrawLimits")))
-					.NameContent()
-					[
-						DetailBuilder.GetProperty(TEXT("bDrawLimits"))->CreatePropertyNameWidget()
-					]
-					.ValueContent()
-					[
-						DetailBuilder.GetProperty(TEXT("bDrawLimits"))->CreatePropertyValueWidget()
-					];
-				}
-
-				uint8* ValuePtr = (uint8*)&Control.GetValue(ValueType);
 
-				const UStruct* ValueStruct = nullptr;
-				switch (Control.ControlType)
-				{
-					case ERigControlType::Vector2D:
-					{
-						ValueStruct = TBaseStructure<FVector2D>::Get();
-						break;
-					}
-					case ERigControlType::Position:
-					case ERigControlType::Scale:
-					{
-						ValueStruct = TBaseStructure<FVector>::Get();
-						break;
-					}
-					case ERigControlType::Rotator:
-					{
-						ValueStruct = TBaseStructure<FRotator>::Get();
-						break;
-					}
-					case ERigControlType::Transform:
-					{
-						ValueStruct = TBaseStructure<FTransform>::Get();
-						break;
-					}
-					case ERigControlType::TransformNoScale:
-					{
-						ValueStruct = TBaseStructure<FTransformNoScale>::Get();
-						break;
-					}
-					case ERigControlType::EulerTransform:
-					{
-						ValueStruct = TBaseStructure<FEulerTransform>::Get();
-						break;
-					}
-					default:
-					{
-						checkNoEntry();
-					}
-				}
-
-				TSharedPtr<FStructOnScope> StructToDisplay = MakeShareable(new FStructOnScope(ValueStruct, ValuePtr));
-
-				IDetailCategoryBuilder& Category = (ValueType == ERigControlValueType::Minimum || ValueType == ERigControlValueType::Maximum) ? LimitsCategory : ControlCategory;
-				FProperty* Property = FRigControl::FindPropertyForValueType(ValueType);
-
-				IDetailPropertyRow* Row = Category.AddExternalStructure(StructToDisplay);
-				Row->DisplayName(PropertyLabel);
-				Row->ToolTip(Property->GetToolTipText());
-				Row->ShouldAutoExpand(true);
-				Row->IsEnabled(TAttribute<bool>::Create(TAttribute<bool>::FGetter::CreateSP(this, &FRigControlDetails::IsEnabled, ValueType)));
+		// setup gizmo transform
+		{
+			const TSharedRef<IPropertyHandle> GizmoHandle = DetailBuilder.GetProperty(TEXT("Gizmo"));
+			const TSharedRef<IPropertyHandle> GizmoInitialLocalTransformHandle = DetailBuilder.GetProperty(TEXT("Gizmo.Initial.Local.Transform"));
+			GizmoCategory.AddProperty(GizmoInitialLocalTransformHandle).DisplayName(FText::FromString(TEXT("Gizmo Transform")));
+			DetailBuilder.HideProperty(TEXT("Gizmo"));
 
-				FSimpleDelegate OnStructContentsChangedDelegate = FSimpleDelegate::CreateSP(this, &FRigControlDetails::OnStructContentsChanged, Property, DetailBuilder.GetPropertyUtilities());
+			GizmoPropertyChain.AddHead(FRigControlElement::StaticStruct()->FindPropertyByName(TEXT("Gizmo")));
+			GizmoPropertyChain.AddTail(FRigCurrentAndInitialTransform::StaticStruct()->FindPropertyByName(TEXT("Initial")));
+			GizmoPropertyChain.AddTail(FRigLocalAndGlobalTransform::StaticStruct()->FindPropertyByName(TEXT("Local")));
+			GizmoPropertyChain.AddTail(FRigComputedTransform::StaticStruct()->FindPropertyByName(TEXT("Transform")));
+			GizmoPropertyChain.SetActiveMemberPropertyNode(GizmoPropertyChain.GetTail()->GetValue());
 
-				TSharedPtr<IPropertyHandle> Handle = Row->GetPropertyHandle();
-				Handle->SetOnPropertyValueChanged(OnStructContentsChangedDelegate);
-				Handle->SetOnChildPropertyValueChanged(OnStructContentsChangedDelegate);
-			}
-			break;
-		}
-		default:
-		{
-			ensure(false);
-			break;
+			const FSimpleDelegate OnTransformChangedDelegate = FSimpleDelegate::CreateStatic(&Local::OnTransformChanged, &GizmoPropertyChain, BlueprintBeingCustomized);
+			GizmoInitialLocalTransformHandle->SetOnPropertyValueChanged(OnTransformChangedDelegate);
+			GizmoInitialLocalTransformHandle->SetOnChildPropertyValueChanged(OnTransformChangedDelegate);
 		}
 	}
 
+	FRigControlElementDetails_SetupValueWidget(ControlCategory, ERigControlValueType::Current, ControlElement, HierarchyBeingCustomized);
+	FRigControlElementDetails_SetupValueWidget(LimitsCategory, ERigControlValueType::Minimum, ControlElement, HierarchyBeingCustomized);
+	FRigControlElementDetails_SetupValueWidget(LimitsCategory, ERigControlValueType::Maximum, ControlElement, HierarchyBeingCustomized);
 
-	switch (Control.ControlType)
+	switch (ControlElement->Settings.ControlType)
 	{
 		case ERigControlType::Float:
 		case ERigControlType::Integer:
@@ -1406,710 +1319,234 @@ void FRigControlDetails::CustomizeDetails(IDetailLayoutBuilder& DetailBuilder)
 		case ERigControlType::TransformNoScale:
 		case ERigControlType::EulerTransform:
 		{
-			DetailBuilder.HideProperty(TEXT("bGizmoEnabled"));
-			DetailBuilder.HideProperty(TEXT("bGizmoVisible"));
-			DetailBuilder.HideProperty(TEXT("GizmoName"));
-
-			GizmoCategory.AddCustomRow(FText::FromString(TEXT("bGizmoEnabled")))
-			.NameContent()
-			[
-				DetailBuilder.GetProperty(TEXT("bGizmoEnabled"))->CreatePropertyNameWidget()
-			]
-			.ValueContent()
-			[
-				DetailBuilder.GetProperty(TEXT("bGizmoEnabled"))->CreatePropertyValueWidget()
-			];
-
-			GizmoCategory.AddCustomRow(FText::FromString(TEXT("GizmoName")))
+			const TSharedRef<IPropertyHandle> GizmoNameHandle = DetailBuilder.GetProperty(TEXT("Settings.GizmoName"));
+			DetailBuilder.HideProperty(GizmoNameHandle);
+			GizmoCategory.AddProperty(GizmoNameHandle).CustomWidget()
 			.NameContent()
 			[
 				SNew(STextBlock)
 				.Text(FText::FromString(TEXT("Gizmo Name")))
 				.Font(IDetailLayoutBuilder::GetDetailFont())
-				.IsEnabled(this, &FRigControlDetails::IsGizmoEnabled)
-			]
-			.ValueContent()
-			[
-				SNew(SControlRigGizmoNameList, &Control, BlueprintBeingCustomized)
-				.OnGetNameListContent(this, &FRigControlDetails::GetGizmoNameList)
-				.IsEnabled(this, &FRigControlDetails::IsGizmoEnabled)
-			];
-
-			GizmoCategory.AddCustomRow(FText::FromString(TEXT("bGizmoVisible")))
-			.NameContent()
-			[
-				DetailBuilder.GetProperty(TEXT("bGizmoVisible"))->CreatePropertyNameWidget()
+				.IsEnabled(this, &FRigControlElementDetails::IsGizmoEnabled)
 			]
 			.ValueContent()
 			[
-				DetailBuilder.GetProperty(TEXT("bGizmoVisible"))->CreatePropertyValueWidget()
+				SNew(SControlRigGizmoNameList, ControlElement, BlueprintBeingCustomized)
+				.OnGetNameListContent(this, &FRigControlElementDetails::GetGizmoNameList)
+				.IsEnabled(this, &FRigControlElementDetails::IsGizmoEnabled)
 			];
-
 			break;
 		}
 		default:
 		{
-			DetailBuilder.HideProperty(TEXT("bGizmoEnabled"));
-			DetailBuilder.HideProperty(TEXT("bGizmoVisible"));
-			DetailBuilder.HideProperty(TEXT("GizmoName"));
-			DetailBuilder.HideProperty(TEXT("GizmoTransform"));
-			DetailBuilder.HideProperty(TEXT("GizmoColor"));
+			DetailBuilder.HideProperty(TEXT("Settings.bGizmoEnabled"));
+			DetailBuilder.HideProperty(TEXT("Settings.bGizmoVisible"));
+			DetailBuilder.HideProperty(TEXT("Settings.GizmoName"));
+			DetailBuilder.HideProperty(TEXT("Settings.GizmoColor"));
 			break;
 		}
 	}
 
-	switch (Control.ControlType)
+	if(ControlElement->Settings.ControlType != ERigControlType::Integer)
+	{
+		DetailBuilder.HideProperty(TEXT("Settings.ControlEnum"));
+	}
+
+	bool bShowLimitTranslation = false;
+	bool bShowLimitRotation = false;
+	bool bShowLimitScale = false;
+
+	switch (ControlElement->Settings.ControlType)
 	{
 		case ERigControlType::Float:
-		case ERigControlType::Integer:
-		case ERigControlType::Vector2D:
-		{
+    	case ERigControlType::Integer:
+    	case ERigControlType::Vector2D:
+    	case ERigControlType::Position:
+    	case ERigControlType::Transform:
+    	case ERigControlType::TransformNoScale:
+    	case ERigControlType::EulerTransform:
+		{
+				bShowLimitTranslation = true;
 			break;
 		}
 		default:
 		{
-			DetailBuilder.HideProperty(TEXT("PrimaryAxis"));
 			break;
 		}
 	}
 
-	if (Control.ControlType == ERigControlType::Integer)
+	switch (ControlElement->Settings.ControlType)
 	{
-		TSharedRef<IPropertyHandle> ControlEnum = DetailBuilder.GetProperty(TEXT("ControlEnum"));
-		ControlEnum->SetOnPropertyValueChanged(FSimpleDelegate::CreateLambda(
-			[this, &DetailBuilder]()
-			{
-				DetailBuilder.ForceRefreshDetails();
-
-				if (this->ContainerBeingCustomized != nullptr && this->ElementKeyBeingCustomized)
-				{
-					int32 Index = ContainerBeingCustomized->ControlHierarchy.GetIndex(ElementKeyBeingCustomized.Name);
-					if (Index != INDEX_NONE)
-					{
-						UEnum* ControlEnum = ContainerBeingCustomized->ControlHierarchy[Index].ControlEnum;
-						if (ControlEnum != nullptr)
-						{
-							int32 Maximum = (int32)ControlEnum->GetMaxEnumValue() - 1;
-							ContainerBeingCustomized->ControlHierarchy[Index].MinimumValue.Set<int32>(0);
-							ContainerBeingCustomized->ControlHierarchy[Index].MaximumValue.Set<int32>(Maximum);
-
-							FRigControlValue InitialValue = ContainerBeingCustomized->ControlHierarchy[Index].InitialValue;
-							FRigControlValue Value = ContainerBeingCustomized->ControlHierarchy[Index].Value;
-
-							ContainerBeingCustomized->ControlHierarchy[Index].ApplyLimits(InitialValue);
-							ContainerBeingCustomized->ControlHierarchy[Index].ApplyLimits(Value);
-
-							ContainerBeingCustomized->ControlHierarchy.SetValue(ElementKeyBeingCustomized.Name, InitialValue, ERigControlValueType::Initial);
-							ContainerBeingCustomized->ControlHierarchy.SetValue(ElementKeyBeingCustomized.Name, Value, ERigControlValueType::Current);
-
-							if (UControlRig* DebuggedRig = Cast<UControlRig>(BlueprintBeingCustomized->GetObjectBeingDebugged()))
-							{
-								FRigControlHierarchy& DebuggedControlHierarchy = *((FRigControlHierarchy*)&DebuggedRig->GetControlHierarchy());
-								DebuggedControlHierarchy[Index].MinimumValue.Set<int32>(0);
-								DebuggedControlHierarchy[Index].MaximumValue.Set<int32>(Maximum);
-								DebuggedControlHierarchy.SetValue(ElementKeyBeingCustomized.Name, InitialValue, ERigControlValueType::Initial);
-								DebuggedControlHierarchy.SetValue(ElementKeyBeingCustomized.Name, Value, ERigControlValueType::Current);
-							}
-						}
-					}
-				}
-			}
-		));
-
-		ControlCategory.AddCustomRow(FText::FromString(TEXT("ControlEnum")))
-		.NameContent()
-		[
-			ControlEnum->CreatePropertyNameWidget()
-		]
-		.ValueContent()
-		[
-			ControlEnum->CreatePropertyValueWidget()
-		];
+    	case ERigControlType::Rotator:
+    	case ERigControlType::Transform:
+    	case ERigControlType::TransformNoScale:
+    	case ERigControlType::EulerTransform:
+		{
+			bShowLimitRotation = true;
+			break;
+		}
+		default:
+		{
+			break;
+		}
 	}
 
-	// setup influences modifier
+	switch (ControlElement->Settings.ControlType)
 	{
-		FName InfluencesEventName = FRigUnit_BeginExecution::EventName;
-		FRigInfluenceMap& InfluenceMap = BlueprintBeingCustomized->Influences.FindOrAdd(InfluencesEventName);
-		FRigInfluenceEntryModifier TempModifier = InfluenceMap.GetEntryModifier(ElementKeyBeingCustomized);
-		InfluenceModifier = MakeShareable(new FRigInfluenceEntryModifier(TempModifier));
-		InfluenceModifierStruct = MakeShareable(new FStructOnScope(FRigInfluenceEntryModifier::StaticStruct(), (uint8*)InfluenceModifier.Get()));
-		InfluenceModifierStruct->SetPackage(BlueprintBeingCustomized->GetOutermost());
-
-		IDetailCategoryBuilder& InversionCategory = DetailBuilder.EditCategory(TEXT("Inversion"), LOCTEXT("InversionCategory", "Inversion"));
-		FProperty* AffectedListProperty = FRigInfluenceEntryModifier::StaticStruct()->FindPropertyByName(TEXT("AffectedList"));
-
-		IDetailPropertyRow* AffectedListRow = InversionCategory.AddExternalStructureProperty(InfluenceModifierStruct, AffectedListProperty->GetFName());
-		AffectedListRow->ShouldAutoExpand(true);
-		AffectedListRow->DisplayName(LOCTEXT("AffectedElements", "Affected Elements"));
-
-		FSimpleDelegate OnInfluencesChangedDelegate = FSimpleDelegate::CreateSP(this, &FRigControlDetails::OnAffectedListChanged, AffectedListProperty, DetailBuilder.GetPropertyUtilities());
-
-		TSharedPtr<IPropertyHandle> AffectedListHandle = AffectedListRow->GetPropertyHandle();
-		AffectedListHandle->SetOnPropertyValueChanged(OnInfluencesChangedDelegate);
-		AffectedListHandle->SetOnChildPropertyValueChanged(OnInfluencesChangedDelegate);
+		case ERigControlType::Scale:
+    	case ERigControlType::Transform:
+    	case ERigControlType::EulerTransform:
+		{
+			bShowLimitScale = true;
+			break;
+		}
+		default:
+		{
+			break;
+		}
 	}
-}
-
-void FRigControlDetails::OnStructContentsChanged(FProperty* InProperty, const TSharedRef<IPropertyUtilities> PropertyUtilities)
-{
-	FPropertyChangedEvent ChangeEvent(InProperty, EPropertyChangeType::ValueSet);
-	PropertyUtilities->NotifyFinishedChangingProperties(ChangeEvent);
-}
 
-void FRigControlDetails::OnAffectedListChanged(FProperty* InProperty, const TSharedRef<IPropertyUtilities> PropertyUtilities)
-{
-	if (BlueprintBeingCustomized != nullptr && InfluenceModifier.IsValid())
+	if(!bShowLimitTranslation)
 	{
-		BlueprintBeingCustomized->Modify();
-
-		FName InfluencesEventName = FRigUnit_BeginExecution::EventName;
-		FRigInfluenceMap& InfluenceMap = BlueprintBeingCustomized->Influences.FindOrAdd(InfluencesEventName);
-		InfluenceMap.SetEntryModifier(ElementKeyBeingCustomized, *InfluenceModifier);
+		DetailBuilder.HideProperty(TEXT("Settings.bLimitTranslation"));
 	}
-
-	OnStructContentsChanged(InProperty, PropertyUtilities);
-}
-
-int32 FRigControlDetails::GetControlEnumValue(ERigControlValueType InValueType) const
-{
-	TOptional<int32> Value = GetComponentValueInteger(InValueType);
-	if (Value.IsSet())
+	if(!bShowLimitRotation)
 	{
-		return Value.GetValue();
+		DetailBuilder.HideProperty(TEXT("Settings.bLimitRotation"));
 	}
-	return 0;
-}
-
-void FRigControlDetails::OnControlEnumChanged(int32 InValue, ESelectInfo::Type InSelectInfo, ERigControlValueType InValueType, const TSharedRef<IPropertyUtilities> PropertyUtilities)
-{
-	if (ContainerBeingCustomized != nullptr && ElementKeyBeingCustomized)
+	if(!bShowLimitScale)
 	{
-		int32 Index = ContainerBeingCustomized->ControlHierarchy.GetIndex(ElementKeyBeingCustomized.Name);
-		if (Index != INDEX_NONE)
-		{
-			FRigControlValue Value = FRigControlValue::Make<int32>(InValue);
-			ContainerBeingCustomized->ControlHierarchy.SetValue(ElementKeyBeingCustomized.Name, Value, InValueType);
-			OnStructContentsChanged(FRigControl::FindPropertyForValueType(InValueType), PropertyUtilities);
-		}
+		DetailBuilder.HideProperty(TEXT("Settings.bLimitScale"));
 	}
-}
 
-FText FRigControlDetails::GetDisplayName() const
-{
-	if (ContainerBeingCustomized != nullptr && ElementKeyBeingCustomized)
+	switch (ControlElement->Settings.ControlType)
 	{
-		int32 Index = ContainerBeingCustomized->ControlHierarchy.GetIndex(ElementKeyBeingCustomized.Name);
-		if (Index != INDEX_NONE)
+		case ERigControlType::Integer:
+		case ERigControlType::Float:
+		case ERigControlType::Vector2D:
 		{
-			FRigControl Control = ContainerBeingCustomized->ControlHierarchy[ElementKeyBeingCustomized.Name];
-			if (Control.DisplayName.IsNone())
-			{
-				return FText();
-			}
-			return FText::FromName(Control.GetDisplayName());
+			break;
 		}
-	}
-	return FText();
-}
-
-void FRigControlDetails::SetDisplayName(const FText& InNewText, ETextCommit::Type InCommitType, const TSharedRef<IPropertyUtilities> PropertyUtilities)
-{
-	if (ContainerBeingCustomized != nullptr && ElementKeyBeingCustomized)
-	{
-		int32 Index = ContainerBeingCustomized->ControlHierarchy.GetIndex(ElementKeyBeingCustomized.Name);
-		if (Index != INDEX_NONE)
+		default:
 		{
-			FRigControl& Control = ContainerBeingCustomized->ControlHierarchy[ElementKeyBeingCustomized.Name];
-
-			FString NewDisplayName = InNewText.ToString().TrimStartAndEnd();
-			if (NewDisplayName.IsEmpty())
-			{
-				Control.DisplayName = FName(NAME_None);
-			}
-			else
-			{
-				Control.DisplayName = *NewDisplayName;
-			}
-
-			OnStructContentsChanged(FRigControl::StaticStruct()->FindPropertyByName(TEXT("DisplayName")), PropertyUtilities);
+			DetailBuilder.HideProperty(TEXT("Settings.PrimaryAxis"));
+			break;
 		}
 	}
-}
 
-ECheckBoxState FRigControlDetails::GetComponentValueBool(bool bInitial) const
-{
-	if (ContainerBeingCustomized != nullptr && ElementKeyBeingCustomized)
+	if (ControlElement->Settings.ControlType == ERigControlType::Integer)
 	{
-		int32 Index = ContainerBeingCustomized->ControlHierarchy.GetIndex(ElementKeyBeingCustomized.Name);
-		if (Index != INDEX_NONE)
-		{
-			FRigControl Control = ContainerBeingCustomized->ControlHierarchy[ElementKeyBeingCustomized.Name];
-			if (!bInitial)
-			{
-				if (UControlRig* DebuggedRig = Cast<UControlRig>(BlueprintBeingCustomized->GetObjectBeingDebugged()))
-				{
-					Control = DebuggedRig->GetControlHierarchy()[Index];
-				}
-			}
-
-			switch (Control.ControlType)
+		TSharedRef<IPropertyHandle> ControlEnum = DetailBuilder.GetProperty(TEXT("Settings.ControlEnum"));
+		ControlEnum->SetOnPropertyValueChanged(FSimpleDelegate::CreateLambda(
+			[this, &DetailBuilder]()
 			{
-				case ERigControlType::Bool:
-				{
-					bool Value = bInitial ? Control.InitialValue.Get<bool>() : Control.Value.Get<bool>();
-					return Value ? ECheckBoxState::Checked : ECheckBoxState::Unchecked;
-				}
-				default:
-				{
-					ensure(false);
-					break;
-				}
-			}
-		}
-	}
-	return ECheckBoxState::Unchecked;
-}
-
-void FRigControlDetails::SetComponentValueBool(ECheckBoxState InNewValue, bool bInitial, const TSharedRef<IPropertyUtilities> PropertyUtilities)
-{
-	if (ContainerBeingCustomized != nullptr && ElementKeyBeingCustomized)
-	{
-		int32 Index = ContainerBeingCustomized->ControlHierarchy.GetIndex(ElementKeyBeingCustomized.Name);
-		if (Index != INDEX_NONE)
-		{
-			const FRigControl& Control = ContainerBeingCustomized->ControlHierarchy[ElementKeyBeingCustomized.Name];
-			FRigControlValue Value = bInitial ? Control.InitialValue : Control.Value;
+				DetailBuilder.ForceRefreshDetails();
 
-			switch (Control.ControlType)
-			{
-				case ERigControlType::Bool:
+				if (this->HierarchyBeingCustomized != nullptr && this->ElementKeyBeingCustomized)
 				{
-					Value.Set<bool>(InNewValue == ECheckBoxState::Checked);
-					break;
-				}
-				default:
-				{
-					ensure(false);
-					break;
-				}
-			}
+					if(FRigControlElement* ControlBeingCustomized = this->HierarchyBeingCustomized->Find<FRigControlElement>(ElementKeyBeingCustomized))
+					{
+						const UEnum* ControlEnum = ControlBeingCustomized->Settings.ControlEnum;
+						if (ControlEnum != nullptr)
+						{
+							int32 Maximum = (int32)ControlEnum->GetMaxEnumValue() - 1;
+							ControlBeingCustomized->Settings.MinimumValue.Set<int32>(0);
+							ControlBeingCustomized->Settings.MaximumValue.Set<int32>(Maximum);
+							HierarchyBeingCustomized->Notify(ERigHierarchyNotification::ControlSettingChanged, ControlBeingCustomized);
 
-			if (bInitial)
-			{
-				ContainerBeingCustomized->ControlHierarchy.SetInitialValue(ElementKeyBeingCustomized.Name, Value);
-				OnStructContentsChanged(FRigControl::FindPropertyForValueType(ERigControlValueType::Initial), PropertyUtilities);
-			}
-			else
-			{
-				ContainerBeingCustomized->ControlHierarchy.SetValue(ElementKeyBeingCustomized.Name, Value);
-				OnStructContentsChanged(FRigControl::FindPropertyForValueType(ERigControlValueType::Current), PropertyUtilities);
-			}
-		}
-	}
-}
+							FRigControlValue InitialValue = HierarchyBeingCustomized->GetControlValue(ControlBeingCustomized, ERigControlValueType::Initial);
+							FRigControlValue CurrentValue = HierarchyBeingCustomized->GetControlValue(ControlBeingCustomized, ERigControlValueType::Current);
 
-TOptional<float> FRigControlDetails::GetComponentValueFloat(ERigControlValueType InValueType, ERigElementDetailsTransformComponent Component) const
-{
-	if (ContainerBeingCustomized != nullptr && ElementKeyBeingCustomized)
-	{
-		int32 Index = ContainerBeingCustomized->ControlHierarchy.GetIndex(ElementKeyBeingCustomized.Name);
-		if (Index != INDEX_NONE)
-		{
-			FRigControl Control = ContainerBeingCustomized->ControlHierarchy[ElementKeyBeingCustomized.Name];
-			if (InValueType == ERigControlValueType::Current)
-			{
-				if (UControlRig* DebuggedRig = Cast<UControlRig>(BlueprintBeingCustomized->GetObjectBeingDebugged()))
-				{
-					Control = DebuggedRig->GetControlHierarchy()[Index];
-				}
-			}
+							ControlBeingCustomized->Settings.ApplyLimits(InitialValue);
+							ControlBeingCustomized->Settings.ApplyLimits(CurrentValue);
+							HierarchyBeingCustomized->SetControlValue(ControlBeingCustomized, InitialValue, ERigControlValueType::Initial);
+							HierarchyBeingCustomized->SetControlValue(ControlBeingCustomized, CurrentValue, ERigControlValueType::Current);
 
-			switch (Control.ControlType)
-			{
-				case ERigControlType::Float:
-				{
-					return Control.GetValue(InValueType).Get<float>();
-				}
-				case ERigControlType::Integer:
-				{
-					return (float)Control.GetValue(InValueType).Get<int32>();
-				}
-				case ERigControlType::Vector2D:
-				{
-					switch (Component)
-					{
-						case ERigElementDetailsTransformComponent::TranslationX:
-						{
-							return Control.GetValue(InValueType).Get<FVector2D>().X;
-						}
-						case ERigElementDetailsTransformComponent::TranslationY:
-						{
-							return Control.GetValue(InValueType).Get<FVector2D>().Y;
-						}
-						default:
-						{
-							break;
-						}
-					}
-					break;
-				}
-				case ERigControlType::Position:
-				{
-					switch (Component)
-					{
-						case ERigElementDetailsTransformComponent::TranslationX:
-						{
-							return Control.GetValue(InValueType).Get<FVector>().X;
-						}
-						case ERigElementDetailsTransformComponent::TranslationY:
-						{
-							return Control.GetValue(InValueType).Get<FVector>().Y;
-						}
-						case ERigElementDetailsTransformComponent::TranslationZ:
-						{
-							return Control.GetValue(InValueType).Get<FVector>().Z;
-						}
-						default:
-						{
-							break;
-						}
-					}
-					break;
-				}
-				case ERigControlType::Scale:
-				{
-					switch (Component)
-					{
-						case ERigElementDetailsTransformComponent::ScaleX:
-						{
-							return Control.GetValue(InValueType).Get<FVector>().X;
-						}
-						case ERigElementDetailsTransformComponent::ScaleY:
-						{
-							return Control.GetValue(InValueType).Get<FVector>().Y;
-						}
-						case ERigElementDetailsTransformComponent::ScaleZ:
-						{
-							return Control.GetValue(InValueType).Get<FVector>().Z;
-						}
-						default:
-						{
-							break;
-						}
-					}
-					break;
-				}
-				case ERigControlType::Rotator:
-				{
-					switch (Component)
-					{
-						case ERigElementDetailsTransformComponent::RotationPitch:
-						{
-							return Control.GetValue(InValueType).Get<FRotator>().Pitch;
-						}
-						case ERigElementDetailsTransformComponent::RotationYaw:
-						{
-							return Control.GetValue(InValueType).Get<FRotator>().Yaw;
-						}
-						case ERigElementDetailsTransformComponent::RotationRoll:
-						{
-							return Control.GetValue(InValueType).Get<FRotator>().Roll;
-						}
-						default:
-						{
-							break;
+							if (UControlRig* DebuggedRig = Cast<UControlRig>(BlueprintBeingCustomized->GetObjectBeingDebugged()))
+							{
+								URigHierarchy* DebuggedHierarchy = DebuggedRig->GetHierarchy();
+								if(FRigControlElement* DebuggedControlElement = DebuggedHierarchy->Find<FRigControlElement>(ElementKeyBeingCustomized))
+								{
+									DebuggedControlElement->Settings.MinimumValue.Set<int32>(0);
+                                    DebuggedControlElement->Settings.MaximumValue.Set<int32>(Maximum);
+                                    DebuggedHierarchy->Notify(ERigHierarchyNotification::ControlSettingChanged, DebuggedControlElement);
+
+                                    DebuggedHierarchy->SetControlValue(DebuggedControlElement, InitialValue, ERigControlValueType::Initial);
+                                    DebuggedHierarchy->SetControlValue(DebuggedControlElement, CurrentValue, ERigControlValueType::Current);
+								}
+							}
 						}
 					}
-					break;
-				}
-				case ERigControlType::Transform:
-				{
-					FTransform Transform = Control.GetValue(InValueType).Get<FTransform>();
-					return GetTransformComponent(Transform, Component);
-				}
-				case ERigControlType::TransformNoScale:
-				{
-					FTransform Transform = Control.GetValue(InValueType).Get<FTransformNoScale>();
-					return GetTransformComponent(Transform, Component);
-				}
-				case ERigControlType::EulerTransform:
-				{
-					FEulerTransform Transform = Control.GetValue(InValueType).Get<FEulerTransform>();
-					return GetEulerTransformComponent(Transform, Component);
-				}
-				case ERigControlType::Bool:
-				{
-					break;
-				}
-				default:
-				{
-					ensure(false);
-					break;
 				}
 			}
-		}
+		));
 	}
-	return 0.f;
-}
-
-void FRigControlDetails::SetComponentValueFloat(float InNewValue, ETextCommit::Type InCommitType, ERigControlValueType InValueType, ERigElementDetailsTransformComponent Component, const TSharedRef<IPropertyUtilities> PropertyUtilities)
-{
-	SetComponentValueFloat(InNewValue, InValueType, Component, PropertyUtilities);
 }
 
-void FRigControlDetails::SetComponentValueFloat(float InNewValue, ERigControlValueType InValueType, ERigElementDetailsTransformComponent Component, const TSharedRef<IPropertyUtilities> PropertyUtilities)
+FText FRigControlElementDetails::GetDisplayName() const
 {
-	if (ContainerBeingCustomized != nullptr && ElementKeyBeingCustomized)
+	if (HierarchyBeingCustomized != nullptr && ElementKeyBeingCustomized)
 	{
-		int32 Index = ContainerBeingCustomized->ControlHierarchy.GetIndex(ElementKeyBeingCustomized.Name);
-		if (Index != INDEX_NONE)
+		if(FRigControlElement* ControlElement = HierarchyBeingCustomized->Find<FRigControlElement>(ElementKeyBeingCustomized))
 		{
-			const FRigControl& Control = ContainerBeingCustomized->ControlHierarchy[ElementKeyBeingCustomized.Name];
-			FRigControlValue Value = Control.GetValue(InValueType);
-
-			switch (Control.ControlType)
+			if (ControlElement->Settings.DisplayName.IsNone())
 			{
-				case ERigControlType::Float:
-				{
-					Value.Set<float>(InNewValue);
-					break;
-				}
-				case ERigControlType::Vector2D:
-				{
-					FVector2D Local = Value.Get<FVector2D>();
-					switch (Component)
-					{
-						case ERigElementDetailsTransformComponent::TranslationX:
-						{
-							Local.X = InNewValue;
-							break;
-						}
-						case ERigElementDetailsTransformComponent::TranslationY:
-						{
-							Local.Y = InNewValue;
-							break;
-						}
-						default:
-						{
-							break;
-						}
-					}
-					Value.Set<FVector2D>(Local);
-					break;
-				}
-				case ERigControlType::Position:
-				{
-					FVector Local = Value.Get<FVector>();
-					switch (Component)
-					{
-						case ERigElementDetailsTransformComponent::TranslationX:
-						{
-							Local.X = InNewValue;
-							break;
-						}
-						case ERigElementDetailsTransformComponent::TranslationY:
-						{
-							Local.Y = InNewValue;
-							break;
-						}
-						case ERigElementDetailsTransformComponent::TranslationZ:
-						{
-							Local.Z = InNewValue;
-							break;
-						}
-						default:
-						{
-							break;
-						}
-					}
-					Value.Set<FVector>(Local);
-					break;
-				}
-				case ERigControlType::Scale:
-				{
-					FVector Local = Value.Get<FVector>();
-					switch (Component)
-					{
-						case ERigElementDetailsTransformComponent::ScaleX:
-						{
-							Local.X = InNewValue;
-							break;
-						}
-						case ERigElementDetailsTransformComponent::ScaleY:
-						{
-							Local.Y = InNewValue;
-							break;
-						}
-						case ERigElementDetailsTransformComponent::ScaleZ:
-						{
-							Local.Z = InNewValue;
-							break;
-						}
-						default:
-						{
-							break;
-						}
-					}
-					Value.Set<FVector>(Local);
-					break;
-				}
-				case ERigControlType::Rotator:
-				{
-					FRotator Local = Value.Get<FRotator>();
-					switch (Component)
-					{
-						case ERigElementDetailsTransformComponent::RotationPitch:
-						{
-							Local.Pitch = InNewValue;
-							break;
-						}
-						case ERigElementDetailsTransformComponent::RotationYaw:
-						{
-							Local.Yaw = InNewValue;
-							break;
-						}
-						case ERigElementDetailsTransformComponent::RotationRoll:
-						{
-							Local.Roll = InNewValue;
-							break;
-						}
-						default:
-						{
-							break;
-						}
-					}
-					Value.Set<FRotator>(Local);
-					break;
-				}
-				case ERigControlType::Transform:
-				{
-					FTransform Transform = Value.Get<FTransform>();
-					SetTransformComponent(Transform, Component, InNewValue);
-					Value.Set<FTransform>(Transform);
-					break;
-				}
-				case ERigControlType::TransformNoScale:
-				{
-					FTransform Transform = Value.Get<FTransformNoScale>();
-					SetTransformComponent(Transform, Component, InNewValue);
-					Value.Set<FTransformNoScale>(Transform);
-					break;
-				}
-				case ERigControlType::EulerTransform:
-				{
-					FEulerTransform Transform = Value.Get<FEulerTransform>();
-					SetEulerTransformComponent(Transform, Component, InNewValue);
-					Value.Set<FEulerTransform>(Transform);
-					break;
-				}
-				case ERigControlType::Bool:
-				{
-					return;
-				}
-				default:
-				{
-					ensure(false);
-					break;
-				}
+				return FText();
 			}
-
-			ContainerBeingCustomized->ControlHierarchy.SetValue(ElementKeyBeingCustomized.Name, Value, InValueType);
-			OnStructContentsChanged(FRigControl::FindPropertyForValueType(InValueType), PropertyUtilities);
+			return FText::FromName(ControlElement->GetDisplayName());
 		}
 	}
+	return FText();
 }
 
-TOptional<int32> FRigControlDetails::GetComponentValueInteger(ERigControlValueType InValueType) const
+void FRigControlElementDetails::SetDisplayName(const FText& InNewText, ETextCommit::Type InCommitType, const TSharedRef<IPropertyUtilities> PropertyUtilities)
 {
-	if (ContainerBeingCustomized != nullptr && ElementKeyBeingCustomized)
+	if (HierarchyBeingCustomized != nullptr && ElementKeyBeingCustomized)
 	{
-		int32 Index = ContainerBeingCustomized->ControlHierarchy.GetIndex(ElementKeyBeingCustomized.Name);
-		if (Index != INDEX_NONE)
+		if(FRigControlElement* ControlElement = HierarchyBeingCustomized->Find<FRigControlElement>(ElementKeyBeingCustomized))
 		{
-			FRigControl Control = ContainerBeingCustomized->ControlHierarchy[ElementKeyBeingCustomized.Name];
-			if (InValueType == ERigControlValueType::Current)
+			FString NewDisplayName = InNewText.ToString().TrimStartAndEnd();
+			if (NewDisplayName.IsEmpty())
 			{
-				if (UControlRig* DebuggedRig = Cast<UControlRig>(BlueprintBeingCustomized->GetObjectBeingDebugged()))
-				{
-					Control = DebuggedRig->GetControlHierarchy()[Index];
-				}
+				ControlElement->Settings.DisplayName = FName(NAME_None);
 			}
-
-			switch (Control.ControlType)
+			else
 			{
-				case ERigControlType::Integer:
-				{
-					return Control.GetValue(InValueType).Get<int32>();
-				}
-				default:
-				{
-					break;
-				}
+				ControlElement->Settings.DisplayName = *NewDisplayName;
 			}
-		}
-	}
-	return 0;
-}
-
-void FRigControlDetails::SetComponentValueInteger(int32 InNewValue, ERigControlValueType InValueType, const TSharedRef<IPropertyUtilities> PropertyUtilities)
-{
-	if (ContainerBeingCustomized != nullptr && ElementKeyBeingCustomized)
-	{
-		int32 Index = ContainerBeingCustomized->ControlHierarchy.GetIndex(ElementKeyBeingCustomized.Name);
-		if (Index != INDEX_NONE)
-		{
-			const FRigControl& Control = ContainerBeingCustomized->ControlHierarchy[ElementKeyBeingCustomized.Name];
-			FRigControlValue Value = Control.GetValue(InValueType);
 
-			switch (Control.ControlType)
-			{
-				case ERigControlType::Integer:
-				{
-					Value.Set<int32>(InNewValue);
-					break;
-				}
-				default:
-				{
-					ensure(false);
-					break;
-				}
-			}
+			HierarchyBeingCustomized->Notify(ERigHierarchyNotification::ControlSettingChanged, ControlElement);
 
-			ContainerBeingCustomized->ControlHierarchy.SetValue(ElementKeyBeingCustomized.Name, Value, InValueType);
-			OnStructContentsChanged(FRigControl::FindPropertyForValueType(InValueType), PropertyUtilities);
+			OnStructContentsChanged(FRigControl::StaticStruct()->FindPropertyByName(TEXT("DisplayName")), PropertyUtilities);
 		}
 	}
 }
 
-bool FRigControlDetails::IsGizmoEnabled() const
+bool FRigControlElementDetails::IsGizmoEnabled() const
 {
-	if (ContainerBeingCustomized != nullptr && ElementKeyBeingCustomized)
+	URigHierarchy* Hierarchy = HierarchyBeingCustomized;
+	if (Hierarchy != nullptr && ElementKeyBeingCustomized)
 	{
-		int32 Index = ContainerBeingCustomized->ControlHierarchy.GetIndex(ElementKeyBeingCustomized.Name);
-		if (Index != INDEX_NONE)
+		if(FRigControlElement* ControlElement = Hierarchy->Find<FRigControlElement>(ElementKeyBeingCustomized))
 		{
-			const FRigControl& Control = ContainerBeingCustomized->ControlHierarchy[ElementKeyBeingCustomized.Name];
-			return Control.bGizmoEnabled;
+			return ControlElement->Settings.bGizmoEnabled;
 		}
 	}
 	return false;
 }
 
-bool FRigControlDetails::IsEnabled(ERigControlValueType InValueType) const
+bool FRigControlElementDetails::IsEnabled(ERigControlValueType InValueType) const
 {
 	switch (InValueType)
 	{
 		case ERigControlValueType::Minimum:
 		case ERigControlValueType::Maximum:
 		{
-			if (ContainerBeingCustomized != nullptr && ElementKeyBeingCustomized)
+			if (HierarchyBeingCustomized != nullptr && ElementKeyBeingCustomized)
 			{
-				int32 Index = ContainerBeingCustomized->ControlHierarchy.GetIndex(ElementKeyBeingCustomized.Name);
-				if (Index != INDEX_NONE)
+				if(FRigControlElement* ControlElement = HierarchyBeingCustomized->Find<FRigControlElement>(ElementKeyBeingCustomized))
 				{
-					const FRigControl& Control = ContainerBeingCustomized->ControlHierarchy[ElementKeyBeingCustomized.Name];
-					return Control.bLimitTranslation || Control.bLimitRotation || Control.bLimitScale;
+					return ControlElement->Settings.bLimitTranslation || ControlElement->Settings.bLimitRotation || ControlElement->Settings.bLimitScale;
 				}
 			}
 			return false;
@@ -2122,40 +1559,19 @@ bool FRigControlDetails::IsEnabled(ERigControlValueType InValueType) const
 	return true;
 }
 
-const TArray<TSharedPtr<FString>>& FRigControlDetails::GetGizmoNameList() const
+const TArray<TSharedPtr<FString>>& FRigControlElementDetails::GetGizmoNameList() const
 {
 	return GizmoNameList;
 }
 
-const TArray<TSharedPtr<FString>>& FRigControlDetails::GetControlTypeList() const
+const TArray<TSharedPtr<FString>>& FRigControlElementDetails::GetControlTypeList() const
 {
 	return ControlTypeList;
 }
 
-void FRigSpaceDetails::CustomizeDetails(IDetailLayoutBuilder& DetailBuilder)
+void FRigSpaceElementDetails::CustomizeDetails(IDetailLayoutBuilder& DetailBuilder)
 {
-	FRigElementDetails::CustomizeDetails(DetailBuilder);
-
-	DetailBuilder.HideProperty(TEXT("SpaceType"));
-	DetailBuilder.HideProperty(TEXT("ParentName"));
-
-	IDetailCategoryBuilder& Category = DetailBuilder.EditCategory(TEXT("FRigElement"), LOCTEXT("SpaceCategory", "Space"));
-	Category.InitiallyCollapsed(false);
-
-	Category.AddCustomRow(FText::FromString(TEXT("Name")))
-	.NameContent()
-	[
-		SNew(STextBlock)
-		.Text(FText::FromString(TEXT("Name")))
-		.Font(IDetailLayoutBuilder::GetDetailFont())
-	]
-	.ValueContent()
-	[
-		SNew(SEditableTextBox)
-		.Font(IDetailLayoutBuilder::GetDetailFont())
-		.Text(this, &FRigElementDetails::GetName)
-		.OnTextCommitted(this, &FRigElementDetails::SetName)
-	];
+	FRigTransformElementDetails::CustomizeDetails(DetailBuilder);
 }
 
 #undef LOCTEXT_NAMESPACE
