Control Rig: New rig hierarchy data structure refactoring

diff --git a/Engine/Plugins/Experimental/ControlRig/Source/ControlRigEditor/Private/ControlRigElementDetails.h b/Engine/Plugins/Experimental/ControlRig/Source/ControlRigEditor/Private/ControlRigElementDetails.h
--- a/Engine/Plugins/Experimental/ControlRig/Source/ControlRigEditor/Private/ControlRigElementDetails.h
+++ b/Engine/Plugins/Experimental/ControlRig/Source/ControlRigEditor/Private/ControlRigElementDetails.h
@@ -6,7 +6,7 @@
 #include "IDetailCustomization.h"
 #include "IPropertyTypeCustomization.h"
 #include "Rigs/RigHierarchyDefines.h"
-#include "Rigs/RigHierarchyContainer.h"
+#include "Rigs/RigHierarchy.h"
 #include "ControlRig.h"
 #include "ControlRigBlueprint.h"
 #include "Graph/ControlRigGraph.h"
@@ -85,7 +85,29 @@ enum class ERigElementDetailsTransformComponent : uint8
 	ScaleZ
 };
 
-class FRigElementDetails : public IDetailCustomization
+class FRigComputedTransformDetails : public IPropertyTypeCustomization
+{
+public:
+
+	static TSharedRef<IPropertyTypeCustomization> MakeInstance()
+	{
+		return MakeShareable(new FRigComputedTransformDetails);
+	}
+
+	/** IPropertyTypeCustomization interface */
+	virtual void CustomizeHeader(TSharedRef<class IPropertyHandle> InStructPropertyHandle, class FDetailWidgetRow& HeaderRow, IPropertyTypeCustomizationUtils& StructCustomizationUtils) override;
+	virtual void CustomizeChildren(TSharedRef<class IPropertyHandle> InStructPropertyHandle, class IDetailChildrenBuilder& StructBuilder, IPropertyTypeCustomizationUtils& StructCustomizationUtils) override;
+
+protected:
+
+	TSharedPtr<IPropertyHandle> TransformHandle;
+	FEditPropertyChain PropertyChain;
+	UControlRigBlueprint* BlueprintBeingCustomized;
+
+	void OnTransformChanged(FEditPropertyChain* InPropertyChain);
+};
+
+class FRigBaseElementDetails : public IDetailCustomization
 {
 public:
 
@@ -93,51 +115,54 @@ public:
 	virtual void CustomizeDetails(IDetailLayoutBuilder& DetailBuilder) override;
 
 	FRigElementKey GetElementKey() const { return ElementKeyBeingCustomized; }
-	FRigHierarchyContainer* GetHierarchy() const { return ContainerBeingCustomized; }
+	URigHierarchy* GetHierarchy() const { return HierarchyBeingCustomized; }
 	FText GetName() const { return FText::FromName(GetElementKey().Name); }
 	void SetName(const FText& InNewText, ETextCommit::Type InCommitType);
 
-	static float GetTransformComponent(const FTransform& InTransform, ERigElementDetailsTransformComponent InComponent);
-	static void SetTransformComponent(FTransform& OutTransform, ERigElementDetailsTransformComponent InComponent, float InNewValue);
-	static float GetEulerTransformComponent(const FEulerTransform& InTransform, ERigElementDetailsTransformComponent InComponent);
-	static void SetEulerTransformComponent(FEulerTransform& OutTransform, ERigElementDetailsTransformComponent InComponent, float InNewValue);
+	void OnStructContentsChanged(FProperty* InProperty, const TSharedRef<IPropertyUtilities> PropertyUtilities);
+
 protected:
 
 	FRigElementKey ElementKeyBeingCustomized;
 	UControlRigBlueprint* BlueprintBeingCustomized;
-	FRigHierarchyContainer* ContainerBeingCustomized;
+	URigHierarchy* HierarchyBeingCustomized;
+};
+
+class FRigTransformElementDetails : public FRigBaseElementDetails
+{
+public:
+
+	/** IDetailCustomization interface */
+	virtual void CustomizeDetails(IDetailLayoutBuilder& DetailBuilder) override;
 };
 
-class FRigBoneDetails : public FRigElementDetails
+class FRigBoneElementDetails : public FRigTransformElementDetails
 {
 public:
 
 	/** Makes a new instance of this detail layout class for a specific detail view requesting it */
 	static TSharedRef<IDetailCustomization> MakeInstance()
 	{
-		return MakeShareable(new FRigBoneDetails);
+		return MakeShareable(new FRigBoneElementDetails);
 	}
 
 	/** IDetailCustomization interface */
 	virtual void CustomizeDetails( IDetailLayoutBuilder& DetailBuilder ) override;
 
-	void OnStructContentsChanged(FProperty* InProperty, const TSharedRef<IPropertyUtilities> PropertyUtilities);
-	void OnAffectedListChanged(FProperty* InProperty, const TSharedRef<IPropertyUtilities> PropertyUtilities);
-
 private:
 
 	TSharedPtr<FRigInfluenceEntryModifier> InfluenceModifier;
 	TSharedPtr<FStructOnScope> InfluenceModifierStruct;
 };
 
-class FRigControlDetails : public FRigElementDetails
+class FRigControlElementDetails : public FRigTransformElementDetails
 {
 public:
 
 	/** Makes a new instance of this detail layout class for a specific detail view requesting it */
 	static TSharedRef<IDetailCustomization> MakeInstance()
 	{
-		return MakeShareable(new FRigControlDetails);
+		return MakeShareable(new FRigControlElementDetails);
 	}
 
 	/** IDetailCustomization interface */
@@ -145,19 +170,8 @@ public:
 
 	FText GetDisplayName() const;
 	void SetDisplayName(const FText& InNewText, ETextCommit::Type InCommitType, const TSharedRef<IPropertyUtilities> PropertyUtilities);
-	ECheckBoxState GetComponentValueBool(bool bInitial) const;
-	void SetComponentValueBool(ECheckBoxState InNewValue, bool bInitial, const TSharedRef<IPropertyUtilities> PropertyUtilities);
-	TOptional<float> GetComponentValueFloat(ERigControlValueType InValueType, ERigElementDetailsTransformComponent Component) const;
-	void SetComponentValueFloat(float InNewValue, ETextCommit::Type InCommitType, ERigControlValueType InValueType, ERigElementDetailsTransformComponent Component, const TSharedRef<IPropertyUtilities> PropertyUtilities);
-	void SetComponentValueFloat(float InNewValue, ERigControlValueType InValueType, ERigElementDetailsTransformComponent Component, const TSharedRef<IPropertyUtilities> PropertyUtilities);
-	TOptional<int32> GetComponentValueInteger(ERigControlValueType InValueType) const;
-	void SetComponentValueInteger(int32 InNewValue, ERigControlValueType InValueType, const TSharedRef<IPropertyUtilities> PropertyUtilities);
 	bool IsGizmoEnabled() const;
 	bool IsEnabled(ERigControlValueType InValueType) const;
-	void OnStructContentsChanged(FProperty* InProperty, const TSharedRef<IPropertyUtilities> PropertyUtilities);
-	void OnAffectedListChanged(FProperty* InProperty, const TSharedRef<IPropertyUtilities> PropertyUtilities);
-	int32 GetControlEnumValue(ERigControlValueType InValueType) const;
-	void OnControlEnumChanged(int32 InValue, ESelectInfo::Type InSelectInfo, ERigControlValueType InValueType, const TSharedRef<IPropertyUtilities> PropertyUtilities);
 
 	const TArray<TSharedPtr<FString>>& GetGizmoNameList() const;
 	const TArray<TSharedPtr<FString>>& GetControlTypeList() const;
@@ -167,16 +181,18 @@ private:
 	static TArray<TSharedPtr<FString>> ControlTypeList;
 	TSharedPtr<FRigInfluenceEntryModifier> InfluenceModifier;
 	TSharedPtr<FStructOnScope> InfluenceModifierStruct;
+	FEditPropertyChain OffsetPropertyChain;
+	FEditPropertyChain GizmoPropertyChain;
 };
 
-class FRigSpaceDetails : public FRigElementDetails
+class FRigSpaceElementDetails : public FRigTransformElementDetails
 {
 public:
 
 	/** Makes a new instance of this detail layout class for a specific detail view requesting it */
 	static TSharedRef<IDetailCustomization> MakeInstance()
 	{
-		return MakeShareable(new FRigSpaceDetails);
+		return MakeShareable(new FRigSpaceElementDetails);
 	}
 
 	/** IDetailCustomization interface */
diff --git a/Engine/Plugins/Experimental/ControlRig/Source/ControlRigEditor/Private/ControlRigThumbnailRenderer.cpp b/Engine/Plugins/Experimental/ControlRig/Source/ControlRigEditor/Private/ControlRigThumbnailRenderer.cpp
--- a/Engine/Plugins/Experimental/ControlRig/Source/ControlRigEditor/Private/ControlRigThumbnailRenderer.cpp
+++ b/Engine/Plugins/Experimental/ControlRig/Source/ControlRigEditor/Private/ControlRigThumbnailRenderer.cpp
@@ -25,18 +25,22 @@ bool UControlRigThumbnailRenderer::CanVisualizeAsset(UObject* Object)
 		{
 			if (InRigBlueprint->GizmoLibrary.IsValid())
 			{
-				for (const FRigControl& Control : InRigBlueprint->HierarchyContainer.ControlHierarchy)
+				bool bMissingMeshCount = true;
+				InRigBlueprint->Hierarchy->ForEach<FRigControlElement>([&](FRigControlElement* ControlElement) -> bool
 				{
-					if (const FControlRigGizmoDefinition* GizmoDef = InRigBlueprint->GizmoLibrary->GetGizmoByName(Control.GizmoName))
+					if (const FControlRigGizmoDefinition* GizmoDef = InRigBlueprint->GizmoLibrary->GetGizmoByName(ControlElement->Settings.GizmoName))
 					{
 						UStaticMesh* StaticMesh = GizmoDef->StaticMesh.Get();
 						if (StaticMesh == nullptr) // not yet loaded
 						{
+							bMissingMeshCount = true;
 							return false;
 						}
 					}
 					return true;
-				}
+				});
+
+				return !bMissingMeshCount;
 			}
 		}
 
@@ -97,9 +101,9 @@ void UControlRigThumbnailRenderer::AddAdditionalPreviewSceneContent(UObject* Obj
 
 		FTransform ComponentToWorld = ThumbnailScene->GetPreviewActor()->GetSkeletalMeshComponent()->GetComponentToWorld();
 
-		for (const FRigControl& Control : ControlRig->GetControlHierarchy())
+		ControlRig->GetHierarchy()->ForEach<FRigControlElement>([&](FRigControlElement* ControlElement) -> bool
 		{
-			switch (Control.ControlType)
+			switch (ControlElement->Settings.ControlType)
 			{
 				case ERigControlType::Float:
 				case ERigControlType::Integer:
@@ -111,15 +115,16 @@ void UControlRigThumbnailRenderer::AddAdditionalPreviewSceneContent(UObject* Obj
 				case ERigControlType::TransformNoScale:
 				case ERigControlType::EulerTransform:
 				{
-					if (const FControlRigGizmoDefinition* GizmoDef = RigBlueprint->GizmoLibrary->GetGizmoByName(Control.GizmoName))
+					if (const FControlRigGizmoDefinition* GizmoDef = RigBlueprint->GizmoLibrary->GetGizmoByName(ControlElement->Settings.GizmoName))
 					{
-						FTransform GizmoTransform = Control.GizmoTransform * (GizmoDef->Transform * ControlRig->GetControlGlobalTransform(Control.Name)) * ComponentToWorld;
 						UStaticMesh* StaticMesh = GizmoDef->StaticMesh.Get();
 						if (StaticMesh == nullptr) // not yet loaded
 						{
-							continue;
+							return true;
 						}
 
+						const FTransform GizmoGlobalTransform = ControlRig->GetHierarchy()->GetGlobalControlGizmoTransform(ControlElement->GetKey());
+
 						FActorSpawnParameters SpawnInfo;
 						SpawnInfo.SpawnCollisionHandlingOverride = ESpawnActorCollisionHandlingMethod::AlwaysSpawn;
 						SpawnInfo.bNoFail = true;
@@ -130,16 +135,17 @@ void UControlRigThumbnailRenderer::AddAdditionalPreviewSceneContent(UObject* Obj
 						UMaterial* DefaultMaterial = RigBlueprint->GizmoLibrary->DefaultMaterial.Get();
 						if (DefaultMaterial == nullptr) // not yet loaded
 						{
-							continue;
+							return true;
 						}
+
 						UMaterialInstanceDynamic* MaterialInstance = UMaterialInstanceDynamic::Create(DefaultMaterial, GizmoActor);
-						MaterialInstance->SetVectorParameterValue(RigBlueprint->GizmoLibrary->MaterialColorParameter, FVector(Control.GizmoColor));
+						MaterialInstance->SetVectorParameterValue(RigBlueprint->GizmoLibrary->MaterialColorParameter, FVector(ControlElement->Settings.GizmoColor));
 						GizmoActor->GetStaticMeshComponent()->SetMaterial(0, MaterialInstance);
 
-						GizmoActors.Add(Control.Name, GizmoActor);
+						GizmoActors.Add(ControlElement->GetName(), GizmoActor);
 
 						GizmoActor->GetStaticMeshComponent()->SetStaticMesh(StaticMesh);
-						GizmoActor->SetActorTransform(GizmoTransform);
+						GizmoActor->SetActorTransform(GizmoGlobalTransform);
 					}
 					break;
 				}
@@ -148,6 +154,7 @@ void UControlRigThumbnailRenderer::AddAdditionalPreviewSceneContent(UObject* Obj
 					break;
 				}
 			}
-		}
+			return true;
+        });
 	}
 }
\ No newline at end of file
diff --git a/Engine/Plugins/Experimental/ControlRig/Source/ControlRigEditor/Private/EditMode/ControlRigControlsProxy.cpp b/Engine/Plugins/Experimental/ControlRig/Source/ControlRigEditor/Private/EditMode/ControlRigControlsProxy.cpp
--- a/Engine/Plugins/Experimental/ControlRig/Source/ControlRigEditor/Private/EditMode/ControlRigControlsProxy.cpp
+++ b/Engine/Plugins/Experimental/ControlRig/Source/ControlRigEditor/Private/EditMode/ControlRigControlsProxy.cpp
@@ -4,7 +4,7 @@
 #include "EditorModeManager.h"
 #include "ControlRigEditMode.h"
 #include "Sequencer/ControlRigSequence.h"
-#include "Rigs/RigControlHierarchy.h"
+#include "Rigs/RigHierarchy.h"
 
 #include "Components/SkeletalMeshComponent.h"
 #include "Subsystems/AssetEditorSubsystem.h"
@@ -73,7 +73,7 @@ void FControlRigEnumControlProxyValueDetails::OnEnumValueChanged(int32 InValue,
 
 void UControlRigControlsProxy::SelectionChanged(bool bInSelected)
 {
-	if (RigControl)
+	if (ControlElement)
 	{
 		Modify();
 		const FName PropertyName("bSelected");
@@ -101,7 +101,7 @@ void UControlRigControlsProxy::PostEditChangeProperty(struct FPropertyChangedEve
 	Super::PostEditChangeProperty(PropertyChangedEvent);
 	if (PropertyChangedEvent.Property && PropertyChangedEvent.Property->GetFName() == GET_MEMBER_NAME_CHECKED(UControlRigControlsProxy, bSelected))
 	{
-		if (RigControl && ControlRig.IsValid())
+		if (ControlElement && ControlRig.IsValid())
 		{
 			if (bSelected)
 			{
@@ -118,7 +118,7 @@ void UControlRigControlsProxy::PostEditChangeProperty(struct FPropertyChangedEve
 #if WITH_EDITOR
 void UControlRigControlsProxy::PostEditUndo()
 {
-	if (RigControl && ControlRig.IsValid() && ControlRig->GetControlHierarchy().GetIndex(ControlName) != INDEX_NONE)
+	if (ControlElement && ControlRig.IsValid() && ControlRig->GetHierarchy()->Contains(FRigElementKey(ControlName, ERigElementType::Control)))
 	{
 		if (bSelected)
 		{
@@ -136,7 +136,7 @@ void UControlRigTransformControlProxy::PostEditChangeProperty(struct FPropertyCh
 	if ((PropertyChangedEvent.Property && PropertyChangedEvent.Property->GetFName() == GET_MEMBER_NAME_CHECKED(UControlRigTransformControlProxy, Transform))
 		|| (PropertyChangedEvent.MemberProperty && PropertyChangedEvent.MemberProperty->GetFName() == GET_MEMBER_NAME_CHECKED(UControlRigTransformControlProxy, Transform)))
 	{
-		if (RigControl && ControlRig.IsValid())
+		if (ControlElement && ControlRig.IsValid())
 		{
 			//MUST set through ControlRig
 			FControlRigInteractionScope InteractionScope(ControlRig.Get());
@@ -149,22 +149,20 @@ void UControlRigTransformControlProxy::PostEditChangeProperty(struct FPropertyCh
 
 void UControlRigTransformControlProxy::ValueChanged()
 {
-	if (RigControl)
+	if (ControlElement && ControlRig.IsValid())
 	{
 		Modify();
 		const FName PropertyName("Transform");
 		FTrackInstancePropertyBindings Binding(PropertyName, PropertyName.ToString());
-		FTransform NewTransform = RigControl->Value.Get<FTransform>();
-
+		const FTransform NewTransform = ControlRig.Get()->GetHierarchy()->GetControlValue(ControlElement, ERigControlValueType::Current).Get<FTransform>();
 		Binding.CallFunction<FTransform>(*this, NewTransform);
-
 	}
 }
 
 #if WITH_EDITOR
 void UControlRigTransformControlProxy::PostEditUndo()
 {
-	if (RigControl && ControlRig.IsValid() && ControlRig->GetControlHierarchy().GetIndex(ControlName) != INDEX_NONE)
+	if (ControlElement && ControlRig.IsValid() && ControlRig->GetHierarchy()->Contains(FRigElementKey(ControlName, ERigElementType::Control)))
 	{
 		if (bSelected)
 		{
@@ -179,7 +177,7 @@ void UControlRigTransformControlProxy::PostEditUndo()
 void UControlRigTransformControlProxy::SetKey(const IPropertyHandle& KeyedPropertyHandle)
 {
 
-	if (RigControl)
+	if (ControlElement)
 	{
 		ControlRig->SetControlValue<FTransform>(ControlName, Transform, true, EControlRigSetKey::Always);
 	}
@@ -191,7 +189,7 @@ void UControlRigTransformNoScaleControlProxy::PostEditChangeProperty(struct FPro
 	if ((PropertyChangedEvent.Property && PropertyChangedEvent.Property->GetFName() == GET_MEMBER_NAME_CHECKED(UControlRigTransformNoScaleControlProxy, Transform))
 		|| (PropertyChangedEvent.MemberProperty && PropertyChangedEvent.MemberProperty->GetFName() == GET_MEMBER_NAME_CHECKED(UControlRigTransformNoScaleControlProxy, Transform)))
 	{
-		if (RigControl && ControlRig.IsValid())
+		if (ControlElement && ControlRig.IsValid())
 		{
 			//MUST set through ControlRig
 			FControlRigInteractionScope InteractionScope(ControlRig.Get());
@@ -204,12 +202,12 @@ void UControlRigTransformNoScaleControlProxy::PostEditChangeProperty(struct FPro
 
 void UControlRigTransformNoScaleControlProxy::ValueChanged()
 {
-	if (RigControl)
+	if (ControlElement)
 	{
 		Modify();
 		const FName PropertyName("Transform");
 		FTrackInstancePropertyBindings Binding(PropertyName, PropertyName.ToString());
-		FTransformNoScale NewTransform = RigControl->Value.Get<FTransformNoScale>();
+		const FTransformNoScale NewTransform = ControlRig.Get()->GetHierarchy()->GetControlValue(ControlElement, ERigControlValueType::Current).Get<FTransformNoScale>();
 		Binding.CallFunction<FTransformNoScale>(*this, NewTransform);
 	}
 }
@@ -217,7 +215,7 @@ void UControlRigTransformNoScaleControlProxy::ValueChanged()
 #if WITH_EDITOR
 void UControlRigTransformNoScaleControlProxy::PostEditUndo()
 {
-	if (RigControl && ControlRig.IsValid() && ControlRig->GetControlHierarchy().GetIndex(ControlName) != INDEX_NONE)
+	if (ControlElement && ControlRig.IsValid() && ControlRig->GetHierarchy()->Contains(FRigElementKey(ControlName, ERigElementType::Control)))
 	{
 		if (bSelected)
 		{
@@ -231,7 +229,7 @@ void UControlRigTransformNoScaleControlProxy::PostEditUndo()
 
 void UControlRigTransformNoScaleControlProxy::SetKey(const IPropertyHandle& KeyedPropertyHandle)
 {
-	if (RigControl)
+	if (ControlElement)
 	{
 		ControlRig->SetControlValue<FTransformNoScale>(ControlName, Transform, true, EControlRigSetKey::Always);
 	}
@@ -244,7 +242,7 @@ void UControlRigEulerTransformControlProxy::PostEditChangeProperty(struct FPrope
 	if ((PropertyChangedEvent.Property && PropertyChangedEvent.Property->GetFName() == GET_MEMBER_NAME_CHECKED(UControlRigEulerTransformControlProxy, Transform))
 		|| (PropertyChangedEvent.MemberProperty && PropertyChangedEvent.MemberProperty->GetFName() == GET_MEMBER_NAME_CHECKED(UControlRigEulerTransformControlProxy, Transform)))
 	{
-		if (RigControl && ControlRig.IsValid())
+		if (ControlElement && ControlRig.IsValid())
 		{
 			//MUST set through ControlRig
 			FControlRigInteractionScope InteractionScope(ControlRig.Get());
@@ -257,12 +255,12 @@ void UControlRigEulerTransformControlProxy::PostEditChangeProperty(struct FPrope
 
 void UControlRigEulerTransformControlProxy::ValueChanged()
 {
-	if (RigControl)
+	if (ControlElement)
 	{
 		Modify();
 		const FName PropertyName("Transform");
 		FTrackInstancePropertyBindings Binding(PropertyName, PropertyName.ToString());
-		FEulerTransform NewTransform = RigControl->Value.Get<FEulerTransform>();
+		const FEulerTransform NewTransform = ControlRig.Get()->GetHierarchy()->GetControlValue(ControlElement, ERigControlValueType::Current).Get<FEulerTransform>();
 
 		Binding.CallFunction<FEulerTransform>(*this, NewTransform);
 
@@ -272,7 +270,7 @@ void UControlRigEulerTransformControlProxy::ValueChanged()
 #if WITH_EDITOR
 void UControlRigEulerTransformControlProxy::PostEditUndo()
 {
-	if (RigControl && ControlRig.IsValid() && ControlRig->GetControlHierarchy().GetIndex(ControlName) != INDEX_NONE)
+	if (ControlElement && ControlRig.IsValid() && ControlRig->GetHierarchy()->Contains(FRigElementKey(ControlName, ERigElementType::Control)))
 	{
 		if (bSelected)
 		{
@@ -286,7 +284,7 @@ void UControlRigEulerTransformControlProxy::PostEditUndo()
 
 void UControlRigEulerTransformControlProxy::SetKey(const IPropertyHandle& KeyedPropertyHandle)
 {
-	if (RigControl)
+	if (ControlElement)
 	{
 		ControlRig->SetControlValue<FEulerTransform>(ControlName, Transform, true, EControlRigSetKey::Always);
 	}
@@ -297,7 +295,7 @@ void UControlRigFloatControlProxy::PostEditChangeProperty(struct FPropertyChange
 	Super::PostEditChangeProperty(PropertyChangedEvent);
 	if (PropertyChangedEvent.Property && PropertyChangedEvent.Property->GetFName() == GET_MEMBER_NAME_CHECKED(UControlRigFloatControlProxy, Float))
 	{
-		if (RigControl && ControlRig.IsValid())
+		if (ControlElement && ControlRig.IsValid())
 		{
 			//MUST set through ControlRig
 			FControlRigInteractionScope InteractionScope(ControlRig.Get());
@@ -310,13 +308,12 @@ void UControlRigFloatControlProxy::PostEditChangeProperty(struct FPropertyChange
 
 void UControlRigFloatControlProxy::ValueChanged()
 {
-	if (RigControl)
+	if (ControlElement)
 	{
 		Modify();
 		const FName PropertyName("Float");
 		FTrackInstancePropertyBindings Binding(PropertyName, PropertyName.ToString());
-		float Val = RigControl->Value.Get<float>();
-
+		const float Val = ControlRig.Get()->GetHierarchy()->GetControlValue(ControlElement, ERigControlValueType::Current).Get<float>();
 		Binding.CallFunction<float>(*this, Val);
 	}
 }
@@ -324,7 +321,7 @@ void UControlRigFloatControlProxy::ValueChanged()
 #if WITH_EDITOR
 void UControlRigFloatControlProxy::PostEditUndo()
 {
-	if (RigControl && ControlRig.IsValid() && ControlRig->GetControlHierarchy().GetIndex(ControlName) != INDEX_NONE)
+	if (ControlElement && ControlRig.IsValid() && ControlRig->GetHierarchy()->Contains(FRigElementKey(ControlName, ERigElementType::Control)))
 	{
 		if (bSelected)
 		{
@@ -339,7 +336,7 @@ void UControlRigFloatControlProxy::PostEditUndo()
 
 void UControlRigFloatControlProxy::SetKey(const IPropertyHandle& KeyedPropertyHandle)
 {
-	if (RigControl)
+	if (ControlElement)
 	{
 		ControlRig->SetControlValue<float>(ControlName, Float, true, EControlRigSetKey::Always);
 	}
@@ -350,7 +347,7 @@ void UControlRigIntegerControlProxy::PostEditChangeProperty(struct FPropertyChan
 	Super::PostEditChangeProperty(PropertyChangedEvent);
 	if (PropertyChangedEvent.Property && PropertyChangedEvent.Property->GetFName() == GET_MEMBER_NAME_CHECKED(UControlRigIntegerControlProxy, Integer))
 	{
-		if (RigControl && ControlRig.IsValid())
+		if (ControlElement && ControlRig.IsValid())
 		{
 			//MUST set through ControlRig
 			FControlRigInteractionScope InteractionScope(ControlRig.Get());
@@ -362,13 +359,12 @@ void UControlRigIntegerControlProxy::PostEditChangeProperty(struct FPropertyChan
 
 void UControlRigIntegerControlProxy::ValueChanged()
 {
-	if (RigControl)
+	if (ControlElement)
 	{
 		Modify();
 		const FName PropertyName("Integer");
 		FTrackInstancePropertyBindings Binding(PropertyName, PropertyName.ToString());
-		int32 Val = RigControl->Value.Get<int32>();
-
+		const int32 Val = ControlRig.Get()->GetHierarchy()->GetControlValue(ControlElement, ERigControlValueType::Current).Get<int32>();
 		Binding.CallFunction<int32>(*this, Val);
 	}
 }
@@ -376,7 +372,7 @@ void UControlRigIntegerControlProxy::ValueChanged()
 #if WITH_EDITOR
 void UControlRigIntegerControlProxy::PostEditUndo()
 {
-	if (RigControl && ControlRig.IsValid() && ControlRig->GetControlHierarchy().GetIndex(ControlName) != INDEX_NONE)
+	if (ControlElement && ControlRig.IsValid() && ControlRig->GetHierarchy()->Contains(FRigElementKey(ControlName, ERigElementType::Control)))
 	{
 		if (bSelected)
 		{
@@ -390,7 +386,7 @@ void UControlRigIntegerControlProxy::PostEditUndo()
 
 void UControlRigIntegerControlProxy::SetKey(const IPropertyHandle& KeyedPropertyHandle)
 {
-	if (RigControl)
+	if (ControlElement)
 	{
 		ControlRig->SetControlValue<int32>(ControlName, Integer, true, EControlRigSetKey::Always);
 	}
@@ -401,7 +397,7 @@ void UControlRigEnumControlProxy::PostEditChangeProperty(struct FPropertyChanged
 	Super::PostEditChangeProperty(PropertyChangedEvent);
 	if (PropertyChangedEvent.Property && PropertyChangedEvent.Property->GetFName() == GET_MEMBER_NAME_CHECKED(UControlRigEnumControlProxy, Enum))
 	{
-		if (RigControl && ControlRig.IsValid())
+		if (ControlElement && ControlRig.IsValid())
 		{
 			//MUST set through ControlRig
 			FControlRigInteractionScope InteractionScope(ControlRig.Get());
@@ -414,15 +410,15 @@ void UControlRigEnumControlProxy::PostEditChangeProperty(struct FPropertyChanged
 
 void UControlRigEnumControlProxy::ValueChanged()
 {
-	if (RigControl)
+	if (ControlElement)
 	{
 		Modify();
 		const FName PropertyName("Enum");
 		FTrackInstancePropertyBindings Binding(PropertyName, PropertyName.ToString());
 
 		FControlRigEnumControlProxyValue Val;
-		Val.EnumType = RigControl->ControlEnum;
-		Val.EnumIndex = RigControl->Value.Get<int32>();
+		Val.EnumType = ControlElement->Settings.ControlEnum;
+		Val.EnumIndex = ControlRig.Get()->GetHierarchy()->GetControlValue(ControlElement, ERigControlValueType::Current).Get<int32>();
 
 		Binding.CallFunction<FControlRigEnumControlProxyValue>(*this, Val);
 	}
@@ -431,7 +427,7 @@ void UControlRigEnumControlProxy::ValueChanged()
 #if WITH_EDITOR
 void UControlRigEnumControlProxy::PostEditUndo()
 {
-	if (RigControl && ControlRig.IsValid() && ControlRig->GetControlHierarchy().GetIndex(ControlName) != INDEX_NONE)
+	if (ControlElement && ControlRig.IsValid() && ControlRig->GetHierarchy()->Contains(FRigElementKey(ControlName, ERigElementType::Control)))
 	{
 		if (bSelected)
 		{
@@ -445,7 +441,7 @@ void UControlRigEnumControlProxy::PostEditUndo()
 
 void UControlRigEnumControlProxy::SetKey(const IPropertyHandle& KeyedPropertyHandle)
 {
-	if (RigControl)
+	if (ControlElement)
 	{
 		ControlRig->SetControlValue<int32>(ControlName, Enum.EnumIndex, true, EControlRigSetKey::Always);
 	}
@@ -457,7 +453,7 @@ void UControlRigVectorControlProxy::PostEditChangeProperty(struct FPropertyChang
 	if ((PropertyChangedEvent.Property && PropertyChangedEvent.Property->GetFName() == GET_MEMBER_NAME_CHECKED(UControlRigVectorControlProxy, Vector))
 		|| (PropertyChangedEvent.MemberProperty && PropertyChangedEvent.MemberProperty->GetFName() == GET_MEMBER_NAME_CHECKED(UControlRigVectorControlProxy, Vector)))
 	{
-		if (RigControl && ControlRig.IsValid())
+		if (ControlElement && ControlRig.IsValid())
 		{
 			//MUST set through ControlRig
 			FControlRigInteractionScope InteractionScope(ControlRig.Get());
@@ -470,13 +466,12 @@ void UControlRigVectorControlProxy::PostEditChangeProperty(struct FPropertyChang
 
 void UControlRigVectorControlProxy::ValueChanged()
 {
-	if (RigControl)
+	if (ControlElement)
 	{
 		Modify();
 		const FName PropertyName("Vector");
 		FTrackInstancePropertyBindings Binding(PropertyName, PropertyName.ToString());
-		FVector Val = RigControl->Value.Get<FVector>();
-
+		const FVector Val = ControlRig.Get()->GetHierarchy()->GetControlValue(ControlElement, ERigControlValueType::Current).Get<FVector>();
 		Binding.CallFunction<FVector>(*this, Val);
 	}
 }
@@ -484,7 +479,7 @@ void UControlRigVectorControlProxy::ValueChanged()
 #if WITH_EDITOR
 void UControlRigVectorControlProxy::PostEditUndo()
 {
-	if (RigControl && ControlRig.IsValid() && ControlRig->GetControlHierarchy().GetIndex(ControlName) != INDEX_NONE)
+	if (ControlElement && ControlRig.IsValid() && ControlRig->GetHierarchy()->Contains(FRigElementKey(ControlName, ERigElementType::Control)))
 	{
 		if (bSelected)
 		{
@@ -499,7 +494,7 @@ void UControlRigVectorControlProxy::PostEditUndo()
 
 void UControlRigVectorControlProxy::SetKey(const IPropertyHandle& KeyedPropertyHandle)
 {
-	if (RigControl)
+	if (ControlElement)
 	{
 		ControlRig->SetControlValue<FVector>(ControlName, Vector, true, EControlRigSetKey::Always);
 	}
@@ -511,7 +506,7 @@ void UControlRigVector2DControlProxy::PostEditChangeProperty(struct FPropertyCha
 	if ((PropertyChangedEvent.Property && PropertyChangedEvent.Property->GetFName() == GET_MEMBER_NAME_CHECKED(UControlRigVector2DControlProxy, Vector2D))
 		|| ((PropertyChangedEvent.MemberProperty && PropertyChangedEvent.MemberProperty->GetFName() == GET_MEMBER_NAME_CHECKED(UControlRigVector2DControlProxy, Vector2D))))
 	{
-		if (RigControl && ControlRig.IsValid())
+		if (ControlElement && ControlRig.IsValid())
 		{
 			//MUST set through ControlRig
 			FControlRigInteractionScope InteractionScope(ControlRig.Get());
@@ -524,13 +519,12 @@ void UControlRigVector2DControlProxy::PostEditChangeProperty(struct FPropertyCha
 
 void UControlRigVector2DControlProxy::ValueChanged()
 {
-	if (RigControl)
+	if (ControlElement)
 	{
 		Modify();
 		const FName PropertyName("Vector2D");
 		FTrackInstancePropertyBindings Binding(PropertyName, PropertyName.ToString());
-		FVector2D Val = RigControl->Value.Get<FVector2D>();
-
+		const FVector2D Val = ControlRig.Get()->GetHierarchy()->GetControlValue(ControlElement, ERigControlValueType::Current).Get<FVector2D>();
 		Binding.CallFunction<FVector2D>(*this, Val);
 	}
 }
@@ -538,7 +532,7 @@ void UControlRigVector2DControlProxy::ValueChanged()
 #if WITH_EDITOR
 void UControlRigVector2DControlProxy::PostEditUndo()
 {
-	if (RigControl && ControlRig.IsValid() && ControlRig->GetControlHierarchy().GetIndex(ControlName) != INDEX_NONE)
+	if (ControlElement && ControlRig.IsValid() && ControlRig->GetHierarchy()->Contains(FRigElementKey(ControlName, ERigElementType::Control)))
 	{
 		if (bSelected)
 		{
@@ -553,7 +547,7 @@ void UControlRigVector2DControlProxy::PostEditUndo()
 
 void UControlRigVector2DControlProxy::SetKey(const IPropertyHandle& KeyedPropertyHandle)
 {
-	if (RigControl)
+	if (ControlElement)
 	{
 		ControlRig->SetControlValue<FVector2D>(ControlName, Vector2D, true, EControlRigSetKey::Always);
 	}
@@ -565,7 +559,7 @@ void UControlRigBoolControlProxy::PostEditChangeProperty(struct FPropertyChanged
 	Super::PostEditChangeProperty(PropertyChangedEvent);
 	if (PropertyChangedEvent.Property && PropertyChangedEvent.Property->GetFName() == GET_MEMBER_NAME_CHECKED(UControlRigBoolControlProxy, Bool))
 	{
-		if (RigControl && ControlRig.IsValid())
+		if (ControlElement && ControlRig.IsValid())
 		{
 			//MUST set through ControlRig
 			FControlRigInteractionScope InteractionScope(ControlRig.Get());
@@ -577,13 +571,12 @@ void UControlRigBoolControlProxy::PostEditChangeProperty(struct FPropertyChanged
 
 void UControlRigBoolControlProxy::ValueChanged()
 {
-	if (RigControl)
+	if (ControlElement)
 	{
 		Modify();
 		const FName PropertyName("Bool");
 		FTrackInstancePropertyBindings Binding(PropertyName, PropertyName.ToString());
-		bool Val = RigControl->Value.Get<bool>();
-
+		const bool Val = ControlRig.Get()->GetHierarchy()->GetControlValue(ControlElement, ERigControlValueType::Current).Get<bool>();
 		Binding.CallFunction<bool>(*this, Val);
 	}
 }
@@ -591,7 +584,7 @@ void UControlRigBoolControlProxy::ValueChanged()
 #if WITH_EDITOR
 void UControlRigBoolControlProxy::PostEditUndo()
 {
-	if (RigControl && ControlRig.IsValid() && ControlRig->GetControlHierarchy().GetIndex(ControlName) != INDEX_NONE)
+	if (ControlElement && ControlRig.IsValid() && ControlRig->GetHierarchy()->Contains(FRigElementKey(ControlName, ERigElementType::Control)))
 	{
 		if (bSelected)
 		{
@@ -606,7 +599,7 @@ void UControlRigBoolControlProxy::PostEditUndo()
 
 void UControlRigBoolControlProxy::SetKey(const IPropertyHandle& KeyedPropertyHandle)
 {
-	if (RigControl)
+	if (ControlElement)
 	{
 		ControlRig->SetControlValue<bool>(ControlName, Bool, true, EControlRigSetKey::Always);
 	}
@@ -625,12 +618,12 @@ UControlRigControlsProxy* UControlRigDetailPanelControlProxies::FindProxy(const
 	return nullptr;
 }
 
-void UControlRigDetailPanelControlProxies::AddProxy(const FName& Name, UControlRig* ControlRig, FRigControl* RigControl)
+void UControlRigDetailPanelControlProxies::AddProxy(const FName& Name, UControlRig* ControlRig, FRigControlElement* ControlElement)
 {
 	UControlRigControlsProxy* Proxy = FindProxy(Name);
-	if (!Proxy && RigControl != nullptr)
+	if (!Proxy && ControlElement != nullptr)
 	{
-		switch(RigControl->ControlType)
+		switch(ControlElement->Settings.ControlType)
 		{
 			case ERigControlType::Transform:
 			{
@@ -657,14 +650,14 @@ void UControlRigDetailPanelControlProxies::AddProxy(const FName& Name, UControlR
 			}
 			case ERigControlType::Integer:
 			{
-				if(RigControl->ControlEnum == nullptr)
+				if(ControlElement->Settings.ControlEnum == nullptr)
 				{
 					Proxy = NewObject<UControlRigIntegerControlProxy>(GetTransientPackage(), NAME_None);
 				}
 				else
 				{
 					UControlRigEnumControlProxy* EnumProxy = NewObject<UControlRigEnumControlProxy>(GetTransientPackage(), NAME_None);
-					EnumProxy->Enum.EnumType = RigControl->ControlEnum;
+					EnumProxy->Enum.EnumType = ControlElement->Settings.ControlEnum;
 					Proxy = EnumProxy;
 				}
 				break;
@@ -697,7 +690,7 @@ void UControlRigDetailPanelControlProxies::AddProxy(const FName& Name, UControlR
 		{
 			Proxy->SetFlags(RF_Transactional);
 			Proxy->SetName(Name);
-			Proxy->RigControl = RigControl;
+			Proxy->ControlElement = ControlElement;
 			Proxy->ControlRig = ControlRig;
 			Proxy->ValueChanged();
 			AllProxies.Add(Name, Proxy);
@@ -732,13 +725,12 @@ void UControlRigDetailPanelControlProxies::RemoveAllProxies()
 void UControlRigDetailPanelControlProxies::RecreateAllProxies(UControlRig* InControlRig)
 {
 	RemoveAllProxies();
-	const TArray<FRigControl>& Controls = InControlRig->AvailableControls();
-	for (const FRigControl& RigControl : Controls)
+	TArray<FRigControlElement*> Controls = InControlRig->AvailableControls();
+	for (FRigControlElement* ControlElement : Controls)
 	{
-		FRigControl* PRigControl = const_cast<FRigControl*>(&RigControl);
-		if(PRigControl->bGizmoEnabled && PRigControl->bAnimatable)
+		if(ControlElement->Settings.bGizmoEnabled && ControlElement->Settings.bAnimatable)
 		{
-			AddProxy(RigControl.Name, InControlRig, PRigControl);
+			AddProxy(ControlElement->GetName(), InControlRig, ControlElement);
 		}
 	}
 }
diff --git a/Engine/Plugins/Experimental/ControlRig/Source/ControlRigEditor/Private/EditMode/ControlRigControlsProxy.h b/Engine/Plugins/Experimental/ControlRig/Source/ControlRigEditor/Private/EditMode/ControlRigControlsProxy.h
--- a/Engine/Plugins/Experimental/ControlRig/Source/ControlRigEditor/Private/EditMode/ControlRigControlsProxy.h
+++ b/Engine/Plugins/Experimental/ControlRig/Source/ControlRigEditor/Private/EditMode/ControlRigControlsProxy.h
@@ -13,7 +13,7 @@
 
 #include "ControlRigControlsProxy.generated.h"
 
-struct FRigControl;
+struct FRigControlElement;
 class UControlRig;
 class IPropertyHandle;
 
@@ -38,7 +38,7 @@ public:
 
 public:
 
-	FRigControl* RigControl;
+	FRigControlElement* ControlElement;
 	TWeakObjectPtr<UControlRig> ControlRig;
 
 	UPROPERTY()
@@ -310,7 +310,7 @@ protected:
 
 public:
 	UControlRigControlsProxy* FindProxy(const FName& Name) const;
-	void AddProxy(const FName& Name, UControlRig* InControlRig, FRigControl* InControl);
+	void AddProxy(const FName& Name, UControlRig* InControlRig, FRigControlElement* ControlElement);
 	void RemoveProxy(const FName& Name );
 	void ProxyChanged(const FName& Name);
 	void RemoveAllProxies();
diff --git a/Engine/Plugins/Experimental/ControlRig/Source/ControlRigEditor/Private/EditMode/ControlRigEditMode.cpp b/Engine/Plugins/Experimental/ControlRig/Source/ControlRigEditor/Private/EditMode/ControlRigEditMode.cpp
--- a/Engine/Plugins/Experimental/ControlRig/Source/ControlRigEditor/Private/EditMode/ControlRigEditMode.cpp
+++ b/Engine/Plugins/Experimental/ControlRig/Source/ControlRigEditor/Private/EditMode/ControlRigEditMode.cpp
@@ -258,10 +258,10 @@ void FControlRigEditMode::SetObjects_Internal()
 
 	if (InteractionControlRig)
 	{
-		InteractionControlRig->Hierarchy.OnElementSelected.RemoveAll(this);
+		InteractionControlRig->GetHierarchy()->OnModified().RemoveAll(this);
 		InteractionControlRig->ControlModified().RemoveAll(this);
 			
-		InteractionControlRig->Hierarchy.OnElementSelected.AddSP(this, &FControlRigEditMode::OnRigElementSelected);
+		InteractionControlRig->GetHierarchy()->OnModified().AddSP(this, &FControlRigEditMode::OnHierarchyModified);
 		InteractionControlRig->ControlModified().AddSP(this, &FControlRigEditMode::OnControlModified);
 	}
 
@@ -415,12 +415,12 @@ void FControlRigEditMode::Tick(FEditorViewportClient* ViewportClient, float Delt
 				{
 					if (UControlRig* ControlRig = GetControlRig(true))
 					{
-						FRigControl* Control = ControlRig->FindControl(SelectedKey.Name);
-						if (Control)
+						FRigControlElement* ControlElement = ControlRig->FindControl(SelectedKey.Name);
+						if (ControlElement)
 						{
-							if (!ControlRig->IsCurveControl(Control))
+							if (!ControlRig->IsCurveControl(ControlElement))
 							{
-								ControlProxy->AddProxy(SelectedKey.Name, ControlRig, Control);
+								ControlProxy->AddProxy(SelectedKey.Name, ControlRig, ControlElement);
 							}
 						}
 					}
@@ -447,9 +447,9 @@ void FControlRigEditMode::Tick(FEditorViewportClient* ViewportClient, float Delt
 				{
 					if (!ModeSupportedByGizmoActor(GizmoActor, CurrentWidgetMode))
 					{
-						if (FRigControl* Control = ControlRig->FindControl(SelectedRigElement.Name))
+						if (FRigControlElement* ControlElement = ControlRig->FindControl(SelectedRigElement.Name))
 						{
-							switch (Control->ControlType)
+							switch (ControlElement->Settings.ControlType)
 							{
 								case ERigControlType::Float:
 								case ERigControlType::Integer:
@@ -509,36 +509,37 @@ void FControlRigEditMode::Render(const FSceneView* View, FViewport* Viewport, FP
 			}
 		}
 
+		URigHierarchy* Hierarchy = ControlRig->GetHierarchy();
 		if (Settings->bDisplayHierarchy)
 		{
-			// each base hierarchy Bone
-			const FRigBoneHierarchy& BaseHierarchy = ControlRig->GetBoneHierarchy();
-			for (int32 BoneIndex = 0; BoneIndex < BaseHierarchy.Num(); ++BoneIndex)
+			Hierarchy->ForEach<FRigTransformElement>([PDI, Hierarchy, ComponentTransform](FRigTransformElement* TransformElement) -> bool
 			{
-				const FRigBone& CurrentBone = BaseHierarchy[BoneIndex];
-				const FTransform Transform = BaseHierarchy.GetGlobalTransform(BoneIndex);
+				const FTransform Transform = Hierarchy->GetTransform(TransformElement, ERigTransformType::CurrentGlobal);
 
-				if (CurrentBone.ParentIndex != INDEX_NONE)
+				TArray<FRigBaseElement*> Parents = Hierarchy->GetParents(TransformElement);
+				for(FRigBaseElement* ParentElement : Parents)
 				{
-					const FTransform ParentTransform = BaseHierarchy.GetGlobalTransform(CurrentBone.ParentIndex);
-
-					PDI->DrawLine(ComponentTransform.TransformPosition(Transform.GetLocation()),ComponentTransform.TransformPosition(ParentTransform.GetLocation()), FLinearColor::White, SDPG_Foreground);
-				}
+					if(FRigTransformElement* ParentTransformElement = Cast<FRigTransformElement>(ParentElement))
+					{
+						const FTransform ParentTransform = Hierarchy->GetTransform(ParentTransformElement, ERigTransformType::CurrentGlobal);
+                        PDI->DrawLine(ComponentTransform.TransformPosition(Transform.GetLocation()),ComponentTransform.TransformPosition(ParentTransform.GetLocation()), FLinearColor::White, SDPG_Foreground);
+					}
+                }
 
 				PDI->DrawPoint(ComponentTransform.TransformPosition(Transform.GetLocation()), FLinearColor::White, 5.0f, SDPG_Foreground);
-			}
+
+				return true;
+			});
 		}
 
 		if (Settings->bDisplaySpaces || ControlRig->IsSetupModeEnabled())
 		{
-			FRigSpaceHierarchy& SpaceHierarchy = ControlRig->GetSpaceHierarchy();
-
 			TArray<FTransform> SpaceTransforms;
-			for (const FRigSpace& Space : SpaceHierarchy)
-			{
-				SpaceTransforms.Add(SpaceHierarchy.GetGlobalTransform(Space.Index));
-			}
-
+			Hierarchy->ForEach<FRigSpaceElement>([&SpaceTransforms, Hierarchy](FRigSpaceElement* SpaceElement) -> bool
+            {
+				SpaceTransforms.Add(Hierarchy->GetTransform(SpaceElement, ERigTransformType::CurrentGlobal));
+				return true;
+			});
 			GetControlRig(true)->DrawInterface.DrawAxes(FTransform::Identity, SpaceTransforms, Settings->AxisScale);
 		}
 
@@ -546,7 +547,6 @@ void FControlRigEditMode::Render(const FSceneView* View, FViewport* Viewport, FP
 		{
 			if (ControlRig->GetWorld() && ControlRig->GetWorld()->IsPreviewWorld())
 			{
-				const FRigHierarchyContainer* Hierarchy = ControlRig->GetHierarchy();
 				const float Scale = Settings->AxisScale;
 				PDI->AddReserveLines(SDPG_Foreground, SelectedRigElements.Num() * 3);
 
@@ -1253,11 +1253,14 @@ void FControlRigEditMode::ClearRigElementSelection(uint32 InTypes)
 	UControlRigBlueprint* Blueprint = Cast<UControlRigBlueprint>(InteractionRig->GetClass()->ClassGeneratedBy);
 	if (Blueprint)
 	{
-		Blueprint->HierarchyContainer.ClearSelection();
+		Blueprint->GetHierarchyController()->ClearSelection();
 	}
 	if (IsInLevelEditor())
 	{
-		InteractionRig->Hierarchy.ClearSelection();
+		if(URigHierarchyController* Controller = InteractionRig->GetHierarchy()->GetController())
+		{
+			Controller->ClearSelection();
+		}
 	}
 }
 
@@ -1273,11 +1276,14 @@ void FControlRigEditMode::SetRigElementSelectionInternal(ERigElementType Type, c
 	UControlRigBlueprint* Blueprint = Cast<UControlRigBlueprint>(InteractionRig->GetClass()->ClassGeneratedBy);
 	if (Blueprint)
 	{
-		Blueprint->HierarchyContainer.Select(FRigElementKey(InRigElementName, Type), bSelected);
+		Blueprint->GetHierarchyController()->SelectElement(FRigElementKey(InRigElementName, Type), bSelected);
 	}
 	if (IsInLevelEditor())
 	{
-		InteractionRig->Hierarchy.Select(FRigElementKey(InRigElementName, Type), bSelected);
+		if(URigHierarchyController* Controller = InteractionRig->GetHierarchy()->GetController())
+		{
+			Controller->SelectElement(FRigElementKey(InRigElementName, Type), bSelected);
+		}
 	}
 }
 
@@ -1616,7 +1622,7 @@ void FControlRigEditMode::ResetTransforms(bool bSelectionOnly)
 		TArray<FRigElementKey> ControlsToReset = SelectedRigElements;
 		if (!bSelectionOnly)
 		{
-			ControlsToReset = ControlRig->GetHierarchy()->GetAllItems();
+			ControlsToReset = ControlRig->GetHierarchy()->GetAllKeys(true, ERigElementType::Control);
 		}
 
 		FScopedTransaction Transaction(LOCTEXT("HierarchyResetTransforms", "Reset Transforms"));
@@ -1624,17 +1630,17 @@ void FControlRigEditMode::ResetTransforms(bool bSelectionOnly)
 		{
 			if (ControlToReset.Type == ERigElementType::Control)
 			{
-				FRigControl* Control = ControlRig->FindControl(ControlToReset.Name);
-				if (Control && !Control->bIsTransientControl)
+				FRigControlElement* ControlElement = ControlRig->FindControl(ControlToReset.Name);
+				if (ControlElement && !ControlElement->Settings.bIsTransientControl)
 				{
-					FTransform Transform = ControlRig->GetControlHierarchy().GetLocalTransform(ControlToReset.Name, ERigControlValueType::Initial);
+					const FTransform InitialLocalTransform = ControlRig->GetHierarchy()->GetInitialLocalTransform(ControlToReset);
 					ControlRig->Modify();
-					ControlRig->GetControlHierarchy().SetLocalTransform(ControlToReset.Name, Transform);
-					ControlRig->ControlModified().Broadcast(ControlRig, *Control, EControlRigSetKey::DoNotCare);
+					ControlRig->GetHierarchy()->SetLocalTransform(ControlToReset, InitialLocalTransform);
+					ControlRig->ControlModified().Broadcast(ControlRig, ControlElement, EControlRigSetKey::DoNotCare);
 
 					if (UControlRigBlueprint* Blueprint = Cast<UControlRigBlueprint>(ControlRig->GetClass()->ClassGeneratedBy))
 					{
-						Blueprint->HierarchyContainer.SetLocalTransform(ControlToReset, Transform);
+						Blueprint->Hierarchy->SetLocalTransform(ControlToReset, InitialLocalTransform);
 					}
 				}
 			}
@@ -1721,13 +1727,15 @@ void FControlRigEditMode::RecreateGizmoActors(const TArray<FRigElementKey>& InSe
 				ControlProxy->RecreateAllProxies(InteractionRig);
 			}
 		}
-	}
 
-	for (const FRigElementKey& SelectedElement : InSelectedElements)
-	{
-		OnRigElementSelected(nullptr, SelectedElement, true);
+		for (const FRigElementKey& SelectedElement : InSelectedElements)
+		{
+			if(FRigControlElement* ControlElement = ControlRig->FindControl(SelectedElement.Name))
+			{
+				OnHierarchyModified(ERigHierarchyNotification::ElementSelected, ControlRig->GetHierarchy(), ControlElement);
+			}
+		}
 	}
-
 }
 
 FControlRigEditMode* FControlRigEditMode::GetEditModeFromWorldContext(UWorld* InWorldContext)
@@ -1803,103 +1811,98 @@ bool FControlRigEditMode::AreRigElementSelectedAndMovable() const
 		return true;
 }
 
-void FControlRigEditMode::OnRigElementAdded(FRigHierarchyContainer* Container, const FRigElementKey& InKey)
-{
-	RequestToRecreateGizmoActors();
-}
-
-void FControlRigEditMode::OnRigElementRemoved(FRigHierarchyContainer* Container, const FRigElementKey& InKey)
-{
-	RequestToRecreateGizmoActors();
-}
-
-void FControlRigEditMode::OnRigElementRenamed(FRigHierarchyContainer* Container, ERigElementType ElementType, const FName& InOldName, const FName& InNewName)
-{
-	RequestToRecreateGizmoActors();
-}
-
-void FControlRigEditMode::OnRigElementReparented(FRigHierarchyContainer* Container, const FRigElementKey& InKey, const FName& InOldParentName, const FName& InNewParentName)
-{
-	RequestToRecreateGizmoActors();
-}
-
-void FControlRigEditMode::OnRigElementSelected(FRigHierarchyContainer* Container, const FRigElementKey& InKey, bool bSelected)
+void FControlRigEditMode::OnHierarchyModified(ERigHierarchyNotification InNotif, URigHierarchy* InHierarchy, const FRigBaseElement* InElement)
 {
-	switch (InKey.Type)
+	switch(InNotif)
 	{
-		case ERigElementType::Bone:
-		case ERigElementType::Control:
-		case ERigElementType::Space:
-		case ERigElementType::Curve:
+		case ERigHierarchyNotification::ElementAdded:
+		case ERigHierarchyNotification::ElementRemoved:
+		case ERigHierarchyNotification::ElementRenamed:
+		case ERigHierarchyNotification::ParentChanged:
+		case ERigHierarchyNotification::HierarchyReset:
+		case ERigHierarchyNotification::ControlSettingChanged:
 		{
-			if (bSelected)
-			{
-				SelectedRigElements.AddUnique(InKey);
-			}
-			else
-			{
-				SelectedRigElements.Remove(InKey);
-			}
+			RequestToRecreateGizmoActors();
+			break;
+		}
+		case ERigHierarchyNotification::ElementSelected:
+		case ERigHierarchyNotification::ElementDeselected:
+		{
+			const FRigElementKey Key = InElement->GetKey();
 
-			// if it's control
-			if (InKey.Type == ERigElementType::Control)
+			switch (InElement->GetType())
 			{
-				FScopedTransaction ScopedTransaction(LOCTEXT("SelectControlTransaction", "Select Control"), IsInLevelEditor() && !GIsTransacting);	
-				if (IsInLevelEditor())
-				{
-					ControlProxy->Modify();
-				}
-				// users may select gizmo and control rig units, so we have to let them go through both of them if they do
-				// first go through gizmo actor
-				AControlRigGizmoActor* GizmoActor = GetGizmoFromControlName(InKey.Name);
-				if (GizmoActor)
-				{
-					GizmoActor->SetSelected(bSelected);
-
-				}
-				if (IsInLevelEditor())
+				case ERigElementType::Bone:
+            	case ERigElementType::Space:
+            	case ERigElementType::Curve:
+            	case ERigElementType::Control:
 				{
+					const bool bSelected = InNotif == ERigHierarchyNotification::ElementSelected;
 					if (bSelected)
 					{
-						if (UControlRig* ControlRig = GetControlRig(true))
+						SelectedRigElements.AddUnique(Key);
+					}
+					else
+					{
+						SelectedRigElements.Remove(Key);
+					}
+
+					// if it's control
+					if (Key.Type == ERigElementType::Control)
+					{
+						FScopedTransaction ScopedTransaction(LOCTEXT("SelectControlTransaction", "Select Control"), IsInLevelEditor() && !GIsTransacting);	
+						if (IsInLevelEditor())
+						{
+							ControlProxy->Modify();
+						}
+						// users may select gizmo and control rig units, so we have to let them go through both of them if they do
+						// first go through gizmo actor
+						AControlRigGizmoActor* GizmoActor = GetGizmoFromControlName(Key.Name);
+						if (GizmoActor)
 						{
-							FRigControl* Control = ControlRig->FindControl(InKey.Name);
-							if (Control)
+							GizmoActor->SetSelected(bSelected);
+
+						}
+						if (IsInLevelEditor())
+						{
+							if (bSelected)
+							{
+								if (UControlRig* ControlRig = GetControlRig(true))
+								{
+									if (ControlRig->GetHierarchy()->Find<FRigControlElement>(Key))
+									{
+										ControlProxy->SelectProxy(Key.Name, true);
+									}
+								}
+							}
+							else
 							{
-								ControlProxy->SelectProxy(InKey.Name, true);
+								ControlProxy->SelectProxy(Key.Name, false);
 							}
 						}
 					}
-					else
-					{
-						ControlProxy->SelectProxy(InKey.Name, false);
-					}
+					bSelectionChanged = true;
+		
+					break;
+				}
+				default:
+				{
+					ensureMsgf(false, TEXT("Unsupported Type of RigElement: %s"), *Key.ToString());
+					break;
 				}
 			}
-			bSelectionChanged = true;
-			
-			break;
 		}
 		default:
 		{
-			ensureMsgf(false, TEXT("Unsupported Type of RigElement: %d"), InKey.Type);
 			break;
 		}
 	}
 }
 
-void FControlRigEditMode::OnRigElementChanged(FRigHierarchyContainer* Container, const FRigElementKey& InKey)
-{
-	if(!InKey) // all of them changed
-	{
-		RequestToRecreateGizmoActors();
-	}
-}
-
-void FControlRigEditMode::OnControlModified(UControlRig* Subject, const FRigControl& Control, const FRigControlModifiedContext& Context)
+void FControlRigEditMode::OnControlModified(UControlRig* Subject, FRigControlElement* InControlElement, const FRigControlModifiedContext& Context)
 {
 	//this makes sure the details panel ui get's updated, don't remove
-	ControlProxy->ProxyChanged(Control.Name);
+	ControlProxy->ProxyChanged(InControlElement->GetName());
 
 	/*
 	FScopedTransaction ScopedTransaction(LOCTEXT("ModifyControlTransaction", "Modify Control"),!GIsTransacting && Context.SetKey != EControlRigSetKey::Never);
@@ -1917,12 +1920,6 @@ void FControlRigEditMode::OnControlModified(UControlRig* Subject, const FRigCont
 	*/
 }
 
-void FControlRigEditMode::OnControlUISettingChanged(FRigHierarchyContainer* Container, const FRigElementKey& InKey)
-{
-	// todo: simplify - this is pretty slow for what it does
-	RequestToRecreateGizmoActors();
-}
-
 void FControlRigEditMode::OnWidgetModeChanged(UE::Widget::EWidgetMode InWidgetMode)
 {
 	if (Settings && Settings->bCoordSystemPerWidgetMode)
@@ -1997,7 +1994,7 @@ void FControlRigEditMode::MoveGizmo(AControlRigGizmoActor* GizmoActor, const boo
 			{
 				InteractionControlRig->SetControlLocalTransform(GizmoActor->ControlName, CurrentLocalTransform);
 
-				FTransform CurrentTransform  = InteractionControlRig->GetGlobalTransform(GizmoActor->ControlName);			// assumes it's attached to actor
+				FTransform CurrentTransform  = InteractionControlRig->GetControlGlobalTransform(GizmoActor->ControlName);			// assumes it's attached to actor
 				CurrentTransform = ToWorldTransform * CurrentTransform;
 
 				GizmoActor->SetGlobalTransform(CurrentTransform);
@@ -2116,16 +2113,16 @@ bool FControlRigEditMode::ModeSupportedByGizmoActor(const AControlRigGizmoActor*
 {
 	if (UControlRig* ControlRig = GetControlRig(true, GizmoActor->ControlRigIndex))
 	{
-		const FRigControl* RigControl = ControlRig->FindControl(GizmoActor->ControlName);
-		if (RigControl)
+		const FRigControlElement* ControlElement = ControlRig->FindControl(GizmoActor->ControlName);
+		if (ControlElement)
 		{
-			if (IsSupportedControlType(RigControl->ControlType))
+			if (IsSupportedControlType(ControlElement->Settings.ControlType))
 			{
 				switch (InMode)
 				{
 					case UE::Widget::WM_Rotate:
 					{
-						switch (RigControl->ControlType)
+						switch (ControlElement->Settings.ControlType)
 						{
 							case ERigControlType::Rotator:
 							case ERigControlType::Transform:
@@ -2143,7 +2140,7 @@ bool FControlRigEditMode::ModeSupportedByGizmoActor(const AControlRigGizmoActor*
 					}
 					case UE::Widget::WM_Translate:
 					{
-						switch (RigControl->ControlType)
+						switch (ControlElement->Settings.ControlType)
 						{
 							case ERigControlType::Float:
 							case ERigControlType::Integer:
@@ -2164,7 +2161,7 @@ bool FControlRigEditMode::ModeSupportedByGizmoActor(const AControlRigGizmoActor*
 					}
 					case UE::Widget::WM_Scale:
 					{
-						switch (RigControl->ControlType)
+						switch (ControlElement->Settings.ControlType)
 						{
 							case ERigControlType::Scale:
 							case ERigControlType::Transform:
@@ -2181,7 +2178,7 @@ bool FControlRigEditMode::ModeSupportedByGizmoActor(const AControlRigGizmoActor*
 					}
 					case UE::Widget::WM_TranslateRotateZ:
 					{
-						switch (RigControl->ControlType)
+						switch (ControlElement->Settings.ControlType)
 						{
 							case ERigControlType::Transform:
 							case ERigControlType::TransformNoScale:
@@ -2209,14 +2206,14 @@ void FControlRigEditMode::TickGizmo(AControlRigGizmoActor* GizmoActor, const FTr
 	{
 		if (UControlRig* ControlRig = GetControlRig(true, GizmoActor->ControlRigIndex))
 		{
-			FTransform Transform = ControlRig->GetControlGlobalTransform(GizmoActor->ControlName);
+			const FTransform Transform = ControlRig->GetControlGlobalTransform(GizmoActor->ControlName);
 			GizmoActor->SetActorTransform(Transform * ComponentTransform);
 
-			if (FRigControl* Control = ControlRig->FindControl(GizmoActor->ControlName))
+			if (FRigControlElement* ControlElement = ControlRig->FindControl(GizmoActor->ControlName))
 			{
-				GizmoActor->SetGizmoColor(Control->GizmoColor);
-				GizmoActor->SetIsTemporarilyHiddenInEditor(!Control->bGizmoVisible || Settings->bHideManipulators);
-				GizmoActor->SetSelectable(Control->bGizmoVisible && !Settings->bHideManipulators && Control->bAnimatable);
+				GizmoActor->SetGizmoColor(ControlElement->Settings.GizmoColor);
+				GizmoActor->SetIsTemporarilyHiddenInEditor(!ControlElement->Settings.bGizmoVisible || Settings->bHideManipulators);
+				GizmoActor->SetSelectable(ControlElement->Settings.bGizmoVisible && !Settings->bHideManipulators && ControlElement->Settings.bAnimatable);
 			}
 		}
 	}
@@ -2238,6 +2235,10 @@ AControlRigGizmoActor* FControlRigEditMode::GetGizmoFromControlName(const FName&
 void FControlRigEditMode::AddControlRig(UControlRig* InControlRig)
 {
 	RuntimeControlRigs.AddUnique(InControlRig);
+
+	InControlRig->GetHierarchy()->OnModified().RemoveAll(this);
+	InControlRig->GetHierarchy()->OnModified().AddSP(this, &FControlRigEditMode::OnHierarchyModified);
+
 	OnControlRigAddedOrRemovedDelegate.Broadcast(InControlRig, true);
 }
 
@@ -2266,6 +2267,7 @@ void FControlRigEditMode::RemoveControlRig(UControlRig* InControlRig)
 	{
 		OnControlRigAddedOrRemovedDelegate.Broadcast(InControlRig, false);
 		RuntimeControlRigs[Index]->ControlModified().RemoveAll(this);
+		RuntimeControlRigs[Index]->GetHierarchy()->OnModified().RemoveAll(this);
 		RuntimeControlRigs.RemoveAt(Index);
 
 		DelegateHelper->RemoveDelegates();
@@ -2309,28 +2311,28 @@ bool FControlRigEditMode::CreateGizmoActors(UWorld* World)
 			continue;
 		}
 
-		const TArray<FRigControl>& Controls = ControlRig->AvailableControls();
+		TArray<FRigControlElement*> Controls = ControlRig->AvailableControls();
 		UControlRigGizmoLibrary* GizmoLibrary = ControlRig->GetGizmoLibrary();
 
-		for (const FRigControl& Control : Controls)
+		for (FRigControlElement* ControlElement : Controls)
 		{
-			if (!Control.bGizmoEnabled)
+			if (!ControlElement->Settings.bGizmoEnabled)
 			{
 				continue;
 			}
-			if (IsSupportedControlType(Control.ControlType))
+			if (IsSupportedControlType(ControlElement->Settings.ControlType))
 			{
 				FGizmoActorCreationParam Param;
 				Param.ManipObj = ControlRig;
 				Param.ControlRigIndex = ControlRigIndex;
-				Param.ControlName = Control.Name;
-				Param.SpawnTransform = ControlRig->GetControlGlobalTransform(Control.Name);
-				Param.GizmoTransform = Control.GizmoTransform;
-				Param.bSelectable = Control.bAnimatable;
+				Param.ControlName = ControlElement->GetName();
+				Param.SpawnTransform = ControlRig->GetControlGlobalTransform(ControlElement->GetName());
+				Param.GizmoTransform = ControlRig->GetHierarchy()->GetControlGizmoTransform(ControlElement, ERigTransformType::CurrentLocal);
+				Param.bSelectable = ControlElement->Settings.bAnimatable;
 
 				if (GizmoLibrary)
 				{
-					if (const FControlRigGizmoDefinition* Gizmo = GizmoLibrary->GetGizmoByName(Control.GizmoName, true /* use default */))
+					if (const FControlRigGizmoDefinition* Gizmo = GizmoLibrary->GetGizmoByName(ControlElement->Settings.GizmoName, true /* use default */))
 					{
 						Param.MeshTransform = Gizmo->Transform;
 						Param.StaticMesh = Gizmo->StaticMesh;
@@ -2339,7 +2341,7 @@ bool FControlRigEditMode::CreateGizmoActors(UWorld* World)
 					}
 				}
 
-				Param.Color = Control.GizmoColor;
+				Param.Color = ControlElement->Settings.GizmoColor;
 
 				AControlRigGizmoActor* GizmoActor = FControlRigGizmoHelper::CreateDefaultGizmoActor(World, Param);
 				if (GizmoActor)
diff --git a/Engine/Plugins/Experimental/ControlRig/Source/ControlRigEditor/Private/EditMode/ControlRigEditMode.h b/Engine/Plugins/Experimental/ControlRig/Source/ControlRigEditor/Private/EditMode/ControlRigEditMode.h
--- a/Engine/Plugins/Experimental/ControlRig/Source/ControlRigEditor/Private/EditMode/ControlRigEditMode.h
+++ b/Engine/Plugins/Experimental/ControlRig/Source/ControlRigEditor/Private/EditMode/ControlRigEditMode.h
@@ -135,14 +135,8 @@ public:
 	FControlRigAddedOrRemoved& OnControlRigAddedOrRemoved() { return OnControlRigAddedOrRemovedDelegate; }
 
 	// callback that gets called when rig element is selected in other view
-	void OnRigElementAdded(FRigHierarchyContainer* Container, const FRigElementKey& InKey);
-	void OnRigElementRemoved(FRigHierarchyContainer* Container, const FRigElementKey& InKey);
-	void OnRigElementRenamed(FRigHierarchyContainer* Container, ERigElementType ElementType, const FName& InOldName, const FName& InNewName);
-	void OnRigElementReparented(FRigHierarchyContainer* Container, const FRigElementKey& InKey, const FName& InOldParentName, const FName& InNewParentName);
-	void OnRigElementSelected(FRigHierarchyContainer* Container, const FRigElementKey& InKey, bool bSelected);
-	void OnRigElementChanged(FRigHierarchyContainer* Container, const FRigElementKey& InKey);
-	void OnControlUISettingChanged(FRigHierarchyContainer* Container, const FRigElementKey& InKey);
-	void OnControlModified(UControlRig* Subject, const FRigControl& Control, const FRigControlModifiedContext& Context);
+	void OnHierarchyModified(ERigHierarchyNotification InNotif, URigHierarchy* InHierarchy, const FRigBaseElement* InElement);
+	void OnControlModified(UControlRig* Subject, FRigControlElement* InControlElement, const FRigControlModifiedContext& Context);
 
 	/** return true if it can be removed from preview scene 
 	- this is to ensure preview scene doesn't remove Gizmo actors */
diff --git a/Engine/Plugins/Experimental/ControlRig/Source/ControlRigEditor/Private/EditMode/SControlHierarchy.cpp b/Engine/Plugins/Experimental/ControlRig/Source/ControlRigEditor/Private/EditMode/SControlHierarchy.cpp
--- a/Engine/Plugins/Experimental/ControlRig/Source/ControlRigEditor/Private/EditMode/SControlHierarchy.cpp
+++ b/Engine/Plugins/Experimental/ControlRig/Source/ControlRigEditor/Private/EditMode/SControlHierarchy.cpp
@@ -12,6 +12,7 @@
 #include "ControlRig.h"
 #include "ControlRigEditorStyle.h"
 #include "ScopedTransaction.h"
+#include "Rigs/RigHierarchyController.h"
 
 #define LOCTEXT_NAMESPACE "SControlHierarchy"
 
@@ -179,7 +180,6 @@ void SControlHierarchy::OnFilterTextChanged(const FText& SearchText)
 
 void SControlHierarchy::RefreshTreeView()
 {
-	
 	TMap<FRigElementKey, bool> ExpansionState;
 	for (TPair<FRigElementKey, TSharedPtr<FControlTreeElement>> Pair : ElementMap)
 	{
@@ -192,13 +192,14 @@ void SControlHierarchy::RefreshTreeView()
 
 	if (ControlRig.IsValid())
 	{
-		TArray<FRigControl> SortedControls;
+		TArray<FRigControlElement*> SortedControls;
 		ControlRig->GetControlsInOrder(SortedControls);
-		for (const FRigControl& Element : SortedControls)
+		
+		for (FRigControlElement* ControlElement : SortedControls)
 		{
-			if (!ControlRig->IsCurveControl(&Element))
+			if (!ControlRig->IsCurveControl(ControlElement))
 			{
-				AddControlElement(Element);
+				AddElement(ControlElement);
 			}
 		}
 		if (ExpansionState.Num() == 0)
@@ -225,17 +226,12 @@ void SControlHierarchy::RefreshTreeView()
 
 		TreeView->RequestTreeRefresh();
 
-		FRigControlHierarchy& ControlHierarchy = ControlRig->GetControlHierarchy();
-		TArray<FName> Selection = ControlHierarchy.CurrentSelection();
-		for (const FName& Name : Selection)
+		TArray<FRigElementKey> Selection = GetHierarchy()->GetSelectedKeys(ERigElementType::Control);
+		for (const FRigElementKey& SelectedKey : Selection)
 		{
-			for (const FRigControl& Control : ControlHierarchy.GetControls())
+			if(FRigControlElement* SelectedElement = GetHierarchy()->Find<FRigControlElement>(SelectedKey))
 			{
-				if (Name == Control.Name)
-				{
-					OnRigElementSelected(ControlRig.Get(),Control, true);
-					break;
-				}
+				OnRigElementSelected(ControlRig.Get(), SelectedElement, true);
 			}
 		}
 	}
@@ -282,13 +278,15 @@ void SControlHierarchy::OnSelectionChanged(TSharedPtr<FControlTreeElement> Selec
 	if (!bSelecting)
 	{
 		TGuardValue<bool> Guard(bSelecting, true);
-		FRigHierarchyContainer* Hierarchy = GetHierarchyContainer();
+		URigHierarchy* Hierarchy = GetHierarchy();
+		URigHierarchyController* Controller = Hierarchy->GetController(true);
+		check(Controller);
 
 		if (Hierarchy)
 		{
 			FScopedTransaction ScopedTransaction(LOCTEXT("SelectControlTransaction", "Select Control"), !GIsTransacting);
 
-			TArray<FRigElementKey> OldSelection = Hierarchy->CurrentSelection();
+			TArray<FRigElementKey> OldSelection = Hierarchy->GetSelectedKeys(ERigElementType::Control);
 			TArray<FRigElementKey> NewSelection;
 
 			TArray<TSharedPtr<FControlTreeElement>> SelectedItems = TreeView->GetSelectedItems();
@@ -303,12 +301,12 @@ void SControlHierarchy::OnSelectionChanged(TSharedPtr<FControlTreeElement> Selec
 				{
 					continue;
 				}
-				Hierarchy->Select(PreviouslySelected, false);
+				Controller->SelectElement(PreviouslySelected, false);
 			}
 
 			for (const FRigElementKey& NewlySelected : NewSelection)
 			{
-				Hierarchy->Select(NewlySelected, true);
+				Controller->SelectElement(NewlySelected, true);
 			}
 		}
 	}
@@ -334,12 +332,9 @@ TSharedPtr<FControlTreeElement> SControlHierarchy::FindElement(const FRigElement
 	return TSharedPtr<FControlTreeElement>();
 }
 
-void SControlHierarchy::OnRigElementSelected(UControlRig* Subject, const FRigControl& Control, bool bSelected)
+void SControlHierarchy::OnRigElementSelected(UControlRig* Subject, FRigControlElement* ControlElement, bool bSelected)
 {
-	
-	FRigElementKey Key;
-	Key.Name = Control.Name;
-	Key.Type = ERigElementType::Control;
+	const FRigElementKey Key = ControlElement->GetKey();
 	for (int32 RootIndex = 0; RootIndex < RootElements.Num(); ++RootIndex)
 	{
 		TSharedPtr<FControlTreeElement> Found = FindElement(Key, RootElements[RootIndex]);
@@ -361,58 +356,15 @@ void SControlHierarchy::OnRigElementSelected(UControlRig* Subject, const FRigCon
 	
 }
 
-void SControlHierarchy::AddControlElement(FRigControl InControl)
-{
-	const FRigHierarchyContainer* Container = GetHierarchyContainer();
-	const FRigControlHierarchy& ControlHierarchy = Container->ControlHierarchy;
-	const FRigSpaceHierarchy& SpaceHierarchy = Container->SpaceHierarchy;
-
-	FRigElementKey ParentKey;
-	if (InControl.SpaceIndex != INDEX_NONE)
-	{
-		AddSpaceElement(SpaceHierarchy[InControl.SpaceIndex]);
-		ParentKey = SpaceHierarchy[InControl.SpaceIndex].GetElementKey();
-	}
-	else if (InControl.ParentIndex != INDEX_NONE)
-	{
-		AddControlElement(ControlHierarchy[InControl.ParentIndex]);
-		ParentKey = ControlHierarchy[InControl.ParentIndex].GetElementKey();
-	}
-	AddElement(InControl.GetElementKey(), ParentKey);
-}
-
-
-void SControlHierarchy::AddSpaceElement(FRigSpace InSpace)
+void SControlHierarchy::AddElement(FRigBaseElement* InElement)
 {
-	const FRigHierarchyContainer* Container = GetHierarchyContainer();
-	const FRigBoneHierarchy& BoneHierarchy = Container->BoneHierarchy;
-	const FRigControlHierarchy& ControlHierarchy = Container->ControlHierarchy;
-	const FRigSpaceHierarchy& SpaceHierarchy = Container->SpaceHierarchy;
-
-	FRigElementKey ParentKey;
-	if (InSpace.ParentIndex != INDEX_NONE)
+	const URigHierarchy* Hierarchy = GetHierarchy();
+	FRigElementKey ParentKey = Hierarchy->GetFirstParent(InElement->GetKey());
+	while(ParentKey.IsValid() && ParentKey.Type != ERigElementType::Control)
 	{
-		switch (InSpace.SpaceType)
-		{
-		case ERigSpaceType::Control:
-		{
-			AddControlElement(ControlHierarchy[InSpace.ParentIndex]);
-			ParentKey = ControlHierarchy[InSpace.ParentIndex].GetElementKey();
-			break;
-		}
-		case ERigSpaceType::Space:
-		{
-			AddSpaceElement(SpaceHierarchy[InSpace.ParentIndex]);
-			ParentKey = SpaceHierarchy[InSpace.ParentIndex].GetElementKey();
-			break;
-		}
-		default:
-		{
-			break;
-		}
-		}
+		ParentKey = Hierarchy->GetFirstParent(InElement->GetKey());
 	}
-	AddElement(InSpace.GetElementKey(), ParentKey);
+	AddElement(InElement->GetKey(), ParentKey);
 }
 
 void SControlHierarchy::AddElement(FRigElementKey InKey, FRigElementKey InParentKey)
@@ -468,14 +420,12 @@ void SControlHierarchy::AddElement(FRigElementKey InKey, FRigElementKey InParent
 	}
 }
 
-
-FRigHierarchyContainer* SControlHierarchy::GetHierarchyContainer() const
+URigHierarchy* SControlHierarchy::GetHierarchy() const
 {
 	if (ControlRig.IsValid())
 	{
 		return ControlRig->GetHierarchy();
 	}
-
 	return nullptr;
 }
 #undef LOCTEXT_NAMESPACE
\ No newline at end of file
diff --git a/Engine/Plugins/Experimental/ControlRig/Source/ControlRigEditor/Private/EditMode/SControlHierarchy.h b/Engine/Plugins/Experimental/ControlRig/Source/ControlRigEditor/Private/EditMode/SControlHierarchy.h
--- a/Engine/Plugins/Experimental/ControlRig/Source/ControlRigEditor/Private/EditMode/SControlHierarchy.h
+++ b/Engine/Plugins/Experimental/ControlRig/Source/ControlRigEditor/Private/EditMode/SControlHierarchy.h
@@ -85,11 +85,10 @@ private:
 	static TSharedPtr<FControlTreeElement> FindElement(const FRigElementKey& InElementKey, TSharedPtr<FControlTreeElement> CurrentItem);
 
 	void SetExpansionRecursive(TSharedPtr<FControlTreeElement> InElements, bool bTowardsParent);
-	void OnRigElementSelected(UControlRig* Subject, const FRigControl& Control, bool bSelected);
-	FRigHierarchyContainer* GetHierarchyContainer() const;
+	void OnRigElementSelected(UControlRig* Subject, FRigControlElement* ControlElement, bool bSelected);
+	URigHierarchy* GetHierarchy() const;
 	void AddElement(FRigElementKey InKey, FRigElementKey InParentKey = FRigElementKey());
-	void AddControlElement(FRigControl InControl);
-	void AddSpaceElement(FRigSpace InSpace);
+	void AddElement(FRigBaseElement* InElement);
 public:
 
 	friend class SControlHierarchyItem;
diff --git a/Engine/Plugins/Experimental/ControlRig/Source/ControlRigEditor/Private/EditMode/SControlRigControlViews.cpp b/Engine/Plugins/Experimental/ControlRig/Source/ControlRigEditor/Private/EditMode/SControlRigControlViews.cpp
--- a/Engine/Plugins/Experimental/ControlRig/Source/ControlRigEditor/Private/EditMode/SControlRigControlViews.cpp
+++ b/Engine/Plugins/Experimental/ControlRig/Source/ControlRigEditor/Private/EditMode/SControlRigControlViews.cpp
@@ -649,7 +649,7 @@ void SControlRigPoseView::HandleControlAdded(UControlRig* ControlRig, bool bIsAd
 	UpdateStatusBlocks();
 }
 
-void SControlRigPoseView::HandleControlSelected(UControlRig* Subject, const FRigControl& Control, bool bSelected)
+void SControlRigPoseView::HandleControlSelected(UControlRig* Subject, FRigControlElement* InControl, bool bSelected)
 {
 	UpdateStatusBlocks();
 }
diff --git a/Engine/Plugins/Experimental/ControlRig/Source/ControlRigEditor/Private/EditMode/SControlRigControlViews.h b/Engine/Plugins/Experimental/ControlRig/Source/ControlRigEditor/Private/EditMode/SControlRigControlViews.h
--- a/Engine/Plugins/Experimental/ControlRig/Source/ControlRigEditor/Private/EditMode/SControlRigControlViews.h
+++ b/Engine/Plugins/Experimental/ControlRig/Source/ControlRigEditor/Private/EditMode/SControlRigControlViews.h
@@ -61,7 +61,7 @@ private:
 	FReply OnPastePose();
 	FReply OnSelectControls();
 	FReply OnCaptureThumbnail();
-	void HandleControlSelected(UControlRig* Subject, const FRigControl& Control, bool bSelected);
+	void HandleControlSelected(UControlRig* Subject, FRigControlElement* InControl, bool bSelected);
 	void HandleControlAdded(UControlRig* ControlRig, bool bIsAdded);
 
 	static bool bIsKey;
diff --git a/Engine/Plugins/Experimental/ControlRig/Source/ControlRigEditor/Private/EditMode/SControlRigEditModeTools.cpp b/Engine/Plugins/Experimental/ControlRig/Source/ControlRigEditor/Private/EditMode/SControlRigEditModeTools.cpp
--- a/Engine/Plugins/Experimental/ControlRig/Source/ControlRigEditor/Private/EditMode/SControlRigEditModeTools.cpp
+++ b/Engine/Plugins/Experimental/ControlRig/Source/ControlRigEditor/Private/EditMode/SControlRigEditModeTools.cpp
@@ -65,13 +65,13 @@ public:
 		for (UControlRigControlsProxy* Proxy : ProxiesBeingCustomized)
 		{
 			FName ValuePropertyName = TEXT("Transform");
-			if (Proxy->RigControl->ControlType == ERigControlType::Float)
+			if (Proxy->ControlElement->Settings.ControlType == ERigControlType::Float)
 			{
 				ValuePropertyName = TEXT("Float");
 			}
-			else if (Proxy->RigControl->ControlType == ERigControlType::Integer)
+			else if (Proxy->ControlElement->Settings.ControlType == ERigControlType::Integer)
 			{
-				if (Proxy->RigControl->ControlEnum == nullptr)
+				if (Proxy->ControlElement->Settings.ControlEnum == nullptr)
 				{
 					ValuePropertyName = TEXT("Integer");
 				}
@@ -80,16 +80,16 @@ public:
 					ValuePropertyName = TEXT("Enum");
 				}
 			}
-			else if (Proxy->RigControl->ControlType == ERigControlType::Bool)
+			else if (Proxy->ControlElement->Settings.ControlType == ERigControlType::Bool)
 			{
 				ValuePropertyName = TEXT("Bool");
 			}
-			else if (Proxy->RigControl->ControlType == ERigControlType::Position ||
-				Proxy->RigControl->ControlType == ERigControlType::Scale)
+			else if (Proxy->ControlElement->Settings.ControlType == ERigControlType::Position ||
+				Proxy->ControlElement->Settings.ControlType == ERigControlType::Scale)
 			{
 				ValuePropertyName = TEXT("Vector");
 			}
-			else if (Proxy->RigControl->ControlType == ERigControlType::Vector2D)
+			else if (Proxy->ControlElement->Settings.ControlType == ERigControlType::Vector2D)
 			{
 				ValuePropertyName = TEXT("Vector2D");
 			}
@@ -97,19 +97,27 @@ public:
 			TSharedPtr<IPropertyHandle> ValuePropertyHandle = DetailLayout.GetProperty(ValuePropertyName, Proxy->GetClass());
 			if (ValuePropertyHandle)
 			{
-				ValuePropertyHandle->SetPropertyDisplayName(FText::FromName(Proxy->RigControl->GetDisplayName()));
+				ValuePropertyHandle->SetPropertyDisplayName(FText::FromName(Proxy->ControlElement->Settings.DisplayName));
 			}
 
-			for (const FRigControl& ChildControl : Proxy->ControlRig->GetControlHierarchy())
-			{
-				if (ChildControl.ParentName == Proxy->RigControl->Name)
+			URigHierarchy* Hierarchy = Proxy->ControlRig->GetHierarchy();
+			Hierarchy->ForEach<FRigControlElement>([Hierarchy, Proxy, &Category](FRigControlElement* ControlElement) -> bool
+            {
+				FName ParentControlName = NAME_None;
+				FRigControlElement* ParentControlElement = Cast<FRigControlElement>(Hierarchy->GetFirstParent(ControlElement));
+				if(ParentControlElement)
+				{
+					ParentControlName = ParentControlElement->GetName();
+				}
+				
+				if (ParentControlName == Proxy->ControlElement->GetName())
 				{
 					if (FControlRigEditMode* EditMode = static_cast<FControlRigEditMode*>(GLevelEditorModeTools().GetActiveMode(FControlRigEditMode::ModeName)))
 					{
-						if (UObject* NestedProxy = EditMode->ControlProxy->FindProxy(ChildControl.Name))
+						if (UObject* NestedProxy = EditMode->ControlProxy->FindProxy(ControlElement->GetName()))
 						{
 							FName PropertyName(NAME_None);
-							switch (ChildControl.ControlType)
+							switch (ControlElement->Settings.ControlType)
 							{
 								case ERigControlType::Bool:
 								{
@@ -123,7 +131,7 @@ public:
 								}
 								case ERigControlType::Integer:
 								{
-									if (ChildControl.ControlEnum == nullptr)
+									if (ControlElement->Settings.ControlEnum == nullptr)
 									{
 										PropertyName = TEXT("Integer");
 									}
@@ -141,7 +149,7 @@ public:
 
 							if (PropertyName.IsNone())
 							{
-								continue;
+								return true;
 							}
 
 							TArray<UObject*> NestedProxies;
@@ -155,13 +163,14 @@ public:
 								PropertyName,
 								EPropertyLocation::Advanced,
 								Params);
-							NestedRow->DisplayName(FText::FromName(ChildControl.GetDisplayName()));
+							NestedRow->DisplayName(FText::FromName(ControlElement->Settings.DisplayName));
 
 							Category.SetShowAdvanced(true);
 						}
 					}
 				}
-			}
+				return true;
+			});
 		}
 	}
 };
diff --git a/Engine/Plugins/Experimental/ControlRig/Source/ControlRigEditor/Private/EditMode/SCurveControlContainer.cpp b/Engine/Plugins/Experimental/ControlRig/Source/ControlRigEditor/Private/EditMode/SCurveControlContainer.cpp
--- a/Engine/Plugins/Experimental/ControlRig/Source/ControlRigEditor/Private/EditMode/SCurveControlContainer.cpp
+++ b/Engine/Plugins/Experimental/ControlRig/Source/ControlRigEditor/Private/EditMode/SCurveControlContainer.cpp
@@ -203,7 +203,7 @@ SCurveControlContainer::~SCurveControlContainer()
 {
 	if (ControlRig.IsValid())
 	{
-		ControlRig->GetHierarchy()->OnElementSelected.RemoveAll(this);
+		ControlRig->ControlSelected().RemoveAll(this);
 	}
 }
 
@@ -235,31 +235,33 @@ TSharedRef<ITableRow> SCurveControlContainer::GenerateCurveControlRow(FDisplayed
 
 void SCurveControlContainer::CreateCurveControlList( const FString& SearchText )
 {
-	const FRigCurveContainer* Container = GetCurveContainer();
-	if (Container)
+	URigHierarchy* Hierarchy = GetHierarchy();
+	if (Hierarchy)
 	{
 		CurveControlList.Reset();
 
-		FString CTRLName(TEXT("CTRL_"));
+		const FString CTRLName(TEXT("CTRL_"));
 
 		// Iterate through all curves..
-		for (const FRigCurve& Curve : (*Container))
-		{
-			FString Name = Curve.Name.ToString();
+		Hierarchy->ForEach<FRigCurveElement>([CTRLName, this](FRigCurveElement* CurveElement) -> bool
+        {
+			const FString Name = CurveElement->GetName().ToString();
 			if (Name.Contains(CTRLName))
 			{
-				FString CurveString = Curve.Name.ToString();
+				const FString CurveString = CurveElement->GetName().ToString();
 
 				// See if we pass the search filter
 				if (!FilterText.IsEmpty() && !CurveString.Contains(*FilterText.ToString()))
 				{
-					continue;
+					return true;
 				}
 
-				TSharedRef<FDisplayedCurveControlInfo> NewItem = FDisplayedCurveControlInfo::Make(Curve.Name);
+				const TSharedRef<FDisplayedCurveControlInfo> NewItem = FDisplayedCurveControlInfo::Make(CurveElement->GetName());
 				CurveControlList.Add(NewItem);
 			}
-		}
+
+			return true;
+		});
 
 		// Sort final list
 		struct FSortNamesAlphabetically
@@ -274,19 +276,16 @@ void SCurveControlContainer::CreateCurveControlList( const FString& SearchText )
 	}
 	CurveControlListView->RequestListRefresh();
 
-	if (ControlRig.IsValid())
+	if (Hierarchy)
 	{
-		FRigControlHierarchy& ControlHierarchy = ControlRig->GetControlHierarchy();
-		TArray<FName> Selection = ControlHierarchy.CurrentSelection();
-		for (const FName& Name : Selection)
+		TArray<FRigElementKey> Selection = Hierarchy->GetSelectedKeys(ERigElementType::Control);
+
+		for (const FRigElementKey& SelectedKey : Selection)
 		{
-			for (const FRigControl& Control : ControlHierarchy.GetControls())
+			if(FRigControlElement* ControlElement = Hierarchy->Find<FRigControlElement>(SelectedKey))
 			{
-				if (Name == Control.Name)
-				{
-					OnRigElementSelected(ControlRig.Get(), Control, true);
-					break;
-				}
+				OnRigElementSelected(ControlRig.Get(), ControlElement, true);
+				break;
 			}
 		}
 	}
@@ -301,8 +300,8 @@ void SCurveControlContainer::SetCurveValue(const FName& CurveName, float CurveVa
 {
 	if (ControlRig.IsValid())
 	{
-		FRigControl* RigControl = ControlRig->FindControl(CurveName);
-		if (RigControl && RigControl->ControlType == ERigControlType::Float)
+		FRigControlElement* ControlElement = ControlRig->FindControl(CurveName);
+		if (ControlElement && ControlElement->Settings.ControlType == ERigControlType::Float)
 		{
 			ControlRig->SetControlValue<float>(CurveName, CurveValue,true, EControlRigSetKey::Always);
 		}
@@ -312,12 +311,11 @@ void SCurveControlContainer::SetCurveValue(const FName& CurveName, float CurveVa
 
 float SCurveControlContainer::GetCurveValue(const FName& CurveName)
 {
-	FRigCurveContainer* Container = GetCurveContainer();
-	if (Container)
+	URigHierarchy* Hierarchy = GetHierarchy();
+	if (Hierarchy)
 	{
-		return Container->GetValue(CurveName);
+		return Hierarchy->GetCurveValue(FRigElementKey(CurveName, ERigElementType::Curve));
 	}
-
 	return 0.f;
 }
 
@@ -327,37 +325,36 @@ void SCurveControlContainer::OnSelectionChanged(FDisplayedCurveControlInfoPtr Se
 	{
 		FScopedTransaction ScopedTransaction(LOCTEXT("SelectControlTransaction", "Select Control"), !GIsTransacting);
 
-		FRigControlHierarchy& ControlHierarchy = ControlRig->GetControlHierarchy();
-		TArray<FName> OldSelection = ControlHierarchy.CurrentSelection();
-		TArray<FName> NewSelection;
+		URigHierarchy* Hierarchy = GetHierarchy();
+		TArray<FRigElementKey> OldSelection = Hierarchy->GetSelectedKeys();
+		TArray<FRigElementKey> NewSelection;
 
 		TArray<FDisplayedCurveControlInfoPtr> SelectedItems = CurveControlListView->GetSelectedItems();
 		for (const FDisplayedCurveControlInfoPtr& SelectedItem : SelectedItems)
 		{
-			NewSelection.Add(SelectedItem->CurveName);
+			NewSelection.Add(FRigElementKey(SelectedItem->CurveName, ERigElementType::Curve));
 		}
 
-		for (const FName& PreviouslySelected : OldSelection)
+		for (const FRigElementKey& PreviouslySelected : OldSelection)
 		{
 			if (NewSelection.Contains(PreviouslySelected))
 			{
 				continue;
 			}
-			ControlRig->SelectControl(PreviouslySelected, false);
+			ControlRig->SelectControl(PreviouslySelected.Name, false);
 		}
 
-		for (const FName& NewlySelected : NewSelection)
+		for (const FRigElementKey& NewlySelected : NewSelection)
 		{
-			ControlRig->SelectControl(NewlySelected, true);
+			ControlRig->SelectControl(NewlySelected.Name, true);
 		}
 	}
 }
-void SCurveControlContainer::OnRigElementSelected(UControlRig* Subject, const FRigControl& Control, bool bSelected)
+void SCurveControlContainer::OnRigElementSelected(UControlRig* Subject, FRigControlElement* ControlElement, bool bSelected)
 {
-
 	for(const FDisplayedCurveControlInfoPtr& Item : CurveControlList)
 	{
-		if (Item->CurveName == Control.Name)
+		if (Item->CurveName == ControlElement->GetName())
 		{
 			CurveControlListView->SetItemSelection(Item, bSelected);
 			break;
@@ -365,11 +362,11 @@ void SCurveControlContainer::OnRigElementSelected(UControlRig* Subject, const FR
 	}
 }
 
-FRigCurveContainer* SCurveControlContainer::GetCurveContainer() const
+URigHierarchy* SCurveControlContainer::GetHierarchy() const
 {
 	if (ControlRig.IsValid())
 	{
-		return &ControlRig->GetHierarchy()->CurveContainer;
+		return ControlRig->GetHierarchy();
 	}
 
 	return nullptr;
diff --git a/Engine/Plugins/Experimental/ControlRig/Source/ControlRigEditor/Private/EditMode/SCurveControlContainer.h b/Engine/Plugins/Experimental/ControlRig/Source/ControlRigEditor/Private/EditMode/SCurveControlContainer.h
--- a/Engine/Plugins/Experimental/ControlRig/Source/ControlRigEditor/Private/EditMode/SCurveControlContainer.h
+++ b/Engine/Plugins/Experimental/ControlRig/Source/ControlRigEditor/Private/EditMode/SCurveControlContainer.h
@@ -190,7 +190,7 @@ private:
 
 	void OnSelectionChanged(FDisplayedCurveControlInfoPtr Selection, ESelectInfo::Type SelectInfo);
 
-	void OnRigElementSelected(UControlRig* Subject, const FRigControl& Control, bool bSelected);
+	void OnRigElementSelected(UControlRig* Subject, FRigControlElement* ControlElement, bool bSelected);
 
 	/** Box to filter to a specific morph target name */
 	TSharedPtr<SSearchBox>	NameFilterBox;
@@ -210,7 +210,7 @@ private:
 	/** Commands that are bound to delegates*/
 	TSharedPtr<FUICommandList> UICommandList;
 
-	FRigCurveContainer* GetCurveContainer() const;
+	URigHierarchy* GetHierarchy() const;
 
 	friend class SCurveControlListRow;
 	friend class SCurveControlTypeList;
