Control Rig: New rig hierarchy data structure refactoring

diff --git a/Engine/Plugins/Experimental/ControlRig/Source/ControlRigEditor/Private/Editor/ControlRigEditor.cpp b/Engine/Plugins/Experimental/ControlRig/Source/ControlRigEditor/Private/Editor/ControlRigEditor.cpp
--- a/Engine/Plugins/Experimental/ControlRig/Source/ControlRigEditor/Private/Editor/ControlRigEditor.cpp
+++ b/Engine/Plugins/Experimental/ControlRig/Source/ControlRigEditor/Private/Editor/ControlRigEditor.cpp
@@ -115,20 +115,16 @@ FControlRigEditor::~FControlRigEditor()
 	{
 		UControlRigBlueprint::sCurrentlyOpenedRigBlueprints.Remove(RigBlueprint);
 
-		RigBlueprint->HierarchyContainer.OnElementChanged.RemoveAll(this);
-		RigBlueprint->HierarchyContainer.OnElementAdded.RemoveAll(this);
-		RigBlueprint->HierarchyContainer.OnElementRemoved.RemoveAll(this);
-		RigBlueprint->HierarchyContainer.OnElementRenamed.RemoveAll(this);
-		RigBlueprint->HierarchyContainer.OnElementReparented.RemoveAll(this);
-		RigBlueprint->HierarchyContainer.OnElementSelected.RemoveAll(this);
+		RigBlueprint->Hierarchy->OnModified().RemoveAll(this);
 		if (FControlRigEditMode* EditMode = GetEditMode())
 		{
-			RigBlueprint->HierarchyContainer.OnElementSelected.RemoveAll(EditMode);
+			RigBlueprint->Hierarchy->OnModified().RemoveAll(EditMode);
 		}
 		RigBlueprint->OnRefreshEditor().RemoveAll(this);
 		RigBlueprint->OnVariableDropped().RemoveAll(this);
 		RigBlueprint->OnNodeDoubleClicked().RemoveAll(this);
 		RigBlueprint->OnGraphImported().RemoveAll(this);
+		RigBlueprint->OnPostEditChangeChainProperty().RemoveAll(this);
 	}
 
 	if (NodeDetailBuffer.Num() > 0 && NodeDetailStruct != nullptr)
@@ -343,28 +339,18 @@ void FControlRigEditor::InitControlRigEditor(const EToolkitMode::Type Mode, cons
 			}
 		}
 
-		InControlRigBlueprint->HierarchyContainer.OnElementAdded.AddSP(this, &FControlRigEditor::OnRigElementAdded);
-		InControlRigBlueprint->HierarchyContainer.OnElementRemoved.AddSP(this, &FControlRigEditor::OnRigElementRemoved, false);
-		InControlRigBlueprint->HierarchyContainer.OnElementRenamed.AddSP(this, &FControlRigEditor::OnRigElementRenamed);
-		InControlRigBlueprint->HierarchyContainer.OnElementReparented.AddSP(this, &FControlRigEditor::OnRigElementReparented);
-		InControlRigBlueprint->HierarchyContainer.OnElementSelected.AddSP(this, &FControlRigEditor::OnRigElementSelected);
-		InControlRigBlueprint->HierarchyContainer.ControlHierarchy.OnControlUISettingsChanged.AddSP(this, &FControlRigEditor::OnControlUISettingChanged);
+		InControlRigBlueprint->Hierarchy->OnModified().AddSP(this, &FControlRigEditor::OnHierarchyModified);
 		InControlRigBlueprint->OnRefreshEditor().AddSP(this, &FControlRigEditor::HandleRefreshEditorFromBlueprint);
 		InControlRigBlueprint->OnVariableDropped().AddSP(this, &FControlRigEditor::HandleVariableDroppedFromBlueprint);
 
 		if (FControlRigEditMode* EditMode = GetEditMode())
 		{
-			InControlRigBlueprint->HierarchyContainer.OnElementAdded.AddSP(EditMode, &FControlRigEditMode::OnRigElementAdded);
-			InControlRigBlueprint->HierarchyContainer.OnElementRemoved.AddSP(EditMode, &FControlRigEditMode::OnRigElementRemoved);
-			InControlRigBlueprint->HierarchyContainer.OnElementRenamed.AddSP(EditMode, &FControlRigEditMode::OnRigElementRenamed);
-			InControlRigBlueprint->HierarchyContainer.OnElementReparented.AddSP(EditMode, &FControlRigEditMode::OnRigElementReparented);
-			InControlRigBlueprint->HierarchyContainer.OnElementSelected.AddSP(EditMode, &FControlRigEditMode::OnRigElementSelected);
-			InControlRigBlueprint->HierarchyContainer.OnElementChanged.AddSP(EditMode, &FControlRigEditMode::OnRigElementChanged);
-			InControlRigBlueprint->HierarchyContainer.ControlHierarchy.OnControlUISettingsChanged.AddSP(EditMode, &FControlRigEditMode::OnControlUISettingChanged);
+			InControlRigBlueprint->Hierarchy->OnModified().AddSP(EditMode, &FControlRigEditMode::OnHierarchyModified);
 		}
 
 		InControlRigBlueprint->OnNodeDoubleClicked().AddSP(this, &FControlRigEditor::OnNodeDoubleClicked);
 		InControlRigBlueprint->OnGraphImported().AddSP(this, &FControlRigEditor::OnGraphImported);
+		InControlRigBlueprint->OnPostEditChangeChainProperty().AddSP(this, &FControlRigEditor::OnBlueprintPropertyChainEvent);
 	}
 
 	UpdateStaleWatchedPins();
@@ -659,7 +645,7 @@ void FControlRigEditor::SetEventQueue(EControlRigEditorEventQueue InEventQueue)
 		if (InEventQueue == EControlRigEditorEventQueue::Setup ||
 			InEventQueue == EControlRigEditorEventQueue::Update)
 		{
-			ControlRig->Hierarchy.ResetTransforms();
+			ControlRig->GetHierarchy()->ResetPoseToInitial(ERigElementType::All);
 		}
 	}
 }
@@ -750,15 +736,11 @@ void FControlRigEditor::ToggleSetupMode()
 		// need to copy here since the removal changes the iterator
 		if (ControlRig)
 		{
-			TArray<FRigControl> TransientControls = ControlRig->TransientControls;
-			for (FRigControl TransientControl : TransientControls)
-			{
-				RigBlueprint->RemoveTransientControl(TransientControl.GetElementKey());
-			}
+			RigBlueprint->ClearTransientControls();
 		}
 
-		PreviousSelection = RigBlueprint->HierarchyContainer.CurrentSelection();
-		RigBlueprint->HierarchyContainer.ClearSelection();
+		PreviousSelection = RigBlueprint->Hierarchy->GetSelectedKeys();
+		RigBlueprint->HierarchyController->ClearSelection();
 	}
 
 	if (ControlRig)
@@ -775,7 +757,7 @@ void FControlRigEditor::ToggleSetupMode()
 	{
 		if (UControlRigBlueprint* RigBlueprint = Cast<UControlRigBlueprint>(GetBlueprintObj()))
 		{
-			EditMode->RecreateGizmoActors(RigBlueprint->HierarchyContainer.CurrentSelection());
+			EditMode->RecreateGizmoActors(RigBlueprint->Hierarchy->GetSelectedKeys());
 		}
 
 		EditMode->Settings->bDisplaySpaces = bSetupModeEnabled;
@@ -787,7 +769,7 @@ void FControlRigEditor::ToggleSetupMode()
 		{
 			for (FRigElementKey SelectedKey : PreviousSelection)
 			{
-				RigBlueprint->HierarchyContainer.Select(SelectedKey, true);
+				RigBlueprint->HierarchyController->SelectElement(SelectedKey, true);
 			}
 		}
 	}
@@ -941,6 +923,7 @@ void FControlRigEditor::HandleSetObjectBeingDebugged(UObject* InObject)
 		bool bIsExternalControlRig = DebuggedControlRig != ControlRig;
 		bool bShouldExecute = (!bIsExternalControlRig) && bExecutionControlRig;
 		DebuggedControlRig->ControlRigLog = &ControlRigLog;
+		DebuggedControlRig->DynamicHierarchy->HierarchyForSelectionPtr = GetControlRigBlueprint()->Hierarchy;
 
 		UControlRigSkeletalMeshComponent* EditorSkelComp = Cast<UControlRigSkeletalMeshComponent>(GetPersonaToolkit()->GetPreviewScene()->GetPreviewMeshComponent());
 		if (EditorSkelComp)
@@ -1166,17 +1149,18 @@ void FControlRigEditor::SetDetailStruct(const FRigElementKey& InElement)
 	ClearDetailObject();
 	
 	UControlRigBlueprint* RigBlueprint = Cast<UControlRigBlueprint>(GetBlueprintObj());
-	FRigHierarchyContainer* Container = &RigBlueprint->HierarchyContainer;
+	URigHierarchy* Hierarchy = RigBlueprint->Hierarchy;
 	
 	if (!bSetupModeEnabled)
 	{
 		if (UControlRig* DebuggedControlRig = Cast<UControlRig>(RigBlueprint->GetObjectBeingDebugged()))
 		{
-			Container = &DebuggedControlRig->Hierarchy;
+			Hierarchy = DebuggedControlRig->GetHierarchy();
 		}
 	}
 
-	if (Container->GetIndex(InElement) == INDEX_NONE)
+	FRigBaseElement* Element = Hierarchy->Find(InElement);
+	if (Element == nullptr)
 	{
 		return;
 	}
@@ -1185,26 +1169,22 @@ void FControlRigEditor::SetDetailStruct(const FRigElementKey& InElement)
 	{
 		case ERigElementType::Bone:
 		{
-			FRigBoneHierarchy& BoneHierarchy = Container->BoneHierarchy;
-			StructToDisplay = MakeShareable(new FStructOnScope(FRigBone::StaticStruct(), (uint8*)&(BoneHierarchy[InElement.Name])));
+			StructToDisplay = MakeShareable(new FStructOnScope(FRigBoneElement::StaticStruct(), (uint8*)Element));
 			break;
 		}
 		case ERigElementType::Control:
 		{
-			FRigControlHierarchy& ControlHierarchy = Container->ControlHierarchy;
-			StructToDisplay = MakeShareable(new FStructOnScope(FRigControl::StaticStruct(), (uint8*)&(ControlHierarchy[InElement.Name])));
+			StructToDisplay = MakeShareable(new FStructOnScope(FRigControlElement::StaticStruct(), (uint8*)Element));
 			break;
 		}
 		case ERigElementType::Space:
 		{
-			FRigSpaceHierarchy& SpaceHierarchy = Container->SpaceHierarchy;
-			StructToDisplay = MakeShareable(new FStructOnScope(FRigSpace::StaticStruct(), (uint8*)&(SpaceHierarchy[InElement.Name])));
+			StructToDisplay = MakeShareable(new FStructOnScope(FRigSpaceElement::StaticStruct(), (uint8*)Element));
 			break;
 		}
 		case ERigElementType::Curve:
 		{
-			FRigCurveContainer& CurveContainer = Container->CurveContainer;
-			StructToDisplay = MakeShareable(new FStructOnScope(FRigBone::StaticStruct(), (uint8*)&(CurveContainer[InElement.Name])));
+			StructToDisplay = MakeShareable(new FStructOnScope(FRigCurveElement::StaticStruct(), (uint8*)Element));
 			break;
 		}
 		default:
@@ -1371,10 +1351,9 @@ void FControlRigEditor::Compile()
 
 		if (UControlRigSettings::Get()->bResetControlTransformsOnCompile)
 		{
-			for (const FRigControl& Control : RigBlueprint->HierarchyContainer.ControlHierarchy)
-			{
-				FRigElementKey Key = Control.GetElementKey();
-				FTransform Transform = RigBlueprint->HierarchyContainer.ControlHierarchy.GetLocalTransform(Control.Index, ERigControlValueType::Initial);
+			RigBlueprint->Hierarchy->ForEach<FRigControlElement>([RigBlueprint](FRigControlElement* ControlElement) -> bool
+            {
+				const FTransform Transform = RigBlueprint->Hierarchy->GetInitialLocalTransform(ControlElement->GetIndex());
 
 				/*/
 				if (ControlRig)
@@ -1385,15 +1364,16 @@ void FControlRigEditor::Compile()
 				}
 				*/
 
-				RigBlueprint->HierarchyContainer.SetLocalTransform(Key, Transform);
-			}
+				RigBlueprint->Hierarchy->SetLocalTransform(ControlElement->GetIndex(), Transform);
+				return true;
+			});
 		}
 
 		RigBlueprint->PropagatePoseFromBPToInstances();
 
 		if (FControlRigEditMode* EditMode = GetEditMode())
 		{
-			EditMode->RecreateGizmoActors(RigBlueprint->HierarchyContainer.CurrentSelection());
+			EditMode->RecreateGizmoActors(RigBlueprint->Hierarchy->GetSelectedKeys());
 		}
 	}
 
@@ -1413,7 +1393,7 @@ void FControlRigEditor::SaveAsset_Execute()
 	// Save the new state of the hierarchy in the default object, so that it has the correct values on load
 	UControlRigBlueprint* RigBlueprint = Cast<UControlRigBlueprint>(GetBlueprintObj());
 	UControlRig* CDO = ControlRig->GetClass()->GetDefaultObject<UControlRig>();
-	CDO->Hierarchy = RigBlueprint->HierarchyContainer;
+	CDO->DynamicHierarchy->CopyHierarchy(RigBlueprint->Hierarchy);
 }
 
 void FControlRigEditor::SaveAssetAs_Execute()
@@ -1424,7 +1404,7 @@ void FControlRigEditor::SaveAssetAs_Execute()
 	// Save the new state of the hierarchy in the default object, so that it has the correct values on load
 	UControlRigBlueprint* RigBlueprint = Cast<UControlRigBlueprint>(GetBlueprintObj());
 	UControlRig* CDO = ControlRig->GetClass()->GetDefaultObject<UControlRig>();
-	CDO->Hierarchy = RigBlueprint->HierarchyContainer;
+	CDO->DynamicHierarchy->CopyHierarchy(RigBlueprint->Hierarchy);
 }
 
 FName FControlRigEditor::GetToolkitFName() const
@@ -1624,7 +1604,7 @@ void FControlRigEditor::PostUndo(bool bSuccess)
 
 		if (FControlRigEditMode* EditMode = GetEditMode())
 		{
-			EditMode->RecreateGizmoActors(RigBlueprint->HierarchyContainer.CurrentSelection());
+			EditMode->RecreateGizmoActors(RigBlueprint->Hierarchy->GetSelectedKeys());
 		}
 	}
 }
@@ -1645,7 +1625,7 @@ void FControlRigEditor::PostRedo(bool bSuccess)
 
 		if (FControlRigEditMode* EditMode = GetEditMode())
 		{
-			EditMode->RecreateGizmoActors(RigBlueprint->HierarchyContainer.CurrentSelection());
+			EditMode->RecreateGizmoActors(RigBlueprint->Hierarchy->GetSelectedKeys());
 		}
 	}
 }
@@ -1656,7 +1636,7 @@ void FControlRigEditor::EnsureValidRigElementInDetailPanel()
 	{
 		if (UControlRigBlueprint * RigBlueprint = Cast<UControlRigBlueprint>(GetBlueprintObj()))
 		{
-			if (RigBlueprint->HierarchyContainer.GetIndex(RigElementInDetailPanel) == INDEX_NONE)
+			if (RigBlueprint->Hierarchy->GetIndex(RigElementInDetailPanel) == INDEX_NONE)
 			{
 				ClearDetailObject();
 			}
@@ -1966,12 +1946,55 @@ void FControlRigEditor::Tick(float DeltaTime)
 			bExecutionControlRig)
 		{
 			// reset transforms here to prevent additive transforms from accumulating to INF
-			ControlRig->GetBoneHierarchy().ResetTransforms();
+			ControlRig->GetHierarchy()->ResetPoseToInitial(ERigElementType::Bone);
 			
 			ControlRig->SetDeltaTime(DeltaTime);
 			ControlRig->Evaluate_AnyThread();
 			bDrawHierarchyBones = true;
 		}
+
+		if(RigElementInDetailPanel.IsValid())
+		{
+			URigHierarchy* Hierarchy = Blueprint->Hierarchy; 
+
+			if(!bSetupModeEnabled)
+			{
+				UControlRig* DebuggedControlRig = Cast<UControlRig>(GetBlueprintObj()->GetObjectBeingDebugged());
+				if (DebuggedControlRig == nullptr)
+				{
+					DebuggedControlRig = ControlRig;
+				}
+				if(DebuggedControlRig)
+				{
+					Hierarchy = DebuggedControlRig->GetHierarchy();
+				}
+			}
+
+			check(Hierarchy);
+
+			FRigBaseElement* ElementInDetailPanel = Hierarchy->Find(RigElementInDetailPanel);
+			if(FRigTransformElement* TransformElement = Cast<FRigTransformElement>(ElementInDetailPanel))
+			{
+				// compute all transforms
+				Hierarchy->GetTransform(TransformElement, ERigTransformType::CurrentGlobal);
+				Hierarchy->GetTransform(TransformElement, ERigTransformType::CurrentLocal);
+				Hierarchy->GetTransform(TransformElement, ERigTransformType::InitialGlobal);
+				Hierarchy->GetTransform(TransformElement, ERigTransformType::InitialLocal);
+			}
+
+			if(FRigControlElement* ControlElement = Cast<FRigControlElement>(ElementInDetailPanel))
+			{
+				// compute all transforms
+				Hierarchy->GetControlOffsetTransform(ControlElement, ERigTransformType::CurrentGlobal);
+				Hierarchy->GetControlOffsetTransform(ControlElement, ERigTransformType::CurrentLocal);
+				Hierarchy->GetControlOffsetTransform(ControlElement, ERigTransformType::InitialGlobal);
+				Hierarchy->GetControlOffsetTransform(ControlElement, ERigTransformType::InitialLocal);
+				Hierarchy->GetControlGizmoTransform(ControlElement, ERigTransformType::CurrentGlobal);
+				Hierarchy->GetControlGizmoTransform(ControlElement, ERigTransformType::CurrentLocal);
+				Hierarchy->GetControlGizmoTransform(ControlElement, ERigTransformType::InitialGlobal);
+				Hierarchy->GetControlGizmoTransform(ControlElement, ERigTransformType::InitialLocal);
+			}
+		}
 	}
 
 	if (bDrawHierarchyBones)
@@ -2444,7 +2467,7 @@ bool FControlRigEditor::IsPinControlNameListEnabled() const
 {
 	if (ControlRig)
 	{
-		if (ControlRig->TransientControls.Num() > 0)
+		if (ControlRig->GetHierarchy()->GetTransientControls().Num() > 0)
 		{
 			return true;
 		}
@@ -2461,9 +2484,20 @@ FText FControlRigEditor::GetPinControlNameListText() const
 {
 	if (ControlRig)
 	{
-		if (ControlRig->TransientControls.Num() > 0)
+		FText Result;
+		ControlRig->GetHierarchy()->ForEach<FRigControlElement>([&Result](FRigControlElement* ControlElement) -> bool
+        {
+			if (ControlElement->Settings.bIsTransientControl)
+			{
+				Result = FText::FromName(ControlElement->GetName());
+				return false;
+			}
+			return true;
+		});
+		
+		if(!Result.IsEmpty())
 		{
-			return FText::FromName(ControlRig->TransientControls[0].ParentName);
+			return Result;
 		}
 	}
 	return FText::FromName(NAME_None);
@@ -2486,36 +2520,40 @@ void FControlRigEditor::SetPinControlNameListText(const FText& NewTypeInValue, E
 {
 	if (ControlRig)
 	{
-		if (ControlRig->TransientControls.Num() > 0)
-		{
-			FRigControl& Control = ControlRig->TransientControls[0];
-
-			Control.ParentIndex = ControlRig->Hierarchy.BoneHierarchy.GetIndex(*NewTypeInValue.ToString());
-			if (Control.ParentIndex == INDEX_NONE)
-			{
-				Control.ParentName = NAME_None;
-			}
-			else
+		ControlRig->GetHierarchy()->ForEach<FRigControlElement>([this, NewTypeInValue](FRigControlElement* ControlElement) -> bool
+        {
+            if (ControlElement->Settings.bIsTransientControl)
 			{
-				Control.ParentName = ControlRig->Hierarchy.BoneHierarchy[Control.ParentIndex].Name;
-			}
+				FName NewParentName = *NewTypeInValue.ToString();
+				const int32 NewParentIndex = ControlRig->GetHierarchy()->GetIndex(FRigElementKey(NewParentName, ERigElementType::Bone));
+				if (NewParentIndex == INDEX_NONE)
+				{
+					NewParentName = NAME_None;
+					GetControlRigBlueprint()->GetHierarchyController()->RemoveAllParents(ControlElement->GetKey(), true, false);
+				}
+				else
+				{
+					GetControlRigBlueprint()->GetHierarchyController()->SetParent(ControlElement->GetKey(), FRigElementKey(NewParentName, ERigElementType::Bone), true, false);
+				}
 
-			// find out if the controlled pin is part of a visual debug node
-			if (UControlRigBlueprint* ControlRigBlueprint = CastChecked<UControlRigBlueprint>(GetBlueprintObj()))
-			{
-				if (URigVMPin* ControlledPin = GetFocusedModel()->FindPin(Control.Name.ToString()))
+				// find out if the controlled pin is part of a visual debug node
+				if (UControlRigBlueprint* ControlRigBlueprint = CastChecked<UControlRigBlueprint>(GetBlueprintObj()))
 				{
-					URigVMNode* ControlledNode = ControlledPin->GetPinForLink()->GetNode();
-					if (URigVMPin* BoneSpacePin = ControlledNode->FindPin(TEXT("BoneSpace")))
+					if (URigVMPin* ControlledPin = GetFocusedModel()->FindPin(ControlElement->GetName().ToString()))
 					{
-						if (BoneSpacePin->GetCPPType() == TEXT("FName") && BoneSpacePin->GetCustomWidgetName() == TEXT("BoneName"))
+						URigVMNode* ControlledNode = ControlledPin->GetPinForLink()->GetNode();
+						if (URigVMPin* BoneSpacePin = ControlledNode->FindPin(TEXT("BoneSpace")))
 						{
-							GetFocusedController()->SetPinDefaultValue(BoneSpacePin->GetPinPath(), Control.ParentName.ToString(), false, false, false);
+							if (BoneSpacePin->GetCPPType() == TEXT("FName") && BoneSpacePin->GetCustomWidgetName() == TEXT("BoneName"))
+							{
+								GetFocusedController()->SetPinDefaultValue(BoneSpacePin->GetPinPath(), NewParentName.ToString(), false, false, false);
+							}
 						}
 					}
 				}
 			}
-		}
+			return true;
+		});
 	}
 }
 
@@ -2676,7 +2714,7 @@ void FControlRigEditor::CacheNameLists()
 			{
 				continue;
 			}
-			RigGraph->CacheNameLists(&ControlRigBP->HierarchyContainer, &ControlRigBP->DrawContainer);
+			RigGraph->CacheNameLists(ControlRigBP->Hierarchy, &ControlRigBP->DrawContainer);
 		}
 	}
 }
@@ -2831,9 +2869,9 @@ FTransform FControlRigEditor::GetRigElementTransform(const FRigElementKey& InEle
 
 	if (bLocal)
 	{
-		return GetControlRigBlueprint()->HierarchyContainer.GetLocalTransform(InElement);
+		return GetControlRigBlueprint()->Hierarchy->GetLocalTransform(InElement);
 	}
-	return GetControlRigBlueprint()->HierarchyContainer.GetGlobalTransform(InElement);
+	return GetControlRigBlueprint()->Hierarchy->GetGlobalTransform(InElement);
 }
 
 void FControlRigEditor::SetRigElementTransform(const FRigElementKey& InElement, const FTransform& InTransform, bool bLocal)
@@ -2852,7 +2890,7 @@ void FControlRigEditor::SetRigElementTransform(const FRigElementKey& InElement,
 			if (bLocal)
 			{
 				FTransform ParentTransform = FTransform::Identity;
-				FRigElementKey ParentKey = ControlRigBP->HierarchyContainer.BoneHierarchy[InElement.Name].GetParentElementKey();
+				FRigElementKey ParentKey = ControlRigBP->Hierarchy->GetFirstParent(InElement);
 				if (ParentKey.IsValid())
 				{
 					ParentTransform = GetRigElementTransform(ParentKey, false, false);
@@ -2861,8 +2899,8 @@ void FControlRigEditor::SetRigElementTransform(const FRigElementKey& InElement,
 				Transform.NormalizeRotation();
 			}
 
-			ControlRigBP->HierarchyContainer.BoneHierarchy.SetInitialGlobalTransform(InElement.Name, Transform);
-			ControlRigBP->HierarchyContainer.BoneHierarchy.SetGlobalTransform(InElement.Name, Transform);
+			ControlRigBP->Hierarchy->SetInitialGlobalTransform(InElement, Transform);
+			ControlRigBP->Hierarchy->SetGlobalTransform(InElement, Transform);
 			OnHierarchyChanged();
 			break;
 		}
@@ -2872,16 +2910,16 @@ void FControlRigEditor::SetRigElementTransform(const FRigElementKey& InElement,
 			FTransform GlobalTransform = InTransform;
 			if (!bLocal)
 			{
-				ControlRigBP->HierarchyContainer.ControlHierarchy.SetGlobalTransform(InElement.Name, InTransform);
-				LocalTransform = ControlRigBP->HierarchyContainer.ControlHierarchy.GetLocalTransform(InElement.Name);
+				ControlRigBP->Hierarchy->SetGlobalTransform(InElement, InTransform);
+				LocalTransform = ControlRigBP->Hierarchy->GetLocalTransform(InElement);
 			}
 			else
 			{
-				ControlRigBP->HierarchyContainer.ControlHierarchy.SetLocalTransform(InElement.Name, InTransform);
-				GlobalTransform = ControlRigBP->HierarchyContainer.ControlHierarchy.GetGlobalTransform(InElement.Name);
+				ControlRigBP->Hierarchy->SetLocalTransform(InElement, InTransform);
+				GlobalTransform = ControlRigBP->Hierarchy->GetGlobalTransform(InElement);
 			}
-			ControlRigBP->HierarchyContainer.ControlHierarchy.SetLocalTransform(InElement.Name, LocalTransform, ERigControlValueType::Initial);
-			ControlRigBP->HierarchyContainer.ControlHierarchy.SetGlobalTransform(InElement.Name, GlobalTransform);
+			ControlRigBP->Hierarchy->SetInitialLocalTransform(InElement, LocalTransform);
+			ControlRigBP->Hierarchy->SetGlobalTransform(InElement, GlobalTransform);
 			OnHierarchyChanged();
 			break;
 		}
@@ -2891,17 +2929,17 @@ void FControlRigEditor::SetRigElementTransform(const FRigElementKey& InElement,
 			FTransform GlobalTransform = InTransform;
 			if (!bLocal)
 			{
-				ControlRigBP->HierarchyContainer.SpaceHierarchy.SetGlobalTransform(InElement.Name, InTransform);
-				LocalTransform = ControlRigBP->HierarchyContainer.SpaceHierarchy.GetLocalTransform(InElement.Name);
+				ControlRigBP->Hierarchy->SetGlobalTransform(InElement, InTransform);
+				LocalTransform = ControlRigBP->Hierarchy->GetLocalTransform(InElement);
 			}
 			else
 			{
-				ControlRigBP->HierarchyContainer.SpaceHierarchy.SetLocalTransform(InElement.Name, InTransform);
-				GlobalTransform = ControlRigBP->HierarchyContainer.SpaceHierarchy.GetGlobalTransform(InElement.Name);
+				ControlRigBP->Hierarchy->SetLocalTransform(InElement, InTransform);
+				GlobalTransform = ControlRigBP->Hierarchy->GetGlobalTransform(InElement);
 			}
 
-			ControlRigBP->HierarchyContainer.SpaceHierarchy.SetInitialTransform(InElement.Name, LocalTransform);
-			ControlRigBP->HierarchyContainer.SpaceHierarchy.SetGlobalTransform(InElement.Name, GlobalTransform);
+			ControlRigBP->Hierarchy->SetInitialLocalTransform(InElement, LocalTransform);
+			ControlRigBP->Hierarchy->SetGlobalTransform(InElement, GlobalTransform);
 			OnHierarchyChanged();
 			break;
 		}
@@ -2975,128 +3013,153 @@ void FControlRigEditor::OnFinishedChangingProperties(const FPropertyChangedEvent
 		UControlRig* DebuggedControlRig = Cast<UControlRig>(GetBlueprintObj()->GetObjectBeingDebugged());
 		check(DebuggedControlRig);
 
-		if (DebuggedControlRig->GetHierarchy()->GetIndex(RigElementInDetailPanel) == INDEX_NONE)
+		UScriptStruct* ScriptStruct = PropertyChangedEvent.Property->GetOwner<UScriptStruct>();
+		if (ScriptStruct == FRigControlSettings::StaticStruct())
 		{
-			return;
-		}
+			URigHierarchy* Hierarchy = ControlRigBP->Hierarchy;
 
-		UScriptStruct* ScriptStruct = PropertyChangedEvent.MemberProperty->GetOwner<UScriptStruct>();
-		if (ScriptStruct)
-		{
-			FRigHierarchyContainer* Container = &ControlRigBP->HierarchyContainer;
-			if (!bSetupModeEnabled)
+			if(FRigControlElement* ControlElement = Hierarchy->Find<FRigControlElement>(RigElementInDetailPanel))
 			{
-				Container = &DebuggedControlRig->Hierarchy;
-			}
-
-			if (ScriptStruct == FRigBone::StaticStruct() && RigElementInDetailPanel.Type == ERigElementType::Bone)
-			{
-				if (PropertyChangedEvent.MemberProperty->GetFName() == TEXT("LocalTransform"))
+				FRigControlElement* DebuggedControlElement = DebuggedControlRig->GetHierarchy()->Find<FRigControlElement>(RigElementInDetailPanel);
+				if(DebuggedControlElement)
 				{
-					FRigBone& Bone = Container->BoneHierarchy[RigElementInDetailPanel.Name];
-					FTransform ParentTransform = FTransform::Identity;
-					if (Bone.ParentIndex != INDEX_NONE)
-					{
-						ParentTransform = Container->BoneHierarchy.GetGlobalTransform(Bone.ParentIndex);
-					}
-					Bone.GlobalTransform = Bone.LocalTransform * ParentTransform;
-
-					if (!bSetupModeEnabled)
-					{
-						ControlRigBP->PropagatePropertyFromInstanceToBP(RigElementInDetailPanel, PropertyChangedEvent.MemberProperty, DebuggedControlRig);
-					}
-					ControlRigBP->PropagatePropertyFromBPToInstances(RigElementInDetailPanel, PropertyChangedEvent.MemberProperty);
-					ControlRigBP->PropagatePropertyFromBPToInstances(RigElementInDetailPanel, ScriptStruct->FindPropertyByName(TEXT("GlobalTransform")));
+					ControlElement->Settings = DebuggedControlElement->Settings;
 				}
-				else if (PropertyChangedEvent.MemberProperty->GetFName() == TEXT("GlobalTransform"))
+			
+				if (PropertyChangedEvent.Property->GetName() == TEXT("GizmoColor"))
 				{
-					FRigBone& Bone = Container->BoneHierarchy[RigElementInDetailPanel.Name];
-					FTransform ParentTransform = FTransform::Identity;
-					if (Bone.ParentIndex != INDEX_NONE)
-					{
-						ParentTransform = Container->BoneHierarchy.GetGlobalTransform(Bone.ParentIndex);
-					}
-					Bone.LocalTransform = Bone.GlobalTransform.GetRelativeTransform(ParentTransform);
-
-					if (!bSetupModeEnabled)
-					{
-						ControlRigBP->PropagatePropertyFromInstanceToBP(RigElementInDetailPanel, PropertyChangedEvent.MemberProperty, DebuggedControlRig);
-					}
-					ControlRigBP->PropagatePropertyFromBPToInstances(RigElementInDetailPanel, PropertyChangedEvent.MemberProperty);
-					ControlRigBP->PropagatePropertyFromBPToInstances(RigElementInDetailPanel, ScriptStruct->FindPropertyByName(TEXT("LocalTransform")));
+					ControlElement->Settings.GizmoColor.R = FMath::Clamp<float>(ControlElement->Settings.GizmoColor.R, 0.f, 1.f);
+					ControlElement->Settings.GizmoColor.G = FMath::Clamp<float>(ControlElement->Settings.GizmoColor.G, 0.f, 1.f);
+					ControlElement->Settings.GizmoColor.B = FMath::Clamp<float>(ControlElement->Settings.GizmoColor.B, 0.f, 1.f);
+					ControlElement->Settings.GizmoColor.A = FMath::Clamp<float>(ControlElement->Settings.GizmoColor.A, 0.f, 1.f);
 				}
 
-				ControlRig->SetTransientControlValue(RigElementInDetailPanel);
-				if (DebuggedControlRig && DebuggedControlRig != ControlRig)
+				Hierarchy->Notify(ERigHierarchyNotification::ControlSettingChanged, ControlElement);
+				if(DebuggedControlElement)
 				{
-					DebuggedControlRig->SetTransientControlValue(RigElementInDetailPanel);
+					DebuggedControlRig->GetHierarchy()->Notify(ERigHierarchyNotification::ControlSettingChanged, DebuggedControlElement);
 				}
+                //ControlRigBP->PropagateHierarchyFromBPToInstances();
+			}
 
-				if (PreviewInstance)
+			ControlRigBP->Modify();
+			ControlRigBP->MarkPackageDirty();
+		}
+	}
+}
+
+void FControlRigEditor::OnBlueprintPropertyChainEvent(FPropertyChangedChainEvent& PropertyChangedChainEvent)
+{
+	if(PropertyChangedChainEvent.PropertyChain.Num() < 4)
+	{
+		return;
+	}
+	
+	UControlRigBlueprint* ControlRigBP = GetControlRigBlueprint();
+	if (ControlRig && ControlRigBP && RigElementInDetailPanel)
+	{
+		UControlRig* DebuggedControlRig = Cast<UControlRig>(GetBlueprintObj()->GetObjectBeingDebugged());
+		check(DebuggedControlRig);
+
+		const TDoubleLinkedList<FProperty*>::TDoubleLinkedListNode* HeadProperty = PropertyChangedChainEvent.PropertyChain.GetHead();
+		const TDoubleLinkedList<FProperty*>::TDoubleLinkedListNode* CurrentInitialProperty = HeadProperty->GetNextNode();
+		const TDoubleLinkedList<FProperty*>::TDoubleLinkedListNode* LocalGlobalProperty = CurrentInitialProperty->GetNextNode();
+
+		const bool bIsInitial = CurrentInitialProperty->GetValue()->GetFName() == TEXT("Initial");
+		const bool bIsLocal = LocalGlobalProperty->GetValue()->GetFName() == TEXT("Local");
+
+		ERigTransformType::Type TransformType = ERigTransformType::CurrentGlobal;
+		if(bIsInitial)
+		{
+			TransformType = ERigTransformType::MakeInitial(TransformType); 
+		}
+		if(bIsLocal)
+		{
+			TransformType = ERigTransformType::MakeLocal(TransformType); 
+		}
+
+		URigHierarchy* Hierarchy = DebuggedControlRig->GetHierarchy();
+		if(bSetupModeEnabled || bIsInitial)
+		{
+			Hierarchy = ControlRigBP->Hierarchy;
+		}
+
+		FRigBaseElement* SourceElement = DebuggedControlRig->GetHierarchy()->Find(RigElementInDetailPanel);
+		if(SourceElement == nullptr)
+		{
+			SourceElement = Hierarchy->Find(RigElementInDetailPanel);;
+		}
+		if(SourceElement == nullptr)
+		{
+			return;
+		}
+		FRigBaseElement* TargetElement = Hierarchy->Find(RigElementInDetailPanel);
+		if(TargetElement == nullptr)
+		{
+			return;
+		}
+		
+		if(HeadProperty->GetValue() == FRigTransformElement::StaticStruct()->FindPropertyByName(TEXT("Pose")))
+		{
+			if(FRigTransformElement* SourceTransformElement = Cast<FRigTransformElement>(SourceElement))
+			{
+				if(FRigTransformElement* TargetTransformElement = Cast<FRigTransformElement>(TargetElement))
 				{
-					if (FAnimNode_ModifyBone* Modify = PreviewInstance->FindModifiedBone(RigElementInDetailPanel.Name))
+					const FTransform Transform = SourceTransformElement->Pose.Get(TransformType);
+					if(ERigTransformType::IsLocal(TransformType) && TargetElement->IsA<FRigControlElement>())
 					{
-						FTransform LocalTransform = Container->BoneHierarchy[RigElementInDetailPanel.Name].LocalTransform;
-						Modify->Translation = LocalTransform.GetTranslation();
-						Modify->Rotation = LocalTransform.GetRotation().Rotator();
-						Modify->TranslationSpace = EBoneControlSpace::BCS_ParentBoneSpace;
-						Modify->RotationSpace = EBoneControlSpace::BCS_ParentBoneSpace;
+						FRigControlElement* TargetControlElement = Cast<FRigControlElement>(TargetTransformElement);
+						check(TargetControlElement);
+						
+						FRigControlValue Value;
+						Value.SetFromTransform(Transform, TargetControlElement->Settings.ControlType, TargetControlElement->Settings.PrimaryAxis);
+						
+						if(ERigTransformType::IsInitial(TransformType))
+						{
+							Hierarchy->SetControlValue(TargetControlElement, Value, ERigControlValueType::Initial, true);
+						}
+						else
+						{
+							Hierarchy->SetControlValue(TargetControlElement, Value, ERigControlValueType::Current, true);
+						}
 					}
-				}
-			}
-			else if (ScriptStruct == FRigSpace::StaticStruct() && RigElementInDetailPanel.Type == ERigElementType::Space)
-			{
-				ControlRigBP->PropagatePropertyFromBPToInstances(RigElementInDetailPanel, PropertyChangedEvent.MemberProperty);
-
-				if (PropertyChangedEvent.MemberProperty->GetName() == TEXT("InitialTransform"))
-				{
-					FRigSpace& Space = Container->SpaceHierarchy[RigElementInDetailPanel.Name];
-					Space.LocalTransform = Space.InitialTransform;
-
-					if (!bSetupModeEnabled)
+					else
 					{
-						ControlRigBP->PropagatePropertyFromInstanceToBP(RigElementInDetailPanel, ScriptStruct->FindPropertyByName(TEXT("LocalTransform")), DebuggedControlRig);
+						Hierarchy->SetTransform(TargetTransformElement, Transform, TransformType, true, true, true);
 					}
-					ControlRigBP->PropagatePropertyFromBPToInstances(RigElementInDetailPanel, ScriptStruct->FindPropertyByName(TEXT("LocalTransform")));
-				}
-
-				ControlRig->SetTransientControlValue(RigElementInDetailPanel);
-				if (DebuggedControlRig && DebuggedControlRig != ControlRig)
-				{
-					DebuggedControlRig->SetTransientControlValue(RigElementInDetailPanel);
 				}
 			}
-			else if (ScriptStruct == FRigControl::StaticStruct() && RigElementInDetailPanel.Type == ERigElementType::Control)
-			{
-				if (PropertyChangedEvent.MemberProperty->GetName() == TEXT("GizmoColor"))
-				{
-					FRigControl& Control = Container->ControlHierarchy[RigElementInDetailPanel.Name];
-					Control.GizmoColor.R = FMath::Clamp<float>(Control.GizmoColor.R, 0.f, 1.f);
-					Control.GizmoColor.G = FMath::Clamp<float>(Control.GizmoColor.G, 0.f, 1.f);
-					Control.GizmoColor.B = FMath::Clamp<float>(Control.GizmoColor.B, 0.f, 1.f);
-					Control.GizmoColor.A = FMath::Clamp<float>(Control.GizmoColor.A, 0.f, 1.f);
-				}
-
-				if (!bSetupModeEnabled)
-				{
-					ControlRigBP->PropagatePropertyFromInstanceToBP(RigElementInDetailPanel, PropertyChangedEvent.MemberProperty, DebuggedControlRig);
-				}
-				ControlRigBP->PropagatePropertyFromBPToInstances(RigElementInDetailPanel, PropertyChangedEvent.MemberProperty);
+		}
 
-				if (PropertyChangedEvent.MemberProperty->GetName().Contains(TEXT("Gizmo")))
+		if(HeadProperty->GetValue() == FRigControlElement::StaticStruct()->FindPropertyByName(TEXT("Offset")))
+		{
+			if(FRigControlElement* SourceControlElement = Cast<FRigControlElement>(SourceElement))
+			{
+				if(FRigControlElement* TargetControlElement = Cast<FRigControlElement>(TargetElement))
 				{
-					ControlRigBP->HierarchyContainer.ControlHierarchy.OnControlUISettingsChanged.Broadcast(&ControlRigBP->HierarchyContainer, FRigElementKey(RigElementInDetailPanel.Name, ERigElementType::Control));
+					const FTransform Transform = SourceControlElement->Offset.Get(TransformType);
+					Hierarchy->SetControlOffsetTransform(TargetControlElement, Transform, TransformType, true, true, true);
+					Hierarchy->SetControlOffsetTransform(TargetControlElement, Transform, ERigTransformType::MakeCurrent(TransformType), true, true, true);
 				}
 			}
-			else if (ScriptStruct == FRigCurve::StaticStruct() && RigElementInDetailPanel.Type == ERigElementType::Curve)
+		}
+
+		if(HeadProperty->GetValue() == FRigControlElement::StaticStruct()->FindPropertyByName(TEXT("Gizmo")))
+		{
+			if(FRigControlElement* SourceControlElement = Cast<FRigControlElement>(SourceElement))
 			{
-				if (!bSetupModeEnabled)
+				if(FRigControlElement* TargetControlElement = Cast<FRigControlElement>(TargetElement))
 				{
-					ControlRigBP->PropagatePropertyFromInstanceToBP(RigElementInDetailPanel, PropertyChangedEvent.MemberProperty, DebuggedControlRig);
+					const FTransform Transform = SourceControlElement->Gizmo.Get(TransformType);
+					Hierarchy->SetControlGizmoTransform(TargetControlElement, Transform, TransformType, true, true);
+					Hierarchy->SetControlGizmoTransform(TargetControlElement, Transform, ERigTransformType::MakeCurrent(TransformType), true, true);
 				}
-				ControlRigBP->PropagatePropertyFromBPToInstances(RigElementInDetailPanel, PropertyChangedEvent.MemberProperty);
 			}
+		}
 
+		if(bIsInitial)
+		{
+			ControlRigBP->PropagatePoseFromBPToInstances();
 			ControlRigBP->Modify();
 			ControlRigBP->MarkPackageDirty();
 		}
@@ -3126,10 +3189,11 @@ void FControlRigEditor::OnHierarchyChanged()
 		ControlRigBP->PropagateHierarchyFromBPToInstances();
 
 		FBlueprintEditorUtils::MarkBlueprintAsModified(GetControlRigBlueprint());
-		TArray<FRigElementKey> SelectedElements = ControlRigBP->HierarchyContainer.CurrentSelection();
-		for(const FRigElementKey& SelectedElement : SelectedElements)
+		
+		TArray<FRigBaseElement*> SelectedElements = ControlRigBP->Hierarchy->GetSelectedElements();
+		for(const FRigBaseElement* SelectedElement : SelectedElements)
 		{
-			ControlRigBP->HierarchyContainer.OnElementSelected.Broadcast(&ControlRigBP->HierarchyContainer, SelectedElement, true);
+			ControlRigBP->Hierarchy->OnModified().Broadcast(ERigHierarchyNotification::ElementSelected, ControlRigBP->Hierarchy, SelectedElement);
 		}
 		GetControlRigBlueprint()->RecompileVM();
 
@@ -3153,7 +3217,7 @@ void FControlRigEditor::OnHierarchyChanged()
 				ClearDetailObject();
 			}
 		}
-		else if (ControlRigBP->HierarchyContainer.GetIndex(RigElementInDetailPanel) == INDEX_NONE)
+		else if (!ControlRigBP->Hierarchy->Contains(RigElementInDetailPanel))
 		{
 			ClearDetailObject();
 		}
@@ -3165,86 +3229,94 @@ void FControlRigEditor::OnHierarchyChanged()
 }
 
 
-
-void FControlRigEditor::OnRigElementAdded(FRigHierarchyContainer* Container, const FRigElementKey& InKey)
+void FControlRigEditor::OnHierarchyModified(ERigHierarchyNotification InNotif, URigHierarchy* InHierarchy, const FRigBaseElement* InElement)
 {
-	if (UControlRigBlueprint* RigBlueprint = Cast<UControlRigBlueprint>(GetBlueprintObj()))
+	UControlRigBlueprint* RigBlueprint = Cast<UControlRigBlueprint>(GetBlueprintObj());
+	if(RigBlueprint == nullptr)
 	{
-		if (RigBlueprint->bSuspendAllNotifications)
-		{
-			return;
-		}
+		return;
 	}
-	OnHierarchyChanged();
-}
-
-void FControlRigEditor::OnRigElementRemoved(FRigHierarchyContainer* Container, const FRigElementKey& InKey, bool bForce)
-{
-	UControlRigBlueprint* Blueprint = GetControlRigBlueprint();
 
-	if (Blueprint->bSuspendAllNotifications && !bForce)
+	if (RigBlueprint->bSuspendAllNotifications)
 	{
 		return;
 	}
 
-	UEnum* RigElementTypeEnum = StaticEnum<ERigElementType>();
-	if (RigElementTypeEnum == nullptr)
+	if(InHierarchy != RigBlueprint->Hierarchy)
 	{
 		return;
 	}
 
-	CacheNameLists();
-
-	FString RemovedElementName = InKey.Name.ToString();
-	ERigElementType RemovedElementType = InKey.Type;
-
-	TArray<UEdGraph*> EdGraphs;
-	Blueprint->GetAllGraphs(EdGraphs);
-
-	for (UEdGraph* Graph : EdGraphs)
+	switch(InNotif)
 	{
-		UControlRigGraph* RigGraph = Cast<UControlRigGraph>(Graph);
-		if (RigGraph == nullptr)
+		case ERigHierarchyNotification::ElementAdded:
+		case ERigHierarchyNotification::ParentChanged:
+		case ERigHierarchyNotification::HierarchyReset:
 		{
-			continue;
+			OnHierarchyChanged();
+			break;
 		}
-
-		for (UEdGraphNode* Node : RigGraph->Nodes)
+		case ERigHierarchyNotification::ElementRemoved:
 		{
-			if (UControlRigGraphNode* RigNode = Cast<UControlRigGraphNode>(Node))
+			UEnum* RigElementTypeEnum = StaticEnum<ERigElementType>();
+			if (RigElementTypeEnum == nullptr)
 			{
-				if (URigVMNode* ModelNode = RigNode->GetModelNode())
+				return;
+			}
+
+			CacheNameLists();
+
+			const FString RemovedElementName = InElement->GetName().ToString();
+			const ERigElementType RemovedElementType = InElement->GetType();
+
+			TArray<UEdGraph*> EdGraphs;
+			RigBlueprint->GetAllGraphs(EdGraphs);
+
+			for (UEdGraph* Graph : EdGraphs)
+			{
+				UControlRigGraph* RigGraph = Cast<UControlRigGraph>(Graph);
+				if (RigGraph == nullptr)
+				{
+					continue;
+				}
+
+				for (UEdGraphNode* Node : RigGraph->Nodes)
 				{
-					TArray<URigVMPin*> ModelPins = ModelNode->GetAllPinsRecursively();
-					for (URigVMPin* ModelPin : ModelPins)
+					if (UControlRigGraphNode* RigNode = Cast<UControlRigGraphNode>(Node))
 					{
-						if ((ModelPin->GetCPPType() == TEXT("FName") && ModelPin->GetCustomWidgetName() == TEXT("BoneName") && RemovedElementType == ERigElementType::Bone) ||
-							(ModelPin->GetCPPType() == TEXT("FName") && ModelPin->GetCustomWidgetName() == TEXT("ControlName") && RemovedElementType == ERigElementType::Control) ||
-							(ModelPin->GetCPPType() == TEXT("FName") && ModelPin->GetCustomWidgetName() == TEXT("SpaceName") && RemovedElementType == ERigElementType::Space) ||
-							(ModelPin->GetCPPType() == TEXT("FName") && ModelPin->GetCustomWidgetName() == TEXT("CurveName") && RemovedElementType == ERigElementType::Curve))
+						if (URigVMNode* ModelNode = RigNode->GetModelNode())
 						{
-							if (ModelPin->GetDefaultValue() == RemovedElementName)
-							{
-								RigNode->ReconstructNode();
-								break;
-							}
-						}
-						else if (ModelPin->GetCPPTypeObject() == FRigElementKey::StaticStruct())
-						{
-
-							if (URigVMPin* TypePin = ModelPin->FindSubPin(TEXT("Type")))
+							TArray<URigVMPin*> ModelPins = ModelNode->GetAllPinsRecursively();
+							for (URigVMPin* ModelPin : ModelPins)
 							{
-								FString TypeStr = TypePin->GetDefaultValue();
-								int64 TypeValue = RigElementTypeEnum->GetValueByNameString(TypeStr);
-								if (TypeValue == (int64)RemovedElementType)
+								if ((ModelPin->GetCPPType() == TEXT("FName") && ModelPin->GetCustomWidgetName() == TEXT("BoneName") && RemovedElementType == ERigElementType::Bone) ||
+									(ModelPin->GetCPPType() == TEXT("FName") && ModelPin->GetCustomWidgetName() == TEXT("ControlName") && RemovedElementType == ERigElementType::Control) ||
+									(ModelPin->GetCPPType() == TEXT("FName") && ModelPin->GetCustomWidgetName() == TEXT("SpaceName") && RemovedElementType == ERigElementType::Space) ||
+									(ModelPin->GetCPPType() == TEXT("FName") && ModelPin->GetCustomWidgetName() == TEXT("CurveName") && RemovedElementType == ERigElementType::Curve))
 								{
-									if (URigVMPin* NamePin = ModelPin->FindSubPin(TEXT("Name")))
+									if (ModelPin->GetDefaultValue() == RemovedElementName)
 									{
-										FString NameStr = NamePin->GetDefaultValue();
-										if (NameStr == RemovedElementName)
+										RigNode->ReconstructNode();
+										break;
+									}
+								}
+								else if (ModelPin->GetCPPTypeObject() == FRigElementKey::StaticStruct())
+								{
+									if (URigVMPin* TypePin = ModelPin->FindSubPin(TEXT("Type")))
+									{
+										FString TypeStr = TypePin->GetDefaultValue();
+										int64 TypeValue = RigElementTypeEnum->GetValueByNameString(TypeStr);
+										if (TypeValue == (int64)RemovedElementType)
 										{
-											RigNode->ReconstructNode();
-											break;
+											if (URigVMPin* NamePin = ModelPin->FindSubPin(TEXT("Name")))
+											{
+												FString NameStr = NamePin->GetDefaultValue();
+												if (NameStr == RemovedElementName)
+												{
+													RigNode->ReconstructNode();
+													break;
+												}
+											}
 										}
 									}
 								}
@@ -3253,76 +3325,68 @@ void FControlRigEditor::OnRigElementRemoved(FRigHierarchyContainer* Container, c
 					}
 				}
 			}
-		}
-	}
-
-	OnHierarchyChanged();
-}
 
-void FControlRigEditor::OnRigElementRenamed(FRigHierarchyContainer* Container, ERigElementType ElementType, const FName& InOldName, const FName& InNewName)
-{
-	DECLARE_SCOPE_HIERARCHICAL_COUNTER_FUNC()
-
-	UControlRigBlueprint* Blueprint = GetControlRigBlueprint();
-
-	if (Blueprint->bSuspendAllNotifications)
-	{
-		return;
-	}
-
-	UEnum* RigElementTypeEnum = StaticEnum<ERigElementType>();
-	if (RigElementTypeEnum == nullptr)
-	{
-		return;
-	}
-
-	FString OldNameStr = InOldName.ToString();
-	FString NewNameStr = InNewName.ToString();
+			OnHierarchyChanged();
+			break;
+		}
+		case ERigHierarchyNotification::ElementRenamed:
+		{
+			UEnum* RigElementTypeEnum = StaticEnum<ERigElementType>();
+			if (RigElementTypeEnum == nullptr)
+			{
+				return;
+			}
 
-	TArray<UEdGraph*> EdGraphs;
-	Blueprint->GetAllGraphs(EdGraphs);
+			const FString OldNameStr = InHierarchy->GetPreviousName(InElement->GetKey()).ToString();
+			const FString NewNameStr = InElement->GetName().ToString();
+			const ERigElementType ElementType = InElement->GetType(); 
 
-	for (UEdGraph* Graph : EdGraphs)
-	{
-		UControlRigGraph* RigGraph = Cast<UControlRigGraph>(Graph);
-		if (RigGraph == nullptr)
-		{
-			continue;
-		}
+			TArray<UEdGraph*> EdGraphs;
+			RigBlueprint->GetAllGraphs(EdGraphs);
 
-		for (UEdGraphNode* Node : RigGraph->Nodes)
-		{
-			if (UControlRigGraphNode* RigNode = Cast<UControlRigGraphNode>(Node))
+			for (UEdGraph* Graph : EdGraphs)
 			{
-				if (URigVMNode* ModelNode = RigNode->GetModelNode())
+				UControlRigGraph* RigGraph = Cast<UControlRigGraph>(Graph);
+				if (RigGraph == nullptr)
+				{
+					continue;
+				}
+
+				for (UEdGraphNode* Node : RigGraph->Nodes)
 				{
-					TArray<URigVMPin*> ModelPins = ModelNode->GetAllPinsRecursively();
-					for (URigVMPin * ModelPin : ModelPins)
+					if (UControlRigGraphNode* RigNode = Cast<UControlRigGraphNode>(Node))
 					{
-						if ((ModelPin->GetCPPType() == TEXT("FName") && ModelPin->GetCustomWidgetName() == TEXT("BoneName") && ElementType == ERigElementType::Bone) ||
-							(ModelPin->GetCPPType() == TEXT("FName") && ModelPin->GetCustomWidgetName() == TEXT("ControlName") && ElementType == ERigElementType::Control) ||
-							(ModelPin->GetCPPType() == TEXT("FName") && ModelPin->GetCustomWidgetName() == TEXT("SpaceName") && ElementType == ERigElementType::Space) ||
-							(ModelPin->GetCPPType() == TEXT("FName") && ModelPin->GetCustomWidgetName() == TEXT("CurveName") && ElementType == ERigElementType::Curve))
-						{
-							if (ModelPin->GetDefaultValue() == InOldName.ToString())
-							{
-								GetFocusedController()->SetPinDefaultValue(ModelPin->GetPinPath(), InNewName.ToString(), false);
-							}
-						}
-						else if (ModelPin->GetCPPTypeObject() == FRigElementKey::StaticStruct())
+						if (URigVMNode* ModelNode = RigNode->GetModelNode())
 						{
-							if (URigVMPin* TypePin = ModelPin->FindSubPin(TEXT("Type")))
+							TArray<URigVMPin*> ModelPins = ModelNode->GetAllPinsRecursively();
+							for (URigVMPin * ModelPin : ModelPins)
 							{
-								FString TypeStr = TypePin->GetDefaultValue();
-								int64 TypeValue = RigElementTypeEnum->GetValueByNameString(TypeStr);
-								if (TypeValue == (int64)ElementType)
+								if ((ModelPin->GetCPPType() == TEXT("FName") && ModelPin->GetCustomWidgetName() == TEXT("BoneName") && ElementType == ERigElementType::Bone) ||
+                                    (ModelPin->GetCPPType() == TEXT("FName") && ModelPin->GetCustomWidgetName() == TEXT("ControlName") && ElementType == ERigElementType::Control) ||
+                                    (ModelPin->GetCPPType() == TEXT("FName") && ModelPin->GetCustomWidgetName() == TEXT("SpaceName") && ElementType == ERigElementType::Space) ||
+                                    (ModelPin->GetCPPType() == TEXT("FName") && ModelPin->GetCustomWidgetName() == TEXT("CurveName") && ElementType == ERigElementType::Curve))
 								{
-									if (URigVMPin* NamePin = ModelPin->FindSubPin(TEXT("Name")))
+									if (ModelPin->GetDefaultValue() == OldNameStr)
 									{
-										FString NameStr = NamePin->GetDefaultValue();
-										if (NameStr == OldNameStr)
+										GetFocusedController()->SetPinDefaultValue(ModelPin->GetPinPath(), NewNameStr, false);
+									}
+								}
+								else if (ModelPin->GetCPPTypeObject() == FRigElementKey::StaticStruct())
+								{
+									if (URigVMPin* TypePin = ModelPin->FindSubPin(TEXT("Type")))
+									{
+										const FString TypeStr = TypePin->GetDefaultValue();
+										const int64 TypeValue = RigElementTypeEnum->GetValueByNameString(TypeStr);
+										if (TypeValue == (int64)ElementType)
 										{
-											GetFocusedController()->SetPinDefaultValue(NamePin->GetPinPath(), NewNameStr);
+											if (URigVMPin* NamePin = ModelPin->FindSubPin(TEXT("Name")))
+											{
+												FString NameStr = NamePin->GetDefaultValue();
+												if (NameStr == OldNameStr)
+												{
+													GetFocusedController()->SetPinDefaultValue(NamePin->GetPinPath(), NewNameStr);
+												}
+											}
 										}
 									}
 								}
@@ -3331,24 +3395,50 @@ void FControlRigEditor::OnRigElementRenamed(FRigHierarchyContainer* Container, E
 					}
 				}
 			}
-		}
+				
+			CacheNameLists();
+			OnHierarchyChanged();
 
-		CacheNameLists();
-	}
+			break;
+		}
+		case ERigHierarchyNotification::ElementSelected:
+		case ERigHierarchyNotification::ElementDeselected:
+		{
+			check(InElement);
+			const bool bSelected = InNotif == ERigHierarchyNotification::ElementSelected;
 
-}
+			if (InElement->GetType() == ERigElementType::Bone)
+			{
+				SynchronizeViewportBoneSelection();
+			}
 
-void FControlRigEditor::OnRigElementReparented(FRigHierarchyContainer* Container, const FRigElementKey& InKey, const FName& InOldParentName, const FName& InNewParentName)
-{
-	if (UControlRigBlueprint* RigBlueprint = Cast<UControlRigBlueprint>(GetBlueprintObj()))
-	{
-		if (RigBlueprint->bSuspendAllNotifications)
+			if (bSelected)
+			{
+				SetDetailStruct(InElement->GetKey());
+			}
+			else
+			{
+				TArray<FRigElementKey> CurrentSelection = RigBlueprint->Hierarchy->GetSelectedKeys();
+				if (CurrentSelection.Num() > 0)
+				{
+					if(FRigBaseElement* LastSelectedElement = InHierarchy->Find(CurrentSelection.Last()))
+					{
+						OnHierarchyModified(ERigHierarchyNotification::ElementSelected, InHierarchy, LastSelectedElement);
+					}
+				}
+				else
+				{
+					ClearDetailObject();
+				}
+			}
+				
+			break;
+		}
+		default:
 		{
-			return;
+			break;
 		}
 	}
-
-	OnHierarchyChanged();
 }
 
 void FControlRigEditor::SynchronizeViewportBoneSelection()
@@ -3364,60 +3454,17 @@ void FControlRigEditor::SynchronizeViewportBoneSelection()
 	{
 		EditorSkelComp->BonesOfInterest.Reset();
 
-		TArray<FName> SelectedBones = RigBlueprint->HierarchyContainer.BoneHierarchy.CurrentSelection();
-		for (const FName& Bone : SelectedBones)
-		{
-			int32 Index = ControlRig->Hierarchy.BoneHierarchy.GetIndex(Bone); 
-			EditorSkelComp->BonesOfInterest.AddUnique(Index);
-		}
-	}
-}
-
-void FControlRigEditor::OnRigElementSelected(FRigHierarchyContainer* Container, const FRigElementKey& InKey, bool bSelected)
-{
-	UControlRigBlueprint* RigBlueprint = GetControlRigBlueprint();
-	if (RigBlueprint == nullptr)
-	{
-		return;
-	}
-
-	if (RigBlueprint->bSuspendAllNotifications)
-	{
-		return;
-	}
-
-	if (Container->GetIndex(InKey) == INDEX_NONE)
-	{
-		return;
-	}
-
-	if (InKey.Type == ERigElementType::Bone)
-	{
-		SynchronizeViewportBoneSelection();
-	}
-
-	if (bSelected)
-	{
-		SetDetailStruct(InKey);
-	}
-	else
-	{
-		TArray<FRigElementKey> CurrentSelection = RigBlueprint->HierarchyContainer.CurrentSelection();
-		if (CurrentSelection.Num() > 0)
+		TArray<FRigBaseElement*> SelectedBones = RigBlueprint->Hierarchy->GetSelectedElements(ERigElementType::Bone);
+		for (const FRigBaseElement* SelectedBone : SelectedBones)
 		{
-			OnRigElementSelected(Container, CurrentSelection.Last(), true);
-		}
-		else
-		{
-			ClearDetailObject();
+			if(EditorSkelComp->GetReferenceSkeleton().IsValidIndex(SelectedBone->GetIndex()))
+			{
+				EditorSkelComp->BonesOfInterest.AddUnique(SelectedBone->GetIndex());
+			}
 		}
 	}
 }
 
-void FControlRigEditor::OnControlUISettingChanged(FRigHierarchyContainer* Container, const FRigElementKey& InKey)
-{
-}
-
 FControlRigEditorEditMode* FControlRigEditor::GetEditMode() const
 {
 	return static_cast<FControlRigEditorEditMode*>(GetEditorModeManager().GetActiveMode(FControlRigEditorEditMode::ModeName));
@@ -3687,8 +3734,10 @@ void FControlRigEditor::HandleMakeElementGetterSetter(ERigElementGetterSetterTyp
 				{
 					if (Key.Type == ERigElementType::Control)
 					{
-						const FRigControl& Control = Blueprint->HierarchyContainer.ControlHierarchy[Key.Name];
-						switch (Control.ControlType)
+						FRigControlElement* ControlElement = Blueprint->Hierarchy->Find<FRigControlElement>(Key);
+						check(ControlElement);
+						
+						switch (ControlElement->Settings.ControlType)
 						{
 							case ERigControlType::Bool:
 							{
@@ -3773,8 +3822,10 @@ void FControlRigEditor::HandleMakeElementGetterSetter(ERigElementGetterSetterTyp
 				{
 					if (Key.Type == ERigElementType::Control)
 					{
-						const FRigControl& Control = Blueprint->HierarchyContainer.ControlHierarchy[Key.Name];
-						switch (Control.ControlType)
+						FRigControlElement* ControlElement = Blueprint->Hierarchy->Find<FRigControlElement>(Key);
+						check(ControlElement);
+
+						switch (ControlElement->Settings.ControlType)
 						{
 							case ERigControlType::Bool:
 							{
@@ -3806,7 +3857,7 @@ void FControlRigEditor::HandleMakeElementGetterSetter(ERigElementGetterSetterTyp
 								StructTemplate = FRigUnit_SetControlVector::StaticStruct();
 								NewNode.ValuePinName = TEXT("Vector");
 								NewNode.ValueType = ERigControlType::Position;
-								NewNode.Value = FRigControlValue::Make<FVector>(Blueprint->HierarchyContainer.GetGlobalTransform(Key).GetLocation());
+								NewNode.Value = FRigControlValue::Make<FVector>(Blueprint->Hierarchy->GetGlobalTransform(Key).GetLocation());
 								break;
 							}
 							case ERigControlType::Scale:
@@ -3815,7 +3866,7 @@ void FControlRigEditor::HandleMakeElementGetterSetter(ERigElementGetterSetterTyp
 								StructTemplate = FRigUnit_SetControlVector::StaticStruct();
 								NewNode.ValuePinName = TEXT("Vector");
 								NewNode.ValueType = ERigControlType::Scale;
-								NewNode.Value = FRigControlValue::Make<FVector>(Blueprint->HierarchyContainer.GetGlobalTransform(Key).GetScale3D());
+								NewNode.Value = FRigControlValue::Make<FVector>(Blueprint->Hierarchy->GetGlobalTransform(Key).GetScale3D());
 								break;
 							}
 							case ERigControlType::Rotator:
@@ -3824,7 +3875,7 @@ void FControlRigEditor::HandleMakeElementGetterSetter(ERigElementGetterSetterTyp
 								StructTemplate = FRigUnit_SetControlRotator::StaticStruct();
 								NewNode.ValuePinName = TEXT("Rotator");
 								NewNode.ValueType = ERigControlType::Rotator;
-								NewNode.Value = FRigControlValue::Make<FRotator>(Blueprint->HierarchyContainer.GetGlobalTransform(Key).Rotator());
+								NewNode.Value = FRigControlValue::Make<FRotator>(Blueprint->Hierarchy->GetGlobalTransform(Key).Rotator());
 								break;
 							}
 							case ERigControlType::Transform:
@@ -3834,7 +3885,7 @@ void FControlRigEditor::HandleMakeElementGetterSetter(ERigElementGetterSetterTyp
 								StructTemplate = FRigUnit_SetTransform::StaticStruct();
 								NewNode.ValuePinName = TEXT("Transform");
 								NewNode.ValueType = ERigControlType::Transform;
-								NewNode.Value = FRigControlValue::Make<FTransform>(Blueprint->HierarchyContainer.GetGlobalTransform(Key));
+								NewNode.Value = FRigControlValue::Make<FTransform>(Blueprint->Hierarchy->GetGlobalTransform(Key));
 								break;
 							}
 							default:
@@ -3848,7 +3899,7 @@ void FControlRigEditor::HandleMakeElementGetterSetter(ERigElementGetterSetterTyp
 						StructTemplate = FRigUnit_SetTransform::StaticStruct();
 						NewNode.ValuePinName = TEXT("Transform");
 						NewNode.ValueType = ERigControlType::Transform;
-						NewNode.Value = FRigControlValue::Make<FTransform>(Blueprint->HierarchyContainer.GetGlobalTransform(Key));
+						NewNode.Value = FRigControlValue::Make<FTransform>(Blueprint->Hierarchy->GetGlobalTransform(Key));
 					}
 					break;
 				}
@@ -3865,7 +3916,7 @@ void FControlRigEditor::HandleMakeElementGetterSetter(ERigElementGetterSetterTyp
 					StructTemplate = FRigUnit_SetRotation::StaticStruct();
 					NewNode.ValuePinName = TEXT("Rotation");
 					NewNode.ValueType = ERigControlType::Rotator;
-					NewNode.Value = FRigControlValue::Make<FRotator>(Blueprint->HierarchyContainer.GetGlobalTransform(Key).Rotator());
+					NewNode.Value = FRigControlValue::Make<FRotator>(Blueprint->Hierarchy->GetGlobalTransform(Key).Rotator());
 					break;
 				}
 				case ERigElementGetterSetterType_Translation:
@@ -3873,7 +3924,7 @@ void FControlRigEditor::HandleMakeElementGetterSetter(ERigElementGetterSetterTyp
 					StructTemplate = FRigUnit_SetTranslation::StaticStruct();
 					NewNode.ValuePinName = TEXT("Translation");
 					NewNode.ValueType = ERigControlType::Position;
-					NewNode.Value = FRigControlValue::Make<FVector>(Blueprint->HierarchyContainer.GetGlobalTransform(Key).GetLocation());
+					NewNode.Value = FRigControlValue::Make<FVector>(Blueprint->Hierarchy->GetGlobalTransform(Key).GetLocation());
 					break;
 				}
 				case ERigElementGetterSetterType_Offset:
@@ -3972,7 +4023,7 @@ void FControlRigEditor::HandleMakeElementGetterSetter(ERigElementGetterSetterTyp
 	}
 }
 
-void FControlRigEditor::HandleOnControlModified(UControlRig* Subject, const FRigControl& Control, const FRigControlModifiedContext& Context)
+void FControlRigEditor::HandleOnControlModified(UControlRig* Subject, FRigControlElement* ControlElement, const FRigControlModifiedContext& Context)
 {
 	if (Subject != ControlRig)
 	{
@@ -3985,39 +4036,44 @@ void FControlRigEditor::HandleOnControlModified(UControlRig* Subject, const FRig
 		return;
 	}
 
-	if (Control.bIsTransientControl)
+	URigHierarchy* Hierarchy = Subject->GetHierarchy();
+
+	if (ControlElement->Settings.bIsTransientControl)
 	{
-		if (URigVMPin* Pin = Blueprint->Model->FindPin(Control.Name.ToString()))
+		FRigControlValue ControlValue = Hierarchy->GetControlValue(ControlElement, ERigControlValueType::Current);
+
+		const FString PinPath = UControlRig::GetPinNameFromTransientControl(ControlElement->GetKey());
+		if (URigVMPin* Pin = Blueprint->Model->FindPin(PinPath))
 		{
 			FString NewDefaultValue;
-			switch (Control.ControlType)
+			switch (ControlElement->Settings.ControlType)
 			{
 				case ERigControlType::Position:
 				case ERigControlType::Scale:
 				{
-					NewDefaultValue = Control.Value.ToString<FVector>();
+					NewDefaultValue = ControlValue.ToString<FVector>();
 					break;
 				}
 				case ERigControlType::Rotator:
 				{
-					FRotator Rotator = Control.Value.Get<FRotator>();
+					FRotator Rotator = ControlValue.Get<FRotator>();
 					FRigControlValue QuatValue = FRigControlValue::Make<FQuat>(FQuat(Rotator));
 					NewDefaultValue = QuatValue.ToString<FQuat>();
 					break;
 				}
 				case ERigControlType::Transform:
 				{
-					NewDefaultValue = Control.Value.ToString<FTransform>();
+					NewDefaultValue = ControlValue.ToString<FTransform>();
 					break;
 				}
 				case ERigControlType::TransformNoScale:
 				{
-					NewDefaultValue = Control.Value.ToString<FTransformNoScale>();
+					NewDefaultValue = ControlValue.ToString<FTransformNoScale>();
 					break;
 				}
 				case ERigControlType::EulerTransform:
 				{
-					NewDefaultValue = Control.Value.ToString<FEulerTransform>();
+					NewDefaultValue = ControlValue.ToString<FEulerTransform>();
 					break;
 				}
 				default:
@@ -4033,34 +4089,23 @@ void FControlRigEditor::HandleOnControlModified(UControlRig* Subject, const FRig
 		}
 		else
 		{
-			static FString ControlRigForElementBoneName;
-			static FString ControlRigForElementSpaceName;
-
-			if (ControlRigForElementBoneName.IsEmpty())
-			{
-				ControlRigForElementBoneName = FString::Printf(TEXT("ControlForRigElement_%s_"),
-					*StaticEnum<ERigElementType>()->GetNameByValue((int64)ERigElementType::Bone).ToString());
-				ControlRigForElementSpaceName = FString::Printf(TEXT("ControlForRigElement_%s_"),
-					*StaticEnum<ERigElementType>()->GetNameByValue((int64)ERigElementType::Space).ToString());
-			}
+			const FRigElementKey ElementKey = UControlRig::GetElementKeyFromTransientControl(ControlElement->GetKey());
 
-			if (Control.Name.ToString().StartsWith(ControlRigForElementBoneName))
+			if (ElementKey.Type == ERigElementType::Bone)
 			{
-				FName BoneName = *Control.Name.ToString().RightChop(ControlRigForElementBoneName.Len());
-
-				FTransform Transform = Control.Value.Get<FTransform>() * Control.OffsetTransform;
-				Blueprint->HierarchyContainer.BoneHierarchy.SetLocalTransform(BoneName, Transform);
+				const FTransform Transform = ControlValue.Get<FTransform>() * Hierarchy->GetControlOffsetTransform(ControlElement, ERigTransformType::CurrentLocal);
+				Blueprint->Hierarchy->SetLocalTransform(ElementKey, Transform);
+				Hierarchy->SetLocalTransform(ElementKey, Transform);
 
 				if (bSetupModeEnabled)
 				{
-					Blueprint->HierarchyContainer.BoneHierarchy.CopyGlobalToInitial(Blueprint->HierarchyContainer.BoneHierarchy.GetIndex(BoneName));
+					Blueprint->Hierarchy->SetInitialLocalTransform(ElementKey, Transform);
+					Hierarchy->SetInitialLocalTransform(ElementKey, Transform);
 				}
 
-				Blueprint->PropagateHierarchyFromBPToInstances(false, false);
-
 				if (PreviewInstance)
 				{
-					if (FAnimNode_ModifyBone* Modify = PreviewInstance->FindModifiedBone(BoneName))
+					if (FAnimNode_ModifyBone* Modify = PreviewInstance->FindModifiedBone(ElementKey.Name))
 					{
 						Modify->Translation = Transform.GetTranslation();
 						Modify->Rotation = Transform.GetRotation().Rotator();
@@ -4069,21 +4114,25 @@ void FControlRigEditor::HandleOnControlModified(UControlRig* Subject, const FRig
 					}
 				}
 			}
-			else if (Control.Name.ToString().StartsWith(ControlRigForElementSpaceName))
+			else if (ElementKey.Type == ERigElementType::Space)
 			{
-				FName SpaceName = *Control.Name.ToString().RightChop(ControlRigForElementSpaceName.Len());
-
-				FTransform GlobalTransform = ControlRig->GetControlGlobalTransform(Control.Name);
-				Blueprint->HierarchyContainer.SpaceHierarchy.SetGlobalTransform(SpaceName, GlobalTransform);
-				Blueprint->HierarchyContainer.SpaceHierarchy.SetInitialGlobalTransform(SpaceName, GlobalTransform);
-				Blueprint->PropagateHierarchyFromBPToInstances(false, false);
+				const FTransform GlobalTransform = ControlRig->GetControlGlobalTransform(ControlElement->GetName());
+				Blueprint->Hierarchy->SetGlobalTransform(ElementKey, GlobalTransform);
+				Blueprint->Hierarchy->SetInitialGlobalTransform(ElementKey, GlobalTransform);
+				Hierarchy->SetGlobalTransform(ElementKey, GlobalTransform);
+				Hierarchy->SetInitialGlobalTransform(ElementKey, GlobalTransform);
 			}
 		}
 	}
 	else if (bSetupModeEnabled)
 	{
-		FRigControlHierarchy& ControlHierarchy = ControlRig->GetControlHierarchy();
-		Blueprint->HierarchyContainer.ControlHierarchy[Control.Index] = ControlHierarchy[Control.Index];
+		FRigControlElement* SourceControlElement = Hierarchy->Find<FRigControlElement>(ControlElement->GetKey());
+		FRigControlElement* TargetControlElement = Blueprint->Hierarchy->Find<FRigControlElement>(ControlElement->GetKey());
+
+		TargetControlElement->Settings = SourceControlElement->Settings;
+		Blueprint->Hierarchy->OnModified().Broadcast(ERigHierarchyNotification::ControlSettingChanged, Blueprint->Hierarchy, TargetControlElement);
+
+		TargetControlElement->CopyPose(SourceControlElement, true, true);
 	}
 }
 
@@ -4328,116 +4377,6 @@ void FControlRigEditor::UpdateGraphCompilerErrors()
 
 }
 
-void FControlRigEditor::DumpUnitTestCode()
-{
-	/*
-	if (UEdGraph* Graph = GetFocusedGraph())
-	{
-		TArray<FString> Code;
-
-		// dump the hierarchy
-		if (ControlRig)
-		{
-			const FRigBoneHierarchy& Hierarchy = ControlRig->GetBoneHierarchy();
-			if (Hierarchy.Bones.Num() > 0)
-			{
-				Code.Add(TEXT("FRigBoneHierarchy& Hierarchy = Rig->GetBoneHierarchy();"));
-			}
-			for (const FRigBone& Bone : Hierarchy.Bones)
-			{
-				FString ParentName = Bone.ParentName.IsNone() ? TEXT("NAME_None") : FString::Printf(TEXT("TEXT(\"%s\")"), *Bone.ParentName.ToString());
-				FTransform T = Bone.InitialTransform;
-				FString QuaternionString = FString::Printf(TEXT("FQuat(%.03f, %.03f, %.03f, %.03f)"), T.GetRotation().X, T.GetRotation().Y, T.GetRotation().Z, T.GetRotation().W);
-				FString TranslationString = FString::Printf(TEXT("FVector(%.03f, %.03f, %.03f)"), T.GetLocation().X, T.GetLocation().Y, T.GetLocation().Z);
-				FString ScaleString = FString::Printf(TEXT("FVector(%.03f, %.03f, %.03f)"), T.GetLocation().X, T.GetLocation().Y, T.GetLocation().Z);
-				FString TransformString = FString::Printf(TEXT("FTransform(%s, %s, %s)"), *QuaternionString, *TranslationString, *ScaleString);
-				Code.Add(FString::Printf(TEXT("Hierarchy.AddBone(TEXT(\"%s\"), %s, %s);"), *Bone.Name.ToString(), *ParentName, *TransformString));
-			}
-		}
-
-		// dump the nodes
-		for (UEdGraphNode* GraphNode : Graph->Nodes)
-		{
-			if (UControlRigGraphNode* RigNode = Cast<UControlRigGraphNode>(GraphNode))
-			{
-				FStructProperty* Property = RigNode->GetUnitProperty();
-				if (Property == nullptr)
-				{
-					return;
-				}
-				
-				Code.Add(FString::Printf(TEXT("FString %s = Rig->AddUnit(TEXT(\"%s\"));"), *Property->GetName(), *Property->Struct->GetName()));
-			}
-		}
-
-		// dump the pin links
-		for (UEdGraphNode* GraphNode : Graph->Nodes)
-		{
-			if (UControlRigGraphNode* RigNode = Cast<UControlRigGraphNode>(GraphNode))
-			{
-				for (UEdGraphPin* Pin : RigNode->Pins)
-				{
-					if (Pin->Direction != EEdGraphPinDirection::EGPD_Output)
-					{
-						continue;
-					}
-
-					for (UEdGraphPin* Linkedpin : Pin->LinkedTo)
-					{
-						if (UControlRigGraphNode* LinkedRigNode = Cast<UControlRigGraphNode>(Linkedpin->GetOwningNode()))
-						{
-							FString PropertyPathA = Pin->GetName();
-							FString PropertyPathB = Linkedpin->GetName();
-							FString NodeNameA, PinNameA, NodeNameB, PinNameB;
-							PropertyPathA.Split(TEXT("."), &NodeNameA, &PinNameA);
-							PropertyPathB.Split(TEXT("."), &NodeNameB, &PinNameB);
-
-							Code.Add(FString::Printf(TEXT("Rig->LinkProperties(%s + TEXT(\".%s\"), %s + TEXT(\".%s\"));"), *NodeNameA, *PinNameA, *NodeNameB, *PinNameB));
-						}
-					}
-				}
-			}
-		}
-
-		// set the pin values
-		for (UEdGraphNode* GraphNode : Graph->Nodes)
-		{
-			if (UControlRigGraphNode* RigNode = Cast<UControlRigGraphNode>(GraphNode))
-			{
-				for (UEdGraphPin* Pin : RigNode->Pins)
-				{
-					if (Pin->Direction != EEdGraphPinDirection::EGPD_Input)
-					{
-						continue;
-					}
-
-					if (Pin->ParentPin != nullptr)
-					{
-						continue;
-					}
-
-					if (Pin->LinkedTo.Num() > 0)
-					{
-						continue;
-					}
-
-					if (!Pin->DefaultValue.IsEmpty())
-					{
-						FString PropertyPath = Pin->GetName();
-						FString NodeName, PinName;
-						PropertyPath.Split(TEXT("."), &NodeName, &PinName);
-						Code.Add(FString::Printf(TEXT("Rig->SetPinDefault(%s + TEXT(\".%s\"), TEXT(\"%s\"));"), *NodeName, *PinName, *Pin->DefaultValue));
-					}
-				}
-			}
-		}
-		Code.Add(TEXT("Rig->Compile();"));
-
-		UE_LOG(LogControlRigEditor, Display, TEXT("\n%s\n"), *FString::Join(Code, TEXT("\n")));
-	}
-	*/
-}
-
 void FControlRigEditor::HandleOnViewportContextMenuDelegate(class FMenuBuilder& MenuBuilder)
 {
 	if (OnViewportContextMenuDelegate.IsBound())
diff --git a/Engine/Plugins/Experimental/ControlRig/Source/ControlRigEditor/Private/Editor/ControlRigEditor.h b/Engine/Plugins/Experimental/ControlRig/Source/ControlRigEditor/Private/Editor/ControlRigEditor.h
--- a/Engine/Plugins/Experimental/ControlRig/Source/ControlRigEditor/Private/Editor/ControlRigEditor.h
+++ b/Engine/Plugins/Experimental/ControlRig/Source/ControlRigEditor/Private/Editor/ControlRigEditor.h
@@ -155,12 +155,7 @@ public:
 
 	void OnCurveContainerChanged();
 
-	void OnRigElementAdded(FRigHierarchyContainer* Container, const FRigElementKey& InKey);
-	void OnRigElementRemoved(FRigHierarchyContainer* Container, const FRigElementKey& InKey, bool bForce = false);
-	void OnRigElementRenamed(FRigHierarchyContainer* Container, ERigElementType ElementType, const FName& InOldName, const FName& InNewName);
-	void OnRigElementReparented(FRigHierarchyContainer* Container, const FRigElementKey& InKey, const FName& InOldParentName, const FName& InNewParentName);
-	void OnRigElementSelected(FRigHierarchyContainer* Container, const FRigElementKey& InKey, bool bSelected);
-	void OnControlUISettingChanged(FRigHierarchyContainer* Container, const FRigElementKey& InKey);
+	void OnHierarchyModified(ERigHierarchyNotification InNotif, URigHierarchy* InHierarchy, const FRigBaseElement* InElement);
 
 	void OnGraphNodeDropToPerform(TSharedPtr<FGraphNodeDragDropOp> DragDropOp, UEdGraph* Graph, const FVector2D& NodePosition, const FVector2D& ScreenPosition);
 
@@ -297,7 +292,7 @@ private:
 
 	 void HandleMakeElementGetterSetter(ERigElementGetterSetterType Type, bool bIsGetter, TArray<FRigElementKey> Keys, UEdGraph* Graph, FVector2D NodePosition);
 
-	 void HandleOnControlModified(UControlRig* Subject, const FRigControl& Control, const FRigControlModifiedContext& Context);
+	 void HandleOnControlModified(UControlRig* Subject, FRigControlElement* ControlElement, const FRigControlModifiedContext& Context);
 
 	 void HandleRefreshEditorFromBlueprint(UControlRigBlueprint* InBlueprint);
 
@@ -355,6 +350,7 @@ protected:
 	virtual void NotifyPostChange(const FPropertyChangedEvent& PropertyChangedEvent, FProperty* PropertyThatChanged) override;
 	/** delegate for changing property */
 	virtual void OnFinishedChangingProperties(const FPropertyChangedEvent& PropertyChangedEvent) override;
+	void OnBlueprintPropertyChainEvent(FPropertyChangedChainEvent& PropertyChangedChainEvent);
 
 	URigVMGraph* GetFocusedModel() const;
 	URigVMController* GetFocusedController() const;
@@ -378,9 +374,6 @@ protected:
 	/** Once the log is collected update the graph */
 	void UpdateGraphCompilerErrors();
 
-	/** This can be used to enable dumping of a unit test */
-	void DumpUnitTestCode();
-
 	void OnAnimInitialized();
 
 	/** Are we currently in setup mode */
diff --git a/Engine/Plugins/Experimental/ControlRig/Source/ControlRigEditor/Private/Editor/ControlRigEditorEditMode.cpp b/Engine/Plugins/Experimental/ControlRig/Source/ControlRigEditor/Private/Editor/ControlRigEditorEditMode.cpp
--- a/Engine/Plugins/Experimental/ControlRig/Source/ControlRigEditor/Private/Editor/ControlRigEditorEditMode.cpp
+++ b/Engine/Plugins/Experimental/ControlRig/Source/ControlRigEditor/Private/Editor/ControlRigEditorEditMode.cpp
@@ -24,56 +24,59 @@ void FControlRigEditorEditMode::Render(const FSceneView* View, FViewport* Viewpo
 	{
 		if (UControlRig* ControlRig = GetControlRig(false))
 		{
-			const FRigBoneHierarchy& Hierarchy = ControlRig->GetBoneHierarchy();
-			const FRigBoneHierarchy* HierarchyForSelection = &Hierarchy;
+			URigHierarchy* Hierarchy = ControlRig->GetHierarchy();
+			URigHierarchy* HierarchyForSelection = Hierarchy;
 
 			if (UClass* ControlRigClass = ControlRig->GetClass())
 			{
 				if (UControlRigBlueprint* RigBlueprint = Cast<UControlRigBlueprint>(ControlRigClass->ClassGeneratedBy))
 				{
-					HierarchyForSelection = &RigBlueprint->HierarchyContainer.BoneHierarchy;
+					HierarchyForSelection = RigBlueprint->Hierarchy;
 				}
 			}
 
 			if (BoneDrawMode == EBoneDrawMode::SelectedAndParents)
 			{
-				if (BoneHasSelectedChild.Num() < Hierarchy.Num())
+				if (BoneHasSelectedChild.Num() != Hierarchy->Num())
 				{
-					BoneHasSelectedChild.SetNumZeroed(Hierarchy.Num());
+					BoneHasSelectedChild.SetNumZeroed(Hierarchy->Num());
 				}
-				for (const FRigBone& Bone : Hierarchy)
+
+				Hierarchy->ForEach<FRigBoneElement>([this, Hierarchy, HierarchyForSelection](FRigBoneElement* BoneElement) -> bool
 				{
-					BoneHasSelectedChild[Bone.Index] = false;
-					if (!HierarchyForSelection->IsSelected(Bone.Name))
+					BoneHasSelectedChild[BoneElement->GetIndex()] = false;
+					if (!HierarchyForSelection->IsSelected(BoneElement->GetKey()))
 					{
-						continue;
+						return true;
 					}
 
-					int32 ParentIndex = Bone.ParentIndex;
+					int32 ParentIndex = Hierarchy->GetFirstParent(BoneElement->GetIndex());
 					while (ParentIndex != INDEX_NONE)
 					{
 						BoneHasSelectedChild[ParentIndex] = true;
-						ParentIndex = Hierarchy[ParentIndex].ParentIndex;
+						ParentIndex = Hierarchy->GetFirstParent(ParentIndex);
 					}
-				}
+
+					return true;
+				});
 			}
 			
-			for (const FRigBone& Bone : Hierarchy)
+			Hierarchy->ForEach<FRigBoneElement>([this, PDI, Hierarchy, HierarchyForSelection, &BoneDrawMode](FRigBoneElement* BoneElement) -> bool
 			{
-				const int32 ParentIndex = Bone.ParentIndex;
-				const bool bSelected = HierarchyForSelection->IsSelected(Bone.Name);
+				const int32 ParentIndex = Hierarchy->GetFirstParent(BoneElement->GetIndex());
+				const bool bSelected = HierarchyForSelection->IsSelected(BoneElement->GetKey());
 
 				if (!bSelected)
 				{
 					if (BoneDrawMode == EBoneDrawMode::Selected)
 					{
-						continue;
+						return true;
 					}
 					else if (BoneDrawMode == EBoneDrawMode::SelectedAndParents)
 					{
-						if (!BoneHasSelectedChild[Bone.Index])
+						if (!BoneHasSelectedChild[BoneElement->GetIndex()])
 						{
-							continue;
+							return true;
 						}
 					}
 				}
@@ -81,15 +84,15 @@ void FControlRigEditorEditMode::Render(const FSceneView* View, FViewport* Viewpo
 				const FLinearColor LineColor = bSelected ? FLinearColor(1.0f, 0.34f, 0.0f, 1.0f) : FLinearColor::White;
 
 				FVector Start, End;
-				if (ParentIndex >= 0)
+				if (ParentIndex != INDEX_NONE)
 				{
-					Start = Hierarchy[ParentIndex].GlobalTransform.GetLocation();
-					End = Bone.GlobalTransform.GetLocation();
+					Start = Hierarchy->GetGlobalTransform(ParentIndex).GetLocation();
+					End = Hierarchy->GetGlobalTransform(BoneElement->GetIndex()).GetLocation();
 				}
 				else
 				{
 					Start = FVector::ZeroVector;
-					End = Bone.GlobalTransform.GetLocation();
+					End = Hierarchy->GetGlobalTransform(BoneElement->GetIndex()).GetLocation();
 				}
 
 				const float BoneLength = (End - Start).Size();
@@ -97,7 +100,8 @@ void FControlRigEditorEditMode::Render(const FSceneView* View, FViewport* Viewpo
 
 				//Render Sphere for bone end point and a cone between it and its parent.
 				SkeletalDebugRendering::DrawWireBone(PDI, Start, End, LineColor, SDPG_Foreground, Radius);
-			}
+				return true;
+			});
 		}
 	}
 }
diff --git a/Engine/Plugins/Experimental/ControlRig/Source/ControlRigEditor/Private/Editor/ControlRigSkeletalMeshComponent.cpp b/Engine/Plugins/Experimental/ControlRig/Source/ControlRigEditor/Private/Editor/ControlRigSkeletalMeshComponent.cpp
--- a/Engine/Plugins/Experimental/ControlRig/Source/ControlRigEditor/Private/Editor/ControlRigSkeletalMeshComponent.cpp
+++ b/Engine/Plugins/Experimental/ControlRig/Source/ControlRigEditor/Private/Editor/ControlRigSkeletalMeshComponent.cpp
@@ -47,7 +47,7 @@ void UControlRigSkeletalMeshComponent::RebuildDebugDrawSkeleton()
 		if (ControlRig)
 		{
 			// just copy it because this is not thread safe
-			const FRigBoneHierarchy& BaseHiearchy = ControlRig->GetBoneHierarchy();
+			URigHierarchy* Hierarchy = ControlRig->GetHierarchy();
 
 			DebugDrawSkeleton.Empty();
 			DebugDrawBones.Reset();
@@ -55,16 +55,24 @@ void UControlRigSkeletalMeshComponent::RebuildDebugDrawSkeleton()
 			// create ref modifier
 			FReferenceSkeletonModifier RefSkelModifier(DebugDrawSkeleton, nullptr);
 
-			for (int32 Index = 0; Index < BaseHiearchy.Num(); Index++)
+			Hierarchy->ForEach<FRigBoneElement>([&RefSkelModifier, this, Hierarchy](FRigBoneElement* BoneElement) -> bool
 			{
+				const int32 Index = BoneElement->GetIndex();
+				int32 ParentIndex = Hierarchy->GetFirstParent(Index);
+				if(ParentIndex != INDEX_NONE)
+				{
+					ParentIndex = Hierarchy->Get(ParentIndex)->GetSubIndex();
+				}
+					
 				FMeshBoneInfo NewMeshBoneInfo;
-				NewMeshBoneInfo.Name = BaseHiearchy.GetName(Index);
-				NewMeshBoneInfo.ParentIndex = BaseHiearchy[Index].ParentIndex;
+				NewMeshBoneInfo.Name = BoneElement->GetName();
+				NewMeshBoneInfo.ParentIndex = ParentIndex; 
 				// give ref pose here
-				RefSkelModifier.Add(NewMeshBoneInfo, BaseHiearchy.GetInitialGlobalTransform(Index));
+				RefSkelModifier.Add(NewMeshBoneInfo, Hierarchy->GetInitialGlobalTransform(Index));
 
-				DebugDrawBones.Add(Index);
-			}
+				DebugDrawBones.Add(DebugDrawBones.Num());
+				return true;
+			});
 		}
 	}
 }
@@ -79,8 +87,8 @@ FTransform UControlRigSkeletalMeshComponent::GetDrawTransform(int32 BoneIndex) c
 		if (ControlRig)
 		{
 			// just copy it because this is not thread safe
-			const FRigBoneHierarchy& BaseHiearchy = ControlRig->GetBoneHierarchy();
-			return BaseHiearchy.GetGlobalTransform(BoneIndex);
+			URigHierarchy* Hierarchy = ControlRig->GetHierarchy();
+			return Hierarchy->GetGlobalTransform(BoneIndex);
 		}
 	}
 
diff --git a/Engine/Plugins/Experimental/ControlRig/Source/ControlRigEditor/Private/Editor/SControlRigGizmoNameList.cpp b/Engine/Plugins/Experimental/ControlRig/Source/ControlRigEditor/Private/Editor/SControlRigGizmoNameList.cpp
--- a/Engine/Plugins/Experimental/ControlRig/Source/ControlRigEditor/Private/Editor/SControlRigGizmoNameList.cpp
+++ b/Engine/Plugins/Experimental/ControlRig/Source/ControlRigEditor/Private/Editor/SControlRigGizmoNameList.cpp
@@ -10,10 +10,10 @@
 #include "ControlRigBlueprint.h"
 #include "ControlRig.h"
 
-void SControlRigGizmoNameList::Construct(const FArguments& InArgs, FRigControl* InControl, UControlRigBlueprint* InBlueprint)
+void SControlRigGizmoNameList::Construct(const FArguments& InArgs, FRigControlElement* ControlElement, UControlRigBlueprint* InBlueprint)
 {
 	this->OnGetNameListContent = InArgs._OnGetNameListContent;
-	this->ControlKey = InControl->GetElementKey();
+	this->ControlKey = ControlElement->GetKey();
 	this->Blueprint = InBlueprint;
 
 	SBox::Construct(SBox::FArguments());
@@ -58,40 +58,41 @@ const TArray<TSharedPtr<FString>>& SControlRigGizmoNameList::GetNameList() const
 
 FText SControlRigGizmoNameList::GetNameListText() const
 {
-	int32 ControlIndex = Blueprint->HierarchyContainer.GetIndex(ControlKey);
+	int32 ControlIndex = Blueprint->Hierarchy->GetIndex(ControlKey);
 	if (ControlIndex != INDEX_NONE)
 	{
-		return FText::FromName(Blueprint->HierarchyContainer.ControlHierarchy[ControlIndex].GizmoName);
+		return FText::FromName(Blueprint->Hierarchy->GetChecked<FRigControlElement>(ControlIndex)->Settings.GizmoName);
 	}
 	return FText();
 }
 
 void SControlRigGizmoNameList::SetNameListText(const FText& NewTypeInValue, ETextCommit::Type /*CommitInfo*/)
 {
-	int32 ControlIndex = Blueprint->HierarchyContainer.GetIndex(ControlKey);
+	int32 ControlIndex = Blueprint->Hierarchy->GetIndex(ControlKey);
 	if (ControlIndex != INDEX_NONE)
 	{
 		FName NewName = *NewTypeInValue.ToString();
 
 		bool bIsOnInstance = false;
-		FRigHierarchyContainer* Container = &Blueprint->HierarchyContainer;
+		URigHierarchy* Hierarchy = Blueprint->Hierarchy;
 		UControlRig* DebuggedControlRig = Cast<UControlRig>(Blueprint->GetObjectBeingDebugged());
 
 		if (DebuggedControlRig)
 		{
 			if (!DebuggedControlRig->IsSetupModeEnabled())
 			{
-				Container = DebuggedControlRig->GetHierarchy();
+				Hierarchy = DebuggedControlRig->GetHierarchy();
 				bIsOnInstance = true;
 			}
 		}
 
-		FRigControl& Control = Container->ControlHierarchy[ControlIndex];
-		if (Control.GizmoName != NewName)
+		FRigControlElement* ControlElement = Hierarchy->Get<FRigControlElement>(ControlIndex);
+		if ((ControlElement != nullptr) && (ControlElement->Settings.GizmoName != NewName))
 		{
 			const FScopedTransaction Transaction(NSLOCTEXT("ControlRigEditor", "ChangeGizmoName", "Change Gizmo Name"));
 			Blueprint->Modify();
-			Control.GizmoName = NewName;
+			
+			ControlElement->Settings.GizmoName = NewName;
 
 			if (bIsOnInstance && DebuggedControlRig)
 			{
@@ -100,7 +101,7 @@ void SControlRigGizmoNameList::SetNameListText(const FText& NewTypeInValue, ETex
 
 			Blueprint->PropagatePropertyFromBPToInstances(ControlKey, FRigControl::StaticStruct()->FindPropertyByName(TEXT("GizmoName")));
 
-			Blueprint->HierarchyContainer.ControlHierarchy.OnControlUISettingsChanged.Broadcast(&Blueprint->HierarchyContainer, ControlKey);
+			Blueprint->Hierarchy->Notify(ERigHierarchyNotification::ControlSettingChanged, Blueprint->Hierarchy->Find(ControlElement->GetKey()));
 		}
 	}
 }
diff --git a/Engine/Plugins/Experimental/ControlRig/Source/ControlRigEditor/Private/Editor/SControlRigGizmoNameList.h b/Engine/Plugins/Experimental/ControlRig/Source/ControlRigEditor/Private/Editor/SControlRigGizmoNameList.h
--- a/Engine/Plugins/Experimental/ControlRig/Source/ControlRigEditor/Private/Editor/SControlRigGizmoNameList.h
+++ b/Engine/Plugins/Experimental/ControlRig/Source/ControlRigEditor/Private/Editor/SControlRigGizmoNameList.h
@@ -22,7 +22,7 @@ public:
 
 	SLATE_END_ARGS()
 
-	void Construct(const FArguments& InArgs, FRigControl* InControl, UControlRigBlueprint* InBlueprint);
+	void Construct(const FArguments& InArgs, FRigControlElement* ControlElement, UControlRigBlueprint* InBlueprint);
 
 protected:
 
diff --git a/Engine/Plugins/Experimental/ControlRig/Source/ControlRigEditor/Private/Editor/SControlRigValidationWidget.cpp b/Engine/Plugins/Experimental/ControlRig/Source/ControlRigEditor/Private/Editor/SControlRigValidationWidget.cpp
--- a/Engine/Plugins/Experimental/ControlRig/Source/ControlRigEditor/Private/Editor/SControlRigValidationWidget.cpp
+++ b/Engine/Plugins/Experimental/ControlRig/Source/ControlRigEditor/Private/Editor/SControlRigValidationWidget.cpp
@@ -251,10 +251,14 @@ void SControlRigValidationWidget::HandleMessageTokenClicked(const TSharedRef<cla
 					if (BlueprintClass)
 					{
 						UControlRigBlueprint* RigBlueprint = Cast<UControlRigBlueprint>(BlueprintClass->ClassGeneratedBy);
-						RigBlueprint->HierarchyContainer.ClearSelection();
-						RigBlueprint->HierarchyContainer.Select(Key, true);
-						ControlRig->GetHierarchy()->ClearSelection();
-						ControlRig->GetHierarchy()->Select(Key, true);
+						RigBlueprint->GetHierarchyController()->ClearSelection();
+						RigBlueprint->GetHierarchyController()->SelectElement(Key, true);
+
+						if(URigHierarchyController* Controller = ControlRig->GetHierarchy()->GetController())
+						{
+							Controller->ClearSelection();
+							Controller->SelectElement(Key, true);
+						}
 					}
 				}
 			}
diff --git a/Engine/Plugins/Experimental/ControlRig/Source/ControlRigEditor/Private/Editor/SRigCurveContainer.cpp b/Engine/Plugins/Experimental/ControlRig/Source/ControlRigEditor/Private/Editor/SRigCurveContainer.cpp
--- a/Engine/Plugins/Experimental/ControlRig/Source/ControlRigEditor/Private/Editor/SRigCurveContainer.cpp
+++ b/Engine/Plugins/Experimental/ControlRig/Source/ControlRigEditor/Private/Editor/SRigCurveContainer.cpp
@@ -144,10 +144,8 @@ void SRigCurveContainer::Construct(const FArguments& InArgs, TSharedRef<FControl
 	ControlRigEditor = InControlRigEditor;
 	ControlRigBlueprint = InControlRigEditor.Get().GetControlRigBlueprint();
 	bIsChangingRigHierarchy = false;
-	ControlRigBlueprint->HierarchyContainer.OnElementAdded.AddRaw(this, &SRigCurveContainer::OnRigElementAdded);
-	ControlRigBlueprint->HierarchyContainer.OnElementRemoved.AddRaw(this, &SRigCurveContainer::OnRigElementRemoved);
-	ControlRigBlueprint->HierarchyContainer.OnElementRenamed.AddRaw(this, &SRigCurveContainer::OnRigElementRenamed);
-	ControlRigBlueprint->HierarchyContainer.OnElementSelected.AddRaw(this, &SRigCurveContainer::OnRigElementSelected);
+
+	ControlRigBlueprint->Hierarchy->OnModified().AddRaw(this, &SRigCurveContainer::OnHierarchyModified);
 	ControlRigBlueprint->OnRefreshEditor().AddRaw(this, &SRigCurveContainer::HandleRefreshEditorFromBlueprint);
 
 	// Register and bind all our menu commands
@@ -208,10 +206,7 @@ SRigCurveContainer::~SRigCurveContainer()
 		ControlRigBlueprint = ControlRigEditor.Pin()->GetControlRigBlueprint();
 		if (ControlRigBlueprint.IsValid())
 		{
-			ControlRigBlueprint->HierarchyContainer.OnElementAdded.RemoveAll(this);
-			ControlRigBlueprint->HierarchyContainer.OnElementRemoved.RemoveAll(this);
-			ControlRigBlueprint->HierarchyContainer.OnElementRenamed.RemoveAll(this);
-			ControlRigBlueprint->HierarchyContainer.OnElementSelected.RemoveAll(this);
+			ControlRigBlueprint->Hierarchy->OnModified().RemoveAll(this);
 			ControlRigBlueprint->OnRefreshEditor().RemoveAll(this);
 		}
 	}
@@ -347,15 +342,19 @@ void SRigCurveContainer::CreateNewNameEntry(const FText& CommittedText, ETextCom
 	FSlateApplication::Get().DismissAllMenus();
 	if (!CommittedText.IsEmpty() && CommitType == ETextCommit::OnEnter)
 	{
-		FRigCurveContainer* Container = GetCurveContainer();
-		if (Container)
+		URigHierarchy* Hierarchy = GetHierarchy();
+		if (Hierarchy)
 		{
 			TGuardValue<bool> GuardReentry(bIsChangingRigHierarchy, true);
 
-			FName NewName = FName(*CommittedText.ToString());
-			Container->Add(NewName);
-			Container->ClearSelection();
-			Container->Select(NewName);
+			const FName NewName = FName(*CommittedText.ToString());
+
+			if(URigHierarchyController* Controller = Hierarchy->GetController())
+			{
+				const FRigElementKey CurveKey = Controller->AddCurve(NewName, 0.f, true);
+				Controller->ClearSelection();
+				Controller->SelectElement(CurveKey);
+			}
 		}
 
 		FSlateApplication::Get().DismissAllMenus();
@@ -365,25 +364,27 @@ void SRigCurveContainer::CreateNewNameEntry(const FText& CommittedText, ETextCom
 
 void SRigCurveContainer::CreateRigCurveList( const FString& SearchText )
 {
-	const FRigCurveContainer* Container = GetCurveContainer();
-	if (Container)
+	URigHierarchy* Hierarchy = GetHierarchy();
+	if (Hierarchy)
 	{
 		RigCurveList.Reset();
 
 		// Iterate through all curves..
-		for (const FRigCurve& Curve : (*Container))
+		Hierarchy->ForEach<FRigCurveElement>([this](FRigCurveElement* CurveElement) -> bool
 		{
-			FString CurveString = Curve.Name.ToString();
-
+			const FString CurveString = CurveElement->GetName().ToString();
+			
 			// See if we pass the search filter
-			if (!FilterText.IsEmpty() && !CurveString.Contains(*FilterText.ToString()))
-			{
-				continue;
-			}
+            if (!FilterText.IsEmpty() && !CurveString.Contains(*FilterText.ToString()))
+            {
+                return true;
+            }
 
-			TSharedRef<FDisplayedRigCurveInfo> NewItem = FDisplayedRigCurveInfo::Make(Curve.Name);
-			RigCurveList.Add(NewItem);
-		}
+            TSharedRef<FDisplayedRigCurveInfo> NewItem = FDisplayedRigCurveInfo::Make(CurveElement->GetName());
+            RigCurveList.Add(NewItem);
+
+			return true;
+		});
 
 		// Sort final list
 		struct FSortNamesAlphabetically
@@ -398,11 +399,15 @@ void SRigCurveContainer::CreateRigCurveList( const FString& SearchText )
 	}
 	RigCurveListView->RequestListRefresh();
 
-	if (Container)
+	if (Hierarchy)
 	{
-		for (const FName& SelectedCurve : Container->CurrentSelection())
+		TArray<FRigElementKey> SelectedCurveKeys = Hierarchy->GetSelectedKeys(ERigElementType::Curve);
+		for (const FRigElementKey& SelectedCurveKey : SelectedCurveKeys)
 		{
-			OnRigElementSelected(&ControlRigBlueprint->HierarchyContainer, FRigElementKey(SelectedCurve, ERigElementType::Curve), true);
+			if(FRigCurveElement* CurveElement = Hierarchy->Find<FRigCurveElement>(SelectedCurveKey))
+			{
+				OnHierarchyModified(ERigHierarchyNotification::ElementSelected, Hierarchy, CurveElement);
+			}
 		}
 	}
 
@@ -415,33 +420,39 @@ void SRigCurveContainer::RefreshCurveList()
 
 void SRigCurveContainer::OnNameCommitted(const FText& InNewName, ETextCommit::Type CommitType, FDisplayedRigCurveInfoPtr Item)
 {
-	FRigCurveContainer* Container = GetCurveContainer();
-	if (Container)
+	URigHierarchy* Hierarchy = GetHierarchy();
+	if (Hierarchy)
 	{
 		if (CommitType == ETextCommit::OnEnter)
 		{
-			FName NewName = FName(*InNewName.ToString());
-			FName OldName = Item->CurveName;
-			Container->Rename(OldName, NewName);
+			if(URigHierarchyController* Controller = Hierarchy->GetController())
+			{
+				FName NewName = FName(*InNewName.ToString());
+				FName OldName = Item->CurveName;
+				Controller->RenameElement(FRigElementKey(OldName, ERigElementType::Curve), NewName, true);
+			}
 		}
 	}
 }
 
 void SRigCurveContainer::OnDeleteNameClicked()
 {
-	FRigCurveContainer* Container = GetCurveContainer();
-	if (Container)
+	URigHierarchy* Hierarchy = GetHierarchy();
+	if (Hierarchy)
 	{
 		TGuardValue<bool> SuspendBlueprintNotifs(ControlRigBlueprint->bSuspendAllNotifications, true);
 
 		TArray< FDisplayedRigCurveInfoPtr > SelectedItems = RigCurveListView->GetSelectedItems();
 		for (auto Item : SelectedItems)
 		{
-			Container->Remove(Item->CurveName);
+			if(URigHierarchyController* Controller = Hierarchy->GetController())
+			{
+				Controller->RemoveElement(FRigElementKey(Item->CurveName, ERigElementType::Curve), true);
+			}
 		}
 	}
 
-	ControlRigBlueprint->PropagateHierarchyFromBPToInstances(true);
+	ControlRigBlueprint->PropagateHierarchyFromBPToInstances();
 	RefreshCurveList();
 }
 
@@ -452,21 +463,20 @@ bool SRigCurveContainer::CanDelete()
 
 void SRigCurveContainer::SetCurveValue(const FName& CurveName, float CurveValue)
 {
-	FRigCurveContainer* Container = GetCurveContainer();
-	if (Container)
+	URigHierarchy* Hierarchy = GetHierarchy();
+	if (Hierarchy)
 	{
-		Container->SetValue(CurveName, CurveValue);
+		Hierarchy->SetCurveValue(FRigElementKey(CurveName, ERigElementType::Curve), CurveValue, true);
 	}
 }
 
 float SRigCurveContainer::GetCurveValue(const FName& CurveName)
 {
-	FRigCurveContainer* Container = GetInstanceCurveContainer();
-	if (Container)
+	URigHierarchy* Hierarchy = GetInstanceHierarchy();
+	if (Hierarchy)
 	{
-		return Container->GetValue(CurveName);
+		return Hierarchy->GetCurveValue(FRigElementKey(CurveName, ERigElementType::Curve));
 	}
-
 	return 0.f;
 }
 
@@ -474,10 +484,13 @@ void SRigCurveContainer::ChangeCurveName(const FName& OldName, const FName& NewN
 {
 	TGuardValue<bool> GuardReentry(bIsChangingRigHierarchy, true);
 
-	FRigCurveContainer* Container = GetCurveContainer();
-	if (Container)
+	URigHierarchy* Hierarchy = GetHierarchy();
+	if (Hierarchy)
 	{
-		Container->Rename(OldName, NewName);
+		if(URigHierarchyController* Controller = Hierarchy->GetController())
+		{
+			Controller->RenameElement(FRigElementKey(OldName, ERigElementType::Curve), NewName, true);
+		}
 	}
 }
 
@@ -488,61 +501,52 @@ void SRigCurveContainer::OnSelectionChanged(FDisplayedRigCurveInfoPtr Selection,
 		return;
 	}
 
-	FRigCurveContainer* Container = GetCurveContainer();
-
-	if (Container)
+	URigHierarchy* Hierarchy = GetHierarchy();
+	if (Hierarchy)
 	{
+		URigHierarchyController* Controller = Hierarchy->GetController();
+		if(Controller == nullptr)
+		{
+			return;
+		}
 
 		FScopedTransaction ScopedTransaction(LOCTEXT("SelectCurveTransaction", "Select Curve"), !GIsTransacting);
 
 		TGuardValue<bool> GuardRigHierarchyChanges(bIsChangingRigHierarchy, true);
 
-		TArray<FName> OldSelection = Container->CurrentSelection();
-		TArray<FName> NewSelection;
+		TArray<FRigElementKey> OldSelection = Hierarchy->GetSelectedKeys(ERigElementType::Curve);
+		TArray<FRigElementKey> NewSelection;
 
 		TArray<FDisplayedRigCurveInfoPtr> SelectedItems = RigCurveListView->GetSelectedItems();
 		for (const FDisplayedRigCurveInfoPtr& SelectedItem : SelectedItems)
 		{
-			NewSelection.Add(SelectedItem->CurveName);
+			NewSelection.Add(FRigElementKey(SelectedItem->CurveName, ERigElementType::Curve));
 		}
 
-		for (const FName& PreviouslySelected : OldSelection)
+		for (const FRigElementKey& PreviouslySelected : OldSelection)
 		{
 			if (NewSelection.Contains(PreviouslySelected))
 			{
 				continue;
 			}
-			Container->Select(PreviouslySelected, false);
+			Controller->DeselectElement(PreviouslySelected);
 		}
 
-		for (const FName& NewlySelected : NewSelection)
+		for (const FRigElementKey& NewlySelected : NewSelection)
 		{
-			Container->Select(NewlySelected, true);
+			Controller->SelectElement(NewlySelected);
 		}
 	}
 }
 
-void SRigCurveContainer::OnRigElementAdded(FRigHierarchyContainer* Container, const FRigElementKey& InKey)
+void SRigCurveContainer::OnHierarchyModified(ERigHierarchyNotification InNotif, URigHierarchy* InHierarchy, const FRigBaseElement* InElement)
 {
-	if (bIsChangingRigHierarchy || InKey.Type != ERigElementType::Curve)
+	if (bIsChangingRigHierarchy)
 	{
 		return;
 	}
 
-	if (ControlRigBlueprint.IsValid())
-	{
-		if (ControlRigBlueprint->bSuspendAllNotifications)
-		{
-			return;
-		}
-	}
-
-	RefreshCurveList();
-}
-
-void SRigCurveContainer::OnRigElementRemoved(FRigHierarchyContainer* Container, const FRigElementKey& InKey)
-{
-	if (bIsChangingRigHierarchy || InKey.Type != ERigElementType::Curve)
+	if(InElement && !InElement->IsTypeOf(ERigElementType::Curve))
 	{
 		return;
 	}
@@ -555,50 +559,35 @@ void SRigCurveContainer::OnRigElementRemoved(FRigHierarchyContainer* Container,
 		}
 	}
 
-	RefreshCurveList();
-}
-
-void SRigCurveContainer::OnRigElementRenamed(FRigHierarchyContainer* Container, ERigElementType ElementType, const FName& InOldName, const FName& InNewName)
-{
-	if (bIsChangingRigHierarchy || ElementType != ERigElementType::Curve)
-	{
-		return;
-	}
-
-	if (ControlRigBlueprint.IsValid())
+	switch(InNotif)
 	{
-		if (ControlRigBlueprint->bSuspendAllNotifications)
+		case ERigHierarchyNotification::ElementAdded:
+		case ERigHierarchyNotification::ElementRemoved:
+		case ERigHierarchyNotification::ElementRenamed:
+		case ERigHierarchyNotification::HierarchyReset:
 		{
-			return;
+			RefreshCurveList();
+			break;
 		}
-	}
-
-	RefreshCurveList();
-}
-
-void SRigCurveContainer::OnRigElementSelected(FRigHierarchyContainer* Container, const FRigElementKey& InKey, bool bSelected)
-{
-	if (bIsChangingRigHierarchy || InKey.Type != ERigElementType::Curve)
-	{
-		return;
-	}
-
-	if (ControlRigBlueprint.IsValid())
-	{
-		if (ControlRigBlueprint->bSuspendAllNotifications)
+		case ERigHierarchyNotification::ElementSelected:
+    	case ERigHierarchyNotification::ElementDeselected:
 		{
-			return;
+			const bool bSelected = InNotif == ERigHierarchyNotification::ElementSelected;
+			for(const FDisplayedRigCurveInfoPtr& Item : RigCurveList)
+			{
+				if (Item->CurveName == InElement->GetName())
+				{
+					RigCurveListView->SetItemSelection(Item, bSelected);
+					break;
+				}
+			}
+			break;
 		}
-	}
-
-	for(const FDisplayedRigCurveInfoPtr& Item : RigCurveList)
-	{
-		if (Item->CurveName == InKey.Name)
+		default:
 		{
-			RigCurveListView->SetItemSelection(Item, bSelected);
 			break;
 		}
-	}
+    }
 }
 
 void SRigCurveContainer::HandleRefreshEditorFromBlueprint(UControlRigBlueprint* InBlueprint)
@@ -642,12 +631,12 @@ bool SRigCurveContainer::ShouldFilterOnImport(const FAssetData& AssetData) const
 
 void SRigCurveContainer::ImportCurve(const FAssetData& InAssetData)
 {
-	FRigCurveContainer* Container = GetCurveContainer();
-	if (Container)
+	URigHierarchy* Hierarchy = GetHierarchy();
+	if (Hierarchy)
 	{
 		TGuardValue<bool> SuspendBlueprintNotifs(ControlRigBlueprint->bSuspendAllNotifications, true);
 
-		const USkeleton* Skeleton = nullptr;
+		USkeleton* Skeleton = nullptr;
 		if (USkeletalMesh* Mesh = Cast<USkeletalMesh>(InAssetData.GetAsset()))
 		{
 			Skeleton = Mesh->GetSkeleton();
@@ -665,38 +654,41 @@ void SRigCurveContainer::ImportCurve(const FAssetData& InAssetData)
 
 			FScopedTransaction Transaction(LOCTEXT("CurveImport", "Import Curve"));
 			ControlRigBlueprint->Modify();
-			Container->ClearSelection();
-			Container->ImportCurvesFromSkeleton(Skeleton, NAME_None, false, false, false /* notify */);
+			
+			if(URigHierarchyController* Controller = Hierarchy->GetController())
+			{
+				Controller->ClearSelection();
+				Controller->ImportCurves(Skeleton, NAME_None, false, true);
+			}
 
 			FSlateApplication::Get().DismissAllMenus();
 		}
 	}
 
 	RefreshCurveList();
-	ControlRigBlueprint->PropagateHierarchyFromBPToInstances(true);
+	ControlRigBlueprint->PropagateHierarchyFromBPToInstances();
 }
 
-FRigCurveContainer* SRigCurveContainer::GetInstanceCurveContainer() const
+URigHierarchy* SRigCurveContainer::GetHierarchy() const
 {
-	if (ControlRigEditor.IsValid())
+	if (ControlRigBlueprint.IsValid())
 	{
-		UControlRig* ControlRig = ControlRigEditor.Pin()->GetInstanceRig();
-		if (ControlRig)
-		{
-			return &ControlRig->Hierarchy.CurveContainer;
-		}
+		return ControlRigBlueprint->Hierarchy;
 	}
-
 	return nullptr;
 }
 
-FRigCurveContainer* SRigCurveContainer::GetCurveContainer() const
+URigHierarchy* SRigCurveContainer::GetInstanceHierarchy() const
 {
-	if (ControlRigBlueprint.IsValid())
+	if (ControlRigEditor.IsValid())
 	{
-		return &ControlRigBlueprint->HierarchyContainer.CurveContainer;
+		UControlRig* ControlRig = ControlRigEditor.Pin()->GetInstanceRig();
+		if (ControlRig)
+		{
+			return ControlRig->GetHierarchy();
+		}
 	}
-
 	return nullptr;
 }
+
 #undef LOCTEXT_NAMESPACE
diff --git a/Engine/Plugins/Experimental/ControlRig/Source/ControlRigEditor/Private/Editor/SRigCurveContainer.h b/Engine/Plugins/Experimental/ControlRig/Source/ControlRigEditor/Private/Editor/SRigCurveContainer.h
--- a/Engine/Plugins/Experimental/ControlRig/Source/ControlRigEditor/Private/Editor/SRigCurveContainer.h
+++ b/Engine/Plugins/Experimental/ControlRig/Source/ControlRigEditor/Private/Editor/SRigCurveContainer.h
@@ -7,7 +7,7 @@
 #include "Widgets/Views/STableViewBase.h"
 #include "Widgets/Views/STableRow.h"
 #include "Widgets/Views/SListView.h"
-#include "Rigs/RigHierarchyContainer.h"
+#include "Rigs/RigHierarchy.h"
 
 //////////////////////////////////////////////////////////////////////////
 // FDisplayedRigCurveInfo
@@ -16,7 +16,6 @@ class SInlineEditableTextBlock;
 class FControlRigEditor;
 struct FAssetData;
 class UControlRigBlueprint;
-struct FRigCurveContainer;
 
 class FDisplayedRigCurveInfo
 {
@@ -234,10 +233,7 @@ private:
 	void OnSelectionChanged(FDisplayedRigCurveInfoPtr Selection, ESelectInfo::Type SelectInfo);
 
 	bool bIsChangingRigHierarchy;
-	void OnRigElementAdded(FRigHierarchyContainer* Container, const FRigElementKey& InKey);
-	void OnRigElementRemoved(FRigHierarchyContainer* Container, const FRigElementKey& InKey);
-	void OnRigElementRenamed(FRigHierarchyContainer* Container, ERigElementType ElementType, const FName& InOldName, const FName& InNewName);
-	void OnRigElementSelected(FRigHierarchyContainer* Container, const FRigElementKey& InKey, bool bSelected);
+	void OnHierarchyModified(ERigHierarchyNotification InNotif, URigHierarchy* InHierarchy, const FRigBaseElement* InElement);
 	void HandleRefreshEditorFromBlueprint(UControlRigBlueprint* InBlueprint);
 
 	// import curve part
@@ -264,8 +260,8 @@ private:
 	/** Commands that are bound to delegates*/
 	TSharedPtr<FUICommandList> UICommandList;
 
-	FRigCurveContainer* GetCurveContainer() const;
-	FRigCurveContainer* GetInstanceCurveContainer() const;
+	URigHierarchy* GetHierarchy() const;
+	URigHierarchy* GetInstanceHierarchy() const;
 
 	friend class SRigCurveListRow;
 	friend class SRigCurveTypeList;
