Control Rig: New rig hierarchy data structure refactoring

diff --git a/Engine/Plugins/Experimental/ControlRig/Source/ControlRigEditor/Private/Editor/SRigHierarchy.cpp b/Engine/Plugins/Experimental/ControlRig/Source/ControlRigEditor/Private/Editor/SRigHierarchy.cpp
--- a/Engine/Plugins/Experimental/ControlRig/Source/ControlRigEditor/Private/Editor/SRigHierarchy.cpp
+++ b/Engine/Plugins/Experimental/ControlRig/Source/ControlRigEditor/Private/Editor/SRigHierarchy.cpp
@@ -147,10 +147,10 @@ void SRigHierarchyItem::Construct(const FArguments& InArgs, TSharedPtr<FControlR
 		{
 			ERigBoneType BoneType = ERigBoneType::User;
 
-			int32 BoneIndex = InHierarchy->ControlRigBlueprint->HierarchyContainer.GetIndex(InRigTreeElement->Key);
-			if (BoneIndex != INDEX_NONE)
+			FRigBoneElement* BoneElement = InHierarchy->ControlRigBlueprint->Hierarchy->Find<FRigBoneElement>(InRigTreeElement->Key);
+			if(BoneElement)
 			{
-				BoneType = InHierarchy->ControlRigBlueprint->HierarchyContainer.BoneHierarchy[BoneIndex].Type;
+				BoneType = BoneElement->BoneType;
 			}
 
 			switch (BoneType)
@@ -264,11 +264,7 @@ SRigHierarchy::~SRigHierarchy()
 
 	if (ControlRigBlueprint.IsValid())
 	{
-		ControlRigBlueprint->HierarchyContainer.OnElementAdded.RemoveAll(this);
-		ControlRigBlueprint->HierarchyContainer.OnElementRemoved.RemoveAll(this);
-		ControlRigBlueprint->HierarchyContainer.OnElementRenamed.RemoveAll(this);
-		ControlRigBlueprint->HierarchyContainer.OnElementReparented.RemoveAll(this);
-		ControlRigBlueprint->HierarchyContainer.OnElementSelected.RemoveAll(this);
+		ControlRigBlueprint->Hierarchy->OnModified().RemoveAll(this);
 		ControlRigBlueprint->OnRefreshEditor().RemoveAll(this);
 	}
 }
@@ -278,13 +274,8 @@ void SRigHierarchy::Construct(const FArguments& InArgs, TSharedRef<FControlRigEd
 	ControlRigEditor = InControlRigEditor;
 
 	ControlRigBlueprint = ControlRigEditor.Pin()->GetControlRigBlueprint();
-	ControlRigBlueprint->HierarchyContainer.Initialize();
 
-	ControlRigBlueprint->HierarchyContainer.OnElementAdded.AddRaw(this, &SRigHierarchy::OnRigElementAdded);
-	ControlRigBlueprint->HierarchyContainer.OnElementRemoved.AddRaw(this, &SRigHierarchy::OnRigElementRemoved);
-	ControlRigBlueprint->HierarchyContainer.OnElementRenamed.AddRaw(this, &SRigHierarchy::OnRigElementRenamed);
-	ControlRigBlueprint->HierarchyContainer.OnElementReparented.AddRaw(this, &SRigHierarchy::OnRigElementReparented);
-	ControlRigBlueprint->HierarchyContainer.OnElementSelected.AddRaw(this, &SRigHierarchy::OnRigElementSelected);
+	ControlRigBlueprint->Hierarchy->OnModified().AddRaw(this, &SRigHierarchy::OnHierarchyModified);
 	ControlRigBlueprint->OnRefreshEditor().AddRaw(this, &SRigHierarchy::HandleRefreshEditorFromBlueprint);
 
 	// for deleting, renaming, dragging
@@ -574,9 +565,9 @@ FReply SRigHierarchy::OnKeyDown(const FGeometry& MyGeometry, const FKeyEvent& In
 
 EVisibility SRigHierarchy::IsToolbarVisible() const
 {
-	if (FRigHierarchyContainer* Container = GetHierarchyContainer())
+	if (URigHierarchy* Hierarchy = GetHierarchy())
 	{
-		if (Container->BoneHierarchy.Num() > 0)
+		if (Hierarchy->Num(ERigElementType::Bone) > 0)
 		{
 			return EVisibility::Collapsed;
 		}
@@ -586,9 +577,11 @@ EVisibility SRigHierarchy::IsToolbarVisible() const
 
 EVisibility SRigHierarchy::IsSearchbarVisible() const
 {
-	if (FRigHierarchyContainer* Container = GetHierarchyContainer())
+	if (URigHierarchy* Hierarchy = GetHierarchy())
 	{
-		if ((Container->BoneHierarchy.Num() + Container->SpaceHierarchy.Num() + Container->ControlHierarchy.Num()) > 0)
+		if ((Hierarchy->Num(ERigElementType::Bone) +
+			Hierarchy->Num(ERigElementType::Space) +
+			Hierarchy->Num(ERigElementType::Control)) > 0)
 		{
 			return EVisibility::Visible;
 		}
@@ -638,37 +631,48 @@ void SRigHierarchy::RefreshTreeView()
 
 	if (ControlRigBlueprint.IsValid())
 	{
-		FRigHierarchyContainer* Container = GetHierarchyContainer();
-		FRigBoneHierarchy& BoneHierarchy = Container->BoneHierarchy;
-		FRigControlHierarchy& ControlHierarchy = Container->ControlHierarchy;
-		FRigSpaceHierarchy& SpaceHierarchy = Container->SpaceHierarchy;
+		URigHierarchy* Hierarchy = GetHierarchy();
 
-		if (bShowBones)
+		Hierarchy->Traverse([&](FRigBaseElement* Element, bool& bContinue)
 		{
-			for (const FRigBone& Element : BoneHierarchy)
+			switch(Element->GetType())
 			{
-				if (bShowImportedBones || Element.Type != ERigBoneType::Imported)
+				case ERigElementType::Bone:
 				{
-					AddBoneElement(Element);
+					if(bShowBones)
+					{
+						FRigBoneElement* BoneElement = CastChecked<FRigBoneElement>(Element);
+						if (bShowImportedBones || BoneElement->BoneType != ERigBoneType::Imported)
+						{
+                            AddElement(BoneElement);
+                        }
+                    }
+					break;
 				}
-			}
-		}
-
-		if (bShowControls)
-		{
-			for (const FRigControl& Element : ControlHierarchy)
-			{
-				AddControlElement(Element);
-			}
-		}
-
-		if (bShowSpaces)
-		{
-			for (const FRigSpace& Element : SpaceHierarchy)
-			{
-				AddSpaceElement(Element);
-			}
-		}
+				case ERigElementType::Space:
+				{
+                    if(bShowSpaces)
+                    {
+                        FRigSpaceElement* SpaceElement = CastChecked<FRigSpaceElement>(Element);
+                        AddElement(SpaceElement);
+                    }
+                    break;
+                }
+				case ERigElementType::Control:
+				{
+                    if(bShowControls)
+                    {
+                        FRigControlElement* ControlElement = CastChecked<FRigControlElement>(Element);
+                        AddElement(ControlElement);
+                    }
+                    break;
+                }
+				default:
+				{
+					break;
+				}
+            }
+		});
 
 		for (const auto& Pair : ElementMap)
 		{
@@ -691,10 +695,10 @@ void SRigHierarchy::RefreshTreeView()
 
 		TreeView->RequestTreeRefresh();
 
-		TArray<FRigElementKey> Selection = Container->CurrentSelection();
+		TArray<FRigElementKey> Selection = Hierarchy->GetSelectedKeys();
 		for (const FRigElementKey& Key : Selection)
 		{
-			OnRigElementSelected(Container, Key, true);
+			OnHierarchyModified(ERigHierarchyNotification::ElementSelected, Hierarchy, Hierarchy->FindChecked(Key));
 		}
 	}
 }
@@ -751,28 +755,16 @@ void SRigHierarchy::OnSelectionChanged(TSharedPtr<FRigTreeElement> Selection, ES
 		return;
 	}
 
-	FRigHierarchyContainer* Hierarchy = GetHierarchyContainer();
-
+	URigHierarchy* Hierarchy = GetHierarchy();
 	if (Hierarchy)
 	{
+		URigHierarchyController* Controller = Hierarchy->GetController(true);
+		check(Controller);
+		
 		TGuardValue<bool> GuardRigHierarchyChanges(bIsChangingRigHierarchy, true);
 
-		TArray<FRigElementKey> OldSelection = Hierarchy->CurrentSelection();
-		TArray<FRigElementKey> NewSelection = GetSelectedKeys();
-
-		for (const FRigElementKey& PreviouslySelected : OldSelection)
-		{
-			if (NewSelection.Contains(PreviouslySelected))
-			{
-				continue;
-			}
-			Hierarchy->Select(PreviouslySelected, false);
-		}
-
-		for (const FRigElementKey& NewlySelected : NewSelection)
-		{
-			Hierarchy->Select(NewlySelected, true);
-		}
+		const TArray<FRigElementKey> NewSelection = GetSelectedKeys();
+		Controller->SetSelection(NewSelection);
 
 		if (NewSelection.Num() > 0)
 		{
@@ -854,126 +846,32 @@ void SRigHierarchy::AddElement(FRigElementKey InKey, FRigElementKey InParentKey,
 	}
 }
 
-void SRigHierarchy::AddSpacerElement()
-{
-	AddElement(FRigElementKey(), FRigElementKey());
-}
-
-void SRigHierarchy::AddBoneElement(FRigBone InBone, const bool bIgnoreTextFilter)
-{
-	if (ElementMap.Contains(InBone.GetElementKey()))
-	{
-		return;
-	}
-
-	const FRigHierarchyContainer* Container = GetHierarchyContainer();
-	const FRigBoneHierarchy& BoneHierarchy = Container->BoneHierarchy;
-	const FRigControlHierarchy& ControlHierarchy = Container->ControlHierarchy;
-	const FRigSpaceHierarchy& SpaceHierarchy = Container->SpaceHierarchy;
-
-	AddElement(InBone.GetElementKey(), FRigElementKey(), bIgnoreTextFilter);
-
-	if (ElementMap.Contains(InBone.GetElementKey()))
-	{
-		FRigElementKey ParentKey;
-		if (InBone.ParentIndex != INDEX_NONE)
-		{
-			AddBoneElement(BoneHierarchy[InBone.ParentIndex], !bFlattenHierarchyOnFilter);
-			ParentKey = BoneHierarchy[InBone.ParentIndex].GetElementKey();
-		}
-
-		if (ParentKey.IsValid())
-		{
-			ReparentElement(InBone.GetElementKey(), ParentKey);
-		}
-	}
-}
-
-void SRigHierarchy::AddControlElement(FRigControl InControl, const bool bIgnoreTextFilter)
+void SRigHierarchy::AddElement(FRigBaseElement* InElement, const bool bIgnoreTextFilter)
 {
-	if (ElementMap.Contains(InControl.GetElementKey()))
+	check(InElement);
+	
+	if (ElementMap.Contains(InElement->GetKey()))
 	{
 		return;
 	}
 
-	const FRigHierarchyContainer* Container = GetHierarchyContainer();
-	const FRigControlHierarchy& ControlHierarchy = Container->ControlHierarchy;
-	const FRigSpaceHierarchy& SpaceHierarchy = Container->SpaceHierarchy;
-
-	AddElement(InControl.GetElementKey(), FRigElementKey(), bIgnoreTextFilter);
+	const URigHierarchy* Hierarchy = GetHierarchy();
+	AddElement(InElement->GetKey(), FRigElementKey(), bIgnoreTextFilter);
 
-	if (ElementMap.Contains(InControl.GetElementKey()))
+	if (ElementMap.Contains(InElement->GetKey()))
 	{
-		FRigElementKey ParentKey;
-		if (InControl.SpaceIndex != INDEX_NONE)
-		{
-			AddSpaceElement(SpaceHierarchy[InControl.SpaceIndex], !bFlattenHierarchyOnFilter);
-			ParentKey = SpaceHierarchy[InControl.SpaceIndex].GetElementKey();
-		}
-		else if (InControl.ParentIndex != INDEX_NONE)
-		{
-			AddControlElement(ControlHierarchy[InControl.ParentIndex], !bFlattenHierarchyOnFilter);
-			ParentKey = ControlHierarchy[InControl.ParentIndex].GetElementKey();
-		}
-
+		FRigElementKey ParentKey = Hierarchy->GetFirstParent(InElement->GetKey());
 		if (ParentKey.IsValid())
 		{
-			ReparentElement(InControl.GetElementKey(), ParentKey);
+			check(ElementMap.Contains(ParentKey));
+			ReparentElement(InElement->GetKey(), ParentKey);
 		}
 	}
 }
 
-void SRigHierarchy::AddSpaceElement(FRigSpace InSpace, const bool bIgnoreTextFilter)
+void SRigHierarchy::AddSpacerElement()
 {
-	if (ElementMap.Contains(InSpace.GetElementKey()))
-	{
-		return;
-	}
-
-	const FRigHierarchyContainer* Container = GetHierarchyContainer();
-	const FRigBoneHierarchy& BoneHierarchy = Container->BoneHierarchy;
-	const FRigControlHierarchy& ControlHierarchy = Container->ControlHierarchy;
-	const FRigSpaceHierarchy& SpaceHierarchy = Container->SpaceHierarchy;
-
-	AddElement(InSpace.GetElementKey(), FRigElementKey(), bIgnoreTextFilter);
-
-	if (ElementMap.Contains(InSpace.GetElementKey()))
-	{
-		FRigElementKey ParentKey;
-		if (InSpace.ParentIndex != INDEX_NONE)
-		{
-			switch (InSpace.SpaceType)
-			{
-				case ERigSpaceType::Bone:
-				{
-					AddBoneElement(BoneHierarchy[InSpace.ParentIndex], !bFlattenHierarchyOnFilter);
-					ParentKey = BoneHierarchy[InSpace.ParentIndex].GetElementKey();
-					break;
-				}
-				case ERigSpaceType::Control:
-				{
-					AddControlElement(ControlHierarchy[InSpace.ParentIndex], !bFlattenHierarchyOnFilter);
-					ParentKey = ControlHierarchy[InSpace.ParentIndex].GetElementKey();
-					break;
-				}
-				case ERigSpaceType::Space:
-				{
-					AddSpaceElement(SpaceHierarchy[InSpace.ParentIndex], !bFlattenHierarchyOnFilter);
-					ParentKey = SpaceHierarchy[InSpace.ParentIndex].GetElementKey();
-					break;
-				}
-				default:
-				{
-					break;
-				}
-			}
-		}
-
-		if (ParentKey.IsValid())
-		{
-			ReparentElement(InSpace.GetElementKey(), ParentKey);
-		}
-	}
+	AddElement(FRigElementKey(), FRigElementKey());
 }
 
 void SRigHierarchy::ReparentElement(FRigElementKey InKey, FRigElementKey InParentKey)
@@ -1033,113 +931,76 @@ void SRigHierarchy::ReparentElement(FRigElementKey InKey, FRigElementKey InParen
 
 }
 
-void SRigHierarchy::OnRigElementAdded(FRigHierarchyContainer* Container, const FRigElementKey& InKey)
+void SRigHierarchy::OnHierarchyModified(ERigHierarchyNotification InNotif, URigHierarchy* InHierarchy, const FRigBaseElement* InElement)
 {
 	if (ControlRigBlueprint->bSuspendAllNotifications)
 	{
 		return;
 	}
 
-	if (bIsChangingRigHierarchy || InKey.Type == ERigElementType::Curve)
+	if (bIsChangingRigHierarchy)
 	{
 		return;
 	}
 
-	if (InKey.Type == ERigElementType::Control)
+	if(InElement)
 	{
-		if (Container->GetIndex(InKey) == INDEX_NONE)
+		if(InElement->IsTypeOf(ERigElementType::Curve))
 		{
 			return;
 		}
-	}
-
-	RefreshTreeView();
-}
-
-void SRigHierarchy::OnRigElementRemoved(FRigHierarchyContainer* Container, const FRigElementKey& InKey)
-{
-	if (bIsChangingRigHierarchy || InKey.Type == ERigElementType::Curve)
-	{
-		return;
-	}
-	if (ControlRigBlueprint->bSuspendAllNotifications)
-	{
-		return;
-	}
 
-	if (InKey.Type == ERigElementType::Control)
-	{
-		if (Container->GetIndex(InKey) == INDEX_NONE)
+		if(const FRigControlElement* ControlElement = Cast<FRigControlElement>(InElement))
 		{
-			return;
+			if(ControlElement->Settings.bIsTransientControl)
+			{
+				return;
+			}
 		}
-	}
-
-	RefreshTreeView();
-}
 
-void SRigHierarchy::OnRigElementRenamed(FRigHierarchyContainer* Container, ERigElementType ElementType, const FName& InOldName, const FName& InNewName)
-{
-	if (bIsChangingRigHierarchy || ElementType == ERigElementType::Curve)
-	{
-		return;
-	}
-	if (ControlRigBlueprint->bSuspendAllNotifications)
-	{
-		return;
-	}
-	RefreshTreeView();
-}
-
-void SRigHierarchy::OnRigElementReparented(FRigHierarchyContainer* Container, const FRigElementKey& InKey, const FName& InOldParentName, const FName& InNewParentName)
-{
-	if (bIsChangingRigHierarchy || InKey.Type == ERigElementType::Curve)
-	{
-		return;
-	}
-	if (ControlRigBlueprint->bSuspendAllNotifications)
-	{
-		return;
 	}
-	RefreshTreeView();
-}
 
-void SRigHierarchy::OnRigElementSelected(FRigHierarchyContainer* Container, const FRigElementKey& InKey, bool bSelected)
-{
-	if (bIsChangingRigHierarchy || InKey.Type == ERigElementType::Curve)
+	switch(InNotif)
 	{
-		return;
-	}
-	if (ControlRigBlueprint->bSuspendAllNotifications)
-	{
-		return;
-	}
-
-	if (InKey.Type == ERigElementType::Control)
-	{
-		if (Container->GetIndex(InKey) == INDEX_NONE)
+		case ERigHierarchyNotification::ElementAdded:
+		case ERigHierarchyNotification::ElementRemoved:
+		case ERigHierarchyNotification::ElementRenamed:
+		case ERigHierarchyNotification::ParentChanged:
+		case ERigHierarchyNotification::HierarchyReset:
 		{
-			return;
+			RefreshTreeView();
+			break;
 		}
-	}
-
-	for (int32 RootIndex = 0; RootIndex < RootElements.Num(); ++RootIndex)
-	{
-		TSharedPtr<FRigTreeElement> Found = FindElement(InKey, RootElements[RootIndex]);
-		if (Found.IsValid())
+		case ERigHierarchyNotification::ElementSelected:
+		case ERigHierarchyNotification::ElementDeselected:
 		{
-			TreeView->SetItemSelection(Found, bSelected, ESelectInfo::OnNavigation);
-			HandleFrameSelection();
+			check(InElement);
 
-			if (ControlRigEditor.IsValid())
+			const bool bSelected = (InNotif == ERigHierarchyNotification::ElementSelected); 
+				
+			for (int32 RootIndex = 0; RootIndex < RootElements.Num(); ++RootIndex)
 			{
-				if (ControlRigEditor.Pin()->GetEventQueue() == EControlRigEditorEventQueue::Setup)
+				TSharedPtr<FRigTreeElement> Found = FindElement(InElement->GetKey(), RootElements[RootIndex]);
+				if (Found.IsValid())
 				{
-					TGuardValue<bool> GuardRigHierarchyChanges(bIsChangingRigHierarchy, true);
-					HandleControlBoneOrSpaceTransform();
+					TreeView->SetItemSelection(Found, bSelected, ESelectInfo::OnNavigation);
+					HandleFrameSelection();
+
+					if (ControlRigEditor.IsValid())
+					{
+						if (ControlRigEditor.Pin()->GetEventQueue() == EControlRigEditorEventQueue::Setup)
+						{
+							TGuardValue<bool> GuardRigHierarchyChanges(bIsChangingRigHierarchy, true);
+							HandleControlBoneOrSpaceTransform();
+						}
+					}
 				}
 			}
 		}
+		default:
+		{
+			break;
+		}
 	}
 }
 
@@ -1198,7 +1059,7 @@ TSharedPtr< SWidget > SRigHierarchy::CreateContextMenu()
 
 void SRigHierarchy::OnItemClicked(TSharedPtr<FRigTreeElement> InItem)
 {
-	FRigHierarchyContainer* Hierarchy = GetHierarchyContainer();
+	URigHierarchy* Hierarchy = GetHierarchy();
 	if (Hierarchy->IsSelected(InItem->Key))
 	{
 		if (ControlRigEditor.IsValid())
@@ -1208,10 +1069,12 @@ void SRigHierarchy::OnItemClicked(TSharedPtr<FRigTreeElement> InItem)
 
 		if (InItem->Key.Type == ERigElementType::Bone)
 		{
-			const FRigBone& Bone = Hierarchy->BoneHierarchy[InItem->Key.Name];
-			if (Bone.Type == ERigBoneType::Imported)
+			if(FRigBoneElement* BoneElement = Hierarchy->Find<FRigBoneElement>(InItem->Key))
 			{
-				return;
+				if (BoneElement->BoneType == ERigBoneType::Imported)
+				{
+					return;
+				}
 			}
 		}
 
@@ -1369,14 +1232,13 @@ void SRigHierarchy::RefreshHierarchy(const FAssetData& InAssetData)
 	}
 	TGuardValue<bool> GuardRigHierarchyChanges(bIsChangingRigHierarchy, true);
 
-	FRigHierarchyContainer* Hierarchy = GetHierarchyContainer();
+	URigHierarchy* Hierarchy = GetHierarchy();
 	USkeletalMesh* Mesh = Cast<USkeletalMesh>(InAssetData.GetAsset());
 	if (Mesh && Hierarchy)
 	{
 		TGuardValue<bool> SuspendBlueprintNotifs(ControlRigBlueprint->bSuspendAllNotifications, true);
 
 		FScopedTransaction Transaction(LOCTEXT("HierarchyRefresh", "Refresh Transform"));
-		ControlRigBlueprint->Modify();
 
 		// don't select bone if we are in setup mode.
 		// we do this to avoid the editmode / viewport gizmos to refresh recursively,
@@ -1388,10 +1250,14 @@ void SRigHierarchy::RefreshHierarchy(const FAssetData& InAssetData)
 		}
 
 		const FReferenceSkeleton& RefSkeleton = Mesh->GetRefSkeleton();
-		Hierarchy->BoneHierarchy.ImportSkeleton(RefSkeleton, NAME_None, true, true, bSelectBones, false /* notify */);
+
+		URigHierarchyController* Controller = Hierarchy->GetController(true);
+		check(Controller);
+		
+		Controller->ImportBones(RefSkeleton, NAME_None, true, true, bSelectBones, true);
 	}
 
-	ControlRigBlueprint->PropagateHierarchyFromBPToInstances(true);
+	ControlRigBlueprint->PropagateHierarchyFromBPToInstances();
 	ControlRigEditor.Pin()->OnHierarchyChanged();
 	ControlRigBlueprint->BroadcastRefreshEditor();
 	RefreshTreeView();
@@ -1433,14 +1299,13 @@ void SRigHierarchy::ImportHierarchy(const FAssetData& InAssetData)
 	}
 	TGuardValue<bool> GuardRigHierarchyChanges(bIsChangingRigHierarchy, true);
 
-	FRigHierarchyContainer* Hierarchy = GetHierarchyContainer();
+	URigHierarchy* Hierarchy = GetHierarchy();
 	USkeletalMesh* Mesh = Cast<USkeletalMesh>(InAssetData.GetAsset());
 	if (Mesh && Hierarchy)
 	{
 		TGuardValue<bool> SuspendBlueprintNotifs(ControlRigBlueprint->bSuspendAllNotifications, true);
 
 		FScopedTransaction Transaction(LOCTEXT("HierarchyImport", "Import Hierarchy"));
-		ControlRigBlueprint->Modify();
 
 		// don't select bone if we are in setup mode.
 		// we do this to avoid the editmode / viewport gizmos to refresh recursively,
@@ -1452,14 +1317,18 @@ void SRigHierarchy::ImportHierarchy(const FAssetData& InAssetData)
 		}
 
 		const FReferenceSkeleton& RefSkeleton = Mesh->GetRefSkeleton();
-		Hierarchy->BoneHierarchy.ImportSkeleton(RefSkeleton, NAME_None, false, false, bSelectBones, false /* notify */);
-		Hierarchy->CurveContainer.ImportCurvesFromSkeleton(Mesh->GetSkeleton(), NAME_None, true, false, false /* notify */);
+
+		URigHierarchyController* Controller = Hierarchy->GetController(true);
+		check(Controller);
+
+		Controller->ImportBones(RefSkeleton, NAME_None, false, false, bSelectBones, true);
+		Controller->ImportCurves(Mesh->GetSkeleton(), NAME_None, true, true);
 
 		ControlRigBlueprint->SourceHierarchyImport = Mesh->GetSkeleton();
 		ControlRigBlueprint->SourceCurveImport = Mesh->GetSkeleton();
 	}
 
-	ControlRigBlueprint->PropagateHierarchyFromBPToInstances(true);
+	ControlRigBlueprint->PropagateHierarchyFromBPToInstances();
 	ControlRigEditor.Pin()->OnHierarchyChanged();
 	ControlRigBlueprint->BroadcastRefreshEditor();
 	RefreshTreeView();
@@ -1538,14 +1407,13 @@ bool SRigHierarchy::IsControlOrSpaceSelected() const
 
 void SRigHierarchy::HandleDeleteItem()
 {
-	FRigHierarchyContainer* Hierarchy = GetHierarchyContainer();
+	URigHierarchy* Hierarchy = GetHierarchy();
  	if (Hierarchy)
  	{
 		TArray<FRigElementKey> RemovedItems;
 
 		ClearDetailPanel();
 		FScopedTransaction Transaction(LOCTEXT("HierarchyTreeDeleteSelected", "Delete selected items from hierarchy"));
-		ControlRigBlueprint->Modify();
 
 		// clear detail view display
 		ControlRigEditor.Pin()->ClearDetailObject();
@@ -1553,77 +1421,48 @@ void SRigHierarchy::HandleDeleteItem()
 		bool bConfirmedByUser = false;
 		bool bDeleteImportedBones = false;
 
+ 		URigHierarchyController* Controller = Hierarchy->GetController(true);
+ 		check(Controller);
+
  		TArray<FRigElementKey> SelectedKeys = GetSelectedKeys();
 		for (const FRigElementKey& SelectedKey : SelectedKeys)
 		{
 			TGuardValue<bool> GuardRigHierarchyChanges(bIsChangingRigHierarchy, true);
 			TGuardValue<bool> SuspendBlueprintNotifs(ControlRigBlueprint->bSuspendAllNotifications, true);
 
-			switch (SelectedKey.Type)
+			if (SelectedKey.Type == ERigElementType::Bone)
 			{
-				case ERigElementType::Bone:
+				if (FRigBoneElement* BoneElement = Hierarchy->Find<FRigBoneElement>(SelectedKey))
 				{
-					int32 BoneIndex = Hierarchy->BoneHierarchy.GetIndex(SelectedKey.Name);
-					if (BoneIndex != INDEX_NONE)
+					if (BoneElement->BoneType == ERigBoneType::Imported && BoneElement->ParentElement != nullptr)
 					{
-						const FRigBone& Bone = Hierarchy->BoneHierarchy[BoneIndex];
-						if (Bone.Type == ERigBoneType::Imported && Bone.ParentIndex != INDEX_NONE)
+						if (!bConfirmedByUser)
 						{
-							if (!bConfirmedByUser)
-							{
-								FText ConfirmDelete = LOCTEXT("ConfirmDeleteBoneHierarchy", "Deleting imported(white) bones can cause issues with animation - are you sure ?");
+							FText ConfirmDelete = LOCTEXT("ConfirmDeleteBoneHierarchy", "Deleting imported(white) bones can cause issues with animation - are you sure ?");
 
-								FSuppressableWarningDialog::FSetupInfo Info(ConfirmDelete, LOCTEXT("DeleteImportedBone", "Delete Imported Bone"), "DeleteImportedBoneHierarchy_Warning");
-								Info.ConfirmText = LOCTEXT("DeleteImportedBoneHierarchy_Yes", "Yes");
-								Info.CancelText = LOCTEXT("DeleteImportedBoneHierarchy_No", "No");
+							FSuppressableWarningDialog::FSetupInfo Info(ConfirmDelete, LOCTEXT("DeleteImportedBone", "Delete Imported Bone"), "DeleteImportedBoneHierarchy_Warning");
+							Info.ConfirmText = LOCTEXT("DeleteImportedBoneHierarchy_Yes", "Yes");
+							Info.CancelText = LOCTEXT("DeleteImportedBoneHierarchy_No", "No");
 
-								FSuppressableWarningDialog DeleteImportedBonesInHierarchy(Info);
-								bDeleteImportedBones = DeleteImportedBonesInHierarchy.ShowModal() != FSuppressableWarningDialog::Cancel;
-								bConfirmedByUser = true;
-							}
+							FSuppressableWarningDialog DeleteImportedBonesInHierarchy(Info);
+							bDeleteImportedBones = DeleteImportedBonesInHierarchy.ShowModal() != FSuppressableWarningDialog::Cancel;
+							bConfirmedByUser = true;
+						}
 
-							if (!bDeleteImportedBones)
-							{
-								break;
-							}
+						if (!bDeleteImportedBones)
+						{
+							break;
 						}
-						Hierarchy->BoneHierarchy.Remove(SelectedKey.Name);
-					}
-					RemovedItems.Add(SelectedKey);
-					break;
-				}
-				case ERigElementType::Control:
-				{
-					if (Hierarchy->ControlHierarchy.GetIndex(SelectedKey.Name) != INDEX_NONE)
-					{
-						Hierarchy->ControlHierarchy.Remove(SelectedKey.Name);
-					}
-					RemovedItems.Add(SelectedKey);
-					break;
-				}
-				case ERigElementType::Space:
-				{
-					if (Hierarchy->SpaceHierarchy.GetIndex(SelectedKey.Name) != INDEX_NONE)
-					{
-						Hierarchy->SpaceHierarchy.Remove(SelectedKey.Name);
 					}
-					RemovedItems.Add(SelectedKey);
-					break;
-				}
-				default:
-				{
-					break;
 				}
-			} 
-		}
+			}
 
-		for (const FRigElementKey& RemovedItem : RemovedItems)
-		{
-			Hierarchy->HandleOnElementRemoved(Hierarchy, RemovedItem);
+			Controller->RemoveElement(SelectedKey, true);
+			RemovedItems.Add(SelectedKey);
 		}
 	}
 
-	ControlRigBlueprint->PropagateHierarchyFromBPToInstances(true);
+	ControlRigBlueprint->PropagateHierarchyFromBPToInstances();
 	ControlRigEditor.Pin()->OnHierarchyChanged();
 	RefreshTreeView();
 	FSlateApplication::Get().DismissAllMenus();
@@ -1637,15 +1476,18 @@ bool SRigHierarchy::CanDeleteItem() const
 /** Delete Item */
 void SRigHierarchy::HandleNewItem(ERigElementType InElementType)
 {
-	FRigHierarchyContainer* Hierarchy = GetHierarchyContainer();
+	URigHierarchy* Hierarchy = GetHierarchy();
 	if (Hierarchy)
 	{
 		// unselect current selected item
 		ClearDetailPanel();
 
+		URigHierarchyController* Controller = Hierarchy->GetController(true);
+		check(Controller);
+
 		FScopedTransaction Transaction(LOCTEXT("HierarchyTreeAdded", "Add new item to hierarchy"));
-		ControlRigBlueprint->Modify();
 
+		FRigElementKey NewItemKey;
 		FRigElementKey ParentKey;
 		FTransform ParentTransform = FTransform::Identity;
 
@@ -1664,66 +1506,20 @@ void SRigHierarchy::HandleNewItem(ERigElementType InElementType)
 			{
 				case ERigElementType::Bone:
 				{
-					FName ParentName = ParentKey.Type == ERigElementType::Bone ? ParentKey.Name : NAME_None;
-					ParentTransform = ParentName == NAME_None ? FTransform::Identity : ParentTransform;
-					Hierarchy->BoneHierarchy.Add(NewElementName, ParentName, ERigBoneType::User, ParentTransform);
+					NewItemKey = Controller->AddBone(NewElementName, ParentKey, ParentTransform, true, ERigBoneType::User, true);
 					break;
 				}
 				case ERigElementType::Control:
 				{
-					if (ParentKey.Type == ERigElementType::Bone)
-					{
-						for (const FRigSpace& ExistingSpace : Hierarchy->SpaceHierarchy)
-						{
-							if (ExistingSpace.SpaceType == ERigSpaceType::Bone && ExistingSpace.ParentName == ParentKey.Name)
-							{
-								ParentKey = ExistingSpace.GetElementKey();
-								break;
-							}
-						}
-
-						if (ParentKey.Type != ERigElementType::Space)
-						{
-							FString SpaceName = FString::Printf(TEXT("%sSpace"), *NewElementName.ToString());
-							FRigSpace& NewSpace = Hierarchy->SpaceHierarchy.Add(*SpaceName, ERigSpaceType::Bone, ParentKey.Name);
-							ParentKey = NewSpace.GetElementKey();
-						}
-					}
-
-					if(!ParentKey)
-					{
-						Hierarchy->ControlHierarchy.Add(NewElementName, ERigControlType::EulerTransform);
-					}
-					else if (ParentKey.Type == ERigElementType::Space)
-					{
-						Hierarchy->ControlHierarchy.Add(NewElementName, ERigControlType::EulerTransform, NAME_None, ParentKey.Name);
-					}
-					else if (ParentKey.Type == ERigElementType::Control)
-					{
-						Hierarchy->ControlHierarchy.Add(NewElementName, ERigControlType::EulerTransform, ParentKey.Name);
-					}
+					FRigControlSettings Settings;
+					Settings.ControlType = ERigControlType::EulerTransform;
+						
+					NewItemKey = Controller->AddControl(NewElementName, ParentKey, Settings, Settings.GetIdentityValue(), FTransform::Identity, FTransform::Identity, true);
 					break;
 				}
 				case ERigElementType::Space:
 				{
-					if (!ParentKey)
-					{
-						Hierarchy->SpaceHierarchy.Add(NewElementName, ERigSpaceType::Global);
-					}
-					else if (ParentKey.Type == ERigElementType::Bone)
-					{
-						FString SpaceName = FString::Printf(TEXT("%sSpace"), *NewElementName.ToString());
-						Hierarchy->SpaceHierarchy.Add(*SpaceName, ERigSpaceType::Bone, ParentKey.Name);
-					}
-					else if (ParentKey.Type == ERigElementType::Control)
-					{
-						FString SpaceName = FString::Printf(TEXT("%sSpace"), *NewElementName.ToString());
-						Hierarchy->SpaceHierarchy.Add(*SpaceName, ERigSpaceType::Control, ParentKey.Name);
-					}
-					else if (ParentKey.Type == ERigElementType::Space)
-					{
-						Hierarchy->SpaceHierarchy.Add(NewElementName, ERigSpaceType::Space, ParentKey.Name);
-					}
+					NewItemKey = Controller->AddSpace(NewElementName, ParentKey, ParentTransform, true, true);
 					break;
 				}
 				default:
@@ -1733,8 +1529,8 @@ void SRigHierarchy::HandleNewItem(ERigElementType InElementType)
 			}
 		}
 
-		Hierarchy->ClearSelection();
-		Hierarchy->Select(FRigElementKey(NewElementName, InElementType));
+		Controller->ClearSelection();
+		Controller->SelectElement(NewItemKey);
 	}
 
 	FSlateApplication::Get().DismissAllMenus();
@@ -1750,7 +1546,7 @@ bool SRigHierarchy::CanDuplicateItem() const
 /** Duplicate Item */
 void SRigHierarchy::HandleDuplicateItem()
 {
-	FRigHierarchyContainer* Hierarchy = GetHierarchyContainer();
+	URigHierarchy* Hierarchy = GetHierarchy();
 	if (Hierarchy)
 	{
 		ClearDetailPanel();
@@ -1759,12 +1555,14 @@ void SRigHierarchy::HandleDuplicateItem()
 			TGuardValue<bool> SuspendBlueprintNotifs(ControlRigBlueprint->bSuspendAllNotifications, true);
 
 			FScopedTransaction Transaction(LOCTEXT("HierarchyTreeDuplicateSelected", "Duplicate selected items from hierarchy"));
-			ControlRigBlueprint->Modify();
 
-			TArray<FRigElementKey> KeysToDuplicate = GetSelectedKeys();
-			Hierarchy->DuplicateItems(KeysToDuplicate, true);
+			URigHierarchyController* Controller = Hierarchy->GetController(true);
+			check(Controller);
+
+			const TArray<FRigElementKey> KeysToDuplicate = GetSelectedKeys();
+			Controller->DuplicateElements(KeysToDuplicate);
 		}
-		ControlRigBlueprint->PropagateHierarchyFromBPToInstances(true);
+		ControlRigBlueprint->PropagateHierarchyFromBPToInstances();
 	}
 
 	FSlateApplication::Get().DismissAllMenus();
@@ -1779,9 +1577,12 @@ void SRigHierarchy::HandleDuplicateItem()
 /** Mirror Item */
 void SRigHierarchy::HandleMirrorItem()
 {
-	FRigHierarchyContainer* Hierarchy = GetHierarchyContainer();
+	URigHierarchy* Hierarchy = GetHierarchy();
 	if (Hierarchy)
 	{
+		URigHierarchyController* Controller = Hierarchy->GetController(true);
+		check(Controller);
+		
 		FRigMirrorSettings Settings;
 		TSharedPtr<FStructOnScope> StructToDisplay = MakeShareable(new FStructOnScope(FRigMirrorSettings::StaticStruct(), (uint8*)&Settings));
 
@@ -1796,12 +1597,12 @@ void SRigHierarchy::HandleMirrorItem()
 			TGuardValue<bool> SuspendBlueprintNotifs(ControlRigBlueprint->bSuspendAllNotifications, true);
 
 			FScopedTransaction Transaction(LOCTEXT("HierarchyTreeMirrorSelected", "Mirror selected items from hierarchy"));
-			ControlRigBlueprint->Modify();
 
-			TArray<FRigElementKey> KeysToMirror = GetSelectedKeys();
-			Hierarchy->MirrorItems(KeysToMirror, Settings, true);
+			const TArray<FRigElementKey> KeysToMirror = GetSelectedKeys();
+			const TArray<FRigElementKey> KeysToDuplicate = GetSelectedKeys();
+			Controller->MirrorElements(KeysToDuplicate, Settings, true);
 		}
-		ControlRigBlueprint->PropagateHierarchyFromBPToInstances(true);
+		ControlRigBlueprint->PropagateHierarchyFromBPToInstances();
 	}
 
 	FSlateApplication::Get().DismissAllMenus();
@@ -1823,22 +1624,19 @@ void SRigHierarchy::HandleRenameItem()
 		return;
 	}
 
-	FRigHierarchyContainer* Hierarchy = GetHierarchyContainer();
+	URigHierarchy* Hierarchy = GetHierarchy();
 	if (Hierarchy)
 	{
 		FScopedTransaction Transaction(LOCTEXT("HierarchyTreeRenameSelected", "Rename selected item from hierarchy"));
-		ControlRigBlueprint->Modify();
 
 		TArray<TSharedPtr<FRigTreeElement>> SelectedItems = TreeView->GetSelectedItems();
 		if (SelectedItems.Num() == 1)
 		{
 			if (SelectedItems[0]->Key.Type == ERigElementType::Bone)
 			{
-				int32 BoneIndex = Hierarchy->GetIndex(SelectedItems[0]->Key);
-				if (BoneIndex != INDEX_NONE)
+				if(FRigBoneElement* BoneElement = Hierarchy->Find<FRigBoneElement>(SelectedItems[0]->Key))
 				{
-					const FRigBone& Bone = Hierarchy->BoneHierarchy[BoneIndex];
-					if (Bone.Type == ERigBoneType::Imported)
+					if (BoneElement->BoneType == ERigBoneType::Imported)
 					{
 						FText ConfirmRename = LOCTEXT("RenameDeleteBoneHierarchy", "Renaming imported(white) bones can cause issues with animation - are you sure ?");
 
@@ -1871,17 +1669,20 @@ bool SRigHierarchy::CanCopyOrPasteItems() const
 
 void SRigHierarchy::HandleCopyItems()
 {
-	if (FRigHierarchyContainer* Hierarchy = GetDebuggedHierarchyContainer())
+	if (URigHierarchy* Hierarchy = GetDebuggedHierarchy())
 	{
-		TArray<FRigElementKey> Selection = GetHierarchyContainer()->CurrentSelection();
-		FString Content = Hierarchy->ExportToText(Selection);
+		URigHierarchyController* Controller = Hierarchy->GetController(true);
+		check(Controller);
+
+		const TArray<FRigElementKey> Selection = GetHierarchy()->GetSelectedKeys();
+		const FString Content = Controller->ExportToText(Selection);
 		FPlatformApplicationMisc::ClipboardCopy(*Content);
 	}
 }
 
 void SRigHierarchy::HandlePasteItems()
 {
-	if (FRigHierarchyContainer* Hierarchy = GetHierarchyContainer())
+	if (URigHierarchy* Hierarchy = GetHierarchy())
 	{
 		TGuardValue<bool> GuardRigHierarchyChanges(bIsChangingRigHierarchy, true);
 		TGuardValue<bool> SuspendBlueprintNotifs(ControlRigBlueprint->bSuspendAllNotifications, true);
@@ -1890,12 +1691,14 @@ void SRigHierarchy::HandlePasteItems()
 		FPlatformApplicationMisc::ClipboardPaste(Content);
 
 		FScopedTransaction Transaction(LOCTEXT("HierarchyTreePaste", "Pasted rig elements."));
-		ControlRigBlueprint->Modify();
 
-		Hierarchy->ImportFromText(Content, ERigHierarchyImportMode::Append, true);
+		URigHierarchyController* Controller = Hierarchy->GetController(true);
+		check(Controller);
+
+		Controller->ImportFromText(Content, false, true, true);
 	}
 
-	//ControlRigBlueprint->PropagateHierarchyFromBPToInstances(true);
+	//ControlRigBlueprint->PropagateHierarchyFromBPToInstances();
 	ControlRigEditor.Pin()->OnHierarchyChanged();
 	{
 		TGuardValue<bool> GuardRigHierarchyChanges(bIsChangingRigHierarchy, true);
@@ -1906,100 +1709,90 @@ void SRigHierarchy::HandlePasteItems()
 
 void SRigHierarchy::HandlePasteLocalTransforms()
 {
-	if (FRigHierarchyContainer* Hierarchy = GetHierarchyContainer())
-	{
-		FString Content;
-		FPlatformApplicationMisc::ClipboardPaste(Content);
-
-		FScopedTransaction Transaction(LOCTEXT("HierarchyTreePasteLocal", "Pasted local transforms."));
-		ControlRigBlueprint->Modify();
-
-		Hierarchy->ImportFromText(Content, ERigHierarchyImportMode::ReplaceCurrentLocalTransform, true);
-
-		FRigHierarchyContainer* DebuggedHierarchy = GetDebuggedHierarchyContainer();
-		if(DebuggedHierarchy && DebuggedHierarchy != Hierarchy)
-		{
-			DebuggedHierarchy->ImportFromText(Content, ERigHierarchyImportMode::ReplaceCurrentLocalTransform, false, Hierarchy->CurrentSelection());
-		}
-	}
+	HandlePasteTransforms(ERigTransformType::CurrentLocal, true);
 }
 
 void SRigHierarchy::HandlePasteGlobalTransforms()
 {
-	if (FRigHierarchyContainer* Hierarchy = GetHierarchyContainer())
+	HandlePasteTransforms(ERigTransformType::CurrentGlobal, false);
+}
+
+void SRigHierarchy::HandlePasteTransforms(ERigTransformType::Type InTransformType, bool bAffectChildren)
+{
+	if (URigHierarchy* Hierarchy = GetHierarchy())
 	{
 		FString Content;
 		FPlatformApplicationMisc::ClipboardPaste(Content);
 
-		FScopedTransaction Transaction(LOCTEXT("HierarchyTreePasteGlobal", "Pasted global transforms."));
-		ControlRigBlueprint->Modify();
+		FScopedTransaction Transaction(LOCTEXT("HierarchyTreePaste", "Pasted transforms."));
 
-		Hierarchy->ImportFromText(Content, ERigHierarchyImportMode::ReplaceCurrentGlobalTransform, true);
+		FRigHierarchyCopyPasteContent Data;
+		FRigHierarchyCopyPasteContent::StaticStruct()->ImportText(*Content, &Data, nullptr, EPropertyPortFlags::PPF_None, nullptr, FRigHierarchyCopyPasteContent::StaticStruct()->GetName(), true);
 
-		FRigHierarchyContainer* DebuggedHierarchy = GetDebuggedHierarchyContainer();
-		if(DebuggedHierarchy && DebuggedHierarchy != Hierarchy)
+		URigHierarchy* DebuggedHierarchy = GetDebuggedHierarchy();
+
+		const TArray<FRigElementKey> CurrentSelection = Hierarchy->GetSelectedKeys();
+		const int32 Count = FMath::Min<int32>(CurrentSelection.Num(), Data.Elements.Num());
+		for(int32 Index = 0; Index < Count; Index++)
 		{
-			DebuggedHierarchy->ImportFromText(Content, ERigHierarchyImportMode::ReplaceCurrentGlobalTransform, false, Hierarchy->CurrentSelection());
+			const FRigHierarchyCopyPasteContentPerElement& PerElementData = Data.Elements[Index];
+			const FTransform Transform =  PerElementData.Pose.Get(InTransformType);
+
+			if(FRigTransformElement* TransformElement = Hierarchy->Find<FRigTransformElement>(CurrentSelection[Index]))
+			{
+				Hierarchy->SetTransform(TransformElement, Transform, InTransformType, bAffectChildren, true);
+			}
+			if(FRigBoneElement* BoneElement = Hierarchy->Find<FRigBoneElement>(CurrentSelection[Index]))
+			{
+				Hierarchy->SetTransform(BoneElement, Transform, ERigTransformType::MakeInitial(InTransformType), bAffectChildren, true);
+			}
+			
+            if(DebuggedHierarchy && DebuggedHierarchy != Hierarchy)
+            {
+            	if(FRigTransformElement* TransformElement = DebuggedHierarchy->Find<FRigTransformElement>(CurrentSelection[Index]))
+            	{
+            		DebuggedHierarchy->SetTransform(TransformElement, Transform, InTransformType, bAffectChildren, true);
+            	}
+            	if(FRigBoneElement* BoneElement = DebuggedHierarchy->Find<FRigBoneElement>(CurrentSelection[Index]))
+            	{
+            		DebuggedHierarchy->SetTransform(BoneElement, Transform, ERigTransformType::MakeInitial(InTransformType), bAffectChildren, true);
+            	}
+            }
 		}
 	}
 }
 
-FRigHierarchyContainer* SRigHierarchy::GetHierarchyContainer() const
+URigHierarchy* SRigHierarchy::GetHierarchy() const
 {
 	if (ControlRigBlueprint.IsValid())
 	{
-		return &ControlRigBlueprint->HierarchyContainer;
+		return ControlRigBlueprint->Hierarchy;
 	}
-
 	return nullptr;
 }
 
-FRigHierarchyContainer* SRigHierarchy::GetDebuggedHierarchyContainer() const
+URigHierarchy* SRigHierarchy::GetDebuggedHierarchy() const
 {
 	if (ControlRigBlueprint.IsValid())
 	{
 		if (UControlRig* DebuggedRig = Cast<UControlRig>(ControlRigBlueprint->GetObjectBeingDebugged()))
 		{
-			return (FRigHierarchyContainer*)DebuggedRig->GetHierarchy();
+			return DebuggedRig->GetHierarchy();
 		}
 	}
 	if (ControlRigEditor.IsValid())
 	{
 		if (UControlRig* CurrentRig = ControlRigEditor.Pin()->ControlRig)
 		{
-			return (FRigHierarchyContainer*)CurrentRig->GetHierarchy();
+			return CurrentRig->GetHierarchy();
 		}
 	}
-
-	return GetHierarchyContainer();
+	return GetHierarchy();
 }
 
 FName SRigHierarchy::CreateUniqueName(const FName& InBaseName, ERigElementType InElementType) const
 {
-	switch (InElementType)
-	{
-		case ERigElementType::Bone:
-		{
-			FRigBoneHierarchy& Hierarchy = GetHierarchyContainer()->BoneHierarchy;
-			return UtilityHelpers::CreateUniqueName(InBaseName, [&Hierarchy](const FName& CurName) { return Hierarchy.GetIndex(CurName) == INDEX_NONE; });
-		}
-		case ERigElementType::Control:
-		{
-			FRigControlHierarchy& Hierarchy = GetHierarchyContainer()->ControlHierarchy;
-			return UtilityHelpers::CreateUniqueName(InBaseName, [&Hierarchy](const FName& CurName) { return Hierarchy.GetIndex(CurName) == INDEX_NONE; });
-		}
-		case ERigElementType::Space:
-		{
-			FRigSpaceHierarchy& Hierarchy = GetHierarchyContainer()->SpaceHierarchy;
-			return UtilityHelpers::CreateUniqueName(InBaseName, [&Hierarchy](const FName& CurName) { return Hierarchy.GetIndex(CurName) == INDEX_NONE; });
-		}
-		default:
-		{
-			break;
-		}
-	}
-	ensure(false);
-	return NAME_None;
+	return GetHierarchy()->GetSafeNewName(InBaseName.ToString(), InElementType);
 }
 
 void SRigHierarchy::PostRedo(bool bSuccess) 
@@ -2041,8 +1834,8 @@ TOptional<EItemDropZone> SRigHierarchy::OnCanAcceptDrop(const FDragDropEvent& Dr
 	TSharedPtr<FRigElementHierarchyDragDropOp> RigDragDropOp = DragDropEvent.GetOperationAs<FRigElementHierarchyDragDropOp>();
 	if (RigDragDropOp.IsValid())
 	{
-		FRigHierarchyContainer* Container = GetHierarchyContainer();
-		if (Container)
+		URigHierarchy* Hierarchy = GetHierarchy();
+		if (Hierarchy)
 		{
 			for (const FRigElementKey& DraggedKey : RigDragDropOp->GetElements())
 			{
@@ -2051,11 +1844,7 @@ TOptional<EItemDropZone> SRigHierarchy::OnCanAcceptDrop(const FDragDropEvent& Dr
 					return ReturnDropZone;
 				}
 
-				if (Container->IsParentedTo(
-					TargetItem->Key.Type,
-					Container->GetIndex(TargetItem->Key),
-					DraggedKey.Type,
-					Container->GetIndex(DraggedKey)))
+				if(Hierarchy->IsParentedTo(TargetItem->Key, DraggedKey))
 				{
 					return ReturnDropZone;
 				}
@@ -2110,17 +1899,18 @@ FReply SRigHierarchy::OnAcceptDrop(const FDragDropEvent& DragDropEvent, EItemDro
 	TSharedPtr<FRigElementHierarchyDragDropOp> RigDragDropOp = DragDropEvent.GetOperationAs<FRigElementHierarchyDragDropOp>();
 	if (RigDragDropOp.IsValid())
 	{
-		FRigHierarchyContainer* Container = GetHierarchyContainer();
-		FRigHierarchyContainer* DebuggedContainer = GetDebuggedHierarchyContainer();
+		URigHierarchy* Hierarchy = GetHierarchy();
+		URigHierarchy* DebuggedHierarchy = GetDebuggedHierarchy();
+		URigHierarchyController* Controller = Hierarchy->GetController(true);
+		check(Controller);
 
-		if (Container && ControlRigBlueprint.IsValid())
+		if (Hierarchy && ControlRigBlueprint.IsValid())
 		{
 			TGuardValue<bool> GuardRigHierarchyChanges(bIsChangingRigHierarchy, true);
 			TGuardValue<bool> SuspendBlueprintNotifs(ControlRigBlueprint->bSuspendAllNotifications, true);
 			FScopedTransaction Transaction(LOCTEXT("HierarchyDragAndDrop", "Drag & Drop"));
-			ControlRigBlueprint->Modify();
 
-			FTransform TargetGlobalTransform = DebuggedContainer->GetGlobalTransform(TargetItem->Key);
+			FTransform TargetGlobalTransform = DebuggedHierarchy->GetGlobalTransform(TargetItem->Key);
 
 			for (const FRigElementKey& DraggedKey : RigDragDropOp->GetElements())
 			{
@@ -2131,11 +1921,7 @@ FReply SRigHierarchy::OnAcceptDrop(const FDragDropEvent& DragDropEvent, EItemDro
 
 				if (bReparentItems)
 				{
-					if (Container->IsParentedTo(
-						TargetItem->Key.Type,
-						Container->GetIndex(TargetItem->Key),
-						DraggedKey.Type,
-						Container->GetIndex(DraggedKey)))
+					if(Hierarchy->IsParentedTo(TargetItem->Key, DraggedKey))
 					{
 						return FReply::Unhandled();
 					}
@@ -2143,11 +1929,9 @@ FReply SRigHierarchy::OnAcceptDrop(const FDragDropEvent& DragDropEvent, EItemDro
 
 				if (DraggedKey.Type == ERigElementType::Bone)
 				{
-					int32 BoneIndex = Container->BoneHierarchy.GetIndex(DraggedKey.Name);
-					if (BoneIndex != INDEX_NONE)
+					if(FRigBoneElement* BoneElement = Hierarchy->Find<FRigBoneElement>(DraggedKey))
 					{
-						const FRigBone& Bone = Container->BoneHierarchy[BoneIndex];
-						if (Bone.Type == ERigBoneType::Imported && Bone.ParentIndex != INDEX_NONE)
+						if (BoneElement->BoneType == ERigBoneType::Imported && BoneElement->ParentElement != nullptr)
 						{
 							FText ConfirmText = bMatchTransforms ?
 								LOCTEXT("ConfirmMatchTransform", "Matching transforms of imported(white) bones can cause issues with animation - are you sure ?") :
@@ -2177,171 +1961,54 @@ FReply SRigHierarchy::OnAcceptDrop(const FDragDropEvent& DragDropEvent, EItemDro
 				{
 					if (DraggedKey.Type == ERigElementType::Control)
 					{
-						int32 ControlIndex = DebuggedContainer->ControlHierarchy.GetIndex(DraggedKey.Name);
+						int32 ControlIndex = DebuggedHierarchy->GetIndex(DraggedKey);
 						if (ControlIndex == INDEX_NONE)
 						{
 							continue;
 						}
 
-						FTransform ParentTransform = DebuggedContainer->ControlHierarchy.GetParentTransform(ControlIndex, false);
+						FTransform ParentTransform = DebuggedHierarchy->GetParentTransformByIndex(ControlIndex, false);
 						FTransform OffsetTransform = TargetGlobalTransform.GetRelativeTransform(ParentTransform);
 
-						Container->ControlHierarchy[ControlIndex].OffsetTransform = OffsetTransform;
-						Container->SetLocalTransform(DraggedKey, FTransform::Identity);
-						Container->SetInitialTransform(DraggedKey, FTransform::Identity);
-						DebuggedContainer->ControlHierarchy[ControlIndex].OffsetTransform = OffsetTransform;
-						DebuggedContainer->SetLocalTransform(DraggedKey, FTransform::Identity);
-						DebuggedContainer->SetInitialTransform(DraggedKey, FTransform::Identity);
-					}
-					else if (DraggedKey.Type == ERigElementType::Bone)
-					{
-						Container->SetGlobalTransform(DraggedKey, TargetGlobalTransform);
-						Container->BoneHierarchy.CopyGlobalToInitial(Container->BoneHierarchy.GetIndex(DraggedKey.Name));
-						DebuggedContainer->SetGlobalTransform(DraggedKey, TargetGlobalTransform);
-						DebuggedContainer->BoneHierarchy.CopyGlobalToInitial(DebuggedContainer->BoneHierarchy.GetIndex(DraggedKey.Name));
+						Hierarchy->SetControlOffsetTransformByIndex(ControlIndex, OffsetTransform, ERigTransformType::InitialLocal, true, true);
+						Hierarchy->SetLocalTransform(DraggedKey, FTransform::Identity, true, true, true);
+						Hierarchy->SetInitialLocalTransform(DraggedKey, FTransform::Identity, true, true);
+						DebuggedHierarchy->SetControlOffsetTransformByIndex(ControlIndex, OffsetTransform, ERigTransformType::InitialLocal, true, true);
+						DebuggedHierarchy->SetLocalTransform(DraggedKey, FTransform::Identity, true, true, true);
+						DebuggedHierarchy->SetInitialLocalTransform(DraggedKey, FTransform::Identity, true, true);
 					}
 					else
 					{
-						Container->SetInitialGlobalTransform(DraggedKey, TargetGlobalTransform);
-						Container->SetGlobalTransform(DraggedKey, TargetGlobalTransform);
-						DebuggedContainer->SetInitialGlobalTransform(DraggedKey, TargetGlobalTransform);
-						DebuggedContainer->SetGlobalTransform(DraggedKey, TargetGlobalTransform);
+						Hierarchy->SetInitialGlobalTransform(DraggedKey, TargetGlobalTransform, true, true);
+						Hierarchy->SetGlobalTransform(DraggedKey, TargetGlobalTransform, false, true, true);
+						DebuggedHierarchy->SetInitialGlobalTransform(DraggedKey, TargetGlobalTransform, true, true);
+						DebuggedHierarchy->SetGlobalTransform(DraggedKey, TargetGlobalTransform, false, true, true);
 					}
 					continue;
 				}
 
 				FRigElementKey ParentKey = TargetItem->Key;
 
-				FTransform InitialTransform = DebuggedContainer->GetInitialGlobalTransform(DraggedKey);
-				FTransform GlobalTransform = DebuggedContainer->GetGlobalTransform(DraggedKey);
+				FTransform InitialTransform = DebuggedHierarchy->GetInitialGlobalTransform(DraggedKey);
+				FTransform GlobalTransform = DebuggedHierarchy->GetGlobalTransform(DraggedKey);
 
-				switch (DraggedKey.Type)
+				if(ParentKey.IsValid())
 				{
-					case ERigElementType::Bone:
-					{
-						if (TargetItem->Key.Type == ERigElementType::Bone)
-						{
-							Container->BoneHierarchy.Reparent(DraggedKey.Name, TargetItem->Key.Name);
-						}
-						else if (!TargetItem->Key.IsValid())
-						{
-							Container->BoneHierarchy.Reparent(DraggedKey.Name, NAME_None);
-						}
-						break;
-					}
-					case ERigElementType::Control:
-					{
-						if (ParentKey.Type == ERigElementType::Bone)
-						{
-							for (const FRigSpace& ExistingSpace : Container->SpaceHierarchy)
-							{
-								if (ExistingSpace.SpaceType == ERigSpaceType::Bone && ExistingSpace.ParentName == ParentKey.Name)
-								{
-									ParentKey = ExistingSpace.GetElementKey();
-									break;
-								}
-							}
-
-							if (ParentKey.Type != ERigElementType::Space)
-							{
-								FString SpaceName = FString::Printf(TEXT("%sSpace"), *DraggedKey.Name.ToString());
-								FRigSpace& NewSpace = Container->SpaceHierarchy.Add(*SpaceName, ERigSpaceType::Bone, ParentKey.Name);
-								if (DebuggedContainer != Container)
-								{
-									DebuggedContainer->SpaceHierarchy.Add(*SpaceName, ERigSpaceType::Bone, ParentKey.Name);
-								}
-								ParentKey = NewSpace.GetElementKey();
-							}
-						}
-
-						if (ParentKey.Type == ERigElementType::Control)
-						{
-							Container->ControlHierarchy.SetSpace(DraggedKey.Name, NAME_None);
-							Container->ControlHierarchy.Reparent(DraggedKey.Name, ParentKey.Name);
-							if (DebuggedContainer != Container)
-							{
-								DebuggedContainer->ControlHierarchy.SetSpace(DraggedKey.Name, NAME_None);
-								DebuggedContainer->ControlHierarchy.Reparent(DraggedKey.Name, ParentKey.Name);
-							}
-						}
-						else if (ParentKey.Type == ERigElementType::Space)
-						{
-							Container->ControlHierarchy.Reparent(DraggedKey.Name, NAME_None);
-							Container->ControlHierarchy.SetSpace(DraggedKey.Name, ParentKey.Name);
-							if (DebuggedContainer != Container)
-							{
-								DebuggedContainer->ControlHierarchy.Reparent(DraggedKey.Name, NAME_None);
-								DebuggedContainer->ControlHierarchy.SetSpace(DraggedKey.Name, ParentKey.Name);
-							}
-						}
-						else if (!ParentKey.IsValid())
-						{
-							Container->ControlHierarchy.Reparent(DraggedKey.Name, NAME_None);
-							Container->ControlHierarchy.SetSpace(DraggedKey.Name, NAME_None);
-							if (DebuggedContainer != Container)
-							{
-								DebuggedContainer->ControlHierarchy.Reparent(DraggedKey.Name, NAME_None);
-								DebuggedContainer->ControlHierarchy.SetSpace(DraggedKey.Name, NAME_None);
-							}
-						}
-
-						break;
-					}
-					case ERigElementType::Space:
-					{
-						if (TargetItem->Key.Type == ERigElementType::Bone)
-						{
-							Container->SpaceHierarchy.Reparent(DraggedKey.Name, ERigSpaceType::Bone, TargetItem->Key.Name);
-							if (DebuggedContainer != Container)
-							{
-								DebuggedContainer->SpaceHierarchy.Reparent(DraggedKey.Name, ERigSpaceType::Bone, TargetItem->Key.Name);
-							}
-						}
-						else if (TargetItem->Key.Type == ERigElementType::Control)
-						{
-							Container->SpaceHierarchy.Reparent(DraggedKey.Name, ERigSpaceType::Control, TargetItem->Key.Name);
-							if (DebuggedContainer != Container)
-							{
-								DebuggedContainer->SpaceHierarchy.Reparent(DraggedKey.Name, ERigSpaceType::Control, TargetItem->Key.Name);
-							}
-						}
-						else if (TargetItem->Key.Type == ERigElementType::Space)
-						{
-							Container->SpaceHierarchy.Reparent(DraggedKey.Name, ERigSpaceType::Space, TargetItem->Key.Name);
-							if (DebuggedContainer != Container)
-							{
-								DebuggedContainer->SpaceHierarchy.Reparent(DraggedKey.Name, ERigSpaceType::Space, TargetItem->Key.Name);
-							}
-						}
-						else if (!ParentKey.IsValid())
-						{
-							Container->SpaceHierarchy.Reparent(DraggedKey.Name, ERigSpaceType::Global, NAME_None);
-							if (DebuggedContainer != Container)
-							{
-								DebuggedContainer->SpaceHierarchy.Reparent(DraggedKey.Name, ERigSpaceType::Global, NAME_None);
-							}
-						}
-						break;
-					}
-					case ERigElementType::Curve:
-					{
-						break;
-					}
-					default:
-					{
-						ensure(false);
-						break;
-					}
+					Controller->SetParent(DraggedKey, ParentKey, true, true);
+				}
+				else
+				{
+					Controller->RemoveAllParents(DraggedKey, true, true);
 				}
 
-				DebuggedContainer->SetInitialGlobalTransform(DraggedKey, InitialTransform);
-				DebuggedContainer->SetGlobalTransform(DraggedKey, GlobalTransform);
-				Container->SetInitialGlobalTransform(DraggedKey, InitialTransform);
-				Container->SetGlobalTransform(DraggedKey, GlobalTransform);
+				DebuggedHierarchy->SetInitialGlobalTransform(DraggedKey, InitialTransform, true, true);
+				DebuggedHierarchy->SetGlobalTransform(DraggedKey, GlobalTransform, false, true, true);
+				Hierarchy->SetInitialGlobalTransform(DraggedKey, InitialTransform, true, true);
+				Hierarchy->SetGlobalTransform(DraggedKey, GlobalTransform, false, true, true);
 			}
 		}
 
-		ControlRigBlueprint->PropagateHierarchyFromBPToInstances(true);
+		ControlRigBlueprint->PropagateHierarchyFromBPToInstances();
 
 		if(bReparentItems)
 		{
@@ -2364,46 +2031,27 @@ FName SRigHierarchy::RenameElement(const FRigElementKey& OldKey, const FString&
 	{
 		return NAME_None;
 	}
-	FString SanitizedNameStr = NewName;
-	FRigHierarchyContainer::SanitizeName(SanitizedNameStr);
-	FName SanitizedName = *SanitizedNameStr;
-	FName ResultingName = NAME_None;
 
 	// make sure there is no duplicate
 	if (ControlRigBlueprint.IsValid())
 	{
 		FScopedTransaction Transaction(LOCTEXT("HierarchyRename", "Rename Hierarchy Element"));
-		ControlRigBlueprint->Modify();
 
-		FRigHierarchyContainer* Container = GetHierarchyContainer();
-		switch (OldKey.Type)
-		{
-			case ERigElementType::Bone:
-			{
-				ResultingName = Container->BoneHierarchy.Rename(OldKey.Name, SanitizedName);
-				break;
-			}
-			case ERigElementType::Control:
-			{
-				ResultingName = Container->ControlHierarchy.Rename(OldKey.Name, SanitizedName);
-				break;
-			}
-			case ERigElementType::Space:
-			{
-				ResultingName = Container->SpaceHierarchy.Rename(OldKey.Name, SanitizedName);
-				break;
-			}
-			default:
-			{
-				return NAME_None;
-			}
-		}
+		URigHierarchy* Hierarchy = GetHierarchy();
+		URigHierarchyController* Controller = Hierarchy->GetController(true);
+		check(Controller);
 
-		ControlRigBlueprint->PropagateHierarchyFromBPToInstances(true, true);
+		FString SanitizedNameStr = NewName;
+		Hierarchy->SanitizeName(SanitizedNameStr);
+		const FName SanitizedName = *SanitizedNameStr;
+		FName ResultingName = NAME_None;
+
+		ResultingName = Controller->RenameElement(OldKey, SanitizedName, true).Name;
+		ControlRigBlueprint->PropagateHierarchyFromBPToInstances();
 		return ResultingName;
 	}
 
-	return ResultingName;
+	return NAME_None;
 }
 
 bool SRigHierarchy::OnVerifyNameChanged(const FRigElementKey& OldKey, const FString& NewName, FText& OutErrorMessage)
@@ -2416,10 +2064,10 @@ bool SRigHierarchy::OnVerifyNameChanged(const FRigElementKey& OldKey, const FStr
 	// make sure there is no duplicate
 	if (ControlRigBlueprint.IsValid())
 	{
-		FRigHierarchyContainer* Container = GetHierarchyContainer();
+		URigHierarchy* Hierarchy = GetHierarchy();
 
 		FString OutErrorString;
-		if (!Container->IsNameAvailable(NewName, OldKey.Type, &OutErrorString))
+		if (!Hierarchy->IsNameAvailable(NewName, OldKey.Type, &OutErrorString))
 		{
 			OutErrorMessage = FText::FromString(OutErrorString);
 			return false;
@@ -2434,110 +2082,21 @@ void SRigHierarchy::HandleResetTransform(bool bSelectionOnly)
 	{
 		if (UControlRigBlueprint* Blueprint = ControlRigEditor.Pin()->GetControlRigBlueprint())
 		{
-			if (FRigHierarchyContainer* DebuggedContainer = GetDebuggedHierarchyContainer())
+			if (URigHierarchy* DebuggedHierarchy = GetDebuggedHierarchy())
 			{
 				FScopedTransaction Transaction(LOCTEXT("HierarchyResetTransforms", "Reset Transforms"));
-				Blueprint->Modify();
 
 				TArray<FRigElementKey> KeysToReset = GetSelectedKeys();
 				if (!bSelectionOnly)
 				{
-					for (const FRigControl& Control : DebuggedContainer->ControlHierarchy)
-					{
-						KeysToReset.AddUnique(Control.GetElementKey());
-					}
+					KeysToReset = DebuggedHierarchy->GetAllKeys(true, ERigElementType::Control);
 				}
 
 				for (FRigElementKey Key : KeysToReset)
 				{
-					if (Key.Type == ERigElementType::Control)
-					{
-						FTransform Transform = GetHierarchyContainer()->ControlHierarchy.GetLocalTransform(Key.Name, ERigControlValueType::Initial);
-						GetHierarchyContainer()->SetLocalTransform(Key, Transform);
-						DebuggedContainer->SetLocalTransform(Key, Transform);
-					}
-
-					else if (Key.Type == ERigElementType::Space)
-					{
-						FTransform InitialTransform = GetHierarchyContainer()->GetInitialTransform(Key);
-						GetHierarchyContainer()->SetLocalTransform(Key, InitialTransform);
-						DebuggedContainer->SetLocalTransform(Key, InitialTransform);
-					}
-
-					else if (Key.Type == ERigElementType::Bone)
-					{
-						FTransform InitialTransform = GetHierarchyContainer()->GetInitialGlobalTransform(Key);
-						GetHierarchyContainer()->SetGlobalTransform(Key, InitialTransform);
-						DebuggedContainer->SetGlobalTransform(Key, InitialTransform);
-
-						Blueprint->RemoveTransientControl(Key);
-						if (ControlRigEditor.Pin()->PreviewInstance)
-						{
-							ControlRigEditor.Pin()->PreviewInstance->RemoveBoneModification(Key.Name);
-						}
-					}
-				}
-			}
-		}
-	}
-}
-
-void SRigHierarchy::HandleResetSpace()
-{
-	if (IsControlSelected())
-	{
-		if (UControlRigBlueprint* Blueprint = ControlRigEditor.Pin()->GetControlRigBlueprint())
-		{
-			if (FRigHierarchyContainer* DebuggedContainer = GetDebuggedHierarchyContainer())
-			{
-				FScopedTransaction Transaction(LOCTEXT("HierarchyResetSpace", "Reset Space"));
-				Blueprint->Modify();
-
-				TArray<FRigElementKey> SelectedKeys = GetSelectedKeys();
-				for (const FRigElementKey& SelectedKey : SelectedKeys)
-				{
-					if (SelectedKey.Type == ERigElementType::Control)
-					{
-						FRigControl& Control = DebuggedContainer->ControlHierarchy[SelectedKey.Name];
-						if (FRigElementKey SpaceKey = Control.GetSpaceElementKey())
-						{
-							FRigSpace& Space = DebuggedContainer->SpaceHierarchy[SpaceKey.Name];
-
-							FTransform Transform = DebuggedContainer->GetGlobalTransform(SelectedKey);
-							FTransform SpaceParentTransform = DebuggedContainer->GetGlobalTransform(Space.GetParentElementKey());
-							FTransform LocalTransform = Transform.GetRelativeTransform(SpaceParentTransform);
-
-							GetHierarchyContainer()->SetInitialTransform(SpaceKey, LocalTransform);
-							GetHierarchyContainer()->SetLocalTransform(SpaceKey, LocalTransform);
-							GetHierarchyContainer()->SetInitialTransform(SelectedKey, FTransform::Identity);
-							GetHierarchyContainer()->SetLocalTransform(SelectedKey, FTransform::Identity);
-
-							DebuggedContainer->SetInitialTransform(SpaceKey, LocalTransform);
-							DebuggedContainer->SetLocalTransform(SpaceKey, LocalTransform);
-							DebuggedContainer->SetInitialTransform(SelectedKey, FTransform::Identity);
-							DebuggedContainer->SetLocalTransform(SelectedKey, FTransform::Identity);
-						}
-						else
-						{
-							FString SpaceName = FString::Printf(TEXT("%sSpace"), *Control.Name.ToString());
-							if (FRigElementKey CurrentParentKey = Control.GetParentElementKey())
-							{
-								FRigSpace& NewSpace = GetHierarchyContainer()->SpaceHierarchy.Add(*SpaceName, ERigSpaceType::Control, CurrentParentKey.Name, Control.Value.Get<FTransform>());
-								SpaceKey = NewSpace.GetElementKey();
-							}
-							else
-							{
-								FRigSpace& NewSpace = GetHierarchyContainer()->SpaceHierarchy.Add(*SpaceName, ERigSpaceType::Global, NAME_None, Control.Value.Get<FTransform>());
-								SpaceKey = NewSpace.GetElementKey();
-							}
-
-							GetHierarchyContainer()->ControlHierarchy.Reparent(Control.Name, NAME_None);
-							GetHierarchyContainer()->ControlHierarchy.SetSpace(Control.Name, SpaceKey.Name);
-							GetHierarchyContainer()->ControlHierarchy.SetLocalTransform(Control.Name, FTransform::Identity, ERigControlValueType::Initial);
-							GetHierarchyContainer()->ControlHierarchy.SetLocalTransform(Control.Name, FTransform::Identity);
-							Blueprint->PropagateHierarchyFromBPToInstances();
-						}
-					}
+					const FTransform InitialTransform = GetHierarchy()->GetInitialLocalTransform(Key);
+					GetHierarchy()->SetLocalTransform(Key, InitialTransform, true, true, true);
+					DebuggedHierarchy->SetLocalTransform(Key, InitialTransform, true, true, true);
 				}
 			}
 		}
@@ -2550,10 +2109,9 @@ void SRigHierarchy::HandleSetInitialTransformFromCurrentTransform()
 	{
 		if (UControlRigBlueprint* Blueprint = ControlRigEditor.Pin()->GetControlRigBlueprint())
 		{
-			if (FRigHierarchyContainer* DebuggedContainer = GetDebuggedHierarchyContainer())
+			if (URigHierarchy* DebuggedHierarchy = GetDebuggedHierarchy())
 			{
 				FScopedTransaction Transaction(LOCTEXT("HierarchySetInitialTransforms", "Set Initial Transforms"));
-				Blueprint->Modify();
 
 				TArray<FRigElementKey> SelectedKeys = GetSelectedKeys();
 				TMap<FRigElementKey, FTransform> GlobalTransforms;
@@ -2561,58 +2119,8 @@ void SRigHierarchy::HandleSetInitialTransformFromCurrentTransform()
 
 				for (const FRigElementKey& SelectedKey : SelectedKeys)
 				{
-					if (SelectedKey.Type == ERigElementType::Control || SelectedKey.Type == ERigElementType::Space || SelectedKey.Type == ERigElementType::Bone)
-					{
-						GlobalTransforms.FindOrAdd(SelectedKey) = DebuggedContainer->GetGlobalTransform(SelectedKey);
-						FRigElementKey ParentKey;
-						switch (SelectedKey.Type)
-						{
-							case ERigElementType::Bone:
-							{
-								const FRigBone& Element = DebuggedContainer->BoneHierarchy[SelectedKey.Name];
-								ParentKey = Element.GetParentElementKey();
-								break;
-							}
-							case ERigElementType::Control:
-							{
-								const FRigControl& Element = DebuggedContainer->ControlHierarchy[SelectedKey.Name];
-								if (Element.SpaceName != NAME_None)
-								{
-									ParentKey = Element.GetSpaceElementKey();
-								}
-								else
-								{
-									ParentKey = Element.GetParentElementKey();
-								}
-								break;
-							}
-							case ERigElementType::Space:
-							{
-								const FRigSpace& Element = DebuggedContainer->SpaceHierarchy[SelectedKey.Name];
-								ParentKey = Element.GetParentElementKey();
-								break;
-							}
-							case ERigElementType::Curve:
-							{
-								break;
-							}
-							default:
-							{
-								ensure(false);
-								break;
-							}
-						}
-
-						if (ParentKey)
-						{
-							ParentGlobalTransforms.FindOrAdd(SelectedKey) = DebuggedContainer->GetGlobalTransform(ParentKey);
-						}
-						else
-						{
-							ParentGlobalTransforms.FindOrAdd(SelectedKey) = FTransform::Identity;
-						}
-					}
-
+					GlobalTransforms.FindOrAdd(SelectedKey) = DebuggedHierarchy->GetGlobalTransform(SelectedKey);
+					ParentGlobalTransforms.FindOrAdd(SelectedKey) = DebuggedHierarchy->GetParentTransform(SelectedKey);
 				}
 
 				for (const FRigElementKey& SelectedKey : SelectedKeys)
@@ -2620,25 +2128,27 @@ void SRigHierarchy::HandleSetInitialTransformFromCurrentTransform()
 					FTransform GlobalTransform = GlobalTransforms[SelectedKey];
 					FTransform LocalTransform = GlobalTransform.GetRelativeTransform(ParentGlobalTransforms[SelectedKey]);
 
-					if (SelectedKey.Type == ERigElementType::Control )
-					{
-						GetHierarchyContainer()->ControlHierarchy[SelectedKey.Name].OffsetTransform = LocalTransform;
-						GetHierarchyContainer()->SetLocalTransform(SelectedKey, FTransform::Identity);
-						GetHierarchyContainer()->SetInitialTransform(SelectedKey, FTransform::Identity);
-						DebuggedContainer->ControlHierarchy[SelectedKey.Name].OffsetTransform = LocalTransform;
-						DebuggedContainer->SetLocalTransform(SelectedKey, FTransform::Identity);
-						DebuggedContainer->SetInitialTransform(SelectedKey, FTransform::Identity);
-					}
-					else if (SelectedKey.Type == ERigElementType::Space)
+					if (SelectedKey.Type == ERigElementType::Control)
 					{
-						GetHierarchyContainer()->SetInitialTransform(SelectedKey, LocalTransform);
-						GetHierarchyContainer()->SetLocalTransform(SelectedKey, LocalTransform);
-						DebuggedContainer->SetInitialTransform(SelectedKey, LocalTransform);
-						DebuggedContainer->SetLocalTransform(SelectedKey, LocalTransform);
+						if(FRigControlElement* ControlElement = GetHierarchy()->Find<FRigControlElement>(SelectedKey))
+						{
+							GetHierarchy()->SetControlOffsetTransform(ControlElement, LocalTransform, ERigTransformType::InitialLocal, true, true);
+							GetHierarchy()->SetControlOffsetTransform(ControlElement, LocalTransform, ERigTransformType::CurrentLocal, true, true);
+							GetHierarchy()->SetTransform(ControlElement, FTransform::Identity, ERigTransformType::InitialLocal, true, true);
+							GetHierarchy()->SetTransform(ControlElement, FTransform::Identity, ERigTransformType::CurrentLocal, true, true);
+						}
+						if(FRigControlElement* ControlElement = DebuggedHierarchy->Find<FRigControlElement>(SelectedKey))
+						{
+							DebuggedHierarchy->SetControlOffsetTransform(ControlElement, LocalTransform, ERigTransformType::InitialLocal, true, true);
+							DebuggedHierarchy->SetControlOffsetTransform(ControlElement, LocalTransform, ERigTransformType::CurrentLocal, true, true);
+							DebuggedHierarchy->SetTransform(ControlElement, FTransform::Identity, ERigTransformType::InitialLocal, true, true);
+							DebuggedHierarchy->SetTransform(ControlElement, FTransform::Identity, ERigTransformType::CurrentLocal, true, true);
+						}
 					}
-					else if (SelectedKey.Type == ERigElementType::Bone)
+					else if (SelectedKey.Type == ERigElementType::Space ||
+						SelectedKey.Type == ERigElementType::Bone)
 					{
-						FTransform InitialTransform = GlobalTransform;
+						FTransform InitialTransform = LocalTransform;
 						if (ControlRigEditor.Pin()->PreviewInstance)
 						{
 							if (FAnimNode_ModifyBone* ModifyBone = ControlRigEditor.Pin()->PreviewInstance->FindModifiedBone(SelectedKey.Name))
@@ -2646,11 +2156,19 @@ void SRigHierarchy::HandleSetInitialTransformFromCurrentTransform()
 								InitialTransform.SetTranslation(ModifyBone->Translation);
 								InitialTransform.SetRotation(FQuat(ModifyBone->Rotation));
 								InitialTransform.SetScale3D(ModifyBone->Scale);
-								InitialTransform = InitialTransform * ParentGlobalTransforms[SelectedKey];
 							}
 						}
-						GetHierarchyContainer()->SetInitialGlobalTransform(SelectedKey, InitialTransform);
-						Blueprint->PropagateHierarchyFromBPToInstances(false, false);
+
+						if(FRigTransformElement* TransformElement = GetHierarchy()->Find<FRigTransformElement>(SelectedKey))
+						{
+							GetHierarchy()->SetTransform(TransformElement, LocalTransform, ERigTransformType::InitialLocal, true, true);
+							GetHierarchy()->SetTransform(TransformElement, LocalTransform, ERigTransformType::CurrentLocal, true, true);
+						}
+						if(FRigTransformElement* TransformElement = DebuggedHierarchy->Find<FRigTransformElement>(SelectedKey))
+						{
+							DebuggedHierarchy->SetTransform(TransformElement, LocalTransform, ERigTransformType::InitialLocal, true, true);
+							DebuggedHierarchy->SetTransform(TransformElement, LocalTransform, ERigTransformType::CurrentLocal, true, true);
+						}
 					}
 				}
 			}
@@ -2687,8 +2205,6 @@ void SRigHierarchy::HandleControlBoneOrSpaceTransform()
 			SelectedKey.Type == ERigElementType::Space)
 		{
 			Blueprint->AddTransientControl(SelectedKey);
-			Blueprint->HierarchyContainer.ClearSelection();
-			Blueprint->HierarchyContainer.Select(SelectedKey, true);
 			return;
 		}
 	}
@@ -2703,7 +2219,6 @@ void SRigHierarchy::HandleUnparent()
 	}
 
 	FScopedTransaction Transaction(LOCTEXT("HierarchyTreeUnparentSelected", "Unparent selected items from hierarchy"));
-	ControlRigBlueprint->Modify();
 
 	bool bUnparentImportedBones = false;
 	bool bConfirmedByUser = false;
@@ -2714,9 +2229,9 @@ void SRigHierarchy::HandleUnparent()
 
 	for (const FRigElementKey& SelectedKey : SelectedKeys)
 	{
-		FRigHierarchyContainer* Container = GetHierarchyContainer();
-		InitialTransforms.Add(SelectedKey, Container->GetInitialGlobalTransform(SelectedKey));
-		GlobalTransforms.Add(SelectedKey, Container->GetGlobalTransform(SelectedKey));
+		URigHierarchy* Hierarchy = GetHierarchy();
+		InitialTransforms.Add(SelectedKey, Hierarchy->GetInitialGlobalTransform(SelectedKey));
+		GlobalTransforms.Add(SelectedKey, Hierarchy->GetGlobalTransform(SelectedKey));
 	}
 
 	for (const FRigElementKey& SelectedKey : SelectedKeys)
@@ -2724,10 +2239,11 @@ void SRigHierarchy::HandleUnparent()
 		TGuardValue<bool> GuardRigHierarchyChanges(bIsChangingRigHierarchy, true);
 		TGuardValue<bool> SuspendBlueprintNotifs(ControlRigBlueprint->bSuspendAllNotifications, true);
 
-		FRigHierarchyContainer* Container = GetHierarchyContainer();
-		FRigBoneHierarchy& BoneHierarchy = Container->BoneHierarchy;
-		FRigControlHierarchy& ControlHierarchy = Container->ControlHierarchy;
-		FRigSpaceHierarchy& SpaceHierarchy = Container->SpaceHierarchy;
+		URigHierarchy* Hierarchy = GetHierarchy();
+		check(Hierarchy);
+
+		URigHierarchyController* Controller = Hierarchy->GetController(true);
+		check(Controller);
 
 		const FTransform& InitialTransform = InitialTransforms.FindChecked(SelectedKey);
 		const FTransform& GlobalTransform = GlobalTransforms.FindChecked(SelectedKey);
@@ -2736,7 +2252,13 @@ void SRigHierarchy::HandleUnparent()
 		{
 			case ERigElementType::Bone:
 			{
-				bool bIsImportedBone = BoneHierarchy[SelectedKey.Name].Type == ERigBoneType::Imported;
+				
+				bool bIsImportedBone = false;
+				if(FRigBoneElement* BoneElement = Hierarchy->Find<FRigBoneElement>(SelectedKey))
+				{
+					bIsImportedBone = BoneElement->BoneType == ERigBoneType::Imported;
+				}
+					
 				if (bIsImportedBone && !bConfirmedByUser)
 				{
 					FText ConfirmUnparent = LOCTEXT("ConfirmUnparentBoneHierarchy", "Unparenting imported(white) bones can cause issues with animation - are you sure ?");
@@ -2752,19 +2274,15 @@ void SRigHierarchy::HandleUnparent()
 
 				if (bUnparentImportedBones || !bIsImportedBone)
 				{
-					BoneHierarchy.Reparent(SelectedKey.Name, NAME_None);
+					Controller->RemoveAllParents(SelectedKey, true, true);
 				}
 				break;
 			}
 			case ERigElementType::Space:
-			{
-				SpaceHierarchy.Reparent(SelectedKey.Name, ERigSpaceType::Global, NAME_None);
-				break;
-			}
 			case ERigElementType::Control:
 			{
-				ControlHierarchy.Reparent(SelectedKey.Name, NAME_None);
-				ControlHierarchy.SetSpace(SelectedKey.Name, NAME_None);
+				Controller->RemoveAllParents(SelectedKey, true, true);
+				break;
 			}
 			default:
 			{
@@ -2772,23 +2290,11 @@ void SRigHierarchy::HandleUnparent()
 			}
 		}
 
-		if(SelectedKey.Type == ERigElementType::Control)
-		{
-			int32 ControlIndex = Container->GetIndex(SelectedKey);
-			FTransform ParentTransform = Container->ControlHierarchy.GetParentTransform(ControlIndex, false);
-			FTransform OffsetTransform = GlobalTransform.GetRelativeTransform(ParentTransform);
-			Container->ControlHierarchy[ControlIndex].OffsetTransform = OffsetTransform;
-			Container->SetLocalTransform(SelectedKey, FTransform::Identity);
-			Container->SetInitialTransform(SelectedKey, FTransform::Identity);
-		}
-		else
-		{
-			Container->SetInitialGlobalTransform(SelectedKey, InitialTransform);
-			Container->SetGlobalTransform(SelectedKey, GlobalTransform);
-		}
+		Hierarchy->SetInitialGlobalTransform(SelectedKey, InitialTransform, true, true);
+		Hierarchy->SetGlobalTransform(SelectedKey, GlobalTransform, false, true, true);
 	}
 
-	ControlRigBlueprint->PropagateHierarchyFromBPToInstances(true);
+	ControlRigBlueprint->PropagateHierarchyFromBPToInstances();
 	ControlRigEditor.Pin()->OnHierarchyChanged();
 	RefreshTreeView();
 	FSlateApplication::Get().DismissAllMenus();
@@ -2796,22 +2302,22 @@ void SRigHierarchy::HandleUnparent()
 
 bool SRigHierarchy::FindClosestBone(const FVector& Point, FName& OutRigElementName, FTransform& OutGlobalTransform) const
 {
-	if (FRigHierarchyContainer* DebuggedContainer = GetDebuggedHierarchyContainer())
+	if (URigHierarchy* DebuggedHierarchy = GetDebuggedHierarchy())
 	{
-		const FRigBoneHierarchy& RigHierarchy = DebuggedContainer->BoneHierarchy;
 		float NearestDistance = BIG_NUMBER;
 
-		for (int32 Index = 0; Index < RigHierarchy.Num(); ++Index)
-		{
-			FTransform CurTransform = RigHierarchy.GetGlobalTransform(Index);
-			float CurDistance = FVector::Distance(CurTransform.GetLocation(), Point);
-			if (CurDistance < NearestDistance)
-			{
-				NearestDistance = CurDistance;
-				OutGlobalTransform = CurTransform;
-				OutRigElementName = RigHierarchy[Index].Name;
-			}
-		}
+		DebuggedHierarchy->ForEach<FRigBoneElement>([&] (FRigBoneElement* Element) -> bool
+		{
+			const FTransform CurTransform = DebuggedHierarchy->GetTransform(Element, ERigTransformType::CurrentGlobal);
+            const float CurDistance = FVector::Distance(CurTransform.GetLocation(), Point);
+            if (CurDistance < NearestDistance)
+            {
+                NearestDistance = CurDistance;
+                OutGlobalTransform = CurTransform;
+                OutRigElementName = Element->GetName();
+            }
+            return true;
+		});
 
 		return (OutRigElementName != NAME_None);
 	}
@@ -2824,10 +2330,9 @@ void SRigHierarchy::HandleSetInitialTransformFromClosestBone()
 	{
 		if (UControlRigBlueprint* Blueprint = ControlRigEditor.Pin()->GetControlRigBlueprint())
 		{
-			if (FRigHierarchyContainer* DebuggedContainer = GetDebuggedHierarchyContainer())
+			if (URigHierarchy* DebuggedHierarchy = GetDebuggedHierarchy())
 			{
 				FScopedTransaction Transaction(LOCTEXT("HierarchySetInitialTransforms", "Set Initial Transforms"));
-				Blueprint->Modify();
 
 				TArray<FRigElementKey> SelectedKeys = GetSelectedKeys();
 				TMap<FRigElementKey, FTransform> ClosestTransforms;
@@ -2837,7 +2342,7 @@ void SRigHierarchy::HandleSetInitialTransformFromClosestBone()
 				{
 					if (SelectedKey.Type == ERigElementType::Control || SelectedKey.Type == ERigElementType::Space)
 					{
-						FTransform GlobalTransform = DebuggedContainer->GetGlobalTransform(SelectedKey);
+						const FTransform GlobalTransform = DebuggedHierarchy->GetGlobalTransform(SelectedKey);
 						FTransform ClosestTransform;
 						FName ClosestRigElement;
 
@@ -2847,51 +2352,8 @@ void SRigHierarchy::HandleSetInitialTransformFromClosestBone()
 						}
 
 						ClosestTransforms.FindOrAdd(SelectedKey) = ClosestTransform;
-
-						FRigElementKey ParentKey;
-						switch (SelectedKey.Type)
-						{
-							case ERigElementType::Control:
-							{
-								const FRigControl& Element = DebuggedContainer->ControlHierarchy[SelectedKey.Name];
-								if (Element.SpaceName != NAME_None)
-								{
-									ParentKey = Element.GetSpaceElementKey();
-								}
-								else
-								{
-									ParentKey = Element.GetParentElementKey();
-								}
-								break;
-							}
-							case ERigElementType::Space:
-							{
-								const FRigSpace& Element = DebuggedContainer->SpaceHierarchy[SelectedKey.Name];
-								ParentKey = Element.GetParentElementKey();
-								break;
-							}
-							case ERigElementType::Bone:
-							case ERigElementType::Curve:
-							{
-								break;
-							}
-							default:
-							{
-								ensure(false);
-								break;
-							}
-						}
-
-						if (ParentKey)
-						{
-							ParentGlobalTransforms.FindOrAdd(SelectedKey) = DebuggedContainer->GetGlobalTransform(ParentKey);
-						}
-						else
-						{
-							ParentGlobalTransforms.FindOrAdd(SelectedKey) = FTransform::Identity;
-						}
+						ParentGlobalTransforms.FindOrAdd(SelectedKey) = DebuggedHierarchy->GetParentTransform(SelectedKey);
 					}
-
 				}
 
 				for (const FRigElementKey& SelectedKey : SelectedKeys)
@@ -2900,25 +2362,50 @@ void SRigHierarchy::HandleSetInitialTransformFromClosestBone()
 					{
 						continue;
 					}
+					FTransform GlobalTransform = ClosestTransforms[SelectedKey];
+					FTransform LocalTransform = GlobalTransform.GetRelativeTransform(ParentGlobalTransforms[SelectedKey]);
 
-					if (SelectedKey.Type == ERigElementType::Space || SelectedKey.Type == ERigElementType::Bone)
+					if (SelectedKey.Type == ERigElementType::Control)
 					{
-						FTransform LocalTransform = ClosestTransforms[SelectedKey].GetRelativeTransform(ParentGlobalTransforms[SelectedKey]);
-						GetHierarchyContainer()->SetInitialTransform(SelectedKey, LocalTransform);
-						GetHierarchyContainer()->SetLocalTransform(SelectedKey, LocalTransform);
-						DebuggedContainer->SetInitialTransform(SelectedKey, LocalTransform);
-						DebuggedContainer->SetLocalTransform(SelectedKey, LocalTransform);
-						Blueprint->SetTransientControlValue(SelectedKey);
+						if(FRigControlElement* ControlElement = GetHierarchy()->Find<FRigControlElement>(SelectedKey))
+						{
+							GetHierarchy()->SetControlOffsetTransform(ControlElement, LocalTransform, ERigTransformType::InitialLocal, true, true);
+							GetHierarchy()->SetControlOffsetTransform(ControlElement, LocalTransform, ERigTransformType::CurrentLocal, true, true);
+							GetHierarchy()->SetTransform(ControlElement, FTransform::Identity, ERigTransformType::InitialLocal, true, true);
+							GetHierarchy()->SetTransform(ControlElement, FTransform::Identity, ERigTransformType::CurrentLocal, true, true);
+						}
+						if(FRigControlElement* ControlElement = DebuggedHierarchy->Find<FRigControlElement>(SelectedKey))
+						{
+							DebuggedHierarchy->SetControlOffsetTransform(ControlElement, LocalTransform, ERigTransformType::InitialLocal, true, true);
+							DebuggedHierarchy->SetControlOffsetTransform(ControlElement, LocalTransform, ERigTransformType::CurrentLocal, true, true);
+							DebuggedHierarchy->SetTransform(ControlElement, FTransform::Identity, ERigTransformType::InitialLocal, true, true);
+							DebuggedHierarchy->SetTransform(ControlElement, FTransform::Identity, ERigTransformType::CurrentLocal, true, true);
+						}
 					}
-					if (SelectedKey.Type == ERigElementType::Control)
+					else if (SelectedKey.Type == ERigElementType::Space ||
+                        SelectedKey.Type == ERigElementType::Bone)
 					{
-						FTransform LocalTransform = ClosestTransforms[SelectedKey].GetRelativeTransform(ParentGlobalTransforms[SelectedKey]);
-						GetHierarchyContainer()->ControlHierarchy[SelectedKey.Name].OffsetTransform = LocalTransform;
-						GetHierarchyContainer()->SetLocalTransform(SelectedKey, FTransform::Identity);
-						GetHierarchyContainer()->SetInitialTransform(SelectedKey, FTransform::Identity);
-						DebuggedContainer->ControlHierarchy[SelectedKey.Name].OffsetTransform = LocalTransform;
-						DebuggedContainer->SetLocalTransform(SelectedKey, FTransform::Identity);
-						DebuggedContainer->SetInitialTransform(SelectedKey, FTransform::Identity);
+						FTransform InitialTransform = LocalTransform;
+						if (ControlRigEditor.Pin()->PreviewInstance)
+						{
+							if (FAnimNode_ModifyBone* ModifyBone = ControlRigEditor.Pin()->PreviewInstance->FindModifiedBone(SelectedKey.Name))
+							{
+								InitialTransform.SetTranslation(ModifyBone->Translation);
+								InitialTransform.SetRotation(FQuat(ModifyBone->Rotation));
+								InitialTransform.SetScale3D(ModifyBone->Scale);
+							}
+						}
+
+						if(FRigTransformElement* TransformElement = GetHierarchy()->Find<FRigTransformElement>(SelectedKey))
+						{
+							GetHierarchy()->SetTransform(TransformElement, LocalTransform, ERigTransformType::InitialLocal, true, true);
+							GetHierarchy()->SetTransform(TransformElement, LocalTransform, ERigTransformType::CurrentLocal, true, true);
+						}
+						if(FRigTransformElement* TransformElement = DebuggedHierarchy->Find<FRigTransformElement>(SelectedKey))
+						{
+							DebuggedHierarchy->SetTransform(TransformElement, LocalTransform, ERigTransformType::InitialLocal, true, true);
+							DebuggedHierarchy->SetTransform(TransformElement, LocalTransform, ERigTransformType::CurrentLocal, true, true);
+						}
 					}
 				}
 			}
@@ -2932,31 +2419,32 @@ void SRigHierarchy::HandleSetGizmoTransformFromCurrent()
 	{
 		if (UControlRigBlueprint* Blueprint = ControlRigEditor.Pin()->GetControlRigBlueprint())
 		{
-			if (FRigHierarchyContainer* DebuggedContainer = GetDebuggedHierarchyContainer())
+			if (URigHierarchy* DebuggedHierarchy = GetDebuggedHierarchy())
 			{
-				FScopedTransaction Transaction(LOCTEXT("HierarchySetInitialTransforms", "Set Initial Transforms"));
-				Blueprint->Modify();
+				FScopedTransaction Transaction(LOCTEXT("HierarchySetGizmoTransforms", "Set Gizmo Transforms"));
 
 				TArray<TSharedPtr<FRigTreeElement>> SelectedItems = TreeView->GetSelectedItems();
-				for (TSharedPtr<FRigTreeElement> SelectedItem : SelectedItems)
+				for (const TSharedPtr<FRigTreeElement>& SelectedItem : SelectedItems)
 				{
-					FRigElementKey Key = SelectedItem->Key;
-					if (Key.Type == ERigElementType::Control)
+					if(FRigControlElement* ControlElement = DebuggedHierarchy->Find<FRigControlElement>(SelectedItem->Key))
 					{
-						FRigControl& Control = GetHierarchyContainer()->ControlHierarchy[Key.Name];
-						FRigControl& DebuggedControl = DebuggedContainer->ControlHierarchy[Key.Name];
+						const FRigElementKey Key = ControlElement->GetKey();
 						
-						if (Control.ControlType == ERigControlType::Transform
-							|| Control.ControlType == ERigControlType::TransformNoScale
-							|| Control.ControlType == ERigControlType::EulerTransform)
+						if (ControlElement->Settings.bGizmoEnabled)
 						{
-							if (Control.bGizmoEnabled)
-							{
-								Control.GizmoTransform = Control.GetTransformFromValue(ERigControlValueType::Current);
-								GetHierarchyContainer()->SetLocalTransform(Key, FTransform::Identity);
-								DebuggedControl.GizmoTransform = DebuggedControl.GetTransformFromValue(ERigControlValueType::Current);
-								DebuggedContainer->SetLocalTransform(Key, FTransform::Identity);
-							}
+							const FTransform OffsetGlobalTransform = DebuggedHierarchy->GetGlobalControlOffsetTransform(Key); 
+							const FTransform GizmoGlobalTransform = DebuggedHierarchy->GetGlobalControlGizmoTransform(Key);
+							const FTransform GizmoLocalTransform = GizmoGlobalTransform.GetRelativeTransform(OffsetGlobalTransform);
+							
+							DebuggedHierarchy->SetControlGizmoTransform(Key, GizmoLocalTransform, true, true);
+							DebuggedHierarchy->SetControlGizmoTransform(Key, GizmoLocalTransform, false, true);
+							GetHierarchy()->SetControlGizmoTransform(Key, GizmoLocalTransform, true, true);
+							GetHierarchy()->SetControlGizmoTransform(Key, GizmoLocalTransform, false, true);
+
+							DebuggedHierarchy->SetLocalTransform(Key, FTransform::Identity, false, true, true);
+							DebuggedHierarchy->SetLocalTransform(Key, FTransform::Identity, true, true, true);
+							GetHierarchy()->SetLocalTransform(Key, FTransform::Identity, false, true, true);
+							GetHierarchy()->SetLocalTransform(Key, FTransform::Identity, true, true, true);
 						}
 
 						if (FControlRigEditorEditMode* EditMode = ControlRigEditor.Pin()->GetEditMode())
diff --git a/Engine/Plugins/Experimental/ControlRig/Source/ControlRigEditor/Private/Editor/SRigHierarchy.h b/Engine/Plugins/Experimental/ControlRig/Source/ControlRigEditor/Private/Editor/SRigHierarchy.h
--- a/Engine/Plugins/Experimental/ControlRig/Source/ControlRigEditor/Private/Editor/SRigHierarchy.h
+++ b/Engine/Plugins/Experimental/ControlRig/Source/ControlRigEditor/Private/Editor/SRigHierarchy.h
@@ -4,7 +4,7 @@
 
 #include "CoreMinimal.h"
 #include "Widgets/Views/STreeView.h"
-#include "Rigs/RigHierarchyContainer.h"
+#include "Rigs/RigHierarchy.h"
 #include "EditorUndoClient.h"
 #include "DragAndDrop/GraphNodeDragDropOp.h"
 #include "Engine/SkeletalMesh.h"
@@ -210,6 +210,7 @@ private:
 	void HandlePasteItems();
 	void HandlePasteLocalTransforms();
 	void HandlePasteGlobalTransforms();
+	void HandlePasteTransforms(ERigTransformType::Type InTransformType, bool bAffectChildren);
 
 	/** Set Selection Changed */
 	void OnSelectionChanged(TSharedPtr<FRigTreeElement> Selection, ESelectInfo::Type SelectInfo);
@@ -288,8 +289,8 @@ private:
 	bool IsControlSelected() const;
 	bool IsControlOrSpaceSelected() const;
 
-	FRigHierarchyContainer* GetHierarchyContainer() const;
-	FRigHierarchyContainer* GetDebuggedHierarchyContainer() const;
+	URigHierarchy* GetHierarchy() const;
+	URigHierarchy* GetDebuggedHierarchy() const;
 
 	void ImportHierarchy(const FAssetData& InAssetData);
 	void CreateImportMenu(FMenuBuilder& MenuBuilder);
@@ -299,7 +300,6 @@ private:
 
 	void HandleResetTransform(bool bSelectionOnly);
 	void HandleResetInitialTransform();
-	void HandleResetSpace();
 	void HandleSetInitialTransformFromCurrentTransform();
 	void HandleSetInitialTransformFromClosestBone();
 	void HandleSetGizmoTransformFromCurrent();
@@ -315,19 +315,13 @@ private:
 	void ClearDetailPanel() const;
 
 	bool bIsChangingRigHierarchy;
-	void OnRigElementAdded(FRigHierarchyContainer* Container, const FRigElementKey& InKey);
-	void OnRigElementRemoved(FRigHierarchyContainer* Container, const FRigElementKey& InKey);
-	void OnRigElementRenamed(FRigHierarchyContainer* Container, ERigElementType ElementType, const FName& InOldName, const FName& InNewName);
-	void OnRigElementReparented(FRigHierarchyContainer* Container, const FRigElementKey& InKey, const FName& InOldParentName, const FName& InNewParentName);
-	void OnRigElementSelected(FRigHierarchyContainer* Container, const FRigElementKey& InKey, bool bSelected);
+	void OnHierarchyModified(ERigHierarchyNotification InNotif, URigHierarchy* InHierarchy, const FRigBaseElement* InElement);
 	void HandleRefreshEditorFromBlueprint(UControlRigBlueprint* InBlueprint);
 
 	static TSharedPtr<FRigTreeElement> FindElement(const FRigElementKey& InElementKey, TSharedPtr<FRigTreeElement> CurrentItem);
 	void AddElement(FRigElementKey InKey, FRigElementKey InParentKey = FRigElementKey(), const bool bIgnoreTextFilter = false);
+	void AddElement(FRigBaseElement* InElement, const bool bIgnoreTextFilter = false);
 	void AddSpacerElement();
-	void AddBoneElement(FRigBone InBone, const bool bIgnoreTextFilter = false);
-	void AddControlElement(FRigControl InControl, const bool bIgnoreTextFilter = false);
-	void AddSpaceElement(FRigSpace InSpace, const bool bIgnoreTextFilter = false);
 	void ReparentElement(FRigElementKey InKey, FRigElementKey InParentKey);
 
 public:
