Control Rig: New rig hierarchy data structure refactoring

diff --git a/Engine/Plugins/Experimental/ControlRig/Source/ControlRigEditor/Private/Sequencer/ControlRigParameterTrackEditor.cpp b/Engine/Plugins/Experimental/ControlRig/Source/ControlRigEditor/Private/Sequencer/ControlRigParameterTrackEditor.cpp
--- a/Engine/Plugins/Experimental/ControlRig/Source/ControlRigEditor/Private/Sequencer/ControlRigParameterTrackEditor.cpp
+++ b/Engine/Plugins/Experimental/ControlRig/Source/ControlRigEditor/Private/Sequencer/ControlRigParameterTrackEditor.cpp
@@ -1322,11 +1322,11 @@ void FControlRigParameterTrackEditor::OnTreeViewChanged()
 				const TArray<FName>SelectedControls = ControlRig->CurrentControlSelection();
 				for (const FName  ControlName: SelectedControls)
 				{
-					int32 Index = ControlRig->GetControlHierarchy().GetIndex(ControlName);
+					const int32 Index = ControlRig->GetHierarchy()->GetIndex(FRigElementKey(ControlName, ERigElementType::Control));
 					if (Index != INDEX_NONE)
 					{
-						const FRigControl& RigControl = ControlRig->GetControlHierarchy().GetControls()[Index];
-						HandleControlSelected(ControlRig, RigControl, true);
+						FRigControlElement* ControlElement = ControlRig->GetHierarchy()->GetChecked<FRigControlElement>(Index);
+						HandleControlSelected(ControlRig, ControlElement, true);
 					}
 				}
 			}
@@ -1522,20 +1522,26 @@ FMovieSceneTrackEditor::FFindOrCreateHandleResult FControlRigParameterTrackEdito
 }
 
 
-void FControlRigParameterTrackEditor::HandleControlSelected(UControlRig* Subject, const FRigControl& Control, bool bSelected)
+void FControlRigParameterTrackEditor::HandleControlSelected(UControlRig* Subject, FRigControlElement* ControlElement, bool bSelected)
 {
 	//if parent selected we select child here if it's a bool,integer or single float
 	TArray<FRigControl> Controls;
-	FRigControlHierarchy& ControlHierarchy = Subject->GetControlHierarchy();
-	for (const FRigControl& OtherControl : ControlHierarchy.GetControls())
+
+	URigHierarchy* Hierarchy = Subject->GetHierarchy();
+	if(URigHierarchyController* Controller = Hierarchy->GetController())
 	{
-		if (OtherControl.ParentIndex == Control.Index &&
-			(OtherControl.ControlType == ERigControlType::Bool ||
-				OtherControl.ControlType == ERigControlType::Float ||
-				OtherControl.ControlType == ERigControlType::Integer))
+		Hierarchy->ForEach<FRigControlElement>([ControlElement, Controller, bSelected](FRigControlElement* OtherControlElement) -> bool
 		{
-			Subject->SelectControl(OtherControl.Name, bSelected);
-		}
+			if (OtherControlElement->ParentElements.Contains(ControlElement) &&
+                (OtherControlElement->Settings.ControlType == ERigControlType::Bool ||
+                    OtherControlElement->Settings.ControlType == ERigControlType::Float ||
+                    OtherControlElement->Settings.ControlType == ERigControlType::Integer))
+			{
+				Controller->SelectElement(OtherControlElement->GetKey(), bSelected);
+			}
+
+			return true;
+		});
 	}
 
 	if (bIsDoingSelection)
@@ -1576,7 +1582,7 @@ void FControlRigParameterTrackEditor::HandleControlSelected(UControlRig* Subject
 				UMovieSceneControlRigParameterSection* ParamSection = Cast<UMovieSceneControlRigParameterSection>(Section);
 				if (ParamSection)
 				{
-					FChannelMapInfo* pChannelIndex = ParamSection->ControlChannelMap.Find(Control.Name);
+					FChannelMapInfo* pChannelIndex = ParamSection->ControlChannelMap.Find(ControlElement->GetName());
 					if (pChannelIndex != nullptr)
 					{
 						if (pChannelIndex->ParentControlIndex == INDEX_NONE)
@@ -1697,7 +1703,7 @@ void FControlRigParameterTrackEditor::HandleOnInitialized(UControlRig* ControlRi
 }
 
 
-void FControlRigParameterTrackEditor::HandleControlModified(UControlRig* ControlRig, const FRigControl& Control, const FRigControlModifiedContext& Context)
+void FControlRigParameterTrackEditor::HandleControlModified(UControlRig* ControlRig, FRigControlElement* ControlElement, const FRigControlModifiedContext& Context)
 {
 	if (GetSequencer().IsValid() && !GetSequencer()->IsAllowedToChange())
 	{
@@ -1735,7 +1741,7 @@ void FControlRigParameterTrackEditor::HandleControlModified(UControlRig* Control
 					{
 						KeyMode = ESequencerKeyMode::ManualKey; //best we have here
 					}
-					AddControlKeys(Component, ControlRig, Name, Control.Name, EMovieSceneTransformChannel::All, KeyMode, Context.LocalTime);
+					AddControlKeys(Component, ControlRig, Name, ControlElement->GetName(), EMovieSceneTransformChannel::All, KeyMode, Context.LocalTime);
 				}
 			}
 		}
@@ -1752,7 +1758,7 @@ void FControlRigParameterTrackEditor::HandleControlModified(UControlRig* Control
 
 void FControlRigParameterTrackEditor::GetControlRigKeys(UControlRig* InControlRig, FName ParameterName, EMovieSceneTransformChannel ChannelsToKey, FGeneratedTrackKeys& OutGeneratedKeys)
 {
-	TArray<FRigControl> Controls;
+	TArray<FRigControlElement*> Controls;
 	InControlRig->GetControlsInOrder(Controls);
 	// If key all is enabled, for a key on all the channels
 	if (GetSequencer()->GetKeyGroupMode() == EKeyGroupMode::KeyAll)
@@ -1765,44 +1771,47 @@ void FControlRigParameterTrackEditor::GetControlRigKeys(UControlRig* InControlRi
 	int32 BoolChannelIndex = 0;
 	int32 EnumChannelIndex = 0;
 	int32 IntChannelIndex = 0;
-	for (const FRigControl& RigControl : Controls)
+	for (FRigControlElement* ControlElement : Controls)
 	{
-		if (!RigControl.bAnimatable)
+		if (!ControlElement->Settings.bAnimatable)
 		{
 			continue;
 		}
-		bool bSetKey = RigControl.Name == ParameterName;
-		switch (RigControl.ControlType)
+		bool bSetKey = ControlElement->GetName() == ParameterName;
+
+		FRigControlValue ControlValue = InControlRig->GetHierarchy()->GetControlValue(ControlElement, ERigControlValueType::Current);
+		
+		switch (ControlElement->Settings.ControlType)
 		{
 		case ERigControlType::Bool:
 		{
-			bool Val = RigControl.Value.Get<bool>();
+			bool Val = ControlValue.Get<bool>();
 			OutGeneratedKeys.Add(FMovieSceneChannelValueSetter::Create<FMovieSceneBoolChannel>(BoolChannelIndex++, Val, bSetKey));
 			break;
 		}
 		case ERigControlType::Integer:
 		{
-			if (RigControl.ControlEnum)
+			if (ControlElement->Settings.ControlEnum)
 			{
-				uint8 Val = (uint8)RigControl.Value.Get<uint8>();
+				uint8 Val = (uint8)ControlValue.Get<uint8>();
 				OutGeneratedKeys.Add(FMovieSceneChannelValueSetter::Create<FMovieSceneByteChannel>(EnumChannelIndex++, Val, bSetKey));
 			}
 			else
 			{
-				int32 Val = RigControl.Value.Get<int32>();
+				int32 Val = ControlValue.Get<int32>();
 				OutGeneratedKeys.Add(FMovieSceneChannelValueSetter::Create<FMovieSceneIntegerChannel>(IntChannelIndex++, Val, bSetKey));
 			}
 			break;
 		}
 		case ERigControlType::Float:
 		{
-			float Val = RigControl.Value.Get<float>();
+			float Val = ControlValue.Get<float>();
 			OutGeneratedKeys.Add(FMovieSceneChannelValueSetter::Create<FMovieSceneFloatChannel>(ChannelIndex++, Val, bSetKey));
 			break;
 		}
 		case ERigControlType::Vector2D:
 		{
-			FVector2D Val = RigControl.Value.Get<FVector2D>();
+			FVector2D Val = ControlValue.Get<FVector2D>();
 			OutGeneratedKeys.Add(FMovieSceneChannelValueSetter::Create<FMovieSceneFloatChannel>(ChannelIndex++, Val.X, bSetKey));
 			OutGeneratedKeys.Add(FMovieSceneChannelValueSetter::Create<FMovieSceneFloatChannel>(ChannelIndex++, Val.Y, bSetKey));
 			break;
@@ -1811,7 +1820,7 @@ void FControlRigParameterTrackEditor::GetControlRigKeys(UControlRig* InControlRi
 		case ERigControlType::Scale:
 		case ERigControlType::Rotator:
 		{
-			FVector Val = RigControl.Value.Get<FVector>();
+			FVector Val = ControlValue.Get<FVector>();
 			OutGeneratedKeys.Add(FMovieSceneChannelValueSetter::Create<FMovieSceneFloatChannel>(ChannelIndex++, Val.X, bSetKey));
 			OutGeneratedKeys.Add(FMovieSceneChannelValueSetter::Create<FMovieSceneFloatChannel>(ChannelIndex++, Val.Y, bSetKey));
 			OutGeneratedKeys.Add(FMovieSceneChannelValueSetter::Create<FMovieSceneFloatChannel>(ChannelIndex++, Val.Z, bSetKey));
@@ -1825,22 +1834,22 @@ void FControlRigParameterTrackEditor::GetControlRigKeys(UControlRig* InControlRi
 			FVector Translation, Scale(1.0f, 1.0f, 1.0f);
 			FRotator Rotation;
 
-			if (RigControl.ControlType == ERigControlType::TransformNoScale)
+			if (ControlElement->Settings.ControlType == ERigControlType::TransformNoScale)
 			{
-				FTransformNoScale NoScale = RigControl.Value.Get<FTransformNoScale>();
+				FTransformNoScale NoScale = ControlValue.Get<FTransformNoScale>();
 				Translation = NoScale.Location;
 				Rotation = NoScale.Rotation.Rotator();
 			}
-			else if (RigControl.ControlType == ERigControlType::EulerTransform)
+			else if (ControlElement->Settings.ControlType == ERigControlType::EulerTransform)
 			{
-				FEulerTransform Euler = RigControl.Value.Get < FEulerTransform >();
+				FEulerTransform Euler = ControlValue.Get < FEulerTransform >();
 				Translation = Euler.Location;
 				Rotation = Euler.Rotation;
 				Scale = Euler.Scale;
 			}
 			else
 			{
-				FTransform Val = RigControl.Value.Get<FTransform>();
+				FTransform Val = ControlValue.Get<FTransform>();
 				Translation = Val.GetTranslation();
 				Rotation = Val.GetRotation().Rotator();
 				Scale = Val.GetScale3D();
@@ -1883,7 +1892,7 @@ void FControlRigParameterTrackEditor::GetControlRigKeys(UControlRig* InControlRi
 			OutGeneratedKeys.Add(FMovieSceneChannelValueSetter::Create<FMovieSceneFloatChannel>(ChannelIndex++, CurrentRotator.Pitch, bKeyY));
 			OutGeneratedKeys.Add(FMovieSceneChannelValueSetter::Create<FMovieSceneFloatChannel>(ChannelIndex++, CurrentRotator.Yaw, bKeyZ));
 
-			if (RigControl.ControlType == ERigControlType::Transform || RigControl.ControlType == ERigControlType::EulerTransform)
+			if (ControlElement->Settings.ControlType == ERigControlType::Transform || ControlElement->Settings.ControlType == ERigControlType::EulerTransform)
 			{
 				CurrentVector = Scale;
 				bKeyX = bSetKey && EnumHasAnyFlags(ChannelsToKey, EMovieSceneTransformChannel::ScaleX);
@@ -2050,26 +2059,26 @@ bool FControlRigParameterTrackEditor::ModifyOurGeneratedKeysByCurrentAndWeight(U
 
 	FMovieSceneContext Context(FMovieSceneEvaluationRange(KeyTime, GetSequencer()->GetFocusedTickResolution()));
 	EvalTrack.Interrogate(Context, InterrogationData, Object);
-	const TArray<FRigControl>& Controls = InControlRig->AvailableControls();
+	TArray<FRigControlElement*> Controls = InControlRig->AvailableControls();
 	UMovieSceneControlRigParameterSection* Section = Cast<UMovieSceneControlRigParameterSection>(SectionToKey);
 	FMovieSceneChannelProxy& Proxy = SectionToKey->GetChannelProxy();
 	int32 ChannelIndex = 0;
 	FChannelMapInfo* pChannelIndex = nullptr;
-	for (const FRigControl& RigControl : Controls)
+	for (FRigControlElement* ControlElement : Controls)
 	{
-		if (!RigControl.bAnimatable)
+		if (!ControlElement->Settings.bAnimatable)
 		{
 			continue;
 		}
-		switch (RigControl.ControlType)
+		switch (ControlElement->Settings.ControlType)
 		{
 		case ERigControlType::Float:
 		{
 			for (const FFloatInterrogationData& Val : InterrogationData.Iterate<FFloatInterrogationData>(UMovieSceneControlRigParameterSection::GetFloatInterrogationKey()))
 			{
-				if ((Val.ParameterName == RigControl.Name))
+				if ((Val.ParameterName == ControlElement->GetName()))
 				{
-					pChannelIndex = Section->ControlChannelMap.Find(RigControl.Name);
+					pChannelIndex = Section->ControlChannelMap.Find(ControlElement->GetName());
 					if (pChannelIndex)
 					{
 						ChannelIndex = pChannelIndex->TotalChannelIndex;
@@ -2091,9 +2100,9 @@ bool FControlRigParameterTrackEditor::ModifyOurGeneratedKeysByCurrentAndWeight(U
 		{
 			for (const FVector2DInterrogationData& Val : InterrogationData.Iterate<FVector2DInterrogationData>(UMovieSceneControlRigParameterSection::GetVector2DInterrogationKey()))
 			{
-				if ((Val.ParameterName == RigControl.Name))
+				if ((Val.ParameterName == ControlElement->GetName()))
 				{
-					pChannelIndex = Section->ControlChannelMap.Find(RigControl.Name);
+					pChannelIndex = Section->ControlChannelMap.Find(ControlElement->GetName());
 					if (pChannelIndex)
 					{
 						ChannelIndex = pChannelIndex->TotalChannelIndex;
@@ -2111,16 +2120,16 @@ bool FControlRigParameterTrackEditor::ModifyOurGeneratedKeysByCurrentAndWeight(U
 		{
 			for (const FVectorInterrogationData& Val : InterrogationData.Iterate<FVectorInterrogationData>(UMovieSceneControlRigParameterSection::GetVectorInterrogationKey()))
 			{
-				if ((Val.ParameterName == RigControl.Name))
+				if ((Val.ParameterName == ControlElement->GetName()))
 				{
-					pChannelIndex = Section->ControlChannelMap.Find(RigControl.Name);
+					pChannelIndex = Section->ControlChannelMap.Find(ControlElement->GetName());
 					if (pChannelIndex)
 					{
 						ChannelIndex = pChannelIndex->TotalChannelIndex;
 
 						/* @Mike.Zyracki why is this causing the value to continuously grow?
 						*/
-						if (RigControl.ControlType != ERigControlType::Rotator)
+						if (ControlElement->Settings.ControlType != ERigControlType::Rotator)
 						{
 							GeneratedTotalKeys[ChannelIndex]->ModifyByCurrentAndWeight(Proxy, KeyTime, (void*)&Val.Val.X, Weight);
 							GeneratedTotalKeys[ChannelIndex + 1]->ModifyByCurrentAndWeight(Proxy, KeyTime, (void*)&Val.Val.Y, Weight);
@@ -2140,9 +2149,9 @@ bool FControlRigParameterTrackEditor::ModifyOurGeneratedKeysByCurrentAndWeight(U
 			for (const FTransformInterrogationData& Val : InterrogationData.Iterate<FTransformInterrogationData>(UMovieSceneControlRigParameterSection::GetTransformInterrogationKey()))
 			{
 
-				if ((Val.ParameterName == RigControl.Name))
+				if ((Val.ParameterName == ControlElement->GetName()))
 				{
-					pChannelIndex = Section->ControlChannelMap.Find(RigControl.Name);
+					pChannelIndex = Section->ControlChannelMap.Find(ControlElement->GetName());
 					if (pChannelIndex)
 					{
 						ChannelIndex = pChannelIndex->TotalChannelIndex;
@@ -2156,7 +2165,7 @@ bool FControlRigParameterTrackEditor::ModifyOurGeneratedKeysByCurrentAndWeight(U
 						GeneratedTotalKeys[ChannelIndex + 4]->ModifyByCurrentAndWeight(Proxy, KeyTime, (void *)&CurrentRot.Pitch, Weight);
 						GeneratedTotalKeys[ChannelIndex + 5]->ModifyByCurrentAndWeight(Proxy, KeyTime, (void *)&CurrentRot.Yaw, Weight);
 
-						if (RigControl.ControlType == ERigControlType::Transform || RigControl.ControlType == ERigControlType::EulerTransform)
+						if (ControlElement->Settings.ControlType == ERigControlType::Transform || ControlElement->Settings.ControlType == ERigControlType::EulerTransform)
 						{
 							FVector CurrentScale = Val.Val.GetScale3D();
 							GeneratedTotalKeys[ChannelIndex + 6]->ModifyByCurrentAndWeight(Proxy, KeyTime, (void *)&CurrentScale.X, Weight);
@@ -2621,7 +2630,8 @@ void FControlRigParameterSection::BuildSectionContextMenu(FMenuBuilder& MenuBuil
 				}
 			}
 		}
-		const TArray<FRigControl>& Controls = ControlRig->AvailableControls();
+		
+		TArray<FRigControlElement*> Controls = ControlRig->AvailableControls();
 
 		auto MakeUIAction = [=](EMovieSceneTransformChannel ChannelsToToggle)
 		{
@@ -2733,9 +2743,9 @@ void FControlRigParameterSection::BuildSectionContextMenu(FMenuBuilder& MenuBuil
 				LOCTEXT("ToggleRigControlsText", "Toggle Rig Controls"), LOCTEXT("ToggleRigControlsText_Tooltip", "Toggle Rig Controls"),
 				FNewMenuDelegate::CreateLambda([=](FMenuBuilder& SubMenuBuilder) {
 				int32 Index = 0;
-				for (const FRigControl& RigControl : Controls)
+				for (FRigControlElement* ControlElement : Controls)
 				{
-					const FName RigName = RigControl.Name;
+					const FName RigName = ControlElement->GetName();
 					FText Name = FText::FromName(RigName);
 					FText Text = FText::Format(LOCTEXT("RigControlToggle", "{0}"), Name);
 					FText TooltipText = FText::Format(LOCTEXT("RigControlToggleTooltip", "Toggle Rig Control {0}"), Name);
diff --git a/Engine/Plugins/Experimental/ControlRig/Source/ControlRigEditor/Private/Sequencer/ControlRigParameterTrackEditor.h b/Engine/Plugins/Experimental/ControlRig/Source/ControlRigEditor/Private/Sequencer/ControlRigParameterTrackEditor.h
--- a/Engine/Plugins/Experimental/ControlRig/Source/ControlRigEditor/Private/Sequencer/ControlRigParameterTrackEditor.h
+++ b/Engine/Plugins/Experimental/ControlRig/Source/ControlRigEditor/Private/Sequencer/ControlRigParameterTrackEditor.h
@@ -106,8 +106,8 @@ private:
 	void AddTrackForComponent(USceneComponent* Component);
 
 	/** Control Rig Delegates*/
-	void HandleControlModified(UControlRig* Subject, const FRigControl& Control, const FRigControlModifiedContext& Context);
-	void HandleControlSelected(UControlRig* Subject, const FRigControl& Control, bool bSelected);
+	void HandleControlModified(UControlRig* Subject, FRigControlElement* ControlElement, const FRigControlModifiedContext& Context);
+	void HandleControlSelected(UControlRig* Subject, FRigControlElement* ControlElement, bool bSelected);
 	void HandleOnInitialized(UControlRig* Subject, const EControlRigState InState, const FName& InEventName);
 	/** Post Edit Delegates */
 	void OnPropagateObjectChanges(UObject* InChangedObject);
diff --git a/Engine/Plugins/Experimental/ControlRig/Source/ControlRigEditor/Private/Tools/ControlRigTweener.cpp b/Engine/Plugins/Experimental/ControlRig/Source/ControlRigEditor/Private/Tools/ControlRigTweener.cpp
--- a/Engine/Plugins/Experimental/ControlRig/Source/ControlRigEditor/Private/Tools/ControlRigTweener.cpp
+++ b/Engine/Plugins/Experimental/ControlRig/Source/ControlRigEditor/Private/Tools/ControlRigTweener.cpp
@@ -39,9 +39,9 @@ void FControlsToTween::Setup(const TArray<UControlRig*>& SelectedControlRigs, TW
 				{
 					Section->Modify();
 					UControlRig* ControlRig = Track->GetControlRig();
-					TArray<FRigControl> CurrentControls;
+					TArray<FRigControlElement*> CurrentControls;
 					ControlRig->GetControlsInOrder(CurrentControls);
-					FRigControlHierarchy& Hierarchy = ControlRig->GetControlHierarchy();
+					URigHierarchy* Hierarchy = ControlRig->GetHierarchy();
 					FMovieSceneChannelProxy& Proxy = Section->GetChannelProxy();
 
 					const FName FloatChannelTypeName = FMovieSceneFloatChannel::StaticStruct()->GetFName();
@@ -51,18 +51,18 @@ void FControlsToTween::Setup(const TArray<UControlRig*>& SelectedControlRigs, TW
 					//reuse these arrays
 					TArray<FFrameNumber> KeyTimes;
 					TArray<FKeyHandle> Handles;
-					for (const FRigControl& RigControl : CurrentControls)
+					for (FRigControlElement* ControlElement : CurrentControls)
 					{
-						if (RigControl.bAnimatable &&  ControlRig->IsControlSelected(RigControl.Name))
+						if (ControlElement->Settings.bAnimatable &&  ControlRig->IsControlSelected(ControlElement->GetName()))
 						{
 							FControlRigChannels ControlRigChannels;
 							ControlRigChannels.NumChannels = 0;
 							ControlRigChannels.Section = Section;
 							KeyTimes.SetNum(0);
 							Handles.SetNum(0);
-							FChannelMapInfo* pChannelIndex = Section->ControlChannelMap.Find(RigControl.Name);
+							FChannelMapInfo* pChannelIndex = Section->ControlChannelMap.Find(ControlElement->GetName());
 							int NumChannels = 0;
-							switch (RigControl.ControlType)
+							switch (ControlElement->Settings.ControlType)
 							{
 								case ERigControlType::Float:
 								{
@@ -102,7 +102,7 @@ void FControlsToTween::Setup(const TArray<UControlRig*>& SelectedControlRigs, TW
 							}
 							if (ControlRigChannels.NumChannels > 0)
 							{
-								ControlRigChannelsMap.Add(RigControl.Name, ControlRigChannels);
+								ControlRigChannelsMap.Add(ControlElement->GetName(), ControlRigChannels);
 							}
 						}
 					}
diff --git a/Engine/Plugins/Experimental/ControlRig/Source/ControlRigEditor/Public/ControlRigBlueprintEditorLibrary.h b/Engine/Plugins/Experimental/ControlRig/Source/ControlRigEditor/Public/ControlRigBlueprintEditorLibrary.h
--- a/Engine/Plugins/Experimental/ControlRig/Source/ControlRigEditor/Public/ControlRigBlueprintEditorLibrary.h
+++ b/Engine/Plugins/Experimental/ControlRig/Source/ControlRigEditor/Public/ControlRigBlueprintEditorLibrary.h
@@ -64,6 +64,9 @@ public:
 	static TArray<UStruct*> GetAvailableRigUnits();
 
 	UFUNCTION(BlueprintCallable, BlueprintPure, Category = "Control Rig Blueprint")
-	static UControlRigHierarchyModifier* GetHierarchyModifier(UControlRigBlueprint* InRigBlueprint);
+	static URigHierarchy* GetHierarchy(UControlRigBlueprint* InRigBlueprint);
+
+	UFUNCTION(BlueprintCallable, BlueprintPure, Category = "Control Rig Blueprint")
+	static URigHierarchyController* GetHierarchyController(UControlRigBlueprint* InRigBlueprint);
 };
 
diff --git a/Engine/Plugins/Experimental/FullBodyIK/Source/FullBodyIK/Private/FBIKConstraintLib.cpp b/Engine/Plugins/Experimental/FullBodyIK/Source/FullBodyIK/Private/FBIKConstraintLib.cpp
--- a/Engine/Plugins/Experimental/FullBodyIK/Source/FullBodyIK/Private/FBIKConstraintLib.cpp
+++ b/Engine/Plugins/Experimental/FullBodyIK/Source/FullBodyIK/Private/FBIKConstraintLib.cpp
@@ -3,7 +3,7 @@
 #include "FBIKConstraintLib.h"
 #include "FullBodyIK.h"
 #include "FBIKConstraintOption.h"
-#include "Rigs/RigHierarchyContainer.h"
+#include "Rigs/RigHierarchy.h"
 
 // in the future, we expose rotation axis 
 namespace FBIKConstraintLib
@@ -41,7 +41,7 @@ namespace FBIKConstraintLib
 
 	// for now allow linkdata to be modified here
 	void BuildConstraints(const TArray<FFBIKConstraintOption>& Constraints, TArray<ConstraintType>& OutConstraints,
-		const FRigHierarchyContainer* Hierarchy, TArray<FFBIKLinkData>& InOutLinkData, const TMap<int32, FRigElementKey>& LinkDataToHierarchyIndices,
+		const URigHierarchy* Hierarchy, TArray<FFBIKLinkData>& InOutLinkData, const TMap<int32, FRigElementKey>& LinkDataToHierarchyIndices,
 		const TMap<FRigElementKey, int32>& HierarchyToLinkDataMap)
 	{
 		OutConstraints.Reset();
@@ -86,7 +86,7 @@ namespace FBIKConstraintLib
 								NewLimitConstraint.bYLimitSet = GetLimit(Constraints[Index].AngularLimit.LimitType_Y, FMath::DegreesToRadians(Constraints[Index].AngularLimit.Limit.Y), NewLimitConstraint.Limit.Y);
 								NewLimitConstraint.bZLimitSet = GetLimit(Constraints[Index].AngularLimit.LimitType_Z, FMath::DegreesToRadians(Constraints[Index].AngularLimit.Limit.Z), NewLimitConstraint.Limit.Z);
 
-								FTransform LocalTransform = Hierarchy->GetInitialGlobalTransform(Item).GetRelativeTransform(Hierarchy->GetInitialGlobalTransform(ParentItem));
+								FTransform LocalTransform = Hierarchy->GetGlobalTransform(Item, true).GetRelativeTransform(Hierarchy->GetGlobalTransform(ParentItem, true));
 								// @todo: think about 
 								// set rotation frame
 								NewLimitConstraint.BaseFrameOffset = FQuat(Constraints[Index].OffsetRotation);
@@ -110,7 +110,7 @@ namespace FBIKConstraintLib
 								PoleVectorConstraint.ParentBoneIndex = InOutLinkData[LinkIndex].ParentLinkIndex;
 								PoleVectorConstraint.BoneIndex = LinkIndex; //*ChildLinkIndex;
 
-								TArray<FRigElementKey> Children = Hierarchy->GetChildKeys(Item, false);
+								TArray<FRigElementKey> Children = Hierarchy->GetChildren(Item);
 								for (int32 ChildIndex = 0; ChildIndex < Children.Num(); ++ChildIndex)
 								{
 									const int32* ChildLinkIndex = HierarchyToLinkDataMap.Find(Children[ChildIndex]);
diff --git a/Engine/Plugins/Experimental/FullBodyIK/Source/FullBodyIK/Private/FBIKUtil.cpp b/Engine/Plugins/Experimental/FullBodyIK/Source/FullBodyIK/Private/FBIKUtil.cpp
--- a/Engine/Plugins/Experimental/FullBodyIK/Source/FullBodyIK/Private/FBIKUtil.cpp
+++ b/Engine/Plugins/Experimental/FullBodyIK/Source/FullBodyIK/Private/FBIKUtil.cpp
@@ -1,8 +1,7 @@
 // Copyright Epic Games, Inc. All Rights Reserved.
 
 #include "FBIKUtil.h"
-#include "Rigs/RigHierarchyDefines.h"
-#include "Rigs/RigHierarchyContainer.h"
+#include "Rigs/RigHierarchy.h"
 
 // in the future, we expose rotation axis 
 namespace FBIKUtil
@@ -28,7 +27,7 @@ namespace FBIKUtil
 	}
 
 	// the output should be from current index -> to root parent 
-	bool GetBoneChain(FRigHierarchyContainer* Hierarchy, const FRigElementKey& Root, const FRigElementKey& Current, TArray<FRigElementKey>& ChainIndices)
+	bool GetBoneChain(URigHierarchy* Hierarchy, const FRigElementKey& Root, const FRigElementKey& Current, TArray<FRigElementKey>& ChainIndices)
 	{
 		ChainIndices.Reset();
 
@@ -38,7 +37,7 @@ namespace FBIKUtil
 		while (Iterator.IsValid() && Iterator != Root)
 		{
 			ChainIndices.Insert(Iterator, 0);
-			Iterator = Hierarchy->GetParentKey(Iterator);
+			Iterator = Hierarchy->GetFirstParent(Iterator);
 		}
 
  		// add the last one if valid
diff --git a/Engine/Plugins/Experimental/FullBodyIK/Source/FullBodyIK/Private/RigUnit_FullbodyIK.cpp b/Engine/Plugins/Experimental/FullBodyIK/Source/FullBodyIK/Private/RigUnit_FullbodyIK.cpp
--- a/Engine/Plugins/Experimental/FullBodyIK/Source/FullBodyIK/Private/RigUnit_FullbodyIK.cpp
+++ b/Engine/Plugins/Experimental/FullBodyIK/Source/FullBodyIK/Private/RigUnit_FullbodyIK.cpp
@@ -12,7 +12,7 @@
 /////////////////////////////////////////////////////
 
 
-static float EnsureToAddBoneToLinkData(FRigHierarchyContainer* Hierarchy, const FRigElementKey& CurrentItem, TArray<FFBIKLinkData>& LinkData,
+static float EnsureToAddBoneToLinkData(URigHierarchy* Hierarchy, const FRigElementKey& CurrentItem, TArray<FFBIKLinkData>& LinkData,
 	TMap<FRigElementKey, int32>& HierarchyToLinkDataMap, TMap<int32, FRigElementKey>& LinkDataToHierarchyIndices)
 {
 	float ChainLength = 0;
@@ -25,14 +25,18 @@ static float EnsureToAddBoneToLinkData(FRigHierarchyContainer* Hierarchy, const
 		FFBIKLinkData& NewLink = LinkData[NewLinkIndex];
 
 		// find parent LinkIndex
-		const FRigElementKey& ParentItem = Hierarchy->GetParentKey(CurrentItem);
+		FRigElementKey ParentItem = Hierarchy->GetFirstParent(CurrentItem);
+
+		const int32 CurrentItemIndex = Hierarchy->GetIndex(CurrentItem);
+		const int32 ParentItemIndex = Hierarchy->GetIndex(ParentItem); 
+		
 		FoundLinkIndex = HierarchyToLinkDataMap.Find(ParentItem);
 		NewLink.ParentLinkIndex = (FoundLinkIndex) ? *FoundLinkIndex : INDEX_NONE;
 		NewLink.SetTransform(Hierarchy->GetGlobalTransform(CurrentItem));
 
 		if (ParentItem.IsValid())
 		{
-			FVector DiffLocation = Hierarchy->GetInitialGlobalTransform(CurrentItem).GetLocation() - Hierarchy->GetInitialGlobalTransform(ParentItem).GetLocation();
+			FVector DiffLocation = Hierarchy->GetInitialGlobalTransform(CurrentItemIndex).GetLocation() - Hierarchy->GetInitialGlobalTransform(ParentItemIndex).GetLocation();
 			// set Length
 			NewLink.Length = DiffLocation.Size();
 		}
@@ -75,7 +79,7 @@ static void AddToEffectorTarget(int32 EffectorIndex, const FRigElementKey& Effec
 	}
 }
 
-static void AddEffectors(FRigHierarchyContainer* Hierarchy, const FRigElementKey& Root, const TArray<FFBIKEndEffector>& Effectors,
+static void AddEffectors(URigHierarchy* Hierarchy, const FRigElementKey& Root, const TArray<FFBIKEndEffector>& Effectors,
 	TArray<FFBIKLinkData>& LinkData, TMap<int32, FFBIKEffectorTarget>& EffectorTargets, TArray<int32>& EffectorLinkIndices, 
 	TMap<int32, FRigElementKey>& LinkDataToHierarchyIndices, TMap<FRigElementKey, int32>& HierarchyToLinkDataMap, const FSolverInput& SolverProperty)
 {
@@ -168,7 +172,7 @@ FRigUnit_FullbodyIK_Execute()
 {
     DECLARE_SCOPE_HIERARCHICAL_COUNTER_RIGUNIT()
 
-	FRigHierarchyContainer* Hierarchy = ExecuteContext.Hierarchy;
+	URigHierarchy* Hierarchy = ExecuteContext.Hierarchy;
 	if (Hierarchy == nullptr)
 	{
 		return;
diff --git a/Engine/Plugins/Experimental/FullBodyIK/Source/FullBodyIK/Public/FBIKConstraintLib.h b/Engine/Plugins/Experimental/FullBodyIK/Source/FullBodyIK/Public/FBIKConstraintLib.h
--- a/Engine/Plugins/Experimental/FullBodyIK/Source/FullBodyIK/Public/FBIKConstraintLib.h
+++ b/Engine/Plugins/Experimental/FullBodyIK/Source/FullBodyIK/Public/FBIKConstraintLib.h
@@ -8,14 +8,14 @@
 
 struct FFBIKConstraintOption;
 struct FRigElementKey;
-struct FRigHierarchyContainer;
+class URigHierarchy;
 
 namespace FBIKConstraintLib
 {
 	// constraint related
 	void FULLBODYIK_API ApplyConstraint(TArray<FFBIKLinkData>& InOutLinkData, TArray<ConstraintType>* Constraints);
 	void FULLBODYIK_API BuildConstraints(const TArray<FFBIKConstraintOption>& Constraints, TArray<ConstraintType>& OutConstraints,
-		const FRigHierarchyContainer* Hierarchy, TArray<FFBIKLinkData>& InOutLinkData, const TMap<int32, FRigElementKey>& LinkDataToHierarchyIndices,
+		const URigHierarchy* Hierarchy, TArray<FFBIKLinkData>& InOutLinkData, const TMap<int32, FRigElementKey>& LinkDataToHierarchyIndices,
 		const TMap<FRigElementKey, int32>& HierarchyToLinkDataMap);
 };
 
diff --git a/Engine/Plugins/Experimental/FullBodyIK/Source/FullBodyIK/Public/FBIKUtil.h b/Engine/Plugins/Experimental/FullBodyIK/Source/FullBodyIK/Public/FBIKUtil.h
--- a/Engine/Plugins/Experimental/FullBodyIK/Source/FullBodyIK/Public/FBIKUtil.h
+++ b/Engine/Plugins/Experimental/FullBodyIK/Source/FullBodyIK/Public/FBIKUtil.h
@@ -4,7 +4,7 @@
 
 #include "CoreMinimal.h"
 
-struct FRigHierarchyContainer;
+class URigHierarchy;
 struct FRigElementKey;
 
 namespace FBIKUtil
@@ -12,7 +12,7 @@ namespace FBIKUtil
 	/** Utility functions */
 	FVector FULLBODYIK_API GetScaledRotationAxis(const FQuat& InQuat);
 	bool FULLBODYIK_API CanCrossProduct(const FVector& Vector1, const FVector& Vector2);
-	bool FULLBODYIK_API GetBoneChain(FRigHierarchyContainer* Hierarchy, const FRigElementKey& Root, const FRigElementKey& Current, TArray<FRigElementKey>& ChainIndices);
+	bool FULLBODYIK_API GetBoneChain(URigHierarchy* Hierarchy, const FRigElementKey& Root, const FRigElementKey& Current, TArray<FRigElementKey>& ChainIndices);
 };
 
 
diff --git a/Engine/Plugins/Experimental/FullBodyIK/Source/PBIK/Private/Core/PBIKBody.cpp b/Engine/Plugins/Experimental/FullBodyIK/Source/PBIK/Private/Core/PBIKBody.cpp
--- a/Engine/Plugins/Experimental/FullBodyIK/Source/PBIK/Private/Core/PBIKBody.cpp
+++ b/Engine/Plugins/Experimental/FullBodyIK/Source/PBIK/Private/Core/PBIKBody.cpp
@@ -11,7 +11,8 @@ FBone::FBone(
 	const int& InParentIndex,		// must pass -1 for root of whole skeleton
 	const FVector& InOrigPosition,
 	const FQuat& InOrigRotation,
-	bool bInIsSolverRoot)
+	bool bInIsSolverRoot,
+	const int& InElementIndex)
 {
 	Name = InName;
 	ParentIndex = InParentIndex;
@@ -19,6 +20,7 @@ FBone::FBone(
 	Rotation = InOrigRotation;
 	bIsSolverRoot = bInIsSolverRoot;
 	bIsSolved = false; // default
+	ElementIndex = InElementIndex;
 }
 
 bool FBone::HasChild(const FBone* Bone)
diff --git a/Engine/Plugins/Experimental/FullBodyIK/Source/PBIK/Private/Core/PBIKSolver.cpp b/Engine/Plugins/Experimental/FullBodyIK/Source/PBIK/Private/Core/PBIKSolver.cpp
--- a/Engine/Plugins/Experimental/FullBodyIK/Source/PBIK/Private/Core/PBIKSolver.cpp
+++ b/Engine/Plugins/Experimental/FullBodyIK/Source/PBIK/Private/Core/PBIKSolver.cpp
@@ -496,9 +496,10 @@ int32 FPBIKSolver::AddBone(
 	const int32 ParentIndex,
 	const FVector& InOrigPosition,
 	const FQuat& InOrigRotation,
-	bool bIsSolverRoot)
+	bool bIsSolverRoot,
+	const int& InElementIndex)
 {
-	return Bones.Emplace(Name, ParentIndex, InOrigPosition, InOrigRotation, bIsSolverRoot);
+	return Bones.Emplace(Name, ParentIndex, InOrigPosition, InOrigRotation, bIsSolverRoot, InElementIndex);
 }
 
 int32 FPBIKSolver::AddEffector(FName BoneName)
diff --git a/Engine/Plugins/Experimental/FullBodyIK/Source/PBIK/Private/RigUnit_PBIK.cpp b/Engine/Plugins/Experimental/FullBodyIK/Source/PBIK/Private/RigUnit_PBIK.cpp
--- a/Engine/Plugins/Experimental/FullBodyIK/Source/PBIK/Private/RigUnit_PBIK.cpp
+++ b/Engine/Plugins/Experimental/FullBodyIK/Source/PBIK/Private/RigUnit_PBIK.cpp
@@ -9,21 +9,21 @@ FRigUnit_PBIK_Execute()
 {
 	DECLARE_SCOPE_HIERARCHICAL_COUNTER_RIGUNIT()
 
-	FRigHierarchyContainer* Hierarchy = ExecuteContext.Hierarchy;
+	URigHierarchy* Hierarchy = ExecuteContext.Hierarchy;
 	if (Hierarchy == nullptr)
 	{
 		return;
 	}
 
-	FRigBoneHierarchy& Bones = Hierarchy->BoneHierarchy;
-
 	if (Context.State == EControlRigState::Init)
 	{
+		BoneSettingToSolverBoneIndex.Reset();
+		
 		// check how many effectors are assigned to a bone
 		int NumEffectors = 0;
 		for (const FPBIKEffector& Effector : Effectors)
 		{
-			if (Bones.GetIndex(Effector.Bone) != INDEX_NONE)
+			if (Hierarchy->GetIndex(FRigElementKey(Effector.Bone, ERigElementType::Bone)) != INDEX_NONE)
 			{
 				++NumEffectors; // bone is set and exists!
 			}
@@ -41,14 +41,16 @@ FRigUnit_PBIK_Execute()
 		Solver.Reset();
 
 		// create bones
-		for (int B = 0; B < Bones.Num(); ++B)
+		TArray<FRigBoneElement*> BoneElements = Hierarchy->GetBones(true);
+		for (int B = 0; B < BoneElements.Num(); ++B)
 		{
-			FName Name = Bones[B].Name;
-			int ParentIndex = Bones[B].ParentIndex;
-			FVector InOrigPosition = Bones[B].InitialTransform.GetLocation();
-			FQuat InOrigRotation = Bones[B].InitialTransform.GetRotation();
-			bool bIsRoot = Bones[B].Name == Root;
-			Solver.AddBone(Name, ParentIndex, InOrigPosition, InOrigRotation, bIsRoot);
+			FName Name = BoneElements[B]->GetName();
+			const int ParentIndex = Hierarchy->GetFirstParent(BoneElements[B]->GetIndex());
+			const FTransform OrigTransform = Hierarchy->GetTransform(BoneElements[B], ERigTransformType::InitialGlobal);
+			const FVector InOrigPosition = OrigTransform.GetLocation();
+			const FQuat InOrigRotation = OrigTransform.GetRotation();
+			bool bIsRoot = BoneElements[B]->GetName() == Root;
+			Solver.AddBone(Name, ParentIndex, InOrigPosition, InOrigRotation, bIsRoot, BoneElements[B]->GetIndex());
 		}
 		
 		// create effectors
@@ -71,20 +73,44 @@ FRigUnit_PBIK_Execute()
 	}
 
 	// set bones to input pose
-	for (int B = 0; B < Bones.Num(); ++B)
+	for(int32 BoneIndex = 0; BoneIndex < Solver.Bones.Num(); BoneIndex++)
+	{
+		const FTransform GlobalTransform = Hierarchy->GetGlobalTransform(Solver.Bones[BoneIndex].ElementIndex);
+		Solver.SetBoneTransform(BoneIndex, GlobalTransform);
+	}
+
+	// invalidate the name lookup for the settings
+	if(BoneSettingToSolverBoneIndex.Num() != BoneSettings.Num())
 	{
-		Solver.SetBoneTransform(B, Bones[B].GlobalTransform);
+		BoneSettingToSolverBoneIndex.Reset();
+		while(BoneSettingToSolverBoneIndex.Num() < BoneSettings.Num())
+		{
+			BoneSettingToSolverBoneIndex.Add(INDEX_NONE);
+		}
 	}
 
 	// update bone settings
-	for (const FPBIKBoneSetting& BoneSetting : BoneSettings)
+	for (int32 BoneSettingIndex = 0; BoneSettingIndex < BoneSettings.Num(); BoneSettingIndex++)
 	{
-		int32 BoneIndex = Bones.GetIndex(BoneSetting.Bone);
-		if (BoneIndex == -1)
+		const FPBIKBoneSetting& BoneSetting = BoneSettings[BoneSettingIndex];
+
+		if(BoneSettingToSolverBoneIndex[BoneSettingIndex] == INDEX_NONE)
 		{
-			continue; // no bones to apply it to
+			for(int32 BoneIndex = 0; BoneIndex < Solver.Bones.Num(); BoneIndex++)
+			{
+				if(Solver.Bones[BoneIndex].Name == BoneSetting.Bone)
+				{
+					BoneSettingToSolverBoneIndex[BoneSettingIndex] = BoneIndex;
+					break;
+				}
+			}
+			if(BoneSettingToSolverBoneIndex[BoneSettingIndex] == INDEX_NONE)
+			{
+				continue;
+			}
 		}
 
+		int32 BoneIndex = BoneSettingToSolverBoneIndex[BoneSettingIndex];
 		if (PBIK::FBoneSettings* InternalSettings = Solver.GetBoneSettings(BoneIndex))
 		{
 			BoneSetting.CopyToCoreStruct(*InternalSettings);
@@ -109,12 +135,12 @@ FRigUnit_PBIK_Execute()
 	Solver.Solve(Settings);
 
 	// copy transforms back
-	bool bPropagateTransform = false;
-	for (int BoneIndex = 0; BoneIndex < Bones.Num(); ++BoneIndex)
+	const bool bPropagateTransform = false;
+	for(int32 BoneIndex = 0; BoneIndex < Solver.Bones.Num(); BoneIndex++)
 	{
 		FTransform NewTransform;
 		Solver.GetBoneGlobalTransform(BoneIndex, NewTransform);
-		Bones.SetGlobalTransform(BoneIndex, NewTransform, bPropagateTransform);
+		Hierarchy->SetGlobalTransform(Solver.Bones[BoneIndex].ElementIndex, NewTransform, false, bPropagateTransform);
 	}
 
 	// do all debug drawing
diff --git a/Engine/Plugins/Experimental/FullBodyIK/Source/PBIK/Public/Core/PBIKBody.h b/Engine/Plugins/Experimental/FullBodyIK/Source/PBIK/Public/Core/PBIKBody.h
--- a/Engine/Plugins/Experimental/FullBodyIK/Source/PBIK/Public/Core/PBIKBody.h
+++ b/Engine/Plugins/Experimental/FullBodyIK/Source/PBIK/Public/Core/PBIKBody.h
@@ -18,6 +18,7 @@ struct FEffector;
 struct FBone
 {
 	FName Name;
+	int ElementIndex = INDEX_NONE; // used to map to the element this bone references
 	int ParentIndex = -2; // -2 is unset, -1 for root, or 0...n otherwise
 	bool bIsSolverRoot = false;
 	bool bIsSolved = false;
@@ -38,7 +39,8 @@ struct FBone
 		const int& InParentIndex,		// must pass -1 for root of whole skeleton
 		const FVector& InOrigPosition,
 		const FQuat& InOrigRotation,
-		bool bInIsSolverRoot);
+		bool bInIsSolverRoot,
+		const int& InElementIndex = INDEX_NONE);
 
 	bool HasChild(const FBone* Bone);
 };
diff --git a/Engine/Plugins/Experimental/FullBodyIK/Source/PBIK/Public/Core/PBIKSolver.h b/Engine/Plugins/Experimental/FullBodyIK/Source/PBIK/Public/Core/PBIKSolver.h
--- a/Engine/Plugins/Experimental/FullBodyIK/Source/PBIK/Public/Core/PBIKSolver.h
+++ b/Engine/Plugins/Experimental/FullBodyIK/Source/PBIK/Public/Core/PBIKSolver.h
@@ -115,7 +115,8 @@ public:
 		const int32 ParentIndex,
 		const FVector& InOrigPosition,
 		const FQuat& InOrigRotation,
-		bool bIsSolverRoot);
+		bool bIsSolverRoot,
+		const int& InElementIndex = INDEX_NONE);
 
 	int32 AddEffector(FName BoneName);
 	
@@ -141,4 +142,5 @@ private:
 	
 	PBIK::FDebugDraw DebugDraw;
 	friend PBIK::FDebugDraw;
+	friend struct FRigUnit_PBIK;
 };
diff --git a/Engine/Plugins/Experimental/LiveLinkControlRig/Source/LiveLinkControlRig/Private/LiveLinkRigUnits.cpp b/Engine/Plugins/Experimental/LiveLinkControlRig/Source/LiveLinkControlRig/Private/LiveLinkRigUnits.cpp
--- a/Engine/Plugins/Experimental/LiveLinkControlRig/Source/LiveLinkControlRig/Private/LiveLinkRigUnits.cpp
+++ b/Engine/Plugins/Experimental/LiveLinkControlRig/Source/LiveLinkControlRig/Private/LiveLinkRigUnits.cpp
@@ -47,6 +47,7 @@ FRigUnit_LiveLinkEvaluteFrameAnimation_Execute()
 		// Debug Drawing if required
 		if (bDrawDebug && Context.DrawInterface != nullptr)
 		{
+			/*
 			FRigBoneHierarchy DebugHierarchy;
 
 			for (int BoneIdx = 0; BoneIdx < SubjectFrame.GetNumberOfTransforms(); ++BoneIdx)
@@ -68,6 +69,7 @@ FRigUnit_LiveLinkEvaluteFrameAnimation_Execute()
 				DebugHierarchy.Add(CurrentBoneName, CurrentParentName, ERigBoneType::User, LocalSpaceTransform, LocalSpaceTransform, GlobalSpaceTransform);
 			}
 			Context.DrawInterface->DrawHierarchy(DebugDrawOffset, DebugHierarchy, EControlRigDrawHierarchyMode::Axes, 1.0f, DebugColor, 1.0f);
+			*/
 		}
 	}
 		break;
@@ -160,6 +162,8 @@ FRigUnit_LiveLinkEvaluteFrameTransform_Execute()
 		{
 			// Rather than drawing as an axis guide, draw as a hierarchy with a dummy root bone which makes it clearer what the origin of the transform
 			// space is
+
+			/*
 			FRigBoneHierarchy DebugHierarchy;
 
 			FName RootName("root");
@@ -168,6 +172,7 @@ FRigUnit_LiveLinkEvaluteFrameTransform_Execute()
 			DebugHierarchy.Add(TransformName, RootName, ERigBoneType::User, Transform, Transform, Transform);
 
 			Context.DrawInterface->DrawHierarchy(DebugDrawOffset, DebugHierarchy, EControlRigDrawHierarchyMode::Axes, 1.0f, DebugColor, 1.0f);
+			*/
 		}
 	}
 		break;
diff --git a/Engine/Plugins/Experimental/MotionTrailEditorMode/Source/MotionTrailEditorMode/Private/Sequencer/MovieSceneTransformTrail.cpp b/Engine/Plugins/Experimental/MotionTrailEditorMode/Source/MotionTrailEditorMode/Private/Sequencer/MovieSceneTransformTrail.cpp
--- a/Engine/Plugins/Experimental/MotionTrailEditorMode/Source/MotionTrailEditorMode/Private/Sequencer/MovieSceneTransformTrail.cpp
+++ b/Engine/Plugins/Experimental/MotionTrailEditorMode/Source/MotionTrailEditorMode/Private/Sequencer/MovieSceneTransformTrail.cpp
@@ -228,7 +228,8 @@ void FMovieSceneControlTransformTrail::UpdateCacheTimes(const FTrailEvaluateTime
 	UMovieSceneControlRigParameterSection* Section = Cast<UMovieSceneControlRigParameterSection>(GetSection());
 	check(Section);
 
-	const FTransform InitialTransform = Section->ControlRig->GetControlHierarchy().GetInitialValue<FTransform>(ControlName);
+	const FRigElementKey ControlKey(ControlName, ERigElementType::Control);
+	const FTransform InitialTransform = Section->ControlRig->GetHierarchy()->GetControlValue(ControlKey, ERigControlValueType::Initial).Get<FTransform>();
 
 	TArrayView<FMovieSceneFloatChannel*> FloatChannels = Section->GetChannelProxy().GetChannels<FMovieSceneFloatChannel>();
 	FloatChannels = FloatChannels.Slice(GetChannelOffset(), uint8(EMSTrailTransformChannel::MaxChannel) + 1);
@@ -247,7 +248,9 @@ UE::MovieScene::FIntermediate3DTransform FMovieSceneControlTransformTrail::Calcu
 {
 	UMovieSceneControlRigParameterSection* Section = Cast<UMovieSceneControlRigParameterSection>(GetSection());
 
-	const FTransform InitialTransform = Section->ControlRig->GetControlHierarchy().GetInitialValue<FTransform>(ControlName);
+	const FRigElementKey ControlKey(ControlName, ERigElementType::Control);
+	const FTransform InitialTransform = Section->ControlRig->GetHierarchy()->GetControlValue(ControlKey, ERigControlValueType::Initial).Get<FTransform>();
+
 	const UE::MovieScene::FIntermediate3DTransform Delta = UE::MovieScene::FIntermediate3DTransform(
 		Current.GetTranslation() - Start.GetTranslation(),
 		Current.GetRotation() - Start.GetRotation(),
diff --git a/Engine/Plugins/Experimental/MotionTrailEditorMode/Source/MotionTrailEditorMode/Private/Sequencer/SequencerTrailHierarchy.cpp b/Engine/Plugins/Experimental/MotionTrailEditorMode/Source/MotionTrailEditorMode/Private/Sequencer/SequencerTrailHierarchy.cpp
--- a/Engine/Plugins/Experimental/MotionTrailEditorMode/Source/MotionTrailEditorMode/Private/Sequencer/SequencerTrailHierarchy.cpp
+++ b/Engine/Plugins/Experimental/MotionTrailEditorMode/Source/MotionTrailEditorMode/Private/Sequencer/SequencerTrailHierarchy.cpp
@@ -131,11 +131,8 @@ void FSequencerTrailHierarchy::Destroy()
 	for (const TPair<UMovieSceneSection*, FControlRigDelegateHandles>& SectionHandlesPair : ControlRigDelegateHandles)
 	{
 		UMovieSceneControlRigParameterSection* Section = Cast<UMovieSceneControlRigParameterSection>(SectionHandlesPair.Key);
-		FRigHierarchyContainer* RigHierarchy = Section->ControlRig->GetHierarchy();
-		RigHierarchy->OnElementAdded.Remove(SectionHandlesPair.Value.OnControlAddedHandle);
-		RigHierarchy->OnElementRemoved.Remove(SectionHandlesPair.Value.OnControlRemovedHandle);
-		RigHierarchy->OnElementReparented.Remove(SectionHandlesPair.Value.OnControlReparentedHandle);
-		RigHierarchy->OnElementRenamed.Remove(SectionHandlesPair.Value.OnControlRenamedHandle);
+		URigHierarchy* RigHierarchy = Section->ControlRig->GetHierarchy();
+		RigHierarchy->OnModified().Remove(SectionHandlesPair.Value.OnHierarchyModified);
 	}
 
 	for (TPair<FGuid, TUniquePtr<FTrail>>& GuidTrailPair : AllTrails)
@@ -600,12 +597,12 @@ void FSequencerTrailHierarchy::AddSkeletonToHierarchy(class USkeletalMeshCompone
 	TimingStats.Add("FSequencerTrailHierarchy::AddSkeletonToHierarchy", Timespan);
 }
 
-void FSequencerTrailHierarchy::ResolveRigElementToRootComponent(FRigHierarchyContainer* RigHierarchy, FRigElementKey InElementKey, class USkeletalMeshComponent* Component)
+void FSequencerTrailHierarchy::ResolveRigElementToRootComponent(URigHierarchy* RigHierarchy, FRigElementKey InElementKey, class USkeletalMeshComponent* Component)
 {
 	TMap<FName, FGuid>& ControlMap = ControlsTracked[Component];
 	
-	int32 ElementIndex = RigHierarchy->ControlHierarchy.GetIndex(InElementKey.Name);
-	FRigElementKey ParentKey = RigHierarchy->ControlHierarchy[ElementIndex].SpaceName != NAME_None ? RigHierarchy->ControlHierarchy[ElementIndex].GetSpaceElementKey() : RigHierarchy->ControlHierarchy[ElementIndex].GetParentElementKey();
+	int32 ElementIndex = RigHierarchy->GetIndex(InElementKey);
+	FRigElementKey ParentKey = RigHierarchy->GetFirstParent(InElementKey);
 
 	FRigElementKey ChildItr = InElementKey;
 	while (ParentKey.IsValid() && ParentKey.Type != ERigElementType::Bone)
@@ -613,8 +610,8 @@ void FSequencerTrailHierarchy::ResolveRigElementToRootComponent(FRigHierarchyCon
 		if (ParentKey.Type == ERigElementType::Space) // TODO: add empty space trail with initial transform
 		{
 			ChildItr = ParentKey;
-			ElementIndex = RigHierarchy->ControlHierarchy.GetIndex(ChildItr.Name);
-			ParentKey = RigHierarchy->SpaceHierarchy[ElementIndex].GetParentElementKey();
+			ElementIndex = RigHierarchy->GetIndex(FRigElementKey(ChildItr.Name, ERigElementType::Control));
+			ParentKey = RigHierarchy->GetFirstParent(FRigElementKey(ChildItr.Name, ERigElementType::Space));
 			continue;
 		}
 
@@ -633,8 +630,8 @@ void FSequencerTrailHierarchy::ResolveRigElementToRootComponent(FRigHierarchyCon
 		}
 
 		ChildItr = ParentKey;
-		ElementIndex = RigHierarchy->ControlHierarchy.GetIndex(ChildItr.Name);
-		ParentKey = RigHierarchy->ControlHierarchy[ElementIndex].GetParentElementKey();
+		ElementIndex = RigHierarchy->GetIndex(FRigElementKey(ChildItr.Name, ERigElementType::Control));
+		ParentKey = RigHierarchy->GetFirstParent(FRigElementKey(ChildItr.Name, ERigElementType::Control));
 	}
 
 	if (!ParentKey.IsValid()) // Add to root component trail
@@ -687,7 +684,7 @@ void FSequencerTrailHierarchy::AddControlsToHierarchy(class USkeletalMeshCompone
 	UMovieSceneControlRigParameterSection* CRParamSection = Cast<UMovieSceneControlRigParameterSection>(Sections[0]);
 	check(CRParamSection);
 
-	FRigHierarchyContainer* RigHierarchy = CRParamSection->ControlRig->GetHierarchy();
+	URigHierarchy* RigHierarchy = CRParamSection->ControlRig->GetHierarchy();
 	if (!ControlRigDelegateHandles.Contains(CRParamSection))
 	{
 		RegisterControlRigDelegates(CompToAdd, CRParamSection);
@@ -695,12 +692,13 @@ void FSequencerTrailHierarchy::AddControlsToHierarchy(class USkeletalMeshCompone
 
 	CRParamSection->ReconstructChannelProxy(true);
 
-	TArray<FRigControl> SortedControls;
+	TArray<FRigControlElement*> SortedControls;
 	CRParamSection->ControlRig->GetControlsInOrder(SortedControls);
+	
 	for (const TPair<FName, FChannelMapInfo>& NameInfoPair : CRParamSection->ControlChannelMap)
 	{
-		const FRigControl& Control = SortedControls[NameInfoPair.Value.ControlIndex];
-		const FRigElementKey RigKey = Control.GetElementKey();
+		FRigControlElement* ControlElement = SortedControls[NameInfoPair.Value.ControlIndex];
+		const FRigElementKey RigKey = ControlElement->GetKey();
 
 		if (RigKey.Type != ERigElementType::Control)
 		{
@@ -724,103 +722,99 @@ void FSequencerTrailHierarchy::AddControlsToHierarchy(class USkeletalMeshCompone
 
 void FSequencerTrailHierarchy::RegisterControlRigDelegates(USkeletalMeshComponent* Component, class UMovieSceneControlRigParameterSection* CRParamSection)
 {
-	FRigHierarchyContainer* RigHierarchy = CRParamSection->ControlRig->GetHierarchy();
+	URigHierarchy* RigHierarchy = CRParamSection->ControlRig->GetHierarchy();
 	FControlRigDelegateHandles& DelegateHandles = ControlRigDelegateHandles.Add(CRParamSection);
-	DelegateHandles.OnControlAddedHandle = RigHierarchy->OnElementAdded.AddLambda(
-		[this, RigHierarchy, Component, CRParamSection](FRigHierarchyContainer*, const FRigElementKey& NewElemKey) {
+	DelegateHandles.OnHierarchyModified = RigHierarchy->OnModified().AddLambda(
+		[this, RigHierarchy, Component, CRParamSection](ERigHierarchyNotification InNotif, URigHierarchy* InHierarchy, const FRigBaseElement* InElement) {
 			TSharedPtr<ISequencer> Sequencer = WeakSequencer.Pin();
 			check(Sequencer);
-			if (NewElemKey.Type != ERigElementType::Control) // We only care about controls
+
+			const FRigControlElement* ControlElement = Cast<FRigControlElement>(InElement); 
+			if(ControlElement == nullptr)
 			{
 				return;
 			}
 
-			const FRigControl& Control = RigHierarchy->ControlHierarchy[NewElemKey.Name];
-			if (Control.ParentName != NAME_None)
+			if(InNotif == ERigHierarchyNotification::ElementAdded)
 			{
-				TMap<FName, FGuid>& ControlMap = ControlsTracked[Component];
-
-				const FGuid ControlGuid = ControlMap.FindOrAdd(NewElemKey.Name, FGuid::NewGuid());
-				FTrailHierarchyNode& ControlNode = Hierarchy.FindOrAdd(ControlGuid, FTrailHierarchyNode());
+				const FRigElementKey ParentKey = InHierarchy->GetFirstParent(ControlElement->GetKey());
+                if (ParentKey.IsValid())
+                {
+                    TMap<FName, FGuid>& ControlMap = ControlsTracked[Component];
 
-				ResolveRigElementToRootComponent(RigHierarchy, NewElemKey, Component);
+                    const FGuid ControlGuid = ControlMap.FindOrAdd(ControlElement->GetName(), FGuid::NewGuid());
+                    FTrailHierarchyNode& ControlNode = Hierarchy.FindOrAdd(ControlGuid, FTrailHierarchyNode());
 
-				const int32 ChannelIndex = CRParamSection->ControlChannelMap[NewElemKey.Name].ChannelIndex;
-				TUniquePtr<FTrail> CurTrail = MakeUnique<FMovieSceneControlTransformTrail>(FLinearColor::White, false, CRParamSection, Sequencer, ChannelIndex, NewElemKey.Name);
-				if (AllTrails.Contains(ControlMap[NewElemKey.Name])) 
-				{
-					AllTrails.Remove(ControlMap[NewElemKey.Name]);
-				}
+                    ResolveRigElementToRootComponent(RigHierarchy, ControlElement->GetKey(), Component);
 
-				AddTrail(ControlMap[NewElemKey.Name], Hierarchy[ControlMap[NewElemKey.Name]], MoveTemp(CurTrail));
-			}
-			else
-			{
-				// TODO: handle spaces
-			}
-		}
-	);
+                    const int32 ChannelIndex = CRParamSection->ControlChannelMap[ControlElement->GetName()].ChannelIndex;
+                    TUniquePtr<FTrail> CurTrail = MakeUnique<FMovieSceneControlTransformTrail>(FLinearColor::White, false, CRParamSection, Sequencer, ChannelIndex, ControlElement->GetName());
+                    if (AllTrails.Contains(ControlMap[ControlElement->GetName()])) 
+                    {
+                        AllTrails.Remove(ControlMap[ControlElement->GetName()]);
+                    }
 
-	DelegateHandles.OnControlRemovedHandle = RigHierarchy->OnElementRemoved.AddLambda(
-		[this, Component](FRigHierarchyContainer*, const FRigElementKey& ElemKey) {
-			TSharedPtr<ISequencer> Sequencer = WeakSequencer.Pin();
-			check(Sequencer);
-			if (ElemKey.Type != ERigElementType::Control || !ControlsTracked.Contains(Component) || !ControlsTracked[Component].Contains(ElemKey.Name)) // We only care about controls
-			{
-				return;
+                    AddTrail(ControlMap[ControlElement->GetName()], Hierarchy[ControlMap[ControlElement->GetName()]], MoveTemp(CurTrail));
+                }
 			}
-			
-			const FGuid TrailGuid = ControlsTracked[Component][ElemKey.Name];
-			RemoveTrail(TrailGuid);
-		}
-	);
-
-	DelegateHandles.OnControlReparentedHandle = RigHierarchy->OnElementReparented.AddLambda(
-		[this, Component](FRigHierarchyContainer*, const FRigElementKey& ElemKey, const FName& OldParent, const FName& NewParent) {
-			TSharedPtr<ISequencer> Sequencer = WeakSequencer.Pin();
-			check(Sequencer);
-			if (ElemKey.Type != ERigElementType::Control || !ControlsTracked.Contains(Component) || 
-				!ControlsTracked[Component].Contains(ElemKey.Name)) // We only care about controls
+			else if(InNotif == ERigHierarchyNotification::ElementRemoved)
 			{
-				return;
-			}
-
-			const FGuid ElemGuid = ControlsTracked[Component][ElemKey.Name];
-			const FGuid OldParentGuid = ControlsTracked[Component].Contains(OldParent) ? ControlsTracked[Component][OldParent] : 
-				BonesTracked[Component].Contains(OldParent) ? BonesTracked[Component][OldParent] : FGuid();
-			if (OldParentGuid.IsValid())
-			{
-				Hierarchy[OldParentGuid].Children.Remove(ElemGuid);
-				Hierarchy[ElemGuid].Parents.Remove(OldParentGuid);
-			}
-
-			// Assuming parent is already in hierarchy
-			const FGuid NewParentGuid = ControlsTracked[Component].Contains(NewParent) ? ControlsTracked[Component][NewParent] :
-				BonesTracked[Component].Contains(NewParent) ? BonesTracked[Component][NewParent] : FGuid();
-			check(NewParentGuid.IsValid());
-			if (!Hierarchy[NewParentGuid].Children.Contains(ElemGuid))
-			{
-				Hierarchy[NewParentGuid].Children.Add(ElemGuid);
-			}
-			if (!Hierarchy[ElemGuid].Parents.Contains(NewParentGuid))
-			{
-				Hierarchy[ElemGuid].Parents.Add(NewParentGuid);
+				if (!ControlsTracked.Contains(Component) || !ControlsTracked[Component].Contains(ControlElement->GetName())) // We only care about controls
+				{
+					return;
+				}
+				
+				const FGuid TrailGuid = ControlsTracked[Component][ControlElement->GetName()];
+				RemoveTrail(TrailGuid);
 			}
-		}
-	);
+            else if(InNotif == ERigHierarchyNotification::ParentChanged)
+            {
+				if (!ControlsTracked.Contains(Component) || 
+					!ControlsTracked[Component].Contains(ControlElement->GetName())) // We only care about controls
+				{
+					return;
+				}
 
-	DelegateHandles.OnControlRenamedHandle = RigHierarchy->OnElementRenamed.AddLambda(
-		[this, Component](FRigHierarchyContainer*, ERigElementType ElemType , const FName& OldName, const FName& NewName) {
-			if (ElemType != ERigElementType::Control || !ControlsTracked.Contains(Component) || !ControlsTracked[Component].Contains(OldName))
-			{
-				return;
-			}
+				const FGuid ElemGuid = ControlsTracked[Component][ControlElement->GetName()];
+            	for(TPair<FGuid, FTrailHierarchyNode>& Pair : Hierarchy)
+            	{
+					Pair.Value.Children.Remove(ElemGuid);
+            	}
+            	Hierarchy[ElemGuid].Parents.Reset();
+
+            	const FRigElementKey NewParentKey = InHierarchy->GetFirstParent(ControlElement->GetKey());
+            	if(NewParentKey.IsValid())
+            	{
+            		// Assuming parent is already in hierarchy
+                    const FGuid NewParentGuid = ControlsTracked[Component].Contains(NewParentKey.Name) ? ControlsTracked[Component][NewParentKey.Name] :
+                        BonesTracked[Component].Contains(NewParentKey.Name) ? BonesTracked[Component][NewParentKey.Name] : FGuid();
+                    check(NewParentGuid.IsValid());
+            		
+                    if (!Hierarchy[NewParentGuid].Children.Contains(ElemGuid))
+                    {
+                        Hierarchy[NewParentGuid].Children.Add(ElemGuid);
+                    }
+                    if (!Hierarchy[ElemGuid].Parents.Contains(NewParentGuid))
+                    {
+                        Hierarchy[ElemGuid].Parents.Add(NewParentGuid);
+                    }
+            	}
+            }
+            else if(InNotif == ERigHierarchyNotification::ElementRenamed)
+            {
+            	const FName OldName = InHierarchy->GetPreviousName(ControlElement->GetKey());
+
+            	if (!ControlsTracked.Contains(Component) || !ControlsTracked[Component].Contains(OldName))
+				{
+					return;
+				}
 
-			const FGuid TempTrailGuid = ControlsTracked[Component][OldName];
-			ControlsTracked[Component].Remove(OldName);
-			ControlsTracked[Component].Add(NewName, TempTrailGuid);
-		}
-	);
+				const FGuid TempTrailGuid = ControlsTracked[Component][OldName];
+				ControlsTracked[Component].Remove(OldName);
+				ControlsTracked[Component].Add(ControlElement->GetName(), TempTrailGuid);
+            }
+        }
+    );
 }
 
 } // namespace MovieScene
diff --git a/Engine/Plugins/Experimental/MotionTrailEditorMode/Source/MotionTrailEditorMode/Private/Sequencer/SequencerTrailHierarchy.h b/Engine/Plugins/Experimental/MotionTrailEditorMode/Source/MotionTrailEditorMode/Private/Sequencer/SequencerTrailHierarchy.h
--- a/Engine/Plugins/Experimental/MotionTrailEditorMode/Source/MotionTrailEditorMode/Private/Sequencer/SequencerTrailHierarchy.h
+++ b/Engine/Plugins/Experimental/MotionTrailEditorMode/Source/MotionTrailEditorMode/Private/Sequencer/SequencerTrailHierarchy.h
@@ -72,7 +72,7 @@ private:
 	void ResolveComponentToRoot(USceneComponent* Component);
 	void AddComponentToHierarchy(USceneComponent* CompToAdd, UMovieScene3DTransformTrack* TransformTrack);
 	void AddSkeletonToHierarchy(USkeletalMeshComponent* CompToAdd);
-	void ResolveRigElementToRootComponent(FRigHierarchyContainer* RigHierarchy, FRigElementKey InElementKey, USkeletalMeshComponent* Component);
+	void ResolveRigElementToRootComponent(URigHierarchy* RigHierarchy, FRigElementKey InElementKey, USkeletalMeshComponent* Component);
 	void AddControlsToHierarchy(USkeletalMeshComponent* CompToAdd, UMovieSceneControlRigParameterTrack* CRParamTrack);
 
 	void RegisterControlRigDelegates(USkeletalMeshComponent* Component, UMovieSceneControlRigParameterSection* CRParamSection);
@@ -93,10 +93,7 @@ private:
 
 	struct FControlRigDelegateHandles
 	{
-		FDelegateHandle OnControlAddedHandle;
-		FDelegateHandle OnControlRemovedHandle;
-		FDelegateHandle OnControlReparentedHandle;
-		FDelegateHandle OnControlRenamedHandle;
+		FDelegateHandle OnHierarchyModified;
 	};
 	TMap<UMovieSceneSection*, FControlRigDelegateHandles> ControlRigDelegateHandles;
 };
