Control Rig: New rig hierarchy data structure refactoring - follow up fix for sequencer

diff --git a/Engine/Plugins/Experimental/ControlRig/Source/ControlRig/Private/ControlRig.cpp b/Engine/Plugins/Experimental/ControlRig/Source/ControlRig/Private/ControlRig.cpp
--- a/Engine/Plugins/Experimental/ControlRig/Source/ControlRig/Private/ControlRig.cpp
+++ b/Engine/Plugins/Experimental/ControlRig/Source/ControlRig/Private/ControlRig.cpp
@@ -1664,26 +1664,7 @@ void UControlRig::GetControlsInOrder(TArray<FRigControlElement*>& SortedControls
 		return;
 	}
 
-	TArray<bool> ElementVisited;
-	ElementVisited.AddZeroed(DynamicHierarchy->Num());
-		
-	for (FRigBaseElement* Element : *DynamicHierarchy)
-	{
-		DynamicHierarchy->Traverse(Element, true, [&ElementVisited, &SortedControls](FRigBaseElement* InElement, bool& bContinue)
-        {
-            bContinue = !ElementVisited[InElement->GetIndex()];
-
-            if(bContinue)
-            {
-            	if(FRigControlElement* ControlElement = Cast<FRigControlElement>(InElement))
-            	{
-            		SortedControls.Add(ControlElement);
-            	}
-            	
-                ElementVisited[InElement->GetIndex()] = true;
-            }
-        });
-	}
+	SortedControls = DynamicHierarchy->GetControls(true);
 }
 
 const FRigInfluenceMap* UControlRig::FindInfluenceMap(const FName& InEventName)
diff --git a/Engine/Plugins/Experimental/ControlRig/Source/ControlRig/Public/Rigs/RigHierarchy.h b/Engine/Plugins/Experimental/ControlRig/Source/ControlRig/Public/Rigs/RigHierarchy.h
--- a/Engine/Plugins/Experimental/ControlRig/Source/ControlRig/Public/Rigs/RigHierarchy.h
+++ b/Engine/Plugins/Experimental/ControlRig/Source/ControlRig/Public/Rigs/RigHierarchy.h
@@ -395,24 +395,20 @@ public:
 		{
 			TArray<bool> ElementVisited;
 			ElementVisited.AddZeroed(Elements.Num());
-		
-			for (int32 ElementIndex = 0; ElementIndex < Elements.Num(); ElementIndex++)
-			{
-				FRigBaseElement* Element = Elements[ElementIndex];
-				Traverse(Element, true, [&ElementVisited, &Results](FRigBaseElement* InElement, bool& bContinue)
-                {
-                    bContinue = !ElementVisited[InElement->GetIndex()];
 
-                    if(bContinue)
-                    {
-                        if(T* CastElement = Cast<T>(InElement))
-                        {
-                            Results.Add(CastElement);
-                        }
-                        ElementVisited[InElement->GetIndex()] = true;
-                    }
-                });
-			}
+			Traverse([&ElementVisited, &Results](FRigBaseElement* InElement, bool& bContinue)
+			{
+			    bContinue = !ElementVisited[InElement->GetIndex()];
+
+			    if(bContinue)
+			    {
+			        if(T* CastElement = Cast<T>(InElement))
+			        {
+			            Results.Add(CastElement);
+			        }
+			        ElementVisited[InElement->GetIndex()] = true;
+			    }
+			});
 		}
 		else
 		{
@@ -458,26 +454,22 @@ public:
 			TArray<bool> ElementVisited;
 			ElementVisited.AddZeroed(Elements.Num());
 		
-			for (int32 ElementIndex = 0; ElementIndex < Elements.Num(); ElementIndex++)
-			{
-				FRigBaseElement* Element = Elements[ElementIndex];
-				Traverse(Element, true, [&ElementVisited, &Results, InKeepElementFunction](FRigBaseElement* InElement, bool& bContinue)
-                {
-                    bContinue = !ElementVisited[InElement->GetIndex()];
+			Traverse([&ElementVisited, &Results, InKeepElementFunction](FRigBaseElement* InElement, bool& bContinue)
+            {
+                bContinue = !ElementVisited[InElement->GetIndex()];
 
-                    if(bContinue)
+                if(bContinue)
+                {
+                    if(T* CastElement = Cast<T>(InElement))
                     {
-                        if(T* CastElement = Cast<T>(InElement))
-                        {
-							if(InKeepElementFunction(CastElement))
-							{
-								Results.Add(CastElement);
-							}
-                        }
-                        ElementVisited[InElement->GetIndex()] = true;
+						if(InKeepElementFunction(CastElement))
+						{
+							Results.Add(CastElement);
+						}
                     }
-                });
-			}
+                    ElementVisited[InElement->GetIndex()] = true;
+                }
+            });
 		}
 		else
 		{
diff --git a/Engine/Plugins/Experimental/ControlRig/Source/ControlRigEditor/Private/EditMode/ControlRigControlsProxy.cpp b/Engine/Plugins/Experimental/ControlRig/Source/ControlRigEditor/Private/EditMode/ControlRigControlsProxy.cpp
--- a/Engine/Plugins/Experimental/ControlRig/Source/ControlRigEditor/Private/EditMode/ControlRigControlsProxy.cpp
+++ b/Engine/Plugins/Experimental/ControlRig/Source/ControlRigEditor/Private/EditMode/ControlRigControlsProxy.cpp
@@ -73,6 +73,7 @@ void FControlRigEnumControlProxyValueDetails::OnEnumValueChanged(int32 InValue,
 
 void UControlRigControlsProxy::SelectionChanged(bool bInSelected)
 {
+	FRigControlElement* ControlElement = GetControlElement();
 	if (ControlElement)
 	{
 		Modify();
@@ -101,6 +102,7 @@ void UControlRigControlsProxy::PostEditChangeProperty(struct FPropertyChangedEve
 	Super::PostEditChangeProperty(PropertyChangedEvent);
 	if (PropertyChangedEvent.Property && PropertyChangedEvent.Property->GetFName() == GET_MEMBER_NAME_CHECKED(UControlRigControlsProxy, bSelected))
 	{
+		FRigControlElement* ControlElement = GetControlElement();
 		if (ControlElement && ControlRig.IsValid())
 		{
 			if (bSelected)
@@ -114,10 +116,10 @@ void UControlRigControlsProxy::PostEditChangeProperty(struct FPropertyChangedEve
 	}
 }
 
-
 #if WITH_EDITOR
 void UControlRigControlsProxy::PostEditUndo()
 {
+	FRigControlElement* ControlElement = GetControlElement();
 	if (ControlElement && ControlRig.IsValid() && ControlRig->GetHierarchy()->Contains(FRigElementKey(ControlName, ERigElementType::Control)))
 	{
 		if (bSelected)
@@ -129,6 +131,14 @@ void UControlRigControlsProxy::PostEditUndo()
 }
 #endif
 
+FRigControlElement* UControlRigControlsProxy::GetControlElement() const
+{
+	if(ControlRig.IsValid())
+	{
+		return ControlRig->GetHierarchy()->Find<FRigControlElement>(FRigElementKey(ControlName, ERigElementType::Control));
+	}
+	return nullptr;
+}
 
 void UControlRigTransformControlProxy::PostEditChangeProperty(struct FPropertyChangedEvent& PropertyChangedEvent)
 {
@@ -136,6 +146,7 @@ void UControlRigTransformControlProxy::PostEditChangeProperty(struct FPropertyCh
 	if ((PropertyChangedEvent.Property && PropertyChangedEvent.Property->GetFName() == GET_MEMBER_NAME_CHECKED(UControlRigTransformControlProxy, Transform))
 		|| (PropertyChangedEvent.MemberProperty && PropertyChangedEvent.MemberProperty->GetFName() == GET_MEMBER_NAME_CHECKED(UControlRigTransformControlProxy, Transform)))
 	{
+		FRigControlElement* ControlElement = GetControlElement();
 		if (ControlElement && ControlRig.IsValid())
 		{
 			//MUST set through ControlRig
@@ -149,6 +160,7 @@ void UControlRigTransformControlProxy::PostEditChangeProperty(struct FPropertyCh
 
 void UControlRigTransformControlProxy::ValueChanged()
 {
+	FRigControlElement* ControlElement = GetControlElement();
 	if (ControlElement && ControlRig.IsValid())
 	{
 		Modify();
@@ -162,6 +174,7 @@ void UControlRigTransformControlProxy::ValueChanged()
 #if WITH_EDITOR
 void UControlRigTransformControlProxy::PostEditUndo()
 {
+	FRigControlElement* ControlElement = GetControlElement();
 	if (ControlElement && ControlRig.IsValid() && ControlRig->GetHierarchy()->Contains(FRigElementKey(ControlName, ERigElementType::Control)))
 	{
 		if (bSelected)
@@ -176,7 +189,7 @@ void UControlRigTransformControlProxy::PostEditUndo()
 
 void UControlRigTransformControlProxy::SetKey(const IPropertyHandle& KeyedPropertyHandle)
 {
-
+	FRigControlElement* ControlElement = GetControlElement();
 	if (ControlElement)
 	{
 		ControlRig->SetControlValue<FTransform>(ControlName, Transform, true, EControlRigSetKey::Always);
@@ -189,6 +202,7 @@ void UControlRigTransformNoScaleControlProxy::PostEditChangeProperty(struct FPro
 	if ((PropertyChangedEvent.Property && PropertyChangedEvent.Property->GetFName() == GET_MEMBER_NAME_CHECKED(UControlRigTransformNoScaleControlProxy, Transform))
 		|| (PropertyChangedEvent.MemberProperty && PropertyChangedEvent.MemberProperty->GetFName() == GET_MEMBER_NAME_CHECKED(UControlRigTransformNoScaleControlProxy, Transform)))
 	{
+		FRigControlElement* ControlElement = GetControlElement();
 		if (ControlElement && ControlRig.IsValid())
 		{
 			//MUST set through ControlRig
@@ -202,6 +216,7 @@ void UControlRigTransformNoScaleControlProxy::PostEditChangeProperty(struct FPro
 
 void UControlRigTransformNoScaleControlProxy::ValueChanged()
 {
+	FRigControlElement* ControlElement = GetControlElement();
 	if (ControlElement)
 	{
 		Modify();
@@ -215,6 +230,7 @@ void UControlRigTransformNoScaleControlProxy::ValueChanged()
 #if WITH_EDITOR
 void UControlRigTransformNoScaleControlProxy::PostEditUndo()
 {
+	FRigControlElement* ControlElement = GetControlElement();
 	if (ControlElement && ControlRig.IsValid() && ControlRig->GetHierarchy()->Contains(FRigElementKey(ControlName, ERigElementType::Control)))
 	{
 		if (bSelected)
@@ -229,6 +245,7 @@ void UControlRigTransformNoScaleControlProxy::PostEditUndo()
 
 void UControlRigTransformNoScaleControlProxy::SetKey(const IPropertyHandle& KeyedPropertyHandle)
 {
+	FRigControlElement* ControlElement = GetControlElement();
 	if (ControlElement)
 	{
 		ControlRig->SetControlValue<FTransformNoScale>(ControlName, Transform, true, EControlRigSetKey::Always);
@@ -242,6 +259,7 @@ void UControlRigEulerTransformControlProxy::PostEditChangeProperty(struct FPrope
 	if ((PropertyChangedEvent.Property && PropertyChangedEvent.Property->GetFName() == GET_MEMBER_NAME_CHECKED(UControlRigEulerTransformControlProxy, Transform))
 		|| (PropertyChangedEvent.MemberProperty && PropertyChangedEvent.MemberProperty->GetFName() == GET_MEMBER_NAME_CHECKED(UControlRigEulerTransformControlProxy, Transform)))
 	{
+		FRigControlElement* ControlElement = GetControlElement();
 		if (ControlElement && ControlRig.IsValid())
 		{
 			//MUST set through ControlRig
@@ -255,6 +273,7 @@ void UControlRigEulerTransformControlProxy::PostEditChangeProperty(struct FPrope
 
 void UControlRigEulerTransformControlProxy::ValueChanged()
 {
+	FRigControlElement* ControlElement = GetControlElement();
 	if (ControlElement)
 	{
 		Modify();
@@ -270,6 +289,7 @@ void UControlRigEulerTransformControlProxy::ValueChanged()
 #if WITH_EDITOR
 void UControlRigEulerTransformControlProxy::PostEditUndo()
 {
+	FRigControlElement* ControlElement = GetControlElement();
 	if (ControlElement && ControlRig.IsValid() && ControlRig->GetHierarchy()->Contains(FRigElementKey(ControlName, ERigElementType::Control)))
 	{
 		if (bSelected)
@@ -284,6 +304,7 @@ void UControlRigEulerTransformControlProxy::PostEditUndo()
 
 void UControlRigEulerTransformControlProxy::SetKey(const IPropertyHandle& KeyedPropertyHandle)
 {
+	FRigControlElement* ControlElement = GetControlElement();
 	if (ControlElement)
 	{
 		ControlRig->SetControlValue<FEulerTransform>(ControlName, Transform, true, EControlRigSetKey::Always);
@@ -295,6 +316,7 @@ void UControlRigFloatControlProxy::PostEditChangeProperty(struct FPropertyChange
 	Super::PostEditChangeProperty(PropertyChangedEvent);
 	if (PropertyChangedEvent.Property && PropertyChangedEvent.Property->GetFName() == GET_MEMBER_NAME_CHECKED(UControlRigFloatControlProxy, Float))
 	{
+		FRigControlElement* ControlElement = GetControlElement();
 		if (ControlElement && ControlRig.IsValid())
 		{
 			//MUST set through ControlRig
@@ -308,6 +330,7 @@ void UControlRigFloatControlProxy::PostEditChangeProperty(struct FPropertyChange
 
 void UControlRigFloatControlProxy::ValueChanged()
 {
+	FRigControlElement* ControlElement = GetControlElement();
 	if (ControlElement)
 	{
 		Modify();
@@ -321,6 +344,7 @@ void UControlRigFloatControlProxy::ValueChanged()
 #if WITH_EDITOR
 void UControlRigFloatControlProxy::PostEditUndo()
 {
+	FRigControlElement* ControlElement = GetControlElement();
 	if (ControlElement && ControlRig.IsValid() && ControlRig->GetHierarchy()->Contains(FRigElementKey(ControlName, ERigElementType::Control)))
 	{
 		if (bSelected)
@@ -336,6 +360,7 @@ void UControlRigFloatControlProxy::PostEditUndo()
 
 void UControlRigFloatControlProxy::SetKey(const IPropertyHandle& KeyedPropertyHandle)
 {
+	FRigControlElement* ControlElement = GetControlElement();
 	if (ControlElement)
 	{
 		ControlRig->SetControlValue<float>(ControlName, Float, true, EControlRigSetKey::Always);
@@ -347,6 +372,7 @@ void UControlRigIntegerControlProxy::PostEditChangeProperty(struct FPropertyChan
 	Super::PostEditChangeProperty(PropertyChangedEvent);
 	if (PropertyChangedEvent.Property && PropertyChangedEvent.Property->GetFName() == GET_MEMBER_NAME_CHECKED(UControlRigIntegerControlProxy, Integer))
 	{
+		FRigControlElement* ControlElement = GetControlElement();
 		if (ControlElement && ControlRig.IsValid())
 		{
 			//MUST set through ControlRig
@@ -359,6 +385,7 @@ void UControlRigIntegerControlProxy::PostEditChangeProperty(struct FPropertyChan
 
 void UControlRigIntegerControlProxy::ValueChanged()
 {
+	FRigControlElement* ControlElement = GetControlElement();
 	if (ControlElement)
 	{
 		Modify();
@@ -372,6 +399,7 @@ void UControlRigIntegerControlProxy::ValueChanged()
 #if WITH_EDITOR
 void UControlRigIntegerControlProxy::PostEditUndo()
 {
+	FRigControlElement* ControlElement = GetControlElement();
 	if (ControlElement && ControlRig.IsValid() && ControlRig->GetHierarchy()->Contains(FRigElementKey(ControlName, ERigElementType::Control)))
 	{
 		if (bSelected)
@@ -386,6 +414,7 @@ void UControlRigIntegerControlProxy::PostEditUndo()
 
 void UControlRigIntegerControlProxy::SetKey(const IPropertyHandle& KeyedPropertyHandle)
 {
+	FRigControlElement* ControlElement = GetControlElement();
 	if (ControlElement)
 	{
 		ControlRig->SetControlValue<int32>(ControlName, Integer, true, EControlRigSetKey::Always);
@@ -397,6 +426,7 @@ void UControlRigEnumControlProxy::PostEditChangeProperty(struct FPropertyChanged
 	Super::PostEditChangeProperty(PropertyChangedEvent);
 	if (PropertyChangedEvent.Property && PropertyChangedEvent.Property->GetFName() == GET_MEMBER_NAME_CHECKED(UControlRigEnumControlProxy, Enum))
 	{
+		FRigControlElement* ControlElement = GetControlElement();
 		if (ControlElement && ControlRig.IsValid())
 		{
 			//MUST set through ControlRig
@@ -410,6 +440,7 @@ void UControlRigEnumControlProxy::PostEditChangeProperty(struct FPropertyChanged
 
 void UControlRigEnumControlProxy::ValueChanged()
 {
+	FRigControlElement* ControlElement = GetControlElement();
 	if (ControlElement)
 	{
 		Modify();
@@ -427,6 +458,7 @@ void UControlRigEnumControlProxy::ValueChanged()
 #if WITH_EDITOR
 void UControlRigEnumControlProxy::PostEditUndo()
 {
+	FRigControlElement* ControlElement = GetControlElement();
 	if (ControlElement && ControlRig.IsValid() && ControlRig->GetHierarchy()->Contains(FRigElementKey(ControlName, ERigElementType::Control)))
 	{
 		if (bSelected)
@@ -441,6 +473,7 @@ void UControlRigEnumControlProxy::PostEditUndo()
 
 void UControlRigEnumControlProxy::SetKey(const IPropertyHandle& KeyedPropertyHandle)
 {
+	FRigControlElement* ControlElement = GetControlElement();
 	if (ControlElement)
 	{
 		ControlRig->SetControlValue<int32>(ControlName, Enum.EnumIndex, true, EControlRigSetKey::Always);
@@ -453,6 +486,7 @@ void UControlRigVectorControlProxy::PostEditChangeProperty(struct FPropertyChang
 	if ((PropertyChangedEvent.Property && PropertyChangedEvent.Property->GetFName() == GET_MEMBER_NAME_CHECKED(UControlRigVectorControlProxy, Vector))
 		|| (PropertyChangedEvent.MemberProperty && PropertyChangedEvent.MemberProperty->GetFName() == GET_MEMBER_NAME_CHECKED(UControlRigVectorControlProxy, Vector)))
 	{
+		FRigControlElement* ControlElement = GetControlElement();
 		if (ControlElement && ControlRig.IsValid())
 		{
 			//MUST set through ControlRig
@@ -466,6 +500,7 @@ void UControlRigVectorControlProxy::PostEditChangeProperty(struct FPropertyChang
 
 void UControlRigVectorControlProxy::ValueChanged()
 {
+	FRigControlElement* ControlElement = GetControlElement();
 	if (ControlElement)
 	{
 		Modify();
@@ -479,6 +514,7 @@ void UControlRigVectorControlProxy::ValueChanged()
 #if WITH_EDITOR
 void UControlRigVectorControlProxy::PostEditUndo()
 {
+	FRigControlElement* ControlElement = GetControlElement();
 	if (ControlElement && ControlRig.IsValid() && ControlRig->GetHierarchy()->Contains(FRigElementKey(ControlName, ERigElementType::Control)))
 	{
 		if (bSelected)
@@ -494,6 +530,7 @@ void UControlRigVectorControlProxy::PostEditUndo()
 
 void UControlRigVectorControlProxy::SetKey(const IPropertyHandle& KeyedPropertyHandle)
 {
+	FRigControlElement* ControlElement = GetControlElement();
 	if (ControlElement)
 	{
 		ControlRig->SetControlValue<FVector>(ControlName, Vector, true, EControlRigSetKey::Always);
@@ -506,6 +543,7 @@ void UControlRigVector2DControlProxy::PostEditChangeProperty(struct FPropertyCha
 	if ((PropertyChangedEvent.Property && PropertyChangedEvent.Property->GetFName() == GET_MEMBER_NAME_CHECKED(UControlRigVector2DControlProxy, Vector2D))
 		|| ((PropertyChangedEvent.MemberProperty && PropertyChangedEvent.MemberProperty->GetFName() == GET_MEMBER_NAME_CHECKED(UControlRigVector2DControlProxy, Vector2D))))
 	{
+		FRigControlElement* ControlElement = GetControlElement();
 		if (ControlElement && ControlRig.IsValid())
 		{
 			//MUST set through ControlRig
@@ -519,6 +557,7 @@ void UControlRigVector2DControlProxy::PostEditChangeProperty(struct FPropertyCha
 
 void UControlRigVector2DControlProxy::ValueChanged()
 {
+	FRigControlElement* ControlElement = GetControlElement();
 	if (ControlElement)
 	{
 		Modify();
@@ -532,6 +571,7 @@ void UControlRigVector2DControlProxy::ValueChanged()
 #if WITH_EDITOR
 void UControlRigVector2DControlProxy::PostEditUndo()
 {
+	FRigControlElement* ControlElement = GetControlElement();
 	if (ControlElement && ControlRig.IsValid() && ControlRig->GetHierarchy()->Contains(FRigElementKey(ControlName, ERigElementType::Control)))
 	{
 		if (bSelected)
@@ -547,6 +587,7 @@ void UControlRigVector2DControlProxy::PostEditUndo()
 
 void UControlRigVector2DControlProxy::SetKey(const IPropertyHandle& KeyedPropertyHandle)
 {
+	FRigControlElement* ControlElement = GetControlElement();
 	if (ControlElement)
 	{
 		ControlRig->SetControlValue<FVector2D>(ControlName, Vector2D, true, EControlRigSetKey::Always);
@@ -559,6 +600,7 @@ void UControlRigBoolControlProxy::PostEditChangeProperty(struct FPropertyChanged
 	Super::PostEditChangeProperty(PropertyChangedEvent);
 	if (PropertyChangedEvent.Property && PropertyChangedEvent.Property->GetFName() == GET_MEMBER_NAME_CHECKED(UControlRigBoolControlProxy, Bool))
 	{
+		FRigControlElement* ControlElement = GetControlElement();
 		if (ControlElement && ControlRig.IsValid())
 		{
 			//MUST set through ControlRig
@@ -571,6 +613,7 @@ void UControlRigBoolControlProxy::PostEditChangeProperty(struct FPropertyChanged
 
 void UControlRigBoolControlProxy::ValueChanged()
 {
+	FRigControlElement* ControlElement = GetControlElement();
 	if (ControlElement)
 	{
 		Modify();
@@ -584,6 +627,7 @@ void UControlRigBoolControlProxy::ValueChanged()
 #if WITH_EDITOR
 void UControlRigBoolControlProxy::PostEditUndo()
 {
+	FRigControlElement* ControlElement = GetControlElement();
 	if (ControlElement && ControlRig.IsValid() && ControlRig->GetHierarchy()->Contains(FRigElementKey(ControlName, ERigElementType::Control)))
 	{
 		if (bSelected)
@@ -599,6 +643,7 @@ void UControlRigBoolControlProxy::PostEditUndo()
 
 void UControlRigBoolControlProxy::SetKey(const IPropertyHandle& KeyedPropertyHandle)
 {
+	FRigControlElement* ControlElement = GetControlElement();
 	if (ControlElement)
 	{
 		ControlRig->SetControlValue<bool>(ControlName, Bool, true, EControlRigSetKey::Always);
@@ -690,7 +735,6 @@ void UControlRigDetailPanelControlProxies::AddProxy(const FName& Name, UControlR
 		{
 			Proxy->SetFlags(RF_Transactional);
 			Proxy->SetName(Name);
-			Proxy->ControlElement = ControlElement;
 			Proxy->ControlRig = ControlRig;
 			Proxy->ValueChanged();
 			AllProxies.Add(Name, Proxy);
diff --git a/Engine/Plugins/Experimental/ControlRig/Source/ControlRigEditor/Private/EditMode/ControlRigControlsProxy.h b/Engine/Plugins/Experimental/ControlRig/Source/ControlRigEditor/Private/EditMode/ControlRigControlsProxy.h
--- a/Engine/Plugins/Experimental/ControlRig/Source/ControlRigEditor/Private/EditMode/ControlRigControlsProxy.h
+++ b/Engine/Plugins/Experimental/ControlRig/Source/ControlRigEditor/Private/EditMode/ControlRigControlsProxy.h
@@ -38,7 +38,7 @@ public:
 
 public:
 
-	FRigControlElement* ControlElement;
+	FRigControlElement* GetControlElement() const;
 	TWeakObjectPtr<UControlRig> ControlRig;
 
 	UPROPERTY()
diff --git a/Engine/Plugins/Experimental/ControlRig/Source/ControlRigEditor/Private/EditMode/SControlRigEditModeTools.cpp b/Engine/Plugins/Experimental/ControlRig/Source/ControlRigEditor/Private/EditMode/SControlRigEditModeTools.cpp
--- a/Engine/Plugins/Experimental/ControlRig/Source/ControlRigEditor/Private/EditMode/SControlRigEditModeTools.cpp
+++ b/Engine/Plugins/Experimental/ControlRig/Source/ControlRigEditor/Private/EditMode/SControlRigEditModeTools.cpp
@@ -64,14 +64,20 @@ public:
 
 		for (UControlRigControlsProxy* Proxy : ProxiesBeingCustomized)
 		{
+			FRigControlElement* ControlElement = Proxy->GetControlElement();
+			if(ControlElement == nullptr)
+			{
+				continue;
+			}
+			
 			FName ValuePropertyName = TEXT("Transform");
-			if (Proxy->ControlElement->Settings.ControlType == ERigControlType::Float)
+			if (ControlElement->Settings.ControlType == ERigControlType::Float)
 			{
 				ValuePropertyName = TEXT("Float");
 			}
-			else if (Proxy->ControlElement->Settings.ControlType == ERigControlType::Integer)
+			else if (ControlElement->Settings.ControlType == ERigControlType::Integer)
 			{
-				if (Proxy->ControlElement->Settings.ControlEnum == nullptr)
+				if (ControlElement->Settings.ControlEnum == nullptr)
 				{
 					ValuePropertyName = TEXT("Integer");
 				}
@@ -80,16 +86,16 @@ public:
 					ValuePropertyName = TEXT("Enum");
 				}
 			}
-			else if (Proxy->ControlElement->Settings.ControlType == ERigControlType::Bool)
+			else if (ControlElement->Settings.ControlType == ERigControlType::Bool)
 			{
 				ValuePropertyName = TEXT("Bool");
 			}
-			else if (Proxy->ControlElement->Settings.ControlType == ERigControlType::Position ||
-				Proxy->ControlElement->Settings.ControlType == ERigControlType::Scale)
+			else if (ControlElement->Settings.ControlType == ERigControlType::Position ||
+				ControlElement->Settings.ControlType == ERigControlType::Scale)
 			{
 				ValuePropertyName = TEXT("Vector");
 			}
-			else if (Proxy->ControlElement->Settings.ControlType == ERigControlType::Vector2D)
+			else if (ControlElement->Settings.ControlType == ERigControlType::Vector2D)
 			{
 				ValuePropertyName = TEXT("Vector2D");
 			}
@@ -97,7 +103,7 @@ public:
 			TSharedPtr<IPropertyHandle> ValuePropertyHandle = DetailLayout.GetProperty(ValuePropertyName, Proxy->GetClass());
 			if (ValuePropertyHandle)
 			{
-				ValuePropertyHandle->SetPropertyDisplayName(FText::FromName(Proxy->ControlElement->Settings.DisplayName));
+				ValuePropertyHandle->SetPropertyDisplayName(FText::FromName(ControlElement->Settings.DisplayName));
 			}
 
 			URigHierarchy* Hierarchy = Proxy->ControlRig->GetHierarchy();
@@ -110,7 +116,7 @@ public:
 					ParentControlName = ParentControlElement->GetName();
 				}
 				
-				if (ParentControlName == Proxy->ControlElement->GetName())
+				if (ParentControlName == ControlElement->GetName())
 				{
 					if (FControlRigEditMode* EditMode = static_cast<FControlRigEditMode*>(GLevelEditorModeTools().GetActiveMode(FControlRigEditMode::ModeName)))
 					{
