Control Rig: Crash when filtering hierarchy widget

diff --git a/Engine/Plugins/Experimental/ControlRig/Source/ControlRigEditor/Private/Editor/SRigHierarchy.cpp b/Engine/Plugins/Experimental/ControlRig/Source/ControlRigEditor/Private/Editor/SRigHierarchy.cpp
--- a/Engine/Plugins/Experimental/ControlRig/Source/ControlRigEditor/Private/Editor/SRigHierarchy.cpp
+++ b/Engine/Plugins/Experimental/ControlRig/Source/ControlRigEditor/Private/Editor/SRigHierarchy.cpp
@@ -860,11 +860,13 @@ void SRigHierarchy::AddElement(FRigBaseElement* InElement, const bool bIgnoreTex
 
 	if (ElementMap.Contains(InElement->GetKey()))
 	{
-		FRigElementKey ParentKey = Hierarchy->GetFirstParent(InElement->GetKey());
+		const FRigElementKey ParentKey = Hierarchy->GetFirstParent(InElement->GetKey());
 		if (ParentKey.IsValid())
 		{
-			check(ElementMap.Contains(ParentKey));
-			ReparentElement(InElement->GetKey(), ParentKey);
+			if(ElementMap.Contains(ParentKey))
+			{
+				ReparentElement(InElement->GetKey(), ParentKey);
+			}
 		}
 	}
 }
