Control Rig: Hide pose / offset / gizmo properties for Bool Controls

diff --git a/Engine/Plugins/Experimental/ControlRig/Source/ControlRigEditor/Private/ControlRigElementDetails.cpp b/Engine/Plugins/Experimental/ControlRig/Source/ControlRigEditor/Private/ControlRigElementDetails.cpp
--- a/Engine/Plugins/Experimental/ControlRig/Source/ControlRigEditor/Private/ControlRigElementDetails.cpp
+++ b/Engine/Plugins/Experimental/ControlRig/Source/ControlRigEditor/Private/ControlRigElementDetails.cpp
@@ -1070,9 +1070,6 @@ void FRigControlElementDetails::CustomizeDetails(IDetailLayoutBuilder& DetailBui
 {
 	FRigTransformElementDetails::CustomizeDetails(DetailBuilder);
 
-	//DetailBuilder.HideProperty(TEXT("Offset"));
-	//DetailBuilder.HideProperty(TEXT("Gizmo"));
-
 	GizmoNameList.Reset();
 	if (BlueprintBeingCustomized)
 	{
@@ -1264,7 +1261,13 @@ void FRigControlElementDetails::CustomizeDetails(IDetailLayoutBuilder& DetailBui
 		}
 	};
 
-	if(ControlElement->Settings.ControlType != ERigControlType::Bool)
+	if(ControlElement->Settings.ControlType == ERigControlType::Bool)
+	{
+		DetailBuilder.HideCategory(TEXT("Pose"));
+		DetailBuilder.HideProperty(TEXT("Offset"));
+		DetailBuilder.HideProperty(TEXT("Gizmo"));
+	}
+	else
 	{
 		// setup offset transform
 		{
