Control Rig: Fix name map cache access during compile time

diff --git a/Engine/Plugins/Experimental/ControlRig/Source/ControlRigDeveloper/Private/ControlRigBlueprint.cpp b/Engine/Plugins/Experimental/ControlRig/Source/ControlRigDeveloper/Private/ControlRigBlueprint.cpp
--- a/Engine/Plugins/Experimental/ControlRig/Source/ControlRigDeveloper/Private/ControlRigBlueprint.cpp
+++ b/Engine/Plugins/Experimental/ControlRig/Source/ControlRigDeveloper/Private/ControlRigBlueprint.cpp
@@ -445,12 +445,14 @@ void UControlRigBlueprint::RecompileVM()
 		}
 		CDO->VM->Reset();
 
+		FRigNameCache TempNameCache;
 		FRigUnitContext InitContext;
 		InitContext.State = EControlRigState::Init;
 		InitContext.Hierarchy = CDO->DynamicHierarchy;
+		InitContext.NameCache = &TempNameCache;
 
 		FRigUnitContext UpdateContext = InitContext;
-		InitContext.State = EControlRigState::Update;
+		UpdateContext.State = EControlRigState::Update;
 
 		void* InitContextPtr = &InitContext;
 		void* UpdateContextPtr = &UpdateContext;
