Control Rig: Clear details panel when removing node

diff --git a/Engine/Plugins/Experimental/ControlRig/Source/ControlRigEditor/Private/Editor/ControlRigEditor.cpp b/Engine/Plugins/Experimental/ControlRig/Source/ControlRigEditor/Private/Editor/ControlRigEditor.cpp
--- a/Engine/Plugins/Experimental/ControlRig/Source/ControlRigEditor/Private/Editor/ControlRigEditor.cpp
+++ b/Engine/Plugins/Experimental/ControlRig/Source/ControlRigEditor/Private/Editor/ControlRigEditor.cpp
@@ -1900,8 +1900,13 @@ void FControlRigEditor::HandleModifiedEvent(ERigVMGraphNotifType InNotifType, UR
 				if (UEdGraph* EdGraph = ControlRigBlueprint->GetEdGraph(CollapseNode->GetContainedGraph()))
 				{
 					CloseDocumentTab(EdGraph);
+					ClearDetailObject();
 				}
 			}
+			else if(URigVMFunctionReferenceNode* FunctionRefNode = Cast<URigVMFunctionReferenceNode>(InSubject))
+			{
+				ClearDetailObject();
+			}
 			break;
 		}
 		case ERigVMGraphNotifType::NodeSelectionChanged:
