Control Rig: Initialize hierarchy for unit test lazily

diff --git a/Engine/Plugins/Experimental/ControlRig/Source/ControlRig/Public/Units/RigUnitTest.h b/Engine/Plugins/Experimental/ControlRig/Source/ControlRig/Public/Units/RigUnitTest.h
--- a/Engine/Plugins/Experimental/ControlRig/Source/ControlRig/Public/Units/RigUnitTest.h
+++ b/Engine/Plugins/Experimental/ControlRig/Source/ControlRig/Public/Units/RigUnitTest.h
@@ -16,22 +16,32 @@ public:
 		, Hierarchy(nullptr)
 		, Controller(nullptr)
 	{
-		Hierarchy = NewObject<URigHierarchy>();
-		Controller = NewObject<URigHierarchyController>();
-		Controller->SetHierarchy(Hierarchy);
-
-		Hierarchy->AddToRoot();
-		Controller->AddToRoot();
-		
-		Context.Hierarchy = Hierarchy;;
-		ExecuteContext.Hierarchy= Hierarchy;
-		Context.NameCache = &NameCache;
 	}
 
 	~FControlRigUnitTestBase()
 	{
-		Controller->RemoveFromRoot();
-		Hierarchy->RemoveFromRoot();
+		if (Hierarchy)
+		{
+			Controller->RemoveFromRoot();
+			Hierarchy->RemoveFromRoot();
+		}
+	}
+
+	void InitHierarchy()
+	{
+		if (!Hierarchy)
+		{
+			Hierarchy = NewObject<URigHierarchy>();
+			Controller = NewObject<URigHierarchyController>();
+			Controller->SetHierarchy(Hierarchy);
+
+			Hierarchy->AddToRoot();
+			Controller->AddToRoot();
+
+			Context.Hierarchy = Hierarchy;;
+			ExecuteContext.Hierarchy = Hierarchy;
+			Context.NameCache = &NameCache;
+		}
 	}
 
 	URigHierarchy* Hierarchy;
@@ -63,6 +73,7 @@ public:
 		TUnitStruct Unit; \
 		virtual bool RunTest(const FString& Parameters) override \
 		{ \
+			FControlRigUnitTestBase::InitHierarchy(); \
 			Hierarchy->Reset(); \
 			Unit = TUnitStruct(); \
 			return RunControlRigUnitTest(Parameters); \
