ControlRig: Hierarchy switched over to recursive dirty propagation

diff --git a/Engine/Plugins/Experimental/ControlRig/Source/ControlRig/Private/Rigs/RigHierarchy.cpp b/Engine/Plugins/Experimental/ControlRig/Source/ControlRig/Private/Rigs/RigHierarchy.cpp
--- a/Engine/Plugins/Experimental/ControlRig/Source/ControlRig/Private/Rigs/RigHierarchy.cpp
+++ b/Engine/Plugins/Experimental/ControlRig/Source/ControlRig/Private/Rigs/RigHierarchy.cpp
@@ -752,7 +752,7 @@ bool URigHierarchy::SetParentWeight(FRigBaseElement* InChild, int32 InParentInde
 			TargetWeight = InputWeight;
 			MultiParentElement->Parent.MarkDirty(GlobalType);
 
-			PropagateDirtyFlags(MultiParentElement, LocalType, bAffectChildren);
+			PropagateDirtyFlags(MultiParentElement, ERigTransformType::IsInitial(LocalType), bAffectChildren);
 			return true;
 		}
 	}
@@ -1126,7 +1126,7 @@ void URigHierarchy::SetTransform(FRigTransformElement* InTransformElement, const
 	}
 
 	const FTransform PreviousTransform = GetTransform(InTransformElement, InTransformType);
-	PropagateDirtyFlags(InTransformElement, InTransformType, bAffectChildren);
+	PropagateDirtyFlags(InTransformElement, ERigTransformType::IsInitial(InTransformType), bAffectChildren);
 
 	const ERigTransformType::Type OpposedType = SwapLocalAndGlobal(InTransformType);
 	InTransformElement->Pose.Set(InTransformType, InTransform);
@@ -1196,7 +1196,7 @@ void URigHierarchy::SetControlOffsetTransform(FRigControlElement* InControlEleme
 	}
 
 	const FTransform PreviousTransform = GetControlOffsetTransform(InControlElement, InTransformType);
-	PropagateDirtyFlags(InControlElement, InTransformType, bAffectChildren);
+	PropagateDirtyFlags(InControlElement, ERigTransformType::IsInitial(InTransformType), bAffectChildren);
 
 	GetTransform(InControlElement, MakeLocal(InTransformType));
 	InControlElement->Pose.MarkDirty(MakeGlobal(InTransformType));
@@ -1730,7 +1730,11 @@ void URigHierarchy::UpdateCachedChildren(const FRigBaseElement* InElement, bool
 	InElement->TopologyVersion = TopologyVersion;
 }
 
-void URigHierarchy::PropagateDirtyFlags(FRigTransformElement* InTransformElement, const ERigTransformType::Type InTransformType, bool bAffectChildren) const
+#if URIGHIERARCHY_RECURSIVE_DIRTY_PROPAGATION
+void URigHierarchy::PropagateDirtyFlags(FRigTransformElement* InTransformElement, bool bInitial, bool bAffectChildren, bool bComputeOpposed, bool bMarkDirty) const
+#else
+void URigHierarchy::PropagateDirtyFlags(FRigTransformElement* InTransformElement, bool bInitial, bool bAffectChildren) const
+#endif
 {
 	if(!bEnableDirtyPropagation)
 	{
@@ -1739,59 +1743,84 @@ void URigHierarchy::PropagateDirtyFlags(FRigTransformElement* InTransformElement
 	
 	check(InTransformElement);
 
-	const ERigTransformType::Type LocalType = MakeLocal(InTransformType);
-	const ERigTransformType::Type GlobalType = MakeGlobal(InTransformType);
+	const ERigTransformType::Type LocalType = bInitial ? ERigTransformType::InitialLocal : ERigTransformType::CurrentLocal;
+	const ERigTransformType::Type GlobalType = bInitial ? ERigTransformType::InitialGlobal : ERigTransformType::CurrentGlobal;
 	const ERigTransformType::Type TypeToCompute = bAffectChildren ? LocalType : GlobalType;
 	const ERigTransformType::Type TypeToDirty = SwapLocalAndGlobal(TypeToCompute);
-	
-	for(const FRigTransformElement::FElementToDirty& ElementToDirty : InTransformElement->ElementsToDirty)
+
+#if URIGHIERARCHY_RECURSIVE_DIRTY_PROPAGATION
+	if(bComputeOpposed)
+#endif
 	{
-		if(!bAffectChildren && ElementToDirty.HierarchyDistance > 1)
+		for(const FRigTransformElement::FElementToDirty& ElementToDirty : InTransformElement->ElementsToDirty)
 		{
-			continue;
-		}
+#if URIGHIERARCHY_RECURSIVE_DIRTY_PROPAGATION
+			if(ElementToDirty.Element->Pose.IsDirty(TypeToDirty))
+			{
+				continue;
+			}
+#else
 
-		GetTransform(ElementToDirty.Element, TypeToCompute); // make sure the local / global transform is up 2 date
+			if(!bAffectChildren && ElementToDirty.HierarchyDistance > 1)
+			{
+				continue;
+			}
+
+#endif
+
+			GetTransform(ElementToDirty.Element, TypeToCompute); // make sure the local / global transform is up 2 date
+
+#if URIGHIERARCHY_RECURSIVE_DIRTY_PROPAGATION
+			PropagateDirtyFlags(ElementToDirty.Element, bInitial, bAffectChildren, true, false);
+#endif
+		}
 	}
-	
-	for(const FRigTransformElement::FElementToDirty& ElementToDirty : InTransformElement->ElementsToDirty)
-	{
 #if URIGHIERARCHY_RECURSIVE_DIRTY_PROPAGATION
 
-		if(ElementToDirty.Element->Pose.IsDirty(TypeToDirty))
+	if(bMarkDirty)
+#endif
+	{
+		for(const FRigTransformElement::FElementToDirty& ElementToDirty : InTransformElement->ElementsToDirty)
 		{
-			return;
-		}
-		
+#if URIGHIERARCHY_RECURSIVE_DIRTY_PROPAGATION
+
+			if(ElementToDirty.Element->Pose.IsDirty(TypeToDirty))
+			{
+				continue;
+			}
+			
 #else
 
-		if(!bAffectChildren && ElementToDirty.HierarchyDistance > 1)
-		{
-			return;
-		}
-		
+			if(!bAffectChildren && ElementToDirty.HierarchyDistance > 1)
+			{
+				continue;
+			}
+			
 #endif
 
-		ElementToDirty.Element->Pose.MarkDirty(TypeToDirty);
-	
-		if(FRigMultiParentElement* MultiParentElement = Cast<FRigMultiParentElement>(ElementToDirty.Element))
-		{
-			MultiParentElement->Parent.MarkDirty(GlobalType);
-		}
-		if(FRigControlElement* ControlElement = Cast<FRigControlElement>(ElementToDirty.Element))
-		{
-			ControlElement->Offset.MarkDirty(GlobalType);
-			ControlElement->Gizmo.MarkDirty(GlobalType);
-		}
+			ElementToDirty.Element->Pose.MarkDirty(TypeToDirty);
+		
+			if(FRigMultiParentElement* MultiParentElement = Cast<FRigMultiParentElement>(ElementToDirty.Element))
+			{
+				MultiParentElement->Parent.MarkDirty(GlobalType);
+			}
+			if(FRigControlElement* ControlElement = Cast<FRigControlElement>(ElementToDirty.Element))
+			{
+				ControlElement->Offset.MarkDirty(GlobalType);
+				ControlElement->Gizmo.MarkDirty(GlobalType);
+			}
 
 #if URIGHIERARCHY_RECURSIVE_DIRTY_PROPAGATION
 
-		if(bAffectChildren)
-		{
-			PropagateDirtyFlags(ElementToDirty.Element, InTransformType, bAffectChildren);
-		}
-		
+			if(bAffectChildren)
+			{
+#if URIGHIERARCHY_RECURSIVE_DIRTY_PROPAGATION
+				PropagateDirtyFlags(ElementToDirty.Element, bInitial, bAffectChildren, false, true);
 #endif
+			}
+			
+#endif
+		}
 	}
 }
 
diff --git a/Engine/Plugins/Experimental/ControlRig/Source/ControlRig/Private/Rigs/RigHierarchyController.cpp b/Engine/Plugins/Experimental/ControlRig/Source/ControlRig/Private/Rigs/RigHierarchyController.cpp
--- a/Engine/Plugins/Experimental/ControlRig/Source/ControlRig/Private/Rigs/RigHierarchyController.cpp
+++ b/Engine/Plugins/Experimental/ControlRig/Source/ControlRig/Private/Rigs/RigHierarchyController.cpp
@@ -1219,8 +1219,8 @@ bool URigHierarchyController::AddParent(FRigBaseElement* InChild, FRigBaseElemen
 
 			if(!bMaintainGlobalTransform)
 			{
-				Hierarchy->PropagateDirtyFlags(SingleParentElement, ERigTransformType::CurrentGlobal, true);
-				Hierarchy->PropagateDirtyFlags(SingleParentElement, ERigTransformType::InitialGlobal, true);
+				Hierarchy->PropagateDirtyFlags(SingleParentElement, true, true);
+				Hierarchy->PropagateDirtyFlags(SingleParentElement, false, true);
 			}
 
 			Notify(ERigHierarchyNotification::ParentChanged, SingleParentElement);
@@ -1253,8 +1253,8 @@ bool URigHierarchyController::AddParent(FRigBaseElement* InChild, FRigBaseElemen
 
 			if(!bMaintainGlobalTransform)
 			{
-				Hierarchy->PropagateDirtyFlags(MultiParentElement, ERigTransformType::CurrentGlobal, true);
-				Hierarchy->PropagateDirtyFlags(MultiParentElement, ERigTransformType::InitialGlobal, true);
+				Hierarchy->PropagateDirtyFlags(MultiParentElement, true, true);
+				Hierarchy->PropagateDirtyFlags(MultiParentElement, false, true);
 			}
 
 			Notify(ERigHierarchyNotification::ParentChanged, MultiParentElement);
@@ -1351,8 +1351,8 @@ bool URigHierarchyController::RemoveParent(FRigBaseElement* InChild, FRigBaseEle
 
 			if(!bMaintainGlobalTransform)
 			{
-				Hierarchy->PropagateDirtyFlags(SingleParentElement, ERigTransformType::CurrentGlobal, true);
-				Hierarchy->PropagateDirtyFlags(SingleParentElement, ERigTransformType::InitialGlobal, true);
+				Hierarchy->PropagateDirtyFlags(SingleParentElement, true, true);
+				Hierarchy->PropagateDirtyFlags(SingleParentElement, false, true);
 			}
 
 			Notify(ERigHierarchyNotification::ParentChanged, SingleParentElement);
@@ -1415,8 +1415,8 @@ bool URigHierarchyController::RemoveParent(FRigBaseElement* InChild, FRigBaseEle
 
 			if(!bMaintainGlobalTransform)
 			{
-				Hierarchy->PropagateDirtyFlags(MultiParentElement, ERigTransformType::CurrentGlobal, true);
-				Hierarchy->PropagateDirtyFlags(MultiParentElement, ERigTransformType::InitialGlobal, true);
+				Hierarchy->PropagateDirtyFlags(MultiParentElement, true, true);
+				Hierarchy->PropagateDirtyFlags(MultiParentElement, false, true);
 			}
 
 			Notify(ERigHierarchyNotification::ParentChanged, MultiParentElement);
diff --git a/Engine/Plugins/Experimental/ControlRig/Source/ControlRig/Public/Rigs/RigHierarchy.h b/Engine/Plugins/Experimental/ControlRig/Source/ControlRig/Public/Rigs/RigHierarchy.h
--- a/Engine/Plugins/Experimental/ControlRig/Source/ControlRig/Public/Rigs/RigHierarchy.h
+++ b/Engine/Plugins/Experimental/ControlRig/Source/ControlRig/Public/Rigs/RigHierarchy.h
@@ -11,7 +11,7 @@
 class URigHierarchy;
 class URigHierarchyController;
 
-#define URIGHIERARCHY_RECURSIVE_DIRTY_PROPAGATION 0
+#define URIGHIERARCHY_RECURSIVE_DIRTY_PROPAGATION 1
 
 DECLARE_MULTICAST_DELEGATE_ThreeParams(FRigHierarchyModifiedEvent, ERigHierarchyNotification /* type */, URigHierarchy* /* hierarchy */, const FRigBaseElement* /* element */);
 
@@ -1959,10 +1959,14 @@ private:
 	/**
 	 * Marks all affected elements of a given element as dirty
 	 * @param InTransformElement The element that has changed
-	 * @param InTransformType The type of transform to propagate for
+	 * @param bInitial If true the initial transform will be dirtied
 	 * @param bAffectChildren If set to false children will not move (maintain global).
 	 */
-	void PropagateDirtyFlags(FRigTransformElement* InTransformElement, const ERigTransformType::Type InTransformType, bool bAffectChildren) const;
+#if URIGHIERARCHY_RECURSIVE_DIRTY_PROPAGATION
+	void PropagateDirtyFlags(FRigTransformElement* InTransformElement, bool bInitial, bool bAffectChildren, bool bComputeOpposed = true, bool bMarkDirty = true) const;
+#else
+	void PropagateDirtyFlags(FRigTransformElement* InTransformElement, bool bInitial, bool bAffectChildren) const;
+#endif
 
 	/**
 	 * The topology version of the hierarchy changes when elements are
