ControlRig: Fix CIS ptr check issue SRigCurveContainer

diff --git a/Engine/Plugins/Experimental/ControlRig/Source/ControlRigEditor/Private/Editor/SRigCurveContainer.cpp b/Engine/Plugins/Experimental/ControlRig/Source/ControlRigEditor/Private/Editor/SRigCurveContainer.cpp
--- a/Engine/Plugins/Experimental/ControlRig/Source/ControlRigEditor/Private/Editor/SRigCurveContainer.cpp
+++ b/Engine/Plugins/Experimental/ControlRig/Source/ControlRigEditor/Private/Editor/SRigCurveContainer.cpp
@@ -546,9 +546,12 @@ void SRigCurveContainer::OnHierarchyModified(ERigHierarchyNotification InNotif,
 		return;
 	}
 
-	if(InElement && !InElement->IsTypeOf(ERigElementType::Curve))
+	if(InElement)
 	{
-		return;
+		if(!InElement->IsTypeOf(ERigElementType::Curve))
+		{
+			return;
+		}
 	}
 
 	if (ControlRigBlueprint.IsValid())
@@ -572,13 +575,16 @@ void SRigCurveContainer::OnHierarchyModified(ERigHierarchyNotification InNotif,
 		case ERigHierarchyNotification::ElementSelected:
     	case ERigHierarchyNotification::ElementDeselected:
 		{
-			const bool bSelected = InNotif == ERigHierarchyNotification::ElementSelected;
-			for(const FDisplayedRigCurveInfoPtr& Item : RigCurveList)
+			if(InElement)
 			{
-				if (Item->CurveName == InElement->GetName())
+				const bool bSelected = InNotif == ERigHierarchyNotification::ElementSelected;
+				for(const FDisplayedRigCurveInfoPtr& Item : RigCurveList)
 				{
-					RigCurveListView->SetItemSelection(Item, bSelected);
-					break;
+					if (Item->CurveName == InElement->GetName())
+					{
+						RigCurveListView->SetItemSelection(Item, bSelected);
+						break;
+					}
 				}
 			}
 			break;
