Control Rig: Introduce RigidBody and Auxiliary elements to the hierarchy

diff --git a/Engine/Plugins/Experimental/ControlRig/Source/ControlRig/Private/ControlRig.cpp b/Engine/Plugins/Experimental/ControlRig/Source/ControlRig/Private/ControlRig.cpp
--- a/Engine/Plugins/Experimental/ControlRig/Source/ControlRig/Private/ControlRig.cpp
+++ b/Engine/Plugins/Experimental/ControlRig/Source/ControlRig/Private/ControlRig.cpp
@@ -557,6 +557,12 @@ void UControlRig::Execute(const EControlRigState InState, const FName& InEventNa
 		}
 	}
 
+	if(DynamicHierarchy)
+	{
+		// if we have any aux elements dirty them
+		DynamicHierarchy->UpdateAuxiliaryElements(&Context);
+	}
+
 #if WITH_EDITOR
 	Context.Log = ControlRigLog;
 	if (ControlRigLog != nullptr)
diff --git a/Engine/Plugins/Experimental/ControlRig/Source/ControlRig/Private/Rigs/RigHierarchy.cpp b/Engine/Plugins/Experimental/ControlRig/Source/ControlRig/Private/Rigs/RigHierarchy.cpp
--- a/Engine/Plugins/Experimental/ControlRig/Source/ControlRig/Private/Rigs/RigHierarchy.cpp
+++ b/Engine/Plugins/Experimental/ControlRig/Source/ControlRig/Private/Rigs/RigHierarchy.cpp
@@ -3,6 +3,7 @@
 #include "Rigs/RigHierarchy.h"
 #include "Rigs/RigHierarchyElements.h"
 #include "Rigs/RigHierarchyController.h"
+#include "Units/RigUnitContext.h"
 #include "Math/ControlRigMathLibrary.h"
 #include "UObject/AnimObjectVersion.h"
 #include "Serialization/MemoryReader.h"
@@ -118,6 +119,16 @@ void URigHierarchy::Load(FArchive& Ar)
 				Element = new FRigCurveElement();
 				break;
 			}
+			case ERigElementType::RigidBody:
+			{
+				Element = new FRigRigidBodyElement();
+				break;
+			}
+			case ERigElementType::Auxiliary:
+			{
+				Element = new FRigAuxiliaryElement();
+				break;
+			}
 			default:
 			{
 				ensure(false);
@@ -242,6 +253,32 @@ void URigHierarchy::CopyPose(URigHierarchy* InHierarchy, bool bCurrent, bool bIn
 	}
 }
 
+void URigHierarchy::UpdateAuxiliaryElements(const FRigUnitContext* InContext)
+{
+	check(InContext);
+	
+	for (int32 ElementIndex = 0; ElementIndex < Elements.Num(); ElementIndex++)
+	{
+		if(FRigAuxiliaryElement* AuxiliaryElement = Cast<FRigAuxiliaryElement>(Elements[ElementIndex]))
+		{
+			const FTransform InitialWorldTransform = AuxiliaryElement->GetAuxiliaryWorldTransform(InContext, true);
+			const FTransform CurrentWorldTransform = AuxiliaryElement->GetAuxiliaryWorldTransform(InContext, false);
+
+			const FTransform InitialGlobalTransform = InitialWorldTransform.GetRelativeTransform(InContext->ToWorldSpaceTransform);
+			const FTransform CurrentGlobalTransform = CurrentWorldTransform.GetRelativeTransform(InContext->ToWorldSpaceTransform);
+
+			const FTransform InitialParentTransform = GetParentTransform(AuxiliaryElement, ERigTransformType::InitialGlobal); 
+			const FTransform CurrentParentTransform = GetParentTransform(AuxiliaryElement, ERigTransformType::CurrentGlobal);
+
+			const FTransform InitialLocalTransform = InitialGlobalTransform.GetRelativeTransform(InitialParentTransform);
+			const FTransform CurrentLocalTransform = CurrentGlobalTransform.GetRelativeTransform(CurrentParentTransform);
+
+			SetTransform(AuxiliaryElement, InitialLocalTransform, ERigTransformType::InitialLocal, true, false);
+			SetTransform(AuxiliaryElement, CurrentLocalTransform, ERigTransformType::CurrentLocal, true, false);
+		}
+	}
+}
+
 void URigHierarchy::ResetPoseToInitial(ERigElementType InTypeFilter)
 {
 	const bool bAppliesToAllTypes = InTypeFilter == ERigElementType::All;
diff --git a/Engine/Plugins/Experimental/ControlRig/Source/ControlRig/Private/Rigs/RigHierarchyController.cpp b/Engine/Plugins/Experimental/ControlRig/Source/ControlRig/Private/Rigs/RigHierarchyController.cpp
--- a/Engine/Plugins/Experimental/ControlRig/Source/ControlRig/Private/Rigs/RigHierarchyController.cpp
+++ b/Engine/Plugins/Experimental/ControlRig/Source/ControlRig/Private/Rigs/RigHierarchyController.cpp
@@ -287,6 +287,72 @@ FRigElementKey URigHierarchyController::AddCurve(FName InName, float InValue, bo
 	return NewElement->Key;
 }
 
+FRigElementKey URigHierarchyController::AddRigidBody(FName InName, FRigElementKey InParent,
+	FRigRigidBodySettings InSettings, FTransform InLocalTransform, bool bSetupUndo)
+{
+	if(!IsValid())
+	{
+		return FRigElementKey();
+	}
+
+#if WITH_EDITOR
+	TSharedPtr<FScopedTransaction> TransactionPtr;
+	if(bSetupUndo)
+	{
+		TransactionPtr = MakeShared<FScopedTransaction>(NSLOCTEXT("RigHierarchyController", "Add RigidBody", "Add RigidBody"));
+		Hierarchy->Modify();
+	}
+#endif
+
+	FRigRigidBodyElement* NewElement = new FRigRigidBodyElement();
+	NewElement->Key.Type = ERigElementType::RigidBody;
+	NewElement->Key.Name = Hierarchy->GetSafeNewName(InName.ToString(), NewElement->Key.Type);
+	NewElement->Settings = InSettings;
+	AddElement(NewElement, Hierarchy->Get(Hierarchy->GetIndex(InParent)), true);
+
+	Hierarchy->SetTransform(NewElement, InLocalTransform, ERigTransformType::InitialLocal, true, false);
+	NewElement->Pose.Current = NewElement->Pose.Initial;
+
+#if WITH_EDITOR
+	TransactionPtr.Reset();
+#endif
+
+	return NewElement->Key;
+}
+
+FRigElementKey URigHierarchyController::AddAuxiliaryElement(FName InName, FRigElementKey InParent,
+	FRigAuxiliaryElementGetWorldTransformDelegate InDelegate, bool bSetupUndo)
+{
+	if(!IsValid())
+	{
+		return FRigElementKey();
+	}
+
+#if WITH_EDITOR
+	TSharedPtr<FScopedTransaction> TransactionPtr;
+	if(bSetupUndo)
+	{
+		TransactionPtr = MakeShared<FScopedTransaction>(NSLOCTEXT("RigHierarchyController", "Add Auxiliary Element", "Add Auxiliary Element"));
+		Hierarchy->Modify();
+	}
+#endif
+
+	FRigAuxiliaryElement* NewElement = new FRigAuxiliaryElement();
+	NewElement->Key.Type = ERigElementType::Auxiliary;
+	NewElement->Key.Name = Hierarchy->GetSafeNewName(InName.ToString(), NewElement->Key.Type);
+	NewElement->GetWorldTransformDelegate = InDelegate;
+	AddElement(NewElement, Hierarchy->Get(Hierarchy->GetIndex(InParent)), true);
+
+	Hierarchy->SetTransform(NewElement, FTransform::Identity, ERigTransformType::InitialLocal, true, false);
+	NewElement->Pose.Current = NewElement->Pose.Initial;
+
+#if WITH_EDITOR
+	TransactionPtr.Reset();
+#endif
+
+	return NewElement->Key;
+}
+
 FRigControlSettings URigHierarchyController::GetControlSettings(FRigElementKey InKey) const
 {
 	if(!IsValid())
@@ -669,6 +735,18 @@ FString URigHierarchyController::ExportToText(TArray<FRigElementKey> InKeys) con
 				FRigCurveElement::StaticStruct()->ExportText(PerElementData.Content, Element, &DefaultElement, nullptr, PPF_None, nullptr);
 				break;
 			}
+			case ERigElementType::RigidBody:
+			{
+				FRigRigidBodyElement DefaultElement;
+				FRigRigidBodyElement::StaticStruct()->ExportText(PerElementData.Content, Element, &DefaultElement, nullptr, PPF_None, nullptr);
+				break;
+			}
+			case ERigElementType::Auxiliary:
+			{
+				FRigAuxiliaryElement DefaultElement;
+				FRigAuxiliaryElement::StaticStruct()->ExportText(PerElementData.Content, Element, &DefaultElement, nullptr, PPF_None, nullptr);
+				break;
+			}
 			default:
 			{
 				ensure(false);
@@ -772,6 +850,18 @@ TArray<FRigElementKey> URigHierarchyController::ImportFromText(FString InContent
 				FRigCurveElement::StaticStruct()->ImportText(*PerElementData.Content, NewElement, nullptr, EPropertyPortFlags::PPF_None, &ErrorPipe, FRigCurveElement::StaticStruct()->GetName(), true);
 				break;
 			}
+			case ERigElementType::RigidBody:
+			{
+				NewElement = new FRigRigidBodyElement();
+				FRigRigidBodyElement::StaticStruct()->ImportText(*PerElementData.Content, NewElement, nullptr, EPropertyPortFlags::PPF_None, &ErrorPipe, FRigRigidBodyElement::StaticStruct()->GetName(), true);
+				break;
+			}
+			case ERigElementType::Auxiliary:
+			{
+				NewElement = new FRigAuxiliaryElement();
+				FRigAuxiliaryElement::StaticStruct()->ImportText(*PerElementData.Content, NewElement, nullptr, EPropertyPortFlags::PPF_None, &ErrorPipe, FRigAuxiliaryElement::StaticStruct()->GetName(), true);
+				break;
+			}
 			default:
 			{
 				ensure(false);
diff --git a/Engine/Plugins/Experimental/ControlRig/Source/ControlRig/Private/Rigs/RigHierarchyElements.cpp b/Engine/Plugins/Experimental/ControlRig/Source/ControlRig/Private/Rigs/RigHierarchyElements.cpp
--- a/Engine/Plugins/Experimental/ControlRig/Source/ControlRig/Private/Rigs/RigHierarchyElements.cpp
+++ b/Engine/Plugins/Experimental/ControlRig/Source/ControlRig/Private/Rigs/RigHierarchyElements.cpp
@@ -2,6 +2,7 @@
 
 #include "Rigs/RigHierarchyElements.h"
 #include "Rigs/RigHierarchy.h"
+#include "Units/RigUnitContext.h"
 #include "UObject/AnimObjectVersion.h"
 
 ////////////////////////////////////////////////////////////////////////////////
@@ -121,6 +122,8 @@ void FRigTransformElement::Load(FArchive& Ar, URigHierarchy* Hierarchy, ESeriali
 
 void FRigTransformElement::CopyPose(FRigBaseElement* InOther, bool bCurrent, bool bInitial)
 {
+	Super::CopyPose(InOther, bCurrent, bInitial);
+
 	if(FRigTransformElement* Other = Cast<FRigTransformElement>(InOther))
 	{
 		if(bCurrent)
@@ -236,6 +239,8 @@ void FRigMultiParentElement::Load(FArchive& Ar, URigHierarchy* Hierarchy, ESeria
 
 void FRigMultiParentElement::CopyPose(FRigBaseElement* InOther, bool bCurrent, bool bInitial)
 {
+	Super::CopyPose(InOther, bCurrent, bInitial);
+
 	if(FRigMultiParentElement* Other = Cast<FRigMultiParentElement>(InOther))
 	{
 		if(bCurrent)
@@ -406,6 +411,8 @@ void FRigControlElement::Load(FArchive& Ar, URigHierarchy* Hierarchy, ESerializa
 
 void FRigControlElement::CopyPose(FRigBaseElement* InOther, bool bCurrent, bool bInitial)
 {
+	Super::CopyPose(InOther, bCurrent, bInitial);
+	
 	if(FRigControlElement* Other = Cast<FRigControlElement>(InOther))
 	{
 		if(bCurrent)
@@ -454,3 +461,82 @@ void FRigCurveElement::CopyPose(FRigBaseElement* InOther, bool bCurrent, bool bI
 		Value = Other->Value;
 	}
 }
+
+////////////////////////////////////////////////////////////////////////////////
+// FRigRigidBodySettings
+////////////////////////////////////////////////////////////////////////////////
+
+FRigRigidBodySettings::FRigRigidBodySettings()
+	: Mass(1.f)
+{
+}
+
+void FRigRigidBodySettings::Save(FArchive& Ar)
+{
+	Ar << Mass;
+}
+
+void FRigRigidBodySettings::Load(FArchive& Ar)
+{
+	Ar << Mass;
+}
+
+////////////////////////////////////////////////////////////////////////////////
+// FRigRigidBodyElement
+////////////////////////////////////////////////////////////////////////////////
+
+void FRigRigidBodyElement::Save(FArchive& Ar, URigHierarchy* Hierarchy, ESerializationPhase SerializationPhase)
+{
+	Super::Save(Ar, Hierarchy, SerializationPhase);
+
+	if(SerializationPhase == ESerializationPhase::StaticData)
+	{
+		Settings.Save(Ar);
+	}
+}
+
+void FRigRigidBodyElement::Load(FArchive& Ar, URigHierarchy* Hierarchy, ESerializationPhase SerializationPhase)
+{
+	Super::Load(Ar, Hierarchy, SerializationPhase);
+
+	if(SerializationPhase == ESerializationPhase::StaticData)
+	{
+		Settings.Load(Ar);
+	}
+}
+
+////////////////////////////////////////////////////////////////////////////////
+// FRigAuxiliaryElement
+////////////////////////////////////////////////////////////////////////////////
+
+void FRigAuxiliaryElement::Save(FArchive& Ar, URigHierarchy* Hierarchy, ESerializationPhase SerializationPhase)
+{
+	Super::Save(Ar, Hierarchy, SerializationPhase);
+}
+
+void FRigAuxiliaryElement::Load(FArchive& Ar, URigHierarchy* Hierarchy, ESerializationPhase SerializationPhase)
+{
+	Super::Load(Ar, Hierarchy, SerializationPhase);
+}
+
+FTransform FRigAuxiliaryElement::GetAuxiliaryWorldTransform(const FRigUnitContext* InContext, bool bInitial) const
+{
+	if(GetWorldTransformDelegate.IsBound())
+	{
+		return GetWorldTransformDelegate.Execute(InContext, GetKey(), bInitial);
+	}
+	return FTransform::Identity;
+}
+
+void FRigAuxiliaryElement::CopyPose(FRigBaseElement* InOther, bool bCurrent, bool bInitial)
+{
+	Super::CopyPose(InOther, bCurrent, bInitial);
+	
+	if(FRigAuxiliaryElement* Other = Cast<FRigAuxiliaryElement>(InOther))
+	{
+		if(Other->GetWorldTransformDelegate.IsBound())
+		{
+			GetWorldTransformDelegate = Other->GetWorldTransformDelegate;
+		}
+	}
+}
diff --git a/Engine/Plugins/Experimental/ControlRig/Source/ControlRig/Public/Rigs/RigHierarchy.h b/Engine/Plugins/Experimental/ControlRig/Source/ControlRig/Public/Rigs/RigHierarchy.h
--- a/Engine/Plugins/Experimental/ControlRig/Source/ControlRig/Public/Rigs/RigHierarchy.h
+++ b/Engine/Plugins/Experimental/ControlRig/Source/ControlRig/Public/Rigs/RigHierarchy.h
@@ -59,6 +59,11 @@ public:
 	UFUNCTION(BlueprintCallable, Category = URigHierarchy)
 	void CopyPose(URigHierarchy* InHierarchy, bool bCurrent, bool bInitial);
 
+	/**
+	 * Update all elements that depend on auxiliary data
+	 */
+	void UpdateAuxiliaryElements(const FRigUnitContext* InContext);
+
 	/**
 	 * Resets the current pose of a filtered lost if elements to the initial / ref pose.
 	 */
@@ -572,6 +577,43 @@ public:
 		return GetKeysOfType<FRigCurveElement>(false);
 	}
 
+	/**
+	 * Returns all RigidBody elements
+	 * @param bTraverse Returns the elements in order of a depth first traversal
+	 */
+	FORCEINLINE TArray<FRigRigidBodyElement*> GetRigidBodies(bool bTraverse = false) const
+	{
+		return GetElementsOfType<FRigRigidBodyElement>(bTraverse);
+	}
+
+	/**
+	 * Returns all RigidBody elements
+	 * @param bTraverse Returns the elements in order of a depth first traversal
+	 */
+	UFUNCTION(BlueprintCallable, Category = URigHierarchy, meta = (DisplayName = "Get RigidBodies", ScriptName = "GetRigidBodies"))
+    FORCEINLINE TArray<FRigElementKey> GetRigidBodyKeys(bool bTraverse = true) const
+	{
+		return GetKeysOfType<FRigRigidBodyElement>(bTraverse);
+	}
+
+	/**
+	 * Returns all Auxiliary elements
+	 * @param bTraverse Returns the elements in order of a depth first traversal
+	 */
+	FORCEINLINE TArray<FRigAuxiliaryElement*> GetAuxiliaryElements(bool bTraverse = false) const
+	{
+		return GetElementsOfType<FRigAuxiliaryElement>(bTraverse);
+	}
+
+	/**
+	 * Returns all Auxiliary elements
+	 * @param bTraverse Returns the elements in order of a depth first traversal
+	 */
+	UFUNCTION(BlueprintCallable, Category = URigHierarchy, meta = (DisplayName = "Get Auxiliarys", ScriptName = "GetAuxiliaryElements"))
+    FORCEINLINE TArray<FRigElementKey> GetAuxiliaryKeys(bool bTraverse = true) const
+	{
+		return GetKeysOfType<FRigAuxiliaryElement>(bTraverse);
+	}
 
 	/**
 	 * Returns the selected elements
@@ -1566,7 +1608,7 @@ public:
 	 * @return The Controller for this hierarchy
 	 */
 	UFUNCTION(BlueprintCallable, Category = URigHierarchy)
-	URigHierarchyController* GetController(bool bCreateIfNeeded = false);
+	URigHierarchyController* GetController(bool bCreateIfNeeded = true);
 
 	/**
 	 * Returns the topology version of this hierarchy
diff --git a/Engine/Plugins/Experimental/ControlRig/Source/ControlRig/Public/Rigs/RigHierarchyController.h b/Engine/Plugins/Experimental/ControlRig/Source/ControlRig/Public/Rigs/RigHierarchyController.h
--- a/Engine/Plugins/Experimental/ControlRig/Source/ControlRig/Public/Rigs/RigHierarchyController.h
+++ b/Engine/Plugins/Experimental/ControlRig/Source/ControlRig/Public/Rigs/RigHierarchyController.h
@@ -147,6 +147,37 @@ public:
         bool bSetupUndo = true
         );
 
+	/**
+	* Adds a rigidbody to the hierarchy
+	* @param InName The suggested name of the new rigidbody - will eventually be corrected by the namespace
+	* @param InParent The (optional) parent of the new rigidbody. If you don't need a parent, pass FRigElementKey()
+	* @param InSettings All of the rigidbody's settings
+	* @param InLocalTransform The transform for the new rigidbody - in the space of the provided parent
+	* @param bSetupUndo If set to true the stack will record the change for undo / redo
+	* @return The key for the newly created rigidbody.
+	*/
+	UFUNCTION(BlueprintCallable, Category = URigHierarchyController)
+    FRigElementKey AddRigidBody(
+    	FName InName,
+    	FRigElementKey InParent,
+        FRigRigidBodySettings InSettings,
+    	FTransform InLocalTransform,
+    	bool bSetupUndo = false);
+
+	/**
+	* Adds an auxiliary element to the hierarchy
+	* @param InName The suggested name of the new auxiliary element - will eventually be corrected by the namespace
+	* @param InParent The (optional) parent of the new auxiliary element. If you don't need a parent, pass FRigElementKey()
+	* @param InDelegate The delegate to use to pull the local transform
+	* @param bSetupUndo If set to true the stack will record the change for undo / redo
+	* @return The key for the newly created auxiliary element.
+	*/
+    FRigElementKey AddAuxiliaryElement(
+        FName InName,
+        FRigElementKey InParent,
+        FRigAuxiliaryElementGetWorldTransformDelegate InDelegate,
+        bool bSetupUndo = false);
+
 	/**
 	 * Returns the control settings of a given control
 	 * @param InKey The key of the control to receive the settings for
diff --git a/Engine/Plugins/Experimental/ControlRig/Source/ControlRig/Public/Rigs/RigHierarchyDefines.h b/Engine/Plugins/Experimental/ControlRig/Source/ControlRig/Public/Rigs/RigHierarchyDefines.h
--- a/Engine/Plugins/Experimental/ControlRig/Source/ControlRig/Public/Rigs/RigHierarchyDefines.h
+++ b/Engine/Plugins/Experimental/ControlRig/Source/ControlRig/Public/Rigs/RigHierarchyDefines.h
@@ -21,7 +21,9 @@ enum class ERigElementType : uint8
 	Space = 0x002,
 	Control = 0x004,
 	Curve = 0x008,
-	All = Bone | Space | Control | Curve UMETA(Hidden),
+	RigidBody = 0x010,
+	Auxiliary = 0x020,
+	All = Bone | Space | Control | Curve | RigidBody | Auxiliary UMETA(Hidden),
 };
 
 UENUM(BlueprintType)
@@ -883,6 +885,14 @@ public:
 			{
 				return FString::Printf(TEXT("Curve(%s)"), *Name.ToString());
 			}
+			case ERigElementType::RigidBody:
+			{
+				return FString::Printf(TEXT("RigidBody(%s)"), *Name.ToString());
+			}
+			case ERigElementType::Auxiliary:
+			{
+				return FString::Printf(TEXT("Auxiliary(%s)"), *Name.ToString());
+			}
 		}
 		return FString();
 	}
diff --git a/Engine/Plugins/Experimental/ControlRig/Source/ControlRig/Public/Rigs/RigHierarchyElements.h b/Engine/Plugins/Experimental/ControlRig/Source/ControlRig/Public/Rigs/RigHierarchyElements.h
--- a/Engine/Plugins/Experimental/ControlRig/Source/ControlRig/Public/Rigs/RigHierarchyElements.h
+++ b/Engine/Plugins/Experimental/ControlRig/Source/ControlRig/Public/Rigs/RigHierarchyElements.h
@@ -6,9 +6,12 @@
 #include "RigHierarchyDefines.h"
 #include "RigHierarchyElements.generated.h"
 
+struct FRigUnitContext;
 struct FRigBaseElement;
 class URigHierarchy;
 
+DECLARE_DELEGATE_RetVal_ThreeParams(FTransform, FRigAuxiliaryElementGetWorldTransformDelegate, const FRigUnitContext*, const FRigElementKey& /* Key */, bool /* bInitial */);
+
 #define DECLARE_RIG_ELEMENT_METHODS(ElementType) \
 template<typename T> \
 friend FORCEINLINE const T* Cast(const ElementType* InElement) \
@@ -575,7 +578,9 @@ protected:
 	{
 		return InElement->GetType() == ERigElementType::Bone ||
 			InElement->GetType() == ERigElementType::Space ||
-			InElement->GetType() == ERigElementType::Control;
+			InElement->GetType() == ERigElementType::Control ||
+			InElement->GetType() == ERigElementType::RigidBody ||
+			InElement->GetType() == ERigElementType::Auxiliary;
 	}
 
 public:
@@ -615,7 +620,9 @@ protected:
 	
 	FORCEINLINE static bool IsClassOf(const FRigBaseElement* InElement)
 	{
-		return InElement->GetType() == ERigElementType::Bone;
+		return InElement->GetType() == ERigElementType::Bone ||
+			InElement->GetType() == ERigElementType::RigidBody ||
+			InElement->GetType() == ERigElementType::Auxiliary;
 	}
 
 	friend struct FRigBaseElement;
@@ -917,4 +924,87 @@ public:
 protected:
 	
 	friend struct FRigBaseElement;
+};
+
+USTRUCT(BlueprintType)
+struct CONTROLRIG_API FRigRigidBodySettings
+{
+	GENERATED_BODY()
+
+	FRigRigidBodySettings();
+
+	void Save(FArchive& Ar);
+	void Load(FArchive& Ar);
+
+
+	UPROPERTY(BlueprintReadWrite, EditAnywhere, Category = Control)
+	float Mass;
+};
+
+USTRUCT(BlueprintType)
+struct CONTROLRIG_API FRigRigidBodyElement : public FRigSingleParentElement
+{
+public:
+	
+	GENERATED_BODY()
+	DECLARE_RIG_ELEMENT_METHODS(FRigRigidBodyElement)
+
+    FRigRigidBodyElement()
+        : FRigSingleParentElement()
+	{
+		Key.Type = ERigElementType::RigidBody;
+	}
+	
+	virtual ~FRigRigidBodyElement(){}
+
+	virtual void Save(FArchive& A, URigHierarchy* Hierarchy, ESerializationPhase SerializationPhase) override;
+	virtual void Load(FArchive& Ar, URigHierarchy* Hierarchy, ESerializationPhase SerializationPhase) override;
+
+	UPROPERTY(BlueprintReadOnly, EditAnywhere, Category = Control, meta=(ShowOnlyInnerProperties))
+	FRigRigidBodySettings Settings;
+
+protected:
+	
+	FORCEINLINE static bool IsClassOf(const FRigBaseElement* InElement)
+	{
+		return InElement->GetType() == ERigElementType::RigidBody;
+	}
+
+	friend struct FRigBaseElement;
+};
+
+USTRUCT(BlueprintType)
+struct CONTROLRIG_API FRigAuxiliaryElement : public FRigSingleParentElement
+{
+public:
+	
+	GENERATED_BODY()
+	DECLARE_RIG_ELEMENT_METHODS(FRigAuxiliaryElement)
+
+    FRigAuxiliaryElement()
+        : FRigSingleParentElement()
+	{
+		Key.Type = ERigElementType::Auxiliary;
+	}
+	
+	virtual ~FRigAuxiliaryElement(){}
+
+	virtual void Save(FArchive& A, URigHierarchy* Hierarchy, ESerializationPhase SerializationPhase) override;
+	virtual void Load(FArchive& Ar, URigHierarchy* Hierarchy, ESerializationPhase SerializationPhase) override;
+
+	FTransform GetAuxiliaryWorldTransform(const FRigUnitContext* InContext, bool bInitial) const;
+
+	virtual void CopyPose(FRigBaseElement* InOther, bool bCurrent, bool bInitial) override;
+
+protected:
+
+	FRigAuxiliaryElementGetWorldTransformDelegate GetWorldTransformDelegate;
+
+	FORCEINLINE static bool IsClassOf(const FRigBaseElement* InElement)
+	{
+		return InElement->GetType() == ERigElementType::Auxiliary;
+	}
+
+	friend struct FRigBaseElement;
+	friend class URigHierarchyController;
 };
\ No newline at end of file
diff --git a/Engine/Plugins/Experimental/ControlRig/Source/ControlRigDeveloper/Private/Graph/ControlRigGraph.cpp b/Engine/Plugins/Experimental/ControlRig/Source/ControlRigDeveloper/Private/Graph/ControlRigGraph.cpp
--- a/Engine/Plugins/Experimental/ControlRig/Source/ControlRigDeveloper/Private/Graph/ControlRigGraph.cpp
+++ b/Engine/Plugins/Experimental/ControlRig/Source/ControlRigDeveloper/Private/Graph/ControlRigGraph.cpp
@@ -74,47 +74,29 @@ void UControlRigGraph::CacheNameLists(URigHierarchy* InHierarchy, const FControl
 	check(InHierarchy);
 	check(DrawContainer);
 
+	TArray<TSharedPtr<FString>>& BoneNameList = ElementNameLists.FindOrAdd(ERigElementType::Bone);
+	TArray<TSharedPtr<FString>>& SpaceNameList = ElementNameLists.FindOrAdd(ERigElementType::Space);
+	TArray<TSharedPtr<FString>>& ControlNameList = ElementNameLists.FindOrAdd(ERigElementType::Control);
+	TArray<TSharedPtr<FString>>& CurveNameList = ElementNameLists.FindOrAdd(ERigElementType::Curve);
+	TArray<TSharedPtr<FString>>& RigidBodyNameList = ElementNameLists.FindOrAdd(ERigElementType::RigidBody);
+	TArray<TSharedPtr<FString>>& AuxiliaryNameList = ElementNameLists.FindOrAdd(ERigElementType::Auxiliary);
+	
 	CacheNameListForHierarchy<FRigBoneElement>(InHierarchy, BoneNameList);
-	CacheNameListForHierarchy<FRigControlElement>(InHierarchy, ControlNameList);
 	CacheNameListForHierarchy<FRigSpaceElement>(InHierarchy, SpaceNameList);
+	CacheNameListForHierarchy<FRigControlElement>(InHierarchy, ControlNameList);
 	CacheNameListForHierarchy<FRigCurveElement>(InHierarchy, CurveNameList);
+	CacheNameListForHierarchy<FRigRigidBodyElement>(InHierarchy, RigidBodyNameList);
+	CacheNameListForHierarchy<FRigAuxiliaryElement>(InHierarchy, AuxiliaryNameList);
 	CacheNameList<FControlRigDrawContainer>(*DrawContainer, DrawingNameList);
 }
 
-const TArray<TSharedPtr<FString>>& UControlRigGraph::GetBoneNameList(URigVMPin* InPin) const
-{
-	if (UControlRigGraph* OuterGraph = Cast<UControlRigGraph>(GetOuter()))
-	{
-		return OuterGraph->GetBoneNameList(InPin);
-	}
-	return BoneNameList;
-}
-
-const TArray<TSharedPtr<FString>>& UControlRigGraph::GetControlNameList(URigVMPin* InPin) const
-{
-	if (UControlRigGraph* OuterGraph = Cast<UControlRigGraph>(GetOuter()))
-	{
-		return OuterGraph->GetControlNameList(InPin);
-	}
-	return ControlNameList;
-}
-
-const TArray<TSharedPtr<FString>>& UControlRigGraph::GetSpaceNameList(URigVMPin* InPin) const
-{
-	if (UControlRigGraph* OuterGraph = Cast<UControlRigGraph>(GetOuter()))
-	{
-		return OuterGraph->GetSpaceNameList(InPin);
-	}
-	return SpaceNameList;
-}
-
-const TArray<TSharedPtr<FString>>& UControlRigGraph::GetCurveNameList(URigVMPin* InPin) const
+const TArray<TSharedPtr<FString>>& UControlRigGraph::GetElementNameList(ERigElementType InElementType) const
 {
 	if (UControlRigGraph* OuterGraph = Cast<UControlRigGraph>(GetOuter()))
 	{
-		return OuterGraph->GetCurveNameList(InPin);
+		return OuterGraph->GetElementNameList(InElementType);
 	}
-	return CurveNameList;
+	return ElementNameLists.FindChecked(InElementType);
 }
 
 const TArray<TSharedPtr<FString>>& UControlRigGraph::GetElementNameList(URigVMPin* InPin) const
@@ -141,36 +123,6 @@ const TArray<TSharedPtr<FString>>& UControlRigGraph::GetElementNameList(URigVMPi
 	return GetBoneNameList(nullptr);
 }
 
-const TArray<TSharedPtr<FString>>& UControlRigGraph::GetElementNameList(ERigElementType InElementType) const
-{
-	if (UControlRigGraph* OuterGraph = Cast<UControlRigGraph>(GetOuter()))
-	{
-		return OuterGraph->GetElementNameList(InElementType);
-	}
-
-	switch (InElementType)
-	{
-		case ERigElementType::Bone:
-		{
-			return GetBoneNameList();
-		}
-		case ERigElementType::Control:
-		{
-			return GetControlNameList();
-		}
-		case ERigElementType::Space:
-		{
-			return GetSpaceNameList();
-		}
-		case ERigElementType::Curve:
-		{
-			return GetCurveNameList();
-		}
-	}
-
-	return GetBoneNameList(nullptr);
-}
-
 const TArray<TSharedPtr<FString>>& UControlRigGraph::GetDrawingNameList(URigVMPin* InPin) const
 {
 	if (UControlRigGraph* OuterGraph = Cast<UControlRigGraph>(GetOuter()))
diff --git a/Engine/Plugins/Experimental/ControlRig/Source/ControlRigDeveloper/Public/Graph/ControlRigGraph.h b/Engine/Plugins/Experimental/ControlRig/Source/ControlRigDeveloper/Public/Graph/ControlRigGraph.h
--- a/Engine/Plugins/Experimental/ControlRig/Source/ControlRigDeveloper/Public/Graph/ControlRigGraph.h
+++ b/Engine/Plugins/Experimental/ControlRig/Source/ControlRigDeveloper/Public/Graph/ControlRigGraph.h
@@ -41,13 +41,25 @@ public:
 
 
 	void CacheNameLists(URigHierarchy* InHierarchy, const FControlRigDrawContainer* DrawContainer);
+	const TArray<TSharedPtr<FString>>& GetElementNameList(ERigElementType InElementType = ERigElementType::Bone) const;
 
-	const TArray<TSharedPtr<FString>>& GetBoneNameList(URigVMPin* InPin = nullptr) const;
-	const TArray<TSharedPtr<FString>>& GetControlNameList(URigVMPin* InPin = nullptr) const;
-	const TArray<TSharedPtr<FString>>& GetSpaceNameList(URigVMPin* InPin = nullptr) const;
-	const TArray<TSharedPtr<FString>>& GetCurveNameList(URigVMPin* InPin = nullptr) const;
+	FORCEINLINE const TArray<TSharedPtr<FString>>& GetBoneNameList(URigVMPin* InPin = nullptr) const
+	{
+		return GetElementNameList(ERigElementType::Bone);
+	}
+	FORCEINLINE const TArray<TSharedPtr<FString>>& GetControlNameList(URigVMPin* InPin = nullptr) const
+	{
+		return GetElementNameList(ERigElementType::Control);
+	}
+	FORCEINLINE const TArray<TSharedPtr<FString>>& GetSpaceNameList(URigVMPin* InPin = nullptr) const
+	{
+		return GetElementNameList(ERigElementType::Space);
+	}
+	FORCEINLINE const TArray<TSharedPtr<FString>>& GetCurveNameList(URigVMPin* InPin = nullptr) const
+	{
+		return GetElementNameList(ERigElementType::Curve);
+	}
 	const TArray<TSharedPtr<FString>>& GetElementNameList(URigVMPin* InPin = nullptr) const;
-	const TArray<TSharedPtr<FString>>& GetElementNameList(ERigElementType InElementType) const;
 	const TArray<TSharedPtr<FString>>& GetDrawingNameList(URigVMPin* InPin = nullptr) const;
 
 	bool bSuspendModelNotifications;
@@ -110,10 +122,7 @@ private:
 		}
 	}
 
-	TArray<TSharedPtr<FString>> BoneNameList;
-	TArray<TSharedPtr<FString>> ControlNameList;
-	TArray<TSharedPtr<FString>> SpaceNameList;
-	TArray<TSharedPtr<FString>> CurveNameList;
+	TMap<ERigElementType, TArray<TSharedPtr<FString>>> ElementNameLists;
 	TArray<TSharedPtr<FString>> DrawingNameList;
 
 	bool bIsSelecting;
diff --git a/Engine/Plugins/Experimental/ControlRig/Source/ControlRigEditor/Private/ControlRigEditorStyle.h b/Engine/Plugins/Experimental/ControlRig/Source/ControlRigEditor/Private/ControlRigEditorStyle.h
--- a/Engine/Plugins/Experimental/ControlRig/Source/ControlRigEditor/Private/ControlRigEditorStyle.h
+++ b/Engine/Plugins/Experimental/ControlRig/Source/ControlRigEditor/Private/ControlRigEditorStyle.h
@@ -99,6 +99,8 @@ public:
 			Set("ControlRig.Tree.BoneImported", new IMAGE_BRUSH("Slate/BoneImported_16x", Icon16x16));
 			Set("ControlRig.Tree.Control", new IMAGE_BRUSH("Slate/Control_16x", Icon16x16));
 			Set("ControlRig.Tree.Space", new IMAGE_BRUSH("Slate/Space_16x", Icon16x16));
+			Set("ControlRig.Tree.RigidBody", new IMAGE_BRUSH("Slate/RigidBody_16x", Icon16x16));
+			Set("ControlRig.Tree.Auxiliary", new IMAGE_BRUSH("Slate/Auxiliary_16x", Icon16x16));
 		}
 
 		// Font?
diff --git a/Engine/Plugins/Experimental/ControlRig/Source/ControlRigEditor/Private/ControlRigHierarchyCommands.cpp b/Engine/Plugins/Experimental/ControlRig/Source/ControlRigEditor/Private/ControlRigHierarchyCommands.cpp
--- a/Engine/Plugins/Experimental/ControlRig/Source/ControlRigEditor/Private/ControlRigHierarchyCommands.cpp
+++ b/Engine/Plugins/Experimental/ControlRig/Source/ControlRigEditor/Private/ControlRigHierarchyCommands.cpp
@@ -33,6 +33,8 @@ void FControlRigHierarchyCommands::RegisterCommands()
 	UI_COMMAND(ShowBones, "Show Bones", "Whether to show or hide bones", EUserInterfaceActionType::ToggleButton, FInputChord());
 	UI_COMMAND(ShowControls, "Show Controls", "Whether to show or hide controls", EUserInterfaceActionType::ToggleButton, FInputChord());
 	UI_COMMAND(ShowSpaces, "Show Spaces", "Whether to show or hide spaces", EUserInterfaceActionType::ToggleButton, FInputChord());
+	UI_COMMAND(ShowRigidBodies, "Show RigidBodies", "Whether to show or hide rigidbodies", EUserInterfaceActionType::ToggleButton, FInputChord());
+	UI_COMMAND(ShowAuxiliaryElements, "Show Auxiliary Elements", "Whether to show or hide auxiliary elements", EUserInterfaceActionType::ToggleButton, FInputChord());
 }
 
 #undef LOCTEXT_NAMESPACE
diff --git a/Engine/Plugins/Experimental/ControlRig/Source/ControlRigEditor/Private/ControlRigHierarchyCommands.h b/Engine/Plugins/Experimental/ControlRig/Source/ControlRigEditor/Private/ControlRigHierarchyCommands.h
--- a/Engine/Plugins/Experimental/ControlRig/Source/ControlRigEditor/Private/ControlRigHierarchyCommands.h
+++ b/Engine/Plugins/Experimental/ControlRig/Source/ControlRigEditor/Private/ControlRigHierarchyCommands.h
@@ -99,6 +99,11 @@ public:
 	/* Show spaces */
 	TSharedPtr< FUICommandInfo > ShowSpaces;
 
+	/* Show rigidbodies */
+	TSharedPtr< FUICommandInfo > ShowRigidBodies;
+
+	/* Show auxiliary elements */
+	TSharedPtr< FUICommandInfo > ShowAuxiliaryElements;
 
 	/**
 	 * Initialize commands
diff --git a/Engine/Plugins/Experimental/ControlRig/Source/ControlRigEditor/Private/EditMode/ControlRigEditMode.cpp b/Engine/Plugins/Experimental/ControlRig/Source/ControlRigEditor/Private/EditMode/ControlRigEditMode.cpp
--- a/Engine/Plugins/Experimental/ControlRig/Source/ControlRigEditor/Private/EditMode/ControlRigEditMode.cpp
+++ b/Engine/Plugins/Experimental/ControlRig/Source/ControlRigEditor/Private/EditMode/ControlRigEditMode.cpp
@@ -1836,6 +1836,8 @@ void FControlRigEditMode::OnHierarchyModified(ERigHierarchyNotification InNotif,
             	case ERigElementType::Space:
             	case ERigElementType::Curve:
             	case ERigElementType::Control:
+            	case ERigElementType::RigidBody:
+            	case ERigElementType::Auxiliary:
 				{
 					const bool bSelected = InNotif == ERigHierarchyNotification::ElementSelected;
 					if (bSelected)
diff --git a/Engine/Plugins/Experimental/ControlRig/Source/ControlRigEditor/Private/Editor/ControlRigEditor.cpp b/Engine/Plugins/Experimental/ControlRig/Source/ControlRigEditor/Private/Editor/ControlRigEditor.cpp
--- a/Engine/Plugins/Experimental/ControlRig/Source/ControlRigEditor/Private/Editor/ControlRigEditor.cpp
+++ b/Engine/Plugins/Experimental/ControlRig/Source/ControlRigEditor/Private/Editor/ControlRigEditor.cpp
@@ -1187,6 +1187,16 @@ void FControlRigEditor::SetDetailStruct(const FRigElementKey& InElement)
 			StructToDisplay = MakeShareable(new FStructOnScope(FRigCurveElement::StaticStruct(), (uint8*)Element));
 			break;
 		}
+		case ERigElementType::RigidBody:
+		{
+			StructToDisplay = MakeShareable(new FStructOnScope(FRigRigidBodyElement::StaticStruct(), (uint8*)Element));
+			break;
+		}
+		case ERigElementType::Auxiliary:
+		{
+			StructToDisplay = MakeShareable(new FStructOnScope(FRigAuxiliaryElement::StaticStruct(), (uint8*)Element));
+			break;
+		}
 		default:
 		{
 			break;
diff --git a/Engine/Plugins/Experimental/ControlRig/Source/ControlRigEditor/Private/Editor/SRigHierarchy.cpp b/Engine/Plugins/Experimental/ControlRig/Source/ControlRigEditor/Private/Editor/SRigHierarchy.cpp
--- a/Engine/Plugins/Experimental/ControlRig/Source/ControlRigEditor/Private/Editor/SRigHierarchy.cpp
+++ b/Engine/Plugins/Experimental/ControlRig/Source/ControlRigEditor/Private/Editor/SRigHierarchy.cpp
@@ -170,6 +170,16 @@ void SRigHierarchyItem::Construct(const FArguments& InArgs, TSharedPtr<FControlR
 
 			break;
 		}
+		case ERigElementType::RigidBody:
+		{
+			Brush = FControlRigEditorStyle::Get().GetBrush("ControlRig.Tree.RigidBody");
+			break;
+		}
+		case ERigElementType::Auxiliary:
+		{
+			Brush = FControlRigEditorStyle::Get().GetBrush("ControlRig.Tree.Auxiliary");
+			break;
+		}
 		default:
 		{
 			break;
@@ -418,6 +428,8 @@ void SRigHierarchy::Construct(const FArguments& InArgs, TSharedRef<FControlRigEd
 	bShowBones = true;
 	bShowControls = true;
 	bShowSpaces = true;
+	bShowRigidBodies = true;
+	bShowAuxiliaryElements = true;
 	bIsChangingRigHierarchy = false;
 	RefreshTreeView();
 
@@ -552,6 +564,18 @@ void SRigHierarchy::BindCommands()
 		FExecuteAction::CreateLambda([this]() { bShowSpaces = !bShowSpaces; RefreshTreeView(); }),
 		FCanExecuteAction(),
 		FIsActionChecked::CreateLambda([this]() { return bShowSpaces; }));
+
+	CommandList->MapAction(
+		Commands.ShowRigidBodies,
+		FExecuteAction::CreateLambda([this]() { bShowRigidBodies = !bShowRigidBodies; RefreshTreeView(); }),
+		FCanExecuteAction(),
+		FIsActionChecked::CreateLambda([this]() { return bShowRigidBodies; }));
+
+	CommandList->MapAction(
+		Commands.ShowAuxiliaryElements,
+		FExecuteAction::CreateLambda([this]() { bShowAuxiliaryElements = !bShowAuxiliaryElements; RefreshTreeView(); }),
+		FCanExecuteAction(),
+		FIsActionChecked::CreateLambda([this]() { return bShowAuxiliaryElements; }));
 }
 
 FReply SRigHierarchy::OnKeyDown(const FGeometry& MyGeometry, const FKeyEvent& InKeyEvent)
@@ -666,6 +690,24 @@ void SRigHierarchy::RefreshTreeView()
                         AddElement(ControlElement);
                     }
                     break;
+                }
+				case ERigElementType::RigidBody:
+				{
+                    if(bShowRigidBodies)
+                    {
+                        FRigRigidBodyElement* RigidBodyElement = CastChecked<FRigRigidBodyElement>(Element);
+                        AddElement(RigidBodyElement);
+                    }
+                    break;
+                }
+				case ERigElementType::Auxiliary:
+				{
+                    if(bShowAuxiliaryElements)
+                    {
+                        FRigAuxiliaryElement* AuxiliaryElement = CastChecked<FRigAuxiliaryElement>(Element);
+                        AddElement(AuxiliaryElement);
+                    }
+                    break;
                 }
 				default:
 				{
@@ -1615,7 +1657,16 @@ void SRigHierarchy::HandleMirrorItem()
 /** Check whether we can deleting the selected item(s) */
 bool SRigHierarchy::CanRenameItem() const
 {
-	return IsSingleSelected();
+	if(IsSingleSelected())
+	{
+		if(GetSelectedKeys()[0].Type == ERigElementType::RigidBody ||
+			GetSelectedKeys()[0].Type == ERigElementType::Auxiliary)
+		{
+			return false;
+		}
+		return true;
+	}
+	return false;
 }
 
 /** Delete Item */
@@ -1863,6 +1914,8 @@ TOptional<EItemDropZone> SRigHierarchy::OnCanAcceptDrop(const FDragDropEvent& Dr
 			}
 			case ERigElementType::Control:
 			case ERigElementType::Space:
+			case ERigElementType::RigidBody:
+			case ERigElementType::Auxiliary:
 			{
 				for (const FRigElementKey& DraggedKey : RigDragDropOp->GetElements())
 				{
@@ -1870,6 +1923,8 @@ TOptional<EItemDropZone> SRigHierarchy::OnCanAcceptDrop(const FDragDropEvent& Dr
 					{
 						case ERigElementType::Control:
 						case ERigElementType::Space:
+						case ERigElementType::RigidBody:
+						case ERigElementType::Auxiliary:
 						{
 							break;
 						}
diff --git a/Engine/Plugins/Experimental/ControlRig/Source/ControlRigEditor/Private/Editor/SRigHierarchy.h b/Engine/Plugins/Experimental/ControlRig/Source/ControlRigEditor/Private/Editor/SRigHierarchy.h
--- a/Engine/Plugins/Experimental/ControlRig/Source/ControlRigEditor/Private/Editor/SRigHierarchy.h
+++ b/Engine/Plugins/Experimental/ControlRig/Source/ControlRigEditor/Private/Editor/SRigHierarchy.h
@@ -256,6 +256,12 @@ private:
 	/** Whether or not to show spaces in the hierarchy */
 	bool bShowSpaces;
 
+	/** Whether or not to show rigidbodies in the hierarchy */
+	bool bShowRigidBodies;
+
+	/** Whether or not to show auxiliary elements in the hierarchy */
+	bool bShowAuxiliaryElements;
+
 	/** Search box widget */
 	TSharedPtr<SSearchBox> FilterBox;
 	FText FilterText;
