Control Rig: Fix bone drawing in COntrol Rig skel mesh component

diff --git a/Engine/Plugins/Experimental/ControlRig/Source/ControlRig/Private/Rigs/RigHierarchyController.cpp b/Engine/Plugins/Experimental/ControlRig/Source/ControlRig/Private/Rigs/RigHierarchyController.cpp
--- a/Engine/Plugins/Experimental/ControlRig/Source/ControlRig/Private/Rigs/RigHierarchyController.cpp
+++ b/Engine/Plugins/Experimental/ControlRig/Source/ControlRig/Private/Rigs/RigHierarchyController.cpp
@@ -444,11 +444,11 @@ TArray<FRigElementKey> URigHierarchyController::ImportBones(const FReferenceSkel
 
 	if (bReplaceExistingBones)
 	{
-		GetHierarchy()->ForEach<FRigBoneElement>([&BoneNameMap](FRigBoneElement* BoneElement) -> bool
-        {
+		TArray<FRigBoneElement*> AllBones = GetHierarchy()->GetElementsOfType<FRigBoneElement>(true);
+		for(FRigBoneElement* BoneElement : AllBones)
+		{
 			BoneNameMap.Add(BoneElement->GetName(), BoneElement->GetName());
-			return true;
-		});
+		}
 
 		for (int32 Index = 0; Index < BoneInfos.Num(); ++Index)
 		{
@@ -524,7 +524,8 @@ TArray<FRigElementKey> URigHierarchyController::ImportBones(const FReferenceSkel
 		}
 		
 		TArray<FRigElementKey> BonesToDelete;
-		GetHierarchy()->ForEach<FRigBoneElement>([BoneNameToIndexInSkeleton, &BonesToDelete](FRigBoneElement* BoneElement) -> bool
+		TArray<FRigBoneElement*> AllBones = GetHierarchy()->GetElementsOfType<FRigBoneElement>(true);
+		for(FRigBoneElement* BoneElement : AllBones)
         {
             if (!BoneNameToIndexInSkeleton.Contains(BoneElement->GetName()))
 			{
@@ -533,8 +534,7 @@ TArray<FRigElementKey> URigHierarchyController::ImportBones(const FReferenceSkel
 					BonesToDelete.Add(BoneElement->GetKey());
 				}
 			}
-			return true;
-		});
+		}
 
 		for (const FRigElementKey& BoneToDelete : BonesToDelete)
 		{
@@ -556,6 +556,17 @@ TArray<FRigElementKey> URigHierarchyController::ImportBones(const FReferenceSkel
 			RemoveElement(BoneToDelete);
 			BonesToSelect.Remove(BoneToDelete);
 		}
+
+		// update the sub index to match the bone index in the skeleton
+		for (int32 Index = 0; Index < BoneInfos.Num(); ++Index)
+		{
+			FName DesiredBoneName = Local::DetermineBoneName(BoneInfos[Index].Name, InNameSpace);
+			const FRigElementKey Key(DesiredBoneName, ERigElementType::Bone);
+			if(FRigBoneElement* BoneElement = Hierarchy->Find<FRigBoneElement>(Key))
+			{
+				BoneElement->SubIndex = Index;
+			}
+		}		
 	}
 
 	if (bSelectBones)
@@ -564,7 +575,6 @@ TArray<FRigElementKey> URigHierarchyController::ImportBones(const FReferenceSkel
 	}
 
 	return AddedBones;
-
 }
 
 #if WITH_EDITOR
@@ -1103,7 +1113,7 @@ bool URigHierarchyController::RemoveElement(FRigBaseElement* InElement)
 
 	if(InElement->SubIndex != INDEX_NONE)
 	{
-		for(FRigBaseElement* Element : *Hierarchy)
+		for(FRigBaseElement* Element : Hierarchy->Elements)
 		{
 			if(Element->SubIndex > InElement->SubIndex && Element->GetType() == InElement->GetType())
 			{
diff --git a/Engine/Plugins/Experimental/ControlRig/Source/ControlRigEditor/Private/Editor/ControlRigEditor.cpp b/Engine/Plugins/Experimental/ControlRig/Source/ControlRigEditor/Private/Editor/ControlRigEditor.cpp
--- a/Engine/Plugins/Experimental/ControlRig/Source/ControlRigEditor/Private/Editor/ControlRigEditor.cpp
+++ b/Engine/Plugins/Experimental/ControlRig/Source/ControlRigEditor/Private/Editor/ControlRigEditor.cpp
@@ -3472,9 +3472,9 @@ void FControlRigEditor::SynchronizeViewportBoneSelection()
 		TArray<FRigBaseElement*> SelectedBones = RigBlueprint->Hierarchy->GetSelectedElements(ERigElementType::Bone);
 		for (const FRigBaseElement* SelectedBone : SelectedBones)
 		{
-			if(EditorSkelComp->GetReferenceSkeleton().IsValidIndex(SelectedBone->GetIndex()))
+			if(EditorSkelComp->GetReferenceSkeleton().IsValidIndex(SelectedBone->GetSubIndex()))
 			{
-				EditorSkelComp->BonesOfInterest.AddUnique(SelectedBone->GetIndex());
+				EditorSkelComp->BonesOfInterest.AddUnique(SelectedBone->GetSubIndex());
 			}
 		}
 	}
diff --git a/Engine/Plugins/Experimental/ControlRig/Source/ControlRigEditor/Private/Editor/ControlRigSkeletalMeshComponent.cpp b/Engine/Plugins/Experimental/ControlRig/Source/ControlRigEditor/Private/Editor/ControlRigSkeletalMeshComponent.cpp
--- a/Engine/Plugins/Experimental/ControlRig/Source/ControlRigEditor/Private/Editor/ControlRigSkeletalMeshComponent.cpp
+++ b/Engine/Plugins/Experimental/ControlRig/Source/ControlRigEditor/Private/Editor/ControlRigSkeletalMeshComponent.cpp
@@ -51,17 +51,21 @@ void UControlRigSkeletalMeshComponent::RebuildDebugDrawSkeleton()
 
 			DebugDrawSkeleton.Empty();
 			DebugDrawBones.Reset();
+			DebugDrawBoneIndexInHierarchy.Reset();
 
 			// create ref modifier
 			FReferenceSkeletonModifier RefSkelModifier(DebugDrawSkeleton, nullptr);
 
-			Hierarchy->ForEach<FRigBoneElement>([&RefSkelModifier, this, Hierarchy](FRigBoneElement* BoneElement) -> bool
+			TMap<FName, int32> AddedBoneMap;
+			TArray<FRigBoneElement*> BoneElements = Hierarchy->GetBones(true);
+			for(FRigBoneElement* BoneElement : BoneElements)
 			{
 				const int32 Index = BoneElement->GetIndex();
-				int32 ParentIndex = Hierarchy->GetFirstParent(Index);
-				if(ParentIndex != INDEX_NONE)
+				const FName ParentName = Hierarchy->GetFirstParent(BoneElement->GetKey()).Name;
+				int32 ParentIndex = INDEX_NONE; 
+				if(!ParentName.IsNone())
 				{
-					ParentIndex = Hierarchy->Get(ParentIndex)->GetSubIndex();
+					ParentIndex = AddedBoneMap.FindChecked(ParentName);
 				}
 					
 				FMeshBoneInfo NewMeshBoneInfo;
@@ -70,9 +74,10 @@ void UControlRigSkeletalMeshComponent::RebuildDebugDrawSkeleton()
 				// give ref pose here
 				RefSkelModifier.Add(NewMeshBoneInfo, Hierarchy->GetInitialGlobalTransform(Index));
 
+				AddedBoneMap.FindOrAdd(BoneElement->GetName(), DebugDrawBones.Num());
 				DebugDrawBones.Add(DebugDrawBones.Num());
-				return true;
-			});
+				DebugDrawBoneIndexInHierarchy.Add(BoneElement->GetIndex());
+			}
 		}
 	}
 }
@@ -84,11 +89,11 @@ FTransform UControlRigSkeletalMeshComponent::GetDrawTransform(int32 BoneIndex) c
 	if (ControlRigInstance)
 	{
 		UControlRig* ControlRig = ControlRigInstance->GetFirstAvailableControlRig();
-		if (ControlRig)
+		if (ControlRig && DebugDrawBoneIndexInHierarchy.IsValidIndex(BoneIndex))
 		{
 			// just copy it because this is not thread safe
 			URigHierarchy* Hierarchy = ControlRig->GetHierarchy();
-			return Hierarchy->GetGlobalTransform(BoneIndex);
+			return Hierarchy->GetGlobalTransform(DebugDrawBoneIndexInHierarchy[BoneIndex]);
 		}
 	}
 
diff --git a/Engine/Plugins/Experimental/ControlRig/Source/ControlRigEditor/Private/Editor/ControlRigSkeletalMeshComponent.h b/Engine/Plugins/Experimental/ControlRig/Source/ControlRigEditor/Private/Editor/ControlRigSkeletalMeshComponent.h
--- a/Engine/Plugins/Experimental/ControlRig/Source/ControlRigEditor/Private/Editor/ControlRigSkeletalMeshComponent.h
+++ b/Engine/Plugins/Experimental/ControlRig/Source/ControlRigEditor/Private/Editor/ControlRigSkeletalMeshComponent.h
@@ -46,4 +46,5 @@ public:
 private:
 	FReferenceSkeleton DebugDrawSkeleton;
 	TArray<FBoneIndexType> DebugDrawBones;
+	TArray<int32> DebugDrawBoneIndexInHierarchy;
 };
\ No newline at end of file
