Control Rig: Fix missing include for non-unity build

diff --git a/Engine/Plugins/Experimental/ControlRig/Source/ControlRig/Public/Rigs/RigHierarchyController.h b/Engine/Plugins/Experimental/ControlRig/Source/ControlRig/Public/Rigs/RigHierarchyController.h
--- a/Engine/Plugins/Experimental/ControlRig/Source/ControlRig/Public/Rigs/RigHierarchyController.h
+++ b/Engine/Plugins/Experimental/ControlRig/Source/ControlRig/Public/Rigs/RigHierarchyController.h
@@ -4,6 +4,7 @@
 
 #include "Rigs/RigHierarchy.h"
 #include "ReferenceSkeleton.h"
+#include "Animation/Skeleton.h"
 #include "RigHierarchyController.generated.h"
 
 UCLASS(BlueprintType)
