Sequence: Control Rig:  Allow for selection to persist when not visible. This fixes, most of not all of the selection bugs I have been seeing. Also Process all Curve Display changes at once at next tick, instead at one at a time, otherwise when a whole bunch of curves is displayed or undisplayed peformance is terrible as we do multiple syncs.

diff --git a/Engine/Plugins/Experimental/ControlRig/Source/ControlRigEditor/Private/Sequencer/ControlRigParameterTrackEditor.cpp b/Engine/Plugins/Experimental/ControlRig/Source/ControlRigEditor/Private/Sequencer/ControlRigParameterTrackEditor.cpp
--- a/Engine/Plugins/Experimental/ControlRig/Source/ControlRigEditor/Private/Sequencer/ControlRigParameterTrackEditor.cpp
+++ b/Engine/Plugins/Experimental/ControlRig/Source/ControlRigEditor/Private/Sequencer/ControlRigParameterTrackEditor.cpp
@@ -69,6 +69,7 @@
 #include "ControlRigBlueprint.h"
 #include "ControlRigBlueprintGeneratedClass.h"
 #include "Animation/SkeletalMeshActor.h"
+#include "TimerManager.h"
 
 #define LOCTEXT_NAMESPACE "FControlRigParameterTrackEditor"
 
@@ -164,7 +165,7 @@ static USkeleton* AcquireSkeletonFromObjectGuid(const FGuid& Guid, UObject** Obj
 }
 
 FControlRigParameterTrackEditor::FControlRigParameterTrackEditor(TSharedRef<ISequencer> InSequencer)
-	: FKeyframeTrackEditor<UMovieSceneControlRigParameterTrack>(InSequencer), bIsDoingSelection(false),  bFilterAssetBySkeleton(true), bFilterAssetByAnimatableControls(true)
+	: FKeyframeTrackEditor<UMovieSceneControlRigParameterTrack>(InSequencer), bIsDoingSelection(false), bFilterAssetBySkeleton(true), bCurveDisplayTickIsPending(false), bFilterAssetByAnimatableControls(true)
 
 {
 	FMovieSceneToolsModule::Get().RegisterAnimationBakeHelper(this);
@@ -1252,7 +1253,7 @@ void FControlRigParameterTrackEditor::OnCurveDisplayChanged(FCurveModel* CurveMo
 	TArray<FString> StringArray;
 	FControlRigEditMode* ControlRigEditMode = static_cast<FControlRigEditMode*>(GLevelEditorModeTools().GetActiveMode(FControlRigEditMode::ModeName));
 	UControlRig* ControlRig = nullptr;
-	TArray<FMovieSceneChannelHandle> Channels;
+
 	if (CurveModel)
 	{
 		UMovieSceneControlRigParameterSection* MovieSection = Cast<UMovieSceneControlRigParameterSection>(CurveModel->GetOwningObject());
@@ -1289,27 +1290,54 @@ void FControlRigParameterTrackEditor::OnCurveDisplayChanged(FCurveModel* CurveMo
 			{
 				//Not great but it should always be the third name
 				FName ControlName(*StringArray[2]);
-				ControlRig->SelectControl(ControlName, bDisplayed);
-
-				Channels.Add(FCurveModel->GetChannelHandle());
+				if (bDisplayed)
+				{
+					DisplayedChannels.Add(FCurveModel->GetChannelHandle());
+				}
+				else
+				{
+					UnDisplayedChannels.Add(FCurveModel->GetChannelHandle());
+				}
 			}
 			else
 			{
 				UE_LOG(LogControlRigEditor, Display, TEXT("Could not find Rig Control From FCurveModel::LongName"));
 			}
 		}
-		if (Channels.Num() > 0)
+		if (bCurveDisplayTickIsPending == false)
 		{
-			bool bSync = GetSequencer()->GetSequencerSettings()->ShouldSyncCurveEditorSelection();
-			GetSequencer()->SuspendSelectionBroadcast();
-			GetSequencer()->GetSequencerSettings()->SyncCurveEditorSelection(false);
-			GetSequencer()->SelectByChannels(MovieSection, Channels, false, bDisplayed);
-			GetSequencer()->RefreshTree();
-			GetSequencer()->GetSequencerSettings()->SyncCurveEditorSelection(bSync);
-			GetSequencer()->ResumeSelectionBroadcast();
+			bCurveDisplayTickIsPending = true;
+			GEditor->GetTimerManager()->SetTimerForNextTick([MovieSection, this]()
+			{
+				if (DisplayedChannels.Num() > 0 || UnDisplayedChannels.Num() > 0)
+				{
+					TGuardValue<bool> Guard(bIsDoingSelection, true);
+
+					bool bSync = GetSequencer()->GetSequencerSettings()->ShouldSyncCurveEditorSelection();
+					GetSequencer()->SuspendSelectionBroadcast();
+					GetSequencer()->GetSequencerSettings()->SyncCurveEditorSelection(false);
+					if (UnDisplayedChannels.Num() > 0)
+					{
+						GetSequencer()->SelectByChannels(MovieSection, UnDisplayedChannels, false, false);
+					}
+					if (DisplayedChannels.Num() > 0)
+					{
+						GetSequencer()->SelectByChannels(MovieSection, DisplayedChannels, false, true);
+					}
+					GetSequencer()->RefreshTree();
+					GetSequencer()->GetSequencerSettings()->SyncCurveEditorSelection(bSync);
+					GetSequencer()->ResumeSelectionBroadcast();
+				}
+				UnDisplayedChannels.SetNum(0);
+				DisplayedChannels.SetNum(0);
+				bCurveDisplayTickIsPending = false;
+			});
+
 		}
+
 	}
 }
+
 void FControlRigParameterTrackEditor::OnTreeViewChanged()
 {
 	if (!bIsDoingSelection)
diff --git a/Engine/Plugins/Experimental/ControlRig/Source/ControlRigEditor/Private/Sequencer/ControlRigParameterTrackEditor.h b/Engine/Plugins/Experimental/ControlRig/Source/ControlRigEditor/Private/Sequencer/ControlRigParameterTrackEditor.h
--- a/Engine/Plugins/Experimental/ControlRig/Source/ControlRigEditor/Private/Sequencer/ControlRigParameterTrackEditor.h
+++ b/Engine/Plugins/Experimental/ControlRig/Source/ControlRigEditor/Private/Sequencer/ControlRigParameterTrackEditor.h
@@ -173,6 +173,11 @@ private:
 	FDelegateHandle OnActorAddedToSequencerHandle;
 	FDelegateHandle OnTreeViewChangedHandle;
 
+
+	//used to sync curve editor selections/displays on next tick for performance reasons
+	TArray<FMovieSceneChannelHandle> DisplayedChannels;
+	TArray<FMovieSceneChannelHandle> UnDisplayedChannels;
+	bool bCurveDisplayTickIsPending;
 	void BindControlRig(UControlRig* ControlRig);
 	void UnbindControlRig(UControlRig* ControlRig);
 	void UnbindAllControlRigs();
diff --git a/Engine/Source/Editor/Sequencer/Private/SSequencerTreeView.cpp b/Engine/Source/Editor/Sequencer/Private/SSequencerTreeView.cpp
--- a/Engine/Source/Editor/Sequencer/Private/SSequencerTreeView.cpp
+++ b/Engine/Source/Editor/Sequencer/Private/SSequencerTreeView.cpp
@@ -186,6 +186,7 @@ void SSequencerTreeView::Construct(const FArguments& InArgs, const TSharedRef<FS
 		.OnContextMenuOpening( this, &SSequencerTreeView::OnContextMenuOpening )
 		.OnSetExpansionRecursive(this, &SSequencerTreeView::SetItemExpansionRecursive)
 		.HighlightParentNodesForSelection(true)
+		.AllowInvisibleItemSelection(true)  //without this we deselect everything when we filter or we collapse, etc..
 	);
 }
 
