CIS FIX

diff --git a/Engine/Plugins/Experimental/ControlRig/Source/ControlRigEditor/Private/Sequencer/ControlRigParameterTrackEditor.cpp b/Engine/Plugins/Experimental/ControlRig/Source/ControlRigEditor/Private/Sequencer/ControlRigParameterTrackEditor.cpp
--- a/Engine/Plugins/Experimental/ControlRig/Source/ControlRigEditor/Private/Sequencer/ControlRigParameterTrackEditor.cpp
+++ b/Engine/Plugins/Experimental/ControlRig/Source/ControlRigEditor/Private/Sequencer/ControlRigParameterTrackEditor.cpp
@@ -165,7 +165,7 @@ static USkeleton* AcquireSkeletonFromObjectGuid(const FGuid& Guid, UObject** Obj
 }
 
 FControlRigParameterTrackEditor::FControlRigParameterTrackEditor(TSharedRef<ISequencer> InSequencer)
-	: FKeyframeTrackEditor<UMovieSceneControlRigParameterTrack>(InSequencer), bIsDoingSelection(false), bFilterAssetBySkeleton(true), bCurveDisplayTickIsPending(false), bFilterAssetByAnimatableControls(true)
+	: FKeyframeTrackEditor<UMovieSceneControlRigParameterTrack>(InSequencer), bIsDoingSelection(false), bCurveDisplayTickIsPending(false), bFilterAssetBySkeleton(true), bFilterAssetByAnimatableControls(true)
 
 {
 	FMovieSceneToolsModule::Get().RegisterAnimationBakeHelper(this);
