Control Rig: Crash opening CR Function

diff --git a/Engine/Plugins/Experimental/ControlRig/Source/ControlRigDeveloper/Private/Graph/ControlRigGraph.cpp b/Engine/Plugins/Experimental/ControlRig/Source/ControlRigDeveloper/Private/Graph/ControlRigGraph.cpp
--- a/Engine/Plugins/Experimental/ControlRig/Source/ControlRigDeveloper/Private/Graph/ControlRigGraph.cpp
+++ b/Engine/Plugins/Experimental/ControlRig/Source/ControlRigDeveloper/Private/Graph/ControlRigGraph.cpp
@@ -64,6 +64,8 @@ void UControlRigGraph::Serialize(FArchive& Ar)
 
 #if WITH_EDITOR
 
+TArray<TSharedPtr<FString>> UControlRigGraph::EmptyElementNameList;
+
 void UControlRigGraph::CacheNameLists(URigHierarchy* InHierarchy, const FControlRigDrawContainer* DrawContainer)
 {
 	if (UControlRigGraph* OuterGraph = Cast<UControlRigGraph>(GetOuter()))
@@ -96,6 +98,17 @@ const TArray<TSharedPtr<FString>>& UControlRigGraph::GetElementNameList(ERigElem
 	{
 		return OuterGraph->GetElementNameList(InElementType);
 	}
+	if(!ElementNameLists.Contains(InElementType))
+	{
+		UControlRigBlueprint* Blueprint = GetBlueprint();
+		if(Blueprint == nullptr)
+		{
+			return EmptyElementNameList;
+		}
+
+		UControlRigGraph* MutableThis = (UControlRigGraph*)this;
+		MutableThis->CacheNameLists(Blueprint->Hierarchy, &Blueprint->DrawContainer);
+	}
 	return ElementNameLists.FindChecked(InElementType);
 }
 
diff --git a/Engine/Plugins/Experimental/ControlRig/Source/ControlRigDeveloper/Public/Graph/ControlRigGraph.h b/Engine/Plugins/Experimental/ControlRig/Source/ControlRigDeveloper/Public/Graph/ControlRigGraph.h
--- a/Engine/Plugins/Experimental/ControlRig/Source/ControlRigDeveloper/Public/Graph/ControlRigGraph.h
+++ b/Engine/Plugins/Experimental/ControlRig/Source/ControlRigDeveloper/Public/Graph/ControlRigGraph.h
@@ -39,10 +39,6 @@ public:
 #endif
 #if WITH_EDITOR
 
-
-	void CacheNameLists(URigHierarchy* InHierarchy, const FControlRigDrawContainer* DrawContainer);
-	const TArray<TSharedPtr<FString>>& GetElementNameList(ERigElementType InElementType = ERigElementType::Bone) const;
-
 	FORCEINLINE const TArray<TSharedPtr<FString>>& GetBoneNameList(URigVMPin* InPin = nullptr) const
 	{
 		return GetElementNameList(ERigElementType::Bone);
@@ -59,6 +55,9 @@ public:
 	{
 		return GetElementNameList(ERigElementType::Curve);
 	}
+
+	void CacheNameLists(URigHierarchy* InHierarchy, const FControlRigDrawContainer* DrawContainer);
+	const TArray<TSharedPtr<FString>>& GetElementNameList(ERigElementType InElementType = ERigElementType::Bone) const;
 	const TArray<TSharedPtr<FString>>& GetElementNameList(URigVMPin* InPin = nullptr) const;
 	const TArray<TSharedPtr<FString>>& GetDrawingNameList(URigVMPin* InPin = nullptr) const;
 
@@ -143,6 +142,8 @@ private:
 	URigVMController* GetTemplateController();
 	void HandleVMCompiledEvent(UBlueprint* InBlueprint, URigVM* InVM);
 
+	static TArray<TSharedPtr<FString>> EmptyElementNameList;
+	
 	friend class UControlRigUnitNodeSpawner;
 	friend class UControlRigVariableNodeSpawner;
 	friend class UControlRigParameterNodeSpawner;
