Control Rig: Make sure to store hierarchy only alongside default object

diff --git a/Engine/Plugins/Experimental/ControlRig/Source/ControlRig/Public/ControlRig.h b/Engine/Plugins/Experimental/ControlRig/Source/ControlRig/Public/ControlRig.h
--- a/Engine/Plugins/Experimental/ControlRig/Source/ControlRig/Public/ControlRig.h
+++ b/Engine/Plugins/Experimental/ControlRig/Source/ControlRig/Public/ControlRig.h
@@ -195,6 +195,10 @@ public:
 		if(DynamicHierarchy == nullptr)
 		{
 			DynamicHierarchy = NewObject<URigHierarchy>(this);
+			if (!HasAnyFlags(RF_ClassDefaultObject))
+			{
+				DynamicHierarchy->SetFlags(RF_Transient);
+			}
 		}
 		return DynamicHierarchy;
 	}
