Control Rig: Crash opening CR Function - additional use case

diff --git a/Engine/Plugins/Experimental/ControlRig/Source/ControlRigDeveloper/Private/Graph/ControlRigGraph.cpp b/Engine/Plugins/Experimental/ControlRig/Source/ControlRigDeveloper/Private/Graph/ControlRigGraph.cpp
--- a/Engine/Plugins/Experimental/ControlRig/Source/ControlRigDeveloper/Private/Graph/ControlRigGraph.cpp
+++ b/Engine/Plugins/Experimental/ControlRig/Source/ControlRigDeveloper/Private/Graph/ControlRigGraph.cpp
@@ -76,6 +76,7 @@ void UControlRigGraph::CacheNameLists(URigHierarchy* InHierarchy, const FControl
 	check(InHierarchy);
 	check(DrawContainer);
 
+	TArray<TSharedPtr<FString>>& AllNameList = ElementNameLists.FindOrAdd(ERigElementType::All);
 	TArray<TSharedPtr<FString>>& BoneNameList = ElementNameLists.FindOrAdd(ERigElementType::Bone);
 	TArray<TSharedPtr<FString>>& SpaceNameList = ElementNameLists.FindOrAdd(ERigElementType::Space);
 	TArray<TSharedPtr<FString>>& ControlNameList = ElementNameLists.FindOrAdd(ERigElementType::Control);
@@ -83,6 +84,7 @@ void UControlRigGraph::CacheNameLists(URigHierarchy* InHierarchy, const FControl
 	TArray<TSharedPtr<FString>>& RigidBodyNameList = ElementNameLists.FindOrAdd(ERigElementType::RigidBody);
 	TArray<TSharedPtr<FString>>& AuxiliaryNameList = ElementNameLists.FindOrAdd(ERigElementType::Auxiliary);
 	
+	CacheNameListForHierarchy<FRigBaseElement>(InHierarchy, AllNameList);
 	CacheNameListForHierarchy<FRigBoneElement>(InHierarchy, BoneNameList);
 	CacheNameListForHierarchy<FRigSpaceElement>(InHierarchy, SpaceNameList);
 	CacheNameListForHierarchy<FRigControlElement>(InHierarchy, ControlNameList);
