Control Rig: Ensure to compute local offset before dirtying global

diff --git a/Engine/Plugins/Experimental/ControlRig/Source/ControlRig/Private/Rigs/RigHierarchy.cpp b/Engine/Plugins/Experimental/ControlRig/Source/ControlRig/Private/Rigs/RigHierarchy.cpp
--- a/Engine/Plugins/Experimental/ControlRig/Source/ControlRig/Private/Rigs/RigHierarchy.cpp
+++ b/Engine/Plugins/Experimental/ControlRig/Source/ControlRig/Private/Rigs/RigHierarchy.cpp
@@ -1805,6 +1805,10 @@ void URigHierarchy::PropagateDirtyFlags(FRigTransformElement* InTransformElement
 
 #endif
 
+			if(FRigControlElement* ControlElement = Cast<FRigControlElement>(ElementToDirty.Element))
+			{
+				GetControlOffsetTransform(ControlElement, LocalType);
+			}
 			GetTransform(ElementToDirty.Element, TypeToCompute); // make sure the local / global transform is up 2 date
 
 #if URIGHIERARCHY_RECURSIVE_DIRTY_PROPAGATION
diff --git a/Engine/Plugins/Experimental/ControlRig/Source/ControlRig/Private/Rigs/RigHierarchyController.cpp b/Engine/Plugins/Experimental/ControlRig/Source/ControlRig/Private/Rigs/RigHierarchyController.cpp
--- a/Engine/Plugins/Experimental/ControlRig/Source/ControlRig/Private/Rigs/RigHierarchyController.cpp
+++ b/Engine/Plugins/Experimental/ControlRig/Source/ControlRig/Private/Rigs/RigHierarchyController.cpp
@@ -246,6 +246,8 @@ FRigElementKey URigHierarchyController::AddControl(
 	Hierarchy->SetControlValue(NewElement, InValue, ERigControlValueType::Initial, false);
 
 	NewElement->Parent.MarkDirty(ERigTransformType::InitialGlobal);
+	NewElement->Offset.MarkDirty(ERigTransformType::InitialGlobal);
+	NewElement->Gizmo.MarkDirty(ERigTransformType::InitialGlobal);
 	NewElement->Parent.Current = NewElement->Parent.Initial;
 	NewElement->Offset.Current = NewElement->Offset.Initial;
 	NewElement->Pose.Current = NewElement->Pose.Initial;
