Control Rig: Fix infinite loop in SControlHIerarchy

diff --git a/Engine/Plugins/Experimental/ControlRig/Source/ControlRigEditor/Private/EditMode/SControlHierarchy.cpp b/Engine/Plugins/Experimental/ControlRig/Source/ControlRigEditor/Private/EditMode/SControlHierarchy.cpp
--- a/Engine/Plugins/Experimental/ControlRig/Source/ControlRigEditor/Private/EditMode/SControlHierarchy.cpp
+++ b/Engine/Plugins/Experimental/ControlRig/Source/ControlRigEditor/Private/EditMode/SControlHierarchy.cpp
@@ -362,7 +362,7 @@ void SControlHierarchy::AddElement(FRigBaseElement* InElement)
 	FRigElementKey ParentKey = Hierarchy->GetFirstParent(InElement->GetKey());
 	while(ParentKey.IsValid() && ParentKey.Type != ERigElementType::Control)
 	{
-		ParentKey = Hierarchy->GetFirstParent(InElement->GetKey());
+		ParentKey = Hierarchy->GetFirstParent(ParentKey);
 	}
 	AddElement(InElement->GetKey(), ParentKey);
 }
