Control Rig: Ensure to dirty the pose on controls after creation

diff --git a/Engine/Plugins/Experimental/ControlRig/Source/ControlRig/Private/Rigs/RigHierarchyController.cpp b/Engine/Plugins/Experimental/ControlRig/Source/ControlRig/Private/Rigs/RigHierarchyController.cpp
--- a/Engine/Plugins/Experimental/ControlRig/Source/ControlRig/Private/Rigs/RigHierarchyController.cpp
+++ b/Engine/Plugins/Experimental/ControlRig/Source/ControlRig/Private/Rigs/RigHierarchyController.cpp
@@ -247,6 +247,7 @@ FRigElementKey URigHierarchyController::AddControl(
 
 	NewElement->Parent.MarkDirty(ERigTransformType::InitialGlobal);
 	NewElement->Offset.MarkDirty(ERigTransformType::InitialGlobal);
+	NewElement->Pose.MarkDirty(ERigTransformType::InitialGlobal);
 	NewElement->Gizmo.MarkDirty(ERigTransformType::InitialGlobal);
 	NewElement->Parent.Current = NewElement->Parent.Initial;
 	NewElement->Offset.Current = NewElement->Offset.Initial;
