Control Rig: Fix ResetTransforms within the SRigHierarchy

diff --git a/Engine/Plugins/Experimental/ControlRig/Source/ControlRigEditor/Private/Editor/SRigHierarchy.cpp b/Engine/Plugins/Experimental/ControlRig/Source/ControlRigEditor/Private/Editor/SRigHierarchy.cpp
--- a/Engine/Plugins/Experimental/ControlRig/Source/ControlRigEditor/Private/Editor/SRigHierarchy.cpp
+++ b/Engine/Plugins/Experimental/ControlRig/Source/ControlRigEditor/Private/Editor/SRigHierarchy.cpp
@@ -2152,8 +2152,8 @@ void SRigHierarchy::HandleResetTransform(bool bSelectionOnly)
 				for (FRigElementKey Key : KeysToReset)
 				{
 					const FTransform InitialTransform = GetHierarchy()->GetInitialLocalTransform(Key);
-					GetHierarchy()->SetLocalTransform(Key, InitialTransform, true, true, true);
-					DebuggedHierarchy->SetLocalTransform(Key, InitialTransform, true, true, true);
+					GetHierarchy()->SetLocalTransform(Key, InitialTransform, false, true, true);
+					DebuggedHierarchy->SetLocalTransform(Key, InitialTransform, false, true, true);
 				}
 			}
 		}
