Control Rig: Create the dynamic hierarchy on the CDO using the default subobject mechanism

diff --git a/Engine/Plugins/Experimental/ControlRig/Source/ControlRig/Private/ControlRig.cpp b/Engine/Plugins/Experimental/ControlRig/Source/ControlRig/Private/ControlRig.cpp
--- a/Engine/Plugins/Experimental/ControlRig/Source/ControlRig/Private/ControlRig.cpp
+++ b/Engine/Plugins/Experimental/ControlRig/Source/ControlRig/Private/ControlRig.cpp
@@ -67,6 +67,7 @@ UControlRig::UControlRig(const FObjectInitializer& ObjectInitializer)
 	, InterRigSyncBracket(0)
 {
 	VM = ObjectInitializer.CreateDefaultSubobject<URigVM>(this, TEXT("VM"));
+	DynamicHierarchy = ObjectInitializer.CreateDefaultSubobject<URigHierarchy>(this, TEXT("DynamicHierarchy"));
 
 	EventQueue.Add(FRigUnit_BeginExecution::EventName);
 }
