Control Rig: Static analysis fixes for hierarchy pointer

diff --git a/Engine/Plugins/Experimental/ControlRig/Source/ControlRigEditor/Private/EditMode/SControlHierarchy.cpp b/Engine/Plugins/Experimental/ControlRig/Source/ControlRigEditor/Private/EditMode/SControlHierarchy.cpp
--- a/Engine/Plugins/Experimental/ControlRig/Source/ControlRigEditor/Private/EditMode/SControlHierarchy.cpp
+++ b/Engine/Plugins/Experimental/ControlRig/Source/ControlRigEditor/Private/EditMode/SControlHierarchy.cpp
@@ -279,6 +279,7 @@ void SControlHierarchy::OnSelectionChanged(TSharedPtr<FControlTreeElement> Selec
 	{
 		TGuardValue<bool> Guard(bSelecting, true);
 		URigHierarchy* Hierarchy = GetHierarchy();
+		check(Hierarchy);
 		URigHierarchyController* Controller = Hierarchy->GetController(true);
 		check(Controller);
 
@@ -359,6 +360,7 @@ void SControlHierarchy::OnRigElementSelected(UControlRig* Subject, FRigControlEl
 void SControlHierarchy::AddElement(FRigBaseElement* InElement)
 {
 	const URigHierarchy* Hierarchy = GetHierarchy();
+	check(Hierarchy);
 	FRigElementKey ParentKey = Hierarchy->GetFirstParent(InElement->GetKey());
 	while(ParentKey.IsValid() && ParentKey.Type != ERigElementType::Control)
 	{
diff --git a/Engine/Plugins/Experimental/ControlRig/Source/ControlRigEditor/Private/Editor/SRigHierarchy.cpp b/Engine/Plugins/Experimental/ControlRig/Source/ControlRigEditor/Private/Editor/SRigHierarchy.cpp
--- a/Engine/Plugins/Experimental/ControlRig/Source/ControlRigEditor/Private/Editor/SRigHierarchy.cpp
+++ b/Engine/Plugins/Experimental/ControlRig/Source/ControlRigEditor/Private/Editor/SRigHierarchy.cpp
@@ -656,6 +656,7 @@ void SRigHierarchy::RefreshTreeView()
 	if (ControlRigBlueprint.IsValid())
 	{
 		URigHierarchy* Hierarchy = GetHierarchy();
+		check(Hierarchy);
 
 		Hierarchy->Traverse([&](FRigBaseElement* Element, bool& bContinue)
 		{
@@ -898,6 +899,8 @@ void SRigHierarchy::AddElement(FRigBaseElement* InElement, const bool bIgnoreTex
 	}
 
 	const URigHierarchy* Hierarchy = GetHierarchy();
+	check(Hierarchy);
+
 	AddElement(InElement->GetKey(), FRigElementKey(), bIgnoreTextFilter);
 
 	if (ElementMap.Contains(InElement->GetKey()))
@@ -1104,6 +1107,8 @@ TSharedPtr< SWidget > SRigHierarchy::CreateContextMenu()
 void SRigHierarchy::OnItemClicked(TSharedPtr<FRigTreeElement> InItem)
 {
 	URigHierarchy* Hierarchy = GetHierarchy();
+	check(Hierarchy);
+
 	if (Hierarchy->IsSelected(InItem->Key))
 	{
 		if (ControlRigEditor.IsValid())
@@ -1957,7 +1962,9 @@ FReply SRigHierarchy::OnAcceptDrop(const FDragDropEvent& DragDropEvent, EItemDro
 	if (RigDragDropOp.IsValid())
 	{
 		URigHierarchy* Hierarchy = GetHierarchy();
+		check(Hierarchy);
 		URigHierarchy* DebuggedHierarchy = GetDebuggedHierarchy();
+		check(DebuggedHierarchy);
 		URigHierarchyController* Controller = Hierarchy->GetController(true);
 		check(Controller);
 
