Control Rig: mark dynamic hierarchy with rf_public and default sub object

diff --git a/Engine/Plugins/Experimental/ControlRig/Source/ControlRig/Private/ControlRig.cpp b/Engine/Plugins/Experimental/ControlRig/Source/ControlRig/Private/ControlRig.cpp
--- a/Engine/Plugins/Experimental/ControlRig/Source/ControlRig/Private/ControlRig.cpp
+++ b/Engine/Plugins/Experimental/ControlRig/Source/ControlRig/Private/ControlRig.cpp
@@ -1036,6 +1036,21 @@ void UControlRig::Serialize(FArchive& Ar)
 	Ar.UsingCustomVersion(FControlRigObjectVersion::GUID);
 }
 
+void UControlRig::PostLoad()
+{
+	Super::PostLoad();
+
+	if(HasAnyFlags(RF_ClassDefaultObject))
+	{
+		if(DynamicHierarchy)
+		{
+			// Some dynamic hierarchy objects have been created using NewObject<> instead of CreateDefaultSubObjects.
+			// Assets from that version require the dynamic hierarchy to be flagged as below.
+			DynamicHierarchy->SetFlags(DynamicHierarchy->GetFlags() | RF_Public | RF_DefaultSubObject);
+		}
+	}
+}
+
 TArray<FRigControlElement*> UControlRig::AvailableControls() const
 {
 	if(DynamicHierarchy)
diff --git a/Engine/Plugins/Experimental/ControlRig/Source/ControlRig/Public/ControlRig.h b/Engine/Plugins/Experimental/ControlRig/Source/ControlRig/Public/ControlRig.h
--- a/Engine/Plugins/Experimental/ControlRig/Source/ControlRig/Public/ControlRig.h
+++ b/Engine/Plugins/Experimental/ControlRig/Source/ControlRig/Public/ControlRig.h
@@ -76,6 +76,7 @@ private:
 
 public:
 	virtual void Serialize(FArchive& Ar) override;
+	virtual void PostLoad() override;
 #if WITH_EDITOR
 	virtual void PostEditUndo() override;
 #endif
