Control Rig: Rename auxiliary element to socket, enable dynamic hierarchy option in Rig Hierarchy

diff --git a/Engine/Plugins/Experimental/ControlRig/Source/ControlRig/Private/ControlRig.cpp b/Engine/Plugins/Experimental/ControlRig/Source/ControlRig/Private/ControlRig.cpp
--- a/Engine/Plugins/Experimental/ControlRig/Source/ControlRig/Private/ControlRig.cpp
+++ b/Engine/Plugins/Experimental/ControlRig/Source/ControlRig/Private/ControlRig.cpp
@@ -561,7 +561,7 @@ void UControlRig::Execute(const EControlRigState InState, const FName& InEventNa
 	if(DynamicHierarchy)
 	{
 		// if we have any aux elements dirty them
-		DynamicHierarchy->UpdateAuxiliaryElements(&Context);
+		DynamicHierarchy->UpdateSockets(&Context);
 	}
 
 #if WITH_EDITOR
@@ -1325,6 +1325,7 @@ FName UControlRig::AddTransientControl(URigVMPin* InPin, FRigElementKey SpaceKey
 		}
 	}
 	Settings.bIsTransientControl = true;
+	Settings.DisplayName = TEXT("Temporary Control");
 
 	FRigElementKey Parent;
 	if(SpaceKey.IsValid() && SpaceKey.Type == ERigElementType::Bone)
@@ -1445,6 +1446,7 @@ FName UControlRig::AddTransientControl(const FRigElementKey& InElement)
 
 	FRigControlSettings Settings;
 	Settings.bIsTransientControl = true;
+	Settings.DisplayName = TEXT("Temporary Control");
 
 	FRigElementKey Parent;
 	switch (InElement.Type)
diff --git a/Engine/Plugins/Experimental/ControlRig/Source/ControlRig/Private/Rigs/RigHierarchy.cpp b/Engine/Plugins/Experimental/ControlRig/Source/ControlRig/Private/Rigs/RigHierarchy.cpp
--- a/Engine/Plugins/Experimental/ControlRig/Source/ControlRig/Private/Rigs/RigHierarchy.cpp
+++ b/Engine/Plugins/Experimental/ControlRig/Source/ControlRig/Private/Rigs/RigHierarchy.cpp
@@ -124,9 +124,9 @@ void URigHierarchy::Load(FArchive& Ar)
 				Element = new FRigRigidBodyElement();
 				break;
 			}
-			case ERigElementType::Auxiliary:
+			case ERigElementType::Socket:
 			{
-				Element = new FRigAuxiliaryElement();
+				Element = new FRigSocketElement();
 				break;
 			}
 			default:
@@ -253,28 +253,28 @@ void URigHierarchy::CopyPose(URigHierarchy* InHierarchy, bool bCurrent, bool bIn
 	}
 }
 
-void URigHierarchy::UpdateAuxiliaryElements(const FRigUnitContext* InContext)
+void URigHierarchy::UpdateSockets(const FRigUnitContext* InContext)
 {
 	check(InContext);
 	
 	for (int32 ElementIndex = 0; ElementIndex < Elements.Num(); ElementIndex++)
 	{
-		if(FRigAuxiliaryElement* AuxiliaryElement = Cast<FRigAuxiliaryElement>(Elements[ElementIndex]))
+		if(FRigSocketElement* Socket = Cast<FRigSocketElement>(Elements[ElementIndex]))
 		{
-			const FTransform InitialWorldTransform = AuxiliaryElement->GetAuxiliaryWorldTransform(InContext, true);
-			const FTransform CurrentWorldTransform = AuxiliaryElement->GetAuxiliaryWorldTransform(InContext, false);
+			const FTransform InitialWorldTransform = Socket->GetSocketWorldTransform(InContext, true);
+			const FTransform CurrentWorldTransform = Socket->GetSocketWorldTransform(InContext, false);
 
 			const FTransform InitialGlobalTransform = InitialWorldTransform.GetRelativeTransform(InContext->ToWorldSpaceTransform);
 			const FTransform CurrentGlobalTransform = CurrentWorldTransform.GetRelativeTransform(InContext->ToWorldSpaceTransform);
 
-			const FTransform InitialParentTransform = GetParentTransform(AuxiliaryElement, ERigTransformType::InitialGlobal); 
-			const FTransform CurrentParentTransform = GetParentTransform(AuxiliaryElement, ERigTransformType::CurrentGlobal);
+			const FTransform InitialParentTransform = GetParentTransform(Socket, ERigTransformType::InitialGlobal); 
+			const FTransform CurrentParentTransform = GetParentTransform(Socket, ERigTransformType::CurrentGlobal);
 
 			const FTransform InitialLocalTransform = InitialGlobalTransform.GetRelativeTransform(InitialParentTransform);
 			const FTransform CurrentLocalTransform = CurrentGlobalTransform.GetRelativeTransform(CurrentParentTransform);
 
-			SetTransform(AuxiliaryElement, InitialLocalTransform, ERigTransformType::InitialLocal, true, false);
-			SetTransform(AuxiliaryElement, CurrentLocalTransform, ERigTransformType::CurrentLocal, true, false);
+			SetTransform(Socket, InitialLocalTransform, ERigTransformType::InitialLocal, true, false);
+			SetTransform(Socket, CurrentLocalTransform, ERigTransformType::CurrentLocal, true, false);
 		}
 	}
 }
diff --git a/Engine/Plugins/Experimental/ControlRig/Source/ControlRig/Private/Rigs/RigHierarchyController.cpp b/Engine/Plugins/Experimental/ControlRig/Source/ControlRig/Private/Rigs/RigHierarchyController.cpp
--- a/Engine/Plugins/Experimental/ControlRig/Source/ControlRig/Private/Rigs/RigHierarchyController.cpp
+++ b/Engine/Plugins/Experimental/ControlRig/Source/ControlRig/Private/Rigs/RigHierarchyController.cpp
@@ -323,8 +323,8 @@ FRigElementKey URigHierarchyController::AddRigidBody(FName InName, FRigElementKe
 	return NewElement->Key;
 }
 
-FRigElementKey URigHierarchyController::AddAuxiliaryElement(FName InName, FRigElementKey InParent,
-	FRigAuxiliaryElementGetWorldTransformDelegate InDelegate, bool bSetupUndo)
+FRigElementKey URigHierarchyController::AddSocket(FName InName, FRigElementKey InParent,
+	FRigSocketGetWorldTransformDelegate InDelegate, bool bSetupUndo)
 {
 	if(!IsValid())
 	{
@@ -335,13 +335,13 @@ FRigElementKey URigHierarchyController::AddAuxiliaryElement(FName InName, FRigEl
 	TSharedPtr<FScopedTransaction> TransactionPtr;
 	if(bSetupUndo)
 	{
-		TransactionPtr = MakeShared<FScopedTransaction>(NSLOCTEXT("RigHierarchyController", "Add Auxiliary Element", "Add Auxiliary Element"));
+		TransactionPtr = MakeShared<FScopedTransaction>(NSLOCTEXT("RigHierarchyController", "Add Socket", "Add Socket"));
 		Hierarchy->Modify();
 	}
 #endif
 
-	FRigAuxiliaryElement* NewElement = new FRigAuxiliaryElement();
-	NewElement->Key.Type = ERigElementType::Auxiliary;
+	FRigSocketElement* NewElement = new FRigSocketElement();
+	NewElement->Key.Type = ERigElementType::Socket;
 	NewElement->Key.Name = Hierarchy->GetSafeNewName(InName.ToString(), NewElement->Key.Type);
 	NewElement->GetWorldTransformDelegate = InDelegate;
 	AddElement(NewElement, Hierarchy->Get(Hierarchy->GetIndex(InParent)), true);
@@ -754,10 +754,10 @@ FString URigHierarchyController::ExportToText(TArray<FRigElementKey> InKeys) con
 				FRigRigidBodyElement::StaticStruct()->ExportText(PerElementData.Content, Element, &DefaultElement, nullptr, PPF_None, nullptr);
 				break;
 			}
-			case ERigElementType::Auxiliary:
+			case ERigElementType::Socket:
 			{
-				FRigAuxiliaryElement DefaultElement;
-				FRigAuxiliaryElement::StaticStruct()->ExportText(PerElementData.Content, Element, &DefaultElement, nullptr, PPF_None, nullptr);
+				FRigSocketElement DefaultElement;
+				FRigSocketElement::StaticStruct()->ExportText(PerElementData.Content, Element, &DefaultElement, nullptr, PPF_None, nullptr);
 				break;
 			}
 			default:
@@ -869,10 +869,10 @@ TArray<FRigElementKey> URigHierarchyController::ImportFromText(FString InContent
 				FRigRigidBodyElement::StaticStruct()->ImportText(*PerElementData.Content, NewElement, nullptr, EPropertyPortFlags::PPF_None, &ErrorPipe, FRigRigidBodyElement::StaticStruct()->GetName(), true);
 				break;
 			}
-			case ERigElementType::Auxiliary:
+			case ERigElementType::Socket:
 			{
-				NewElement = new FRigAuxiliaryElement();
-				FRigAuxiliaryElement::StaticStruct()->ImportText(*PerElementData.Content, NewElement, nullptr, EPropertyPortFlags::PPF_None, &ErrorPipe, FRigAuxiliaryElement::StaticStruct()->GetName(), true);
+				NewElement = new FRigSocketElement();
+				FRigSocketElement::StaticStruct()->ImportText(*PerElementData.Content, NewElement, nullptr, EPropertyPortFlags::PPF_None, &ErrorPipe, FRigSocketElement::StaticStruct()->GetName(), true);
 				break;
 			}
 			default:
diff --git a/Engine/Plugins/Experimental/ControlRig/Source/ControlRig/Private/Rigs/RigHierarchyElements.cpp b/Engine/Plugins/Experimental/ControlRig/Source/ControlRig/Private/Rigs/RigHierarchyElements.cpp
--- a/Engine/Plugins/Experimental/ControlRig/Source/ControlRig/Private/Rigs/RigHierarchyElements.cpp
+++ b/Engine/Plugins/Experimental/ControlRig/Source/ControlRig/Private/Rigs/RigHierarchyElements.cpp
@@ -506,20 +506,20 @@ void FRigRigidBodyElement::Load(FArchive& Ar, URigHierarchy* Hierarchy, ESeriali
 }
 
 ////////////////////////////////////////////////////////////////////////////////
-// FRigAuxiliaryElement
+// FRigSocketElement
 ////////////////////////////////////////////////////////////////////////////////
 
-void FRigAuxiliaryElement::Save(FArchive& Ar, URigHierarchy* Hierarchy, ESerializationPhase SerializationPhase)
+void FRigSocketElement::Save(FArchive& Ar, URigHierarchy* Hierarchy, ESerializationPhase SerializationPhase)
 {
 	Super::Save(Ar, Hierarchy, SerializationPhase);
 }
 
-void FRigAuxiliaryElement::Load(FArchive& Ar, URigHierarchy* Hierarchy, ESerializationPhase SerializationPhase)
+void FRigSocketElement::Load(FArchive& Ar, URigHierarchy* Hierarchy, ESerializationPhase SerializationPhase)
 {
 	Super::Load(Ar, Hierarchy, SerializationPhase);
 }
 
-FTransform FRigAuxiliaryElement::GetAuxiliaryWorldTransform(const FRigUnitContext* InContext, bool bInitial) const
+FTransform FRigSocketElement::GetSocketWorldTransform(const FRigUnitContext* InContext, bool bInitial) const
 {
 	if(GetWorldTransformDelegate.IsBound())
 	{
@@ -528,11 +528,11 @@ FTransform FRigAuxiliaryElement::GetAuxiliaryWorldTransform(const FRigUnitContex
 	return FTransform::Identity;
 }
 
-void FRigAuxiliaryElement::CopyPose(FRigBaseElement* InOther, bool bCurrent, bool bInitial)
+void FRigSocketElement::CopyPose(FRigBaseElement* InOther, bool bCurrent, bool bInitial)
 {
 	Super::CopyPose(InOther, bCurrent, bInitial);
 	
-	if(FRigAuxiliaryElement* Other = Cast<FRigAuxiliaryElement>(InOther))
+	if(FRigSocketElement* Other = Cast<FRigSocketElement>(InOther))
 	{
 		if(Other->GetWorldTransformDelegate.IsBound())
 		{
diff --git a/Engine/Plugins/Experimental/ControlRig/Source/ControlRig/Public/Rigs/RigHierarchy.h b/Engine/Plugins/Experimental/ControlRig/Source/ControlRig/Public/Rigs/RigHierarchy.h
--- a/Engine/Plugins/Experimental/ControlRig/Source/ControlRig/Public/Rigs/RigHierarchy.h
+++ b/Engine/Plugins/Experimental/ControlRig/Source/ControlRig/Public/Rigs/RigHierarchy.h
@@ -60,9 +60,9 @@ public:
 	void CopyPose(URigHierarchy* InHierarchy, bool bCurrent, bool bInitial);
 
 	/**
-	 * Update all elements that depend on auxiliary data
+	 * Update all elements that depend on external sockets
 	 */
-	void UpdateAuxiliaryElements(const FRigUnitContext* InContext);
+	void UpdateSockets(const FRigUnitContext* InContext);
 
 	/**
 	 * Resets the current pose of a filtered lost if elements to the initial / ref pose.
@@ -597,22 +597,22 @@ public:
 	}
 
 	/**
-	 * Returns all Auxiliary elements
+	 * Returns all sockets
 	 * @param bTraverse Returns the elements in order of a depth first traversal
 	 */
-	FORCEINLINE TArray<FRigAuxiliaryElement*> GetAuxiliaryElements(bool bTraverse = false) const
+	FORCEINLINE TArray<FRigSocketElement*> GetSockets(bool bTraverse = false) const
 	{
-		return GetElementsOfType<FRigAuxiliaryElement>(bTraverse);
+		return GetElementsOfType<FRigSocketElement>(bTraverse);
 	}
 
 	/**
-	 * Returns all Auxiliary elements
+	 * Returns all sockets
 	 * @param bTraverse Returns the elements in order of a depth first traversal
 	 */
-	UFUNCTION(BlueprintCallable, Category = URigHierarchy, meta = (DisplayName = "Get Auxiliarys", ScriptName = "GetAuxiliaryElements"))
-    FORCEINLINE TArray<FRigElementKey> GetAuxiliaryKeys(bool bTraverse = true) const
+	UFUNCTION(BlueprintCallable, Category = URigHierarchy, meta = (DisplayName = "Get Sockets", ScriptName = "GetSockets"))
+    FORCEINLINE TArray<FRigElementKey> GetSocketKeys(bool bTraverse = true) const
 	{
-		return GetKeysOfType<FRigAuxiliaryElement>(bTraverse);
+		return GetKeysOfType<FRigSocketElement>(bTraverse);
 	}
 
 	/**
diff --git a/Engine/Plugins/Experimental/ControlRig/Source/ControlRig/Public/Rigs/RigHierarchyController.h b/Engine/Plugins/Experimental/ControlRig/Source/ControlRig/Public/Rigs/RigHierarchyController.h
--- a/Engine/Plugins/Experimental/ControlRig/Source/ControlRig/Public/Rigs/RigHierarchyController.h
+++ b/Engine/Plugins/Experimental/ControlRig/Source/ControlRig/Public/Rigs/RigHierarchyController.h
@@ -166,17 +166,17 @@ public:
     	bool bSetupUndo = false);
 
 	/**
-	* Adds an auxiliary element to the hierarchy
-	* @param InName The suggested name of the new auxiliary element - will eventually be corrected by the namespace
-	* @param InParent The (optional) parent of the new auxiliary element. If you don't need a parent, pass FRigElementKey()
+	* Adds an socket to the hierarchy
+	* @param InName The suggested name of the new socket - will eventually be corrected by the namespace
+	* @param InParent The (optional) parent of the new socket. If you don't need a parent, pass FRigElementKey()
 	* @param InDelegate The delegate to use to pull the local transform
 	* @param bSetupUndo If set to true the stack will record the change for undo / redo
-	* @return The key for the newly created auxiliary element.
+	* @return The key for the newly created socket.
 	*/
-    FRigElementKey AddAuxiliaryElement(
+    FRigElementKey AddSocket(
         FName InName,
         FRigElementKey InParent,
-        FRigAuxiliaryElementGetWorldTransformDelegate InDelegate,
+        FRigSocketGetWorldTransformDelegate InDelegate,
         bool bSetupUndo = false);
 
 	/**
diff --git a/Engine/Plugins/Experimental/ControlRig/Source/ControlRig/Public/Rigs/RigHierarchyDefines.h b/Engine/Plugins/Experimental/ControlRig/Source/ControlRig/Public/Rigs/RigHierarchyDefines.h
--- a/Engine/Plugins/Experimental/ControlRig/Source/ControlRig/Public/Rigs/RigHierarchyDefines.h
+++ b/Engine/Plugins/Experimental/ControlRig/Source/ControlRig/Public/Rigs/RigHierarchyDefines.h
@@ -22,8 +22,8 @@ enum class ERigElementType : uint8
 	Control = 0x004,
 	Curve = 0x008,
 	RigidBody = 0x010,
-	Auxiliary = 0x020,
-	All = Bone | Space | Control | Curve | RigidBody | Auxiliary UMETA(Hidden),
+	Socket = 0x020,
+	All = Bone | Space | Control | Curve | RigidBody | Socket UMETA(Hidden),
 };
 
 UENUM(BlueprintType)
@@ -889,9 +889,9 @@ public:
 			{
 				return FString::Printf(TEXT("RigidBody(%s)"), *Name.ToString());
 			}
-			case ERigElementType::Auxiliary:
+			case ERigElementType::Socket:
 			{
-				return FString::Printf(TEXT("Auxiliary(%s)"), *Name.ToString());
+				return FString::Printf(TEXT("Socket(%s)"), *Name.ToString());
 			}
 		}
 		return FString();
diff --git a/Engine/Plugins/Experimental/ControlRig/Source/ControlRig/Public/Rigs/RigHierarchyElements.h b/Engine/Plugins/Experimental/ControlRig/Source/ControlRig/Public/Rigs/RigHierarchyElements.h
--- a/Engine/Plugins/Experimental/ControlRig/Source/ControlRig/Public/Rigs/RigHierarchyElements.h
+++ b/Engine/Plugins/Experimental/ControlRig/Source/ControlRig/Public/Rigs/RigHierarchyElements.h
@@ -10,7 +10,7 @@ struct FRigUnitContext;
 struct FRigBaseElement;
 class URigHierarchy;
 
-DECLARE_DELEGATE_RetVal_ThreeParams(FTransform, FRigAuxiliaryElementGetWorldTransformDelegate, const FRigUnitContext*, const FRigElementKey& /* Key */, bool /* bInitial */);
+DECLARE_DELEGATE_RetVal_ThreeParams(FTransform, FRigSocketGetWorldTransformDelegate, const FRigUnitContext*, const FRigElementKey& /* Key */, bool /* bInitial */);
 
 #define DECLARE_RIG_ELEMENT_METHODS(ElementType) \
 template<typename T> \
@@ -580,7 +580,7 @@ protected:
 			InElement->GetType() == ERigElementType::Space ||
 			InElement->GetType() == ERigElementType::Control ||
 			InElement->GetType() == ERigElementType::RigidBody ||
-			InElement->GetType() == ERigElementType::Auxiliary;
+			InElement->GetType() == ERigElementType::Socket;
 	}
 
 public:
@@ -622,7 +622,7 @@ protected:
 	{
 		return InElement->GetType() == ERigElementType::Bone ||
 			InElement->GetType() == ERigElementType::RigidBody ||
-			InElement->GetType() == ERigElementType::Auxiliary;
+			InElement->GetType() == ERigElementType::Socket;
 	}
 
 	friend struct FRigBaseElement;
@@ -974,35 +974,35 @@ protected:
 };
 
 USTRUCT(BlueprintType)
-struct CONTROLRIG_API FRigAuxiliaryElement : public FRigSingleParentElement
+struct CONTROLRIG_API FRigSocketElement : public FRigSingleParentElement
 {
 public:
 	
 	GENERATED_BODY()
-	DECLARE_RIG_ELEMENT_METHODS(FRigAuxiliaryElement)
+	DECLARE_RIG_ELEMENT_METHODS(FRigSocketElement)
 
-    FRigAuxiliaryElement()
+    FRigSocketElement()
         : FRigSingleParentElement()
 	{
-		Key.Type = ERigElementType::Auxiliary;
+		Key.Type = ERigElementType::Socket;
 	}
 	
-	virtual ~FRigAuxiliaryElement(){}
+	virtual ~FRigSocketElement(){}
 
 	virtual void Save(FArchive& A, URigHierarchy* Hierarchy, ESerializationPhase SerializationPhase) override;
 	virtual void Load(FArchive& Ar, URigHierarchy* Hierarchy, ESerializationPhase SerializationPhase) override;
 
-	FTransform GetAuxiliaryWorldTransform(const FRigUnitContext* InContext, bool bInitial) const;
+	FTransform GetSocketWorldTransform(const FRigUnitContext* InContext, bool bInitial) const;
 
 	virtual void CopyPose(FRigBaseElement* InOther, bool bCurrent, bool bInitial) override;
 
 protected:
 
-	FRigAuxiliaryElementGetWorldTransformDelegate GetWorldTransformDelegate;
+	FRigSocketGetWorldTransformDelegate GetWorldTransformDelegate;
 
 	FORCEINLINE static bool IsClassOf(const FRigBaseElement* InElement)
 	{
-		return InElement->GetType() == ERigElementType::Auxiliary;
+		return InElement->GetType() == ERigElementType::Socket;
 	}
 
 	friend struct FRigBaseElement;
diff --git a/Engine/Plugins/Experimental/ControlRig/Source/ControlRig/Public/Units/RigUnitContext.h b/Engine/Plugins/Experimental/ControlRig/Source/ControlRig/Public/Units/RigUnitContext.h
--- a/Engine/Plugins/Experimental/ControlRig/Source/ControlRig/Public/Units/RigUnitContext.h
+++ b/Engine/Plugins/Experimental/ControlRig/Source/ControlRig/Public/Units/RigUnitContext.h
@@ -50,7 +50,7 @@ struct FRigUnitContext
 	/** The draw interface for the units to use */
 	FControlRigDrawInterface* DrawInterface;
 
-	/** The auxiliary draw container for the units to use */
+	/** The draw container for the units to use */
 	FControlRigDrawContainer* DrawContainer;
 
 	/** The registry to access data source */
diff --git a/Engine/Plugins/Experimental/ControlRig/Source/ControlRigDeveloper/Private/Graph/ControlRigGraph.cpp b/Engine/Plugins/Experimental/ControlRig/Source/ControlRigDeveloper/Private/Graph/ControlRigGraph.cpp
--- a/Engine/Plugins/Experimental/ControlRig/Source/ControlRigDeveloper/Private/Graph/ControlRigGraph.cpp
+++ b/Engine/Plugins/Experimental/ControlRig/Source/ControlRigDeveloper/Private/Graph/ControlRigGraph.cpp
@@ -82,7 +82,7 @@ void UControlRigGraph::CacheNameLists(URigHierarchy* InHierarchy, const FControl
 	TArray<TSharedPtr<FString>>& ControlNameList = ElementNameLists.FindOrAdd(ERigElementType::Control);
 	TArray<TSharedPtr<FString>>& CurveNameList = ElementNameLists.FindOrAdd(ERigElementType::Curve);
 	TArray<TSharedPtr<FString>>& RigidBodyNameList = ElementNameLists.FindOrAdd(ERigElementType::RigidBody);
-	TArray<TSharedPtr<FString>>& AuxiliaryNameList = ElementNameLists.FindOrAdd(ERigElementType::Auxiliary);
+	TArray<TSharedPtr<FString>>& SocketNameList = ElementNameLists.FindOrAdd(ERigElementType::Socket);
 	
 	CacheNameListForHierarchy<FRigBaseElement>(InHierarchy, AllNameList);
 	CacheNameListForHierarchy<FRigBoneElement>(InHierarchy, BoneNameList);
@@ -90,7 +90,7 @@ void UControlRigGraph::CacheNameLists(URigHierarchy* InHierarchy, const FControl
 	CacheNameListForHierarchy<FRigControlElement>(InHierarchy, ControlNameList);
 	CacheNameListForHierarchy<FRigCurveElement>(InHierarchy, CurveNameList);
 	CacheNameListForHierarchy<FRigRigidBodyElement>(InHierarchy, RigidBodyNameList);
-	CacheNameListForHierarchy<FRigAuxiliaryElement>(InHierarchy, AuxiliaryNameList);
+	CacheNameListForHierarchy<FRigSocketElement>(InHierarchy, SocketNameList);
 	CacheNameList<FControlRigDrawContainer>(*DrawContainer, DrawingNameList);
 }
 
diff --git a/Engine/Plugins/Experimental/ControlRig/Source/ControlRigEditor/Private/ControlRigEditorStyle.h b/Engine/Plugins/Experimental/ControlRig/Source/ControlRigEditor/Private/ControlRigEditorStyle.h
--- a/Engine/Plugins/Experimental/ControlRig/Source/ControlRigEditor/Private/ControlRigEditorStyle.h
+++ b/Engine/Plugins/Experimental/ControlRig/Source/ControlRigEditor/Private/ControlRigEditorStyle.h
@@ -100,7 +100,7 @@ public:
 			Set("ControlRig.Tree.Control", new IMAGE_BRUSH("Slate/Control_16x", Icon16x16));
 			Set("ControlRig.Tree.Space", new IMAGE_BRUSH("Slate/Space_16x", Icon16x16));
 			Set("ControlRig.Tree.RigidBody", new IMAGE_BRUSH("Slate/RigidBody_16x", Icon16x16));
-			Set("ControlRig.Tree.Auxiliary", new IMAGE_BRUSH("Slate/Auxiliary_16x", Icon16x16));
+			Set("ControlRig.Tree.Socket", new IMAGE_BRUSH("Slate/Auxiliary_16x", Icon16x16));
 		}
 
 		// Font?
diff --git a/Engine/Plugins/Experimental/ControlRig/Source/ControlRigEditor/Private/ControlRigHierarchyCommands.cpp b/Engine/Plugins/Experimental/ControlRig/Source/ControlRigEditor/Private/ControlRigHierarchyCommands.cpp
--- a/Engine/Plugins/Experimental/ControlRig/Source/ControlRigEditor/Private/ControlRigHierarchyCommands.cpp
+++ b/Engine/Plugins/Experimental/ControlRig/Source/ControlRigEditor/Private/ControlRigHierarchyCommands.cpp
@@ -34,7 +34,8 @@ void FControlRigHierarchyCommands::RegisterCommands()
 	UI_COMMAND(ShowControls, "Show Controls", "Whether to show or hide controls", EUserInterfaceActionType::ToggleButton, FInputChord());
 	UI_COMMAND(ShowSpaces, "Show Spaces", "Whether to show or hide spaces", EUserInterfaceActionType::ToggleButton, FInputChord());
 	UI_COMMAND(ShowRigidBodies, "Show RigidBodies", "Whether to show or hide rigidbodies", EUserInterfaceActionType::ToggleButton, FInputChord());
-	UI_COMMAND(ShowAuxiliaryElements, "Show Auxiliary Elements", "Whether to show or hide auxiliary elements", EUserInterfaceActionType::ToggleButton, FInputChord());
+	UI_COMMAND(ShowSockets, "Show Sockets", "Whether to show or hide sockets", EUserInterfaceActionType::ToggleButton, FInputChord());
+	UI_COMMAND(ShowDynamicHierarchy, "Dynamic Hierarchy", "Whether to show the static or dynamic hierarchy", EUserInterfaceActionType::ToggleButton, FInputChord());
 }
 
 #undef LOCTEXT_NAMESPACE
diff --git a/Engine/Plugins/Experimental/ControlRig/Source/ControlRigEditor/Private/ControlRigHierarchyCommands.h b/Engine/Plugins/Experimental/ControlRig/Source/ControlRigEditor/Private/ControlRigHierarchyCommands.h
--- a/Engine/Plugins/Experimental/ControlRig/Source/ControlRigEditor/Private/ControlRigHierarchyCommands.h
+++ b/Engine/Plugins/Experimental/ControlRig/Source/ControlRigEditor/Private/ControlRigHierarchyCommands.h
@@ -102,8 +102,11 @@ public:
 	/* Show rigidbodies */
 	TSharedPtr< FUICommandInfo > ShowRigidBodies;
 
-	/* Show auxiliary elements */
-	TSharedPtr< FUICommandInfo > ShowAuxiliaryElements;
+	/* Show sockets */
+	TSharedPtr< FUICommandInfo > ShowSockets;
+
+	/* Show dynamic hierarchy */
+	TSharedPtr< FUICommandInfo > ShowDynamicHierarchy;
 
 	/**
 	 * Initialize commands
diff --git a/Engine/Plugins/Experimental/ControlRig/Source/ControlRigEditor/Private/EditMode/ControlRigEditMode.cpp b/Engine/Plugins/Experimental/ControlRig/Source/ControlRigEditor/Private/EditMode/ControlRigEditMode.cpp
--- a/Engine/Plugins/Experimental/ControlRig/Source/ControlRigEditor/Private/EditMode/ControlRigEditMode.cpp
+++ b/Engine/Plugins/Experimental/ControlRig/Source/ControlRigEditor/Private/EditMode/ControlRigEditMode.cpp
@@ -1837,7 +1837,7 @@ void FControlRigEditMode::OnHierarchyModified(ERigHierarchyNotification InNotif,
             	case ERigElementType::Curve:
             	case ERigElementType::Control:
             	case ERigElementType::RigidBody:
-            	case ERigElementType::Auxiliary:
+            	case ERigElementType::Socket:
 				{
 					const bool bSelected = InNotif == ERigHierarchyNotification::ElementSelected;
 					if (bSelected)
diff --git a/Engine/Plugins/Experimental/ControlRig/Source/ControlRigEditor/Private/Editor/ControlRigEditor.cpp b/Engine/Plugins/Experimental/ControlRig/Source/ControlRigEditor/Private/Editor/ControlRigEditor.cpp
--- a/Engine/Plugins/Experimental/ControlRig/Source/ControlRigEditor/Private/Editor/ControlRigEditor.cpp
+++ b/Engine/Plugins/Experimental/ControlRig/Source/ControlRigEditor/Private/Editor/ControlRigEditor.cpp
@@ -1192,9 +1192,9 @@ void FControlRigEditor::SetDetailStruct(const FRigElementKey& InElement)
 			StructToDisplay = MakeShareable(new FStructOnScope(FRigRigidBodyElement::StaticStruct(), (uint8*)Element));
 			break;
 		}
-		case ERigElementType::Auxiliary:
+		case ERigElementType::Socket:
 		{
-			StructToDisplay = MakeShareable(new FStructOnScope(FRigAuxiliaryElement::StaticStruct(), (uint8*)Element));
+			StructToDisplay = MakeShareable(new FStructOnScope(FRigSocketElement::StaticStruct(), (uint8*)Element));
 			break;
 		}
 		default:
diff --git a/Engine/Plugins/Experimental/ControlRig/Source/ControlRigEditor/Private/Editor/SRigHierarchy.cpp b/Engine/Plugins/Experimental/ControlRig/Source/ControlRigEditor/Private/Editor/SRigHierarchy.cpp
--- a/Engine/Plugins/Experimental/ControlRig/Source/ControlRigEditor/Private/Editor/SRigHierarchy.cpp
+++ b/Engine/Plugins/Experimental/ControlRig/Source/ControlRigEditor/Private/Editor/SRigHierarchy.cpp
@@ -45,8 +45,21 @@
 FRigTreeElement::FRigTreeElement(const FRigElementKey& InKey, TWeakPtr<SRigHierarchy> InHierarchyHandler)
 {
 	Key = InKey;
+	bIsTransient = false;
+
+	if(InHierarchyHandler.IsValid())
+	{
+		if(URigHierarchy* Hierarchy = InHierarchyHandler.Pin()->GetDebuggedHierarchy())
+		{
+			if(FRigControlElement* ControlElement = Hierarchy->Find<FRigControlElement>(Key))
+			{
+				bIsTransient = ControlElement->Settings.bIsTransientControl;
+			}
+		}
+	}
 }
 
+
 TSharedRef<ITableRow> FRigTreeElement::MakeTreeRowWidget(TSharedPtr<FControlRigEditor> InControlRigEditor, const TSharedRef<STableViewBase>& InOwnerTable, TSharedRef<FRigTreeElement> InRigTreeElement, TSharedRef<FUICommandList> InCommandList, TSharedPtr<SRigHierarchy> InHierarchy)
 {
 	if (InRigTreeElement->Key.IsValid())
@@ -175,9 +188,9 @@ void SRigHierarchyItem::Construct(const FArguments& InArgs, TSharedPtr<FControlR
 			Brush = FControlRigEditorStyle::Get().GetBrush("ControlRig.Tree.RigidBody");
 			break;
 		}
-		case ERigElementType::Auxiliary:
+		case ERigElementType::Socket:
 		{
-			Brush = FControlRigEditorStyle::Get().GetBrush("ControlRig.Tree.Auxiliary");
+			Brush = FControlRigEditorStyle::Get().GetBrush("ControlRig.Tree.Socket");
 			break;
 		}
 		default:
@@ -220,6 +233,11 @@ void SRigHierarchyItem::Construct(const FArguments& InArgs, TSharedPtr<FControlR
 
 FText SRigHierarchyItem::GetName() const
 {
+	if(WeakRigTreeElement.Pin()->bIsTransient)
+	{
+		static const FText TemporaryControl = FText::FromString(TEXT("Temporary Control"));
+		return TemporaryControl;
+	}
 	return (FText::FromName(WeakRigTreeElement.Pin()->Key.Name));
 }
 
@@ -287,6 +305,7 @@ void SRigHierarchy::Construct(const FArguments& InArgs, TSharedRef<FControlRigEd
 
 	ControlRigBlueprint->Hierarchy->OnModified().AddRaw(this, &SRigHierarchy::OnHierarchyModified);
 	ControlRigBlueprint->OnRefreshEditor().AddRaw(this, &SRigHierarchy::HandleRefreshEditorFromBlueprint);
+	ControlRigBlueprint->OnSetObjectBeingDebugged().AddRaw(this, &SRigHierarchy::HandleSetObjectBeingDebugged);
 
 	// for deleting, renaming, dragging
 	CommandList = MakeShared<FUICommandList>();
@@ -429,8 +448,9 @@ void SRigHierarchy::Construct(const FArguments& InArgs, TSharedRef<FControlRigEd
 	bShowControls = true;
 	bShowSpaces = true;
 	bShowRigidBodies = true;
-	bShowAuxiliaryElements = true;
+	bShowSockets = true;
 	bIsChangingRigHierarchy = false;
+	bShowDynamicHierarchy = false;
 	RefreshTreeView();
 
 	if (ControlRigEditor.IsValid())
@@ -572,10 +592,16 @@ void SRigHierarchy::BindCommands()
 		FIsActionChecked::CreateLambda([this]() { return bShowRigidBodies; }));
 
 	CommandList->MapAction(
-		Commands.ShowAuxiliaryElements,
-		FExecuteAction::CreateLambda([this]() { bShowAuxiliaryElements = !bShowAuxiliaryElements; RefreshTreeView(); }),
+		Commands.ShowSockets,
+		FExecuteAction::CreateLambda([this]() { bShowSockets = !bShowSockets; RefreshTreeView(); }),
+		FCanExecuteAction(),
+		FIsActionChecked::CreateLambda([this]() { return bShowSockets; }));
+
+	CommandList->MapAction(
+		Commands.ShowDynamicHierarchy,
+		FExecuteAction::CreateLambda([this]() { bShowDynamicHierarchy = !bShowDynamicHierarchy; RefreshTreeView(); }),
 		FCanExecuteAction(),
-		FIsActionChecked::CreateLambda([this]() { return bShowAuxiliaryElements; }));
+		FIsActionChecked::CreateLambda([this]() { return bShowDynamicHierarchy; }));
 }
 
 FReply SRigHierarchy::OnKeyDown(const FGeometry& MyGeometry, const FKeyEvent& InKeyEvent)
@@ -656,6 +682,12 @@ void SRigHierarchy::RefreshTreeView()
 	if (ControlRigBlueprint.IsValid())
 	{
 		URigHierarchy* Hierarchy = GetHierarchy();
+
+		if(bShowDynamicHierarchy && ControlRigBeingDebuggedPtr.IsValid())
+		{
+			Hierarchy = ControlRigBeingDebuggedPtr->GetHierarchy();
+        }
+
 		check(Hierarchy);
 
 		Hierarchy->Traverse([&](FRigBaseElement* Element, bool& bContinue)
@@ -701,12 +733,12 @@ void SRigHierarchy::RefreshTreeView()
                     }
                     break;
                 }
-				case ERigElementType::Auxiliary:
+				case ERigElementType::Socket:
 				{
-                    if(bShowAuxiliaryElements)
+                    if(bShowSockets)
                     {
-                        FRigAuxiliaryElement* AuxiliaryElement = CastChecked<FRigAuxiliaryElement>(Element);
-                        AddElement(AuxiliaryElement);
+                        FRigSocketElement* Socket = CastChecked<FRigSocketElement>(Element);
+                        AddElement(Socket);
                     }
                     break;
                 }
@@ -997,6 +1029,7 @@ void SRigHierarchy::OnHierarchyModified(ERigHierarchyNotification InNotif, URigH
 			return;
 		}
 
+		/*
 		if(const FRigControlElement* ControlElement = Cast<FRigControlElement>(InElement))
 		{
 			if(ControlElement->Settings.bIsTransientControl)
@@ -1004,7 +1037,7 @@ void SRigHierarchy::OnHierarchyModified(ERigHierarchyNotification InNotif, URigH
 				return;
 			}
 		}
-
+		*/
 	}
 
 	switch(InNotif)
@@ -1060,6 +1093,33 @@ void SRigHierarchy::HandleRefreshEditorFromBlueprint(UControlRigBlueprint* InBlu
 	RefreshTreeView();
 }
 
+void SRigHierarchy::HandleSetObjectBeingDebugged(UObject* InObject)
+{
+	if(ControlRigBeingDebuggedPtr.Get() == InObject)
+	{
+		return;
+	}
+
+	if(ControlRigBeingDebuggedPtr.IsValid())
+	{
+		if(UControlRig* ControlRigBeingDebugged = ControlRigBeingDebuggedPtr.Get())
+		{
+			if(!ControlRigBeingDebugged->HasAnyFlags(RF_BeginDestroyed))
+			{
+				ControlRigBeingDebugged->GetHierarchy()->OnModified().RemoveAll(this);
+			}
+		}
+	}
+
+	ControlRigBeingDebuggedPtr.Reset();
+	
+	if(UControlRig* ControlRig = Cast<UControlRig>(InObject))
+	{
+		ControlRigBeingDebuggedPtr = ControlRig;
+		ControlRig->GetHierarchy()->OnModified().AddSP(this, &SRigHierarchy::OnHierarchyModified);
+	}
+}
+
 void SRigHierarchy::ClearDetailPanel() const
 {
 	ControlRigEditor.Pin()->ClearDetailObject();
@@ -1075,6 +1135,7 @@ TSharedRef< SWidget > SRigHierarchy::CreateFilterMenu()
 	MenuBuilder.BeginSection("FilterOptions", LOCTEXT("OptionsMenuHeading", "Options"));
 	{
 		MenuBuilder.AddMenuEntry(Actions.FilteringFlattensHierarchy);
+		MenuBuilder.AddMenuEntry(Actions.ShowDynamicHierarchy);
 		//MenuBuilder.AddMenuEntry(Actions.HideParentsWhenFiltering);
 	}
 	MenuBuilder.EndSection();
@@ -1664,11 +1725,25 @@ bool SRigHierarchy::CanRenameItem() const
 {
 	if(IsSingleSelected())
 	{
-		if(GetSelectedKeys()[0].Type == ERigElementType::RigidBody ||
-			GetSelectedKeys()[0].Type == ERigElementType::Auxiliary)
+		const FRigElementKey Key = GetSelectedKeys()[0];
+		if(Key.Type == ERigElementType::RigidBody ||
+			Key.Type == ERigElementType::Socket)
 		{
 			return false;
 		}
+		if(Key.Type == ERigElementType::Control)
+		{
+			if(URigHierarchy* DebuggedHierarchy = GetDebuggedHierarchy())
+			{
+				if(FRigControlElement* ControlElement = DebuggedHierarchy->Find<FRigControlElement>(Key))
+				{
+					if(ControlElement->Settings.bIsTransientControl)
+					{
+						return false;
+					}
+				}
+			}
+		}
 		return true;
 	}
 	return false;
@@ -1833,7 +1908,7 @@ URigHierarchy* SRigHierarchy::GetDebuggedHierarchy() const
 {
 	if (ControlRigBlueprint.IsValid())
 	{
-		if (UControlRig* DebuggedRig = Cast<UControlRig>(ControlRigBlueprint->GetObjectBeingDebugged()))
+		if (UControlRig* DebuggedRig = ControlRigBeingDebuggedPtr.Get())
 		{
 			return DebuggedRig->GetHierarchy();
 		}
@@ -1920,7 +1995,7 @@ TOptional<EItemDropZone> SRigHierarchy::OnCanAcceptDrop(const FDragDropEvent& Dr
 			case ERigElementType::Control:
 			case ERigElementType::Space:
 			case ERigElementType::RigidBody:
-			case ERigElementType::Auxiliary:
+			case ERigElementType::Socket:
 			{
 				for (const FRigElementKey& DraggedKey : RigDragDropOp->GetElements())
 				{
@@ -1929,7 +2004,7 @@ TOptional<EItemDropZone> SRigHierarchy::OnCanAcceptDrop(const FDragDropEvent& Dr
 						case ERigElementType::Control:
 						case ERigElementType::Space:
 						case ERigElementType::RigidBody:
-						case ERigElementType::Auxiliary:
+						case ERigElementType::Socket:
 						{
 							break;
 						}
diff --git a/Engine/Plugins/Experimental/ControlRig/Source/ControlRigEditor/Private/Editor/SRigHierarchy.h b/Engine/Plugins/Experimental/ControlRig/Source/ControlRigEditor/Private/Editor/SRigHierarchy.h
--- a/Engine/Plugins/Experimental/ControlRig/Source/ControlRigEditor/Private/Editor/SRigHierarchy.h
+++ b/Engine/Plugins/Experimental/ControlRig/Source/ControlRigEditor/Private/Editor/SRigHierarchy.h
@@ -30,6 +30,7 @@ public:
 public:
 	/** Element Data to display */
 	FRigElementKey Key;
+	bool bIsTransient;
 	TArray<TSharedPtr<FRigTreeElement>> Children;
 
 	TSharedRef<ITableRow> MakeTreeRowWidget(TSharedPtr<FControlRigEditor> InControlRigEditor, const TSharedRef<STableViewBase>& InOwnerTable, TSharedRef<FRigTreeElement> InRigTreeElement, TSharedRef<FUICommandList> InCommandList, TSharedPtr<SRigHierarchy> InHierarchy);
@@ -80,7 +81,7 @@ private:
 	SLATE_END_ARGS()
 
 	void Construct(const FArguments& InArgs, TSharedPtr<FControlRigEditor> InControlRigEditor, const TSharedRef<STableViewBase>& OwnerTable, TSharedRef<FRigTreeElement> InRigTreeElement, TSharedRef<FUICommandList> InCommandList, TSharedPtr<SRigHierarchy> InHierarchy);
-	void OnNameCommitted(const FText& InText, ETextCommit::Type InCommitType) const;
+ 	void OnNameCommitted(const FText& InText, ETextCommit::Type InCommitType) const;
 	bool OnVerifyNameChanged(const FText& InText, FText& OutErrorMessage);
 
 private:
@@ -259,8 +260,11 @@ private:
 	/** Whether or not to show rigidbodies in the hierarchy */
 	bool bShowRigidBodies;
 
-	/** Whether or not to show auxiliary elements in the hierarchy */
-	bool bShowAuxiliaryElements;
+	/** Whether or not to show sockets in the hierarchy */
+	bool bShowSockets;
+
+	/** Whether or not to show the static or dynamic hierarchy */
+	bool bShowDynamicHierarchy;
 
 	/** Search box widget */
 	TSharedPtr<SSearchBox> FilterBox;
@@ -284,6 +288,7 @@ private:
 	TMap<FRigElementKey, FRigElementKey> ParentMap;
 
 	TWeakObjectPtr<UControlRigBlueprint> ControlRigBlueprint;
+	TWeakObjectPtr<UControlRig> ControlRigBeingDebuggedPtr;
 	
 	/** Command list we bind to */
 	TSharedPtr<FUICommandList> CommandList;
@@ -323,6 +328,7 @@ private:
 	bool bIsChangingRigHierarchy;
 	void OnHierarchyModified(ERigHierarchyNotification InNotif, URigHierarchy* InHierarchy, const FRigBaseElement* InElement);
 	void HandleRefreshEditorFromBlueprint(UControlRigBlueprint* InBlueprint);
+	void HandleSetObjectBeingDebugged(UObject* InObject);
 
 	static TSharedPtr<FRigTreeElement> FindElement(const FRigElementKey& InElementKey, TSharedPtr<FRigTreeElement> CurrentItem);
 	void AddElement(FRigElementKey InKey, FRigElementKey InParentKey = FRigElementKey(), const bool bIgnoreTextFilter = false);
@@ -334,5 +340,6 @@ public:
 	FName RenameElement(const FRigElementKey& OldKey, const FString& NewName);
 	bool OnVerifyNameChanged(const FRigElementKey& OldKey, const FString& NewName, FText& OutErrorMessage);
 
+	friend class FRigTreeElement;
 	friend class SRigHierarchyItem;
 };
\ No newline at end of file
