Control Rig: Change all nodes to default to propagate to children = true

diff --git a/Engine/Plugins/Experimental/ControlRig/Source/ControlRig/Private/Units/Hierarchy/RigUnit_OffsetTransform.h b/Engine/Plugins/Experimental/ControlRig/Source/ControlRig/Private/Units/Hierarchy/RigUnit_OffsetTransform.h
--- a/Engine/Plugins/Experimental/ControlRig/Source/ControlRig/Private/Units/Hierarchy/RigUnit_OffsetTransform.h
+++ b/Engine/Plugins/Experimental/ControlRig/Source/ControlRig/Private/Units/Hierarchy/RigUnit_OffsetTransform.h
@@ -17,7 +17,7 @@ struct CONTROLRIG_API FRigUnit_OffsetTransformForItem : public FRigUnitMutable
 		: Item(NAME_None, ERigElementType::Bone)
 		, OffsetTransform(FTransform::Identity)
 		, Weight(1.f)
-		, bPropagateToChildren(false)
+		, bPropagateToChildren(true)
 		, CachedIndex()
 	{}
 
diff --git a/Engine/Plugins/Experimental/ControlRig/Source/ControlRig/Private/Units/Hierarchy/RigUnit_SetBoneInitialTransform.h b/Engine/Plugins/Experimental/ControlRig/Source/ControlRig/Private/Units/Hierarchy/RigUnit_SetBoneInitialTransform.h
--- a/Engine/Plugins/Experimental/ControlRig/Source/ControlRig/Private/Units/Hierarchy/RigUnit_SetBoneInitialTransform.h
+++ b/Engine/Plugins/Experimental/ControlRig/Source/ControlRig/Private/Units/Hierarchy/RigUnit_SetBoneInitialTransform.h
@@ -16,7 +16,7 @@ struct CONTROLRIG_API FRigUnit_SetBoneInitialTransform : public FRigUnitMutable
 
 	FRigUnit_SetBoneInitialTransform()
 		: Space(EBoneGetterSetterMode::GlobalSpace)
-		, bPropagateToChildren(false)
+		, bPropagateToChildren(true)
 		, CachedBone(FCachedRigElement())
 	{}
 
diff --git a/Engine/Plugins/Experimental/ControlRig/Source/ControlRig/Private/Units/Hierarchy/RigUnit_SetBoneRotation.h b/Engine/Plugins/Experimental/ControlRig/Source/ControlRig/Private/Units/Hierarchy/RigUnit_SetBoneRotation.h
--- a/Engine/Plugins/Experimental/ControlRig/Source/ControlRig/Private/Units/Hierarchy/RigUnit_SetBoneRotation.h
+++ b/Engine/Plugins/Experimental/ControlRig/Source/ControlRig/Private/Units/Hierarchy/RigUnit_SetBoneRotation.h
@@ -18,7 +18,7 @@ struct CONTROLRIG_API FRigUnit_SetBoneRotation : public FRigUnitMutable
 		: Rotation(FQuat::Identity)
 		, Space(EBoneGetterSetterMode::LocalSpace)
 		, Weight(1.f)
-		, bPropagateToChildren(false)
+		, bPropagateToChildren(true)
 		, CachedBone(FCachedRigElement())
 	{}
 
diff --git a/Engine/Plugins/Experimental/ControlRig/Source/ControlRig/Private/Units/Hierarchy/RigUnit_SetBoneTransform.h b/Engine/Plugins/Experimental/ControlRig/Source/ControlRig/Private/Units/Hierarchy/RigUnit_SetBoneTransform.h
--- a/Engine/Plugins/Experimental/ControlRig/Source/ControlRig/Private/Units/Hierarchy/RigUnit_SetBoneTransform.h
+++ b/Engine/Plugins/Experimental/ControlRig/Source/ControlRig/Private/Units/Hierarchy/RigUnit_SetBoneTransform.h
@@ -17,7 +17,7 @@ struct CONTROLRIG_API FRigUnit_SetBoneTransform : public FRigUnitMutable
 	FRigUnit_SetBoneTransform()
 		: Space(EBoneGetterSetterMode::LocalSpace)
 		, Weight(1.f)
-		, bPropagateToChildren(false)
+		, bPropagateToChildren(true)
 		, CachedBone(FCachedRigElement())
 	{}
 
diff --git a/Engine/Plugins/Experimental/ControlRig/Source/ControlRig/Private/Units/Hierarchy/RigUnit_SetBoneTranslation.h b/Engine/Plugins/Experimental/ControlRig/Source/ControlRig/Private/Units/Hierarchy/RigUnit_SetBoneTranslation.h
--- a/Engine/Plugins/Experimental/ControlRig/Source/ControlRig/Private/Units/Hierarchy/RigUnit_SetBoneTranslation.h
+++ b/Engine/Plugins/Experimental/ControlRig/Source/ControlRig/Private/Units/Hierarchy/RigUnit_SetBoneTranslation.h
@@ -18,7 +18,7 @@ struct CONTROLRIG_API FRigUnit_SetBoneTranslation : public FRigUnitMutable
 		: Translation(FVector::ZeroVector)
 		, Space(EBoneGetterSetterMode::LocalSpace)
 		, Weight(1.f)
-		, bPropagateToChildren(false)
+		, bPropagateToChildren(true)
 		, CachedBone(FCachedRigElement())
 	{}
 
diff --git a/Engine/Plugins/Experimental/ControlRig/Source/ControlRig/Private/Units/Hierarchy/RigUnit_SetRelativeBoneTransform.h b/Engine/Plugins/Experimental/ControlRig/Source/ControlRig/Private/Units/Hierarchy/RigUnit_SetRelativeBoneTransform.h
--- a/Engine/Plugins/Experimental/ControlRig/Source/ControlRig/Private/Units/Hierarchy/RigUnit_SetRelativeBoneTransform.h
+++ b/Engine/Plugins/Experimental/ControlRig/Source/ControlRig/Private/Units/Hierarchy/RigUnit_SetRelativeBoneTransform.h
@@ -16,7 +16,7 @@ struct CONTROLRIG_API FRigUnit_SetRelativeBoneTransform : public FRigUnitMutable
 
 	FRigUnit_SetRelativeBoneTransform()
 		: Weight(1.f)
-		, bPropagateToChildren(false)
+		, bPropagateToChildren(true)
 		, CachedBone(FCachedRigElement())
 		, CachedSpaceIndex(FCachedRigElement())
 	{}
diff --git a/Engine/Plugins/Experimental/ControlRig/Source/ControlRig/Private/Units/Hierarchy/RigUnit_SetRelativeTransform.h b/Engine/Plugins/Experimental/ControlRig/Source/ControlRig/Private/Units/Hierarchy/RigUnit_SetRelativeTransform.h
--- a/Engine/Plugins/Experimental/ControlRig/Source/ControlRig/Private/Units/Hierarchy/RigUnit_SetRelativeTransform.h
+++ b/Engine/Plugins/Experimental/ControlRig/Source/ControlRig/Private/Units/Hierarchy/RigUnit_SetRelativeTransform.h
@@ -19,7 +19,7 @@ struct CONTROLRIG_API FRigUnit_SetRelativeTransformForItem : public FRigUnitMuta
 		, bParentInitial(false)
 		, RelativeTransform(FTransform::Identity)
 		, Weight(1.f)
-		, bPropagateToChildren(false)
+		, bPropagateToChildren(true)
 		, CachedChild()
 		, CachedParent()
 	{}
diff --git a/Engine/Plugins/Experimental/ControlRig/Source/ControlRig/Private/Units/Highlevel/Harmonics/RigUnit_BoneHarmonics.h b/Engine/Plugins/Experimental/ControlRig/Source/ControlRig/Private/Units/Highlevel/Harmonics/RigUnit_BoneHarmonics.h
--- a/Engine/Plugins/Experimental/ControlRig/Source/ControlRig/Private/Units/Highlevel/Harmonics/RigUnit_BoneHarmonics.h
+++ b/Engine/Plugins/Experimental/ControlRig/Source/ControlRig/Private/Units/Highlevel/Harmonics/RigUnit_BoneHarmonics.h
@@ -92,7 +92,7 @@ struct CONTROLRIG_API FRigUnit_BoneHarmonics : public FRigUnit_HighlevelBaseMuta
 		WaveMinimum = 0.5f;
 		WaveMaximum = 1.f;
 		RotationOrder = EControlRigRotationOrder::YZX;
-		bPropagateToChildren = false;
+		bPropagateToChildren = true;
 	}
 
 	RIGVM_METHOD()
diff --git a/Engine/Plugins/Experimental/ControlRig/Source/ControlRig/Private/Units/Highlevel/Hierarchy/RigUnit_AimBone.h b/Engine/Plugins/Experimental/ControlRig/Source/ControlRig/Private/Units/Highlevel/Hierarchy/RigUnit_AimBone.h
--- a/Engine/Plugins/Experimental/ControlRig/Source/ControlRig/Private/Units/Highlevel/Hierarchy/RigUnit_AimBone.h
+++ b/Engine/Plugins/Experimental/ControlRig/Source/ControlRig/Private/Units/Highlevel/Hierarchy/RigUnit_AimBone.h
@@ -209,7 +209,7 @@ struct CONTROLRIG_API FRigUnit_AimBone : public FRigUnit_HighlevelBaseMutable
 		Primary.Axis = FVector(1.f, 0.f, 0.f);
 		Secondary.Axis = FVector(0.f, 0.f, 1.f);
 		Weight = 1.f;
-		bPropagateToChildren = false;
+		bPropagateToChildren = true;
 		DebugSettings = FRigUnit_AimBone_DebugSettings();
 		CachedBoneIndex = FCachedRigElement();
 		PrimaryCachedSpace = FCachedRigElement();
diff --git a/Engine/Plugins/Experimental/ControlRig/Source/ControlRig/Private/Units/Highlevel/Hierarchy/RigUnit_CCDIK.h b/Engine/Plugins/Experimental/ControlRig/Source/ControlRig/Private/Units/Highlevel/Hierarchy/RigUnit_CCDIK.h
--- a/Engine/Plugins/Experimental/ControlRig/Source/ControlRig/Private/Units/Highlevel/Hierarchy/RigUnit_CCDIK.h
+++ b/Engine/Plugins/Experimental/ControlRig/Source/ControlRig/Private/Units/Highlevel/Hierarchy/RigUnit_CCDIK.h
@@ -100,7 +100,7 @@ struct CONTROLRIG_API FRigUnit_CCDIK : public FRigUnit_HighlevelBaseMutable
 		Weight = 1.f;
 		MaxIterations = 10;
 		bStartFromTail = true;
-		bPropagateToChildren = false;
+		bPropagateToChildren = true;
 		BaseRotationLimit = 30.f;
 	}
 
@@ -190,7 +190,7 @@ struct CONTROLRIG_API FRigUnit_CCDIKPerItem : public FRigUnit_HighlevelBaseMutab
 		Weight = 1.f;
 		MaxIterations = 10;
 		bStartFromTail = true;
-		bPropagateToChildren = false;
+		bPropagateToChildren = true;
 		BaseRotationLimit = 30.f;
 	}
 
diff --git a/Engine/Plugins/Experimental/ControlRig/Source/ControlRig/Private/Units/Highlevel/Hierarchy/RigUnit_DistributeRotation.h b/Engine/Plugins/Experimental/ControlRig/Source/ControlRig/Private/Units/Highlevel/Hierarchy/RigUnit_DistributeRotation.h
--- a/Engine/Plugins/Experimental/ControlRig/Source/ControlRig/Private/Units/Highlevel/Hierarchy/RigUnit_DistributeRotation.h
+++ b/Engine/Plugins/Experimental/ControlRig/Source/ControlRig/Private/Units/Highlevel/Hierarchy/RigUnit_DistributeRotation.h
@@ -66,7 +66,7 @@ struct CONTROLRIG_API FRigUnit_DistributeRotation : public FRigUnit_HighlevelBas
 		StartBone = EndBone = NAME_None;
 		RotationEaseType = EControlRigAnimEasingType::Linear;
 		Weight = 1.f;
-		bPropagateToChildren = false;
+		bPropagateToChildren = true;
 	}
 
 	virtual FRigElementKey DetermineSpaceForPin(const FString& InPinPath, void* InUserContext) const override
diff --git a/Engine/Plugins/Experimental/ControlRig/Source/ControlRig/Private/Units/Highlevel/Hierarchy/RigUnit_FABRIK.h b/Engine/Plugins/Experimental/ControlRig/Source/ControlRig/Private/Units/Highlevel/Hierarchy/RigUnit_FABRIK.h
--- a/Engine/Plugins/Experimental/ControlRig/Source/ControlRig/Private/Units/Highlevel/Hierarchy/RigUnit_FABRIK.h
+++ b/Engine/Plugins/Experimental/ControlRig/Source/ControlRig/Private/Units/Highlevel/Hierarchy/RigUnit_FABRIK.h
@@ -45,7 +45,7 @@ struct CONTROLRIG_API FRigUnit_FABRIK : public FRigUnit_HighlevelBaseMutable
 		Weight = 1.f;
 		MaxIterations = 10;
 		EffectorTransform = FTransform::Identity;
-		bPropagateToChildren = false;
+		bPropagateToChildren = true;
 	}
 
 	/**
@@ -113,7 +113,7 @@ struct CONTROLRIG_API FRigUnit_FABRIKPerItem : public FRigUnit_HighlevelBaseMuta
 	{
 		Precision = 1.f;
 		Weight = 1.f;
-		bPropagateToChildren = false;
+		bPropagateToChildren = true;
 		MaxIterations = 10;
 		EffectorTransform = FTransform::Identity;
 	}
diff --git a/Engine/Plugins/Experimental/ControlRig/Source/ControlRig/Private/Units/Highlevel/Hierarchy/RigUnit_FitChainToCurve.h b/Engine/Plugins/Experimental/ControlRig/Source/ControlRig/Private/Units/Highlevel/Hierarchy/RigUnit_FitChainToCurve.h
--- a/Engine/Plugins/Experimental/ControlRig/Source/ControlRig/Private/Units/Highlevel/Hierarchy/RigUnit_FitChainToCurve.h
+++ b/Engine/Plugins/Experimental/ControlRig/Source/ControlRig/Private/Units/Highlevel/Hierarchy/RigUnit_FitChainToCurve.h
@@ -145,7 +145,7 @@ struct CONTROLRIG_API FRigUnit_FitChainToCurve : public FRigUnit_HighlevelBaseMu
 		PoleVectorPosition = FVector::ZeroVector;
 		RotationEaseType = EControlRigAnimEasingType::Linear;
 		Weight = 1.f;
-		bPropagateToChildren = false;
+		bPropagateToChildren = true;
 		DebugSettings = FRigUnit_FitChainToCurve_DebugSettings();
 	}
 
@@ -268,7 +268,7 @@ struct CONTROLRIG_API FRigUnit_FitChainToCurvePerItem : public FRigUnit_Highleve
 		PoleVectorPosition = FVector::ZeroVector;
 		RotationEaseType = EControlRigAnimEasingType::Linear;
 		Weight = 1.f;
-		bPropagateToChildren = false;
+		bPropagateToChildren = true;
 		DebugSettings = FRigUnit_FitChainToCurve_DebugSettings();
 	}
 
diff --git a/Engine/Plugins/Experimental/ControlRig/Source/ControlRig/Private/Units/Highlevel/Hierarchy/RigUnit_SlideChain.h b/Engine/Plugins/Experimental/ControlRig/Source/ControlRig/Private/Units/Highlevel/Hierarchy/RigUnit_SlideChain.h
--- a/Engine/Plugins/Experimental/ControlRig/Source/ControlRig/Private/Units/Highlevel/Hierarchy/RigUnit_SlideChain.h
+++ b/Engine/Plugins/Experimental/ControlRig/Source/ControlRig/Private/Units/Highlevel/Hierarchy/RigUnit_SlideChain.h
@@ -44,7 +44,7 @@ struct CONTROLRIG_API FRigUnit_SlideChain: public FRigUnit_HighlevelBaseMutable
 	{
 		StartBone = EndBone = NAME_None;
 		SlideAmount = 0.f;
-		bPropagateToChildren = false;
+		bPropagateToChildren = true;
 	}
 
 	RIGVM_METHOD()
@@ -91,7 +91,7 @@ struct CONTROLRIG_API FRigUnit_SlideChainPerItem: public FRigUnit_HighlevelBaseM
 	FRigUnit_SlideChainPerItem()
 	{
 		SlideAmount = 0.f;
-		bPropagateToChildren = false;
+		bPropagateToChildren = true;
 	}
 
 	RIGVM_METHOD()
diff --git a/Engine/Plugins/Experimental/ControlRig/Source/ControlRig/Private/Units/Highlevel/Hierarchy/RigUnit_SpringIK.h b/Engine/Plugins/Experimental/ControlRig/Source/ControlRig/Private/Units/Highlevel/Hierarchy/RigUnit_SpringIK.h
--- a/Engine/Plugins/Experimental/ControlRig/Source/ControlRig/Private/Units/Highlevel/Hierarchy/RigUnit_SpringIK.h
+++ b/Engine/Plugins/Experimental/ControlRig/Source/ControlRig/Private/Units/Highlevel/Hierarchy/RigUnit_SpringIK.h
@@ -92,7 +92,7 @@ struct CONTROLRIG_API FRigUnit_SpringIK : public FRigUnit_HighlevelBaseMutable
 		bLiveSimulation = false;
 		Iterations = 10;
 		bLimitLocalPosition = true;
-		bPropagateToChildren = false;
+		bPropagateToChildren = true;
 		DebugSettings = FRigUnit_SpringIK_DebugSettings();
 	}
 
diff --git a/Engine/Plugins/Experimental/ControlRig/Source/ControlRig/Private/Units/Highlevel/Hierarchy/RigUnit_TwistBones.h b/Engine/Plugins/Experimental/ControlRig/Source/ControlRig/Private/Units/Highlevel/Hierarchy/RigUnit_TwistBones.h
--- a/Engine/Plugins/Experimental/ControlRig/Source/ControlRig/Private/Units/Highlevel/Hierarchy/RigUnit_TwistBones.h
+++ b/Engine/Plugins/Experimental/ControlRig/Source/ControlRig/Private/Units/Highlevel/Hierarchy/RigUnit_TwistBones.h
@@ -37,7 +37,7 @@ struct CONTROLRIG_API FRigUnit_TwistBones : public FRigUnit_HighlevelBaseMutable
 		PoleAxis = FVector(0.f, 1.f, 0.f);
 		TwistEaseType = EControlRigAnimEasingType::Linear;
 		Weight = 1.f;
-		bPropagateToChildren = false;
+		bPropagateToChildren = true;
 	}
 
 	virtual FRigElementKey DetermineSpaceForPin(const FString& InPinPath, void* InUserContext) const override
@@ -110,7 +110,7 @@ struct CONTROLRIG_API FRigUnit_TwistBonesPerItem : public FRigUnit_HighlevelBase
 		PoleAxis = FVector(0.f, 1.f, 0.f);
 		TwistEaseType = EControlRigAnimEasingType::Linear;
 		Weight = 1.f;
-		bPropagateToChildren = false;
+		bPropagateToChildren = true;
 	}
 
 	RIGVM_METHOD()
diff --git a/Engine/Plugins/Experimental/ControlRig/Source/ControlRig/Private/Units/Highlevel/Hierarchy/RigUnit_TwoBoneIKSimple.h b/Engine/Plugins/Experimental/ControlRig/Source/ControlRig/Private/Units/Highlevel/Hierarchy/RigUnit_TwoBoneIKSimple.h
--- a/Engine/Plugins/Experimental/ControlRig/Source/ControlRig/Private/Units/Highlevel/Hierarchy/RigUnit_TwoBoneIKSimple.h
+++ b/Engine/Plugins/Experimental/ControlRig/Source/ControlRig/Private/Units/Highlevel/Hierarchy/RigUnit_TwoBoneIKSimple.h
@@ -59,7 +59,7 @@ struct CONTROLRIG_API FRigUnit_TwoBoneIKSimple : public FRigUnit_HighlevelBaseMu
 		StretchMaximumRatio = 1.25f;
 		Weight = 1.f;
 		BoneALength = BoneBLength = 0.f;
-		bPropagateToChildren = false;
+		bPropagateToChildren = true;
 		DebugSettings = FRigUnit_TwoBoneIKSimple_DebugSettings();
 
 		CachedBoneAIndex = FCachedRigElement();
@@ -226,7 +226,7 @@ struct CONTROLRIG_API FRigUnit_TwoBoneIKSimplePerItem : public FRigUnit_Highleve
 		StretchMaximumRatio = 1.25f;
 		Weight = 1.f;
 		ItemALength = ItemBLength = 0.f;
-		bPropagateToChildren = false;
+		bPropagateToChildren = true;
 		DebugSettings = FRigUnit_TwoBoneIKSimple_DebugSettings();
 
 		CachedItemAIndex = FCachedRigElement();
diff --git a/Engine/Plugins/Experimental/ControlRig/Source/ControlRig/Private/Units/Simulation/RigUnit_PointSimulation.h b/Engine/Plugins/Experimental/ControlRig/Source/ControlRig/Private/Units/Simulation/RigUnit_PointSimulation.h
--- a/Engine/Plugins/Experimental/ControlRig/Source/ControlRig/Private/Units/Simulation/RigUnit_PointSimulation.h
+++ b/Engine/Plugins/Experimental/ControlRig/Source/ControlRig/Private/Units/Simulation/RigUnit_PointSimulation.h
@@ -128,7 +128,7 @@ struct CONTROLRIG_API FRigUnit_PointSimulation : public FRigUnit_SimBaseMutable
 		IntegratorType = ECRSimPointIntegrateType::Verlet;
 		VerletBlend = 4.f;
 		bLimitLocalPosition = true;
-		bPropagateToChildren = false;
+		bPropagateToChildren = true;
 		PrimaryAimAxis = FVector(1.f, 0.f, 0.f);
 		SecondaryAimAxis = FVector(0.f, 1.f, 0.f);
 		DebugSettings = FRigUnit_PointSimulation_DebugSettings();
diff --git a/Engine/Plugins/Experimental/ControlRig/Source/ControlRig/Public/Units/Hierarchy/RigUnit_AddBoneTransform.h b/Engine/Plugins/Experimental/ControlRig/Source/ControlRig/Public/Units/Hierarchy/RigUnit_AddBoneTransform.h
--- a/Engine/Plugins/Experimental/ControlRig/Source/ControlRig/Public/Units/Hierarchy/RigUnit_AddBoneTransform.h
+++ b/Engine/Plugins/Experimental/ControlRig/Source/ControlRig/Public/Units/Hierarchy/RigUnit_AddBoneTransform.h
@@ -17,7 +17,7 @@ struct FRigUnit_AddBoneTransform : public FRigUnitMutable
 	FRigUnit_AddBoneTransform()
 		: Weight(1.f)
 		, bPostMultiply(false)
-		, bPropagateToChildren(false)
+		, bPropagateToChildren(true)
 		, CachedBone()
 	{}
 
diff --git a/Engine/Plugins/Experimental/ControlRig/Source/ControlRig/Public/Units/Hierarchy/RigUnit_SetTransform.h b/Engine/Plugins/Experimental/ControlRig/Source/ControlRig/Public/Units/Hierarchy/RigUnit_SetTransform.h
--- a/Engine/Plugins/Experimental/ControlRig/Source/ControlRig/Public/Units/Hierarchy/RigUnit_SetTransform.h
+++ b/Engine/Plugins/Experimental/ControlRig/Source/ControlRig/Public/Units/Hierarchy/RigUnit_SetTransform.h
@@ -19,7 +19,7 @@ struct CONTROLRIG_API FRigUnit_SetTransform : public FRigUnitMutable
 		, bInitial(false)
 		, Transform(FTransform::Identity)
 		, Weight(1.f)
-		, bPropagateToChildren(false)
+		, bPropagateToChildren(true)
 		, CachedIndex()
 	{}
 
@@ -91,7 +91,7 @@ struct CONTROLRIG_API FRigUnit_SetTranslation : public FRigUnitMutable
 		, Space(EBoneGetterSetterMode::GlobalSpace)
 		, Translation(FVector::ZeroVector)
 		, Weight(1.f)
-		, bPropagateToChildren(false)
+		, bPropagateToChildren(true)
 		, CachedIndex()
 	{}
 
@@ -155,7 +155,7 @@ struct CONTROLRIG_API FRigUnit_SetRotation : public FRigUnitMutable
 		, Space(EBoneGetterSetterMode::GlobalSpace)
 		, Rotation(FQuat::Identity)
 		, Weight(1.f)
-		, bPropagateToChildren(false)
+		, bPropagateToChildren(true)
 		, CachedIndex()
 	{}
 
@@ -219,7 +219,7 @@ struct FRigUnit_SetScale : public FRigUnitMutable
 		, Space(EBoneGetterSetterMode::GlobalSpace)
 		, Scale(FVector::OneVector)
 		, Weight(1.f)
-		, bPropagateToChildren(false)
+		, bPropagateToChildren(true)
 		, CachedIndex()
 	{}
 
