Control Rig: Guard against the element name list for element type NONE

diff --git a/Engine/Plugins/Experimental/ControlRig/Source/ControlRigDeveloper/Private/Graph/ControlRigGraph.cpp b/Engine/Plugins/Experimental/ControlRig/Source/ControlRigDeveloper/Private/Graph/ControlRigGraph.cpp
--- a/Engine/Plugins/Experimental/ControlRig/Source/ControlRigDeveloper/Private/Graph/ControlRigGraph.cpp
+++ b/Engine/Plugins/Experimental/ControlRig/Source/ControlRigDeveloper/Private/Graph/ControlRigGraph.cpp
@@ -100,6 +100,12 @@ const TArray<TSharedPtr<FString>>& UControlRigGraph::GetElementNameList(ERigElem
 	{
 		return OuterGraph->GetElementNameList(InElementType);
 	}
+	
+	if(InElementType == ERigElementType::None)
+	{
+		return EmptyElementNameList;
+	}
+	
 	if(!ElementNameLists.Contains(InElementType))
 	{
 		UControlRigBlueprint* Blueprint = GetBlueprint();
