Control Rig: Catch nullptr when converting to spawnable in sequencer

diff --git a/Engine/Plugins/Experimental/ControlRig/Source/ControlRigEditor/Private/Sequencer/ControlRigParameterTrackEditor.cpp b/Engine/Plugins/Experimental/ControlRig/Source/ControlRigEditor/Private/Sequencer/ControlRigParameterTrackEditor.cpp
--- a/Engine/Plugins/Experimental/ControlRig/Source/ControlRigEditor/Private/Sequencer/ControlRigParameterTrackEditor.cpp
+++ b/Engine/Plugins/Experimental/ControlRig/Source/ControlRigEditor/Private/Sequencer/ControlRigParameterTrackEditor.cpp
@@ -206,7 +206,12 @@ FControlRigParameterTrackEditor::FControlRigParameterTrackEditor(TSharedRef<ISeq
 						UControlRigComponent* NewControlRigComponent = Cast<UControlRigComponent>(ReplacedObject.Value);
 						if (OldControlRigComponent->GetControlRig())
 						{
-							OldToNewControlRigs.Emplace(OldControlRigComponent->GetControlRig(), NewControlRigComponent->GetControlRig());
+							UControlRig* NewControlRig = nullptr;
+							if(NewControlRigComponent)
+							{
+								NewControlRig = NewControlRigComponent->GetControlRig();
+							}
+							OldToNewControlRigs.Emplace(OldControlRigComponent->GetControlRig(), NewControlRig);
 						}
 					}
 					else if (UControlRig* OldControlRig = Cast<UControlRig>(ReplacedObject.Key))
