Control Rig: Fix non-unity build error

diff --git a/Engine/Plugins/Experimental/ControlRig/Source/ControlRigEditor/Private/Editor/SRigHierarchy.h b/Engine/Plugins/Experimental/ControlRig/Source/ControlRigEditor/Private/Editor/SRigHierarchy.h
--- a/Engine/Plugins/Experimental/ControlRig/Source/ControlRigEditor/Private/Editor/SRigHierarchy.h
+++ b/Engine/Plugins/Experimental/ControlRig/Source/ControlRigEditor/Private/Editor/SRigHierarchy.h
@@ -15,6 +15,7 @@ class FControlRigEditor;
 class SSearchBox;
 class FUICommandList;
 class UControlRigBlueprint;
+class UControlRig;
 struct FAssetData;
 class FMenuBuilder;
 class SRigHierarchyItem;
