Control Rig: Dynamic hierarchy testing within Engine Test

diff --git a/Engine/Plugins/Experimental/ControlRig/Source/ControlRig/Private/ControlRig.cpp b/Engine/Plugins/Experimental/ControlRig/Source/ControlRig/Private/ControlRig.cpp
--- a/Engine/Plugins/Experimental/ControlRig/Source/ControlRig/Private/ControlRig.cpp
+++ b/Engine/Plugins/Experimental/ControlRig/Source/ControlRig/Private/ControlRig.cpp
@@ -57,6 +57,7 @@ UControlRig::UControlRig(const FObjectInitializer& ObjectInitializer)
 	, bRequiresInitExecution(false)
 	, bRequiresSetupEvent(false)
 	, bSetupModeEnabled(false)
+	, bCopyHierarchyBeforeSetup(true)
 	, bResetInitialTransformsBeforeSetup(true)
 	, bManipulationEnabled(false)
 	, InitBracket(0)
@@ -493,7 +494,7 @@ void UControlRig::Execute(const EControlRigState InState, const FName& InEventNa
 	Context.FramesPerSecond = GetCurrentFramesPerSecond();
 	Context.bDuringInteraction = IsInteracting();
 	Context.State = InState;
-	Context.Hierarchy = DynamicHierarchy;
+	Context.Hierarchy = GetHierarchy();
 
 	Context.ToWorldSpaceTransform = FTransform::Identity;
 	Context.OwningComponent = nullptr;
@@ -579,11 +580,19 @@ void UControlRig::Execute(const EControlRigState InState, const FName& InEventNa
 		{
 			bRequiresSetupEvent = bSetupModeEnabled;
 
-			if (bResetInitialTransformsBeforeSetup && !bSetupModeEnabled)
+			if (UControlRig* CDO = Cast<UControlRig>(GetClass()->GetDefaultObject()))
 			{
-				if (UControlRig* CDO = Cast<UControlRig>(GetClass()->GetDefaultObject()))
+				if(bCopyHierarchyBeforeSetup)
 				{
-					DynamicHierarchy->CopyPose(CDO->DynamicHierarchy, false, true);
+					if(CDO->GetHierarchy()->GetTopologyVersion()!= GetHierarchy()->GetTopologyVersion())
+					{
+						GetHierarchy()->CopyHierarchy(CDO->GetHierarchy());
+					}
+				}
+				
+				if (bResetInitialTransformsBeforeSetup && !bSetupModeEnabled)
+				{
+					GetHierarchy()->CopyPose(CDO->GetHierarchy(), false, true);
 				}
 			}
 
diff --git a/Engine/Plugins/Experimental/ControlRig/Source/ControlRig/Private/Rigs/AdditiveControlRig.cpp b/Engine/Plugins/Experimental/ControlRig/Source/ControlRig/Private/Rigs/AdditiveControlRig.cpp
--- a/Engine/Plugins/Experimental/ControlRig/Source/ControlRig/Private/Rigs/AdditiveControlRig.cpp
+++ b/Engine/Plugins/Experimental/ControlRig/Source/ControlRig/Private/Rigs/AdditiveControlRig.cpp
@@ -13,6 +13,7 @@
 UAdditiveControlRig::UAdditiveControlRig(const FObjectInitializer& ObjectInitializer)
 	: Super(ObjectInitializer)
 {
+	bCopyHierarchyBeforeSetup = false;
 	bResetInitialTransformsBeforeSetup = false;
 }
 
diff --git a/Engine/Plugins/Experimental/ControlRig/Source/ControlRig/Private/Rigs/FKControlRig.cpp b/Engine/Plugins/Experimental/ControlRig/Source/ControlRig/Private/Rigs/FKControlRig.cpp
--- a/Engine/Plugins/Experimental/ControlRig/Source/ControlRig/Private/Rigs/FKControlRig.cpp
+++ b/Engine/Plugins/Experimental/ControlRig/Source/ControlRig/Private/Rigs/FKControlRig.cpp
@@ -14,6 +14,7 @@ UFKControlRig::UFKControlRig(const FObjectInitializer& ObjectInitializer)
 	: Super(ObjectInitializer)
 	, ApplyMode(EControlRigFKRigExecuteMode::Replace)
 {
+	bCopyHierarchyBeforeSetup = false;
 	bResetInitialTransformsBeforeSetup = false;
 }
 
diff --git a/Engine/Plugins/Experimental/ControlRig/Source/ControlRig/Private/Rigs/RigHierarchy.cpp b/Engine/Plugins/Experimental/ControlRig/Source/ControlRig/Private/Rigs/RigHierarchy.cpp
--- a/Engine/Plugins/Experimental/ControlRig/Source/ControlRig/Private/Rigs/RigHierarchy.cpp
+++ b/Engine/Plugins/Experimental/ControlRig/Source/ControlRig/Private/Rigs/RigHierarchy.cpp
@@ -237,6 +237,8 @@ void URigHierarchy::CopyHierarchy(URigHierarchy* InHierarchy)
 
 	FMemoryReader ArchiveReader(ArchiveBytes);
 	Load(ArchiveReader);
+
+	TopologyVersion = InHierarchy->GetTopologyVersion();
 }
 
 void URigHierarchy::CopyPose(URigHierarchy* InHierarchy, bool bCurrent, bool bInitial)
diff --git a/Engine/Plugins/Experimental/ControlRig/Source/ControlRig/Private/Units/Execution/RigUnit_Collection.cpp b/Engine/Plugins/Experimental/ControlRig/Source/ControlRig/Private/Units/Execution/RigUnit_Collection.cpp
--- a/Engine/Plugins/Experimental/ControlRig/Source/ControlRig/Private/Units/Execution/RigUnit_Collection.cpp
+++ b/Engine/Plugins/Experimental/ControlRig/Source/ControlRig/Private/Units/Execution/RigUnit_Collection.cpp
@@ -213,3 +213,11 @@ FRigUnit_CollectionLoop_Execute()
 		Item = FRigElementKey();
 	}
 }
+
+FRigUnit_CollectionAddItem_Execute()
+{
+	DECLARE_SCOPE_HIERARCHICAL_COUNTER_RIGUNIT();
+
+	Result = Collection;
+	Result.AddUnique(Item);
+}
\ No newline at end of file
diff --git a/Engine/Plugins/Experimental/ControlRig/Source/ControlRig/Private/Units/Execution/RigUnit_Collection.h b/Engine/Plugins/Experimental/ControlRig/Source/ControlRig/Private/Units/Execution/RigUnit_Collection.h
--- a/Engine/Plugins/Experimental/ControlRig/Source/ControlRig/Private/Units/Execution/RigUnit_Collection.h
+++ b/Engine/Plugins/Experimental/ControlRig/Source/ControlRig/Private/Units/Execution/RigUnit_Collection.h
@@ -397,4 +397,31 @@ struct CONTROLRIG_API FRigUnit_CollectionLoop : public FRigUnit_CollectionBaseMu
 
 	UPROPERTY(meta = (Output))
 	FControlRigExecuteContext Completed;
+};
+
+/**
+* Adds an element to an existing collection
+*/
+USTRUCT(meta = (DisplayName = "Add Item", Keywords = "Item,Add,Push,Insert"))
+struct CONTROLRIG_API FRigUnit_CollectionAddItem : public FRigUnit_CollectionBase
+{
+	GENERATED_BODY()
+
+	FRigUnit_CollectionAddItem()
+	{
+		Collection = Result = FRigElementKeyCollection();
+		Item = FRigElementKey();
+	}
+
+	RIGVM_METHOD()
+    virtual void Execute(const FRigUnitContext& Context) override;
+
+	UPROPERTY(meta = (Input))
+	FRigElementKeyCollection Collection;
+
+	UPROPERTY(meta = (Input))
+	FRigElementKey Item;
+
+	UPROPERTY(meta = (Output))
+	FRigElementKeyCollection Result;
 };
\ No newline at end of file
diff --git a/Engine/Plugins/Experimental/ControlRig/Source/ControlRig/Public/ControlRig.h b/Engine/Plugins/Experimental/ControlRig/Source/ControlRig/Public/ControlRig.h
--- a/Engine/Plugins/Experimental/ControlRig/Source/ControlRig/Public/ControlRig.h
+++ b/Engine/Plugins/Experimental/ControlRig/Source/ControlRig/Public/ControlRig.h
@@ -523,7 +523,8 @@ public:
 	virtual UAssetUserData* GetAssetUserDataOfClass(TSubclassOf<UAssetUserData> InUserDataClass) override;
 	virtual const TArray<UAssetUserData*>* GetAssetUserDataArray() const override;
 	//~ End IInterface_AssetUserData Interface
-protected:
+protected
+	:
 	/** Array of user data stored with the asset */
 	UPROPERTY(EditAnywhere, AdvancedDisplay, Instanced, Category = "Default")
 	TArray<UAssetUserData*> AssetUserData;
@@ -550,6 +551,7 @@ protected:
 	bool bRequiresInitExecution;
 	bool bRequiresSetupEvent;
 	bool bSetupModeEnabled;
+	bool bCopyHierarchyBeforeSetup;
 	bool bResetInitialTransformsBeforeSetup;
 	bool bManipulationEnabled;
 
diff --git a/Engine/Plugins/Experimental/ControlRig/Source/ControlRigEditor/Private/Editor/ControlRigEditor.cpp b/Engine/Plugins/Experimental/ControlRig/Source/ControlRigEditor/Private/Editor/ControlRigEditor.cpp
--- a/Engine/Plugins/Experimental/ControlRig/Source/ControlRigEditor/Private/Editor/ControlRigEditor.cpp
+++ b/Engine/Plugins/Experimental/ControlRig/Source/ControlRigEditor/Private/Editor/ControlRigEditor.cpp
@@ -105,6 +105,7 @@ FControlRigEditor::FControlRigEditor()
 	, bAnyErrorsLeft(false)
 	, LastEventQueue(EControlRigEditorEventQueue::Setup)
 	, LastDebuggedRig()
+	, RigHierarchyTabCount(0)
 {
 }
 
@@ -943,7 +944,8 @@ void FControlRigEditor::HandleSetObjectBeingDebugged(UObject* InObject)
 				AnimInstance->RecalcRequiredBones();
 
 				// since rig has changed, rebuild draw skeleton
-				EditorSkelComp->RebuildDebugDrawSkeleton();
+				EditorSkelComp->SetControlRigBeingDebugged(DebuggedControlRig);
+
 				if (FControlRigEditMode* EditMode = GetEditMode())
 				{
 					EditMode->SetObjects(DebuggedControlRig, EditorSkelComp,nullptr);
diff --git a/Engine/Plugins/Experimental/ControlRig/Source/ControlRigEditor/Private/Editor/ControlRigEditor.h b/Engine/Plugins/Experimental/ControlRig/Source/ControlRigEditor/Private/Editor/ControlRigEditor.h
--- a/Engine/Plugins/Experimental/ControlRig/Source/ControlRigEditor/Private/Editor/ControlRigEditor.h
+++ b/Engine/Plugins/Experimental/ControlRig/Source/ControlRigEditor/Private/Editor/ControlRigEditor.h
@@ -86,6 +86,8 @@ public:
 	// FBlueprintEditor interface
 	virtual UBlueprint* GetBlueprintObj() const override;
 
+	int32 GetRigHierarchyTabCount() const { return RigHierarchyTabCount; }
+
 public:
 	// IToolkit Interface
 	virtual FName GetToolkitFName() const override;
@@ -396,8 +398,10 @@ protected:
 
 	EControlRigEditorEventQueue LastEventQueue;
 	FString LastDebuggedRig;
+	int32 RigHierarchyTabCount;
 
 	friend class FControlRigEditorMode;
 	friend class SControlRigStackView;
 	friend class SRigHierarchy;
+	friend struct FRigHierarchyTabSummoner;
 };
\ No newline at end of file
diff --git a/Engine/Plugins/Experimental/ControlRig/Source/ControlRigEditor/Private/Editor/ControlRigSkeletalMeshComponent.cpp b/Engine/Plugins/Experimental/ControlRig/Source/ControlRigEditor/Private/Editor/ControlRigSkeletalMeshComponent.cpp
--- a/Engine/Plugins/Experimental/ControlRig/Source/ControlRigEditor/Private/Editor/ControlRigSkeletalMeshComponent.cpp
+++ b/Engine/Plugins/Experimental/ControlRig/Source/ControlRigEditor/Private/Editor/ControlRigSkeletalMeshComponent.cpp
@@ -61,7 +61,29 @@ void UControlRigSkeletalMeshComponent::RebuildDebugDrawSkeleton()
 			for(FRigBoneElement* BoneElement : BoneElements)
 			{
 				const int32 Index = BoneElement->GetIndex();
-				const FName ParentName = Hierarchy->GetFirstParent(BoneElement->GetKey()).Name;
+
+				FName ParentName = NAME_None;
+
+				// find the first parent that is a bone
+				Hierarchy->Traverse(BoneElement, false, [&ParentName, BoneElement](FRigBaseElement* InElement, bool& bContinue)
+				{
+					bContinue = true;
+					
+					if(InElement == BoneElement)
+					{
+						return;
+					}
+
+					if(FRigBoneElement* InBoneElement = Cast<FRigBoneElement>(InElement))
+					{
+						if(ParentName.IsNone())
+						{
+							ParentName = InBoneElement->GetName();
+						}
+						bContinue = false;
+					}
+				});
+
 				int32 ParentIndex = INDEX_NONE; 
 				if(!ParentName.IsNone())
 				{
@@ -108,3 +130,70 @@ void UControlRigSkeletalMeshComponent::EnablePreview(bool bEnable, UAnimationAss
 		PreviewInstance->SetAnimationAsset(PreviewAsset);
 	}
 }
+
+void UControlRigSkeletalMeshComponent::SetControlRigBeingDebugged(UControlRig* InControlRig)
+{
+	if(ControlRigBeingDebuggedPtr.Get() == InControlRig)
+	{
+		return;
+	}
+
+	if(ControlRigBeingDebuggedPtr.IsValid())
+	{
+		if(UControlRig* ControlRigBeingDebugged = ControlRigBeingDebuggedPtr.Get())
+		{
+			if(!ControlRigBeingDebugged->HasAnyFlags(RF_BeginDestroyed))
+			{
+				ControlRigBeingDebugged->GetHierarchy()->OnModified().RemoveAll(this);
+			}
+		}
+	}
+
+	ControlRigBeingDebuggedPtr.Reset();
+	
+	if(InControlRig)
+	{
+		ControlRigBeingDebuggedPtr = InControlRig;
+ 		InControlRig->GetHierarchy()->OnModified().AddUObject(this, &UControlRigSkeletalMeshComponent::OnHierarchyModifiedAsync);
+	}
+
+	RebuildDebugDrawSkeleton();
+}
+
+void UControlRigSkeletalMeshComponent::OnHierarchyModified(ERigHierarchyNotification InNotif,
+    URigHierarchy* InHierarchy, const FRigBaseElement* InElement)
+{
+	switch(InNotif)
+	{
+		case ERigHierarchyNotification::ElementAdded:
+		case ERigHierarchyNotification::ElementRemoved:
+		case ERigHierarchyNotification::ElementRenamed:
+		case ERigHierarchyNotification::ParentChanged:
+		case ERigHierarchyNotification::HierarchyReset:
+		{
+			RebuildDebugDrawSkeleton();
+			break;
+		}
+		default:
+		{
+			break;
+		}
+	}
+}
+
+void UControlRigSkeletalMeshComponent::OnHierarchyModifiedAsync(ERigHierarchyNotification InNotif,
+    URigHierarchy* InHierarchy, const FRigBaseElement* InElement)
+{
+	FRigElementKey Key;
+	if(InElement)
+	{
+		Key = InElement->GetKey();
+	}
+	
+	FFunctionGraphTask::CreateAndDispatchWhenReady([this, InNotif, InHierarchy, Key]()
+    {
+        const FRigBaseElement* Element = InHierarchy->Find(Key);
+        OnHierarchyModified(InNotif, InHierarchy, Element);
+		
+    }, TStatId(), NULL, ENamedThreads::GameThread);
+}
diff --git a/Engine/Plugins/Experimental/ControlRig/Source/ControlRigEditor/Private/Editor/ControlRigSkeletalMeshComponent.h b/Engine/Plugins/Experimental/ControlRig/Source/ControlRigEditor/Private/Editor/ControlRigSkeletalMeshComponent.h
--- a/Engine/Plugins/Experimental/ControlRig/Source/ControlRigEditor/Private/Editor/ControlRigSkeletalMeshComponent.h
+++ b/Engine/Plugins/Experimental/ControlRig/Source/ControlRigEditor/Private/Editor/ControlRigSkeletalMeshComponent.h
@@ -2,6 +2,7 @@
 
 #pragma once
 
+#include "Rigs/RigHierarchy.h"
 #include "Animation/DebugSkelMeshComponent.h"
 #include "ControlRigSkeletalMeshComponent.generated.h"
 
@@ -39,12 +40,23 @@ class UControlRigSkeletalMeshComponent : public UDebugSkelMeshComponent
 	// END UDebugSkeletalMeshComponent interface
 
 public:
+
+	/*
+	 * React to a control rig being debugged
+	 */
+	void SetControlRigBeingDebugged(UControlRig* InControlRig);
+	
 	/*
 	 *	Rebuild debug draw skeleton 
 	 */
 	void RebuildDebugDrawSkeleton();
+
+	void OnHierarchyModified(ERigHierarchyNotification InNotif, URigHierarchy* InHierarchy, const FRigBaseElement* InElement);
+	void OnHierarchyModifiedAsync(ERigHierarchyNotification InNotif, URigHierarchy* InHierarchy, const FRigBaseElement* InElement);
+
 private:
 	FReferenceSkeleton DebugDrawSkeleton;
 	TArray<FBoneIndexType> DebugDrawBones;
 	TArray<int32> DebugDrawBoneIndexInHierarchy;
-};
\ No newline at end of file
+	TWeakObjectPtr<UControlRig> ControlRigBeingDebuggedPtr;
+};
diff --git a/Engine/Plugins/Experimental/ControlRig/Source/ControlRigEditor/Private/Editor/RigHierarchyTabSummoner.cpp b/Engine/Plugins/Experimental/ControlRig/Source/ControlRigEditor/Private/Editor/RigHierarchyTabSummoner.cpp
--- a/Engine/Plugins/Experimental/ControlRig/Source/ControlRigEditor/Private/Editor/RigHierarchyTabSummoner.cpp
+++ b/Engine/Plugins/Experimental/ControlRig/Source/ControlRigEditor/Private/Editor/RigHierarchyTabSummoner.cpp
@@ -4,6 +4,7 @@
 #include "SRigHierarchy.h"
 #include "ControlRigEditorStyle.h"
 #include "ControlRigEditor.h"
+#include "SRigHierarchy.h"
 #include "Widgets/Docking/SDockTab.h"
 
 #define LOCTEXT_NAMESPACE "RigHierarchyTabSummoner"
@@ -36,16 +37,37 @@ FTabSpawnerEntry& FRigHierarchyTabSummoner::RegisterTabSpawner(TSharedRef<FTabMa
 
 TSharedRef<SWidget> FRigHierarchyTabSummoner::CreateTabBody(const FWorkflowTabSpawnInfo& Info) const
 {
+	ControlRigEditor.Pin()->RigHierarchyTabCount++;
 	return SNew(SRigHierarchy, ControlRigEditor.Pin().ToSharedRef());
 }
 
 TSharedRef<SDockTab> FRigHierarchyTabSummoner::SpawnTab(const FWorkflowTabSpawnInfo& Info) const
 {
 	TSharedRef<SDockTab>  DockTab = FWorkflowTabFactory::SpawnTab(Info);
-	DockTab->SetCanCloseTab(SDockTab::FCanCloseTab::CreateLambda([]()
+	DockTab->SetCanCloseTab(SDockTab::FCanCloseTab::CreateLambda([DockTab]()
     {
-        return false;
+		int32 HierarchyTabCount = 0;
+		if(SWidget* Content = &DockTab->GetContent().Get())
+		{
+			SRigHierarchy* RigHierarchy = (SRigHierarchy*)Content;
+			if(FControlRigEditor* ControlRigEditor = RigHierarchy->GetControlRigEditor())
+			{
+				HierarchyTabCount = ControlRigEditor->GetRigHierarchyTabCount();
+			}
+		}
+        return HierarchyTabCount > 1;
     }));
+	DockTab->SetOnTabClosed( SDockTab::FOnTabClosedCallback::CreateLambda([](TSharedRef<SDockTab> DockTab)
+	{
+		if(SWidget* Content = &DockTab->GetContent().Get())
+		{
+			SRigHierarchy* RigHierarchy = (SRigHierarchy*)Content;
+			if(FControlRigEditor* ControlRigEditor = RigHierarchy->GetControlRigEditor())
+			{
+				ControlRigEditor->RigHierarchyTabCount--;
+			}
+		}
+	}));
 	return DockTab;
 }
 
diff --git a/Engine/Plugins/Experimental/ControlRig/Source/ControlRigEditor/Private/Editor/SRigHierarchy.cpp b/Engine/Plugins/Experimental/ControlRig/Source/ControlRigEditor/Private/Editor/SRigHierarchy.cpp
--- a/Engine/Plugins/Experimental/ControlRig/Source/ControlRigEditor/Private/Editor/SRigHierarchy.cpp
+++ b/Engine/Plugins/Experimental/ControlRig/Source/ControlRigEditor/Private/Editor/SRigHierarchy.cpp
@@ -599,7 +599,12 @@ void SRigHierarchy::BindCommands()
 
 	CommandList->MapAction(
 		Commands.ShowDynamicHierarchy,
-		FExecuteAction::CreateLambda([this]() { bShowDynamicHierarchy = !bShowDynamicHierarchy; RefreshTreeView(); }),
+		FExecuteAction::CreateLambda([this]()
+		{
+			bShowDynamicHierarchy = !bShowDynamicHierarchy;
+			HandleSetObjectBeingDebugged(GetControlRigEditor()->GetControlRigBlueprint()->GetObjectBeingDebugged());
+			RefreshTreeView();
+		}),
 		FCanExecuteAction(),
 		FIsActionChecked::CreateLambda([this]() { return bShowDynamicHierarchy; }));
 }
@@ -664,110 +669,58 @@ void SRigHierarchy::OnFilterTextChanged(const FText& SearchText)
 	RefreshTreeView();
 }
 
-void SRigHierarchy::RefreshTreeView()
+void SRigHierarchy::RefreshTreeView(bool bRebuildContent)
 {
 	TMap<FRigElementKey, bool> ExpansionState;
-	for (TPair<FRigElementKey, TSharedPtr<FRigTreeElement>> Pair : ElementMap)
+
+	if(bRebuildContent)
 	{
-		ExpansionState.FindOrAdd(Pair.Key) = TreeView->IsItemExpanded(Pair.Value);
-	}
+		for (TPair<FRigElementKey, TSharedPtr<FRigTreeElement>> Pair : ElementMap)
+		{
+			ExpansionState.FindOrAdd(Pair.Key) = TreeView->IsItemExpanded(Pair.Value);
+		}
 
-	// internally save expansion states before rebuilding the tree, so the states can be restored later
-	TreeView->SaveAndClearSparseItemInfos();
+		// internally save expansion states before rebuilding the tree, so the states can be restored later
+		TreeView->SaveAndClearSparseItemInfos();
 
-	RootElements.Reset();
-	ElementMap.Reset();
-	ParentMap.Reset();
+		RootElements.Reset();
+		ElementMap.Reset();
+		ParentMap.Reset();
+	}
 
 	if (ControlRigBlueprint.IsValid())
 	{
-		URigHierarchy* Hierarchy = GetHierarchy();
+		URigHierarchy* Hierarchy = GetHierarchyForTopology();
+		check(Hierarchy);
 
-		if(bShowDynamicHierarchy && ControlRigBeingDebuggedPtr.IsValid())
+		if(bRebuildContent)
 		{
-			Hierarchy = ControlRigBeingDebuggedPtr->GetHierarchy();
-        }
+			Hierarchy->Traverse([&](FRigBaseElement* Element, bool& bContinue)
+            {
+                AddElement(Element);
+                bContinue = true;
+            });
 
-		check(Hierarchy);
+			for (const auto& Pair : ElementMap)
+			{
+				TreeView->RestoreSparseItemInfos(Pair.Value);
+			}
 
-		Hierarchy->Traverse([&](FRigBaseElement* Element, bool& bContinue)
-		{
-			switch(Element->GetType())
+			// expand all elements upon the initial construction of the tree
+			if (ExpansionState.Num() == 0)
 			{
-				case ERigElementType::Bone:
+				for (TSharedPtr<FRigTreeElement> RootElement : RootElements)
 				{
-					if(bShowBones)
-					{
-						FRigBoneElement* BoneElement = CastChecked<FRigBoneElement>(Element);
-						if (bShowImportedBones || BoneElement->BoneType != ERigBoneType::Imported)
-						{
-                            AddElement(BoneElement);
-                        }
-                    }
-					break;
+					SetExpansionRecursive(RootElement, false, true);
 				}
-				case ERigElementType::Space:
-				{
-                    if(bShowSpaces)
-                    {
-                        FRigSpaceElement* SpaceElement = CastChecked<FRigSpaceElement>(Element);
-                        AddElement(SpaceElement);
-                    }
-                    break;
-                }
-				case ERigElementType::Control:
-				{
-                    if(bShowControls)
-                    {
-                        FRigControlElement* ControlElement = CastChecked<FRigControlElement>(Element);
-                        AddElement(ControlElement);
-                    }
-                    break;
-                }
-				case ERigElementType::RigidBody:
-				{
-                    if(bShowRigidBodies)
-                    {
-                        FRigRigidBodyElement* RigidBodyElement = CastChecked<FRigRigidBodyElement>(Element);
-                        AddElement(RigidBodyElement);
-                    }
-                    break;
-                }
-				case ERigElementType::Socket:
-				{
-                    if(bShowSockets)
-                    {
-                        FRigSocketElement* Socket = CastChecked<FRigSocketElement>(Element);
-                        AddElement(Socket);
-                    }
-                    break;
-                }
-				default:
-				{
-					break;
-				}
-            }
-		});
-
-		for (const auto& Pair : ElementMap)
-		{
-			TreeView->RestoreSparseItemInfos(Pair.Value);
-		}
+			}
 
-		// expand all elements upon the initial construction of the tree
-		if (ExpansionState.Num() == 0)
-		{
-			for (TSharedPtr<FRigTreeElement> RootElement : RootElements)
+			if (RootElements.Num() > 0)
 			{
-				SetExpansionRecursive(RootElement, false, true);
+				AddSpacerElement();
 			}
 		}
-
-		if (RootElements.Num() > 0)
-		{
-			AddSpacerElement();
-		}
-
+		
 		TreeView->RequestTreeRefresh();
 
 		TArray<FRigElementKey> Selection = Hierarchy->GetSelectedKeys();
@@ -873,11 +826,11 @@ TSharedPtr<FRigTreeElement> SRigHierarchy::FindElement(const FRigElementKey& InE
 	return TSharedPtr<FRigTreeElement>();
 }
 
-void SRigHierarchy::AddElement(FRigElementKey InKey, FRigElementKey InParentKey, const bool bIgnoreTextFilter)
+bool SRigHierarchy::AddElement(FRigElementKey InKey, FRigElementKey InParentKey, const bool bIgnoreTextFilter)
 {
 	if(ElementMap.Contains(InKey))
 	{
-		return;
+		return false;
 	}
 
 	FString FilteredString = FilterText.ToString();
@@ -919,21 +872,80 @@ void SRigHierarchy::AddElement(FRigElementKey InKey, FRigElementKey InParentKey,
 			RootElements.Add(NewItem);
 		}
 	}
+
+	return true;
 }
 
-void SRigHierarchy::AddElement(FRigBaseElement* InElement, const bool bIgnoreTextFilter)
+bool SRigHierarchy::AddElement(const FRigBaseElement* InElement, const bool bIgnoreTextFilter)
 {
 	check(InElement);
 	
 	if (ElementMap.Contains(InElement->GetKey()))
 	{
-		return;
+		return false;
+	}
+
+	switch(InElement->GetType())
+	{
+		case ERigElementType::Bone:
+		{
+			if(!bShowBones)
+			{
+				return false;
+			}
+
+			const FRigBoneElement* BoneElement = CastChecked<FRigBoneElement>(InElement);
+			if (!bShowImportedBones && BoneElement->BoneType == ERigBoneType::Imported)
+			{
+				return false;
+			}
+			break;
+		}
+		case ERigElementType::Space:
+		{
+			if(!bShowSpaces)
+			{
+				return false;
+			}
+			break;
+		}
+		case ERigElementType::Control:
+		{
+			if(!bShowControls)
+			{
+				return false;
+			}
+			break;
+		}
+		case ERigElementType::RigidBody:
+		{
+			if(!bShowRigidBodies)
+			{
+				return false;
+			}
+			break;
+		}
+		case ERigElementType::Socket:
+		{
+			if(!bShowSockets)
+			{
+				return false;
+			}
+			break;
+		}
+		default:
+		{
+			break;
+		}
 	}
 
-	const URigHierarchy* Hierarchy = GetHierarchy();
+	const URigHierarchy* Hierarchy = GetHierarchyForTopology();
 	check(Hierarchy);
 
-	AddElement(InElement->GetKey(), FRigElementKey(), bIgnoreTextFilter);
+	if(!AddElement(InElement->GetKey(), FRigElementKey(), bIgnoreTextFilter))
+	{
+		return false;
+	}
 
 	if (ElementMap.Contains(InElement->GetKey()))
 	{
@@ -946,6 +958,8 @@ void SRigHierarchy::AddElement(FRigBaseElement* InElement, const bool bIgnoreTex
 			}
 		}
 	}
+
+	return true;
 }
 
 void SRigHierarchy::AddSpacerElement()
@@ -953,29 +967,29 @@ void SRigHierarchy::AddSpacerElement()
 	AddElement(FRigElementKey(), FRigElementKey());
 }
 
-void SRigHierarchy::ReparentElement(FRigElementKey InKey, FRigElementKey InParentKey)
+bool SRigHierarchy::ReparentElement(FRigElementKey InKey, FRigElementKey InParentKey)
 {
 	if (!InKey.IsValid() || InKey == InParentKey)
 	{
-		return;
+		return false;
 	}
 
 	TSharedPtr<FRigTreeElement>* FoundItem = ElementMap.Find(InKey);
 	if (FoundItem == nullptr)
 	{
-		return;
+		return false;
 	}
 
 	if (!FilterText.IsEmpty() && bFlattenHierarchyOnFilter)
 	{
-		return;
+		return false;
 	}
 
 	if (const FRigElementKey* ExistingParentKey = ParentMap.Find(InKey))
 	{
 		if (*ExistingParentKey == InParentKey)
 		{
-			return;
+			return false;
 		}
 
 		if (TSharedPtr<FRigTreeElement>* ExistingParent = ElementMap.Find(*ExistingParentKey))
@@ -989,7 +1003,7 @@ void SRigHierarchy::ReparentElement(FRigElementKey InKey, FRigElementKey InParen
 	{
 		if (!InParentKey.IsValid())
 		{
-			return;
+			return false;
 		}
 
 		RootElements.Remove(*FoundItem);
@@ -1008,6 +1022,21 @@ void SRigHierarchy::ReparentElement(FRigElementKey InKey, FRigElementKey InParen
 		RootElements.Add(*FoundItem);
 	}
 
+	return true;
+}
+
+bool SRigHierarchy::RemoveElement(FRigElementKey InKey)
+{
+	TSharedPtr<FRigTreeElement>* FoundItem = ElementMap.Find(InKey);
+	if (FoundItem == nullptr)
+	{
+		return false;
+	}
+
+	ReparentElement(InKey, FRigElementKey());
+
+	RootElements.Remove(*FoundItem);
+	return ElementMap.Remove(InKey) > 0;
 }
 
 void SRigHierarchy::OnHierarchyModified(ERigHierarchyNotification InNotif, URigHierarchy* InHierarchy, const FRigBaseElement* InElement)
@@ -1028,24 +1057,41 @@ void SRigHierarchy::OnHierarchyModified(ERigHierarchyNotification InNotif, URigH
 		{
 			return;
 		}
-
-		/*
-		if(const FRigControlElement* ControlElement = Cast<FRigControlElement>(InElement))
-		{
-			if(ControlElement->Settings.bIsTransientControl)
-			{
-				return;
-			}
-		}
-		*/
 	}
 
 	switch(InNotif)
 	{
 		case ERigHierarchyNotification::ElementAdded:
+		{
+			check(InElement);
+			if(AddElement(InElement))
+			{
+				RefreshTreeView(false);
+			}
+			break;
+		}
 		case ERigHierarchyNotification::ElementRemoved:
-		case ERigHierarchyNotification::ElementRenamed:
+		{
+			check(InElement);
+			if(RemoveElement(InElement->GetKey()))
+			{
+				RefreshTreeView(false);
+			}
+			break;
+		}
 		case ERigHierarchyNotification::ParentChanged:
+		{
+			check(InHierarchy);
+			check(InElement);
+
+			const FRigElementKey ParentKey = InHierarchy->GetFirstParent(InElement->GetKey());
+			if(ReparentElement(InElement->GetKey(), ParentKey))
+			{
+				RefreshTreeView(false);
+			}
+			break;
+		}
+		case ERigHierarchyNotification::ElementRenamed:
 		case ERigHierarchyNotification::HierarchyReset:
 		{
 			RefreshTreeView();
@@ -1084,6 +1130,37 @@ void SRigHierarchy::OnHierarchyModified(ERigHierarchyNotification InNotif, URigH
 	}
 }
 
+void SRigHierarchy::OnHierarchyModifiedAsync(ERigHierarchyNotification InNotif, URigHierarchy* InHierarchy, const FRigBaseElement* InElement)
+{
+	if(!bShowDynamicHierarchy)
+	{
+		return;
+	}
+	
+	if(!ControlRigBeingDebuggedPtr.IsValid())
+	{
+		return;
+	}
+
+	if(InHierarchy != ControlRigBeingDebuggedPtr->GetHierarchy())
+	{
+		return;
+	}
+
+	FRigElementKey Key;
+	if(InElement)
+	{
+		Key = InElement->GetKey();
+	}
+
+	FFunctionGraphTask::CreateAndDispatchWhenReady([this, InNotif, InHierarchy, Key]()
+    {
+		const FRigBaseElement* Element = InHierarchy->Find(Key);
+		OnHierarchyModified(InNotif, InHierarchy, Element);
+		
+    }, TStatId(), NULL, ENamedThreads::GameThread);
+}
+
 void SRigHierarchy::HandleRefreshEditorFromBlueprint(UControlRigBlueprint* InBlueprint)
 {
 	if (bIsChangingRigHierarchy)
@@ -1116,8 +1193,10 @@ void SRigHierarchy::HandleSetObjectBeingDebugged(UObject* InObject)
 	if(UControlRig* ControlRig = Cast<UControlRig>(InObject))
 	{
 		ControlRigBeingDebuggedPtr = ControlRig;
-		ControlRig->GetHierarchy()->OnModified().AddSP(this, &SRigHierarchy::OnHierarchyModified);
+		ControlRig->GetHierarchy()->OnModified().AddSP(this, &SRigHierarchy::OnHierarchyModifiedAsync);
 	}
+
+	RefreshTreeView();
 }
 
 void SRigHierarchy::ClearDetailPanel() const
@@ -1923,6 +2002,17 @@ URigHierarchy* SRigHierarchy::GetDebuggedHierarchy() const
 	return GetHierarchy();
 }
 
+URigHierarchy* SRigHierarchy::GetHierarchyForTopology() const
+{
+	URigHierarchy* Hierarchy = GetHierarchy();
+	if(bShowDynamicHierarchy && ControlRigBeingDebuggedPtr.IsValid())
+	{
+		Hierarchy = ControlRigBeingDebuggedPtr->GetHierarchy();
+	}
+	return Hierarchy;
+}
+
+
 FName SRigHierarchy::CreateUniqueName(const FName& InBaseName, ERigElementType InElementType) const
 {
 	return GetHierarchy()->GetSafeNewName(InBaseName.ToString(), InElementType);
diff --git a/Engine/Plugins/Experimental/ControlRig/Source/ControlRigEditor/Private/Editor/SRigHierarchy.h b/Engine/Plugins/Experimental/ControlRig/Source/ControlRigEditor/Private/Editor/SRigHierarchy.h
--- a/Engine/Plugins/Experimental/ControlRig/Source/ControlRigEditor/Private/Editor/SRigHierarchy.h
+++ b/Engine/Plugins/Experimental/ControlRig/Source/ControlRigEditor/Private/Editor/SRigHierarchy.h
@@ -163,6 +163,15 @@ public:
 
 	void Construct(const FArguments& InArgs, TSharedRef<FControlRigEditor> InControlRigEditor);
 
+	FORCEINLINE FControlRigEditor* GetControlRigEditor() const
+	{
+		if(ControlRigEditor.IsValid())
+		{
+			return ControlRigEditor.Pin().Get();
+		}
+		return nullptr;
+	}
+
 private:
 	/** Bind commands that this widget handles */
 	void BindCommands();
@@ -171,7 +180,7 @@ private:
 	virtual FReply OnKeyDown(const FGeometry& MyGeometry, const FKeyEvent& InKeyEvent);
 
 	/** Rebuild the tree view */
-	void RefreshTreeView();
+	void RefreshTreeView(bool bRebuildContent = true);
 
 	/** Return all selected keys */
 	TArray<FRigElementKey> GetSelectedKeys() const;
@@ -232,8 +241,6 @@ private:
 	TOptional<EItemDropZone> OnCanAcceptDrop(const FDragDropEvent& DragDropEvent, EItemDropZone DropZone, TSharedPtr<FRigTreeElement> TargetItem);
 	FReply OnAcceptDrop(const FDragDropEvent& DragDropEvent, EItemDropZone DropZone, TSharedPtr<FRigTreeElement> TargetItem);
 
-private:
-
 	void FillContextMenu(class FMenuBuilder& MenuBuilder);
 	TSharedPtr<FUICommandList> GetContextMenuCommands();
 
@@ -303,6 +310,7 @@ private:
 
 	URigHierarchy* GetHierarchy() const;
 	URigHierarchy* GetDebuggedHierarchy() const;
+	URigHierarchy* GetHierarchyForTopology() const;
 
 	void ImportHierarchy(const FAssetData& InAssetData);
 	void CreateImportMenu(FMenuBuilder& MenuBuilder);
@@ -328,14 +336,16 @@ private:
 
 	bool bIsChangingRigHierarchy;
 	void OnHierarchyModified(ERigHierarchyNotification InNotif, URigHierarchy* InHierarchy, const FRigBaseElement* InElement);
+	void OnHierarchyModifiedAsync(ERigHierarchyNotification InNotif, URigHierarchy* InHierarchy, const FRigBaseElement* InElement);
 	void HandleRefreshEditorFromBlueprint(UControlRigBlueprint* InBlueprint);
 	void HandleSetObjectBeingDebugged(UObject* InObject);
 
 	static TSharedPtr<FRigTreeElement> FindElement(const FRigElementKey& InElementKey, TSharedPtr<FRigTreeElement> CurrentItem);
-	void AddElement(FRigElementKey InKey, FRigElementKey InParentKey = FRigElementKey(), const bool bIgnoreTextFilter = false);
-	void AddElement(FRigBaseElement* InElement, const bool bIgnoreTextFilter = false);
+	bool AddElement(FRigElementKey InKey, FRigElementKey InParentKey = FRigElementKey(), const bool bIgnoreTextFilter = false);
+	bool AddElement(const FRigBaseElement* InElement, const bool bIgnoreTextFilter = false);
 	void AddSpacerElement();
-	void ReparentElement(FRigElementKey InKey, FRigElementKey InParentKey);
+	bool ReparentElement(FRigElementKey InKey, FRigElementKey InParentKey);
+	bool RemoveElement(FRigElementKey InKey);
 
 public:
 	FName RenameElement(const FRigElementKey& OldKey, const FString& NewName);
