Control Rig: Make sure to check blueprint weak ptr when closing editor

diff --git a/Engine/Plugins/Experimental/ControlRig/Source/ControlRigEditor/Private/Editor/SRigHierarchy.cpp b/Engine/Plugins/Experimental/ControlRig/Source/ControlRigEditor/Private/Editor/SRigHierarchy.cpp
--- a/Engine/Plugins/Experimental/ControlRig/Source/ControlRigEditor/Private/Editor/SRigHierarchy.cpp
+++ b/Engine/Plugins/Experimental/ControlRig/Source/ControlRigEditor/Private/Editor/SRigHierarchy.cpp
@@ -1041,6 +1041,11 @@ bool SRigHierarchy::RemoveElement(FRigElementKey InKey)
 
 void SRigHierarchy::OnHierarchyModified(ERigHierarchyNotification InNotif, URigHierarchy* InHierarchy, const FRigBaseElement* InElement)
 {
+	if(!ControlRigBlueprint.IsValid())
+	{
+		return;
+	}
+	
 	if (ControlRigBlueprint->bSuspendAllNotifications)
 	{
 		return;
