Control Rig: Fix build break in RigHierarchyTabSummoner on mac

diff --git a/Engine/Plugins/Experimental/ControlRig/Source/ControlRigEditor/Private/Editor/RigHierarchyTabSummoner.cpp b/Engine/Plugins/Experimental/ControlRig/Source/ControlRigEditor/Private/Editor/RigHierarchyTabSummoner.cpp
--- a/Engine/Plugins/Experimental/ControlRig/Source/ControlRigEditor/Private/Editor/RigHierarchyTabSummoner.cpp
+++ b/Engine/Plugins/Experimental/ControlRig/Source/ControlRigEditor/Private/Editor/RigHierarchyTabSummoner.cpp
@@ -50,9 +50,9 @@ TSharedRef<SDockTab> FRigHierarchyTabSummoner::SpawnTab(const FWorkflowTabSpawnI
 		if(SWidget* Content = &DockTab->GetContent().Get())
 		{
 			SRigHierarchy* RigHierarchy = (SRigHierarchy*)Content;
-			if(FControlRigEditor* ControlRigEditor = RigHierarchy->GetControlRigEditor())
+			if(FControlRigEditor* ControlRigEditorForTab = RigHierarchy->GetControlRigEditor())
 			{
-				HierarchyTabCount = ControlRigEditor->GetRigHierarchyTabCount();
+				HierarchyTabCount = ControlRigEditorForTab->GetRigHierarchyTabCount();
 			}
 		}
         return HierarchyTabCount > 1;
@@ -62,9 +62,9 @@ TSharedRef<SDockTab> FRigHierarchyTabSummoner::SpawnTab(const FWorkflowTabSpawnI
 		if(SWidget* Content = &DockTab->GetContent().Get())
 		{
 			SRigHierarchy* RigHierarchy = (SRigHierarchy*)Content;
-			if(FControlRigEditor* ControlRigEditor = RigHierarchy->GetControlRigEditor())
+			if(FControlRigEditor* ControlRigEditorForTab = RigHierarchy->GetControlRigEditor())
 			{
-				ControlRigEditor->RigHierarchyTabCount--;
+				ControlRigEditorForTab->RigHierarchyTabCount--;
 			}
 		}
 	}));
