Control Rig: Support dynamic hierarchy name tags within graph editor

diff --git a/Engine/Plugins/Experimental/ControlRig/Source/ControlRigEditor/Private/Editor/ControlRigEditor.cpp b/Engine/Plugins/Experimental/ControlRig/Source/ControlRigEditor/Private/Editor/ControlRigEditor.cpp
--- a/Engine/Plugins/Experimental/ControlRig/Source/ControlRigEditor/Private/Editor/ControlRigEditor.cpp
+++ b/Engine/Plugins/Experimental/ControlRig/Source/ControlRigEditor/Private/Editor/ControlRigEditor.cpp
@@ -903,6 +903,14 @@ void FControlRigEditor::HandleSetObjectBeingDebugged(UObject* InObject)
 		}
 	}
 
+	if(UControlRig* PreviouslyDebuggedControlRig = Cast<UControlRig>(GetBlueprintObj()->GetObjectBeingDebugged()))
+	{
+		if(!PreviouslyDebuggedControlRig->HasAnyFlags(RF_BeginDestroyed))
+		{
+			PreviouslyDebuggedControlRig->GetHierarchy()->OnModified().RemoveAll(this);
+		}
+	}
+
 	if (UControlRigBlueprint* RigBlueprint = Cast<UControlRigBlueprint>(GetBlueprintObj()))
 	{
 		UpdateStaleWatchedPins();
@@ -952,6 +960,8 @@ void FControlRigEditor::HandleSetObjectBeingDebugged(UObject* InObject)
 				}
 			}
 		}
+
+		DebuggedControlRig->GetHierarchy()->OnModified().AddSP(this, &FControlRigEditor::OnHierarchyModifiedAsync);
 	}
 	else
 	{
@@ -2724,6 +2734,12 @@ void FControlRigEditor::CacheNameLists()
 		TArray<UEdGraph*> EdGraphs;
 		ControlRigBP->GetAllGraphs(EdGraphs);
 
+		URigHierarchy* Hierarchy = ControlRigBP->Hierarchy;
+		if(UControlRig* DebuggedControlRig = Cast<UControlRig>(GetBlueprintObj()->GetObjectBeingDebugged()))
+		{
+			Hierarchy = DebuggedControlRig->GetHierarchy();
+		}
+
 		for (UEdGraph* Graph : EdGraphs)
 		{
 			UControlRigGraph* RigGraph = Cast<UControlRigGraph>(Graph);
@@ -2731,7 +2747,7 @@ void FControlRigEditor::CacheNameLists()
 			{
 				continue;
 			}
-			RigGraph->CacheNameLists(ControlRigBP->Hierarchy, &ControlRigBP->DrawContainer);
+			RigGraph->CacheNameLists(Hierarchy, &ControlRigBP->DrawContainer);
 		}
 	}
 }
@@ -2752,6 +2768,12 @@ void FControlRigEditor::HandlePreviewMeshChanged(USkeletalMesh* InOldSkeletalMes
 		if (UControlRigBlueprint* ControlRigBP = GetControlRigBlueprint())
 		{
 			ControlRigBP->SetPreviewMesh(InNewSkeletalMesh);
+
+			if(UControlRig* DebuggedControlRig = Cast<UControlRig>(GetBlueprintObj()->GetObjectBeingDebugged()))
+			{
+				DebuggedControlRig->GetHierarchy()->Notify(ERigHierarchyNotification::HierarchyReset, nullptr);
+				DebuggedControlRig->RequestSetup();
+			}
 		}
 	}
 }
@@ -3458,6 +3480,38 @@ void FControlRigEditor::OnHierarchyModified(ERigHierarchyNotification InNotif, U
 	}
 }
 
+void FControlRigEditor::OnHierarchyModifiedAsync(ERigHierarchyNotification InNotif, URigHierarchy* InHierarchy, const FRigBaseElement* InElement)
+{
+	FRigElementKey Key;
+	if(InElement)
+	{
+		Key = InElement->GetKey();
+	}
+
+	FFunctionGraphTask::CreateAndDispatchWhenReady([this, InNotif, InHierarchy, Key]()
+    {
+        const FRigBaseElement* Element = InHierarchy->Find(Key);
+
+		switch(InNotif)
+		{
+			case ERigHierarchyNotification::ElementAdded:
+			case ERigHierarchyNotification::ElementRemoved:
+			case ERigHierarchyNotification::ElementRenamed:
+			case ERigHierarchyNotification::ParentChanged:
+			case ERigHierarchyNotification::HierarchyReset:
+			{
+				CacheNameLists();
+				break;
+			}
+			default:
+			{
+				break;
+			}
+		}
+		
+    }, TStatId(), NULL, ENamedThreads::GameThread);
+}
+
 void FControlRigEditor::SynchronizeViewportBoneSelection()
 {
 	UControlRigBlueprint* RigBlueprint = GetControlRigBlueprint();
diff --git a/Engine/Plugins/Experimental/ControlRig/Source/ControlRigEditor/Private/Editor/ControlRigEditor.h b/Engine/Plugins/Experimental/ControlRig/Source/ControlRigEditor/Private/Editor/ControlRigEditor.h
--- a/Engine/Plugins/Experimental/ControlRig/Source/ControlRigEditor/Private/Editor/ControlRigEditor.h
+++ b/Engine/Plugins/Experimental/ControlRig/Source/ControlRigEditor/Private/Editor/ControlRigEditor.h
@@ -158,6 +158,7 @@ public:
 	void OnCurveContainerChanged();
 
 	void OnHierarchyModified(ERigHierarchyNotification InNotif, URigHierarchy* InHierarchy, const FRigBaseElement* InElement);
+	void OnHierarchyModifiedAsync(ERigHierarchyNotification InNotif, URigHierarchy* InHierarchy, const FRigBaseElement* InElement);
 
 	void OnGraphNodeDropToPerform(TSharedPtr<FGraphNodeDragDropOp> DragDropOp, UEdGraph* Graph, const FVector2D& NodePosition, const FVector2D& ScreenPosition);
 
diff --git a/Engine/Plugins/Experimental/ControlRig/Source/ControlRigEditor/Private/Editor/SRigHierarchy.cpp b/Engine/Plugins/Experimental/ControlRig/Source/ControlRigEditor/Private/Editor/SRigHierarchy.cpp
--- a/Engine/Plugins/Experimental/ControlRig/Source/ControlRigEditor/Private/Editor/SRigHierarchy.cpp
+++ b/Engine/Plugins/Experimental/ControlRig/Source/ControlRigEditor/Private/Editor/SRigHierarchy.cpp
@@ -160,7 +160,7 @@ void SRigHierarchyItem::Construct(const FArguments& InArgs, TSharedPtr<FControlR
 		{
 			ERigBoneType BoneType = ERigBoneType::User;
 
-			FRigBoneElement* BoneElement = InHierarchy->ControlRigBlueprint->Hierarchy->Find<FRigBoneElement>(InRigTreeElement->Key);
+			FRigBoneElement* BoneElement = InHierarchy->GetHierarchyForTopology()->Find<FRigBoneElement>(InRigTreeElement->Key);
 			if(BoneElement)
 			{
 				BoneType = BoneElement->BoneType;
@@ -620,7 +620,7 @@ FReply SRigHierarchy::OnKeyDown(const FGeometry& MyGeometry, const FKeyEvent& In
 
 EVisibility SRigHierarchy::IsToolbarVisible() const
 {
-	if (URigHierarchy* Hierarchy = GetHierarchy())
+	if (URigHierarchy* Hierarchy = GetDebuggedHierarchy())
 	{
 		if (Hierarchy->Num(ERigElementType::Bone) > 0)
 		{
@@ -632,7 +632,7 @@ EVisibility SRigHierarchy::IsToolbarVisible() const
 
 EVisibility SRigHierarchy::IsSearchbarVisible() const
 {
-	if (URigHierarchy* Hierarchy = GetHierarchy())
+	if (URigHierarchy* Hierarchy = GetDebuggedHierarchy())
 	{
 		if ((Hierarchy->Num(ERigElementType::Bone) +
 			Hierarchy->Num(ERigElementType::Space) +
@@ -1068,31 +1068,36 @@ void SRigHierarchy::OnHierarchyModified(ERigHierarchyNotification InNotif, URigH
 	{
 		case ERigHierarchyNotification::ElementAdded:
 		{
-			check(InElement);
-			if(AddElement(InElement))
+			if(InElement)
 			{
-				RefreshTreeView(false);
+				if(AddElement(InElement))
+				{
+					RefreshTreeView(false);
+				}
 			}
 			break;
 		}
 		case ERigHierarchyNotification::ElementRemoved:
 		{
-			check(InElement);
-			if(RemoveElement(InElement->GetKey()))
+			if(InElement)
 			{
-				RefreshTreeView(false);
+				if(RemoveElement(InElement->GetKey()))
+				{
+					RefreshTreeView(false);
+				}
 			}
 			break;
 		}
 		case ERigHierarchyNotification::ParentChanged:
 		{
 			check(InHierarchy);
-			check(InElement);
-
-			const FRigElementKey ParentKey = InHierarchy->GetFirstParent(InElement->GetKey());
-			if(ReparentElement(InElement->GetKey(), ParentKey))
+			if(InElement)
 			{
-				RefreshTreeView(false);
+				const FRigElementKey ParentKey = InHierarchy->GetFirstParent(InElement->GetKey());
+				if(ReparentElement(InElement->GetKey(), ParentKey))
+				{
+					RefreshTreeView(false);
+				}
 			}
 			break;
 		}
@@ -1105,28 +1110,30 @@ void SRigHierarchy::OnHierarchyModified(ERigHierarchyNotification InNotif, URigH
 		case ERigHierarchyNotification::ElementSelected:
 		case ERigHierarchyNotification::ElementDeselected:
 		{
-			check(InElement);
-
-			const bool bSelected = (InNotif == ERigHierarchyNotification::ElementSelected); 
-				
-			for (int32 RootIndex = 0; RootIndex < RootElements.Num(); ++RootIndex)
+			if(InElement)
 			{
-				TSharedPtr<FRigTreeElement> Found = FindElement(InElement->GetKey(), RootElements[RootIndex]);
-				if (Found.IsValid())
+				const bool bSelected = (InNotif == ERigHierarchyNotification::ElementSelected); 
+					
+				for (int32 RootIndex = 0; RootIndex < RootElements.Num(); ++RootIndex)
 				{
-					TreeView->SetItemSelection(Found, bSelected, ESelectInfo::OnNavigation);
-					HandleFrameSelection();
-
-					if (ControlRigEditor.IsValid())
+					TSharedPtr<FRigTreeElement> Found = FindElement(InElement->GetKey(), RootElements[RootIndex]);
+					if (Found.IsValid())
 					{
-						if (ControlRigEditor.Pin()->GetEventQueue() == EControlRigEditorEventQueue::Setup)
+						TreeView->SetItemSelection(Found, bSelected, ESelectInfo::OnNavigation);
+						HandleFrameSelection();
+
+						if (ControlRigEditor.IsValid())
 						{
-							TGuardValue<bool> GuardRigHierarchyChanges(bIsChangingRigHierarchy, true);
-							HandleControlBoneOrSpaceTransform();
+							if (ControlRigEditor.Pin()->GetEventQueue() == EControlRigEditorEventQueue::Setup)
+							{
+								TGuardValue<bool> GuardRigHierarchyChanges(bIsChangingRigHierarchy, true);
+								HandleControlBoneOrSpaceTransform();
+							}
 						}
 					}
 				}
 			}
+			break;
 		}
 		default:
 		{
