Control Rig: Use  weak pointers for async events

diff --git a/Engine/Plugins/Experimental/ControlRig/Source/ControlRig/Private/Rigs/RigHierarchy.cpp b/Engine/Plugins/Experimental/ControlRig/Source/ControlRig/Private/Rigs/RigHierarchy.cpp
--- a/Engine/Plugins/Experimental/ControlRig/Source/ControlRig/Private/Rigs/RigHierarchy.cpp
+++ b/Engine/Plugins/Experimental/ControlRig/Source/ControlRig/Private/Rigs/RigHierarchy.cpp
@@ -998,19 +998,19 @@ void URigHierarchy::SendEvent(const FRigEventContext& InEvent, bool bAsynchronou
 {
 	if(EventDelegate.IsBound())
 	{
-		URigHierarchy* LocalHierarchy = this;
+		TWeakObjectPtr<URigHierarchy> WeakThis = this;
 		FRigEventDelegate& Delegate = EventDelegate;
 
 		if (bAsynchronous)
 		{
-			FFunctionGraphTask::CreateAndDispatchWhenReady([LocalHierarchy, Delegate, InEvent]()
+			FFunctionGraphTask::CreateAndDispatchWhenReady([WeakThis, Delegate, InEvent]()
             {
-                Delegate.Broadcast(LocalHierarchy, InEvent);
+                Delegate.Broadcast(WeakThis.Get(), InEvent);
             }, TStatId(), NULL, ENamedThreads::GameThread);
 		}
 		else
 		{
-			Delegate.Broadcast(LocalHierarchy, InEvent);
+			Delegate.Broadcast(this, InEvent);
 		}
 	}
 
diff --git a/Engine/Plugins/Experimental/ControlRig/Source/ControlRig/Private/Rigs/RigHierarchyController.cpp b/Engine/Plugins/Experimental/ControlRig/Source/ControlRig/Private/Rigs/RigHierarchyController.cpp
--- a/Engine/Plugins/Experimental/ControlRig/Source/ControlRig/Private/Rigs/RigHierarchyController.cpp
+++ b/Engine/Plugins/Experimental/ControlRig/Source/ControlRig/Private/Rigs/RigHierarchyController.cpp
@@ -969,6 +969,11 @@ TArray<FRigElementKey> URigHierarchyController::ImportFromText(FString InContent
 	TransactionPtr.Reset();
 #endif
 
+	if(bSelectNewElements)
+	{
+		SetSelection(PastedKeys);
+	}
+
 	return PastedKeys;
 }
 
diff --git a/Engine/Plugins/Experimental/ControlRig/Source/ControlRigEditor/Private/Editor/ControlRigEditor.cpp b/Engine/Plugins/Experimental/ControlRig/Source/ControlRigEditor/Private/Editor/ControlRigEditor.cpp
--- a/Engine/Plugins/Experimental/ControlRig/Source/ControlRigEditor/Private/Editor/ControlRigEditor.cpp
+++ b/Engine/Plugins/Experimental/ControlRig/Source/ControlRigEditor/Private/Editor/ControlRigEditor.cpp
@@ -3488,9 +3488,15 @@ void FControlRigEditor::OnHierarchyModifiedAsync(ERigHierarchyNotification InNot
 		Key = InElement->GetKey();
 	}
 
-	FFunctionGraphTask::CreateAndDispatchWhenReady([this, InNotif, InHierarchy, Key]()
+	TWeakObjectPtr<URigHierarchy> WeakHierarchy = InHierarchy;
+	FFunctionGraphTask::CreateAndDispatchWhenReady([this, InNotif, WeakHierarchy, Key]()
     {
-        const FRigBaseElement* Element = InHierarchy->Find(Key);
+    	if(!WeakHierarchy.IsValid())
+    	{
+    		return;
+    	}
+
+        const FRigBaseElement* Element = WeakHierarchy.Get()->Find(Key);
 
 		switch(InNotif)
 		{
diff --git a/Engine/Plugins/Experimental/ControlRig/Source/ControlRigEditor/Private/Editor/ControlRigSkeletalMeshComponent.cpp b/Engine/Plugins/Experimental/ControlRig/Source/ControlRigEditor/Private/Editor/ControlRigSkeletalMeshComponent.cpp
--- a/Engine/Plugins/Experimental/ControlRig/Source/ControlRigEditor/Private/Editor/ControlRigSkeletalMeshComponent.cpp
+++ b/Engine/Plugins/Experimental/ControlRig/Source/ControlRigEditor/Private/Editor/ControlRigSkeletalMeshComponent.cpp
@@ -189,11 +189,17 @@ void UControlRigSkeletalMeshComponent::OnHierarchyModifiedAsync(ERigHierarchyNot
 	{
 		Key = InElement->GetKey();
 	}
+
+	TWeakObjectPtr<URigHierarchy> WeakHierarchy = InHierarchy;
 	
-	FFunctionGraphTask::CreateAndDispatchWhenReady([this, InNotif, InHierarchy, Key]()
+	FFunctionGraphTask::CreateAndDispatchWhenReady([this, InNotif, WeakHierarchy, Key]()
     {
-        const FRigBaseElement* Element = InHierarchy->Find(Key);
-        OnHierarchyModified(InNotif, InHierarchy, Element);
+    	if(!WeakHierarchy.IsValid())
+    	{
+    		return;
+    	}
+        const FRigBaseElement* Element = WeakHierarchy.Get()->Find(Key);
+        OnHierarchyModified(InNotif, WeakHierarchy.Get(), Element);
 		
     }, TStatId(), NULL, ENamedThreads::GameThread);
 }
diff --git a/Engine/Plugins/Experimental/ControlRig/Source/ControlRigEditor/Private/Editor/SRigHierarchy.cpp b/Engine/Plugins/Experimental/ControlRig/Source/ControlRigEditor/Private/Editor/SRigHierarchy.cpp
--- a/Engine/Plugins/Experimental/ControlRig/Source/ControlRigEditor/Private/Editor/SRigHierarchy.cpp
+++ b/Engine/Plugins/Experimental/ControlRig/Source/ControlRigEditor/Private/Editor/SRigHierarchy.cpp
@@ -720,8 +720,12 @@ void SRigHierarchy::RefreshTreeView(bool bRebuildContent)
 				AddSpacerElement();
 			}
 		}
-		
+
 		TreeView->RequestTreeRefresh();
+		{
+			TGuardValue<bool> GuardRigHierarchyChanges(bIsChangingRigHierarchy, true);
+			TreeView->ClearSelection();
+		}
 
 		TArray<FRigElementKey> Selection = Hierarchy->GetSelectedKeys();
 		for (const FRigElementKey& Key : Selection)
@@ -1165,10 +1169,16 @@ void SRigHierarchy::OnHierarchyModifiedAsync(ERigHierarchyNotification InNotif,
 		Key = InElement->GetKey();
 	}
 
-	FFunctionGraphTask::CreateAndDispatchWhenReady([this, InNotif, InHierarchy, Key]()
+	TWeakObjectPtr<URigHierarchy> WeakHierarchy = InHierarchy;
+
+	FFunctionGraphTask::CreateAndDispatchWhenReady([this, InNotif, WeakHierarchy, Key]()
     {
-		const FRigBaseElement* Element = InHierarchy->Find(Key);
-		OnHierarchyModified(InNotif, InHierarchy, Element);
+    	if(!WeakHierarchy.IsValid())
+    	{
+    		return;
+    	}
+		const FRigBaseElement* Element = WeakHierarchy.Get()->Find(Key);
+		OnHierarchyModified(InNotif, WeakHierarchy.Get(), Element);
 		
     }, TStatId(), NULL, ENamedThreads::GameThread);
 }
@@ -1761,8 +1771,9 @@ void SRigHierarchy::HandleDuplicateItem()
 			check(Controller);
 
 			const TArray<FRigElementKey> KeysToDuplicate = GetSelectedKeys();
-			Controller->DuplicateElements(KeysToDuplicate);
+			Controller->DuplicateElements(KeysToDuplicate, true, true);
 		}
+
 		ControlRigBlueprint->PropagateHierarchyFromBPToInstances();
 	}
 
@@ -1801,7 +1812,7 @@ void SRigHierarchy::HandleMirrorItem()
 
 			const TArray<FRigElementKey> KeysToMirror = GetSelectedKeys();
 			const TArray<FRigElementKey> KeysToDuplicate = GetSelectedKeys();
-			Controller->MirrorElements(KeysToDuplicate, Settings, true);
+			Controller->MirrorElements(KeysToDuplicate, Settings, true, true);
 		}
 		ControlRigBlueprint->PropagateHierarchyFromBPToInstances();
 	}
