Control Rig: Fix workflow for changing gizmo name and display name

diff --git a/Engine/Plugins/Experimental/ControlRig/Source/ControlRigEditor/Private/ControlRigElementDetails.cpp b/Engine/Plugins/Experimental/ControlRig/Source/ControlRigEditor/Private/ControlRigElementDetails.cpp
--- a/Engine/Plugins/Experimental/ControlRig/Source/ControlRigEditor/Private/ControlRigElementDetails.cpp
+++ b/Engine/Plugins/Experimental/ControlRig/Source/ControlRigEditor/Private/ControlRigElementDetails.cpp
@@ -1508,7 +1508,12 @@ void FRigControlElementDetails::SetDisplayName(const FText& InNewText, ETextComm
 	{
 		if(FRigControlElement* ControlElement = HierarchyBeingCustomized->Find<FRigControlElement>(ElementKeyBeingCustomized))
 		{
-			FString NewDisplayName = InNewText.ToString().TrimStartAndEnd();
+			if(BlueprintBeingCustomized)
+			{
+				BlueprintBeingCustomized->Hierarchy->Modify();
+			}
+
+			const FString NewDisplayName = InNewText.ToString().TrimStartAndEnd();
 			if (NewDisplayName.IsEmpty())
 			{
 				ControlElement->Settings.DisplayName = FName(NAME_None);
@@ -1520,7 +1525,14 @@ void FRigControlElementDetails::SetDisplayName(const FText& InNewText, ETextComm
 
 			HierarchyBeingCustomized->Notify(ERigHierarchyNotification::ControlSettingChanged, ControlElement);
 
-			OnStructContentsChanged(FRigControl::StaticStruct()->FindPropertyByName(TEXT("DisplayName")), PropertyUtilities);
+			if (BlueprintBeingCustomized && BlueprintBeingCustomized->Hierarchy != HierarchyBeingCustomized)
+			{
+				if(FRigControlElement* OtherControlElement = BlueprintBeingCustomized->Hierarchy->Find<FRigControlElement>(ElementKeyBeingCustomized))
+				{
+					OtherControlElement->Settings.DisplayName = ControlElement->Settings.DisplayName;
+					BlueprintBeingCustomized->Hierarchy->Notify(ERigHierarchyNotification::ControlSettingChanged, OtherControlElement);
+				}
+			}
 		}
 	}
 }
diff --git a/Engine/Plugins/Experimental/ControlRig/Source/ControlRigEditor/Private/Editor/SControlRigGizmoNameList.cpp b/Engine/Plugins/Experimental/ControlRig/Source/ControlRigEditor/Private/Editor/SControlRigGizmoNameList.cpp
--- a/Engine/Plugins/Experimental/ControlRig/Source/ControlRigEditor/Private/Editor/SControlRigGizmoNameList.cpp
+++ b/Engine/Plugins/Experimental/ControlRig/Source/ControlRigEditor/Private/Editor/SControlRigGizmoNameList.cpp
@@ -90,18 +90,21 @@ void SControlRigGizmoNameList::SetNameListText(const FText& NewTypeInValue, ETex
 		if ((ControlElement != nullptr) && (ControlElement->Settings.GizmoName != NewName))
 		{
 			const FScopedTransaction Transaction(NSLOCTEXT("ControlRigEditor", "ChangeGizmoName", "Change Gizmo Name"));
-			Blueprint->Modify();
-			
+			Blueprint->Hierarchy->Modify();
+
 			ControlElement->Settings.GizmoName = NewName;
 
 			if (bIsOnInstance && DebuggedControlRig)
 			{
-				Blueprint->PropagatePropertyFromInstanceToBP(ControlKey, FRigControl::StaticStruct()->FindPropertyByName(TEXT("GizmoName")), DebuggedControlRig);
+				FRigControlElement* OtherControlElement = Blueprint->Hierarchy->Get<FRigControlElement>(ControlIndex);
+				if(OtherControlElement)
+				{
+					OtherControlElement->Settings.GizmoName = NewName;
+					Blueprint->Hierarchy->Notify(ERigHierarchyNotification::ControlSettingChanged, OtherControlElement);
+				}
 			}
 
-			Blueprint->PropagatePropertyFromBPToInstances(ControlKey, FRigControl::StaticStruct()->FindPropertyByName(TEXT("GizmoName")));
-
-			Blueprint->Hierarchy->Notify(ERigHierarchyNotification::ControlSettingChanged, Blueprint->Hierarchy->Find(ControlElement->GetKey()));
+			Hierarchy->Notify(ERigHierarchyNotification::ControlSettingChanged, ControlElement);
 		}
 	}
 }
