Control Rig: Ensure to force transform changes coming from the ui layer

diff --git a/Engine/Plugins/Experimental/ControlRig/Source/ControlRigEditor/Private/Editor/ControlRigEditor.cpp b/Engine/Plugins/Experimental/ControlRig/Source/ControlRigEditor/Private/Editor/ControlRigEditor.cpp
--- a/Engine/Plugins/Experimental/ControlRig/Source/ControlRigEditor/Private/Editor/ControlRigEditor.cpp
+++ b/Engine/Plugins/Experimental/ControlRig/Source/ControlRigEditor/Private/Editor/ControlRigEditor.cpp
@@ -3155,11 +3155,11 @@ void FControlRigEditor::OnBlueprintPropertyChainEvent(FPropertyChangedChainEvent
 						
 						if(ERigTransformType::IsInitial(TransformType))
 						{
-							Hierarchy->SetControlValue(TargetControlElement, Value, ERigControlValueType::Initial, true);
+							Hierarchy->SetControlValue(TargetControlElement, Value, ERigControlValueType::Initial, true, true);
 						}
 						else
 						{
-							Hierarchy->SetControlValue(TargetControlElement, Value, ERigControlValueType::Current, true);
+							Hierarchy->SetControlValue(TargetControlElement, Value, ERigControlValueType::Current, true, true);
 						}
 					}
 					else
