Control Rig: Fix align control functionality

diff --git a/Engine/Plugins/Experimental/ControlRig/Source/ControlRigEditor/Private/Editor/SRigHierarchy.cpp b/Engine/Plugins/Experimental/ControlRig/Source/ControlRigEditor/Private/Editor/SRigHierarchy.cpp
--- a/Engine/Plugins/Experimental/ControlRig/Source/ControlRigEditor/Private/Editor/SRigHierarchy.cpp
+++ b/Engine/Plugins/Experimental/ControlRig/Source/ControlRigEditor/Private/Editor/SRigHierarchy.cpp
@@ -2223,9 +2223,11 @@ FReply SRigHierarchy::OnAcceptDrop(const FDragDropEvent& DragDropEvent, EItemDro
 						FTransform OffsetTransform = TargetGlobalTransform.GetRelativeTransform(ParentTransform);
 
 						Hierarchy->SetControlOffsetTransformByIndex(ControlIndex, OffsetTransform, ERigTransformType::InitialLocal, true, true);
+						Hierarchy->SetControlOffsetTransformByIndex(ControlIndex, OffsetTransform, ERigTransformType::CurrentLocal, true, true);
 						Hierarchy->SetLocalTransform(DraggedKey, FTransform::Identity, true, true, true);
 						Hierarchy->SetInitialLocalTransform(DraggedKey, FTransform::Identity, true, true);
 						DebuggedHierarchy->SetControlOffsetTransformByIndex(ControlIndex, OffsetTransform, ERigTransformType::InitialLocal, true, true);
+						DebuggedHierarchy->SetControlOffsetTransformByIndex(ControlIndex, OffsetTransform, ERigTransformType::CurrentLocal, true, true);
 						DebuggedHierarchy->SetLocalTransform(DraggedKey, FTransform::Identity, true, true, true);
 						DebuggedHierarchy->SetInitialLocalTransform(DraggedKey, FTransform::Identity, true, true);
 					}
