Control Rig: Support for Undo / Redo for transform changes in the hierarchy

diff --git a/Engine/Plugins/Experimental/ControlRig/Source/ControlRig/Private/ControlRig.cpp b/Engine/Plugins/Experimental/ControlRig/Source/ControlRig/Private/ControlRig.cpp
--- a/Engine/Plugins/Experimental/ControlRig/Source/ControlRig/Private/ControlRig.cpp
+++ b/Engine/Plugins/Experimental/ControlRig/Source/ControlRig/Private/ControlRig.cpp
@@ -582,7 +582,7 @@ void UControlRig::Execute(const EControlRigState InState, const FName& InEventNa
 
 			if (UControlRig* CDO = Cast<UControlRig>(GetClass()->GetDefaultObject()))
 			{
-				if(bCopyHierarchyBeforeSetup)
+				if(bCopyHierarchyBeforeSetup && !bSetupModeEnabled)
 				{
 					if(CDO->GetHierarchy()->GetTopologyVersion()!= GetHierarchy()->GetTopologyVersion())
 					{
@@ -1441,15 +1441,19 @@ FName UControlRig::AddTransientControl(const FRigElementKey& InElement)
 		return NAME_None;
 	}
 
-	RemoveTransientControl(InElement);
+	const FName ControlName = GetNameForTransientControl(InElement);
+	if(DynamicHierarchy->Contains(FRigElementKey(ControlName, ERigElementType::Control)))
+	{
+		SetTransientControlValue(InElement);
+		return ControlName;
+	}
+
 	const int32 ElementIndex = DynamicHierarchy->GetIndex(InElement);
 	if (ElementIndex == INDEX_NONE)
 	{
 		return NAME_None;
 	}
 
-	const FName ControlName = GetNameForTransientControl(InElement);
-
 	FTransform GizmoTransform = FTransform::Identity;
 	GizmoTransform.SetScale3D(FVector::ZeroVector);
 
@@ -1608,6 +1612,11 @@ FName UControlRig::GetNameForTransientControl(const FRigElementKey& InElement)
 
 FRigElementKey UControlRig::GetElementKeyFromTransientControl(const FRigElementKey& InKey)
 {
+	if(InKey.Type != ERigElementType::Control)
+	{
+		return FRigElementKey();
+	}
+	
 	static FString ControlRigForElementBoneName;
 	static FString ControlRigForElementSpaceName;
 
@@ -2223,6 +2232,17 @@ void UControlRig::SetBoneInitialTransformsFromRefSkeleton(const FReferenceSkelet
 	bResetInitialTransformsBeforeSetup = false;
 }
 
+void UControlRig::OnHierarchyTransformUndoRedo(URigHierarchy* InHierarchy, const FRigElementKey& InKey, ERigTransformType::Type InTransformType, const FTransform& InTransform, bool bIsUndo)
+{
+	if(InKey.Type == ERigElementType::Control)
+	{
+		if(FRigControlElement* ControlElement = InHierarchy->Find<FRigControlElement>(InKey))
+		{
+			ControlModified().Broadcast(this, ControlElement, FRigControlModifiedContext(EControlRigSetKey::Never));
+		}
+	}
+}
+
 #undef LOCTEXT_NAMESPACE
 
 
diff --git a/Engine/Plugins/Experimental/ControlRig/Source/ControlRig/Private/Rigs/RigHierarchy.cpp b/Engine/Plugins/Experimental/ControlRig/Source/ControlRig/Private/Rigs/RigHierarchy.cpp
--- a/Engine/Plugins/Experimental/ControlRig/Source/ControlRig/Private/Rigs/RigHierarchy.cpp
+++ b/Engine/Plugins/Experimental/ControlRig/Source/ControlRig/Private/Rigs/RigHierarchy.cpp
@@ -28,7 +28,7 @@ void URigHierarchy::Serialize(FArchive& Ar)
 {
 	Ar.UsingCustomVersion(FAnimObjectVersion::GUID);
 
-	if (Ar.IsSaving() || Ar.IsObjectReferenceCollector())
+	if (Ar.IsSaving() || Ar.IsObjectReferenceCollector() || Ar.IsCountingMemory())
 	{
 		Save(Ar);
 	}
@@ -47,6 +47,8 @@ void URigHierarchy::Save(FArchive& Ar)
 	if(Ar.IsTransacting())
 	{
 		Ar << TransformStackIndex;
+		Ar << bTransactingForTransformChange;
+		
 		if(bTransactingForTransformChange)
 		{
 			return;
@@ -79,8 +81,11 @@ void URigHierarchy::Load(FArchive& Ar)
 {
 	if(Ar.IsTransacting())
 	{
+		bool bOnlySerializedTransformStackIndex = false;
 		Ar << TransformStackIndex;
-		if(TransformStackIndex != TransformUndoStack.Num())
+		Ar << bOnlySerializedTransformStackIndex;
+		
+		if(bOnlySerializedTransformStackIndex)
 		{
 			return;
 		}
@@ -904,6 +909,7 @@ bool URigHierarchy::Undo()
 
 	const FTransformStackEntry Entry = TransformUndoStack.Pop();
 	ApplyTransformFromStack(Entry, true);
+	UndoRedoEvent.Broadcast(this, Entry.Key, Entry.TransformType, Entry.OldTransform, true);
 	TransformRedoStack.Push(Entry);
 	TransformStackIndex = TransformUndoStack.Num();
 	return true;
@@ -926,6 +932,7 @@ bool URigHierarchy::Redo()
 
 	const FTransformStackEntry Entry = TransformRedoStack.Pop();
 	ApplyTransformFromStack(Entry, false);
+	UndoRedoEvent.Broadcast(this, Entry.Key, Entry.TransformType, Entry.NewTransform, false);
 	TransformUndoStack.Push(Entry);
 	TransformStackIndex = TransformUndoStack.Num();
 	return true;
@@ -1873,6 +1880,11 @@ void URigHierarchy::PushTransformToStack(const FRigElementKey& InKey, ETransform
 {
 #if WITH_EDITOR
 
+	if(GIsTransacting)
+	{
+		return;
+	}
+
 	static const FText TransformPoseTitle = NSLOCTEXT("RigHierarchy", "Set Pose Transform", "Set Pose Transform");
 	static const FText ControlOffsetTitle = NSLOCTEXT("RigHierarchy", "Set Control Offset", "Set Control Offset");
 	static const FText ControlGizmoTitle = NSLOCTEXT("RigHierarchy", "Set Control Gizo", "Set Control Gizo");
@@ -1904,8 +1916,38 @@ void URigHierarchy::PushTransformToStack(const FRigElementKey& InKey, ETransform
 	}
 
 	TGuardValue<bool> TransactingGuard(bTransactingForTransformChange, true);
+
 	FScopedTransaction Transaction(Title);
+
+	if(bIsInteracting)
+	{
+		FTransformStackEntry LastEntry;
+		if(!TransformUndoStack.IsEmpty())
+		{
+			LastEntry = TransformUndoStack.Last();
+		}
+
+		if(LastEntry.Key == InKey && LastEntry.EntryType == InEntryType && LastEntry.bAffectChildren == bAffectChildren)
+		{
+			// merge the entries on the stack
+			TransformUndoStack.Last() = 
+                FTransformStackEntry(InKey, InEntryType, InTransformType, LastEntry.OldTransform, InNewTransform, bAffectChildren);
+		}
+		else
+		{
+			Modify();
+
+			TransformUndoStack.Add(
+                FTransformStackEntry(InKey, InEntryType, InTransformType, InOldTransform, InNewTransform, bAffectChildren));
+			TransformStackIndex = TransformUndoStack.Num();
+		}
+
+		TransformRedoStack.Reset();
+		return;
+	}
+
 	Modify();
+
 	TransformUndoStack.Add(
 		FTransformStackEntry(InKey, InEntryType, InTransformType, InOldTransform, InNewTransform, bAffectChildren));
 	TransformStackIndex = TransformUndoStack.Num();
diff --git a/Engine/Plugins/Experimental/ControlRig/Source/ControlRig/Private/Rigs/RigHierarchyController.cpp b/Engine/Plugins/Experimental/ControlRig/Source/ControlRig/Private/Rigs/RigHierarchyController.cpp
--- a/Engine/Plugins/Experimental/ControlRig/Source/ControlRig/Private/Rigs/RigHierarchyController.cpp
+++ b/Engine/Plugins/Experimental/ControlRig/Source/ControlRig/Private/Rigs/RigHierarchyController.cpp
@@ -103,6 +103,7 @@ bool URigHierarchyController::SetSelection(TArray<FRigElementKey> InKeys)
 	}
 
 	TArray<FRigElementKey> PreviousSelection = Hierarchy->GetSelectedKeys();
+	bool bResult = false;
 
 	for(const FRigElementKey& KeyToDeselect : PreviousSelection)
 	{
@@ -112,6 +113,7 @@ bool URigHierarchyController::SetSelection(TArray<FRigElementKey> InKeys)
 			{
 				return false;
 			}
+			bResult = true;
 		}
 	}
 
@@ -123,10 +125,11 @@ bool URigHierarchyController::SetSelection(TArray<FRigElementKey> InKeys)
 			{
 				return false;
 			}
+			bResult = true;
 		}
 	}
 
-	return true;
+	return bResult;
 }
 
 FRigElementKey URigHierarchyController::AddBone(FName InName, FRigElementKey InParent, FTransform InTransform,
diff --git a/Engine/Plugins/Experimental/ControlRig/Source/ControlRig/Public/ControlRig.h b/Engine/Plugins/Experimental/ControlRig/Source/ControlRig/Public/ControlRig.h
--- a/Engine/Plugins/Experimental/ControlRig/Source/ControlRig/Public/ControlRig.h
+++ b/Engine/Plugins/Experimental/ControlRig/Source/ControlRig/Public/ControlRig.h
@@ -190,6 +190,9 @@ public:
 	{
 		if(DynamicHierarchy != nullptr && DynamicHierarchy->GetOuter() != this)
 		{
+#if WITH_EDITOR
+			DynamicHierarchy->OnUndoRedo().RemoveAll(this);
+#endif
 			DynamicHierarchy = nullptr;
 		}
 		
@@ -198,8 +201,11 @@ public:
 			DynamicHierarchy = NewObject<URigHierarchy>(this);
 			if (!HasAnyFlags(RF_ClassDefaultObject))
 			{
-				DynamicHierarchy->SetFlags(RF_Transient);
+				DynamicHierarchy->SetFlags(RF_Transient | RF_Transactional);
 			}
+#if WITH_EDITOR
+			DynamicHierarchy->OnUndoRedo().AddUObject(this, &UControlRig::OnHierarchyTransformUndoRedo);
+#endif
 		}
 		return DynamicHierarchy;
 	}
@@ -275,14 +281,18 @@ public:
 		const FRigControlModifiedContext& Context = FRigControlModifiedContext())
 	{
 		const FRigElementKey Key(InControlName, ERigElementType::Control);
-		DynamicHierarchy->SetControlValue(Key, InValue);
+
+		FRigControlElement* ControlElement = DynamicHierarchy->Find<FRigControlElement>(Key);
+		if(ControlElement == nullptr)
+		{
+			return;
+		}
+
+		DynamicHierarchy->SetControlValue(ControlElement, InValue, ERigControlValueType::Current, true);
 
 		if (bNotify && OnControlModified.IsBound())
 		{
-			if(FRigControlElement* ControlElement = FindControl(InControlName))
-			{
-				OnControlModified.Broadcast(this, ControlElement, Context);
-			}
+			OnControlModified.Broadcast(this, ControlElement, Context);
 		}
 	}
 
@@ -594,6 +604,7 @@ protected:
 		return InterRigSyncBracket > 0;
 	}
 
+	void OnHierarchyTransformUndoRedo(URigHierarchy* InHierarchy, const FRigElementKey& InKey, ERigTransformType::Type InTransformType, const FTransform& InTransform, bool bIsUndo);
 
 	FFilterControlEvent OnFilterControl;
 	FControlModifiedEvent OnControlModified;
@@ -640,17 +651,24 @@ class CONTROLRIG_API FControlRigInteractionScope
 public:
 
 	FORCEINLINE FControlRigInteractionScope(UControlRig* InControlRig)
-		: InteractionBracketScope(InControlRig->InteractionBracket)
+		: ControlRig(InControlRig)
+		, InteractionBracketScope(InControlRig->InteractionBracket)
 		, SyncBracketScope(InControlRig->InterRigSyncBracket)
 	{
+		InControlRig->GetHierarchy()->StartInteraction();
 	}
 
 	FORCEINLINE ~FControlRigInteractionScope()
 	{
+		if(ControlRig.IsValid())
+		{
+			ControlRig->GetHierarchy()->EndInteraction();
+		}
 	}
 
 private:
 
+	TWeakObjectPtr<UControlRig> ControlRig;
 	FControlRigBracketScope InteractionBracketScope;
 	FControlRigBracketScope SyncBracketScope;
 };
diff --git a/Engine/Plugins/Experimental/ControlRig/Source/ControlRig/Public/Rigs/RigHierarchy.h b/Engine/Plugins/Experimental/ControlRig/Source/ControlRig/Public/Rigs/RigHierarchy.h
--- a/Engine/Plugins/Experimental/ControlRig/Source/ControlRig/Public/Rigs/RigHierarchy.h
+++ b/Engine/Plugins/Experimental/ControlRig/Source/ControlRig/Public/Rigs/RigHierarchy.h
@@ -14,6 +14,7 @@ class URigHierarchyController;
 #define URIGHIERARCHY_RECURSIVE_DIRTY_PROPAGATION 1
 
 DECLARE_MULTICAST_DELEGATE_ThreeParams(FRigHierarchyModifiedEvent, ERigHierarchyNotification /* type */, URigHierarchy* /* hierarchy */, const FRigBaseElement* /* element */);
+DECLARE_EVENT_FiveParams(URigHierarchy, FRigHierarchyUndoRedoTransformEvent, URigHierarchy*, const FRigElementKey&, ERigTransformType::Type, const FTransform&, bool /* bUndo */);
 
 UCLASS(BlueprintType)
 class CONTROLRIG_API URigHierarchy : public UObject
@@ -29,6 +30,7 @@ public:
 	, IndexLookup()
 	, TransformStackIndex(0)
 	, bTransactingForTransformChange(false)
+	, bIsInteracting(false)
 	{}
 
 	virtual ~URigHierarchy();
@@ -1584,6 +1586,23 @@ public:
 	 */
 	bool Redo();
 
+	/**
+	 * Returns the event fired during undo / redo
+	 */
+	FRigHierarchyUndoRedoTransformEvent& OnUndoRedo() { return UndoRedoEvent; }
+
+	/**
+	 * Starts an interaction on the rig.
+	 * This will cause all transform actions happening to be merged
+	 */
+	void StartInteraction() { bIsInteracting = true; }
+
+	/**
+	 * Starts an interaction on the rig.
+	 * This will cause all transform actions happening to be merged
+	 */
+	FORCEINLINE void EndInteraction() { bIsInteracting = false; }
+
 	/**
 	 * Returns the transform stack index
 	 */
@@ -2129,6 +2148,16 @@ private:
 	 */
 	bool ApplyTransformFromStack(const FTransformStackEntry& InEntry, bool bUndo);
 
+	/**
+	 * Manages merging transform actions into one during an interaction
+	 */
+	bool bIsInteracting;
+
+	/**
+	 * The event fired during undo / redo
+	 */
+	FRigHierarchyUndoRedoTransformEvent UndoRedoEvent;
+
 	TWeakObjectPtr<URigHierarchy> HierarchyForSelectionPtr;
 	TWeakObjectPtr<UObject> LastControllerPtr;
 
diff --git a/Engine/Plugins/Experimental/ControlRig/Source/ControlRigEditor/Private/ControlRigEditorModule.cpp b/Engine/Plugins/Experimental/ControlRig/Source/ControlRigEditor/Private/ControlRigEditorModule.cpp
--- a/Engine/Plugins/Experimental/ControlRig/Source/ControlRigEditor/Private/ControlRigEditorModule.cpp
+++ b/Engine/Plugins/Experimental/ControlRig/Source/ControlRigEditor/Private/ControlRigEditorModule.cpp
@@ -1675,12 +1675,8 @@ void FControlRigEditorModule::GetContextMenuActions(const UControlRigGraphSchema
 						FSlateIcon(),
 						FUIAction(FExecuteAction::CreateLambda([RigBlueprint, RigElementsToSelect]() {
 
-							RigBlueprint->GetHierarchyController()->ClearSelection();
-							for(const FRigElementKey& RigElementToSelect : RigElementsToSelect)
-							{
-								RigBlueprint->GetHierarchyController()->SelectElement(RigElementToSelect, true);
-							}
-
+							RigBlueprint->GetHierarchyController()->SetSelection(RigElementsToSelect);
+							
 						})
 					));
 				}
diff --git a/Engine/Plugins/Experimental/ControlRig/Source/ControlRigEditor/Private/EditMode/SControlHierarchy.cpp b/Engine/Plugins/Experimental/ControlRig/Source/ControlRigEditor/Private/EditMode/SControlHierarchy.cpp
--- a/Engine/Plugins/Experimental/ControlRig/Source/ControlRigEditor/Private/EditMode/SControlHierarchy.cpp
+++ b/Engine/Plugins/Experimental/ControlRig/Source/ControlRigEditor/Private/EditMode/SControlHierarchy.cpp
@@ -296,19 +296,7 @@ void SControlHierarchy::OnSelectionChanged(TSharedPtr<FControlTreeElement> Selec
 				NewSelection.Add(SelectedItem->Key);
 			}
 
-			for (const FRigElementKey& PreviouslySelected : OldSelection)
-			{
-				if (NewSelection.Contains(PreviouslySelected))
-				{
-					continue;
-				}
-				Controller->SelectElement(PreviouslySelected, false);
-			}
-
-			for (const FRigElementKey& NewlySelected : NewSelection)
-			{
-				Controller->SelectElement(NewlySelected, true);
-			}
+			Controller->SetSelection(NewSelection);
 		}
 	}
 }
diff --git a/Engine/Plugins/Experimental/ControlRig/Source/ControlRigEditor/Private/Editor/ControlRigEditor.cpp b/Engine/Plugins/Experimental/ControlRig/Source/ControlRigEditor/Private/Editor/ControlRigEditor.cpp
--- a/Engine/Plugins/Experimental/ControlRig/Source/ControlRigEditor/Private/Editor/ControlRigEditor.cpp
+++ b/Engine/Plugins/Experimental/ControlRig/Source/ControlRigEditor/Private/Editor/ControlRigEditor.cpp
@@ -768,10 +768,7 @@ void FControlRigEditor::ToggleSetupMode()
 	{
 		if (UControlRigBlueprint* RigBlueprint = Cast<UControlRigBlueprint>(GetBlueprintObj()))
 		{
-			for (FRigElementKey SelectedKey : PreviousSelection)
-			{
-				RigBlueprint->HierarchyController->SelectElement(SelectedKey, true);
-			}
+			RigBlueprint->GetHierarchyController()->SetSelection(PreviousSelection);
 		}
 	}
 
@@ -961,7 +958,7 @@ void FControlRigEditor::HandleSetObjectBeingDebugged(UObject* InObject)
 			}
 		}
 
-		DebuggedControlRig->GetHierarchy()->OnModified().AddSP(this, &FControlRigEditor::OnHierarchyModifiedAsync);
+		DebuggedControlRig->GetHierarchy()->OnModified().AddSP(this, &FControlRigEditor::OnHierarchyModified_AnyThread);
 	}
 	else
 	{
@@ -1626,7 +1623,7 @@ void FControlRigEditor::PostUndo(bool bSuccess)
 
 		if (FControlRigEditMode* EditMode = GetEditMode())
 		{
-			EditMode->RecreateGizmoActors(RigBlueprint->Hierarchy->GetSelectedKeys());
+			EditMode->RequestToRecreateGizmoActors();
 		}
 	}
 }
@@ -1647,7 +1644,7 @@ void FControlRigEditor::PostRedo(bool bSuccess)
 
 		if (FControlRigEditMode* EditMode = GetEditMode())
 		{
-			EditMode->RecreateGizmoActors(RigBlueprint->Hierarchy->GetSelectedKeys());
+			EditMode->RequestToRecreateGizmoActors();
 		}
 	}
 }
@@ -3480,7 +3477,7 @@ void FControlRigEditor::OnHierarchyModified(ERigHierarchyNotification InNotif, U
 	}
 }
 
-void FControlRigEditor::OnHierarchyModifiedAsync(ERigHierarchyNotification InNotif, URigHierarchy* InHierarchy, const FRigBaseElement* InElement)
+void FControlRigEditor::OnHierarchyModified_AnyThread(ERigHierarchyNotification InNotif, URigHierarchy* InHierarchy, const FRigBaseElement* InElement)
 {
 	FRigElementKey Key;
 	if(InElement)
diff --git a/Engine/Plugins/Experimental/ControlRig/Source/ControlRigEditor/Private/Editor/ControlRigEditor.h b/Engine/Plugins/Experimental/ControlRig/Source/ControlRigEditor/Private/Editor/ControlRigEditor.h
--- a/Engine/Plugins/Experimental/ControlRig/Source/ControlRigEditor/Private/Editor/ControlRigEditor.h
+++ b/Engine/Plugins/Experimental/ControlRig/Source/ControlRigEditor/Private/Editor/ControlRigEditor.h
@@ -158,7 +158,7 @@ public:
 	void OnCurveContainerChanged();
 
 	void OnHierarchyModified(ERigHierarchyNotification InNotif, URigHierarchy* InHierarchy, const FRigBaseElement* InElement);
-	void OnHierarchyModifiedAsync(ERigHierarchyNotification InNotif, URigHierarchy* InHierarchy, const FRigBaseElement* InElement);
+	void OnHierarchyModified_AnyThread(ERigHierarchyNotification InNotif, URigHierarchy* InHierarchy, const FRigBaseElement* InElement);
 
 	void OnGraphNodeDropToPerform(TSharedPtr<FGraphNodeDragDropOp> DragDropOp, UEdGraph* Graph, const FVector2D& NodePosition, const FVector2D& ScreenPosition);
 
diff --git a/Engine/Plugins/Experimental/ControlRig/Source/ControlRigEditor/Private/Editor/ControlRigSkeletalMeshComponent.cpp b/Engine/Plugins/Experimental/ControlRig/Source/ControlRigEditor/Private/Editor/ControlRigSkeletalMeshComponent.cpp
--- a/Engine/Plugins/Experimental/ControlRig/Source/ControlRigEditor/Private/Editor/ControlRigSkeletalMeshComponent.cpp
+++ b/Engine/Plugins/Experimental/ControlRig/Source/ControlRigEditor/Private/Editor/ControlRigSkeletalMeshComponent.cpp
@@ -154,7 +154,7 @@ void UControlRigSkeletalMeshComponent::SetControlRigBeingDebugged(UControlRig* I
 	if(InControlRig)
 	{
 		ControlRigBeingDebuggedPtr = InControlRig;
- 		InControlRig->GetHierarchy()->OnModified().AddUObject(this, &UControlRigSkeletalMeshComponent::OnHierarchyModifiedAsync);
+ 		InControlRig->GetHierarchy()->OnModified().AddUObject(this, &UControlRigSkeletalMeshComponent::OnHierarchyModified_AnyThread);
 	}
 
 	RebuildDebugDrawSkeleton();
@@ -181,7 +181,7 @@ void UControlRigSkeletalMeshComponent::OnHierarchyModified(ERigHierarchyNotifica
 	}
 }
 
-void UControlRigSkeletalMeshComponent::OnHierarchyModifiedAsync(ERigHierarchyNotification InNotif,
+void UControlRigSkeletalMeshComponent::OnHierarchyModified_AnyThread(ERigHierarchyNotification InNotif,
     URigHierarchy* InHierarchy, const FRigBaseElement* InElement)
 {
 	FRigElementKey Key;
diff --git a/Engine/Plugins/Experimental/ControlRig/Source/ControlRigEditor/Private/Editor/ControlRigSkeletalMeshComponent.h b/Engine/Plugins/Experimental/ControlRig/Source/ControlRigEditor/Private/Editor/ControlRigSkeletalMeshComponent.h
--- a/Engine/Plugins/Experimental/ControlRig/Source/ControlRigEditor/Private/Editor/ControlRigSkeletalMeshComponent.h
+++ b/Engine/Plugins/Experimental/ControlRig/Source/ControlRigEditor/Private/Editor/ControlRigSkeletalMeshComponent.h
@@ -52,7 +52,7 @@ public:
 	void RebuildDebugDrawSkeleton();
 
 	void OnHierarchyModified(ERigHierarchyNotification InNotif, URigHierarchy* InHierarchy, const FRigBaseElement* InElement);
-	void OnHierarchyModifiedAsync(ERigHierarchyNotification InNotif, URigHierarchy* InHierarchy, const FRigBaseElement* InElement);
+	void OnHierarchyModified_AnyThread(ERigHierarchyNotification InNotif, URigHierarchy* InHierarchy, const FRigBaseElement* InElement);
 
 private:
 	FReferenceSkeleton DebugDrawSkeleton;
diff --git a/Engine/Plugins/Experimental/ControlRig/Source/ControlRigEditor/Private/Editor/SRigHierarchy.cpp b/Engine/Plugins/Experimental/ControlRig/Source/ControlRigEditor/Private/Editor/SRigHierarchy.cpp
--- a/Engine/Plugins/Experimental/ControlRig/Source/ControlRigEditor/Private/Editor/SRigHierarchy.cpp
+++ b/Engine/Plugins/Experimental/ControlRig/Source/ControlRigEditor/Private/Editor/SRigHierarchy.cpp
@@ -725,12 +725,19 @@ void SRigHierarchy::RefreshTreeView(bool bRebuildContent)
 		{
 			TGuardValue<bool> GuardRigHierarchyChanges(bIsChangingRigHierarchy, true);
 			TreeView->ClearSelection();
-		}
 
-		TArray<FRigElementKey> Selection = Hierarchy->GetSelectedKeys();
-		for (const FRigElementKey& Key : Selection)
-		{
-			OnHierarchyModified(ERigHierarchyNotification::ElementSelected, Hierarchy, Hierarchy->FindChecked(Key));
+			TArray<FRigElementKey> Selection = Hierarchy->GetSelectedKeys();
+			for (const FRigElementKey& Key : Selection)
+			{
+				for (int32 RootIndex = 0; RootIndex < RootElements.Num(); ++RootIndex)
+				{
+					TSharedPtr<FRigTreeElement> Found = FindElement(Key, RootElements[RootIndex]);
+					if (Found.IsValid())
+					{
+						TreeView->SetItemSelection(Found, true, ESelectInfo::OnNavigation);
+					}
+				}
+			}
 		}
 	}
 }
@@ -796,7 +803,10 @@ void SRigHierarchy::OnSelectionChanged(TSharedPtr<FRigTreeElement> Selection, ES
 		TGuardValue<bool> GuardRigHierarchyChanges(bIsChangingRigHierarchy, true);
 
 		const TArray<FRigElementKey> NewSelection = GetSelectedKeys();
-		Controller->SetSelection(NewSelection);
+		if(!Controller->SetSelection(NewSelection))
+		{
+			return;
+		}
 
 		if (NewSelection.Num() > 0)
 		{
@@ -1126,7 +1136,7 @@ void SRigHierarchy::OnHierarchyModified(ERigHierarchyNotification InNotif, URigH
 						TreeView->SetItemSelection(Found, bSelected, ESelectInfo::OnNavigation);
 						HandleFrameSelection();
 
-						if (ControlRigEditor.IsValid())
+						if (ControlRigEditor.IsValid() && !GIsTransacting)
 						{
 							if (ControlRigEditor.Pin()->GetEventQueue() == EControlRigEditorEventQueue::Setup)
 							{
@@ -1146,12 +1156,16 @@ void SRigHierarchy::OnHierarchyModified(ERigHierarchyNotification InNotif, URigH
 	}
 }
 
-void SRigHierarchy::OnHierarchyModifiedAsync(ERigHierarchyNotification InNotif, URigHierarchy* InHierarchy, const FRigBaseElement* InElement)
+void SRigHierarchy::OnHierarchyModified_AnyThread(ERigHierarchyNotification InNotif, URigHierarchy* InHierarchy, const FRigBaseElement* InElement)
 {
 	if(!bShowDynamicHierarchy)
 	{
 		return;
 	}
+	if(bIsChangingRigHierarchy)
+	{
+		return;
+	}
 	
 	if(!ControlRigBeingDebuggedPtr.IsValid())
 	{
@@ -1163,24 +1177,31 @@ void SRigHierarchy::OnHierarchyModifiedAsync(ERigHierarchyNotification InNotif,
 		return;
 	}
 
-	FRigElementKey Key;
-	if(InElement)
+	if(IsInGameThread())
 	{
-		Key = InElement->GetKey();
+		OnHierarchyModified(InNotif, InHierarchy, InElement);
 	}
+	else
+	{
+		FRigElementKey Key;
+		if(InElement)
+		{
+			Key = InElement->GetKey();
+		}
 
-	TWeakObjectPtr<URigHierarchy> WeakHierarchy = InHierarchy;
+		TWeakObjectPtr<URigHierarchy> WeakHierarchy = InHierarchy;
 
-	FFunctionGraphTask::CreateAndDispatchWhenReady([this, InNotif, WeakHierarchy, Key]()
-    {
-    	if(!WeakHierarchy.IsValid())
-    	{
-    		return;
-    	}
-		const FRigBaseElement* Element = WeakHierarchy.Get()->Find(Key);
-		OnHierarchyModified(InNotif, WeakHierarchy.Get(), Element);
-		
-    }, TStatId(), NULL, ENamedThreads::GameThread);
+		FFunctionGraphTask::CreateAndDispatchWhenReady([this, InNotif, WeakHierarchy, Key]()
+        {
+            if(!WeakHierarchy.IsValid())
+            {
+                return;
+            }
+            const FRigBaseElement* Element = WeakHierarchy.Get()->Find(Key);
+            OnHierarchyModified(InNotif, WeakHierarchy.Get(), Element);
+			
+        }, TStatId(), NULL, ENamedThreads::GameThread);
+	}
 }
 
 void SRigHierarchy::HandleRefreshEditorFromBlueprint(UControlRigBlueprint* InBlueprint)
@@ -1215,7 +1236,7 @@ void SRigHierarchy::HandleSetObjectBeingDebugged(UObject* InObject)
 	if(UControlRig* ControlRig = Cast<UControlRig>(InObject))
 	{
 		ControlRigBeingDebuggedPtr = ControlRig;
-		ControlRig->GetHierarchy()->OnModified().AddSP(this, &SRigHierarchy::OnHierarchyModifiedAsync);
+		ControlRig->GetHierarchy()->OnModified().AddSP(this, &SRigHierarchy::OnHierarchyModified_AnyThread);
 	}
 
 	RefreshTreeView();
@@ -2452,6 +2473,17 @@ void SRigHierarchy::HandleControlBoneOrSpaceTransform()
 		return;
 	}
 
+	UControlRig* DebuggedControlRig = Cast<UControlRig>(Blueprint->GetObjectBeingDebugged());
+	if(DebuggedControlRig == nullptr)
+	{
+		return;
+	}
+
+	if(!DebuggedControlRig->IsSetupModeEnabled())
+	{
+		return;
+	}
+
 	TArray<FRigElementKey> SelectedKeys = GetSelectedKeys();
 	for (const FRigElementKey& SelectedKey : SelectedKeys)
 	{
diff --git a/Engine/Plugins/Experimental/ControlRig/Source/ControlRigEditor/Private/Editor/SRigHierarchy.h b/Engine/Plugins/Experimental/ControlRig/Source/ControlRigEditor/Private/Editor/SRigHierarchy.h
--- a/Engine/Plugins/Experimental/ControlRig/Source/ControlRigEditor/Private/Editor/SRigHierarchy.h
+++ b/Engine/Plugins/Experimental/ControlRig/Source/ControlRigEditor/Private/Editor/SRigHierarchy.h
@@ -336,7 +336,7 @@ private:
 
 	bool bIsChangingRigHierarchy;
 	void OnHierarchyModified(ERigHierarchyNotification InNotif, URigHierarchy* InHierarchy, const FRigBaseElement* InElement);
-	void OnHierarchyModifiedAsync(ERigHierarchyNotification InNotif, URigHierarchy* InHierarchy, const FRigBaseElement* InElement);
+	void OnHierarchyModified_AnyThread(ERigHierarchyNotification InNotif, URigHierarchy* InHierarchy, const FRigBaseElement* InElement);
 	void HandleRefreshEditorFromBlueprint(UControlRigBlueprint* InBlueprint);
 	void HandleSetObjectBeingDebugged(UObject* InObject);
 
