Control Rig: Fix static analysis with DynamicHierarchy pointer

diff --git a/Engine/Plugins/Experimental/ControlRig/Source/ControlRig/Private/ControlRig.cpp b/Engine/Plugins/Experimental/ControlRig/Source/ControlRig/Private/ControlRig.cpp
--- a/Engine/Plugins/Experimental/ControlRig/Source/ControlRig/Private/ControlRig.cpp
+++ b/Engine/Plugins/Experimental/ControlRig/Source/ControlRig/Private/ControlRig.cpp
@@ -137,10 +137,11 @@ void UControlRig::Initialize(bool bInitRigUnits)
 		RequestInit();
 	}
 
-	DynamicHierarchy->OnModified().RemoveAll(this);
-	DynamicHierarchy->OnModified().AddUObject(this, &UControlRig::HandleHierarchyModified);
-	DynamicHierarchy->OnEventReceived().RemoveAll(this);
-	DynamicHierarchy->OnEventReceived().AddUObject(this, &UControlRig::HandleHierarchyEvent);
+	
+	GetHierarchy()->OnModified().RemoveAll(this);
+	GetHierarchy()->OnModified().AddUObject(this, &UControlRig::HandleHierarchyModified);
+	GetHierarchy()->OnEventReceived().RemoveAll(this);
+	GetHierarchy()->OnEventReceived().AddUObject(this, &UControlRig::HandleHierarchyEvent);
 }
 
 void UControlRig::InitializeFromCDO()
@@ -559,10 +560,10 @@ void UControlRig::Execute(const EControlRigState InState, const FName& InEventNa
 		}
 	}
 
-	if(DynamicHierarchy)
+	if(GetHierarchy())
 	{
 		// if we have any aux elements dirty them
-		DynamicHierarchy->UpdateSockets(&Context);
+		GetHierarchy()->UpdateSockets(&Context);
 	}
 
 #if WITH_EDITOR
@@ -612,7 +613,7 @@ void UControlRig::Execute(const EControlRigState InState, const FName& InEventNa
 
 			if (bSetupModeEnabled)
 			{
-				DynamicHierarchy->ResetPoseToInitial(ERigElementType::Bone);
+				GetHierarchy()->ResetPoseToInitial(ERigElementType::Bone);
 			}
 		}
 		else
@@ -724,7 +725,7 @@ void UControlRig::Execute(const EControlRigState InState, const FName& InEventNa
 	{
 		Context.DrawInterface->Instructions.Append(Context.DrawContainer->Instructions);
 
-		DynamicHierarchy->ForEach<FRigControlElement>([this](FRigControlElement* ControlElement) -> bool
+		GetHierarchy()->ForEach<FRigControlElement>([this](FRigControlElement* ControlElement) -> bool
 		{
 			const FRigControlSettings& Settings = ControlElement->Settings;
 			
@@ -742,7 +743,7 @@ void UControlRig::Execute(const EControlRigState InState, const FName& InEventNa
 					return true;
 				}
 
-				FTransform Transform = DynamicHierarchy->GetGlobalControlOffsetTransformByIndex(ControlElement->GetIndex());
+				FTransform Transform = GetHierarchy()->GetGlobalControlOffsetTransformByIndex(ControlElement->GetIndex());
 				FControlRigDrawInstruction Instruction(EControlRigDrawSettings::Lines, Settings.GizmoColor, 0.f, Transform);
 
 				switch (Settings.ControlType)
@@ -977,6 +978,8 @@ void UControlRig::GetMappableNodeData(TArray<FName>& OutNames, TArray<FNodeItem>
 	OutNames.Reset();
 	OutNodeItems.Reset();
 
+	check(DynamicHierarchy);
+
 	// now add all nodes
 	DynamicHierarchy->ForEach<FRigBoneElement>([&OutNames, &OutNodeItems, this](FRigBoneElement* BoneElement) -> bool
     {
@@ -1071,6 +1074,10 @@ TArray<FRigControlElement*> UControlRig::AvailableControls() const
 
 FRigControlElement* UControlRig::FindControl(const FName& InControlName) const
 {
+	if(DynamicHierarchy == nullptr)
+	{
+		return nullptr;
+	}
 	return DynamicHierarchy->Find<FRigControlElement>(FRigElementKey(InControlName, ERigElementType::Control));
 }
 
@@ -1081,9 +1088,9 @@ FTransform UControlRig::SetupControlFromGlobalTransform(const FName& InControlNa
 		FRigControlElement* ControlElement = FindControl(InControlName);
 		if (ControlElement && !ControlElement->Settings.bIsTransientControl)
 		{
-			const FTransform ParentTransform = DynamicHierarchy->GetParentTransform(ControlElement, ERigTransformType::CurrentGlobal);
+			const FTransform ParentTransform = GetHierarchy()->GetParentTransform(ControlElement, ERigTransformType::CurrentGlobal);
 			const FTransform OffsetTransform = InGlobalTransform.GetRelativeTransform(ParentTransform);
-			DynamicHierarchy->SetControlOffsetTransform(ControlElement, OffsetTransform, ERigTransformType::InitialLocal, true, true);
+			GetHierarchy()->SetControlOffsetTransform(ControlElement, OffsetTransform, ERigTransformType::InitialLocal, true, true);
 			ControlElement->Offset.Current = ControlElement->Offset.Initial; 
 		}
 	}
@@ -1146,6 +1153,10 @@ bool UControlRig::IsCurveControl(const FRigControlElement* InControlElement) con
 
 FTransform UControlRig::GetControlGlobalTransform(const FName& InControlName) const
 {
+	if(DynamicHierarchy == nullptr)
+	{
+		return FTransform::Identity;
+	}
 	return DynamicHierarchy->GetGlobalTransform(FRigElementKey(InControlName, ERigElementType::Control), false);
 }
 
@@ -1177,13 +1188,16 @@ FRigControlValue UControlRig::GetControlValueFromGlobalTransform(const FName& In
 
 	if (FRigControlElement* ControlElement = FindControl(InControlName))
 	{
-		FTransform ParentTransform = DynamicHierarchy->GetControlOffsetTransform(ControlElement, ERigTransformType::CurrentGlobal);
-		FTransform Transform = InGlobalTransform.GetRelativeTransform(ParentTransform);
-		Value.SetFromTransform(Transform, ControlElement->Settings.ControlType, ControlElement->Settings.PrimaryAxis);
-
-		if (ShouldApplyLimits())
+		if(DynamicHierarchy)
 		{
-			ControlElement->Settings.ApplyLimits(Value);
+			FTransform ParentTransform = DynamicHierarchy->GetControlOffsetTransform(ControlElement, ERigTransformType::CurrentGlobal);
+			FTransform Transform = InGlobalTransform.GetRelativeTransform(ParentTransform);
+			Value.SetFromTransform(Transform, ControlElement->Settings.ControlType, ControlElement->Settings.PrimaryAxis);
+
+			if (ShouldApplyLimits())
+			{
+				ControlElement->Settings.ApplyLimits(Value);
+			}
 		}
 	}
 
@@ -1208,6 +1222,10 @@ void UControlRig::SetControlLocalTransform(const FName& InControlName, const FTr
 
 FTransform UControlRig::GetControlLocalTransform(const FName& InControlName)
 {
+	if(DynamicHierarchy == nullptr)
+	{
+		return FTransform::Identity;
+	}
 	return DynamicHierarchy->GetLocalTransform(FRigElementKey(InControlName, ERigElementType::Control));
 }
 
@@ -1229,17 +1247,23 @@ UControlRigGizmoLibrary* UControlRig::GetGizmoLibrary() const
 
 void UControlRig::SelectControl(const FName& InControlName, bool bSelect)
 {
-	if(URigHierarchyController* Controller = DynamicHierarchy->GetController(true))
+	if(DynamicHierarchy)
 	{
-		Controller->SelectElement(FRigElementKey(InControlName, ERigElementType::Control), bSelect);
+		if(URigHierarchyController* Controller = DynamicHierarchy->GetController(true))
+		{
+			Controller->SelectElement(FRigElementKey(InControlName, ERigElementType::Control), bSelect);
+		}
 	}
 }
 
 bool UControlRig::ClearControlSelection()
 {
-	if(URigHierarchyController* Controller = DynamicHierarchy->GetController(true))
+	if(DynamicHierarchy)
 	{
-		return Controller->ClearSelection();
+		if(URigHierarchyController* Controller = DynamicHierarchy->GetController(true))
+		{
+			return Controller->ClearSelection();
+		}
 	}
 	return false;
 }
@@ -1247,19 +1271,26 @@ bool UControlRig::ClearControlSelection()
 TArray<FName> UControlRig::CurrentControlSelection() const
 {
 	TArray<FName> SelectedControlNames;
-	TArray<FRigBaseElement*> SelectedControls = DynamicHierarchy->GetSelectedElements(ERigElementType::Control);
-	for (FRigBaseElement* SelectedControl : SelectedControls)
+
+	if(DynamicHierarchy)
 	{
-		SelectedControlNames.Add(SelectedControl->GetName());
+		TArray<FRigBaseElement*> SelectedControls = DynamicHierarchy->GetSelectedElements(ERigElementType::Control);
+		for (FRigBaseElement* SelectedControl : SelectedControls)
+		{
+			SelectedControlNames.Add(SelectedControl->GetName());
+		}
 	}
 	return SelectedControlNames;
 }
 
 bool UControlRig::IsControlSelected(const FName& InControlName)const
 {
-	if(FRigControlElement* ControlElement = FindControl(InControlName))
+	if(DynamicHierarchy)
 	{
-		return DynamicHierarchy->IsSelected(ControlElement);
+		if(FRigControlElement* ControlElement = FindControl(InControlName))
+		{
+			return DynamicHierarchy->IsSelected(ControlElement);
+		}
 	}
 	return false;
 }
@@ -1522,6 +1553,11 @@ bool UControlRig::SetTransientControlValue(const FRigElementKey& InElement)
 		return false;
 	}
 
+	if(DynamicHierarchy == nullptr)
+	{
+		return false;
+	}
+
 	const FName ControlName = GetNameForTransientControl(InElement);
 	if(FRigControlElement* ControlElement = FindControl(ControlName))
 	{
@@ -1561,6 +1597,11 @@ FName UControlRig::RemoveTransientControl(const FRigElementKey& InElement)
 		return NAME_None;
 	}
 
+	if(DynamicHierarchy == nullptr)
+	{
+		return NAME_None;
+	}
+
 	URigHierarchyController* Controller = DynamicHierarchy->GetController(true);
 	if(Controller == nullptr)
 	{
