Animation Recorder: In correct number of keys.

diff --git a/Engine/Plugins/Experimental/ControlRig/Source/ControlRig/Private/Sequencer/MovieSceneControlRigParameterSection.cpp b/Engine/Plugins/Experimental/ControlRig/Source/ControlRig/Private/Sequencer/MovieSceneControlRigParameterSection.cpp
--- a/Engine/Plugins/Experimental/ControlRig/Source/ControlRig/Private/Sequencer/MovieSceneControlRigParameterSection.cpp
+++ b/Engine/Plugins/Experimental/ControlRig/Source/ControlRig/Private/Sequencer/MovieSceneControlRigParameterSection.cpp
@@ -2095,7 +2095,7 @@ bool UMovieSceneControlRigParameterSection::LoadAnimSequenceIntoThisSection(UAni
 		SetEndFrame(EndFrame);
 	}
 
-	const int32 NumberOfFrames = FrameRate.AsFrameNumber(Length).Value;
+	const int32 NumberOfFrames = FrameRate.AsFrameTime(Length).CeilToFrame().Value + 1;
 	FFrameNumber FrameRateInFrameNumber = TickResolution.AsFrameNumber(FrameRate.AsInterval());
 	int32 ExtraProgress = bKeyReduce ? FloatChannels.Num() : 0;
 	
diff --git a/Engine/Source/Editor/SequenceRecorder/Private/AnimationRecorder.cpp b/Engine/Source/Editor/SequenceRecorder/Private/AnimationRecorder.cpp
--- a/Engine/Source/Editor/SequenceRecorder/Private/AnimationRecorder.cpp
+++ b/Engine/Source/Editor/SequenceRecorder/Private/AnimationRecorder.cpp
@@ -586,7 +586,7 @@ void FAnimationRecorder::UpdateRecord(USkeletalMeshComponent* Component, float D
 	// time passed has been updated
 	// now find what frames we need to update
 	int32 FramesRecorded = LastFrame;
-	int32 FramesToRecord = FPlatformMath::TruncToInt(TimePassed / IntervalTime);
+	int32 FramesToRecord = FPlatformMath::RoundToInt(TimePassed / IntervalTime);
 
 	// notifies need to be done regardless of sample rate
 	if (Component->GetAnimInstance())
