Control Rig: Fix static analysis in rig hierarchy views and hierarchy pointer

diff --git a/Engine/Plugins/Experimental/ControlRig/Source/ControlRigEditor/Private/EditMode/SControlHierarchy.cpp b/Engine/Plugins/Experimental/ControlRig/Source/ControlRigEditor/Private/EditMode/SControlHierarchy.cpp
--- a/Engine/Plugins/Experimental/ControlRig/Source/ControlRigEditor/Private/EditMode/SControlHierarchy.cpp
+++ b/Engine/Plugins/Experimental/ControlRig/Source/ControlRigEditor/Private/EditMode/SControlHierarchy.cpp
@@ -279,24 +279,25 @@ void SControlHierarchy::OnSelectionChanged(TSharedPtr<FControlTreeElement> Selec
 	{
 		TGuardValue<bool> Guard(bSelecting, true);
 		URigHierarchy* Hierarchy = GetHierarchy();
-		check(Hierarchy);
-		URigHierarchyController* Controller = Hierarchy->GetController(true);
-		check(Controller);
-
+		
 		if (Hierarchy)
 		{
-			FScopedTransaction ScopedTransaction(LOCTEXT("SelectControlTransaction", "Select Control"), !GIsTransacting);
+			URigHierarchyController* Controller = Hierarchy->GetController(true);
+			if(Controller)
+			{
+				FScopedTransaction ScopedTransaction(LOCTEXT("SelectControlTransaction", "Select Control"), !GIsTransacting);
 
-			TArray<FRigElementKey> OldSelection = Hierarchy->GetSelectedKeys(ERigElementType::Control);
-			TArray<FRigElementKey> NewSelection;
+				TArray<FRigElementKey> OldSelection = Hierarchy->GetSelectedKeys(ERigElementType::Control);
+				TArray<FRigElementKey> NewSelection;
 
-			TArray<TSharedPtr<FControlTreeElement>> SelectedItems = TreeView->GetSelectedItems();
-			for (const TSharedPtr<FControlTreeElement>& SelectedItem : SelectedItems)
-			{
-				NewSelection.Add(SelectedItem->Key);
-			}
+				TArray<TSharedPtr<FControlTreeElement>> SelectedItems = TreeView->GetSelectedItems();
+				for (const TSharedPtr<FControlTreeElement>& SelectedItem : SelectedItems)
+				{
+					NewSelection.Add(SelectedItem->Key);
+				}
 
-			Controller->SetSelection(NewSelection);
+				Controller->SetSelection(NewSelection);
+			}
 		}
 	}
 }
diff --git a/Engine/Plugins/Experimental/ControlRig/Source/ControlRigEditor/Private/Editor/SRigHierarchy.cpp b/Engine/Plugins/Experimental/ControlRig/Source/ControlRigEditor/Private/Editor/SRigHierarchy.cpp
--- a/Engine/Plugins/Experimental/ControlRig/Source/ControlRigEditor/Private/Editor/SRigHierarchy.cpp
+++ b/Engine/Plugins/Experimental/ControlRig/Source/ControlRigEditor/Private/Editor/SRigHierarchy.cpp
@@ -2171,14 +2171,16 @@ FReply SRigHierarchy::OnAcceptDrop(const FDragDropEvent& DragDropEvent, EItemDro
 	if (RigDragDropOp.IsValid())
 	{
 		URigHierarchy* Hierarchy = GetHierarchy();
-		check(Hierarchy);
 		URigHierarchy* DebuggedHierarchy = GetDebuggedHierarchy();
-		check(DebuggedHierarchy);
-		URigHierarchyController* Controller = Hierarchy->GetController(true);
-		check(Controller);
 
 		if (Hierarchy && ControlRigBlueprint.IsValid())
 		{
+			URigHierarchyController* Controller = Hierarchy->GetController(true);
+			if(Controller == nullptr)
+			{
+				return FReply::Unhandled();
+			}
+
 			TGuardValue<bool> GuardRigHierarchyChanges(bIsChangingRigHierarchy, true);
 			TGuardValue<bool> SuspendBlueprintNotifs(ControlRigBlueprint->bSuspendAllNotifications, true);
 			FScopedTransaction Transaction(LOCTEXT("HierarchyDragAndDrop", "Drag & Drop"));
