Control Rig: Fix crash with contextual menu

diff --git a/Engine/Plugins/Experimental/ControlRig/Source/ControlRigEditor/Private/ControlRigEditorModule.cpp b/Engine/Plugins/Experimental/ControlRig/Source/ControlRigEditor/Private/ControlRigEditorModule.cpp
--- a/Engine/Plugins/Experimental/ControlRig/Source/ControlRigEditor/Private/ControlRigEditorModule.cpp
+++ b/Engine/Plugins/Experimental/ControlRig/Source/ControlRigEditor/Private/ControlRigEditorModule.cpp
@@ -1563,9 +1563,12 @@ void FControlRigEditorModule::GetContextMenuActions(const UControlRigGraphSchema
 							StructOnScope = UnitNode->ConstructStructInstance(false /* default */);
 							StructMemory = (FRigUnit*)StructOnScope->GetStructMemory();
 
+							FRigNameCache NameCache;
 							FRigUnitContext RigUnitContext;
 							RigUnitContext.Hierarchy = TemporaryHierarchy;
 							RigUnitContext.State = EControlRigState::Update;
+							RigUnitContext.NameCache = &NameCache;
+							
 							StructMemory->Execute(RigUnitContext);
 						}
 
