Control Rig: Skip curve items when importing the hierarchy

diff --git a/Engine/Plugins/Experimental/ControlRig/Source/ControlRigEditor/Private/Editor/SRigHierarchy.cpp b/Engine/Plugins/Experimental/ControlRig/Source/ControlRigEditor/Private/Editor/SRigHierarchy.cpp
--- a/Engine/Plugins/Experimental/ControlRig/Source/ControlRigEditor/Private/Editor/SRigHierarchy.cpp
+++ b/Engine/Plugins/Experimental/ControlRig/Source/ControlRigEditor/Private/Editor/SRigHierarchy.cpp
@@ -947,6 +947,10 @@ bool SRigHierarchy::AddElement(const FRigBaseElement* InElement, const bool bIgn
 			}
 			break;
 		}
+		case ERigElementType::Curve:
+		{
+			return false;
+		}
 		default:
 		{
 			break;
