Sequencer:Control Rig: New AnimSnapper Tool

diff --git a/Engine/Plugins/Experimental/ControlRig/Source/ControlRig/Private/Sequencer/ControlRigSequencerAnimInstance.cpp b/Engine/Plugins/Experimental/ControlRig/Source/ControlRig/Private/Sequencer/ControlRigSequencerAnimInstance.cpp
deleted file mode 100644
diff --git a/Engine/Plugins/Experimental/ControlRig/Source/ControlRig/Private/Sequencer/ControlRigSequencerAnimInstanceProxy.cpp b/Engine/Plugins/Experimental/ControlRig/Source/ControlRig/Private/Sequencer/ControlRigSequencerAnimInstanceProxy.cpp
deleted file mode 100644
diff --git a/Engine/Plugins/Experimental/ControlRig/Source/ControlRig/Private/Sequencer/ControlRigSequencerAnimInstanceProxy.h b/Engine/Plugins/Experimental/ControlRig/Source/ControlRig/Private/Sequencer/ControlRigSequencerAnimInstanceProxy.h
deleted file mode 100644
diff --git a/Engine/Plugins/Experimental/ControlRig/Source/ControlRig/Private/Tools/ControlRigSnapSettings.cpp b/Engine/Plugins/Experimental/ControlRig/Source/ControlRig/Private/Tools/ControlRigSnapSettings.cpp
new file mode 100644
--- /dev/null
+++ b/Engine/Plugins/Experimental/ControlRig/Source/ControlRig/Private/Tools/ControlRigSnapSettings.cpp
@@ -0,0 +1,6 @@
+// Copyright Epic Games, Inc. All Rights Reserved.
+#include "Tools/ControlRigSnapSettings.h"
+
+UControlRigSnapSettings::UControlRigSnapSettings()
+{
+}
diff --git a/Engine/Plugins/Experimental/ControlRig/Source/ControlRig/Public/Sequencer/ControlRigSequencerAnimInstance.h b/Engine/Plugins/Experimental/ControlRig/Source/ControlRig/Public/Sequencer/ControlRigSequencerAnimInstance.h
deleted file mode 100644
diff --git a/Engine/Plugins/Experimental/ControlRig/Source/ControlRig/Public/Tools/ControlRigSnapSettings.h b/Engine/Plugins/Experimental/ControlRig/Source/ControlRig/Public/Tools/ControlRigSnapSettings.h
new file mode 100644
--- /dev/null
+++ b/Engine/Plugins/Experimental/ControlRig/Source/ControlRig/Public/Tools/ControlRigSnapSettings.h
@@ -0,0 +1,40 @@
+// Copyright Epic Games, Inc. All Rights Reserved.
+#pragma once
+/**
+* Per Project Settings Animation Snapping
+*/
+#include "CoreMinimal.h"
+#include "UObject/ObjectMacros.h"
+#include "UObject/Object.h"
+#include "ControlRig.h"
+#include "ControlRigSnapSettings.generated.h"
+
+
+UCLASS(config = EditorPerProjectUserSettings)
+class CONTROLRIG_API UControlRigSnapSettings : public UObject
+{
+
+public:
+	GENERATED_BODY()
+
+	UControlRigSnapSettings();
+
+	/** When snapping keep offset, If false snap completely */
+	UPROPERTY(config, EditAnywhere, BlueprintReadWrite, Category = "Snap Settings", meta = (ToolTip = "When snapping keep offset, if false snap completely"))
+	bool bKeepOffset = false;
+
+	/** When snapping, also snap position */
+	UPROPERTY(config, EditAnywhere, BlueprintReadWrite, Category = "Snap Settings", meta = (ToolTip = "When snapping, also snap position"))
+	bool bSnapPosition = true;
+
+	/** When snapping, also snap rotation */
+	UPROPERTY(config, EditAnywhere, BlueprintReadWrite, Category = "Snap Settings", meta = (ToolTip = "When snapping, also snap rotation"))
+	bool bSnapRotation = false;
+
+	/** When snapping, also snap scale */
+	UPROPERTY(config, EditAnywhere, BlueprintReadWrite, Category = "Snap Settings", meta = (ToolTip = "When snapping, also snap scale"))
+	bool bSnapScale = false;
+
+};
+
+
diff --git a/Engine/Plugins/Experimental/ControlRig/Source/ControlRigEditor/Private/ControlRigEditorModule.cpp b/Engine/Plugins/Experimental/ControlRig/Source/ControlRigEditor/Private/ControlRigEditorModule.cpp
--- a/Engine/Plugins/Experimental/ControlRig/Source/ControlRigEditor/Private/ControlRigEditorModule.cpp
+++ b/Engine/Plugins/Experimental/ControlRig/Source/ControlRigEditor/Private/ControlRigEditorModule.cpp
@@ -112,6 +112,7 @@
 #include "AssetTypeActions_ControlRigPose.h"
 #include "EditMode/SControlRigBaseListWidget.h"
 #include "EditMode/SControlRigTweenWidget.h"
+#include "EditMode/SControlRigSnapper.h"
 
 #define LOCTEXT_NAMESPACE "ControlRigEditorModule"
 
@@ -119,6 +120,9 @@ DEFINE_LOG_CATEGORY(LogControlRigEditor);
 
 const FName IControlRigEditorModule::ControlRigPoseTab = FName("ControlRigPoseTab");
 const FName IControlRigEditorModule::ControlRigTweenTab = FName("ControlRigTweenTab");
+const FName IControlRigEditorModule::ControlRigSnapperTab = FName("ControlRigSnapperTab");
+const FName IControlRigEditorModule::ControlRigTempPivotTab = FName("ControlRigTempPivotTab");
+
 
 TSharedRef<SDockTab> SpawnRigProfiler( const FSpawnTabArgs& Args )
 {
@@ -147,6 +151,24 @@ TSharedRef<SDockTab> SpawnTweenTab(const FSpawnTabArgs& Args)
 		];
 }
 
+TSharedRef<SDockTab> SpawnSnapperTab(const FSpawnTabArgs& Args)
+{
+	return SNew(SDockTab)
+		.TabRole(ETabRole::NomadTab)
+		[
+			SNew(SControlRigSnapper)
+		];
+}
+
+TSharedRef<SDockTab> SpawnPivotTab(const FSpawnTabArgs& Args)
+{
+	return SNew(SDockTab)
+		.TabRole(ETabRole::NomadTab)
+		[
+			SNew(SControlRigTweenWidget)
+		];
+}
+
 void FControlRigEditorModule::StartupModule()
 {
 	FControlRigEditModeCommands::Register();
@@ -271,6 +293,22 @@ void FControlRigEditorModule::StartupModule()
 			.SetMenuType(ETabSpawnerMenuType::Hidden)
 			.SetIcon(FSlateIcon(FEditorStyle::GetStyleSetName(), "SkeletonTree.BlendProfile")); //MZ todo replace with correct icon
 		FGlobalTabmanager::Get()->RegisterDefaultTabWindowSize(IControlRigEditorModule::ControlRigTweenTab, FVector2D(400, 100));
+
+		FGlobalTabmanager::Get()->RegisterNomadTabSpawner(IControlRigEditorModule::ControlRigSnapperTab, FOnSpawnTab::CreateStatic(&SpawnSnapperTab))
+			.SetDisplayName(NSLOCTEXT("UnrealEditor", "ControlRigSnapperTab", "Control Rig Snapper"))
+			.SetTooltipText(NSLOCTEXT("UnrealEditor", "ControlRigSnapperTabTooltip", "Open the Control Rig Snapper tab."))
+			.SetMenuType(ETabSpawnerMenuType::Hidden)
+			.SetIcon(FSlateIcon(FEditorStyle::GetStyleSetName(), "Persona.Tabs.AnimSlotManager"));//MZ todo replace with correct icon SkeletonTree.SkeletonSocket
+		FGlobalTabmanager::Get()->RegisterDefaultTabWindowSize(IControlRigEditorModule::ControlRigSnapperTab, FVector2D(400, 400));
+
+		FGlobalTabmanager::Get()->RegisterNomadTabSpawner(IControlRigEditorModule::ControlRigTempPivotTab, FOnSpawnTab::CreateStatic(&SpawnPivotTab))
+			.SetDisplayName(NSLOCTEXT("UnrealEditor", "ControlRigTempPivotTab", "Control Rig Temp Pivot Tab"))
+			.SetTooltipText(NSLOCTEXT("UnrealEditor", "ControlRigTweenTooltip", "Open the Control Rig Pivot tab."))
+			.SetMenuType(ETabSpawnerMenuType::Hidden)
+			.SetIcon(FSlateIcon(FEditorStyle::GetStyleSetName(), "SkeletonTree.BlendProfile")); //MZ todo replace with correct icon SkeletonTree.LODBone
+		FGlobalTabmanager::Get()->RegisterDefaultTabWindowSize(IControlRigEditorModule::ControlRigTempPivotTab, FVector2D(400, 100));
+
+
 	};
 #endif
 
@@ -297,6 +335,8 @@ void FControlRigEditorModule::ShutdownModule()
 		FGlobalTabmanager::Get()->UnregisterNomadTabSpawner("ControlRigProfiler");
 		FGlobalTabmanager::Get()->UnregisterNomadTabSpawner(IControlRigEditorModule::ControlRigPoseTab);
 		FGlobalTabmanager::Get()->UnregisterNomadTabSpawner(IControlRigEditorModule::ControlRigTweenTab);
+		FGlobalTabmanager::Get()->UnregisterNomadTabSpawner(IControlRigEditorModule::ControlRigSnapperTab);
+		FGlobalTabmanager::Get()->UnregisterNomadTabSpawner(IControlRigEditorModule::ControlRigTempPivotTab);
 	}
 #endif
 
diff --git a/Engine/Plugins/Experimental/ControlRig/Source/ControlRigEditor/Private/EditMode/SControlRigControlViews.cpp b/Engine/Plugins/Experimental/ControlRig/Source/ControlRigEditor/Private/EditMode/SControlRigControlViews.cpp
--- a/Engine/Plugins/Experimental/ControlRig/Source/ControlRigEditor/Private/EditMode/SControlRigControlViews.cpp
+++ b/Engine/Plugins/Experimental/ControlRig/Source/ControlRigEditor/Private/EditMode/SControlRigControlViews.cpp
@@ -206,6 +206,7 @@ void SControlRigPoseView::Construct(const FArguments& InArgs)
 	DetailsViewArgs.bShowPropertyMatrixButton = false;
 	DetailsViewArgs.bUpdatesFromSelection = false;
 	DetailsViewArgs.bLockable = false;
+	DetailsViewArgs.bAllowFavoriteSystem = false;
 	DetailsViewArgs.NameAreaSettings = FDetailsViewArgs::HideNameArea;
 	DetailsViewArgs.NotifyHook = this;
 	DetailsViewArgs.ViewIdentifier = "Create Control Asset";
diff --git a/Engine/Plugins/Experimental/ControlRig/Source/ControlRigEditor/Private/EditMode/SControlRigEditModeTools.cpp b/Engine/Plugins/Experimental/ControlRig/Source/ControlRigEditor/Private/EditMode/SControlRigEditModeTools.cpp
--- a/Engine/Plugins/Experimental/ControlRig/Source/ControlRigEditor/Private/EditMode/SControlRigEditModeTools.cpp
+++ b/Engine/Plugins/Experimental/ControlRig/Source/ControlRigEditor/Private/EditMode/SControlRigEditModeTools.cpp
@@ -532,7 +532,6 @@ void SControlRigEditModeTools::CustomizeToolBarPalette(FToolBarBuilder& ToolBarB
 	ToolBarBuilder.AddSeparator();
 
 	//POSES
-
 	ToolBarBuilder.AddToolBarButton(
 		FExecuteAction::CreateSP(this, &SControlRigEditModeTools::MakePoseDialog),
 		NAME_None,
@@ -552,6 +551,27 @@ void SControlRigEditModeTools::CustomizeToolBarPalette(FToolBarBuilder& ToolBarB
 		FSlateIcon(FEditorStyle::GetStyleSetName(), "SkeletonTree.BlendProfile"), //MZ todo replace with correct icon
 		EUserInterfaceActionType::Button
 	);
+
+	// Snap
+	ToolBarBuilder.AddToolBarButton(
+		FExecuteAction::CreateSP(this, &SControlRigEditModeTools::MakeSnapperDialog),
+		NAME_None,
+		LOCTEXT("Snapper", "Snapper"),
+		LOCTEXT("SnapperTooltip", "Snap Control to another over a set of frames"),
+		FSlateIcon(FEditorStyle::GetStyleSetName(), "Persona.Tabs.AnimSlotManager"), //MZ todo replace with correct icon
+		EUserInterfaceActionType::Button
+	);
+	/*
+	// Pivot
+	ToolBarBuilder.AddToolBarButton(
+		FExecuteAction::CreateSP(this, &SControlRigEditModeTools::MakeTempPivotDialog),
+		NAME_None,
+		LOCTEXT("TempPivot", "Temp Pivot"),
+		LOCTEXT("TempPivotTooltip", "Create a temporary pivot to transform the selected Control"),
+		FSlateIcon(FEditorStyle::GetStyleSetName(), "AnimationEditor.ApplyAnimation"), //MZ todo replace with correct icon
+		EUserInterfaceActionType::Button
+	);
+	*/
 	ToolBarBuilder.AddSeparator();
 }
 
@@ -573,6 +593,28 @@ void SControlRigEditModeTools::MakeTweenDialog()
 	}
 }
 
+
+void SControlRigEditModeTools::MakeSnapperDialog()
+{
+	FControlRigEditMode* ControlRigEditMode = static_cast<FControlRigEditMode*>(GLevelEditorModeTools().GetActiveMode(FControlRigEditMode::ModeName));
+	if (ControlRigEditMode)
+	{
+		FGlobalTabmanager::Get()->TryInvokeTab(IControlRigEditorModule::ControlRigSnapperTab);
+	}
+}
+
+
+
+void SControlRigEditModeTools::MakeTempPivotDialog()
+{
+	FControlRigEditMode* ControlRigEditMode = static_cast<FControlRigEditMode*>(GLevelEditorModeTools().GetActiveMode(FControlRigEditMode::ModeName));
+	if (ControlRigEditMode)
+	{
+		FGlobalTabmanager::Get()->TryInvokeTab(IControlRigEditorModule::ControlRigTempPivotTab);
+	}
+}
+
+
 /* MZ TODO
 void SControlRigEditModeTools::MakeSelectionSetDialog()
 {
diff --git a/Engine/Plugins/Experimental/ControlRig/Source/ControlRigEditor/Private/EditMode/SControlRigEditModeTools.h b/Engine/Plugins/Experimental/ControlRig/Source/ControlRigEditor/Private/EditMode/SControlRigEditModeTools.h
--- a/Engine/Plugins/Experimental/ControlRig/Source/ControlRigEditor/Private/EditMode/SControlRigEditModeTools.h
+++ b/Engine/Plugins/Experimental/ControlRig/Source/ControlRigEditor/Private/EditMode/SControlRigEditModeTools.h
@@ -77,11 +77,10 @@ private:
 private:
 	/** Toolbar functions and windows*/
 
-	/** Make the Pose Dialg*/
 	void MakePoseDialog();
-
-	/** Make the Pose Dialg*/
 	void MakeTweenDialog();
+	void MakeSnapperDialog();
+	void MakeTempPivotDialog();
 
 	//TODO may put back void MakeSelectionSetDialog();
 	//TWeakPtr<SWindow> SelectionSetWindow;
diff --git a/Engine/Plugins/Experimental/ControlRig/Source/ControlRigEditor/Private/EditMode/SControlRigSnapper.cpp b/Engine/Plugins/Experimental/ControlRig/Source/ControlRigEditor/Private/EditMode/SControlRigSnapper.cpp
new file mode 100644
--- /dev/null
+++ b/Engine/Plugins/Experimental/ControlRig/Source/ControlRigEditor/Private/EditMode/SControlRigSnapper.cpp
@@ -0,0 +1,567 @@
+// Copyright Epic Games, Inc. All Rights Reserved.
+
+#include "SControlRigSnapper.h"
+#include "Framework/Notifications/NotificationManager.h"
+#include "Widgets/Notifications/SNotificationList.h"
+#include "Framework/MultiBox/MultiBoxBuilder.h"
+#include "AssetData.h"
+#include "EditorStyleSet.h"
+#include "Widgets/Text/SInlineEditableTextBlock.h"
+#include "EditorStyleSet.h"
+#include "Styling/CoreStyle.h"
+#include "ScopedTransaction.h"
+#include "ControlRig.h"
+#include "UnrealEdGlobals.h"
+#include "ControlRigEditMode.h"
+#include "Tools/ControlRigPose.h"
+#include "EditorModeManager.h"
+#include "ISequencer.h"
+#include "LevelSequence.h"
+#include "UnrealEd/Public/Selection.h"
+#include "Editor.h"
+#include "Tools/ControlRigSnapSettings.h"
+#include "LevelEditor.h"
+#include "Editor/SceneOutliner/Private/SSocketChooser.h"
+
+#define LOCTEXT_NAMESPACE "ControlRigSnapper"
+
+class SComponentPickerPopup : public SCompoundWidget
+{
+public:
+	DECLARE_DELEGATE_OneParam(FOnComponentChosen, FName);
+
+	SLATE_BEGIN_ARGS(SComponentPickerPopup)
+		: _Actor(NULL)
+	{}
+
+	/** An actor with components */
+	SLATE_ARGUMENT(AActor*, Actor)
+
+		/** Called when the text is chosen. */
+		SLATE_EVENT(FOnComponentChosen, OnComponentChosen)
+
+		SLATE_END_ARGS()
+
+		/** Delegate to call when component is selected */
+		FOnComponentChosen OnComponentChosen;
+
+	/** List of tag names selected in the tag containers*/
+	TArray< TSharedPtr<FName> > ComponentNames;
+
+private:
+	TSharedRef<ITableRow> MakeListViewWidget(TSharedPtr<FName> InItem, const TSharedRef<STableViewBase>& OwnerTable)
+	{
+		return SNew(STableRow< TSharedPtr<FName> >, OwnerTable)
+			[
+				SNew(STextBlock).Text(FText::FromName(*InItem.Get()))
+			];
+	}
+
+	void OnComponentSelected(TSharedPtr<FName> InItem, ESelectInfo::Type InSelectInfo)
+	{
+		FSlateApplication::Get().DismissAllMenus();
+
+		if (OnComponentChosen.IsBound())
+		{
+			OnComponentChosen.Execute(*InItem.Get());
+		}
+	}
+
+public:
+	void Construct(const FArguments& InArgs)
+	{
+		OnComponentChosen = InArgs._OnComponentChosen;
+		AActor* Actor = InArgs._Actor;
+
+		TInlineComponentArray<USceneComponent*> Components(Actor);
+
+		ComponentNames.Empty();
+		for (USceneComponent* Component : Components)
+		{
+			if (Component->HasAnySockets())
+			{
+				ComponentNames.Add(MakeShareable(new FName(Component->GetFName())));
+			}
+		}
+
+		// Then make widget
+		this->ChildSlot
+			[
+				SNew(SBorder)
+				.BorderImage(FEditorStyle::GetBrush(TEXT("Menu.Background")))
+			.Padding(5)
+			.Content()
+			[
+				SNew(SVerticalBox)
+				+ SVerticalBox::Slot()
+			.AutoHeight()
+			.Padding(0.0f, 1.0f)
+			[
+				SNew(STextBlock)
+				.Font(FEditorStyle::GetFontStyle(TEXT("SocketChooser.TitleFont")))
+			.Text(NSLOCTEXT("ComponentChooser", "ChooseComponentLabel", "Choose Component"))
+			]
+		+ SVerticalBox::Slot()
+			.AutoHeight()
+			.MaxHeight(512)
+			[
+				SNew(SBox)
+				.WidthOverride(256)
+			.Content()
+			[
+				SNew(SListView< TSharedPtr<FName> >)
+				.ListItemsSource(&ComponentNames)
+			.OnGenerateRow(this, &SComponentPickerPopup::MakeListViewWidget)
+			.OnSelectionChanged(this, &SComponentPickerPopup::OnComponentSelected)
+			]
+			]
+			]
+			];
+	}
+};
+
+void SControlRigSnapper::Construct(const FArguments& InArgs)
+{
+	ClearActors();
+	SetStartEndFrames();
+
+	//for snapper settings
+	UControlRigSnapSettings* SnapperSettings = GetMutableDefault<UControlRigSnapSettings>();
+	FPropertyEditorModule& PropertyEditor = FModuleManager::LoadModuleChecked<FPropertyEditorModule>("PropertyEditor");
+
+	FDetailsViewArgs DetailsViewArgs;
+	DetailsViewArgs.bShowOptions = false;
+	DetailsViewArgs.bAllowSearch = false;
+	DetailsViewArgs.bShowPropertyMatrixButton = false;
+	DetailsViewArgs.bUpdatesFromSelection = false;
+	DetailsViewArgs.bLockable = false;
+	DetailsViewArgs.bAllowFavoriteSystem = false;
+	DetailsViewArgs.NameAreaSettings = FDetailsViewArgs::HideNameArea;
+	DetailsViewArgs.ViewIdentifier = "ControlRigSnapper";
+
+	SnapperDetailsView = PropertyEditor.CreateDetailView(DetailsViewArgs);
+	SnapperDetailsView->SetObject(SnapperSettings);
+
+
+	ChildSlot
+		[
+			SNew(SBorder)
+			.HAlign(HAlign_Fill)
+		.VAlign(VAlign_Fill)
+		.Padding(FMargin(10.0,5.0,10.0,5.0))
+		[
+			SNew(SVerticalBox)
+			+ SVerticalBox::Slot()
+			.AutoHeight()
+			.HAlign(HAlign_Fill)
+			[
+				SNew(SHorizontalBox)
+				+ SHorizontalBox::Slot()
+					.AutoWidth()
+					.Padding(10.f)
+					.VAlign(VAlign_Center)
+					[
+						SNew(SVerticalBox)
+						+ SVerticalBox::Slot()
+						.AutoHeight()
+						.HAlign(HAlign_Center)
+						[
+
+							SNew(SBox)
+							.Padding(0.0f)
+							[
+								SNew(STextBlock)
+								.Text(LOCTEXT("Children", "Children"))
+							]
+
+						]
+						+ SVerticalBox::Slot()
+						.AutoHeight()
+						.HAlign(HAlign_Fill)
+						[
+							SNew(SButton)
+							.HAlign(HAlign_Center)
+							.ContentPadding(FMargin(10.0f, 2.0f, 10.0f, 2.0f))
+							.OnClicked(this, &SControlRigSnapper::OnActorToSnapClicked)
+							[
+								SNew(STextBlock)
+								.ToolTipText(LOCTEXT("ActorToSnapTooltip","Select child object(s) you want to snap over the interval range"))
+								.Text_Lambda([this]()
+									{
+										return GetActorToSnapText();
+									})
+							]
+						]
+
+					]
+					+ SHorizontalBox::Slot()
+						.AutoWidth()
+						.VAlign(VAlign_Center)
+						.Padding(10.f)
+						[
+							SNew(SVerticalBox)
+							+ SVerticalBox::Slot()
+							.AutoHeight()
+							.HAlign(HAlign_Center)
+							[
+
+								SNew(SBox)
+								.Padding(0.0f)
+								[
+									SNew(STextBlock)
+									.Text(LOCTEXT("Parent", "Parent"))
+								]
+
+							]
+							+ SVerticalBox::Slot()
+							.AutoHeight()
+							.HAlign(HAlign_Fill)
+							[
+								SNew(SButton)
+									.HAlign(HAlign_Center)
+									.ContentPadding(FMargin(10.0f, 2.0f, 10.0f, 2.0f))
+									.OnClicked(this, &SControlRigSnapper::OnParentToSnapToClicked)
+									[
+										SNew(STextBlock)
+										.ToolTipText(LOCTEXT("ParentToSnapTooltip", "Select parent object you want children to snap to. If one is not selected it will snap to World Location at the start."))
+										.Text_Lambda([this]()
+											{
+												return GetParentToSnapText();
+											})
+									]
+							]
+
+						]
+			 ]
+
+			+ SVerticalBox::Slot()
+				.AutoHeight()
+				.HAlign(HAlign_Fill)
+				[
+					SNew(SHorizontalBox)
+					+ SHorizontalBox::Slot()
+						.AutoWidth()
+						.Padding(10.f)
+						.VAlign(VAlign_Center)
+						[
+							SNew(SButton)
+							.HAlign(HAlign_Center)
+							.ContentPadding(FMargin(10.0f, 2.0f, 10.0f, 2.0f))
+							.OnClicked(this, &SControlRigSnapper::OnStartFrameClicked)
+							[
+								SNew(SEditableTextBox)
+								.ToolTipText(LOCTEXT("GetStartFrameTooltip", "Set first frame to snap"))
+								.Text_Lambda([this]()
+									{
+										return GetStartFrameToSnapText();
+									})
+							]
+
+						]
+						+ SHorizontalBox::Slot()
+							.AutoWidth()
+							.Padding(10.f)
+							.VAlign(VAlign_Center)
+						[
+
+							SNew(SButton)
+								.HAlign(HAlign_Center)
+								.ContentPadding(FMargin(10.0f, 2.0f, 10.0f, 2.0f))
+								.OnClicked(this, &SControlRigSnapper::OnEndFrameClicked)
+								[
+									SNew(SEditableTextBox)
+									.ToolTipText(LOCTEXT("GetEndFrameTooltip", "Set end frame to snap"))
+									.Text_Lambda([this]()
+										{
+											return GetEndFrameToSnapText();
+										})
+								]
+
+						]
+				 ]
+
+				+ SVerticalBox::Slot()
+				.AutoHeight()
+				.HAlign(HAlign_Fill)
+				[
+					SnapperDetailsView.ToSharedRef()
+				]
+				+ SVerticalBox::Slot()
+				.AutoHeight()
+				.HAlign(HAlign_Fill)
+				.VAlign(VAlign_Bottom)
+				[
+					SNew(SHorizontalBox)
+					+SHorizontalBox::Slot()
+					.Padding(5.f)
+					.HAlign(HAlign_Center)
+					.VAlign(VAlign_Center)
+					[
+						SNew(SButton)
+						.HAlign(HAlign_Fill)
+						.ContentPadding(FMargin(10.0f, 2.0f, 10.0f, 2.0f))
+						.OnClicked(this, &SControlRigSnapper::OnSnapAnimationClicked)
+						[
+							SNew(STextBlock)
+							.Text(LOCTEXT("SnapAnimation", "Snap Animation"))
+						]
+					]
+				]
+			]
+		];
+}
+
+FReply SControlRigSnapper::OnActorToSnapClicked()
+{
+	ActorToSnap = GetSelection(true);
+	return FReply::Handled();
+}
+
+FReply SControlRigSnapper::OnParentToSnapToClicked()
+{
+	ParentToSnap = GetSelection(false);
+	return FReply::Handled();
+}
+
+FText SControlRigSnapper::GetActorToSnapText()
+{
+	if (ActorToSnap.IsValid())
+	{
+		return (ActorToSnap.GetName());
+	}
+
+	return LOCTEXT("SelectActor", "Select Actor");
+}
+
+FText SControlRigSnapper::GetParentToSnapText()
+{
+	if (ParentToSnap.IsValid())
+	{
+		return (ParentToSnap.GetName());
+	}
+
+	return LOCTEXT("World", "World");
+}
+
+
+FReply SControlRigSnapper::OnStartFrameClicked()
+{
+	TWeakPtr<ISequencer> Sequencer = Snapper.GetSequencer();
+	if (Sequencer.IsValid())
+	{
+		const FFrameRate TickResolution = Sequencer.Pin()->GetFocusedTickResolution();
+		const FFrameTime FrameTime = Sequencer.Pin()->GetLocalTime().ConvertTo(TickResolution);
+		StartFrame = FrameTime.GetFrame();
+	}
+	return FReply::Handled();
+}
+
+FReply SControlRigSnapper::OnEndFrameClicked()
+{
+	TWeakPtr<ISequencer> Sequencer = Snapper.GetSequencer();
+	if (Sequencer.IsValid())
+	{
+		const FFrameRate TickResolution = Sequencer.Pin()->GetFocusedTickResolution();
+		const FFrameTime FrameTime = Sequencer.Pin()->GetLocalTime().ConvertTo(TickResolution);
+		EndFrame = FrameTime.GetFrame();
+	}
+	return FReply::Handled();
+}
+
+FReply SControlRigSnapper::OnSnapAnimationClicked()
+{
+	Snapper.SnapIt(StartFrame, EndFrame, ActorToSnap, ParentToSnap);
+	return FReply::Handled();
+}
+
+FText SControlRigSnapper::GetStartFrameToSnapText()
+{
+	FText Value;
+	TWeakPtr<ISequencer> Sequencer = Snapper.GetSequencer();
+	if (Sequencer.IsValid() && Sequencer.Pin()->GetFocusedMovieSceneSequence())
+	{
+		Value = FText::FromString(Sequencer.Pin()->GetNumericTypeInterface()->ToString(StartFrame.Value));
+	}
+	return Value;
+}
+
+FText SControlRigSnapper::GetEndFrameToSnapText()
+{
+	FText Value;
+	TWeakPtr<ISequencer> Sequencer = Snapper.GetSequencer();
+	if (Sequencer.IsValid() && Sequencer.Pin()->GetFocusedMovieSceneSequence())
+	{
+		Value = FText::FromString(Sequencer.Pin()->GetNumericTypeInterface()->ToString(EndFrame.Value));
+	}
+	return Value;
+}
+
+void SControlRigSnapper::ClearActors()
+{
+	ActorToSnap.Clear();
+	ParentToSnap.Clear();
+}
+
+void SControlRigSnapper::SetStartEndFrames()
+{
+	TWeakPtr<ISequencer> Sequencer = Snapper.GetSequencer();
+	if (Sequencer.IsValid()  && Sequencer.Pin()->GetFocusedMovieSceneSequence())
+	{
+		UMovieScene* MovieScene = Sequencer.Pin()->GetFocusedMovieSceneSequence()->GetMovieScene();
+		StartFrame = MovieScene->GetPlaybackRange().GetLowerBoundValue();
+		EndFrame = MovieScene->GetPlaybackRange().GetUpperBoundValue();
+	}
+}
+
+FControlRigSnapperSelection SControlRigSnapper::GetSelection(bool bGetAll) 
+{
+	FControlRigSnapperSelection Selection;
+	UControlRig* ControlRig = GetControlRig();
+	if (ControlRig)
+	{
+		TArray<FName> SelectedControls = ControlRig->CurrentControlSelection();
+		if (SelectedControls.Num() > 0)
+		{
+			FControlRigForWorldTransforms ControlRigAndSelection;
+			ControlRigAndSelection.ControlRig = ControlRig;
+			ControlRigAndSelection.ControlNames = SelectedControls;
+			Selection.ControlRigs.Add(ControlRigAndSelection);
+			if (bGetAll == false)
+			{
+				return Selection;
+			}
+		}
+	}
+	USelection* SelectedActors = GEditor->GetSelectedActors();
+	for (FSelectionIterator Iter(*SelectedActors); Iter; ++Iter)
+	{
+		AActor* Actor = Cast<AActor>(*Iter);
+		if (Actor)
+		{
+			FActorForWorldTransforms ActorSelection;
+			ActorSelection.Actor = Actor;
+			Selection.Actors.Add(ActorSelection);
+			if (bGetAll == false)
+			{
+				ActorParentPicked(ActorSelection);
+				return Selection;
+			}
+		}
+	}
+	return Selection;
+}
+
+UControlRig* SControlRigSnapper::GetControlRig() const
+{
+	if (FControlRigEditMode* EditMode = static_cast<FControlRigEditMode*>(GLevelEditorModeTools().GetActiveMode(FControlRigEditMode::ModeName)))
+	{
+		return EditMode->GetControlRig(true);
+	}
+	return nullptr;
+}
+
+
+void SControlRigSnapper::ActorParentSocketPicked(const FName SocketPicked, FActorForWorldTransforms Selection) 
+{
+	ParentToSnap.Actors.SetNum(0);
+	Selection.SocketName = SocketPicked;
+	ParentToSnap.Actors.Add(Selection);
+	//ParentToSnap
+}
+
+void SControlRigSnapper::ActorParentPicked(FActorForWorldTransforms Selection) 
+{
+	TArray<USceneComponent*> ComponentsWithSockets;
+	if (Selection.Actor.IsValid())
+	{
+		TInlineComponentArray<USceneComponent*> Components(Selection.Actor.Get());
+
+		for (USceneComponent* Component : Components)
+		{
+			if (Component->HasAnySockets())
+			{
+				ComponentsWithSockets.Add(Component);
+			}
+		}
+	}
+
+	if (ComponentsWithSockets.Num() == 0)
+	{
+		FSlateApplication::Get().DismissAllMenus();
+		ActorParentSocketPicked(NAME_None, Selection);
+		return;
+	}
+
+	FLevelEditorModule& LevelEditorModule = FModuleManager::GetModuleChecked<FLevelEditorModule>("LevelEditor");
+	TSharedPtr< ILevelEditor > LevelEditor = LevelEditorModule.GetFirstLevelEditor();
+
+	TSharedPtr<SWidget> MenuWidget;
+	if (ComponentsWithSockets.Num() > 1)
+	{
+		MenuWidget =
+			SNew(SComponentPickerPopup)
+			.Actor(Selection.Actor.Get())
+			.OnComponentChosen(this, &SControlRigSnapper::ActorParentComponentPicked,Selection);
+
+		// Create as context menu
+		FSlateApplication::Get().PushMenu(
+			LevelEditor.ToSharedRef(),
+			FWidgetPath(),
+			MenuWidget.ToSharedRef(),
+			FSlateApplication::Get().GetCursorPos(),
+			FPopupTransitionEffect(FPopupTransitionEffect::ContextMenu)
+		);
+	}
+	else
+	{
+		ActorParentComponentPicked(ComponentsWithSockets[0]->GetFName(),Selection);
+	}
+}
+
+
+void SControlRigSnapper::ActorParentComponentPicked(FName ComponentName, FActorForWorldTransforms Selection) 
+{
+	USceneComponent* ComponentWithSockets = nullptr;
+	if (Selection.Actor.IsValid())
+	{
+		AActor* Actor = Selection.Actor.Get();
+		TInlineComponentArray<USceneComponent*> Components(Actor);
+
+		for (USceneComponent* Component : Components)
+		{
+			if (Component->GetFName() == ComponentName)
+			{
+				ComponentWithSockets = Component;
+				break;
+			}
+		}
+	}
+
+	if (ComponentWithSockets == nullptr)
+	{
+		return;
+	}
+	Selection.Component = ComponentWithSockets;
+	FLevelEditorModule& LevelEditorModule = FModuleManager::GetModuleChecked<FLevelEditorModule>("LevelEditor");
+	TSharedPtr< ILevelEditor > LevelEditor = LevelEditorModule.GetFirstLevelEditor();
+
+	TSharedPtr<SWidget> MenuWidget;
+
+	MenuWidget =
+		SNew(SSocketChooserPopup)
+		.SceneComponent(ComponentWithSockets)
+		.OnSocketChosen(this, &SControlRigSnapper::ActorParentSocketPicked, Selection);
+
+	// Create as context menu
+	FSlateApplication::Get().PushMenu(
+		LevelEditor.ToSharedRef(),
+		FWidgetPath(),
+		MenuWidget.ToSharedRef(),
+		FSlateApplication::Get().GetCursorPos(),
+		FPopupTransitionEffect(FPopupTransitionEffect::ContextMenu)
+	);
+}
+
+
+
+#undef LOCTEXT_NAMESPACE
diff --git a/Engine/Plugins/Experimental/ControlRig/Source/ControlRigEditor/Private/EditMode/SControlRigSnapper.h b/Engine/Plugins/Experimental/ControlRig/Source/ControlRigEditor/Private/EditMode/SControlRigSnapper.h
new file mode 100644
--- /dev/null
+++ b/Engine/Plugins/Experimental/ControlRig/Source/ControlRigEditor/Private/EditMode/SControlRigSnapper.h
@@ -0,0 +1,75 @@
+// Copyright Epic Games, Inc. All Rights Reserved.
+/**
+* Hold the View for the Snapper Widget
+*/
+#pragma once
+
+#include "CoreMinimal.h"
+#include "Widgets/SWidget.h"
+#include "Widgets/DeclarativeSyntaxSupport.h"
+#include "Widgets/SCompoundWidget.h"
+#include "Widgets/Input/SCheckBox.h"
+#include "Widgets/Input/SNumericEntryBox.h"
+#include "Widgets/Views/SListView.h"
+#include "Widgets/Text/STextBlock.h"
+#include "Widgets/Views/STableRow.h"
+#include "Tools/ControlRigSnapper.h"
+#include "Misc/FrameNumber.h"
+
+class UControlRig;
+class ISequencer;
+class AActor;
+
+
+class SControlRigSnapper : public SCompoundWidget
+{
+	SLATE_BEGIN_ARGS(SControlRigSnapper) {}
+
+	SLATE_END_ARGS()
+	~SControlRigSnapper()
+	{
+	}
+
+	void Construct(const FArguments& InArgs);
+
+
+private:
+
+	UControlRig* GetControlRig() const;
+
+	/*
+	* Delegates and Helpers
+	*/
+	void ActorParentPicked(FActorForWorldTransforms Selection);
+	void ActorParentSocketPicked(const FName SocketName, FActorForWorldTransforms Selection);
+	void ActorParentComponentPicked(FName ComponentName, FActorForWorldTransforms Selection);
+
+	FReply OnActorToSnapClicked();
+	FReply OnParentToSnapToClicked();
+	FText GetActorToSnapText();
+	FText GetParentToSnapText();
+
+	FReply OnStartFrameClicked();
+	FReply OnEndFrameClicked();
+	FText GetStartFrameToSnapText();
+	FText GetEndFrameToSnapText();
+
+	void ClearActors();
+	void SetStartEndFrames();
+	FReply OnSnapAnimationClicked();
+
+	FControlRigSnapperSelection GetSelection(bool bGetAll);
+
+	//FControlRigSnapper  ControlRigSnapper;
+	TSharedPtr<IDetailsView> SnapperDetailsView;
+
+	FControlRigSnapperSelection ActorToSnap;
+	FControlRigSnapperSelection ParentToSnap;
+
+	FFrameNumber StartFrame;
+	FFrameNumber EndFrame;
+
+	FControlRigSnapper Snapper;
+
+};
+
diff --git a/Engine/Plugins/Experimental/ControlRig/Source/ControlRigEditor/Private/Editor/ControlRigEditor.cpp b/Engine/Plugins/Experimental/ControlRig/Source/ControlRigEditor/Private/Editor/ControlRigEditor.cpp
--- a/Engine/Plugins/Experimental/ControlRig/Source/ControlRigEditor/Private/Editor/ControlRigEditor.cpp
+++ b/Engine/Plugins/Experimental/ControlRig/Source/ControlRigEditor/Private/Editor/ControlRigEditor.cpp
@@ -22,7 +22,6 @@
 #include "ControlRigBlueprintGeneratedClass.h"
 #include "AnimCustomInstanceHelper.h"
 #include "Sequencer/ControlRigLayerInstance.h"
-#include "Sequencer/ControlRigSequencerAnimInstance.h"
 #include "Animation/DebugSkelMeshComponent.h"
 #include "IPersonaPreviewScene.h"
 #include "Persona/Private/AnimationEditorPreviewScene.h"
diff --git a/Engine/Plugins/Experimental/ControlRig/Source/ControlRigEditor/Private/Tools/ControlRigSnapper.cpp b/Engine/Plugins/Experimental/ControlRig/Source/ControlRigEditor/Private/Tools/ControlRigSnapper.cpp
new file mode 100644
--- /dev/null
+++ b/Engine/Plugins/Experimental/ControlRig/Source/ControlRigEditor/Private/Tools/ControlRigSnapper.cpp
@@ -0,0 +1,598 @@
+// Copyright Epic Games, Inc. All Rights Reserved.
+
+#include "Tools/ControlRigSnapper.h"
+#include "Tools/ControlRigTweener.h" //remove
+
+#include "Channels/MovieSceneFloatChannel.h"
+#include "ControlRig.h"
+#include "ISequencer.h"
+#include "MovieSceneSequence.h"
+#include "MovieScene.h"
+#include "Rigs/RigControlHierarchy.h"
+#include "Sequencer/MovieSceneControlRigParameterTrack.h"
+#include "Sequencer/MovieSceneControlRigParameterSection.h"
+#include "IKeyArea.h"
+#include "LevelSequence.h"
+#include "LevelSequenceEditor/Private/LevelSequenceEditorBlueprintLibrary.h"
+#include "ILevelSequenceEditorToolkit.h"
+#include "Tools/ControlRigSnapper.h"
+#include "EntitySystem/Interrogation/MovieSceneInterrogationLinker.h"
+#include "ControlRigObjectBinding.h"
+#include "Sections/MovieScene3DTransformSection.h"
+#include "Tracks/MovieScene3DTransformTrack.h"
+#include "Channels/MovieSceneChannelProxy.h"
+#include "MatineeImportTools.h"
+#include "Tools/ControlRigSnapSettings.h"
+#include "MovieSceneToolsModule.h"
+#include "Components/SkeletalMeshComponent.h"
+#include "ScopedTransaction.h"
+#include "MovieSceneToolHelpers.h"
+
+#define LOCTEXT_NAMESPACE "ControlRigSnapper"
+
+
+FText FControlRigSnapperSelection::GetName() const
+{
+	FText Name;
+
+	int32 Num = NumSelected();
+	if (Num == 0)
+	{
+		Name = LOCTEXT("None", "None");
+	}
+	else if (Num == 1)
+	{
+		for (const FActorForWorldTransforms& Actor : Actors)
+		{
+			if (Actor.Actor.IsValid())
+			{
+				FString ActorLabel = Actor.Actor->GetActorLabel();
+				if (Actor.SocketName != NAME_None)
+				{
+					FString SocketString = Actor.SocketName.ToString();
+					ActorLabel += (FString(":") + SocketString);
+				}
+				Name = FText::FromString(ActorLabel);
+				break;
+			}
+		}
+		for (const FControlRigForWorldTransforms& Selection : ControlRigs)
+		{
+			if (Selection.ControlRig.IsValid())
+			{
+				if (Selection.ControlNames.Num() > 0)
+				{
+					FName ControlName = Selection.ControlNames[0];
+					Name = FText::FromString(ControlName.ToString());
+					break;
+				}
+			}
+		}
+	}
+	else
+	{
+		Name = LOCTEXT("Multiple", "--Multiple--");
+	}
+	return Name;
+}
+
+int32 FControlRigSnapperSelection::NumSelected() const
+{
+	int32 Selected = 0;
+	for (const FActorForWorldTransforms& Actor : Actors)
+	{
+		if (Actor.Actor.IsValid())
+		{
+			++Selected;
+		}
+	}
+	for (const FControlRigForWorldTransforms& Selection : ControlRigs)
+	{
+		if (Selection.ControlRig.IsValid())
+		{
+			Selected += Selection.ControlNames.Num();
+		}
+	}
+	return Selected;
+}
+
+TWeakPtr<ISequencer> FControlRigSnapper::GetSequencer()
+{
+	TWeakPtr<ISequencer> WeakSequencer = nullptr;
+	ULevelSequence* LevelSequence = ULevelSequenceEditorBlueprintLibrary::GetCurrentLevelSequence();
+	if (LevelSequence)
+	{
+		IAssetEditorInstance* AssetEditor = GEditor->GetEditorSubsystem<UAssetEditorSubsystem>()->FindEditorForAsset(LevelSequence, false);
+		ILevelSequenceEditorToolkit* LevelSequenceEditor = static_cast<ILevelSequenceEditorToolkit*>(AssetEditor);
+		WeakSequencer = LevelSequenceEditor ? LevelSequenceEditor->GetSequencer() : nullptr;
+	}
+	return WeakSequencer;
+}
+
+static bool GetControlRigControlTransforms(ISequencer* Sequencer,  UControlRig* ControlRig, const FName& ControlName, 
+	const TArray<FFrameNumber> &Frames, const TArray<FTransform>& ParentTransforms,TArray<FTransform>& OutTransforms)
+{
+	if (Frames.Num() > ParentTransforms.Num())
+	{
+		UE_LOG(LogControlRig, Error, TEXT("Number of Frames %d to Snap greater  than Parent Frames %d"), Frames.Num(), ParentTransforms.Num());
+		return false;
+	}
+	if (ControlRig->FindControl(ControlName) == nullptr)
+	{
+		UE_LOG(LogControlRig, Error, TEXT("Can not find Control %s"), *(ControlName.ToString()));
+		return false;
+	}
+	if (UMovieScene* MovieScene = Sequencer->GetFocusedMovieSceneSequence()->GetMovieScene())
+	{
+
+		FMovieSceneSequenceIDRef Template = Sequencer->GetFocusedTemplateID();
+		FMovieSceneSequenceTransform RootToLocalTransform;
+
+		FFrameRate TickResolution = MovieScene->GetTickResolution();
+		FFrameRate DisplayRate = MovieScene->GetDisplayRate();
+
+		OutTransforms.SetNum(Frames.Num());
+		for (int32 Index = 0; Index <Frames.Num(); ++Index)
+		{
+			const FFrameNumber& FrameNumber = Frames[Index];
+			FFrameTime GlobalTime(FrameNumber);
+
+			FMovieSceneContext Context = FMovieSceneContext(FMovieSceneEvaluationRange(GlobalTime, TickResolution), Sequencer->GetPlaybackStatus()).SetHasJumped(true);
+
+			Sequencer->GetEvaluationTemplate().Evaluate(Context, *Sequencer);
+			ControlRig->Evaluate_AnyThread();
+			OutTransforms[Index] =  ControlRig->GetControlGlobalTransform(ControlName) * ParentTransforms[Index];
+		}
+	}
+	return true;
+}
+
+//Matinee version of this doesn't actually set the key it only adds..sigh
+static void SetOrAddKey(TMovieSceneChannelData<FMovieSceneFloatValue>& ChannelData, FFrameNumber Time, float Value)
+{
+	int32 ExistingIndex = ChannelData.FindKey(Time);
+	if (ExistingIndex != INDEX_NONE)
+	{
+		FMovieSceneFloatValue& FloatValue = ChannelData.GetValues()[ExistingIndex]; //-V758
+		FloatValue.Value = Value;
+	}
+	else
+	{
+		FMovieSceneFloatValue NewKey(Value);
+		ERichCurveTangentWeightMode WeightedMode = RCTWM_WeightedNone;
+		NewKey.InterpMode = ERichCurveInterpMode::RCIM_Cubic;
+		NewKey.TangentMode = ERichCurveTangentMode::RCTM_Auto;
+		NewKey.Tangent.ArriveTangent = 0.0f;
+		NewKey.Tangent.LeaveTangent = 0.0f;
+		NewKey.Tangent.TangentWeightMode = WeightedMode;
+		NewKey.Tangent.ArriveTangentWeight = 0.0f;
+		NewKey.Tangent.LeaveTangentWeight = 0.0f;
+		ChannelData.AddKey(Time, NewKey);
+	}
+}
+
+struct FGuidAndActor
+{
+	FGuidAndActor(FGuid InGuid, AActor* InActor) : Guid(InGuid), Actor(InActor) {};
+	FGuid Guid;
+	AActor* Actor;
+	
+	bool SetTransform(ISequencer* Sequencer, const TArray<FFrameNumber>& Frames,
+		const TArray<FTransform>& WorldTransformsToSnapTo)
+	{
+		if (!Sequencer || !Sequencer->GetFocusedMovieSceneSequence())
+		{
+			return false;
+		}
+		UMovieScene* MovieScene = Sequencer->GetFocusedMovieSceneSequence()->GetMovieScene();
+
+		UMovieScene3DTransformTrack* TransformTrack = MovieScene->FindTrack<UMovieScene3DTransformTrack>(Guid);
+		if (!TransformTrack)
+		{
+			MovieScene->Modify();
+			TransformTrack = MovieScene->AddTrack<UMovieScene3DTransformTrack>(Guid);
+		}
+		TransformTrack->Modify();
+
+		bool bSectionAdded = false;
+		UMovieScene3DTransformSection* TransformSection = Cast<UMovieScene3DTransformSection>(TransformTrack->FindOrAddSection(0, bSectionAdded));
+		if (!TransformSection)
+		{
+			return false;
+		}
+
+		TransformSection->Modify();
+
+		if (bSectionAdded)
+		{
+			TransformSection->SetRange(TRange<FFrameNumber>::All());
+		}
+
+		TArray<FTransform> ParentWorldTransforms; 
+		AActor* ParentActor = Actor->GetAttachParentActor();
+		if (ParentActor)
+		{
+			FActorForWorldTransforms ActorSelection;
+			ActorSelection.SocketName = Actor->GetAttachParentSocketName();
+			MovieSceneToolHelpers::GetActorWorldTransforms(Sequencer, ActorSelection, Frames, ParentWorldTransforms);
+		}
+		else
+		{
+			ParentWorldTransforms.SetNumUninitialized(Frames.Num());
+			for (FTransform& Transform : ParentWorldTransforms)
+			{
+				Transform = FTransform::Identity;
+			}
+		}
+		const UControlRigSnapSettings* SnapSettings = GetDefault<UControlRigSnapSettings>();
+
+		TArrayView<FMovieSceneFloatChannel*> Channels = TransformSection->GetChannelProxy().GetChannels<FMovieSceneFloatChannel>();
+		for (int32 Index = 0; Index < Frames.Num(); ++Index)
+		{
+			const FFrameNumber& Frame = Frames[Index];
+			FTransform ParentTransform = ParentWorldTransforms[Index];
+			FTransform WorldTransform = WorldTransformsToSnapTo[Index];
+			FTransform LocalTransform = WorldTransform.GetRelativeTransform(ParentTransform);
+			FVector Location = LocalTransform.GetLocation();
+			FRotator Rotation = LocalTransform.GetRotation().Rotator();
+			FVector Scale3D = LocalTransform.GetScale3D();
+			if (Index == 0) 
+			{
+				if (SnapSettings->bSnapPosition)
+				{
+					if (!Channels[0]->GetDefault().IsSet())
+					{
+						Channels[0]->SetDefault(Location.X);
+					}
+					if (!Channels[1]->GetDefault().IsSet())
+					{
+						Channels[1]->SetDefault(Location.Y);
+					}
+					if (!Channels[2]->GetDefault().IsSet())
+					{
+						Channels[2]->SetDefault(Location.Z);
+					}
+				}
+				if (SnapSettings->bSnapRotation)
+				{
+					if (!Channels[3]->GetDefault().IsSet())
+					{
+						Channels[3]->SetDefault(Rotation.Yaw);
+					}
+					if (!Channels[4]->GetDefault().IsSet())
+					{
+						Channels[4]->SetDefault(Rotation.Pitch);
+					}
+					if (!Channels[5]->GetDefault().IsSet())
+					{
+						Channels[5]->SetDefault(Rotation.Roll);
+					}
+				}
+				if (SnapSettings->bSnapScale)
+				{
+					if (!Channels[6]->GetDefault().IsSet())
+					{
+						Channels[6]->SetDefault(Scale3D.X);
+					}
+					if (!Channels[7]->GetDefault().IsSet())
+					{
+						Channels[7]->SetDefault(Scale3D.Y);
+					}
+					if (!Channels[8]->GetDefault().IsSet())
+					{
+						Channels[8]->SetDefault(Scale3D.Z);
+					}
+				}
+			}
+
+			if (SnapSettings->bSnapPosition)
+			{
+				TMovieSceneChannelData<FMovieSceneFloatValue> ChannelData = Channels[0]->GetData();
+				SetOrAddKey(ChannelData, Frame, Location.X);
+				ChannelData = Channels[1]->GetData();
+				SetOrAddKey(ChannelData, Frame, Location.Y);
+				ChannelData = Channels[2]->GetData();
+				SetOrAddKey(ChannelData, Frame, Location.Z);
+			}
+			if (SnapSettings->bSnapRotation)
+			{
+				//todo winding
+				TMovieSceneChannelData<FMovieSceneFloatValue> ChannelData = Channels[3]->GetData();
+				SetOrAddKey(ChannelData, Frame, Rotation.Yaw);
+				ChannelData = Channels[4]->GetData();
+				SetOrAddKey(ChannelData, Frame, Rotation.Pitch);
+				ChannelData = Channels[5]->GetData();
+				SetOrAddKey(ChannelData, Frame, Rotation.Roll);
+			}
+			if (SnapSettings->bSnapScale)
+			{
+				TMovieSceneChannelData<FMovieSceneFloatValue> ChannelData = Channels[6]->GetData();
+				SetOrAddKey(ChannelData, Frame, Scale3D.X);
+				ChannelData = Channels[7]->GetData();
+				SetOrAddKey(ChannelData, Frame, Scale3D.X);
+				ChannelData = Channels[8]->GetData();
+				SetOrAddKey(ChannelData, Frame, Scale3D.X);
+			}		
+		}
+		//now we need to set auto tangents
+		if (SnapSettings->bSnapPosition)
+		{
+			Channels[0]->AutoSetTangents();
+			Channels[1]->AutoSetTangents();
+			Channels[2]->AutoSetTangents();		
+		}
+		if (SnapSettings->bSnapRotation)
+		{
+			Channels[3]->AutoSetTangents();
+			Channels[4]->AutoSetTangents();
+			Channels[5]->AutoSetTangents();
+		}
+		if (SnapSettings->bSnapScale)
+		{
+			Channels[6]->AutoSetTangents();
+			Channels[7]->AutoSetTangents();
+			Channels[8]->AutoSetTangents();
+		}
+		Sequencer->NotifyMovieSceneDataChanged(EMovieSceneDataChangeType::MovieSceneStructureItemsChanged);
+
+		return true;
+	}
+};
+
+static void CalculateFramesToSnap(ISequencer* InSequencer, UMovieScene* MovieScene, FFrameNumber StartFrame, FFrameNumber EndFrame, TArray<FFrameNumber>& OutFrames)
+{
+	FFrameRate TickResolution = MovieScene->GetTickResolution();
+	FFrameRate DisplayResolution = MovieScene->GetDisplayRate();
+
+	FFrameNumber StartTimeInDisplay = FFrameRate::TransformTime(FFrameTime(StartFrame), TickResolution, DisplayResolution).FloorToFrame();
+	FFrameNumber EndTimeInDisplay = FFrameRate::TransformTime(FFrameTime(EndFrame), TickResolution, DisplayResolution).CeilToFrame();
+	for (FFrameNumber DisplayFrameNumber = StartTimeInDisplay; DisplayFrameNumber <= EndTimeInDisplay; ++DisplayFrameNumber)
+	{
+		FFrameNumber TickFrameNumber = FFrameRate::TransformTime(FFrameTime(DisplayFrameNumber), DisplayResolution, TickResolution).FrameNumber;
+		OutFrames.Add(TickFrameNumber);
+	}
+}
+
+//returns true if world is calculated, false if there are no parents
+static bool CalculateWorldTransformsFromParents(ISequencer* Sequencer, const FControlRigSnapperSelection& ParentToSnap,
+	const TArray<FFrameNumber>& Frames, TArray<FTransform>& OutParentWorldTransforms)
+{
+	//just do first for now but may do average later.
+	for (const FActorForWorldTransforms& ActorSelection : ParentToSnap.Actors)
+	{
+		if (ActorSelection.Actor.IsValid())
+		{
+			MovieSceneToolHelpers::GetActorWorldTransforms(Sequencer, ActorSelection, Frames, OutParentWorldTransforms);
+			return true;
+		}
+	}
+
+	for (const FControlRigForWorldTransforms& ControlRigAndSelection : ParentToSnap.ControlRigs)
+	{
+		//get actor transform...
+		UControlRig* ControlRig = ControlRigAndSelection.ControlRig.Get();
+
+		if (ControlRig)
+		{
+			if (TSharedPtr<IControlRigObjectBinding> ObjectBinding = ControlRig->GetObjectBinding())
+			{
+				USceneComponent* Component = Cast<USceneComponent>(ObjectBinding->GetBoundObject());
+				if (!Component)
+				{
+					continue;
+				}
+				AActor* Actor = Component->GetTypedOuter< AActor >();
+				if (!Actor)
+				{
+					continue;
+				}
+				TArray<FTransform> ParentTransforms;
+				FActorForWorldTransforms ActorSelection;
+				ActorSelection.Actor = Actor;
+				MovieSceneToolHelpers::GetActorWorldTransforms(Sequencer, ActorSelection, Frames, ParentTransforms);
+
+				//just do first for now but may do average later.
+				for (const FName& Name : ControlRigAndSelection.ControlNames)
+				{
+					GetControlRigControlTransforms(Sequencer, ControlRig, Name, Frames, ParentTransforms, OutParentWorldTransforms);
+					return true;
+				}
+			}
+		}
+	}
+	OutParentWorldTransforms.SetNum(Frames.Num());
+	for (FTransform& Transform : OutParentWorldTransforms)
+	{
+		Transform = FTransform::Identity;
+	}
+	return false;
+}
+
+
+bool FControlRigSnapper::SnapIt(FFrameNumber StartFrame, FFrameNumber EndFrame,const FControlRigSnapperSelection& ActorToSnap,
+	const FControlRigSnapperSelection& ParentToSnap)
+{
+	TWeakPtr<ISequencer> InSequencer = GetSequencer();
+	if (InSequencer.IsValid() && InSequencer.Pin()->GetFocusedMovieSceneSequence() && ActorToSnap.IsValid())
+	{
+		const UControlRigSnapSettings* SnapSettings = GetDefault<UControlRigSnapSettings>();
+
+		FScopedTransaction ScopedTransaction(LOCTEXT("SnapAnimation", "Snap Animation"));
+
+		ISequencer* Sequencer = InSequencer.Pin().Get();
+
+		UMovieScene* MovieScene = Sequencer->GetFocusedMovieSceneSequence()->GetMovieScene();
+		MovieScene->Modify();
+		
+		TArray<FFrameNumber> Frames;
+		CalculateFramesToSnap(Sequencer,MovieScene, StartFrame, EndFrame, Frames);
+
+		TArray<FTransform> WorldTransformToSnap;
+		bool bSnapToFirstFrameNotParents = !CalculateWorldTransformsFromParents(Sequencer, ParentToSnap, Frames, WorldTransformToSnap);
+		if (Frames.Num() != WorldTransformToSnap.Num())
+		{
+			UE_LOG(LogControlRig, Error, TEXT("Number of Frames %d to Snap different than Parent Frames %d"), Frames.Num(),WorldTransformToSnap.Num());
+			return false;
+		}
+
+		TArray<FGuidAndActor > ActorsToSnap;
+		//There may be Actors here not in Sequencer so we add them to sequencer also
+		for (const FActorForWorldTransforms& ActorSelection : ActorToSnap.Actors)
+		{
+			if (ActorSelection.Actor.IsValid())
+			{
+				AActor* Actor = ActorSelection.Actor.Get();
+				FGuid ObjectHandle = Sequencer->GetHandleToObject(Actor,false);
+				if (ObjectHandle.IsValid() == false)
+				{
+					TArray<TWeakObjectPtr<AActor> > ActorsToAdd;
+					ActorsToAdd.Add(Actor);
+					TArray<FGuid> ActorTracks = Sequencer->AddActors(ActorsToAdd, false);
+					if (ActorTracks[0].IsValid())
+					{
+						ActorsToSnap.Add(FGuidAndActor(ActorTracks[0], Actor));
+					}
+				}
+				else
+				{
+					ActorsToSnap.Add(FGuidAndActor(ObjectHandle, Actor));
+				}
+			}
+		}
+		//set transforms on these guys
+		for (FGuidAndActor& GuidActor : ActorsToSnap)
+		{
+			if (bSnapToFirstFrameNotParents || SnapSettings->bKeepOffset) //if we are snapping to the first frame or keep offset we just don't set the parent transforms
+			{
+				TArray<FFrameNumber> OneFrame;
+				OneFrame.Add(Frames[0]);
+				TArray<FTransform> OneTransform;
+				FActorForWorldTransforms ActorSelection;
+				ActorSelection.Actor = GuidActor.Actor;
+				MovieSceneToolHelpers::GetActorWorldTransforms(Sequencer, ActorSelection, OneFrame, OneTransform);
+				if (bSnapToFirstFrameNotParents)
+				{
+					for (FTransform& Transform : WorldTransformToSnap)
+					{
+						Transform = OneTransform[0];
+					}
+				}
+				else //we keep offset
+				{
+					FTransform FirstWorld = WorldTransformToSnap[0];
+					for (FTransform& Transform : WorldTransformToSnap)
+					{
+						Transform = OneTransform[0].GetRelativeTransform(FirstWorld) * Transform;
+					}
+				}
+
+			}
+			GuidActor.SetTransform(Sequencer, Frames, WorldTransformToSnap);
+		}
+
+		//Now do Control Rigs
+		FFrameRate TickResolution = MovieScene->GetTickResolution();
+		FRigControlModifiedContext Context;
+		Context.SetKey = EControlRigSetKey::Always;
+
+		for (const FControlRigForWorldTransforms& ControlRigAndSelection : ActorToSnap.ControlRigs)
+		{
+			//get actor transform...
+			UControlRig* ControlRig = ControlRigAndSelection.ControlRig.Get();
+
+			if (ControlRig)
+			{
+				if (TSharedPtr<IControlRigObjectBinding> ObjectBinding = ControlRig->GetObjectBinding())
+				{
+					USceneComponent* Component = Cast<USceneComponent>(ObjectBinding->GetBoundObject());
+					if (!Component)
+					{
+						continue;
+					}
+					AActor* Actor = Component->GetTypedOuter< AActor >();
+					if (!Actor)
+					{
+						continue;
+					}
+					ControlRig->Modify();
+					TArray<FTransform> ControlRigParentWorldTransforms;
+					FActorForWorldTransforms ControlRigActorSelection;
+					ControlRigActorSelection.Actor = Actor;
+					MovieSceneToolHelpers::GetActorWorldTransforms(Sequencer, ControlRigActorSelection, Frames, ControlRigParentWorldTransforms);
+
+					for (const FName& Name : ControlRigAndSelection.ControlNames)
+					{
+						TArray<FFrameNumber> OneFrame;
+						OneFrame.SetNum(1);
+						TArray<FTransform> CurrentControlRigTransform, CurrentParentWorldTransform;
+						CurrentControlRigTransform.SetNum(1);
+						CurrentParentWorldTransform.SetNum(1);
+						if (bSnapToFirstFrameNotParents || SnapSettings->bKeepOffset)
+						{
+							OneFrame[0] = Frames[0];
+							CurrentParentWorldTransform[0] = ControlRigParentWorldTransforms[0];
+							GetControlRigControlTransforms(Sequencer, ControlRig, Name, OneFrame, CurrentParentWorldTransform, CurrentControlRigTransform);
+							if (bSnapToFirstFrameNotParents)
+							{
+								for (FTransform& Transform : WorldTransformToSnap)
+								{
+									Transform = CurrentControlRigTransform[0];
+								}
+							}
+							else if (SnapSettings->bKeepOffset)
+							{
+								FTransform FirstWorld = WorldTransformToSnap[0];
+								for (FTransform& Transform : WorldTransformToSnap)
+								{
+									Transform =  CurrentControlRigTransform[0].GetRelativeTransform(FirstWorld) * Transform;
+								}
+							}
+						}
+
+						for (int32 Index = 0; Index < WorldTransformToSnap.Num(); ++Index)
+						{
+							OneFrame[0] = Frames[Index];
+							CurrentParentWorldTransform[0] = ControlRigParentWorldTransforms[Index];
+							//this will evaluate at the current frame which we want
+							GetControlRigControlTransforms(Sequencer, ControlRig, Name, OneFrame, CurrentParentWorldTransform, CurrentControlRigTransform);
+							if (SnapSettings->bSnapPosition == false || SnapSettings->bSnapRotation == false || SnapSettings->bSnapScale == false)
+							{
+								FTransform& Transform = WorldTransformToSnap[Index];
+								const FTransform& CurrentTransform = CurrentControlRigTransform[0];
+								if (SnapSettings->bSnapPosition == false)
+								{
+									FVector CurrentPosition = CurrentTransform.GetLocation();
+									Transform.SetLocation(CurrentPosition);
+								}
+								if (SnapSettings->bSnapRotation == false)
+								{
+									FQuat CurrentRotation = CurrentTransform.GetRotation();
+									Transform.SetRotation(CurrentRotation);
+								}
+								if (SnapSettings->bSnapScale == false)
+								{
+									FVector Scale = CurrentTransform.GetScale3D();
+									Transform.SetScale3D(Scale);
+								}
+								
+							}
+							const FFrameNumber& FrameNumber = Frames[Index];
+							Context.LocalTime = TickResolution.AsSeconds(FFrameTime(FrameNumber));
+							FTransform GlobalTransform = WorldTransformToSnap[Index].GetRelativeTransform(ControlRigParentWorldTransforms[Index]);
+							//TODO need interrogator to get control rig values OR use different functon
+							GlobalTransform.SetScale3D(FVector(1.0f, 1.0f, 1.0f)); //note this is scale is off...
+							ControlRig->SetControlGlobalTransform(Name, GlobalTransform, true, Context);
+						}
+					}
+				}
+			}
+		}
+	}
+	return false;
+}
+
+
+
+#undef LOCTEXT_NAMESPACE
+
+
diff --git a/Engine/Plugins/Experimental/ControlRig/Source/ControlRigEditor/Public/IControlRigEditorModule.h b/Engine/Plugins/Experimental/ControlRig/Source/ControlRigEditor/Public/IControlRigEditorModule.h
--- a/Engine/Plugins/Experimental/ControlRig/Source/ControlRigEditor/Public/IControlRigEditorModule.h
+++ b/Engine/Plugins/Experimental/ControlRig/Source/ControlRigEditor/Public/IControlRigEditorModule.h
@@ -48,4 +48,6 @@ public:
 
 	static const FName ControlRigPoseTab;
 	static const FName ControlRigTweenTab;
+	static const FName ControlRigSnapperTab;
+	static const FName ControlRigTempPivotTab;
 };
diff --git a/Engine/Plugins/Experimental/ControlRig/Source/ControlRigEditor/Public/Tools/ControlRigSnapper.h b/Engine/Plugins/Experimental/ControlRig/Source/ControlRigEditor/Public/Tools/ControlRigSnapper.h
new file mode 100644
--- /dev/null
+++ b/Engine/Plugins/Experimental/ControlRig/Source/ControlRigEditor/Public/Tools/ControlRigSnapper.h
@@ -0,0 +1,67 @@
+// Copyright Epic Games, Inc. All Rights Reserved.
+
+#pragma once
+#include "CoreMinimal.h"
+#include "UObject/Object.h"
+#include "UObject/ObjectMacros.h"
+#include "Misc/FrameNumber.h"
+#include "ActorForWorldTransforms.h"
+#include "ControlRigSnapper.generated.h"
+
+class AActor;
+class UControlRig;
+class ISequencer;
+class UMovieScene;
+
+//Specification containing a Control Rig and a list of selected Controls we use to get World Transforms for Snapping.
+USTRUCT(BlueprintType)
+struct FControlRigForWorldTransforms
+{
+	GENERATED_BODY()
+
+	FControlRigForWorldTransforms() : ControlRig(nullptr) {};
+	UPROPERTY(BlueprintReadWrite, EditAnywhere, Category = "Control Rig")
+	TWeakObjectPtr<UControlRig> ControlRig;
+	UPROPERTY(BlueprintReadWrite, EditAnywhere, Category = "Control Rig")
+	TArray<FName> ControlNames;
+};
+
+//Selection from the UI to Snap To. Contains a set of Actors and/or Control Rigs to snap onto or from.
+USTRUCT(BlueprintType)
+struct FControlRigSnapperSelection
+{
+	GENERATED_BODY()
+
+	UPROPERTY(BlueprintReadWrite, EditAnywhere, Category = "Selection")
+	TArray<FActorForWorldTransforms> Actors;
+	UPROPERTY(BlueprintReadWrite, EditAnywhere, Category = "Selection")
+	TArray<FControlRigForWorldTransforms>  ControlRigs;
+
+	bool IsValid() const { return (NumSelected() > 0); }
+	FText GetName() const;
+	int32 NumSelected() const;
+	void Clear() { Actors.SetNum(0); ControlRigs.SetNum(0); }
+};
+
+
+
+struct FControlRigSnapper 
+{
+	/**
+	*  Get Current Active Sequencer in the Level.
+	*/
+	TWeakPtr<ISequencer> GetSequencer();
+
+	/**
+	*  Snap the passed in children over the start and end frames.
+	* @param StartFrame Start of the snap interval.
+	* @param EndFrame End of the snap interval.
+	* @param ChildrenToSnap The children to snap over the interval. Will set a key per frame on them.
+	* @param ParentToSnap  The parent to snap to.
+	*/
+	bool SnapIt(FFrameNumber StartFrame, FFrameNumber EndFrame, const FControlRigSnapperSelection& ChildrenToSnap,
+		const FControlRigSnapperSelection& ParentToSnap);
+
+};
+
+
diff --git a/Engine/Source/Editor/MovieSceneTools/Private/MovieSceneToolHelpers.cpp b/Engine/Source/Editor/MovieSceneTools/Private/MovieSceneToolHelpers.cpp
--- a/Engine/Source/Editor/MovieSceneTools/Private/MovieSceneToolHelpers.cpp
+++ b/Engine/Source/Editor/MovieSceneTools/Private/MovieSceneToolHelpers.cpp
@@ -1,6 +1,7 @@
 // Copyright Epic Games, Inc. All Rights Reserved.
 
 #include "MovieSceneToolHelpers.h"
+#include "ActorForWorldTransforms.h"
 #include "MovieSceneToolsModule.h"
 #include "MovieScene.h"
 #include "Layout/Margin.h"
@@ -78,6 +79,7 @@
 #include "Channels/MovieSceneIntegerChannel.h"
 #include "Channels/MovieSceneByteChannel.h"
 #include "Framework/MultiBox/MultiBoxBuilder.h"
+#include "EntitySystem/Interrogation/MovieSceneInterrogationLinker.h"
 
 /* FSkelMeshRecorder
  ***********/
@@ -3922,4 +3924,202 @@ void MovieSceneToolHelpers::GetLocationAtTime(const FMovieSceneEvaluationTrack*
 }
 
 
+static USkeletalMeshComponent* AcquireSkeletalMeshFromObject(UObject* BoundObject)
+{
+	if (AActor* Actor = Cast<AActor>(BoundObject))
+	{
+		for (UActorComponent* Component : Actor->GetComponents())
+		{
+			if (USkeletalMeshComponent* SkeletalMeshComp = Cast<USkeletalMeshComponent>(Component))
+			{
+				return SkeletalMeshComp;
+			}
+		}
+	}
+	else if (USkeletalMeshComponent* SkeletalMeshComponent = Cast<USkeletalMeshComponent>(BoundObject))
+	{
+		if (SkeletalMeshComponent->SkeletalMesh)
+		{
+			return SkeletalMeshComponent;
+		}
+	}
+	return nullptr;
+}
+
+static void GetSequencerActorWorldTransforms(ISequencer* Sequencer, const FActorForWorldTransforms& ActorSelection, const TArray<FFrameNumber>& Frames, TArray<FTransform>& OutTransforms)
+{
+	if (AActor* Actor = ActorSelection.Actor.Get())
+	{
+
+		USkeletalMeshComponent* SkelMeshComp = AcquireSkeletalMeshFromObject(Actor);
+
+		if (ActorSelection.SocketName != NAME_None && SkelMeshComp)
+		{
+
+			if (UMovieScene* MovieScene = Sequencer->GetFocusedMovieSceneSequence()->GetMovieScene())
+			{
+				OutTransforms.SetNum(Frames.Num());
+
+				FMovieSceneSequenceIDRef Template = Sequencer->GetFocusedTemplateID();
+				FMovieSceneSequenceTransform RootToLocalTransform;
+
+				FFrameRate TickResolution = MovieScene->GetTickResolution();
+				FFrameRate DisplayRate = MovieScene->GetDisplayRate();
+				const TArray<IMovieSceneToolsAnimationBakeHelper*>& BakeHelpers = FMovieSceneToolsModule::Get().GetAnimationBakeHelpers();
+
+				for (IMovieSceneToolsAnimationBakeHelper* BakeHelper : BakeHelpers)
+				{
+					if (BakeHelper)
+					{
+						BakeHelper->StartBaking(MovieScene);
+					}
+				}
+
+				for (int32 Index = 0; Index < Frames.Num(); ++Index)
+				{
+					const FFrameNumber& FrameNumber = Frames[Index];
+					FFrameTime GlobalTime(FrameNumber);
+
+					for (IMovieSceneToolsAnimationBakeHelper* BakeHelper : BakeHelpers)
+					{
+						if (BakeHelper)
+						{
+							BakeHelper->PreEvaluation(MovieScene, FrameNumber);
+						}
+					}
+
+					FMovieSceneContext Context = FMovieSceneContext(FMovieSceneEvaluationRange(GlobalTime, TickResolution), Sequencer->GetPlaybackStatus()).SetHasJumped(true);
+
+					Sequencer->GetEvaluationTemplate().Evaluate(Context, *Sequencer);
+
+					for (IMovieSceneToolsAnimationBakeHelper* BakeHelper : BakeHelpers)
+					{
+						if (BakeHelper)
+						{
+							BakeHelper->PostEvaluation(MovieScene, FrameNumber);
+						}
+					}
+
+					SkelMeshComp->TickAnimation(0.03f, false);
+					SkelMeshComp->RefreshBoneTransforms();
+					SkelMeshComp->RefreshSlaveComponents();
+					SkelMeshComp->UpdateComponentToWorld();
+					SkelMeshComp->FinalizeBoneTransform();
+					SkelMeshComp->MarkRenderTransformDirty();
+					SkelMeshComp->MarkRenderDynamicDataDirty();
+
+					OutTransforms[Index] = SkelMeshComp->GetSocketTransform(ActorSelection.SocketName);// GetSocketTransofrm is world space in theory*OutTransforms[Index];
+
+				}
+
+				for (IMovieSceneToolsAnimationBakeHelper* BakeHelper : BakeHelpers)
+				{
+					if (BakeHelper)
+					{
+						BakeHelper->StopBaking(MovieScene);
+					}
+				}
+			}
+		}
+		else //no attached skelmesh socket so use Interrogator
+		{
+			// Hack: static system interrogator for now to avoid re-allocating UObjects all the time
+			static UE::MovieScene::FSystemInterrogator Interrogator;
+			Interrogator.Reset();
+
+			USceneComponent* SceneComponent = Actor->GetRootComponent();
+
+			Interrogator.ImportTransformHierarchy(SceneComponent, Sequencer, Sequencer->GetFocusedTemplateID());
+
+			if (Frames[0] < Frames[Frames.Num() - 1])
+			{
+				for (const FFrameNumber& FrameNumber : Frames)
+				{
+					const FFrameTime FrameTime(FrameNumber);
+					Interrogator.AddInterrogation(FrameTime);
+				}
+				Interrogator.Update();
+
+				Interrogator.QueryWorldSpaceTransforms(SceneComponent, OutTransforms);
+			}
+			else //time is backward we need to do swap things around
+			{
+				for (int32 Index = Frames.Num() - 1; Index >= 0; --Index)
+				{
+					const FFrameTime FrameTime(Frames[Index]);
+					Interrogator.AddInterrogation(FrameTime);
+				}
+				Interrogator.Update();
+
+				TArray<FTransform> WorldTransforms;
+				Interrogator.QueryWorldSpaceTransforms(SceneComponent, WorldTransforms);
+				OutTransforms.SetNum(0);
+				OutTransforms.Reserve(WorldTransforms.Num());
+				for (int32 Index2 = WorldTransforms.Num() - 1; Index2 >= 0; --Index2)
+				{
+					OutTransforms.Add(WorldTransforms[Index2]);
+				}
+			}
+			Interrogator.Reset();
+		}
+	}
+}
+
+static void GetNonSequencerActorWorldTransforms(ISequencer* Sequencer, const FActorForWorldTransforms& ActorSelection, const TArray<FFrameNumber>& Frames, TArray<FTransform>& OutTransforms)
+{
+	FName SocketName = ActorSelection.SocketName;
+	AActor* Actor = ActorSelection.Actor.Get();
+	FActorForWorldTransforms NewActorSelection = ActorSelection;
+	FTransform WorldTransform = FTransform::Identity;
+	if (Actor)
+	{
+		do
+		{
+			FGuid ActorHandle = Sequencer->GetHandleToObject(Actor, false);
+			if (ActorHandle.IsValid())
+			{
+				GetSequencerActorWorldTransforms(Sequencer, NewActorSelection, Frames, OutTransforms);
+				for (FTransform& OutTransform : OutTransforms)
+				{
+					OutTransform = WorldTransform * OutTransform;
+				}
+				return;
+			}
+			else if (Actor->GetRootComponent()->DoesSocketExist(SocketName))
+			{
+				WorldTransform = WorldTransform * Actor->GetRootComponent()->GetSocketTransform(SocketName);
+			}
+			else
+			{
+				WorldTransform = WorldTransform * Actor->GetActorTransform();
+			}
 
+			Actor = Actor->GetAttachParentActor();
+			NewActorSelection.Actor = Actor;
+			if (Actor)
+			{
+				SocketName = Actor->GetAttachParentSocketName();
+				NewActorSelection.SocketName = SocketName;
+			}
+		} while (Actor);
+		//if we get to here world transform is what we want 
+	}
+	OutTransforms.SetNumUninitialized(Frames.Num());
+	for (FTransform& OutTransform : OutTransforms)
+	{
+		OutTransform = WorldTransform;
+	}
+}
+
+void MovieSceneToolHelpers::GetActorWorldTransforms(ISequencer* Sequencer, const FActorForWorldTransforms& ActorSelection, const TArray<FFrameNumber>& Frames, TArray<FTransform>& OutWorldTransforms)
+{
+	FGuid ObjectHandle = Sequencer->GetHandleToObject(ActorSelection.Actor.Get(), false);
+	if (ObjectHandle.IsValid())
+	{
+		GetSequencerActorWorldTransforms(Sequencer, ActorSelection, Frames, OutWorldTransforms);
+	}
+	else
+	{
+		GetNonSequencerActorWorldTransforms(Sequencer, ActorSelection, Frames, OutWorldTransforms);
+	}
+}
diff --git a/Engine/Source/Editor/MovieSceneTools/Public/MovieSceneToolHelpers.h b/Engine/Source/Editor/MovieSceneTools/Public/MovieSceneToolHelpers.h
--- a/Engine/Source/Editor/MovieSceneTools/Public/MovieSceneToolHelpers.h
+++ b/Engine/Source/Editor/MovieSceneTools/Public/MovieSceneToolHelpers.h
@@ -36,6 +36,7 @@ struct FMovieSceneSequenceTransform;
 class UAnimSeqExportOption;
 template<typename ChannelType> struct TMovieSceneChannelData;
 enum class EVisibilityBasedAnimTickOption : uint8;
+struct FActorForWorldTransforms;
 
 namespace fbxsdk
 {
@@ -62,7 +63,6 @@ DECLARE_DELEGATE(FStartAnimationCB);
 DECLARE_DELEGATE_OneParam(FTickAnimationCB, float);
 DECLARE_DELEGATE(FEndAnimationCB);
 
-
 //Skel Mesh Recorder to set up and restore various parameters on the skelmesh
 struct MOVIESCENETOOLS_API FSkelMeshRecorderState
 {
@@ -87,7 +87,6 @@ public:
 	bool bCachedEnableUpdateRateOptimizations;
 };
 
-
 class MOVIESCENETOOLS_API MovieSceneToolHelpers
 {
 public:
@@ -502,6 +501,17 @@ public:
 	*/	
 	static bool ImportFBXIntoControlRigChannels(UMovieScene* MovieScene, const FString& ImportFilename,  UMovieSceneUserImportFBXControlRigSettings *ControlRigSettings,
 		TArray<FFBXNodeAndChannels>* NodeAndChannels, const TArray<FName>& SelectedControlNames, FFrameRate FrameRate);
+
+	/*
+	*  Get an actors word transforms at the specified times
+	* @param Sequencer Sequencer to evaluation
+    * @param ActorForWorldTransforms The actor and possible component and socket that we want to get the world transforms for.
+	* @param Frames The times we want to get the world transforms
+	* @param OutWorldTransforms The calculated world transforms, one for each specified frame.
+	*/
+	static void GetActorWorldTransforms(ISequencer* Sequencer, const FActorForWorldTransforms& Actors, const TArray<FFrameNumber>& Frames, TArray<FTransform>& OutWorldTransforms);
+
+
 };
 
 // Helper to make spawnables persist throughout the export process and then restore properly afterwards
diff --git a/Engine/Source/Runtime/MovieScene/Public/ActorForWorldTransforms.h b/Engine/Source/Runtime/MovieScene/Public/ActorForWorldTransforms.h
new file mode 100644
--- /dev/null
+++ b/Engine/Source/Runtime/MovieScene/Public/ActorForWorldTransforms.h
@@ -0,0 +1,25 @@
+// Copyright Epic Games, Inc. All Rights Reserved.
+
+#pragma once
+
+#include "UObject/ObjectMacros.h"
+
+#include "ActorForWorldTransforms.generated.h"
+
+//Description of an actor selected parts we can find world transforms on 
+USTRUCT(BlueprintType)
+struct  MOVIESCENE_API FActorForWorldTransforms
+{
+	GENERATED_BODY()
+
+	FActorForWorldTransforms() : Actor(nullptr), Component(nullptr), SocketName(NAME_None) {};
+
+	UPROPERTY(BlueprintReadWrite, EditAnywhere, Category = "Actor")
+	TWeakObjectPtr<AActor> Actor;
+
+	UPROPERTY(BlueprintReadWrite, EditAnywhere, Category = "Actor")	
+	TWeakObjectPtr<USceneComponent> Component;
+
+	UPROPERTY(BlueprintReadWrite, EditAnywhere, Category = "Actor")
+	FName SocketName;
+};
\ No newline at end of file
