Control Rig Fix non-unity compile issue

diff --git a/Engine/Plugins/Experimental/ControlRig/Source/ControlRigEditor/Private/Editor/ControlRigSkeletalMeshComponent.h b/Engine/Plugins/Experimental/ControlRig/Source/ControlRigEditor/Private/Editor/ControlRigSkeletalMeshComponent.h
--- a/Engine/Plugins/Experimental/ControlRig/Source/ControlRigEditor/Private/Editor/ControlRigSkeletalMeshComponent.h
+++ b/Engine/Plugins/Experimental/ControlRig/Source/ControlRigEditor/Private/Editor/ControlRigSkeletalMeshComponent.h
@@ -2,7 +2,7 @@
 
 #pragma once
 
-#include "Rigs/RigHierarchy.h"
+#include "ControlRig.h"
 #include "Animation/DebugSkelMeshComponent.h"
 #include "ControlRigSkeletalMeshComponent.generated.h"
 
