Control Rig: Update initial transform during refresh

diff --git a/Engine/Plugins/Experimental/ControlRig/Source/ControlRig/Private/Rigs/RigHierarchyController.cpp b/Engine/Plugins/Experimental/ControlRig/Source/ControlRig/Private/Rigs/RigHierarchyController.cpp
--- a/Engine/Plugins/Experimental/ControlRig/Source/ControlRig/Private/Rigs/RigHierarchyController.cpp
+++ b/Engine/Plugins/Experimental/ControlRig/Source/ControlRig/Private/Rigs/RigHierarchyController.cpp
@@ -484,6 +484,7 @@ TArray<FRigElementKey> URigHierarchyController::ImportBones(const FReferenceSkel
 					SetParent(ExistingBoneKey, ParentKey, bSetupUndo);
 				}
 
+				Hierarchy->SetInitialLocalTransform(ExistingBoneIndex, BonePoses[Index], true, bSetupUndo);
 				Hierarchy->SetLocalTransform(ExistingBoneIndex, BonePoses[Index], true, bSetupUndo);
 
 				BonesToSelect.Add(ExistingBoneKey);
