Control Rig Function Libraries - First Pass. Feature to allow function sharing between rigs

diff --git a/Engine/Plugins/Experimental/ControlRig/Source/ControlRigDeveloper/Private/ControlRigBlueprint.cpp b/Engine/Plugins/Experimental/ControlRig/Source/ControlRigDeveloper/Private/ControlRigBlueprint.cpp
--- a/Engine/Plugins/Experimental/ControlRig/Source/ControlRigDeveloper/Private/ControlRigBlueprint.cpp
+++ b/Engine/Plugins/Experimental/ControlRig/Source/ControlRigDeveloper/Private/ControlRigBlueprint.cpp
@@ -37,6 +37,21 @@
 
 #define LOCTEXT_NAMESPACE "ControlRigBlueprint"
 
+FEdGraphPinType FControlRigPublicFunctionArg::GetPinType() const
+{
+	FRigVMExternalVariable Variable;
+	Variable.Name = Name;
+	Variable.bIsArray = bIsArray;
+	Variable.TypeName = CPPType;
+	
+	if(CPPTypeObjectPath.IsValid())
+	{
+		Variable.TypeObject = URigVMPin::FindObjectFromCPPTypeObjectPath(CPPTypeObjectPath.ToString());
+	}
+
+	return UControlRig::GetPinTypeFromExternalVariable(Variable);
+}
+
 TArray<UControlRigBlueprint*> UControlRigBlueprint::sCurrentlyOpenedRigBlueprints;
 
 UControlRigBlueprint::UControlRigBlueprint(const FObjectInitializer& ObjectInitializer)
@@ -220,6 +235,51 @@ void UControlRigBlueprint::SetPreviewMesh(USkeletalMesh* PreviewMesh, bool bMark
 	PreviewSkeletalMesh = PreviewMesh;
 }
 
+void UControlRigBlueprint::Serialize(FArchive& Ar)
+{
+	Super::Serialize(Ar);
+
+	if(Ar.IsObjectReferenceCollector())
+	{
+		TArray<UBlueprint*> ReferencedBlueprints;
+		
+		TArray<UEdGraph*> EdGraphs;
+		GetAllGraphs(EdGraphs);
+		for (UEdGraph* EdGraph : EdGraphs)
+		{
+			for(UEdGraphNode* Node : EdGraph->Nodes)
+			{
+				if(UControlRigGraphNode* RigNode = Cast<UControlRigGraphNode>(Node))
+				{
+					if(URigVMFunctionReferenceNode* FunctionRefNode = Cast<URigVMFunctionReferenceNode>(RigNode->GetModelNode()))
+					{
+						if(URigVMLibraryNode* ReferencedNode = FunctionRefNode->GetReferencedNode())
+						{
+							if(URigVMFunctionLibrary* ReferencedFunctionLibrary = ReferencedNode->GetLibrary())
+							{
+								if(ReferencedFunctionLibrary == GetLocalFunctionLibrary())
+								{
+									continue;
+								}
+
+								if(UBlueprint* ReferencedBlueprint = Cast<UBlueprint>(ReferencedFunctionLibrary->GetOuter()))
+								{
+									ReferencedBlueprints.AddUnique(ReferencedBlueprint);
+								}
+							}
+						}
+					}
+				}
+			}
+		}
+
+		for(UBlueprint* ReferencedBlueprint : ReferencedBlueprints)
+		{
+			Ar << ReferencedBlueprints;
+		}
+	}
+}
+
 void UControlRigBlueprint::PreSave(const class ITargetPlatform* TargetPlatform)
 {
 	Super::PreSave(TargetPlatform);
@@ -243,6 +303,17 @@ void UControlRigBlueprint::PreSave(const class ITargetPlatform* TargetPlatform)
 		}
 		return true;
 	});
+
+	for(FControlRigPublicFunctionData& FunctionData : PublicFunctions)
+	{
+		if(URigVMLibraryNode* FunctionNode = FunctionLibrary->FindFunction(FunctionData.Name))
+		{
+			if(UControlRigGraph* Graph = Cast<UControlRigGraph>(GetEdGraph(FunctionNode->GetContainedGraph())))
+			{
+				FunctionData = Graph->GetPublicFunctionData();
+			}
+		}
+	}
 }
 
 void UControlRigBlueprint::PostLoad()
@@ -1787,6 +1858,19 @@ void UControlRigBlueprint::HandleModifiedEvent(ERigVMGraphNotifType InNotifType,
 
 					ClearTransientControls();
 					RequestAutoVMRecompilation();
+
+					if(CollapseNode->GetOuter()->IsA<URigVMFunctionLibrary>())
+					{
+						for(int32 Index = 0; Index < PublicFunctions.Num(); Index++)
+						{
+							if(PublicFunctions[Index].Name == CollapseNode->GetFName())
+							{
+								Modify();
+								PublicFunctions.RemoveAt(Index);
+							}
+						}
+					}
+
 					MarkPackageDirty();
 					FBlueprintEditorUtils::MarkBlueprintAsStructurallyModified(this);
 					break;
@@ -1884,6 +1968,18 @@ void UControlRigBlueprint::HandleModifiedEvent(ERigVMGraphNotifType InNotifType,
 						ContainedEdGraph->Rename(*CollapseNode->GetEditorSubGraphName(), nullptr);
 					}
 
+					if(CollapseNode->GetOuter()->IsA<URigVMFunctionLibrary>())
+					{
+						for(int32 Index = 0; Index < PublicFunctions.Num(); Index++)
+						{
+							if(PublicFunctions[Index].Name == CollapseNode->GetPreviousFName())
+							{
+								Modify();
+								PublicFunctions[Index].Name = CollapseNode->GetFName();
+							}
+						}
+					}
+
 					FBlueprintEditorUtils::MarkBlueprintAsStructurallyModified(this);
 				}
 				break;
diff --git a/Engine/Plugins/Experimental/ControlRig/Source/ControlRigDeveloper/Private/Graph/ControlRigGraph.cpp b/Engine/Plugins/Experimental/ControlRig/Source/ControlRigDeveloper/Private/Graph/ControlRigGraph.cpp
--- a/Engine/Plugins/Experimental/ControlRig/Source/ControlRigDeveloper/Private/Graph/ControlRigGraph.cpp
+++ b/Engine/Plugins/Experimental/ControlRig/Source/ControlRigDeveloper/Private/Graph/ControlRigGraph.cpp
@@ -41,6 +41,18 @@ void UControlRigGraph::Initialize(UControlRigBlueprint* InBlueprint)
 	InBlueprint->OnModified().AddUObject(this, &UControlRigGraph::HandleModifiedEvent);
 	InBlueprint->OnVMCompiled().RemoveAll(this);
 	InBlueprint->OnVMCompiled().AddUObject(this, &UControlRigGraph::HandleVMCompiledEvent);
+
+	URigHierarchy* Hierarchy = InBlueprint->Hierarchy;
+
+	if(UControlRig* ControlRig = Cast<UControlRig>(InBlueprint->GetObjectBeingDebugged()))
+	{
+		Hierarchy = ControlRig->GetHierarchy();
+	}
+
+	if(Hierarchy)
+	{
+		CacheNameLists(Hierarchy, &InBlueprint->DrawContainer);
+	}
 }
 
 const UControlRigGraphSchema* UControlRigGraph::GetControlRigGraphSchema()
@@ -749,4 +761,41 @@ void UControlRigGraph::HandleVMCompiledEvent(UBlueprint* InBlueprint, URigVM* In
 
 #endif
 
+FControlRigPublicFunctionData UControlRigGraph::GetPublicFunctionData() const
+{
+	FControlRigPublicFunctionData Data;
+
+	FString Prefix, ModelNodeName;
+	if(!URigVMNode::SplitNodePathAtEnd(ModelNodePath, Prefix, ModelNodeName))
+	{
+		ModelNodeName = ModelNodePath;
+	}
+	Data.Name = *ModelNodeName;
+
+	if(URigVMGraph* RigGraph = GetModel())
+	{
+		if(URigVMCollapseNode* FunctionNode = Cast<URigVMCollapseNode>(RigGraph->GetOuter()))
+		{
+			Data.Category = FunctionNode->GetNodeCategory();
+			Data.Keywords = FunctionNode->GetNodeKeywords();
+			
+			for(URigVMPin* Pin : FunctionNode->GetPins())
+			{
+				FControlRigPublicFunctionArg Arg;
+				Arg.Name = Pin->GetFName();
+				Arg.bIsArray = Pin->IsArray();
+				Arg.Direction = Pin->GetDirection();
+				Arg.CPPType = *Pin->GetCPPType();
+				if(Pin->GetCPPTypeObject())
+				{
+					Arg.CPPTypeObjectPath = *Pin->GetCPPTypeObject()->GetPathName();
+				}
+				Data.Arguments.Add(Arg);
+			}
+		}
+	}
+
+	return Data;
+}
+
 #undef LOCTEXT_NAMESPACE
diff --git a/Engine/Plugins/Experimental/ControlRig/Source/ControlRigDeveloper/Public/ControlRigBlueprint.h b/Engine/Plugins/Experimental/ControlRig/Source/ControlRigDeveloper/Public/ControlRigBlueprint.h
--- a/Engine/Plugins/Experimental/ControlRig/Source/ControlRigDeveloper/Public/ControlRigBlueprint.h
+++ b/Engine/Plugins/Experimental/ControlRig/Source/ControlRigDeveloper/Public/ControlRigBlueprint.h
@@ -38,6 +38,66 @@ DECLARE_EVENT_TwoParams(UControlRigBlueprint, FOnNodeDoubleClicked, UControlRigB
 DECLARE_EVENT_OneParam(UControlRigBlueprint, FOnGraphImported, UEdGraph*);
 DECLARE_EVENT_OneParam(UControlRigBlueprint, FOnPostEditChangeChainProperty, FPropertyChangedChainEvent&);
 
+
+USTRUCT()
+struct CONTROLRIGDEVELOPER_API FControlRigPublicFunctionArg
+{
+	GENERATED_BODY();
+	
+	FControlRigPublicFunctionArg()
+	: Name(NAME_None)
+	, CPPType(NAME_None)
+	, CPPTypeObjectPath(NAME_None)
+	, bIsArray(false)
+	, Direction(ERigVMPinDirection::Input)
+	{}
+
+	UPROPERTY()
+	FName Name;
+
+	UPROPERTY()
+	FName CPPType;
+
+	UPROPERTY()
+	FName CPPTypeObjectPath;
+
+	UPROPERTY()
+	bool bIsArray;
+
+	UPROPERTY()
+	ERigVMPinDirection Direction;
+
+	FEdGraphPinType GetPinType() const;
+};
+
+USTRUCT()
+struct CONTROLRIGDEVELOPER_API FControlRigPublicFunctionData
+{
+	GENERATED_BODY();
+
+	FControlRigPublicFunctionData()
+		:Name(NAME_None)
+	{}
+
+	UPROPERTY()
+	FName Name;
+
+	UPROPERTY()
+	FString DisplayName;
+
+	UPROPERTY()
+	FString Category;
+
+	UPROPERTY()
+	FString Keywords;
+
+	UPROPERTY()
+	FControlRigPublicFunctionArg ReturnValue;
+
+	UPROPERTY()
+	TArray<FControlRigPublicFunctionArg> Arguments;
+};
+
 UCLASS(BlueprintType, meta=(IgnoreClassThumbnail))
 class CONTROLRIGDEVELOPER_API UControlRigBlueprint : public UBlueprint, public IInterface_PreviewMeshProvider
 {
@@ -54,6 +114,8 @@ public:
 	/** Get the (skeleton) generated class for this control rig blueprint */
 	UControlRigBlueprintGeneratedClass* GetControlRigBlueprintSkeletonClass() const;
 
+	virtual void Serialize(FArchive& Ar) override;
+
 #if WITH_EDITOR
 
 	// UBlueprint interface
@@ -152,6 +214,10 @@ protected:
 	UPROPERTY(BlueprintReadOnly, Category = "VM")
 	URigVMFunctionLibrary* FunctionLibrary;
 
+	/** Asset searchable information about exposed public functions on this rig */
+	UPROPERTY(AssetRegistrySearchable)
+	TArray<FControlRigPublicFunctionData> PublicFunctions;
+
 	UPROPERTY(BlueprintReadOnly, transient, Category = "VM")
 	TMap<URigVMGraph*, URigVMController*> Controllers;
 
@@ -405,4 +471,5 @@ private:
 	friend struct FRigValidationTabSummoner;
 	friend class UAnimGraphNode_ControlRig;
 	friend class UControlRigThumbnailRenderer;
+	friend class FControlRigGraphDetails;
 };
diff --git a/Engine/Plugins/Experimental/ControlRig/Source/ControlRigDeveloper/Public/Graph/ControlRigGraph.h b/Engine/Plugins/Experimental/ControlRig/Source/ControlRigDeveloper/Public/Graph/ControlRigGraph.h
--- a/Engine/Plugins/Experimental/ControlRig/Source/ControlRigDeveloper/Public/Graph/ControlRigGraph.h
+++ b/Engine/Plugins/Experimental/ControlRig/Source/ControlRigDeveloper/Public/Graph/ControlRigGraph.h
@@ -2,6 +2,7 @@
 
 #pragma once
 
+#include "ControlRigBlueprint.h"
 #include "EdGraph/EdGraph.h"
 #include "Graph/ControlRigGraphNode.h"
 #include "Rigs/RigHierarchy.h"
@@ -80,6 +81,8 @@ public:
 	UPROPERTY()
 	bool bIsFunctionDefinition;
 
+	FControlRigPublicFunctionData GetPublicFunctionData() const;
+
 private:
 
 	template<class T>
diff --git a/Engine/Plugins/Experimental/ControlRig/Source/ControlRigEditor/Private/ControlRigEditorModule.cpp b/Engine/Plugins/Experimental/ControlRig/Source/ControlRigEditor/Private/ControlRigEditorModule.cpp
--- a/Engine/Plugins/Experimental/ControlRig/Source/ControlRigEditor/Private/ControlRigEditorModule.cpp
+++ b/Engine/Plugins/Experimental/ControlRig/Source/ControlRigEditor/Private/ControlRigEditorModule.cpp
@@ -1032,6 +1032,8 @@ TSharedRef<IControlRigEditor> FControlRigEditorModule::CreateControlRigEditor(co
 
 void FControlRigEditorModule::GetTypeActions(UControlRigBlueprint* CRB, FBlueprintActionDatabaseRegistrar& ActionRegistrar)
 {
+	FAssetRegistryModule& AssetRegistryModule = FModuleManager::LoadModuleChecked<FAssetRegistryModule>("AssetRegistry");
+
 	// actions get registered under specific object-keys; the idea is that 
 	// actions might have to be updated (or deleted) if their object-key is  
 	// mutated (or removed)... here we use the class (so if the class 
@@ -1153,6 +1155,36 @@ void FControlRigEditorModule::GetTypeActions(UControlRigBlueprint* CRB, FBluepri
 		check(NodeSpawner != nullptr);
 		ActionRegistrar.AddBlueprintAction(ActionKey, NodeSpawner);
 	}
+
+	FArrayProperty* PublicFunctionsProperty = CastField<FArrayProperty>(UControlRigBlueprint::StaticClass()->FindPropertyByName(TEXT("PublicFunctions")));
+	if(PublicFunctionsProperty)
+	{
+		// find all control rigs in the project
+		TArray<FAssetData> ControlRigAssetDatas;
+		FARFilter ControlRigAssetFilter;
+		ControlRigAssetFilter.ClassNames.Add(UControlRigBlueprint::StaticClass()->GetFName());
+		AssetRegistryModule.Get().GetAssets(ControlRigAssetFilter, ControlRigAssetDatas);
+
+		// loop over all control rigs in the project
+		for(const FAssetData& ControlRigAssetData : ControlRigAssetDatas)
+		{
+			const FString PublicFunctionsString = ControlRigAssetData.GetTagValueRef<FString>(PublicFunctionsProperty->GetFName());
+			if(PublicFunctionsString.IsEmpty())
+			{
+				continue;
+			}
+
+			TArray<FControlRigPublicFunctionData> PublicFunctions;
+			PublicFunctionsProperty->ImportText(*PublicFunctionsString, &PublicFunctions, EPropertyPortFlags::PPF_None, nullptr);
+
+			for(const FControlRigPublicFunctionData& PublicFunction : PublicFunctions)
+			{
+				UBlueprintNodeSpawner* NodeSpawner = UControlRigFunctionRefNodeSpawner::CreateFromAssetData(ControlRigAssetData, PublicFunction);
+				check(NodeSpawner != nullptr);
+				ActionRegistrar.AddBlueprintAction(ActionKey, NodeSpawner);
+			}
+		}
+	}
 }
 
 void FControlRigEditorModule::GetInstanceActions(UControlRigBlueprint* CRB, FBlueprintActionDatabaseRegistrar& ActionRegistrar)
diff --git a/Engine/Plugins/Experimental/ControlRig/Source/ControlRigEditor/Private/ControlRigGraphDetails.cpp b/Engine/Plugins/Experimental/ControlRig/Source/ControlRigEditor/Private/ControlRigGraphDetails.cpp
--- a/Engine/Plugins/Experimental/ControlRig/Source/ControlRigEditor/Private/ControlRigGraphDetails.cpp
+++ b/Engine/Plugins/Experimental/ControlRig/Source/ControlRigEditor/Private/ControlRigGraphDetails.cpp
@@ -9,6 +9,7 @@
 #include "EditorStyleSet.h"
 #include "SPinTypeSelector.h"
 #include "Widgets/Input/SButton.h"
+#include "Widgets/Input/SCheckBox.h"
 #include "Widgets/Text/STextBlock.h"
 #include "Widgets/Colors/SColorPicker.h"
 #include "Widgets/Text/SInlineEditableTextBlock.h"
@@ -324,22 +325,6 @@ FEdGraphPinType FControlRigArgumentLayout::OnGetPinInfo() const
 	return FEdGraphPinType();
 }
 
-ECheckBoxState FControlRigArgumentLayout::IsRefChecked() const
-{
-	FEdGraphPinType PinType = OnGetPinInfo();
-	return PinType.bIsReference ? ECheckBoxState::Checked : ECheckBoxState::Unchecked;
-}
-
-void FControlRigArgumentLayout::OnRefCheckStateChanged(ECheckBoxState InState)
-{
-	const FScopedTransaction Transaction(LOCTEXT("ChangeByRef", "Change Pass By Reference"));
-
-	FEdGraphPinType PinType = OnGetPinInfo();
-	PinType.bIsReference = (InState == ECheckBoxState::Checked) ? true : false;
-
-	PinInfoChanged(PinType);
-}
-
 void FControlRigArgumentLayout::PinInfoChanged(const FEdGraphPinType& PinType)
 {
 	if (PinPtr.IsValid() && ControlRigBlueprintPtr.IsValid() && FBlueprintEditorUtils::IsPinTypeValid(PinType))
@@ -675,6 +660,39 @@ void FControlRigGraphDetails::CustomizeDetails(IDetailLayoutBuilder& DetailLayou
 			.Text(this, &FControlRigGraphDetails::GetNodeKeywords)
 			.OnTextCommitted(this, &FControlRigGraphDetails::SetNodeKeywords)
 		];
+
+		if(AccessSpecifierStrings.IsEmpty())
+		{
+			AccessSpecifierStrings.Add(TSharedPtr<FString>(new FString(TEXT("Public"))));
+			AccessSpecifierStrings.Add(TSharedPtr<FString>(new FString(TEXT("Private"))));
+		}
+
+		// access specifier
+		SettingsCategory.AddCustomRow( LOCTEXT( "AccessSpecifier", "Access Specifier" ) )
+        .NameContent()
+        [
+            SNew(STextBlock)
+                .Text( LOCTEXT( "AccessSpecifier", "Access Specifier" ) )
+                .Font( IDetailLayoutBuilder::GetDetailFont() )
+        ]
+        .ValueContent()
+        [
+            SNew(SComboButton)
+            .ContentPadding(0)
+            .ButtonContent()
+            [
+                SNew(STextBlock)
+                    .Text(this, &FControlRigGraphDetails::GetCurrentAccessSpecifierName)
+                    .Font( IDetailLayoutBuilder::GetDetailFont() )
+            ]
+            .MenuContent()
+            [
+                SNew(SListView<TSharedPtr<FString> >)
+                    .ListItemsSource( &AccessSpecifierStrings )
+                    .OnGenerateRow(this, &FControlRigGraphDetails::HandleGenerateRowAccessSpecifier)
+                    .OnSelectionChanged(this, &FControlRigGraphDetails::OnAccessSpecifierSelected)
+            ]
+        ];
 	}
 
 	// node color
@@ -912,4 +930,73 @@ FReply FControlRigGraphDetails::OnNodeColorClicked()
 	return FReply::Handled();
 }
 
+TArray<TSharedPtr<FString>> FControlRigGraphDetails::AccessSpecifierStrings;
+
+FText FControlRigGraphDetails::GetCurrentAccessSpecifierName() const
+{
+	if(ControlRigBlueprintPtr.IsValid() && GraphPtr.IsValid())
+	{
+		UControlRigGraph* Graph = GraphPtr.Get();
+		UControlRigBlueprint* ControlRigBlueprint = ControlRigBlueprintPtr.Get();
+
+		const FControlRigPublicFunctionData ExpectedFunctionData = Graph->GetPublicFunctionData();
+		for(const FControlRigPublicFunctionData& PublicFunction : ControlRigBlueprint->PublicFunctions)
+		{
+			if(PublicFunction.Name == ExpectedFunctionData.Name)
+			{
+				return FText::FromString(*AccessSpecifierStrings[0].Get()); // public
+			}
+		}
+
+		return FText::FromString(*AccessSpecifierStrings[1].Get()); // private
+	}
+
+	return FText::FromString(*AccessSpecifierStrings[1].Get()); // private
+}
+
+void FControlRigGraphDetails::OnAccessSpecifierSelected( TSharedPtr<FString> SpecifierName, ESelectInfo::Type SelectInfo )
+{
+	if(ControlRigBlueprintPtr.IsValid() && GraphPtr.IsValid())
+	{
+		UControlRigGraph* Graph = GraphPtr.Get();
+		UControlRigBlueprint* ControlRigBlueprint = ControlRigBlueprintPtr.Get();
+		ControlRigBlueprint->Modify();
+
+		if(SpecifierName->Equals(TEXT("Private")))
+		{
+			for(int32 Index = 0; Index < ControlRigBlueprint->PublicFunctions.Num(); Index++)
+			{
+				if(ControlRigBlueprint->PublicFunctions[Index].Name == Graph->GetFName())
+				{
+					ControlRigBlueprint->PublicFunctions.RemoveAt(Index);
+					return;
+				}
+			}
+		}
+		else
+		{
+			const FControlRigPublicFunctionData NewFunctionData = Graph->GetPublicFunctionData();
+			for(FControlRigPublicFunctionData& ExistingFunctionData : ControlRigBlueprint->PublicFunctions)
+			{
+				if(ExistingFunctionData.Name == NewFunctionData.Name)
+				{
+					ExistingFunctionData = NewFunctionData;
+					return;
+				}
+			}
+			ControlRigBlueprint->PublicFunctions.Add(NewFunctionData);
+		}
+	}
+}
+
+TSharedRef<ITableRow> FControlRigGraphDetails::HandleGenerateRowAccessSpecifier( TSharedPtr<FString> SpecifierName, const TSharedRef<STableViewBase>& OwnerTable )
+{
+	return SNew(STableRow< TSharedPtr<FString> >, OwnerTable)
+        .Content()
+        [
+            SNew( STextBlock ) 
+                .Text(FText::FromString(*SpecifierName.Get()) )
+        ];
+}
+
 #undef LOCTEXT_NAMESPACE
diff --git a/Engine/Plugins/Experimental/ControlRig/Source/ControlRigEditor/Private/ControlRigGraphDetails.h b/Engine/Plugins/Experimental/ControlRig/Source/ControlRigEditor/Private/ControlRigGraphDetails.h
--- a/Engine/Plugins/Experimental/ControlRig/Source/ControlRigEditor/Private/ControlRigGraphDetails.h
+++ b/Engine/Plugins/Experimental/ControlRig/Source/ControlRigEditor/Private/ControlRigGraphDetails.h
@@ -43,7 +43,6 @@ private:
 	TWeakObjectPtr<UControlRigBlueprint> ControlRigBlueprintPtr;
 	bool bIsInputGroup;
 	FSimpleDelegate OnRebuildChildren;
-
 };
 
 class FControlRigArgumentLayout : public IDetailCustomNodeBuilder, public TSharedFromThis<FControlRigArgumentLayout>
@@ -91,12 +90,6 @@ private:
 	void PinInfoChanged(const FEdGraphPinType& PinType);
 	void OnPrePinInfoChange(const FEdGraphPinType& PinType);
 
-	/** Returns whether the "Pass-by-Reference" checkbox is checked or not */
-	ECheckBoxState IsRefChecked() const;
-
-	/** Handles toggling the "Pass-by-Reference" checkbox */
-	void OnRefCheckStateChanged(ECheckBoxState InState);
-
 private:
 
 	/** The argument pin that this layout reflects */
@@ -174,6 +167,9 @@ public:
 	void OnNodeColorEnd();
 	void OnNodeColorCancelled(FLinearColor OriginalColor);
 	FReply OnNodeColorClicked();
+	FText GetCurrentAccessSpecifierName() const;
+	void OnAccessSpecifierSelected( TSharedPtr<FString> SpecifierName, ESelectInfo::Type SelectInfo );
+	TSharedRef<ITableRow> HandleGenerateRowAccessSpecifier( TSharedPtr<FString> SpecifierName, const TSharedRef<STableViewBase>& OwnerTable );
 
 private:
 
@@ -197,4 +193,6 @@ private:
 
 	/** Set to true if the UI is currently picking a color */
 	bool bIsPickingColor;
+
+	static TArray<TSharedPtr<FString>> AccessSpecifierStrings;
 };
\ No newline at end of file
diff --git a/Engine/Plugins/Experimental/ControlRig/Source/ControlRigEditor/Private/Editor/ControlRigEditor.cpp b/Engine/Plugins/Experimental/ControlRig/Source/ControlRigEditor/Private/Editor/ControlRigEditor.cpp
--- a/Engine/Plugins/Experimental/ControlRig/Source/ControlRigEditor/Private/Editor/ControlRigEditor.cpp
+++ b/Engine/Plugins/Experimental/ControlRig/Source/ControlRigEditor/Private/Editor/ControlRigEditor.cpp
@@ -77,6 +77,7 @@
 #include "BlueprintEditorTabs.h"
 #include "RigVMModel/Nodes/RigVMFunctionEntryNode.h"
 #include "RigVMModel/Nodes/RigVMFunctionReturnNode.h"
+#include "BlueprintActionDatabaseRegistrar.h"
 
 #define LOCTEXT_NAMESPACE "ControlRigEditor"
 
@@ -1410,8 +1411,15 @@ void FControlRigEditor::SaveAsset_Execute()
 
 	// Save the new state of the hierarchy in the default object, so that it has the correct values on load
 	UControlRigBlueprint* RigBlueprint = Cast<UControlRigBlueprint>(GetBlueprintObj());
-	UControlRig* CDO = ControlRig->GetClass()->GetDefaultObject<UControlRig>();
-	CDO->DynamicHierarchy->CopyHierarchy(RigBlueprint->Hierarchy);
+	if(ControlRig)
+	{
+		UControlRig* CDO = ControlRig->GetClass()->GetDefaultObject<UControlRig>();
+		CDO->DynamicHierarchy->CopyHierarchy(RigBlueprint->Hierarchy);
+	}
+
+	FBlueprintActionDatabase& ActionDatabase = FBlueprintActionDatabase::Get();
+	ActionDatabase.ClearAssetActions(UControlRigBlueprint::StaticClass());
+	ActionDatabase.RefreshClassActions(UControlRigBlueprint::StaticClass());
 }
 
 void FControlRigEditor::SaveAssetAs_Execute()
@@ -1421,8 +1429,15 @@ void FControlRigEditor::SaveAssetAs_Execute()
 
 	// Save the new state of the hierarchy in the default object, so that it has the correct values on load
 	UControlRigBlueprint* RigBlueprint = Cast<UControlRigBlueprint>(GetBlueprintObj());
-	UControlRig* CDO = ControlRig->GetClass()->GetDefaultObject<UControlRig>();
-	CDO->DynamicHierarchy->CopyHierarchy(RigBlueprint->Hierarchy);
+	if(ControlRig)
+	{
+		UControlRig* CDO = ControlRig->GetClass()->GetDefaultObject<UControlRig>();
+		CDO->DynamicHierarchy->CopyHierarchy(RigBlueprint->Hierarchy);
+	}
+
+	FBlueprintActionDatabase& ActionDatabase = FBlueprintActionDatabase::Get();
+	ActionDatabase.ClearAssetActions(UControlRigBlueprint::StaticClass());
+	ActionDatabase.RefreshClassActions(UControlRigBlueprint::StaticClass());
 }
 
 FName FControlRigEditor::GetToolkitFName() const
@@ -4302,6 +4317,22 @@ void FControlRigEditor::OnNodeDoubleClicked(UControlRigBlueprint* InBlueprint, U
 			{
 				OpenGraphAndBringToFront(EdGraph, true);
 			}
+			else
+			{
+				if(URigVMCollapseNode* FunctionLibraryNode = Cast<URigVMCollapseNode>(ContainedGraph->GetOuter()))
+				{
+					if(URigVMFunctionLibrary* FunctionLibrary = FunctionLibraryNode->GetLibrary())
+					{
+						if(UControlRigBlueprint* FunctionBlueprint = Cast<UControlRigBlueprint>(FunctionLibrary->GetOuter()))
+						{
+							if (UEdGraph* FunctionEdGraph = FunctionBlueprint->GetEdGraph(ContainedGraph))
+							{
+								FKismetEditorUtilities::BringKismetToFocusAttentionOnObject(FunctionEdGraph);
+							}
+						}
+					}
+				}
+			}
 		}
 	}
 }
diff --git a/Engine/Plugins/Experimental/ControlRig/Source/ControlRigEditor/Private/Editor/SControlRigStackView.cpp b/Engine/Plugins/Experimental/ControlRig/Source/ControlRigEditor/Private/Editor/SControlRigStackView.cpp
--- a/Engine/Plugins/Experimental/ControlRig/Source/ControlRigEditor/Private/Editor/SControlRigStackView.cpp
+++ b/Engine/Plugins/Experimental/ControlRig/Source/ControlRigEditor/Private/Editor/SControlRigStackView.cpp
@@ -131,7 +131,10 @@ SControlRigStackView::~SControlRigStackView()
 		}
 		if (OnControlRigInitializedHandle.IsValid())
 		{
-			ControlRigEditor.Pin()->ControlRig->OnInitialized_AnyThread().Remove(OnControlRigInitializedHandle);
+			if(ControlRigEditor.Pin()->ControlRig)
+			{
+				ControlRigEditor.Pin()->ControlRig->OnInitialized_AnyThread().Remove(OnControlRigInitializedHandle);
+			}
 		}
 		if (OnPreviewControlRigUpdatedHandle.IsValid())
 		{
@@ -509,8 +512,10 @@ void SControlRigStackView::OnVMCompiled(UBlueprint* InCompiledBlueprint, URigVM*
 
 	if (ControlRigEditor.IsValid() && !OnControlRigInitializedHandle.IsValid())
 	{
-		UControlRig* ControlRig = ControlRigEditor.Pin()->ControlRig;
-		OnControlRigInitializedHandle = ControlRig->OnInitialized_AnyThread().AddSP(this, &SControlRigStackView::HandleControlRigInitializedEvent);
+		if(UControlRig* ControlRig = ControlRigEditor.Pin()->ControlRig)
+		{
+			OnControlRigInitializedHandle = ControlRig->OnInitialized_AnyThread().AddSP(this, &SControlRigStackView::HandleControlRigInitializedEvent);
+		}
 	}
 }
 
diff --git a/Engine/Plugins/Experimental/ControlRig/Source/ControlRigEditor/Private/Graph/NodeSpawners/ControlRigFunctionRefNodeSpawner.cpp b/Engine/Plugins/Experimental/ControlRig/Source/ControlRigEditor/Private/Graph/NodeSpawners/ControlRigFunctionRefNodeSpawner.cpp
--- a/Engine/Plugins/Experimental/ControlRig/Source/ControlRigEditor/Private/Graph/NodeSpawners/ControlRigFunctionRefNodeSpawner.cpp
+++ b/Engine/Plugins/Experimental/ControlRig/Source/ControlRigEditor/Private/Graph/NodeSpawners/ControlRigFunctionRefNodeSpawner.cpp
@@ -33,25 +33,12 @@ UControlRigFunctionRefNodeSpawner* UControlRigFunctionRefNodeSpawner::CreateFrom
 	UControlRigFunctionRefNodeSpawner* NodeSpawner = NewObject<UControlRigFunctionRefNodeSpawner>(GetTransientPackage());
 	NodeSpawner->ReferencedFunctionPtr = InFunction;
 	NodeSpawner->NodeClass = UControlRigGraphNode::StaticClass();
+	NodeSpawner->bIsLocalFunction = true;
 
 	FBlueprintActionUiSpec& MenuSignature = NodeSpawner->DefaultMenuSignature;
 
-	FString Category = InFunction->GetNodeCategory();
-	FString CategoryPrefix;
-
-	if (UControlRigBlueprint* Blueprint = Cast<UControlRigBlueprint>(InFunction->GetLibrary()->GetOuter()))
-	{
-		/*if(Blueprint->IsFunctionLibrary())
-		{
-			// todo
-			ensure(false);
-		}
-		else
-		*/
-		{
-			CategoryPrefix = TEXT("Local Functions");
-		}
-	}
+	const FString Category = InFunction->GetNodeCategory();
+	FString CategoryPrefix = TEXT("Local Functions");
 
 	if (!Category.IsEmpty() && !CategoryPrefix.IsEmpty())
 	{
@@ -76,6 +63,41 @@ UControlRigFunctionRefNodeSpawner* UControlRigFunctionRefNodeSpawner::CreateFrom
 	return NodeSpawner;
 }
 
+UControlRigFunctionRefNodeSpawner* UControlRigFunctionRefNodeSpawner::CreateFromAssetData(const FAssetData& InAssetData, const FControlRigPublicFunctionData& InPublicFunction)
+{
+	UControlRigFunctionRefNodeSpawner* NodeSpawner = NewObject<UControlRigFunctionRefNodeSpawner>(GetTransientPackage());
+	NodeSpawner->ReferencedAssetObjectPath = InAssetData.ObjectPath;
+	NodeSpawner->ReferencedPublicFunctionData = InPublicFunction;
+	NodeSpawner->NodeClass = UControlRigGraphNode::StaticClass();
+	NodeSpawner->bIsLocalFunction = false;
+
+	FBlueprintActionUiSpec& MenuSignature = NodeSpawner->DefaultMenuSignature;
+
+	const FString Category = InPublicFunction.Category;
+	FString CategoryPrefix =  InAssetData.ToSoftObjectPath().GetAssetName();
+
+	if (!Category.IsEmpty() && !CategoryPrefix.IsEmpty())
+	{
+		CategoryPrefix += TEXT("|");
+	}
+
+	MenuSignature.MenuName = FText::FromName(InPublicFunction.Name);
+	MenuSignature.Category = FText::FromString(CategoryPrefix + Category);
+	MenuSignature.Keywords = FText::FromString(InPublicFunction.Keywords);
+
+	// add at least one character, so that PrimeDefaultUiSpec() doesn't 
+	// attempt to query the template node
+	if (MenuSignature.Keywords.IsEmpty())
+	{
+		// want to set it to something so we won't end up back in this condition
+		MenuSignature.Keywords = FText::FromString(TEXT(" "));
+	}
+
+	MenuSignature.Icon = FSlateIcon("EditorStyle", "Kismet.AllClasses.FunctionIcon");
+
+	return NodeSpawner;
+}
+
 void UControlRigFunctionRefNodeSpawner::Prime()
 {
 	// we expect that you don't need a node template to construct menu entries
@@ -84,11 +106,34 @@ void UControlRigFunctionRefNodeSpawner::Prime()
 
 FBlueprintNodeSignature UControlRigFunctionRefNodeSpawner::GetSpawnerSignature() const
 {
+	FString SignatureString = TEXT("Invalid RigFunction");
 	if(ReferencedFunctionPtr.IsValid())
 	{
-		return FBlueprintNodeSignature(TEXT("RigFunction=") + ReferencedFunctionPtr->GetPathName());
+		if(UObject* FunctionLibrary = ReferencedFunctionPtr->GetOuter())
+		{
+			if(UObject* Blueprint = FunctionLibrary->GetOuter())
+			{
+				SignatureString = 
+					FString::Printf(TEXT("RigFunction=%s::%s"),
+					*Blueprint->GetPathName(),
+					*ReferencedFunctionPtr->GetName());
+			}
+		}
+	}
+	else if(ReferencedAssetObjectPath.IsValid() && ReferencedPublicFunctionData.Name.IsValid())
+	{
+		SignatureString = 
+			FString::Printf(TEXT("RigFunction=%s::%s"),
+			*ReferencedAssetObjectPath.ToString(),
+			*ReferencedPublicFunctionData.Name.ToString());
+	}
+
+	if(bIsLocalFunction)
+	{
+		SignatureString += TEXT(" (local)");
 	}
-	return FBlueprintNodeSignature(TEXT("Invalid RigFunction"));
+
+	return FBlueprintNodeSignature(SignatureString);
 }
 
 FBlueprintActionUiSpec UControlRigFunctionRefNodeSpawner::GetUiSpec(FBlueprintActionContext const& Context, FBindingSet const& Bindings) const
@@ -104,6 +149,33 @@ UEdGraphNode* UControlRigFunctionRefNodeSpawner::Invoke(UEdGraph* ParentGraph, F
 {
 	UControlRigGraphNode* NewNode = nullptr;
 
+	// if we are trying to build the real function ref - but we haven't loaded the asset yet...
+	if(!FBlueprintNodeTemplateCache::IsTemplateOuter(ParentGraph))
+	{
+		if(!ReferencedFunctionPtr.IsValid() && !ReferencedAssetObjectPath.IsNone() && !ReferencedPublicFunctionData.Name.IsNone())
+		{
+			UControlRigBlueprint* ReferencedBlueprint = LoadObject<UControlRigBlueprint>(nullptr, *ReferencedAssetObjectPath.ToString());
+			if(ReferencedBlueprint == nullptr)
+			{
+				return nullptr;
+			}
+
+			URigVMFunctionLibrary* FunctionLibrary = ReferencedBlueprint->GetLocalFunctionLibrary();
+			if(FunctionLibrary == nullptr)
+			{
+				return nullptr;
+			}
+
+			URigVMLibraryNode* FunctionNode = FunctionLibrary->FindFunction(ReferencedPublicFunctionData.Name);
+			if(FunctionNode == nullptr)
+			{
+				return nullptr;
+			}
+
+			ReferencedFunctionPtr = FunctionNode;
+		}
+	}
+
 	if(ReferencedFunctionPtr.IsValid())
 	{
 #if WITH_EDITOR
@@ -116,6 +188,40 @@ UEdGraphNode* UControlRigFunctionRefNodeSpawner::Invoke(UEdGraph* ParentGraph, F
 		UBlueprint* Blueprint = FBlueprintEditorUtils::FindBlueprintForGraph(ParentGraph);
 		NewNode = SpawnNode(ParentGraph, Blueprint, ReferencedFunctionPtr.Get(), Location);
 	}
+	else
+	{
+		// we are only going to get here if we are spawning a template node
+		NewNode = NewObject<UControlRigGraphNode>(ParentGraph);
+		ParentGraph->AddNode(NewNode, false);
+
+		NewNode->CreateNewGuid();
+		NewNode->PostPlacedNewNode();
+
+		for(const FControlRigPublicFunctionArg& Arg : ReferencedPublicFunctionData.Arguments)
+		{
+			if(Arg.Direction ==  ERigVMPinDirection::Input ||
+				Arg.Direction ==  ERigVMPinDirection::IO)
+			{
+				UEdGraphPin* InputPin = UEdGraphPin::CreatePin(NewNode);
+				NewNode->Pins.Add(InputPin);
+
+				InputPin->Direction = EGPD_Input;
+				InputPin->PinType = Arg.GetPinType();
+			}
+
+			if(Arg.Direction ==  ERigVMPinDirection::Output ||
+				Arg.Direction ==  ERigVMPinDirection::IO)
+			{
+				UEdGraphPin* OutputPin = UEdGraphPin::CreatePin(NewNode);
+				NewNode->Pins.Add(OutputPin);
+
+				OutputPin->Direction = EGPD_Output;
+				OutputPin->PinType = Arg.GetPinType();
+			}
+		}
+
+		NewNode->SetFlags(RF_Transactional);
+	}
 
 	return NewNode;
 }
@@ -168,5 +274,33 @@ UControlRigGraphNode* UControlRigFunctionRefNodeSpawner::SpawnNode(UEdGraph* Par
 	return NewNode;
 }
 
+bool UControlRigFunctionRefNodeSpawner::IsTemplateNodeFilteredOut(FBlueprintActionFilter const& Filter) const
+{
+	if(bIsLocalFunction)
+	{
+		if(ReferencedFunctionPtr.IsValid())
+		{
+			for (UBlueprint* Blueprint : Filter.Context.Blueprints)
+			{
+				if(Blueprint->GetOutermost() != ReferencedFunctionPtr.Get()->GetOutermost())
+				{
+					return true;
+				}
+			}
+		}
+	}
+	else if(ReferencedAssetObjectPath.IsValid())
+	{
+		const FString ReferencedAssetObjectPathString = ReferencedAssetObjectPath.ToString();
+		for (UBlueprint* Blueprint : Filter.Context.Blueprints)
+		{
+			if(Blueprint->GetPathName() == ReferencedAssetObjectPathString)
+			{
+				return true;
+			}
+		}
+	}
+	return false;
+}
 
 #undef LOCTEXT_NAMESPACE
diff --git a/Engine/Plugins/Experimental/ControlRig/Source/ControlRigEditor/Private/Graph/NodeSpawners/ControlRigFunctionRefNodeSpawner.h b/Engine/Plugins/Experimental/ControlRig/Source/ControlRigEditor/Private/Graph/NodeSpawners/ControlRigFunctionRefNodeSpawner.h
--- a/Engine/Plugins/Experimental/ControlRig/Source/ControlRigEditor/Private/Graph/NodeSpawners/ControlRigFunctionRefNodeSpawner.h
+++ b/Engine/Plugins/Experimental/ControlRig/Source/ControlRigEditor/Private/Graph/NodeSpawners/ControlRigFunctionRefNodeSpawner.h
@@ -26,6 +26,7 @@ class CONTROLRIGEDITOR_API UControlRigFunctionRefNodeSpawner : public UBlueprint
 	GENERATED_BODY()
 
 public:
+
 	/**
 	 * Creates a new UControlRigVariableNodeSpawner, charged with spawning 
 	 * a new member-variable node
@@ -34,17 +35,39 @@ public:
 	 */
 	static UControlRigFunctionRefNodeSpawner* CreateFromFunction(URigVMLibraryNode* InFunction);
 
+	/**
+	 * Creates a new UControlRigVariableNodeSpawner, charged with spawning 
+	 * a new member-variable node
+	 * 
+	 * @return A newly allocated instance of this class.
+	 */
+	static UControlRigFunctionRefNodeSpawner* CreateFromAssetData(const FAssetData& InAssetData, const FControlRigPublicFunctionData& InPublicFunction);
+
 	// UBlueprintNodeSpawner interface
 	virtual void Prime() override;
 	virtual FBlueprintNodeSignature GetSpawnerSignature() const override;
 	virtual FBlueprintActionUiSpec GetUiSpec(FBlueprintActionContext const& Context, FBindingSet const& Bindings) const override;
 	virtual UEdGraphNode* Invoke(UEdGraph* ParentGraph, FBindingSet const& Bindings, FVector2D const Location) const override;
+	virtual bool IsTemplateNodeFilteredOut(FBlueprintActionFilter const& Filter) const override;
 	// End UBlueprintNodeSpawner interface
 
 private:
-	/** The unit type we will spawn */
+
+	/** The unit type we will spawn [optional] */
+	UPROPERTY(Transient)
+	mutable TWeakObjectPtr<URigVMLibraryNode> ReferencedFunctionPtr;
+
+	/** The asset object path we'll spawn from [optional] */
+	UPROPERTY(Transient)
+	FName ReferencedAssetObjectPath;
+
+	/** The public function definition we will spawn from [optional] */
+	UPROPERTY(Transient)
+	FControlRigPublicFunctionData ReferencedPublicFunctionData;
+
+	/** Marked as true for local function definitions */
 	UPROPERTY(Transient)
-	TWeakObjectPtr<URigVMLibraryNode> ReferencedFunctionPtr;
+	bool bIsLocalFunction;
 
 	static UControlRigGraphNode* SpawnNode(UEdGraph* ParentGraph, UBlueprint* Blueprint, URigVMLibraryNode* InFunction, FVector2D const Location);
 
