CIS Unity buid Fix
#trivial

diff --git a/Engine/Plugins/Experimental/ControlRig/Source/ControlRigEditor/Private/EditMode/SControlRigSnapper.cpp b/Engine/Plugins/Experimental/ControlRig/Source/ControlRigEditor/Private/EditMode/SControlRigSnapper.cpp
--- a/Engine/Plugins/Experimental/ControlRig/Source/ControlRigEditor/Private/EditMode/SControlRigSnapper.cpp
+++ b/Engine/Plugins/Experimental/ControlRig/Source/ControlRigEditor/Private/EditMode/SControlRigSnapper.cpp
@@ -22,6 +22,10 @@
 #include "Tools/ControlRigSnapSettings.h"
 #include "LevelEditor.h"
 #include "Editor/SceneOutliner/Private/SSocketChooser.h"
+#include "Widgets/Views/SListView.h"
+#include "Widgets/Text/STextBlock.h"
+#include "Widgets/Views/STableRow.h"
+#include "Widgets/Input/SButton.h"
 
 #define LOCTEXT_NAMESPACE "ControlRigSnapper"
 
diff --git a/Engine/Plugins/Experimental/ControlRig/Source/ControlRigEditor/Private/EditMode/SControlRigSnapper.h b/Engine/Plugins/Experimental/ControlRig/Source/ControlRigEditor/Private/EditMode/SControlRigSnapper.h
--- a/Engine/Plugins/Experimental/ControlRig/Source/ControlRigEditor/Private/EditMode/SControlRigSnapper.h
+++ b/Engine/Plugins/Experimental/ControlRig/Source/ControlRigEditor/Private/EditMode/SControlRigSnapper.h
@@ -8,13 +8,9 @@
 #include "Widgets/SWidget.h"
 #include "Widgets/DeclarativeSyntaxSupport.h"
 #include "Widgets/SCompoundWidget.h"
-#include "Widgets/Input/SCheckBox.h"
-#include "Widgets/Input/SNumericEntryBox.h"
-#include "Widgets/Views/SListView.h"
-#include "Widgets/Text/STextBlock.h"
-#include "Widgets/Views/STableRow.h"
 #include "Tools/ControlRigSnapper.h"
 #include "Misc/FrameNumber.h"
+#include "IDetailsView.h"
 
 class UControlRig;
 class ISequencer;
