Control Rig Hierarchy: When changing parent weight also dirty offset transform

diff --git a/Engine/Plugins/Experimental/ControlRig/Source/ControlRig/Private/Rigs/RigHierarchy.cpp b/Engine/Plugins/Experimental/ControlRig/Source/ControlRig/Private/Rigs/RigHierarchy.cpp
--- a/Engine/Plugins/Experimental/ControlRig/Source/ControlRig/Private/Rigs/RigHierarchy.cpp
+++ b/Engine/Plugins/Experimental/ControlRig/Source/ControlRig/Private/Rigs/RigHierarchy.cpp
@@ -796,6 +796,11 @@ bool URigHierarchy::SetParentWeight(FRigBaseElement* InChild, int32 InParentInde
 			TargetWeight = InputWeight;
 			MultiParentElement->Parent.MarkDirty(GlobalType);
 
+			if(FRigControlElement* ControlElement = Cast<FRigControlElement>(MultiParentElement))
+			{
+				ControlElement->Offset.MarkDirty(GlobalType);
+			}
+
 			PropagateDirtyFlags(MultiParentElement, ERigTransformType::IsInitial(LocalType), bAffectChildren);
 			return true;
 		}
