Control Rig: Prepare framework for dynamic hierarchy demos

diff --git a/Engine/Plugins/Experimental/ControlRig/Source/ControlRig/Public/ControlRig.h b/Engine/Plugins/Experimental/ControlRig/Source/ControlRig/Public/ControlRig.h
--- a/Engine/Plugins/Experimental/ControlRig/Source/ControlRig/Public/ControlRig.h
+++ b/Engine/Plugins/Experimental/ControlRig/Source/ControlRig/Public/ControlRig.h
@@ -98,7 +98,7 @@ public:
 	void SetFramesPerSecond(float InFramesPerSecond);
 
 	/** Returns the current frames per second (this may change over time) */
-	UFUNCTION(BlueprintCallable, Category = "Control Rig")
+	UFUNCTION(BlueprintPure, Category = "Control Rig")
 	float GetCurrentFramesPerSecond() const;
 
 #if WITH_EDITOR
@@ -128,15 +128,15 @@ public:
 	FRigVMExternalVariable GetPublicVariableByName(const FName& InVariableName) const;
 
 	/** Returns the names of variables accessible in scripting */
-	UFUNCTION(BlueprintCallable, Category = "Control Rig", meta=(DisplayName="Get Variables"))
+	UFUNCTION(BlueprintPure, Category = "Control Rig", meta=(DisplayName="Get Variables"))
 	TArray<FName> GetScriptAccessibleVariables() const;
 
 	/** Returns the type of a given variable */
-	UFUNCTION(BlueprintCallable, Category = "Control Rig")
+	UFUNCTION(BlueprintPure, Category = "Control Rig")
 	FName GetVariableType(const FName& InVariableName) const;
 
 	/** Returns the value of a given variable as a string */
-	UFUNCTION(BlueprintCallable, Category = "Control Rig")
+	UFUNCTION(BlueprintPure, Category = "Control Rig")
 	FString GetVariableAsString(const FName& InVariableName) const;
 
 	/** Returns the value of a given variable as a string */
@@ -161,10 +161,10 @@ public:
 		return SupportsEvent(T::EventName);
 	}
 
-	UFUNCTION(BlueprintCallable, Category = "Control Rig")
+	UFUNCTION(BlueprintPure, Category = "Control Rig")
 	bool SupportsEvent(const FName& InEventName) const;
 
-	UFUNCTION(BlueprintCallable, Category = "Control Rig")
+	UFUNCTION(BlueprintPure, Category = "Control Rig")
 	TArray<FName> GetSupportedEvents() const;
 
 	/** Setup bindings to a runtime object (or clear by passing in nullptr). */
@@ -185,7 +185,7 @@ public:
 		return ObjectName;
 	}
 
-	UFUNCTION(BlueprintCallable, Category = "Control Rig")
+	UFUNCTION(BlueprintPure, Category = "Control Rig")
 	FORCEINLINE_DEBUGGABLE URigHierarchy* GetHierarchy()
 	{
 		if(DynamicHierarchy != nullptr && DynamicHierarchy->GetOuter() != this)
diff --git a/Engine/Plugins/Experimental/ControlRig/Source/ControlRigEditor/Private/EditMode/ControlRigEditMode.cpp b/Engine/Plugins/Experimental/ControlRig/Source/ControlRigEditor/Private/EditMode/ControlRigEditMode.cpp
--- a/Engine/Plugins/Experimental/ControlRig/Source/ControlRigEditor/Private/EditMode/ControlRigEditMode.cpp
+++ b/Engine/Plugins/Experimental/ControlRig/Source/ControlRigEditor/Private/EditMode/ControlRigEditMode.cpp
@@ -1818,7 +1818,6 @@ void FControlRigEditMode::OnHierarchyModified(ERigHierarchyNotification InNotif,
 		case ERigHierarchyNotification::ElementAdded:
 		case ERigHierarchyNotification::ElementRemoved:
 		case ERigHierarchyNotification::ElementRenamed:
-		case ERigHierarchyNotification::ParentChanged:
 		case ERigHierarchyNotification::HierarchyReset:
 		case ERigHierarchyNotification::ControlSettingChanged:
 		{
