Control Rig: Hierarchy: Rename space to null, horde preflight successful

diff --git a/Engine/Plugins/Experimental/ControlRig/Config/BaseControlRig.ini b/Engine/Plugins/Experimental/ControlRig/Config/BaseControlRig.ini
--- a/Engine/Plugins/Experimental/ControlRig/Config/BaseControlRig.ini
+++ b/Engine/Plugins/Experimental/ControlRig/Config/BaseControlRig.ini
@@ -11,4 +11,6 @@
 +PropertyRedirects=(OldName="RigBoneHierarchy.Joints",NewName="RigBoneHierarchy.Bones")
 +PropertyRedirects=(OldName="RigHierarchyContainer.BaseHierarchy",NewName="RigHierarchyContainer.BoneHierarchy")
 
-+ClassRedirects=(OldName="RigVMGraph",NewName="RigVMGraph",ValueChanges=(("Model","RigVMModel")))
\ No newline at end of file
++ClassRedirects=(OldName="RigVMGraph",NewName="RigVMGraph",ValueChanges=(("Model","RigVMModel")))
+
++EnumRedirects=(OldName="ERigElementType",ValueChanges=(("Space","Null"))
\ No newline at end of file
diff --git a/Engine/Plugins/Experimental/ControlRig/Source/ControlRig/Private/ControlRig.cpp b/Engine/Plugins/Experimental/ControlRig/Source/ControlRig/Private/ControlRig.cpp
--- a/Engine/Plugins/Experimental/ControlRig/Source/ControlRig/Private/ControlRig.cpp
+++ b/Engine/Plugins/Experimental/ControlRig/Source/ControlRig/Private/ControlRig.cpp
@@ -1500,7 +1500,7 @@ FName UControlRig::AddTransientControl(const FRigElementKey& InElement)
 			Parent = DynamicHierarchy->GetFirstParent(InElement);
 			break;
 		}
-		case ERigElementType::Space:
+		case ERigElementType::Null:
 		{
 			Parent = InElement;
 			break;
@@ -1578,7 +1578,7 @@ bool UControlRig::SetTransientControlValue(const FRigElementKey& InElement)
 				}
 			}
 		}
-		else if (InElement.Type == ERigElementType::Space)
+		else if (InElement.Type == ERigElementType::Null)
 		{
 			const FTransform GlobalTransform = DynamicHierarchy->GetGlobalTransform(InElement);
 			DynamicHierarchy->SetTransform(ControlElement, GlobalTransform, ERigTransformType::InitialGlobal, true, false);
@@ -1659,14 +1659,14 @@ FRigElementKey UControlRig::GetElementKeyFromTransientControl(const FRigElementK
 	}
 	
 	static FString ControlRigForElementBoneName;
-	static FString ControlRigForElementSpaceName;
+	static FString ControlRigForElementNullName;
 
 	if (ControlRigForElementBoneName.IsEmpty())
 	{
 		ControlRigForElementBoneName = FString::Printf(TEXT("ControlForRigElement_%s_"),
             *StaticEnum<ERigElementType>()->GetDisplayNameTextByValue((int64)ERigElementType::Bone).ToString());
-		ControlRigForElementSpaceName = FString::Printf(TEXT("ControlForRigElement_%s_"),
-            *StaticEnum<ERigElementType>()->GetDisplayNameTextByValue((int64)ERigElementType::Space).ToString());
+		ControlRigForElementNullName = FString::Printf(TEXT("ControlForRigElement_%s_"),
+            *StaticEnum<ERigElementType>()->GetDisplayNameTextByValue((int64)ERigElementType::Null).ToString());
 	}
 	
 	FString Name = InKey.Name.ToString();
@@ -1675,10 +1675,10 @@ FRigElementKey UControlRig::GetElementKeyFromTransientControl(const FRigElementK
 		Name.RightChopInline(ControlRigForElementBoneName.Len());
 		return FRigElementKey(*Name, ERigElementType::Bone);
 	}
-	if(Name.StartsWith(ControlRigForElementSpaceName))
+	if(Name.StartsWith(ControlRigForElementNullName))
 	{
-		Name.RightChopInline(ControlRigForElementSpaceName.Len());
-		return FRigElementKey(*Name, ERigElementType::Space);
+		Name.RightChopInline(ControlRigForElementNullName.Len());
+		return FRigElementKey(*Name, ERigElementType::Null);
 	}
 	
 	return FRigElementKey();;
diff --git a/Engine/Plugins/Experimental/ControlRig/Source/ControlRig/Private/ControlRigComponent.cpp b/Engine/Plugins/Experimental/ControlRig/Source/ControlRig/Private/ControlRigComponent.cpp
--- a/Engine/Plugins/Experimental/ControlRig/Source/ControlRig/Private/ControlRigComponent.cpp
+++ b/Engine/Plugins/Experimental/ControlRig/Source/ControlRig/Private/ControlRigComponent.cpp
@@ -959,15 +959,15 @@ FTransform UControlRigComponent::GetSpaceTransform(FName SpaceName, EControlRigC
 {
 	if (UControlRig* CR = SetupControlRigIfRequired())
 	{
-		if(FRigSpaceElement* SpaceElement = CR->GetHierarchy()->Find<FRigSpaceElement>(FRigElementKey(SpaceName, ERigElementType::Control)))
+		if(FRigNullElement* NullElement = CR->GetHierarchy()->Find<FRigNullElement>(FRigElementKey(SpaceName, ERigElementType::Control)))
 		{
 			if (Space == EControlRigComponentSpace::LocalSpace)
 			{
-				return CR->GetHierarchy()->GetTransform(SpaceElement, ERigTransformType::CurrentLocal);
+				return CR->GetHierarchy()->GetTransform(NullElement, ERigTransformType::CurrentLocal);
 			}
 			else
 			{
-				FTransform RootTransform = CR->GetHierarchy()->GetTransform(SpaceElement, ERigTransformType::CurrentGlobal);
+				FTransform RootTransform = CR->GetHierarchy()->GetTransform(NullElement, ERigTransformType::CurrentGlobal);
 				ConvertTransformFromRigSpace(RootTransform, Space);
 				return RootTransform;
 			}
@@ -981,15 +981,15 @@ FTransform UControlRigComponent::GetInitialSpaceTransform(FName SpaceName, ECont
 {
 	if (UControlRig* CR = SetupControlRigIfRequired())
 	{
-		if(FRigSpaceElement* SpaceElement = CR->GetHierarchy()->Find<FRigSpaceElement>(FRigElementKey(SpaceName, ERigElementType::Control)))
+		if(FRigNullElement* NullElement = CR->GetHierarchy()->Find<FRigNullElement>(FRigElementKey(SpaceName, ERigElementType::Control)))
 		{
 			if (Space == EControlRigComponentSpace::LocalSpace)
 			{
-				return CR->GetHierarchy()->GetTransform(SpaceElement, ERigTransformType::InitialLocal);
+				return CR->GetHierarchy()->GetTransform(NullElement, ERigTransformType::InitialLocal);
 			}
 			else
 			{
-				FTransform RootTransform = CR->GetHierarchy()->GetTransform(SpaceElement, ERigTransformType::InitialGlobal);
+				FTransform RootTransform = CR->GetHierarchy()->GetTransform(NullElement, ERigTransformType::InitialGlobal);
 				ConvertTransformFromRigSpace(RootTransform, Space);
 				return RootTransform;
 			}
@@ -1003,16 +1003,16 @@ void UControlRigComponent::SetInitialSpaceTransform(FName SpaceName, FTransform
 {
 	if (UControlRig* CR = SetupControlRigIfRequired())
 	{
-		if(FRigSpaceElement* SpaceElement = CR->GetHierarchy()->Find<FRigSpaceElement>(FRigElementKey(SpaceName, ERigElementType::Control)))
+		if(FRigNullElement* NullElement = CR->GetHierarchy()->Find<FRigNullElement>(FRigElementKey(SpaceName, ERigElementType::Control)))
 		{
 			if (Space == EControlRigComponentSpace::LocalSpace)
 			{
-				CR->GetHierarchy()->SetTransform(SpaceElement, InitialTransform, ERigTransformType::InitialLocal, true, false);
+				CR->GetHierarchy()->SetTransform(NullElement, InitialTransform, ERigTransformType::InitialLocal, true, false);
 			}
 			else
 			{
 				ConvertTransformToRigSpace(InitialTransform, Space);
-				CR->GetHierarchy()->SetTransform(SpaceElement, InitialTransform, ERigTransformType::InitialGlobal, true, false);
+				CR->GetHierarchy()->SetTransform(NullElement, InitialTransform, ERigTransformType::InitialGlobal, true, false);
 			}
 		}
 	}
@@ -1254,7 +1254,7 @@ void UControlRigComponent::TransferOutputs()
 
 			if (MappedElement.ElementType == ERigElementType::Bone ||
 				MappedElement.ElementType == ERigElementType::Control ||
-				MappedElement.ElementType == ERigElementType::Space)
+				MappedElement.ElementType == ERigElementType::Null)
 			{
 				FTransform Transform = ControlRig->GetHierarchy()->GetGlobalTransform(MappedElement.ElementIndex);
 				ConvertTransformFromRigSpace(Transform, MappedElement.Space);
diff --git a/Engine/Plugins/Experimental/ControlRig/Source/ControlRig/Private/Rigs/AdditiveControlRig.cpp b/Engine/Plugins/Experimental/ControlRig/Source/ControlRig/Private/Rigs/AdditiveControlRig.cpp
--- a/Engine/Plugins/Experimental/ControlRig/Source/ControlRig/Private/Rigs/AdditiveControlRig.cpp
+++ b/Engine/Plugins/Experimental/ControlRig/Source/ControlRig/Private/Rigs/AdditiveControlRig.cpp
@@ -28,11 +28,11 @@ FName UAdditiveControlRig::GetControlName(const FName& InBoneName)
 	return NAME_None;
 }
 
-FName UAdditiveControlRig::GetSpaceName(const FName& InBoneName)
+FName UAdditiveControlRig::GetNullName(const FName& InBoneName)
 {
 	if (InBoneName != NAME_None)
 	{
-		return FName(*(InBoneName.ToString() + TEXT("_SPACE")));
+		return FName(*(InBoneName.ToString() + TEXT("_NULL")));
 	}
 
 	// if bone name is coming as none, we don't append
@@ -115,27 +115,27 @@ void UAdditiveControlRig::CreateRigElements(const FReferenceSkeleton& InReferenc
     {
 		const FName BoneName = BoneElement->GetName();
 		const int32 ParentIndex = GetHierarchy()->GetFirstParent(BoneElement->GetIndex());
-		const FName SpaceName = GetSpaceName(BoneName);// name conflict?
+		const FName NullName = GetNullName(BoneName);// name conflict?
 		const FName ControlName = GetControlName(BoneName); // name conflict?
 
-		FRigElementKey SpaceKey;
+		FRigElementKey NullKey;
 		
 		if (ParentIndex != INDEX_NONE)
 		{
 			FTransform GlobalTransform = GetHierarchy()->GetGlobalTransform(BoneElement->GetIndex());
 			FTransform ParentTransform = GetHierarchy()->GetGlobalTransform(ParentIndex);
 			FTransform LocalTransform = GlobalTransform.GetRelativeTransform(ParentTransform);
-			SpaceKey = Controller->AddSpace(SpaceName, GetHierarchy()->GetKey(ParentIndex), LocalTransform, false, false);
+			NullKey = Controller->AddNull(NullName, GetHierarchy()->GetKey(ParentIndex), LocalTransform, false, false);
 		}
 		else
 		{
 			FTransform GlobalTransform = GetHierarchy()->GetGlobalTransform(BoneElement->GetIndex());
-			SpaceKey = Controller->AddSpace(SpaceName, FRigElementKey(), GlobalTransform, true, false);
+			NullKey = Controller->AddNull(NullName, FRigElementKey(), GlobalTransform, true, false);
 		}
 
 		FRigControlSettings Settings;
 		Settings.DisplayName = BoneName;
-		Controller->AddControl(ControlName, SpaceKey, Settings, FRigControlValue::Make(FTransform::Identity), FTransform::Identity, FTransform::Identity, false);
+		Controller->AddControl(ControlName, NullKey, Settings, FRigControlValue::Make(FTransform::Identity), FTransform::Identity, FTransform::Identity, false);
 
 		return true;
 	});
diff --git a/Engine/Plugins/Experimental/ControlRig/Source/ControlRig/Private/Rigs/FKControlRig.cpp b/Engine/Plugins/Experimental/ControlRig/Source/ControlRig/Private/Rigs/FKControlRig.cpp
--- a/Engine/Plugins/Experimental/ControlRig/Source/ControlRig/Private/Rigs/FKControlRig.cpp
+++ b/Engine/Plugins/Experimental/ControlRig/Source/ControlRig/Private/Rigs/FKControlRig.cpp
@@ -29,11 +29,11 @@ FName UFKControlRig::GetControlName(const FName& InBoneName)
 	return NAME_None;
 }
 
-FName UFKControlRig::GetSpaceName(const FName& InBoneName)
+FName UFKControlRig::GetNullName(const FName& InBoneName)
 {
 	if (InBoneName != NAME_None)
 	{
-		return FName(*(InBoneName.ToString() + TEXT("_SPACE")));
+		return FName(*(InBoneName.ToString() + TEXT("_NULL")));
 	}
 
 	// if bone name is coming as none, we don't append
@@ -226,10 +226,10 @@ void UFKControlRig::CreateRigElements(const FReferenceSkeleton& InReferenceSkele
 	{
 		const FName BoneName = BoneElement->GetName();
 		const int32 ParentIndex = GetHierarchy()->GetFirstParent(BoneElement->GetIndex());
-		const FName SpaceName = GetSpaceName(BoneName);// name conflict?
+		const FName NullName = GetNullName(BoneName);// name conflict?
 		const FName ControlName = GetControlName(BoneName); // name conflict?
 
-		FRigElementKey SpaceKey;
+		FRigElementKey NullKey;
 
 		FTransform LocalTransform;
 		if (ParentIndex != INDEX_NONE)
@@ -237,18 +237,18 @@ void UFKControlRig::CreateRigElements(const FReferenceSkeleton& InReferenceSkele
 			FTransform GlobalTransform = GetHierarchy()->GetGlobalTransform(BoneElement->GetIndex());
 			FTransform ParentTransform = GetHierarchy()->GetGlobalTransform(ParentIndex);
 			LocalTransform = GlobalTransform.GetRelativeTransform(ParentTransform);
-			SpaceKey = Controller->AddSpace(SpaceName, GetHierarchy()->GetKey(ParentIndex), FTransform::Identity, false, false);
+			NullKey = Controller->AddNull(NullName, GetHierarchy()->GetKey(ParentIndex), FTransform::Identity, false, false);
 		}
 		else
 		{
 			LocalTransform = GetHierarchy()->GetLocalTransform(BoneElement->GetIndex());
-			SpaceKey = Controller->AddSpace(SpaceName, FRigElementKey(), FTransform::Identity, true, false);
+			NullKey = Controller->AddNull(NullName, FRigElementKey(), FTransform::Identity, true, false);
 		}
 
 		FRigControlSettings Settings;
 		Settings.ControlType = ERigControlType::EulerTransform;
 		Settings.DisplayName = BoneName;
-		const FRigElementKey ControlKey = Controller->AddControl(ControlName, SpaceKey, Settings, FRigControlValue::Make(FEulerTransform::Identity), FTransform::Identity, FTransform::Identity, false);
+		const FRigElementKey ControlKey = Controller->AddControl(ControlName, NullKey, Settings, FRigControlValue::Make(FEulerTransform::Identity), FTransform::Identity, FTransform::Identity, false);
 		GetHierarchy()->SetLocalTransform(ControlKey, LocalTransform, true, true, false);
 
 		return true;
diff --git a/Engine/Plugins/Experimental/ControlRig/Source/ControlRig/Private/Rigs/RigHierarchy.cpp b/Engine/Plugins/Experimental/ControlRig/Source/ControlRig/Private/Rigs/RigHierarchy.cpp
--- a/Engine/Plugins/Experimental/ControlRig/Source/ControlRig/Private/Rigs/RigHierarchy.cpp
+++ b/Engine/Plugins/Experimental/ControlRig/Source/ControlRig/Private/Rigs/RigHierarchy.cpp
@@ -109,9 +109,9 @@ void URigHierarchy::Load(FArchive& Ar)
 				Element = new FRigBoneElement();
 				break;
 			}
-			case ERigElementType::Space:
+			case ERigElementType::Null:
 			{
-				Element = new FRigSpaceElement();
+				Element = new FRigNullElement();
 				break;
 			}
 			case ERigElementType::Control:
diff --git a/Engine/Plugins/Experimental/ControlRig/Source/ControlRig/Private/Rigs/RigHierarchyController.cpp b/Engine/Plugins/Experimental/ControlRig/Source/ControlRig/Private/Rigs/RigHierarchyController.cpp
--- a/Engine/Plugins/Experimental/ControlRig/Source/ControlRig/Private/Rigs/RigHierarchyController.cpp
+++ b/Engine/Plugins/Experimental/ControlRig/Source/ControlRig/Private/Rigs/RigHierarchyController.cpp
@@ -173,7 +173,7 @@ FRigElementKey URigHierarchyController::AddBone(FName InName, FRigElementKey InP
 	return NewElement->Key;
 }
 
-FRigElementKey URigHierarchyController::AddSpace(FName InName, FRigElementKey InParent, FTransform InTransform,
+FRigElementKey URigHierarchyController::AddNull(FName InName, FRigElementKey InParent, FTransform InTransform,
 	bool bTransformInGlobal, bool bSetupUndo)
 {
 	if(!IsValid())
@@ -185,13 +185,13 @@ FRigElementKey URigHierarchyController::AddSpace(FName InName, FRigElementKey In
 	TSharedPtr<FScopedTransaction> TransactionPtr;
 	if(bSetupUndo)
 	{
-		TransactionPtr = MakeShared<FScopedTransaction>(NSLOCTEXT("RigHierarchyController", "Add Space", "Add Space"));
+		TransactionPtr = MakeShared<FScopedTransaction>(NSLOCTEXT("RigHierarchyController", "Add Null", "Add Null"));
 		Hierarchy->Modify();
 	}
 #endif
 
-	FRigSpaceElement* NewElement = new FRigSpaceElement();
-	NewElement->Key.Type = ERigElementType::Space;
+	FRigNullElement* NewElement = new FRigNullElement();
+	NewElement->Key.Type = ERigElementType::Null;
 	NewElement->Key.Name = Hierarchy->GetSafeNewName(InName.ToString(), NewElement->Key.Type);
 	AddElement(NewElement, Hierarchy->Get(Hierarchy->GetIndex(InParent)), false);
 
@@ -740,10 +740,10 @@ FString URigHierarchyController::ExportToText(TArray<FRigElementKey> InKeys) con
 				FRigControlElement::StaticStruct()->ExportText(PerElementData.Content, Element, &DefaultElement, nullptr, PPF_None, nullptr);
 				break;
 			}
-			case ERigElementType::Space:
+			case ERigElementType::Null:
 			{
-				FRigSpaceElement DefaultElement;
-				FRigSpaceElement::StaticStruct()->ExportText(PerElementData.Content, Element, &DefaultElement, nullptr, PPF_None, nullptr);
+				FRigNullElement DefaultElement;
+				FRigNullElement::StaticStruct()->ExportText(PerElementData.Content, Element, &DefaultElement, nullptr, PPF_None, nullptr);
 				break;
 			}
 			case ERigElementType::Curve:
@@ -849,10 +849,10 @@ TArray<FRigElementKey> URigHierarchyController::ImportFromText(FString InContent
 				FRigBoneElement::StaticStruct()->ImportText(*PerElementData.Content, NewElement, nullptr, EPropertyPortFlags::PPF_None, &ErrorPipe, FRigBoneElement::StaticStruct()->GetName(), true);
 				break;
 			}
-			case ERigElementType::Space:
+			case ERigElementType::Null:
 			{
-				NewElement = new FRigSpaceElement();
-				FRigSpaceElement::StaticStruct()->ImportText(*PerElementData.Content, NewElement, nullptr, EPropertyPortFlags::PPF_None, &ErrorPipe, FRigSpaceElement::StaticStruct()->GetName(), true);
+				NewElement = new FRigNullElement();
+				FRigNullElement::StaticStruct()->ImportText(*PerElementData.Content, NewElement, nullptr, EPropertyPortFlags::PPF_None, &ErrorPipe, FRigNullElement::StaticStruct()->GetName(), true);
 				break;
 			}
 			case ERigElementType::Control:
diff --git a/Engine/Plugins/Experimental/ControlRig/Source/ControlRig/Private/Rigs/RigHierarchyDefines.cpp b/Engine/Plugins/Experimental/ControlRig/Source/ControlRig/Private/Rigs/RigHierarchyDefines.cpp
--- a/Engine/Plugins/Experimental/ControlRig/Source/ControlRig/Private/Rigs/RigHierarchyDefines.cpp
+++ b/Engine/Plugins/Experimental/ControlRig/Source/ControlRig/Private/Rigs/RigHierarchyDefines.cpp
@@ -76,7 +76,7 @@ FRigElementKeyCollection FRigElementKeyCollection::MakeFromChildren(
 
 	bool bAddBones = (InElementTypes & (uint8)ERigElementType::Bone) == (uint8)ERigElementType::Bone;
 	bool bAddControls = (InElementTypes & (uint8)ERigElementType::Control) == (uint8)ERigElementType::Control;
-	bool bAddSpaces = (InElementTypes & (uint8)ERigElementType::Space) == (uint8)ERigElementType::Space;
+	bool bAddNulls = (InElementTypes & (uint8)ERigElementType::Null) == (uint8)ERigElementType::Null;
 	bool bAddCurves = (InElementTypes & (uint8)ERigElementType::Curve) == (uint8)ERigElementType::Curve;
 
 	for (int32 ParentIndex = 0; ParentIndex < ParentKeys.Num(); ParentIndex++)
diff --git a/Engine/Plugins/Experimental/ControlRig/Source/ControlRig/Private/Units/Hierarchy/RigUnit_GetSpaceTransform.cpp b/Engine/Plugins/Experimental/ControlRig/Source/ControlRig/Private/Units/Hierarchy/RigUnit_GetSpaceTransform.cpp
--- a/Engine/Plugins/Experimental/ControlRig/Source/ControlRig/Private/Units/Hierarchy/RigUnit_GetSpaceTransform.cpp
+++ b/Engine/Plugins/Experimental/ControlRig/Source/ControlRig/Private/Units/Hierarchy/RigUnit_GetSpaceTransform.cpp
@@ -17,7 +17,7 @@ FRigUnit_GetSpaceTransform_Execute()
 			}
 			case EControlRigState::Update:
 			{
-				if (CachedSpaceIndex.UpdateCache(FRigElementKey(Space, ERigElementType::Space), Hierarchy))
+				if (CachedSpaceIndex.UpdateCache(FRigElementKey(Space, ERigElementType::Null), Hierarchy))
 				{
 					switch (SpaceType)
 					{
diff --git a/Engine/Plugins/Experimental/ControlRig/Source/ControlRig/Private/Units/Hierarchy/RigUnit_ParentSwitchConstraint.h b/Engine/Plugins/Experimental/ControlRig/Source/ControlRig/Private/Units/Hierarchy/RigUnit_ParentSwitchConstraint.h
--- a/Engine/Plugins/Experimental/ControlRig/Source/ControlRig/Private/Units/Hierarchy/RigUnit_ParentSwitchConstraint.h
+++ b/Engine/Plugins/Experimental/ControlRig/Source/ControlRig/Private/Units/Hierarchy/RigUnit_ParentSwitchConstraint.h
@@ -16,7 +16,7 @@ struct CONTROLRIG_API FRigUnit_ParentSwitchConstraint : public FRigUnitMutable
 	GENERATED_BODY()
 
 	FRigUnit_ParentSwitchConstraint()
-		: Subject(NAME_None, ERigElementType::Space)
+		: Subject(NAME_None, ERigElementType::Null)
 		, ParentIndex(0)
 		, Weight(1.0f)
 		, Switched(false)
diff --git a/Engine/Plugins/Experimental/ControlRig/Source/ControlRig/Private/Units/Hierarchy/RigUnit_SetSpaceInitialTransform.cpp b/Engine/Plugins/Experimental/ControlRig/Source/ControlRig/Private/Units/Hierarchy/RigUnit_SetSpaceInitialTransform.cpp
--- a/Engine/Plugins/Experimental/ControlRig/Source/ControlRig/Private/Units/Hierarchy/RigUnit_SetSpaceInitialTransform.cpp
+++ b/Engine/Plugins/Experimental/ControlRig/Source/ControlRig/Private/Units/Hierarchy/RigUnit_SetSpaceInitialTransform.cpp
@@ -20,7 +20,7 @@ FRigUnit_SetSpaceInitialTransform_Execute()
 			}
 			case EControlRigState::Update:
 			{
-				const FRigElementKey SpaceKey(SpaceName, ERigElementType::Space);
+				const FRigElementKey SpaceKey(SpaceName, ERigElementType::Null);
 				if (!CachedSpaceIndex.UpdateCache(SpaceKey, Hierarchy))
 				{
 					UE_CONTROLRIG_RIGUNIT_REPORT_WARNING(TEXT("Space '%s' is not valid."), *SpaceName.ToString());
diff --git a/Engine/Plugins/Experimental/ControlRig/Source/ControlRig/Private/Units/Hierarchy/RigUnit_SetSpaceTransform.cpp b/Engine/Plugins/Experimental/ControlRig/Source/ControlRig/Private/Units/Hierarchy/RigUnit_SetSpaceTransform.cpp
--- a/Engine/Plugins/Experimental/ControlRig/Source/ControlRig/Private/Units/Hierarchy/RigUnit_SetSpaceTransform.cpp
+++ b/Engine/Plugins/Experimental/ControlRig/Source/ControlRig/Private/Units/Hierarchy/RigUnit_SetSpaceTransform.cpp
@@ -19,7 +19,7 @@ FRigUnit_SetSpaceTransform_Execute()
 			}
 			case EControlRigState::Update:
 			{
-				const FRigElementKey SpaceKey(Space, ERigElementType::Space);
+				const FRigElementKey SpaceKey(Space, ERigElementType::Null);
 
 				if (CachedSpaceIndex.UpdateCache(SpaceKey, Hierarchy))
 				{
diff --git a/Engine/Plugins/Experimental/ControlRig/Source/ControlRig/Public/Rigs/AdditiveControlRig.h b/Engine/Plugins/Experimental/ControlRig/Source/ControlRig/Public/Rigs/AdditiveControlRig.h
--- a/Engine/Plugins/Experimental/ControlRig/Source/ControlRig/Public/Rigs/AdditiveControlRig.h
+++ b/Engine/Plugins/Experimental/ControlRig/Source/ControlRig/Public/Rigs/AdditiveControlRig.h
@@ -23,7 +23,7 @@ public:
 
 	// utility function to 
 	static FName GetControlName(const FName& InBoneName);
-	static FName GetSpaceName(const FName& InBoneName);
+	static FName GetNullName(const FName& InBoneName);
 
 private:
 	// custom units that are used to execute this rig
diff --git a/Engine/Plugins/Experimental/ControlRig/Source/ControlRig/Public/Rigs/FKControlRig.h b/Engine/Plugins/Experimental/ControlRig/Source/ControlRig/Public/Rigs/FKControlRig.h
--- a/Engine/Plugins/Experimental/ControlRig/Source/ControlRig/Public/Rigs/FKControlRig.h
+++ b/Engine/Plugins/Experimental/ControlRig/Source/ControlRig/Public/Rigs/FKControlRig.h
@@ -46,7 +46,7 @@ public:
 
 	// utility function to 
 	static FName GetControlName(const FName& InBoneName);
-	static FName GetSpaceName(const FName& InBoneName);
+	static FName GetNullName(const FName& InBoneName);
 
 
 	TArray<FName> GetControlNames();
diff --git a/Engine/Plugins/Experimental/ControlRig/Source/ControlRig/Public/Rigs/RigControlHierarchy.h b/Engine/Plugins/Experimental/ControlRig/Source/ControlRig/Public/Rigs/RigControlHierarchy.h
--- a/Engine/Plugins/Experimental/ControlRig/Source/ControlRig/Public/Rigs/RigControlHierarchy.h
+++ b/Engine/Plugins/Experimental/ControlRig/Source/ControlRig/Public/Rigs/RigControlHierarchy.h
@@ -177,7 +177,7 @@ struct CONTROLRIG_API FRigControl : public FRigElement
 	{
 		if (SpaceIndex != INDEX_NONE || bForce)
 		{
-			return FRigElementKey(SpaceName, ERigElementType::Space);
+			return FRigElementKey(SpaceName, ERigElementType::Null);
 		}
 		return FRigElementKey();
 	}
diff --git a/Engine/Plugins/Experimental/ControlRig/Source/ControlRig/Public/Rigs/RigHierarchy.h b/Engine/Plugins/Experimental/ControlRig/Source/ControlRig/Public/Rigs/RigHierarchy.h
--- a/Engine/Plugins/Experimental/ControlRig/Source/ControlRig/Public/Rigs/RigHierarchy.h
+++ b/Engine/Plugins/Experimental/ControlRig/Source/ControlRig/Public/Rigs/RigHierarchy.h
@@ -517,9 +517,9 @@ public:
 	 * Returns all Space elements
 	 * @param bTraverse Returns the elements in order of a depth first traversal
 	 */
-	FORCEINLINE TArray<FRigSpaceElement*> GetSpaces(bool bTraverse = false) const
+	FORCEINLINE TArray<FRigNullElement*> GetSpaces(bool bTraverse = false) const
 	{
-		return GetElementsOfType<FRigSpaceElement>(bTraverse);
+		return GetElementsOfType<FRigNullElement>(bTraverse);
 	}
 
 	/**
@@ -529,7 +529,7 @@ public:
 	UFUNCTION(BlueprintCallable, Category = URigHierarchy, meta = (DisplayName = "Get Spaces", ScriptName = "GetSpaces"))
 	FORCEINLINE TArray<FRigElementKey> GetSpaceKeys(bool bTraverse = true) const
 	{
-		return GetKeysOfType<FRigSpaceElement>(bTraverse);
+		return GetKeysOfType<FRigNullElement>(bTraverse);
 	}
 
 	/**
diff --git a/Engine/Plugins/Experimental/ControlRig/Source/ControlRig/Public/Rigs/RigHierarchyController.h b/Engine/Plugins/Experimental/ControlRig/Source/ControlRig/Public/Rigs/RigHierarchyController.h
--- a/Engine/Plugins/Experimental/ControlRig/Source/ControlRig/Public/Rigs/RigHierarchyController.h
+++ b/Engine/Plugins/Experimental/ControlRig/Source/ControlRig/Public/Rigs/RigHierarchyController.h
@@ -81,16 +81,16 @@ public:
 	FRigElementKey AddBone(FName InName, FRigElementKey InParent, FTransform InTransform, bool bTransformInGlobal = true, ERigBoneType InBoneType = ERigBoneType::User, bool bSetupUndo = false);
 
 	/**
-	 * Adds a space to the hierarchy
-	 * @param InName The suggested name of the new space - will eventually be corrected by the namespace
-	 * @param InParent The (optional) parent of the new space. If you don't need a parent, pass FRigElementKey()
-	 * @param InTransform The transform for the new space - either in local or global space, based on bTransformInGlobal
-	 * @param bTransformInGlobal Set this to true if the Transform passed is expressed in global space, false for local space.
+	 * Adds a null to the hierarchy
+	 * @param InName The suggested name of the new null - will eventually be corrected by the namespace
+	 * @param InParent The (optional) parent of the new null. If you don't need a parent, pass FRigElementKey()
+	 * @param InTransform The transform for the new null - either in local or global null, based on bTransformInGlobal
+	 * @param bTransformInGlobal Set this to true if the Transform passed is expressed in global null, false for local null.
 	 * @param bSetupUndo If set to true the stack will record the change for undo / redo
-	 * @return The key for the newly created space.
+	 * @return The key for the newly created null.
 	 */
 	UFUNCTION(BlueprintCallable, Category = URigHierarchyController)
-    FRigElementKey AddSpace(FName InName, FRigElementKey InParent, FTransform InTransform, bool bTransformInGlobal = true, bool bSetupUndo = false);
+    FRigElementKey AddNull(FName InName, FRigElementKey InParent, FTransform InTransform, bool bTransformInGlobal = true, bool bSetupUndo = false);
 
 	/**
 	 * Adds a control to the hierarchy
diff --git a/Engine/Plugins/Experimental/ControlRig/Source/ControlRig/Public/Rigs/RigHierarchyDefines.h b/Engine/Plugins/Experimental/ControlRig/Source/ControlRig/Public/Rigs/RigHierarchyDefines.h
--- a/Engine/Plugins/Experimental/ControlRig/Source/ControlRig/Public/Rigs/RigHierarchyDefines.h
+++ b/Engine/Plugins/Experimental/ControlRig/Source/ControlRig/Public/Rigs/RigHierarchyDefines.h
@@ -18,12 +18,12 @@ enum class ERigElementType : uint8
 {
 	None,
 	Bone = 0x001,
-	Space = 0x002,
+	Null = 0x002,
 	Control = 0x004,
 	Curve = 0x008,
 	RigidBody = 0x010,
 	Socket = 0x020,
-	All = Bone | Space | Control | Curve | RigidBody | Socket UMETA(Hidden),
+	All = Bone | Null | Control | Curve | RigidBody | Socket UMETA(Hidden),
 };
 
 UENUM(BlueprintType)
@@ -873,7 +873,7 @@ public:
 			{
 				return FString::Printf(TEXT("Bone(%s)"), *Name.ToString());
 			}
-			case ERigElementType::Space:
+			case ERigElementType::Null:
 			{
 				return FString::Printf(TEXT("Space(%s)"), *Name.ToString());
 			}
diff --git a/Engine/Plugins/Experimental/ControlRig/Source/ControlRig/Public/Rigs/RigHierarchyElements.h b/Engine/Plugins/Experimental/ControlRig/Source/ControlRig/Public/Rigs/RigHierarchyElements.h
--- a/Engine/Plugins/Experimental/ControlRig/Source/ControlRig/Public/Rigs/RigHierarchyElements.h
+++ b/Engine/Plugins/Experimental/ControlRig/Source/ControlRig/Public/Rigs/RigHierarchyElements.h
@@ -577,7 +577,7 @@ protected:
 	FORCEINLINE static bool IsClassOf(const FRigBaseElement* InElement)
 	{
 		return InElement->GetType() == ERigElementType::Bone ||
-			InElement->GetType() == ERigElementType::Space ||
+			InElement->GetType() == ERigElementType::Null ||
 			InElement->GetType() == ERigElementType::Control ||
 			InElement->GetType() == ERigElementType::RigidBody ||
 			InElement->GetType() == ERigElementType::Socket;
@@ -657,7 +657,7 @@ protected:
 
 	FORCEINLINE static bool IsClassOf(const FRigBaseElement* InElement)
 	{
-		return InElement->GetType() == ERigElementType::Space ||
+		return InElement->GetType() == ERigElementType::Null ||
 			InElement->GetType() == ERigElementType::Control;
 	}
 
@@ -705,26 +705,26 @@ protected:
 };
 
 USTRUCT(BlueprintType)
-struct CONTROLRIG_API FRigSpaceElement : public FRigMultiParentElement
+struct CONTROLRIG_API FRigNullElement : public FRigMultiParentElement
 {
 public:
 	
 	GENERATED_BODY()
-	DECLARE_RIG_ELEMENT_METHODS(FRigSpaceElement)
+	DECLARE_RIG_ELEMENT_METHODS(FRigNullElement)
 
-	FRigSpaceElement()
+	FRigNullElement()
     : FRigMultiParentElement()
 	{
-		Key.Type = ERigElementType::Space; 
+		Key.Type = ERigElementType::Null; 
 	}
 
-	virtual ~FRigSpaceElement(){}
+	virtual ~FRigNullElement(){}
 
 protected:
 	
 	FORCEINLINE static bool IsClassOf(const FRigBaseElement* InElement)
 	{
-		return InElement->GetType() == ERigElementType::Space;
+		return InElement->GetType() == ERigElementType::Null;
 	}
 
 	friend struct FRigBaseElement;
diff --git a/Engine/Plugins/Experimental/ControlRig/Source/ControlRig/Public/Rigs/RigSpaceHierarchy.h b/Engine/Plugins/Experimental/ControlRig/Source/ControlRig/Public/Rigs/RigSpaceHierarchy.h
--- a/Engine/Plugins/Experimental/ControlRig/Source/ControlRig/Public/Rigs/RigSpaceHierarchy.h
+++ b/Engine/Plugins/Experimental/ControlRig/Source/ControlRig/Public/Rigs/RigSpaceHierarchy.h
@@ -59,7 +59,7 @@ struct CONTROLRIG_API FRigSpace : public FRigElement
 
 	FORCEINLINE virtual ERigElementType GetElementType() const override
 	{
-		return ERigElementType::Space;
+		return ERigElementType::Null;
 	}
 
 	FORCEINLINE virtual FRigElementKey GetParentElementKey(bool bForce = false) const
@@ -78,7 +78,7 @@ struct CONTROLRIG_API FRigSpace : public FRigElement
 				}
 				case ERigSpaceType::Space:
 				{
-					return FRigElementKey(ParentName, ERigElementType::Space);
+					return FRigElementKey(ParentName, ERigElementType::Null);
 				}
 				default:
 				{
diff --git a/Engine/Plugins/Experimental/ControlRig/Source/ControlRigDeveloper/Private/ControlRigBlueprint.cpp b/Engine/Plugins/Experimental/ControlRig/Source/ControlRigDeveloper/Private/ControlRigBlueprint.cpp
--- a/Engine/Plugins/Experimental/ControlRig/Source/ControlRigDeveloper/Private/ControlRigBlueprint.cpp
+++ b/Engine/Plugins/Experimental/ControlRig/Source/ControlRigDeveloper/Private/ControlRigBlueprint.cpp
@@ -366,7 +366,7 @@ void UControlRigBlueprint::PostLoad()
 		}
 		for(const FRigSpace& Space : HierarchyContainer_DEPRECATED.SpaceHierarchy)
 		{
-			HierarchyController->AddSpace(Space.Name, FRigElementKey(), Space.InitialTransform, false, false);
+			HierarchyController->AddNull(Space.Name, FRigElementKey(), Space.InitialTransform, false, false);
 		}
 		for(const FRigControl& Control : HierarchyContainer_DEPRECATED.ControlHierarchy)
 		{
@@ -476,6 +476,44 @@ void UControlRigBlueprint::PostLoad()
 		GetOrCreateController(GraphToDetach)->ReattachLinksToPinObjects(true /* follow redirectors */);
 	}
 
+	// perform backwards compat value upgrades
+	TArray<URigVMGraph*> GraphsToValidate;
+	GraphsToValidate.Add(GetModel());
+	GraphsToValidate.Add(GetLocalFunctionLibrary());
+	for (int32 GraphIndex = 0; GraphIndex < GraphsToValidate.Num(); GraphIndex++)
+	{
+		URigVMGraph* GraphToValidate = GraphsToValidate[GraphIndex];
+		if(GraphToValidate == nullptr)
+		{
+			continue;
+		}
+			
+		for(URigVMNode* Node : GraphToValidate->GetNodes())
+		{
+			if(URigVMCollapseNode* CollapseNode = Cast<URigVMCollapseNode>(Node))
+			{
+				GraphsToValidate.AddUnique(CollapseNode->GetContainedGraph());
+			}
+
+			TArray<URigVMPin*> Pins = Node->GetAllPinsRecursively();
+			for(URigVMPin* Pin : Pins)
+			{
+				if(Pin->GetCPPTypeObject() == StaticEnum<ERigElementType>())
+				{
+					if(Pin->GetDefaultValue() == TEXT("Space"))
+					{
+						if(URigVMController* Controller = GetController(GraphToValidate))
+						{
+							Controller->SuspendNotifications(true);
+							Controller->SetPinDefaultValue(Pin->GetPinPath(), TEXT("Null"), false, false, false);
+							Controller->SuspendNotifications(false);
+						}
+					}
+				}
+			}
+		}
+	}	
+
 	RecompileVM();
 	RequestControlRigInit();
 
diff --git a/Engine/Plugins/Experimental/ControlRig/Source/ControlRigDeveloper/Private/Graph/ControlRigGraph.cpp b/Engine/Plugins/Experimental/ControlRig/Source/ControlRigDeveloper/Private/Graph/ControlRigGraph.cpp
--- a/Engine/Plugins/Experimental/ControlRig/Source/ControlRigDeveloper/Private/Graph/ControlRigGraph.cpp
+++ b/Engine/Plugins/Experimental/ControlRig/Source/ControlRigDeveloper/Private/Graph/ControlRigGraph.cpp
@@ -90,7 +90,7 @@ void UControlRigGraph::CacheNameLists(URigHierarchy* InHierarchy, const FControl
 
 	TArray<TSharedPtr<FString>>& AllNameList = ElementNameLists.FindOrAdd(ERigElementType::All);
 	TArray<TSharedPtr<FString>>& BoneNameList = ElementNameLists.FindOrAdd(ERigElementType::Bone);
-	TArray<TSharedPtr<FString>>& SpaceNameList = ElementNameLists.FindOrAdd(ERigElementType::Space);
+	TArray<TSharedPtr<FString>>& NullNameList = ElementNameLists.FindOrAdd(ERigElementType::Null);
 	TArray<TSharedPtr<FString>>& ControlNameList = ElementNameLists.FindOrAdd(ERigElementType::Control);
 	TArray<TSharedPtr<FString>>& CurveNameList = ElementNameLists.FindOrAdd(ERigElementType::Curve);
 	TArray<TSharedPtr<FString>>& RigidBodyNameList = ElementNameLists.FindOrAdd(ERigElementType::RigidBody);
@@ -98,7 +98,7 @@ void UControlRigGraph::CacheNameLists(URigHierarchy* InHierarchy, const FControl
 	
 	CacheNameListForHierarchy<FRigBaseElement>(InHierarchy, AllNameList);
 	CacheNameListForHierarchy<FRigBoneElement>(InHierarchy, BoneNameList);
-	CacheNameListForHierarchy<FRigSpaceElement>(InHierarchy, SpaceNameList);
+	CacheNameListForHierarchy<FRigNullElement>(InHierarchy, NullNameList);
 	CacheNameListForHierarchy<FRigControlElement>(InHierarchy, ControlNameList);
 	CacheNameListForHierarchy<FRigCurveElement>(InHierarchy, CurveNameList);
 	CacheNameListForHierarchy<FRigRigidBodyElement>(InHierarchy, RigidBodyNameList);
diff --git a/Engine/Plugins/Experimental/ControlRig/Source/ControlRigDeveloper/Public/Graph/ControlRigGraph.h b/Engine/Plugins/Experimental/ControlRig/Source/ControlRigDeveloper/Public/Graph/ControlRigGraph.h
--- a/Engine/Plugins/Experimental/ControlRig/Source/ControlRigDeveloper/Public/Graph/ControlRigGraph.h
+++ b/Engine/Plugins/Experimental/ControlRig/Source/ControlRigDeveloper/Public/Graph/ControlRigGraph.h
@@ -48,9 +48,9 @@ public:
 	{
 		return GetElementNameList(ERigElementType::Control);
 	}
-	FORCEINLINE const TArray<TSharedPtr<FString>>& GetSpaceNameList(URigVMPin* InPin = nullptr) const
+	FORCEINLINE const TArray<TSharedPtr<FString>>& GetNullNameList(URigVMPin* InPin = nullptr) const
 	{
-		return GetElementNameList(ERigElementType::Space);
+		return GetElementNameList(ERigElementType::Null);
 	}
 	FORCEINLINE const TArray<TSharedPtr<FString>>& GetCurveNameList(URigVMPin* InPin = nullptr) const
 	{
diff --git a/Engine/Plugins/Experimental/ControlRig/Source/ControlRigEditor/Private/ControlRigEditorModule.cpp b/Engine/Plugins/Experimental/ControlRig/Source/ControlRigEditor/Private/ControlRigEditorModule.cpp
--- a/Engine/Plugins/Experimental/ControlRig/Source/ControlRigEditor/Private/ControlRigEditorModule.cpp
+++ b/Engine/Plugins/Experimental/ControlRig/Source/ControlRigEditor/Private/ControlRigEditorModule.cpp
@@ -201,8 +201,8 @@ void FControlRigEditorModule::StartupModule()
 	ClassesToUnregisterOnShutdown.Add(FRigControlElement::StaticStruct()->GetFName());
 	PropertyEditorModule.RegisterCustomClassLayout(ClassesToUnregisterOnShutdown.Last(), FOnGetDetailCustomizationInstance::CreateStatic(&FRigControlElementDetails::MakeInstance));
 
-	ClassesToUnregisterOnShutdown.Add(FRigSpaceElement::StaticStruct()->GetFName());
-	PropertyEditorModule.RegisterCustomClassLayout(ClassesToUnregisterOnShutdown.Last(), FOnGetDetailCustomizationInstance::CreateStatic(&FRigSpaceElementDetails::MakeInstance));
+	ClassesToUnregisterOnShutdown.Add(FRigNullElement::StaticStruct()->GetFName());
+	PropertyEditorModule.RegisterCustomClassLayout(ClassesToUnregisterOnShutdown.Last(), FOnGetDetailCustomizationInstance::CreateStatic(&FRigNullElementDetails::MakeInstance));
 
 	ClassesToUnregisterOnShutdown.Add(FRigInfluenceMapPerEvent::StaticStruct()->GetFName());
 	PropertyEditorModule.RegisterCustomClassLayout(ClassesToUnregisterOnShutdown.Last(), FOnGetDetailCustomizationInstance::CreateStatic(&FRigInfluenceMapPerEventDetails::MakeInstance));
@@ -1659,7 +1659,7 @@ void FControlRigEditorModule::GetContextMenuActions(const UControlRigGraphSchema
 								}
 								else if (Pin->GetCustomWidgetName() == TEXT("SpaceName"))
 								{
-									Key = FRigElementKey(*Pin->GetDefaultValue(), ERigElementType::Space);
+									Key = FRigElementKey(*Pin->GetDefaultValue(), ERigElementType::Null);
 								}
 								else if (Pin->GetCustomWidgetName() == TEXT("CurveName"))
 								{
diff --git a/Engine/Plugins/Experimental/ControlRig/Source/ControlRigEditor/Private/ControlRigEditorStyle.h b/Engine/Plugins/Experimental/ControlRig/Source/ControlRigEditor/Private/ControlRigEditorStyle.h
--- a/Engine/Plugins/Experimental/ControlRig/Source/ControlRigEditor/Private/ControlRigEditorStyle.h
+++ b/Engine/Plugins/Experimental/ControlRig/Source/ControlRigEditor/Private/ControlRigEditorStyle.h
@@ -98,9 +98,9 @@ public:
 			Set("ControlRig.Tree.BoneUser", new IMAGE_BRUSH("Slate/BoneUser_16x", Icon16x16));
 			Set("ControlRig.Tree.BoneImported", new IMAGE_BRUSH("Slate/BoneImported_16x", Icon16x16));
 			Set("ControlRig.Tree.Control", new IMAGE_BRUSH("Slate/Control_16x", Icon16x16));
-			Set("ControlRig.Tree.Space", new IMAGE_BRUSH("Slate/Space_16x", Icon16x16));
+			Set("ControlRig.Tree.Null", new IMAGE_BRUSH("Slate/Null_16x", Icon16x16));
 			Set("ControlRig.Tree.RigidBody", new IMAGE_BRUSH("Slate/RigidBody_16x", Icon16x16));
-			Set("ControlRig.Tree.Socket", new IMAGE_BRUSH("Slate/Auxiliary_16x", Icon16x16));
+			Set("ControlRig.Tree.Socket", new IMAGE_BRUSH("Slate/Socket_16x", Icon16x16));
 		}
 
 		// Font?
diff --git a/Engine/Plugins/Experimental/ControlRig/Source/ControlRigEditor/Private/ControlRigElementDetails.cpp b/Engine/Plugins/Experimental/ControlRig/Source/ControlRigEditor/Private/ControlRigElementDetails.cpp
--- a/Engine/Plugins/Experimental/ControlRig/Source/ControlRigEditor/Private/ControlRigElementDetails.cpp
+++ b/Engine/Plugins/Experimental/ControlRig/Source/ControlRigEditor/Private/ControlRigElementDetails.cpp
@@ -335,7 +335,7 @@ void FRigUnitDetails::CustomizeDetails(IDetailLayoutBuilder& DetailBuilder)
 				}
 				else if (CustomWidgetName == TEXT("SpaceName"))
 				{
-					NameList = &GraphBeingCustomized->GetSpaceNameList();
+					NameList = &GraphBeingCustomized->GetNullNameList();
 				}
 				else if (CustomWidgetName == TEXT("CurveName"))
 				{
@@ -1584,7 +1584,7 @@ const TArray<TSharedPtr<FString>>& FRigControlElementDetails::GetControlTypeList
 	return ControlTypeList;
 }
 
-void FRigSpaceElementDetails::CustomizeDetails(IDetailLayoutBuilder& DetailBuilder)
+void FRigNullElementDetails::CustomizeDetails(IDetailLayoutBuilder& DetailBuilder)
 {
 	FRigTransformElementDetails::CustomizeDetails(DetailBuilder);
 }
diff --git a/Engine/Plugins/Experimental/ControlRig/Source/ControlRigEditor/Private/ControlRigElementDetails.h b/Engine/Plugins/Experimental/ControlRig/Source/ControlRigEditor/Private/ControlRigElementDetails.h
--- a/Engine/Plugins/Experimental/ControlRig/Source/ControlRigEditor/Private/ControlRigElementDetails.h
+++ b/Engine/Plugins/Experimental/ControlRig/Source/ControlRigEditor/Private/ControlRigElementDetails.h
@@ -185,14 +185,14 @@ private:
 	FEditPropertyChain GizmoPropertyChain;
 };
 
-class FRigSpaceElementDetails : public FRigTransformElementDetails
+class FRigNullElementDetails : public FRigTransformElementDetails
 {
 public:
 
 	/** Makes a new instance of this detail layout class for a specific detail view requesting it */
 	static TSharedRef<IDetailCustomization> MakeInstance()
 	{
-		return MakeShareable(new FRigSpaceElementDetails);
+		return MakeShareable(new FRigNullElementDetails);
 	}
 
 	/** IDetailCustomization interface */
diff --git a/Engine/Plugins/Experimental/ControlRig/Source/ControlRigEditor/Private/ControlRigHierarchyCommands.cpp b/Engine/Plugins/Experimental/ControlRig/Source/ControlRigEditor/Private/ControlRigHierarchyCommands.cpp
--- a/Engine/Plugins/Experimental/ControlRig/Source/ControlRigEditor/Private/ControlRigHierarchyCommands.cpp
+++ b/Engine/Plugins/Experimental/ControlRig/Source/ControlRigEditor/Private/ControlRigHierarchyCommands.cpp
@@ -8,7 +8,7 @@ void FControlRigHierarchyCommands::RegisterCommands()
 {
 	UI_COMMAND(AddBoneItem, "New Bone", "Add new bone at the origin (0, 0, 0) to the hierarchy.", EUserInterfaceActionType::Button, FInputChord());
 	UI_COMMAND(AddControlItem, "New Control", "Add new control at the origin (0, 0, 0) to the hierarchy.", EUserInterfaceActionType::Button, FInputChord(EKeys::N, EModifierKey::Control));
-	UI_COMMAND(AddSpaceItem, "New Space", "Add new space at the origin (0, 0, 0) to the hierarchy.", EUserInterfaceActionType::Button, FInputChord());
+	UI_COMMAND(AddNullItem, "New Null", "Add new null at the origin (0, 0, 0) to the hierarchy.", EUserInterfaceActionType::Button, FInputChord());
 	UI_COMMAND(DuplicateItem, "Duplicate", "Duplicate the selected items in the hierarchy.", EUserInterfaceActionType::Button, FInputChord(EKeys::W, EModifierKey::Control));
 	UI_COMMAND(MirrorItem, "Mirror", "Mirror the selected items in the hierarchy.", EUserInterfaceActionType::Button, FInputChord());
 	UI_COMMAND(DeleteItem, "Delete", "Deletes the selected items from the hierarchy.", EUserInterfaceActionType::Button, FInputChord(EKeys::Delete));
@@ -22,7 +22,7 @@ void FControlRigHierarchyCommands::RegisterCommands()
 	UI_COMMAND(SetGizmoTransformFromCurrent, "Set Gizmo Transform From Current", "Transfer the Current Local Transform of the Control to its Gizmo Transform, the Control's Transform Resets", EUserInterfaceActionType::Button, FInputChord());
 	UI_COMMAND(ResetTransform, "Reset Transform", "Reset the Transform", EUserInterfaceActionType::Button, FInputChord(EKeys::G, EModifierKey::Control));
 	UI_COMMAND(ResetAllTransforms, "Reset All Transforms", "Resets all Transforms", EUserInterfaceActionType::Button, FInputChord(EKeys::G, EModifierKey::Control | EModifierKey::Shift));
-	UI_COMMAND(ResetSpace, "Reset Space", "Resets or injects a Space below the Control", EUserInterfaceActionType::Button, FInputChord());
+	UI_COMMAND(ResetNull, "Reset Null", "Resets or injects a Null below the Control", EUserInterfaceActionType::Button, FInputChord());
 	UI_COMMAND(FrameSelection, "Frame Selection", "Expands and frames the selection in the tree", EUserInterfaceActionType::Button, FInputChord(EKeys::F));
 	UI_COMMAND(ControlBoneTransform, "Control Bone Transform", "Sets the bone transform using a gizmo", EUserInterfaceActionType::Button, FInputChord());
 	UI_COMMAND(ControlSpaceTransform, "Control Space Transform", "Sets the space transform using a gizmo", EUserInterfaceActionType::Button, FInputChord());
@@ -32,7 +32,7 @@ void FControlRigHierarchyCommands::RegisterCommands()
 	UI_COMMAND(ShowImportedBones, "Show Imported Bones", "Whether to show or hide imported bones", EUserInterfaceActionType::ToggleButton, FInputChord());
 	UI_COMMAND(ShowBones, "Show Bones", "Whether to show or hide bones", EUserInterfaceActionType::ToggleButton, FInputChord());
 	UI_COMMAND(ShowControls, "Show Controls", "Whether to show or hide controls", EUserInterfaceActionType::ToggleButton, FInputChord());
-	UI_COMMAND(ShowSpaces, "Show Spaces", "Whether to show or hide spaces", EUserInterfaceActionType::ToggleButton, FInputChord());
+	UI_COMMAND(ShowNulls, "Show Nulls", "Whether to show or hide nulls", EUserInterfaceActionType::ToggleButton, FInputChord());
 	UI_COMMAND(ShowRigidBodies, "Show RigidBodies", "Whether to show or hide rigidbodies", EUserInterfaceActionType::ToggleButton, FInputChord());
 	UI_COMMAND(ShowSockets, "Show Sockets", "Whether to show or hide sockets", EUserInterfaceActionType::ToggleButton, FInputChord());
 	UI_COMMAND(ShowDynamicHierarchy, "Dynamic Hierarchy", "Whether to show the static or dynamic hierarchy", EUserInterfaceActionType::ToggleButton, FInputChord());
diff --git a/Engine/Plugins/Experimental/ControlRig/Source/ControlRigEditor/Private/ControlRigHierarchyCommands.h b/Engine/Plugins/Experimental/ControlRig/Source/ControlRigEditor/Private/ControlRigHierarchyCommands.h
--- a/Engine/Plugins/Experimental/ControlRig/Source/ControlRigEditor/Private/ControlRigHierarchyCommands.h
+++ b/Engine/Plugins/Experimental/ControlRig/Source/ControlRigEditor/Private/ControlRigHierarchyCommands.h
@@ -25,7 +25,7 @@ public:
 	TSharedPtr< FUICommandInfo > AddControlItem;
 
 	/** Add Item at origin */
-	TSharedPtr< FUICommandInfo > AddSpaceItem;
+	TSharedPtr< FUICommandInfo > AddNullItem;
 
 	/** Duplicate currently selected items */
 	TSharedPtr< FUICommandInfo > DuplicateItem;
@@ -58,7 +58,7 @@ public:
 	TSharedPtr<FUICommandInfo> ResetAllTransforms;
 
 	/* Reset space */
-	TSharedPtr<FUICommandInfo> ResetSpace;
+	TSharedPtr<FUICommandInfo> ResetNull;
 
 	/* Set initial transform from current */
 	TSharedPtr<FUICommandInfo> SetInitialTransformFromCurrentTransform;
@@ -97,7 +97,7 @@ public:
 	TSharedPtr< FUICommandInfo > ShowControls;
 
 	/* Show spaces */
-	TSharedPtr< FUICommandInfo > ShowSpaces;
+	TSharedPtr< FUICommandInfo > ShowNulls;
 
 	/* Show rigidbodies */
 	TSharedPtr< FUICommandInfo > ShowRigidBodies;
diff --git a/Engine/Plugins/Experimental/ControlRig/Source/ControlRigEditor/Private/EditMode/ControlRigEditMode.cpp b/Engine/Plugins/Experimental/ControlRig/Source/ControlRigEditor/Private/EditMode/ControlRigEditMode.cpp
--- a/Engine/Plugins/Experimental/ControlRig/Source/ControlRigEditor/Private/EditMode/ControlRigEditMode.cpp
+++ b/Engine/Plugins/Experimental/ControlRig/Source/ControlRigEditor/Private/EditMode/ControlRigEditMode.cpp
@@ -535,9 +535,9 @@ void FControlRigEditMode::Render(const FSceneView* View, FViewport* Viewport, FP
 		if (Settings->bDisplaySpaces || ControlRig->IsSetupModeEnabled())
 		{
 			TArray<FTransform> SpaceTransforms;
-			Hierarchy->ForEach<FRigSpaceElement>([&SpaceTransforms, Hierarchy](FRigSpaceElement* SpaceElement) -> bool
+			Hierarchy->ForEach<FRigNullElement>([&SpaceTransforms, Hierarchy](FRigNullElement* NullElement) -> bool
             {
-				SpaceTransforms.Add(Hierarchy->GetTransform(SpaceElement, ERigTransformType::CurrentGlobal));
+				SpaceTransforms.Add(Hierarchy->GetTransform(NullElement, ERigTransformType::CurrentGlobal));
 				return true;
 			});
 			GetControlRig(true)->DrawInterface.DrawAxes(FTransform::Identity, SpaceTransforms, Settings->AxisScale);
@@ -1832,7 +1832,7 @@ void FControlRigEditMode::OnHierarchyModified(ERigHierarchyNotification InNotif,
 			switch (InElement->GetType())
 			{
 				case ERigElementType::Bone:
-            	case ERigElementType::Space:
+            	case ERigElementType::Null:
             	case ERigElementType::Curve:
             	case ERigElementType::Control:
             	case ERigElementType::RigidBody:
diff --git a/Engine/Plugins/Experimental/ControlRig/Source/ControlRigEditor/Private/Editor/ControlRigEditor.cpp b/Engine/Plugins/Experimental/ControlRig/Source/ControlRigEditor/Private/Editor/ControlRigEditor.cpp
--- a/Engine/Plugins/Experimental/ControlRig/Source/ControlRigEditor/Private/Editor/ControlRigEditor.cpp
+++ b/Engine/Plugins/Experimental/ControlRig/Source/ControlRigEditor/Private/Editor/ControlRigEditor.cpp
@@ -1186,9 +1186,9 @@ void FControlRigEditor::SetDetailStruct(const FRigElementKey& InElement)
 			StructToDisplay = MakeShareable(new FStructOnScope(FRigControlElement::StaticStruct(), (uint8*)Element));
 			break;
 		}
-		case ERigElementType::Space:
+		case ERigElementType::Null:
 		{
-			StructToDisplay = MakeShareable(new FStructOnScope(FRigSpaceElement::StaticStruct(), (uint8*)Element));
+			StructToDisplay = MakeShareable(new FStructOnScope(FRigNullElement::StaticStruct(), (uint8*)Element));
 			break;
 		}
 		case ERigElementType::Curve:
@@ -2339,12 +2339,12 @@ void FControlRigEditor::HandleViewportCreated(const TSharedRef<class IPersonaVie
 						SNew(SBox)
 						.Padding(FMargin(4.0f, 0.0f, 0.0f, 0.0f))
 						.WidthOverride(100.0f)
-						.IsEnabled(this, &FControlRigEditor::IsToolbarDrawSpacesEnabled)
+						.IsEnabled(this, &FControlRigEditor::IsToolbarDrawNullsEnabled)
 						[
 							SNew(SCheckBox)
-							.IsChecked(this, &FControlRigEditor::GetToolbarDrawSpaces)
-							.OnCheckStateChanged(this, &FControlRigEditor::OnToolbarDrawSpacesChanged)
-							.ToolTipText(LOCTEXT("ControlRigDrawSpacesToolTip", "If checked all spaces are drawn as axes."))
+							.IsChecked(this, &FControlRigEditor::GetToolbarDrawNulls)
+							.OnCheckStateChanged(this, &FControlRigEditor::OnToolbarDrawNullsChanged)
+							.ToolTipText(LOCTEXT("ControlRigDrawNullsToolTip", "If checked all spaces are drawn as axes."))
 						]
 					],
 					LOCTEXT("ControlRigDisplaySpaces", "Display Spaces")
@@ -2472,7 +2472,7 @@ void FControlRigEditor::OnToolbarDrawAxesOnSelectionChanged(ECheckBoxState InNew
 	}
 }
 
-bool FControlRigEditor::IsToolbarDrawSpacesEnabled() const
+bool FControlRigEditor::IsToolbarDrawNullsEnabled() const
 {
 	if (ControlRig)
 	{
@@ -2484,7 +2484,7 @@ bool FControlRigEditor::IsToolbarDrawSpacesEnabled() const
 	return false;
 }
 
-ECheckBoxState FControlRigEditor::GetToolbarDrawSpaces() const
+ECheckBoxState FControlRigEditor::GetToolbarDrawNulls() const
 {
 	if (FControlRigEditMode* EditMode = GetEditMode())
 	{
@@ -2493,7 +2493,7 @@ ECheckBoxState FControlRigEditor::GetToolbarDrawSpaces() const
 	return ECheckBoxState::Unchecked;
 }
 
-void FControlRigEditor::OnToolbarDrawSpacesChanged(ECheckBoxState InNewValue)
+void FControlRigEditor::OnToolbarDrawNullsChanged(ECheckBoxState InNewValue)
 {
 	if (FControlRigEditMode* EditMode = GetEditMode())
 	{
@@ -2973,7 +2973,7 @@ void FControlRigEditor::SetRigElementTransform(const FRigElementKey& InElement,
 			OnHierarchyChanged();
 			break;
 		}
-		case ERigElementType::Space:
+		case ERigElementType::Null:
 		{
 			FTransform LocalTransform = InTransform;
 			FTransform GlobalTransform = InTransform;
@@ -3341,7 +3341,7 @@ void FControlRigEditor::OnHierarchyModified(ERigHierarchyNotification InNotif, U
 							{
 								if ((ModelPin->GetCPPType() == TEXT("FName") && ModelPin->GetCustomWidgetName() == TEXT("BoneName") && RemovedElementType == ERigElementType::Bone) ||
 									(ModelPin->GetCPPType() == TEXT("FName") && ModelPin->GetCustomWidgetName() == TEXT("ControlName") && RemovedElementType == ERigElementType::Control) ||
-									(ModelPin->GetCPPType() == TEXT("FName") && ModelPin->GetCustomWidgetName() == TEXT("SpaceName") && RemovedElementType == ERigElementType::Space) ||
+									(ModelPin->GetCPPType() == TEXT("FName") && ModelPin->GetCustomWidgetName() == TEXT("SpaceName") && RemovedElementType == ERigElementType::Null) ||
 									(ModelPin->GetCPPType() == TEXT("FName") && ModelPin->GetCustomWidgetName() == TEXT("CurveName") && RemovedElementType == ERigElementType::Curve))
 								{
 									if (ModelPin->GetDefaultValue() == RemovedElementName)
@@ -3413,7 +3413,7 @@ void FControlRigEditor::OnHierarchyModified(ERigHierarchyNotification InNotif, U
 							{
 								if ((ModelPin->GetCPPType() == TEXT("FName") && ModelPin->GetCustomWidgetName() == TEXT("BoneName") && ElementType == ERigElementType::Bone) ||
                                     (ModelPin->GetCPPType() == TEXT("FName") && ModelPin->GetCustomWidgetName() == TEXT("ControlName") && ElementType == ERigElementType::Control) ||
-                                    (ModelPin->GetCPPType() == TEXT("FName") && ModelPin->GetCustomWidgetName() == TEXT("SpaceName") && ElementType == ERigElementType::Space) ||
+                                    (ModelPin->GetCPPType() == TEXT("FName") && ModelPin->GetCustomWidgetName() == TEXT("SpaceName") && ElementType == ERigElementType::Null) ||
                                     (ModelPin->GetCPPType() == TEXT("FName") && ModelPin->GetCustomWidgetName() == TEXT("CurveName") && ElementType == ERigElementType::Curve))
 								{
 									if (ModelPin->GetDefaultValue() == OldNameStr)
@@ -3632,7 +3632,7 @@ void FControlRigEditor::OnGraphNodeDropToPerform(TSharedPtr<FGraphNodeDragDropOp
 			}
 
 			if (((DraggedTypes & (uint8)ERigElementType::Bone) != 0) ||
-				((DraggedTypes & (uint8)ERigElementType::Space) != 0))
+				((DraggedTypes & (uint8)ERigElementType::Null) != 0))
 			{
 				MenuBuilder.AddMenuEntry(
 					LOCTEXT("CreateGetTransform", "Get Transform"),
@@ -3656,7 +3656,7 @@ void FControlRigEditor::OnGraphNodeDropToPerform(TSharedPtr<FGraphNodeDragDropOp
 
 			if (((DraggedTypes & (uint8)ERigElementType::Bone) != 0) ||
 				((DraggedTypes & (uint8)ERigElementType::Control) != 0) ||
-				((DraggedTypes & (uint8)ERigElementType::Space) != 0))
+				((DraggedTypes & (uint8)ERigElementType::Null) != 0))
 			{
 				MenuBuilder.AddMenuSeparator();
 
@@ -4202,7 +4202,7 @@ void FControlRigEditor::HandleOnControlModified(UControlRig* Subject, FRigContro
 					}
 				}
 			}
-			else if (ElementKey.Type == ERigElementType::Space)
+			else if (ElementKey.Type == ERigElementType::Null)
 			{
 				const FTransform GlobalTransform = ControlRig->GetControlGlobalTransform(ControlElement->GetName());
 				Blueprint->Hierarchy->SetGlobalTransform(ElementKey, GlobalTransform);
diff --git a/Engine/Plugins/Experimental/ControlRig/Source/ControlRigEditor/Private/Editor/ControlRigEditor.h b/Engine/Plugins/Experimental/ControlRig/Source/ControlRigEditor/Private/Editor/ControlRigEditor.h
--- a/Engine/Plugins/Experimental/ControlRig/Source/ControlRigEditor/Private/Editor/ControlRigEditor.h
+++ b/Engine/Plugins/Experimental/ControlRig/Source/ControlRigEditor/Private/Editor/ControlRigEditor.h
@@ -247,9 +247,9 @@ private:
 	/** Handle preview scene setup */
 	void HandlePreviewSceneCreated(const TSharedRef<IPersonaPreviewScene>& InPersonaPreviewScene);
 	void HandleViewportCreated(const TSharedRef<class IPersonaViewport>& InViewport);
-	bool IsToolbarDrawSpacesEnabled() const;
-	ECheckBoxState GetToolbarDrawSpaces() const;
-	void OnToolbarDrawSpacesChanged(ECheckBoxState InNewValue);
+	bool IsToolbarDrawNullsEnabled() const;
+	ECheckBoxState GetToolbarDrawNulls() const;
+	void OnToolbarDrawNullsChanged(ECheckBoxState InNewValue);
 	ECheckBoxState GetToolbarDrawAxesOnSelection() const;
 	void OnToolbarDrawAxesOnSelectionChanged(ECheckBoxState InNewValue);
 	TOptional<float> GetToolbarAxesScale() const;
diff --git a/Engine/Plugins/Experimental/ControlRig/Source/ControlRigEditor/Private/Editor/ControlRigEditorEditMode.cpp b/Engine/Plugins/Experimental/ControlRig/Source/ControlRigEditor/Private/Editor/ControlRigEditorEditMode.cpp
--- a/Engine/Plugins/Experimental/ControlRig/Source/ControlRigEditor/Private/Editor/ControlRigEditorEditMode.cpp
+++ b/Engine/Plugins/Experimental/ControlRig/Source/ControlRigEditor/Private/Editor/ControlRigEditorEditMode.cpp
@@ -113,7 +113,7 @@ bool FControlRigEditorEditMode::GetCameraTarget(FSphere& OutTarget) const
 	for (int32 Index = 0; Index < SelectedRigElements.Num(); ++Index)
 	{
 		static const float Radius = 20.f;
-		if (SelectedRigElements[Index].Type == ERigElementType::Bone || SelectedRigElements[Index].Type == ERigElementType::Space)
+		if (SelectedRigElements[Index].Type == ERigElementType::Bone || SelectedRigElements[Index].Type == ERigElementType::Null)
 		{
 			FTransform Transform = OnGetRigElementTransformDelegate.Execute(SelectedRigElements[Index], false, true);
 			Box += Transform.TransformPosition(FVector::OneVector * Radius);
diff --git a/Engine/Plugins/Experimental/ControlRig/Source/ControlRigEditor/Private/Editor/SRigHierarchy.cpp b/Engine/Plugins/Experimental/ControlRig/Source/ControlRigEditor/Private/Editor/SRigHierarchy.cpp
--- a/Engine/Plugins/Experimental/ControlRig/Source/ControlRigEditor/Private/Editor/SRigHierarchy.cpp
+++ b/Engine/Plugins/Experimental/ControlRig/Source/ControlRigEditor/Private/Editor/SRigHierarchy.cpp
@@ -151,9 +151,9 @@ void SRigHierarchyItem::Construct(const FArguments& InArgs, TSharedPtr<FControlR
 			Brush = FControlRigEditorStyle::Get().GetBrush("ControlRig.Tree.Control");
 			break;
 		}
-		case ERigElementType::Space:
+		case ERigElementType::Null:
 		{
-			Brush = FControlRigEditorStyle::Get().GetBrush("ControlRig.Tree.Space");
+			Brush = FControlRigEditorStyle::Get().GetBrush("ControlRig.Tree.Null");
 			break;
 		}
 		case ERigElementType::Bone:
@@ -446,7 +446,7 @@ void SRigHierarchy::Construct(const FArguments& InArgs, TSharedRef<FControlRigEd
 	bShowImportedBones = true;
 	bShowBones = true;
 	bShowControls = true;
-	bShowSpaces = true;
+	bShowNulls = true;
 	bShowRigidBodies = true;
 	bShowSockets = true;
 	bIsChangingRigHierarchy = false;
@@ -474,8 +474,8 @@ void SRigHierarchy::BindCommands()
 	CommandList->MapAction(Commands.AddControlItem,
 		FExecuteAction::CreateSP(this, &SRigHierarchy::HandleNewItem, ERigElementType::Control));
 
-	CommandList->MapAction(Commands.AddSpaceItem,
-		FExecuteAction::CreateSP(this, &SRigHierarchy::HandleNewItem, ERigElementType::Space));
+	CommandList->MapAction(Commands.AddNullItem,
+		FExecuteAction::CreateSP(this, &SRigHierarchy::HandleNewItem, ERigElementType::Null));
 
 	CommandList->MapAction(Commands.DuplicateItem,
 		FExecuteAction::CreateSP(this, &SRigHierarchy::HandleDuplicateItem),
@@ -522,7 +522,7 @@ void SRigHierarchy::BindCommands()
 	CommandList->MapAction(
 		Commands.SetInitialTransformFromClosestBone,
 		FExecuteAction::CreateSP(this, &SRigHierarchy::HandleSetInitialTransformFromClosestBone),
-		FCanExecuteAction::CreateSP(this, &SRigHierarchy::IsControlOrSpaceSelected));
+		FCanExecuteAction::CreateSP(this, &SRigHierarchy::IsControlOrNullSelected));
 
 	CommandList->MapAction(
 		Commands.SetInitialTransformFromCurrentTransform,
@@ -580,10 +580,10 @@ void SRigHierarchy::BindCommands()
 		FIsActionChecked::CreateLambda([this]() { return bShowControls; }));
 	
 	CommandList->MapAction(
-		Commands.ShowSpaces,
-		FExecuteAction::CreateLambda([this]() { bShowSpaces = !bShowSpaces; RefreshTreeView(); }),
+		Commands.ShowNulls,
+		FExecuteAction::CreateLambda([this]() { bShowNulls = !bShowNulls; RefreshTreeView(); }),
 		FCanExecuteAction(),
-		FIsActionChecked::CreateLambda([this]() { return bShowSpaces; }));
+		FIsActionChecked::CreateLambda([this]() { return bShowNulls; }));
 
 	CommandList->MapAction(
 		Commands.ShowRigidBodies,
@@ -635,7 +635,7 @@ EVisibility SRigHierarchy::IsSearchbarVisible() const
 	if (URigHierarchy* Hierarchy = GetDebuggedHierarchy())
 	{
 		if ((Hierarchy->Num(ERigElementType::Bone) +
-			Hierarchy->Num(ERigElementType::Space) +
+			Hierarchy->Num(ERigElementType::Null) +
 			Hierarchy->Num(ERigElementType::Control)) > 0)
 		{
 			return EVisibility::Visible;
@@ -915,9 +915,9 @@ bool SRigHierarchy::AddElement(const FRigBaseElement* InElement, const bool bIgn
 			}
 			break;
 		}
-		case ERigElementType::Space:
+		case ERigElementType::Null:
 		{
-			if(!bShowSpaces)
+			if(!bShowNulls)
 			{
 				return false;
 			}
@@ -1276,7 +1276,7 @@ TSharedRef< SWidget > SRigHierarchy::CreateFilterMenu()
 	MenuBuilder.BeginSection("FilterControls", LOCTEXT("ControlsMenuHeading", "Controls"));
 	{
 		MenuBuilder.AddMenuEntry(Actions.ShowControls);
-		MenuBuilder.AddMenuEntry(Actions.ShowSpaces);
+		MenuBuilder.AddMenuEntry(Actions.ShowNulls);
 	}
 	MenuBuilder.EndSection();
 
@@ -1367,7 +1367,7 @@ void SRigHierarchy::FillContextMenu(class FMenuBuilder& MenuBuilder)
 					InSubMenuBuilder.AddMenuEntry(Actions.AddBoneItem);
 				}
 				InSubMenuBuilder.AddMenuEntry(Actions.AddControlItem);
-				InSubMenuBuilder.AddMenuEntry(Actions.AddSpaceItem);
+				InSubMenuBuilder.AddMenuEntry(Actions.AddNullItem);
 			}
 		};
 
@@ -1389,7 +1389,7 @@ void SRigHierarchy::FillContextMenu(class FMenuBuilder& MenuBuilder)
 		}
 
 		/*
-		if (IsSingleSpaceSelected())
+		if (IsSingleNullSelected())
 		{
 			MenuBuilder.BeginSection("Interaction", LOCTEXT("InteractionHeader", "Interaction"));
 			MenuBuilder.AddMenuEntry(Actions.ControlSpaceTransform);
@@ -1602,13 +1602,13 @@ bool SRigHierarchy::IsSingleBoneSelected() const
 	return GetSelectedKeys()[0].Type == ERigElementType::Bone;
 }
 
-bool SRigHierarchy::IsSingleSpaceSelected() const
+bool SRigHierarchy::IsSingleNullSelected() const
 {
 	if(!IsSingleSelected())
 	{
 		return false;
 	}
-	return GetSelectedKeys()[0].Type == ERigElementType::Space;
+	return GetSelectedKeys()[0].Type == ERigElementType::Null;
 }
 
 bool SRigHierarchy::IsControlSelected() const
@@ -1624,7 +1624,7 @@ bool SRigHierarchy::IsControlSelected() const
 	return false;
 }
 
-bool SRigHierarchy::IsControlOrSpaceSelected() const
+bool SRigHierarchy::IsControlOrNullSelected() const
 {
 	TArray<FRigElementKey> SelectedKeys = GetSelectedKeys();
 	for (const FRigElementKey& SelectedKey : SelectedKeys)
@@ -1633,7 +1633,7 @@ bool SRigHierarchy::IsControlOrSpaceSelected() const
 		{
 			return true;
 		}
-		if (SelectedKey.Type == ERigElementType::Space)
+		if (SelectedKey.Type == ERigElementType::Null)
 		{
 			return true;
 		}
@@ -1753,9 +1753,9 @@ void SRigHierarchy::HandleNewItem(ERigElementType InElementType)
 					NewItemKey = Controller->AddControl(NewElementName, ParentKey, Settings, Settings.GetIdentityValue(), FTransform::Identity, FTransform::Identity, true);
 					break;
 				}
-				case ERigElementType::Space:
+				case ERigElementType::Null:
 				{
-					NewItemKey = Controller->AddSpace(NewElementName, ParentKey, ParentTransform, true, true);
+					NewItemKey = Controller->AddNull(NewElementName, ParentKey, ParentTransform, true, true);
 					break;
 				}
 				default:
@@ -2131,7 +2131,7 @@ TOptional<EItemDropZone> SRigHierarchy::OnCanAcceptDrop(const FDragDropEvent& Dr
 				break;
 			}
 			case ERigElementType::Control:
-			case ERigElementType::Space:
+			case ERigElementType::Null:
 			case ERigElementType::RigidBody:
 			case ERigElementType::Socket:
 			{
@@ -2140,7 +2140,7 @@ TOptional<EItemDropZone> SRigHierarchy::OnCanAcceptDrop(const FDragDropEvent& Dr
 					switch (DraggedKey.Type)
 					{
 						case ERigElementType::Control:
-						case ERigElementType::Space:
+						case ERigElementType::Null:
 						case ERigElementType::RigidBody:
 						case ERigElementType::Socket:
 						{
@@ -2426,7 +2426,7 @@ void SRigHierarchy::HandleSetInitialTransformFromCurrentTransform()
 							DebuggedHierarchy->SetTransform(ControlElement, FTransform::Identity, ERigTransformType::CurrentLocal, true, true);
 						}
 					}
-					else if (SelectedKey.Type == ERigElementType::Space ||
+					else if (SelectedKey.Type == ERigElementType::Null ||
 						SelectedKey.Type == ERigElementType::Bone)
 					{
 						FTransform InitialTransform = LocalTransform;
@@ -2494,7 +2494,7 @@ void SRigHierarchy::HandleControlBoneOrSpaceTransform()
 	for (const FRigElementKey& SelectedKey : SelectedKeys)
 	{
 		if (SelectedKey.Type == ERigElementType::Bone ||
-			SelectedKey.Type == ERigElementType::Space)
+			SelectedKey.Type == ERigElementType::Null)
 		{
 			Blueprint->AddTransientControl(SelectedKey);
 			return;
@@ -2570,7 +2570,7 @@ void SRigHierarchy::HandleUnparent()
 				}
 				break;
 			}
-			case ERigElementType::Space:
+			case ERigElementType::Null:
 			case ERigElementType::Control:
 			{
 				Controller->RemoveAllParents(SelectedKey, true, true);
@@ -2618,7 +2618,7 @@ bool SRigHierarchy::FindClosestBone(const FVector& Point, FName& OutRigElementNa
 
 void SRigHierarchy::HandleSetInitialTransformFromClosestBone()
 {
-	if (IsControlOrSpaceSelected())
+	if (IsControlOrNullSelected())
 	{
 		if (UControlRigBlueprint* Blueprint = ControlRigEditor.Pin()->GetControlRigBlueprint())
 		{
@@ -2632,7 +2632,7 @@ void SRigHierarchy::HandleSetInitialTransformFromClosestBone()
 
 				for (const FRigElementKey& SelectedKey : SelectedKeys)
 				{
-					if (SelectedKey.Type == ERigElementType::Control || SelectedKey.Type == ERigElementType::Space)
+					if (SelectedKey.Type == ERigElementType::Control || SelectedKey.Type == ERigElementType::Null)
 					{
 						const FTransform GlobalTransform = DebuggedHierarchy->GetGlobalTransform(SelectedKey);
 						FTransform ClosestTransform;
@@ -2674,7 +2674,7 @@ void SRigHierarchy::HandleSetInitialTransformFromClosestBone()
 							DebuggedHierarchy->SetTransform(ControlElement, FTransform::Identity, ERigTransformType::CurrentLocal, true, true);
 						}
 					}
-					else if (SelectedKey.Type == ERigElementType::Space ||
+					else if (SelectedKey.Type == ERigElementType::Null ||
                         SelectedKey.Type == ERigElementType::Bone)
 					{
 						FTransform InitialTransform = LocalTransform;
diff --git a/Engine/Plugins/Experimental/ControlRig/Source/ControlRigEditor/Private/Editor/SRigHierarchy.h b/Engine/Plugins/Experimental/ControlRig/Source/ControlRigEditor/Private/Editor/SRigHierarchy.h
--- a/Engine/Plugins/Experimental/ControlRig/Source/ControlRigEditor/Private/Editor/SRigHierarchy.h
+++ b/Engine/Plugins/Experimental/ControlRig/Source/ControlRigEditor/Private/Editor/SRigHierarchy.h
@@ -263,7 +263,7 @@ private:
 	bool bShowControls;
 
 	/** Whether or not to show spaces in the hierarchy */
-	bool bShowSpaces;
+	bool bShowNulls;
 
 	/** Whether or not to show rigidbodies in the hierarchy */
 	bool bShowRigidBodies;
@@ -304,9 +304,9 @@ private:
 	bool IsMultiSelected() const;
 	bool IsSingleSelected() const;
 	bool IsSingleBoneSelected() const;
-	bool IsSingleSpaceSelected() const;
+	bool IsSingleNullSelected() const;
 	bool IsControlSelected() const;
-	bool IsControlOrSpaceSelected() const;
+	bool IsControlOrNullSelected() const;
 
 	URigHierarchy* GetHierarchy() const;
 	URigHierarchy* GetDebuggedHierarchy() const;
diff --git a/Engine/Plugins/Experimental/ControlRig/Source/ControlRigEditor/Private/Graph/ControlRigGraphPanelPinFactory.cpp b/Engine/Plugins/Experimental/ControlRig/Source/ControlRigEditor/Private/Graph/ControlRigGraphPanelPinFactory.cpp
--- a/Engine/Plugins/Experimental/ControlRig/Source/ControlRigEditor/Private/Graph/ControlRigGraphPanelPinFactory.cpp
+++ b/Engine/Plugins/Experimental/ControlRig/Source/ControlRigEditor/Private/Graph/ControlRigGraphPanelPinFactory.cpp
@@ -55,7 +55,7 @@ TSharedPtr<SGraphPin> FControlRigGraphPanelPinFactory::CreatePin(UEdGraphPin* In
 				{
 					return SNew(SControlRigGraphPinNameList, InPin)
 						.ModelPin(ModelPin)
-						.OnGetNameListContent_UObject(RigGraph, &UControlRigGraph::GetSpaceNameList);
+						.OnGetNameListContent_UObject(RigGraph, &UControlRigGraph::GetNullNameList);
 				}
 				else if (CustomWidgetName == TEXT("CurveName"))
 				{
diff --git a/Engine/Plugins/Experimental/MotionTrailEditorMode/Source/MotionTrailEditorMode/Private/Sequencer/SequencerTrailHierarchy.cpp b/Engine/Plugins/Experimental/MotionTrailEditorMode/Source/MotionTrailEditorMode/Private/Sequencer/SequencerTrailHierarchy.cpp
--- a/Engine/Plugins/Experimental/MotionTrailEditorMode/Source/MotionTrailEditorMode/Private/Sequencer/SequencerTrailHierarchy.cpp
+++ b/Engine/Plugins/Experimental/MotionTrailEditorMode/Source/MotionTrailEditorMode/Private/Sequencer/SequencerTrailHierarchy.cpp
@@ -607,11 +607,11 @@ void FSequencerTrailHierarchy::ResolveRigElementToRootComponent(URigHierarchy* R
 	FRigElementKey ChildItr = InElementKey;
 	while (ParentKey.IsValid() && ParentKey.Type != ERigElementType::Bone)
 	{
-		if (ParentKey.Type == ERigElementType::Space) // TODO: add empty space trail with initial transform
+		if (ParentKey.Type == ERigElementType::Null) // TODO: add empty space trail with initial transform
 		{
 			ChildItr = ParentKey;
 			ElementIndex = RigHierarchy->GetIndex(FRigElementKey(ChildItr.Name, ERigElementType::Control));
-			ParentKey = RigHierarchy->GetFirstParent(FRigElementKey(ChildItr.Name, ERigElementType::Space));
+			ParentKey = RigHierarchy->GetFirstParent(FRigElementKey(ChildItr.Name, ERigElementType::Null));
 			continue;
 		}
 
