Control Rig: Bracket hierarchy changes during load with a notification guard

diff --git a/Engine/Plugins/Experimental/ControlRig/Source/ControlRig/Private/Rigs/RigHierarchy.cpp b/Engine/Plugins/Experimental/ControlRig/Source/ControlRig/Private/Rigs/RigHierarchy.cpp
--- a/Engine/Plugins/Experimental/ControlRig/Source/ControlRig/Private/Rigs/RigHierarchy.cpp
+++ b/Engine/Plugins/Experimental/ControlRig/Source/ControlRig/Private/Rigs/RigHierarchy.cpp
@@ -1104,6 +1104,10 @@ void URigHierarchy::SetPose(const FRigPose& InPose, ERigTransformType::Type InTr
 
 void URigHierarchy::Notify(ERigHierarchyNotification InNotifType, const FRigBaseElement* InElement)
 {
+	if(bSuspendNotifications)
+	{
+		return;
+	}
 	ModifiedEvent.Broadcast(InNotifType, this, InElement);
 }
 
diff --git a/Engine/Plugins/Experimental/ControlRig/Source/ControlRig/Private/Rigs/RigHierarchyController.cpp b/Engine/Plugins/Experimental/ControlRig/Source/ControlRig/Private/Rigs/RigHierarchyController.cpp
--- a/Engine/Plugins/Experimental/ControlRig/Source/ControlRig/Private/Rigs/RigHierarchyController.cpp
+++ b/Engine/Plugins/Experimental/ControlRig/Source/ControlRig/Private/Rigs/RigHierarchyController.cpp
@@ -987,11 +987,19 @@ void URigHierarchyController::Notify(ERigHierarchyNotification InNotifType, cons
 	{
 		return;
 	}
+	if(bSuspendNotifications)
+	{
+		return;
+	}
 	Hierarchy->Notify(InNotifType, InElement);
 }
 
 void URigHierarchyController::HandleHierarchyModified(ERigHierarchyNotification InNotifType, URigHierarchy* InHierarchy, const FRigBaseElement* InElement) const
 {
+	if(bSuspendNotifications)
+	{
+		return;
+	}
 	ensure(IsValid());
 	ensure(InHierarchy == Hierarchy);
 	ModifiedEvent.Broadcast(InNotifType, InHierarchy, InElement);
diff --git a/Engine/Plugins/Experimental/ControlRig/Source/ControlRig/Public/Rigs/RigHierarchy.h b/Engine/Plugins/Experimental/ControlRig/Source/ControlRig/Public/Rigs/RigHierarchy.h
--- a/Engine/Plugins/Experimental/ControlRig/Source/ControlRig/Public/Rigs/RigHierarchy.h
+++ b/Engine/Plugins/Experimental/ControlRig/Source/ControlRig/Public/Rigs/RigHierarchy.h
@@ -31,6 +31,7 @@ public:
 	, TransformStackIndex(0)
 	, bTransactingForTransformChange(false)
 	, bIsInteracting(false)
+	, bSuspendNotifications(false)
 	{}
 
 	virtual ~URigHierarchy();
@@ -1994,6 +1995,11 @@ public:
 	 * @return True if the child is parented to the parent
 	 */
 	bool IsParentedTo(FRigBaseElement* InChild, FRigBaseElement* InParent) const;
+	
+	/**
+	 * Returns a reference to the suspend notifications flag
+	 */
+	FORCEINLINE bool& GetSuspendNotificationsFlag() { return bSuspendNotifications; } 
 
 private:
 
@@ -2152,6 +2158,11 @@ private:
 	 * Manages merging transform actions into one during an interaction
 	 */
 	bool bIsInteracting;
+	
+	/** 
+	 * If set to true all notifs coming from this hierarchy will be suspended
+	 */
+	bool bSuspendNotifications;
 
 	/**
 	 * The event fired during undo / redo
diff --git a/Engine/Plugins/Experimental/ControlRig/Source/ControlRig/Public/Rigs/RigHierarchyController.h b/Engine/Plugins/Experimental/ControlRig/Source/ControlRig/Public/Rigs/RigHierarchyController.h
--- a/Engine/Plugins/Experimental/ControlRig/Source/ControlRig/Public/Rigs/RigHierarchyController.h
+++ b/Engine/Plugins/Experimental/ControlRig/Source/ControlRig/Public/Rigs/RigHierarchyController.h
@@ -17,6 +17,7 @@ public:
 	URigHierarchyController()
 	: bReportWarningsAndErrors(true)
 	, Hierarchy(nullptr)
+	, bSuspendNotifications(false)
 	{}
 
 	virtual ~URigHierarchyController();
@@ -451,6 +452,11 @@ public:
 	UPROPERTY(transient)
 	bool bReportWarningsAndErrors;
 
+	/**
+	 * Returns a reference to the suspend notifications flag
+	 */
+	FORCEINLINE bool& GetSuspendNotificationsFlag() { return bSuspendNotifications; } 
+
 private:
 
 	UPROPERTY(transient)
@@ -547,5 +553,10 @@ private:
 	static USkeleton* GetSkeletonFromAssetPath(const FString& InAssetPath);
 #endif
 
+	/** 
+	 * If set to true all notifs coming from this hierarchy will be suspended
+	 */
+	bool bSuspendNotifications;
+
 	friend class UControlRig;
 };
diff --git a/Engine/Plugins/Experimental/ControlRig/Source/ControlRigDeveloper/Private/ControlRigBlueprint.cpp b/Engine/Plugins/Experimental/ControlRig/Source/ControlRigDeveloper/Private/ControlRigBlueprint.cpp
--- a/Engine/Plugins/Experimental/ControlRig/Source/ControlRigDeveloper/Private/ControlRigBlueprint.cpp
+++ b/Engine/Plugins/Experimental/ControlRig/Source/ControlRigDeveloper/Private/ControlRigBlueprint.cpp
@@ -358,6 +358,8 @@ void UControlRigBlueprint::PostLoad()
 	{
 		Modify();
 		
+		TGuardValue<bool> SuspendNotifGuard(Hierarchy->GetSuspendNotificationsFlag(), true);
+		
 		Hierarchy->Reset();
 
 		for(const FRigBone& Bone : HierarchyContainer_DEPRECATED.BoneHierarchy)
diff --git a/Engine/Plugins/Experimental/ControlRig/Source/ControlRigDeveloper/Private/Graph/ControlRigGraph.cpp b/Engine/Plugins/Experimental/ControlRig/Source/ControlRigDeveloper/Private/Graph/ControlRigGraph.cpp
--- a/Engine/Plugins/Experimental/ControlRig/Source/ControlRigDeveloper/Private/Graph/ControlRigGraph.cpp
+++ b/Engine/Plugins/Experimental/ControlRig/Source/ControlRigDeveloper/Private/Graph/ControlRigGraph.cpp
@@ -88,13 +88,21 @@ void UControlRigGraph::CacheNameLists(URigHierarchy* InHierarchy, const FControl
 	check(InHierarchy);
 	check(DrawContainer);
 
-	TArray<TSharedPtr<FString>>& AllNameList = ElementNameLists.FindOrAdd(ERigElementType::All);
-	TArray<TSharedPtr<FString>>& BoneNameList = ElementNameLists.FindOrAdd(ERigElementType::Bone);
-	TArray<TSharedPtr<FString>>& NullNameList = ElementNameLists.FindOrAdd(ERigElementType::Null);
-	TArray<TSharedPtr<FString>>& ControlNameList = ElementNameLists.FindOrAdd(ERigElementType::Control);
-	TArray<TSharedPtr<FString>>& CurveNameList = ElementNameLists.FindOrAdd(ERigElementType::Curve);
-	TArray<TSharedPtr<FString>>& RigidBodyNameList = ElementNameLists.FindOrAdd(ERigElementType::RigidBody);
-	TArray<TSharedPtr<FString>>& SocketNameList = ElementNameLists.FindOrAdd(ERigElementType::Socket);
+	ElementNameLists.FindOrAdd(ERigElementType::All);
+	ElementNameLists.FindOrAdd(ERigElementType::Bone);
+	ElementNameLists.FindOrAdd(ERigElementType::Null);
+	ElementNameLists.FindOrAdd(ERigElementType::Control);
+	ElementNameLists.FindOrAdd(ERigElementType::Curve);
+	ElementNameLists.FindOrAdd(ERigElementType::RigidBody);
+	ElementNameLists.FindOrAdd(ERigElementType::Socket);
+
+	TArray<TSharedPtr<FString>>& AllNameList = ElementNameLists.FindChecked(ERigElementType::All);
+	TArray<TSharedPtr<FString>>& BoneNameList = ElementNameLists.FindChecked(ERigElementType::Bone);
+	TArray<TSharedPtr<FString>>& NullNameList = ElementNameLists.FindChecked(ERigElementType::Null);
+	TArray<TSharedPtr<FString>>& ControlNameList = ElementNameLists.FindChecked(ERigElementType::Control);
+	TArray<TSharedPtr<FString>>& CurveNameList = ElementNameLists.FindChecked(ERigElementType::Curve);
+	TArray<TSharedPtr<FString>>& RigidBodyNameList = ElementNameLists.FindChecked(ERigElementType::RigidBody);
+	TArray<TSharedPtr<FString>>& SocketNameList = ElementNameLists.FindChecked(ERigElementType::Socket);
 	
 	CacheNameListForHierarchy<FRigBaseElement>(InHierarchy, AllNameList);
 	CacheNameListForHierarchy<FRigBoneElement>(InHierarchy, BoneNameList);
