Sequencer: Control Rig: Fix for Q for select mode, new widget mode exists that we need to handle.

diff --git a/Engine/Plugins/Experimental/ControlRig/Source/ControlRigEditor/Private/EditMode/ControlRigEditMode.cpp b/Engine/Plugins/Experimental/ControlRig/Source/ControlRigEditor/Private/EditMode/ControlRigEditMode.cpp
--- a/Engine/Plugins/Experimental/ControlRig/Source/ControlRigEditor/Private/EditMode/ControlRigEditMode.cpp
+++ b/Engine/Plugins/Experimental/ControlRig/Source/ControlRigEditor/Private/EditMode/ControlRigEditMode.cpp
@@ -1928,7 +1928,10 @@ void FControlRigEditMode::OnWidgetModeChanged(UE::Widget::EWidgetMode InWidgetMo
 		TGuardValue<bool> ReentrantGuardSelf(bIsChangingCoordSystem, true);
 
 		int32 WidgetMode = (int32)GLevelEditorModeTools().GetWidgetMode();
-		GLevelEditorModeTools().SetCoordSystem(CoordSystemPerWidgetMode[WidgetMode]);
+		if (WidgetMode >= 0 && WidgetMode < CoordSystemPerWidgetMode.Num())
+		{
+			GLevelEditorModeTools().SetCoordSystem(CoordSystemPerWidgetMode[WidgetMode]);
+		}
 	}
 }
 
@@ -1938,7 +1941,10 @@ void FControlRigEditMode::OnCoordSystemChanged(ECoordSystem InCoordSystem)
 
 	int32 WidgetMode = (int32)GLevelEditorModeTools().GetWidgetMode();
 	ECoordSystem CoordSystem = GLevelEditorModeTools().GetCoordSystem();
-	CoordSystemPerWidgetMode[WidgetMode] = CoordSystem;
+	if (WidgetMode >= 0 && WidgetMode < CoordSystemPerWidgetMode.Num())
+	{
+		CoordSystemPerWidgetMode[WidgetMode] = CoordSystem;
+	}
 }
 
 void FControlRigEditMode::SetGizmoTransform(AControlRigGizmoActor* GizmoActor, const FTransform& InTransform)
@@ -2121,6 +2127,8 @@ bool FControlRigEditMode::ModeSupportedByGizmoActor(const AControlRigGizmoActor*
 			{
 				switch (InMode)
 				{
+					case UE::Widget::WM_None:
+						return true;
 					case UE::Widget::WM_Rotate:
 					{
 						switch (ControlElement->Settings.ControlType)
