Sequencer: Rework control rig track so that it's not dependent on the name of the control rig, which can change. Now instead of calling FindOrCreateTrackForObject with the name of the control rig, there's a specialized implementation of FindOrCreateControlRigTrackForObject that looks for a track for the given binding that has a control rig that matches the control rig that's being edited.

diff --git a/Engine/Plugins/Experimental/ControlRig/Source/ControlRigEditor/Private/Sequencer/ControlRigParameterTrackEditor.cpp b/Engine/Plugins/Experimental/ControlRig/Source/ControlRigEditor/Private/Sequencer/ControlRigParameterTrackEditor.cpp
--- a/Engine/Plugins/Experimental/ControlRig/Source/ControlRigEditor/Private/Sequencer/ControlRigParameterTrackEditor.cpp
+++ b/Engine/Plugins/Experimental/ControlRig/Source/ControlRigEditor/Private/Sequencer/ControlRigParameterTrackEditor.cpp
@@ -1554,6 +1554,38 @@ FMovieSceneTrackEditor::FFindOrCreateHandleResult FControlRigParameterTrackEdito
 	return Result;
 }
 
+FMovieSceneTrackEditor::FFindOrCreateTrackResult FControlRigParameterTrackEditor::FindOrCreateControlRigTrackForObject(FGuid ObjectHandle, UControlRig* ControlRig, FName PropertyName, bool bCreateTrackIfMissing)
+{
+	FFindOrCreateTrackResult Result;
+	bool bTrackExisted = false;
+
+	UMovieScene* MovieScene = GetSequencer()->GetFocusedMovieSceneSequence()->GetMovieScene();
+
+	if (FMovieSceneBinding* Binding = MovieScene->FindBinding(ObjectHandle))
+	{
+		for (UMovieSceneTrack* Track : Binding->GetTracks())
+		{
+			if (UMovieSceneControlRigParameterTrack* ControlRigParameterTrack = Cast<UMovieSceneControlRigParameterTrack>(Track))
+			{
+				if (ControlRigParameterTrack->GetControlRig() == ControlRig)
+				{
+					Result.Track = ControlRigParameterTrack;
+					bTrackExisted = true;
+				}
+			}
+		}
+	}
+
+	if (!Result.Track && bCreateTrackIfMissing)
+	{
+		Result.Track = AddTrack(MovieScene, ObjectHandle, UMovieSceneControlRigParameterTrack::StaticClass(), PropertyName);
+	}
+
+	Result.bWasCreated = bTrackExisted == false && Result.Track != nullptr;
+
+	return Result;
+}
+
 
 void FControlRigParameterTrackEditor::HandleControlSelected(UControlRig* Subject, FRigControlElement* ControlElement, bool bSelected)
 {
@@ -1601,7 +1633,7 @@ void FControlRigParameterTrackEditor::HandleControlSelected(UControlRig* Subject
 			return;
 		}
 
-		FFindOrCreateTrackResult TrackResult = FindOrCreateTrackForObject(ObjectHandle, UMovieSceneControlRigParameterTrack::StaticClass(), ControlRigName, bCreateTrack);
+		FFindOrCreateTrackResult TrackResult = FindOrCreateControlRigTrackForObject(ObjectHandle, Subject, ControlRigName, bCreateTrack);
 		UMovieSceneControlRigParameterTrack* Track = CastChecked<UMovieSceneControlRigParameterTrack>(TrackResult.Track, ECastCheckedType::NullAllowed);
 		if (Track)
 		{
@@ -1712,7 +1744,7 @@ void FControlRigParameterTrackEditor::HandleOnInitialized(UControlRig* ControlRi
 			return;
 		}
 
-		FFindOrCreateTrackResult TrackResult = FindOrCreateTrackForObject(ObjectHandle, UMovieSceneControlRigParameterTrack::StaticClass(), ControlRigName, bCreateTrack);
+		FFindOrCreateTrackResult TrackResult = FindOrCreateControlRigTrackForObject(ObjectHandle, ControlRig, ControlRigName, bCreateTrack);
 		UMovieSceneControlRigParameterTrack* Track = CastChecked<UMovieSceneControlRigParameterTrack>(TrackResult.Track, ECastCheckedType::NullAllowed);
 		if (Track)
 		{
@@ -1964,7 +1996,7 @@ FKeyPropertyResult FControlRigParameterTrackEditor::AddKeysToControlRigHandle(US
 
 	// Try to find an existing Track, and if one doesn't exist check the key params and create one if requested.
 
-	FFindOrCreateTrackResult TrackResult = FindOrCreateTrackForObject(ObjectHandle, TrackClass, ControlRigName, bCreateTrack);
+	FFindOrCreateTrackResult TrackResult = FindOrCreateControlRigTrackForObject(ObjectHandle, InControlRig, ControlRigName, bCreateTrack);
 	UMovieSceneControlRigParameterTrack* Track = CastChecked<UMovieSceneControlRigParameterTrack>(TrackResult.Track, ECastCheckedType::NullAllowed);
 
 	bool bTrackCreated = TrackResult.bWasCreated;
@@ -2046,7 +2078,7 @@ void FControlRigParameterTrackEditor::AddControlKeys(USceneComponent *InSceneCom
 	{
 		return;
 	}
-	FFindOrCreateTrackResult TrackResult = FindOrCreateTrackForObject(ObjectHandle, UMovieSceneControlRigParameterTrack::StaticClass(), ControlRigName, bCreateTrack);
+	FFindOrCreateTrackResult TrackResult = FindOrCreateControlRigTrackForObject(ObjectHandle, InControlRig, ControlRigName, bCreateTrack);
 	UMovieSceneControlRigParameterTrack* Track = CastChecked<UMovieSceneControlRigParameterTrack>(TrackResult.Track, ECastCheckedType::NullAllowed);
 	if (Track)
 	{
diff --git a/Engine/Plugins/Experimental/ControlRig/Source/ControlRigEditor/Private/Sequencer/ControlRigParameterTrackEditor.h b/Engine/Plugins/Experimental/ControlRig/Source/ControlRigEditor/Private/Sequencer/ControlRigParameterTrackEditor.h
--- a/Engine/Plugins/Experimental/ControlRig/Source/ControlRigEditor/Private/Sequencer/ControlRigParameterTrackEditor.h
+++ b/Engine/Plugins/Experimental/ControlRig/Source/ControlRigEditor/Private/Sequencer/ControlRigParameterTrackEditor.h
@@ -112,9 +112,12 @@ private:
 	/** Post Edit Delegates */
 	void OnPropagateObjectChanges(UObject* InChangedObject);
 
-	/** Handle Creattion for SkelMeshComp or Actor Owner, either may have a binding*/
+	/** Handle Creation for SkelMeshComp or Actor Owner, either may have a binding*/
 	FMovieSceneTrackEditor::FFindOrCreateHandleResult FindOrCreateHandleToSceneCompOrOwner(USceneComponent* InComp);
-	
+
+	/** Handle Creation for control rig track given the object binding and the control rig */	
+	FMovieSceneTrackEditor::FFindOrCreateTrackResult FindOrCreateControlRigTrackForObject(FGuid ObjectBinding, UControlRig* ControlRig, FName PropertyName = NAME_None, bool bCreateTrackIfMissing = true);
+
 	/** Import FBX*/
 	void ImportFBX(UMovieSceneControlRigParameterTrack* InTrack, UMovieSceneControlRigParameterSection* InSection, 
 		TArray<FFBXNodeAndChannels>* NodeAndChannels);
