Spring Interpolation rig unit to help with leg IK.

diff --git a/Engine/Plugins/Experimental/ControlRig/Source/ControlRig/Private/Units/Simulation/RigUnit_SpringInterp.cpp b/Engine/Plugins/Experimental/ControlRig/Source/ControlRig/Private/Units/Simulation/RigUnit_SpringInterp.cpp
new file mode 100644
--- /dev/null
+++ b/Engine/Plugins/Experimental/ControlRig/Source/ControlRig/Private/Units/Simulation/RigUnit_SpringInterp.cpp
@@ -0,0 +1,32 @@
+// Copyright Epic Games, Inc. All Rights Reserved.
+
+#include "Units/Simulation/RigUnit_SpringInterp.h"
+#include "Units/RigUnitContext.h"
+
+FRigUnit_SpringInterp_Execute()
+{
+	DECLARE_SCOPE_HIERARCHICAL_COUNTER_RIGUNIT()
+
+	if (Context.State == EControlRigState::Init)
+	{
+		SpringState.Reset();
+	}
+	else
+	{
+		Result = UKismetMathLibrary::FloatSpringInterp(Current, Target, SpringState, Stiffness, CriticalDamping, Context.DeltaTime, Mass);
+	}
+}
+
+FRigUnit_SpringInterpVector_Execute()
+{
+    DECLARE_SCOPE_HIERARCHICAL_COUNTER_RIGUNIT()
+
+	if (Context.State == EControlRigState::Init)
+	{
+		SpringState.Reset();
+	}
+	else
+	{
+		Result = UKismetMathLibrary::VectorSpringInterp(Current, Target, SpringState, Stiffness, CriticalDamping, Context.DeltaTime, Mass);
+	}
+}
\ No newline at end of file
diff --git a/Engine/Plugins/Experimental/ControlRig/Source/ControlRig/Private/Units/Simulation/RigUnit_SpringInterp.h b/Engine/Plugins/Experimental/ControlRig/Source/ControlRig/Private/Units/Simulation/RigUnit_SpringInterp.h
new file mode 100644
--- /dev/null
+++ b/Engine/Plugins/Experimental/ControlRig/Source/ControlRig/Private/Units/Simulation/RigUnit_SpringInterp.h
@@ -0,0 +1,91 @@
+// Copyright Epic Games, Inc. All Rights Reserved.
+
+#pragma once
+
+#include "RigUnit_SimBase.h"
+#include "Kismet/KismetMathLibrary.h"
+#include "RigUnit_SpringInterp.generated.h"
+
+/**
+ * Uses a simple spring model to interpolate a float from Current to Target.
+ */
+USTRUCT(meta=(DisplayName="Spring Interpolate", Keywords="Alpha,SpringInterpolate", PrototypeName = "SpringInterp"))
+struct CONTROLRIG_API FRigUnit_SpringInterp : public FRigUnit_SimBase
+{
+	GENERATED_BODY()
+
+	FRigUnit_SpringInterp()
+	{
+		Current = Target = Result = 0.0f;
+		Stiffness = 10.0f;
+		CriticalDamping = 2.0f;
+		Mass = 10.0f;
+		SpringState = FFloatSpringState();
+	}
+
+	RIGVM_METHOD()
+	virtual void Execute(const FRigUnitContext& Context) override;
+
+	UPROPERTY(meta = (Input))
+	float Current;
+
+	UPROPERTY(meta=(Input))
+	float Target;
+
+	UPROPERTY(meta=(Input))
+	float Stiffness;
+
+	UPROPERTY(meta=(Input))
+	float CriticalDamping;
+
+	UPROPERTY(meta=(Input))
+	float Mass;
+
+	UPROPERTY(meta=(Output))
+	float Result;
+
+	UPROPERTY()
+	FFloatSpringState SpringState;
+};
+
+/**
+ * Uses a simple spring model to interpolate a vector from Current to Target.
+ */
+USTRUCT(meta=(DisplayName="Spring Interpolate", Keywords="Alpha,SpringInterpolate", PrototypeName = "SpringInterp", MenuDescSuffix = "(Vector)"))
+struct CONTROLRIG_API FRigUnit_SpringInterpVector : public FRigUnit_SimBase
+{
+	GENERATED_BODY()
+	
+	FRigUnit_SpringInterpVector()
+	{
+		Current = Target = Result = FVector::ZeroVector;
+		Stiffness = 10.0f;
+		CriticalDamping = 2.0f;
+		Mass = 10.0f;
+		SpringState = FVectorSpringState();
+	}
+
+	RIGVM_METHOD()
+	virtual void Execute(const FRigUnitContext& Context) override;
+
+	UPROPERTY(meta = (Input))
+	FVector Current;
+
+	UPROPERTY(meta=(Input))
+	FVector Target;
+
+	UPROPERTY(meta=(Input))
+	float Stiffness;
+
+	UPROPERTY(meta=(Input))
+	float CriticalDamping;
+
+	UPROPERTY(meta=(Input))
+	float Mass;
+
+	UPROPERTY(meta=(Output))
+	FVector Result;
+
+	UPROPERTY()
+	FVectorSpringState SpringState;
+};
\ No newline at end of file
