Sequencer:Control Rig: Recreate Rig on Load since it may have changed.

diff --git a/Engine/Plugins/Experimental/ControlRig/Source/ControlRig/Private/Sequencer/MovieSceneControlRigParameterTrack.cpp b/Engine/Plugins/Experimental/ControlRig/Source/ControlRig/Private/Sequencer/MovieSceneControlRigParameterTrack.cpp
--- a/Engine/Plugins/Experimental/ControlRig/Source/ControlRig/Private/Sequencer/MovieSceneControlRigParameterTrack.cpp
+++ b/Engine/Plugins/Experimental/ControlRig/Source/ControlRig/Private/Sequencer/MovieSceneControlRigParameterTrack.cpp
@@ -329,8 +329,7 @@ void UMovieSceneControlRigParameterTrack::ReconstructControlRig()
 				UMovieSceneControlRigParameterSection* CRSection = Cast<UMovieSceneControlRigParameterSection>(Section);
 				if (CRSection)
 				{
-					CRSection->ControlRig = ControlRig;
-					CRSection->ReconstructChannelProxy(true);
+					CRSection->RecreateWithThisControlRig(ControlRig, false);
 				}
 			}
 		}
