Apply time delta clamp and fixed time step to Spring Interp rig unit.

diff --git a/Engine/Plugins/Experimental/ControlRig/Source/ControlRig/Private/Units/Simulation/RigUnit_SpringInterp.cpp b/Engine/Plugins/Experimental/ControlRig/Source/ControlRig/Private/Units/Simulation/RigUnit_SpringInterp.cpp
--- a/Engine/Plugins/Experimental/ControlRig/Source/ControlRig/Private/Units/Simulation/RigUnit_SpringInterp.cpp
+++ b/Engine/Plugins/Experimental/ControlRig/Source/ControlRig/Private/Units/Simulation/RigUnit_SpringInterp.cpp
@@ -3,6 +3,12 @@
 #include "Units/Simulation/RigUnit_SpringInterp.h"
 #include "Units/RigUnitContext.h"
 
+namespace RigUnitSpringInterpConstants
+{
+	static const float FixedTimeStep = 1.0f / 60.0f;
+	static const float MaxTimeStep = 0.1f;
+}
+
 FRigUnit_SpringInterp_Execute()
 {
 	DECLARE_SCOPE_HIERARCHICAL_COUNTER_RIGUNIT()
@@ -13,7 +19,17 @@ FRigUnit_SpringInterp_Execute()
 	}
 	else
 	{
-		Result = UKismetMathLibrary::FloatSpringInterp(Current, Target, SpringState, Stiffness, CriticalDamping, Context.DeltaTime, Mass);
+		// Clamp to avoid large time deltas.
+		float RemainingTime = FMath::Min(Context.DeltaTime, RigUnitSpringInterpConstants::MaxTimeStep);
+
+		Result = Current;
+		while (RemainingTime >= RigUnitSpringInterpConstants::FixedTimeStep)
+		{
+			Result = UKismetMathLibrary::FloatSpringInterp(Result, Target, SpringState, Stiffness, CriticalDamping, RigUnitSpringInterpConstants::FixedTimeStep, Mass);
+			RemainingTime -= RigUnitSpringInterpConstants::FixedTimeStep;
+		}
+
+		Result = UKismetMathLibrary::FloatSpringInterp(Result, Target, SpringState, Stiffness, CriticalDamping, RemainingTime, Mass);
 	}
 }
 
@@ -27,6 +43,16 @@ FRigUnit_SpringInterpVector_Execute()
 	}
 	else
 	{
-		Result = UKismetMathLibrary::VectorSpringInterp(Current, Target, SpringState, Stiffness, CriticalDamping, Context.DeltaTime, Mass);
+		// Clamp to avoid large time deltas.
+		float RemainingTime = FMath::Min(Context.DeltaTime, RigUnitSpringInterpConstants::MaxTimeStep);
+
+		Result = Current;
+		while (RemainingTime >= RigUnitSpringInterpConstants::FixedTimeStep)
+		{
+			Result = UKismetMathLibrary::VectorSpringInterp(Result, Target, SpringState, Stiffness, CriticalDamping, RigUnitSpringInterpConstants::FixedTimeStep, Mass);
+			RemainingTime -= RigUnitSpringInterpConstants::FixedTimeStep;
+		}
+
+		Result = UKismetMathLibrary::VectorSpringInterp(Result, Target, SpringState, Stiffness, CriticalDamping, RemainingTime, Mass);
 	}
 }
\ No newline at end of file
