Control Rig: Add sub title functionality for graph nodes

diff --git a/Engine/Plugins/Experimental/ControlRig/Source/ControlRigDeveloper/Private/Graph/ControlRigGraph.cpp b/Engine/Plugins/Experimental/ControlRig/Source/ControlRigDeveloper/Private/Graph/ControlRigGraph.cpp
--- a/Engine/Plugins/Experimental/ControlRig/Source/ControlRigDeveloper/Private/Graph/ControlRigGraph.cpp
+++ b/Engine/Plugins/Experimental/ControlRig/Source/ControlRigDeveloper/Private/Graph/ControlRigGraph.cpp
@@ -582,6 +582,7 @@ void UControlRigGraph::HandleModifiedEvent(ERigVMGraphNotifType InNotifType, URi
 			break;
 		}
 		case ERigVMGraphNotifType::VariableRenamed:
+		case ERigVMGraphNotifType::NodeReferenceChanged:
 		{
 			if (URigVMNode* ModelNode = Cast<URigVMNode>(InSubject))
 			{
diff --git a/Engine/Plugins/Experimental/ControlRig/Source/ControlRigDeveloper/Private/Graph/ControlRigGraphNode.cpp b/Engine/Plugins/Experimental/ControlRig/Source/ControlRigDeveloper/Private/Graph/ControlRigGraphNode.cpp
--- a/Engine/Plugins/Experimental/ControlRig/Source/ControlRigDeveloper/Private/Graph/ControlRigGraphNode.cpp
+++ b/Engine/Plugins/Experimental/ControlRig/Source/ControlRigDeveloper/Private/Graph/ControlRigGraphNode.cpp
@@ -26,6 +26,7 @@
 #include "RigVMModel/Nodes/RigVMFunctionReferenceNode.h"
 #include "RigVMModel/Nodes/RigVMFunctionEntryNode.h"
 #include "RigVMModel/Nodes/RigVMFunctionReturnNode.h"
+#include "RigVMModel/Nodes/RigVMCollapseNode.h"
 
 #if WITH_EDITOR
 #include "IControlRigEditorModule.h"
@@ -36,6 +37,7 @@
 UControlRigGraphNode::UControlRigGraphNode()
 : Dimensions(0.0f, 0.0f)
 , NodeTitle(FText::GetEmpty())
+, FullNodeTitle(FText::GetEmpty())
 , CachedTitleColor(FLinearColor(0.f, 0.f, 0.f, 0.f))
 , CachedNodeColor(FLinearColor(0.f, 0.f, 0.f, 0.f))
 {
@@ -47,6 +49,7 @@ FText UControlRigGraphNode::GetNodeTitle(ENodeTitleType::Type TitleType) const
 {
 	if(NodeTitle.IsEmpty())
 	{
+		FString SubTitle;
 		if(URigVMNode* ModelNode = GetModelNode())
 		{
 			if (URigVMUnitNode* UnitNode = Cast<URigVMUnitNode>(ModelNode))
@@ -61,6 +64,26 @@ FText UControlRigGraphNode::GetNodeTitle(ENodeTitleType::Type TitleType) const
 				}
 			}
 
+			else if(URigVMFunctionReferenceNode* FunctionReferenceNode = Cast<URigVMFunctionReferenceNode>(ModelNode))
+			{
+				UPackage* ReferencedPackage = FunctionReferenceNode->GetReferencedNode()->GetOutermost();
+				if(ReferencedPackage != ModelNode->GetOutermost())
+				{
+					SubTitle = FString::Printf(TEXT("From %s"), *ReferencedPackage->GetName());
+				}
+				else
+				{
+					static const FString LocalFunctionString = TEXT("Local Function");
+					SubTitle = LocalFunctionString;
+				}
+			}
+
+			else if(URigVMCollapseNode* CollapseNode = Cast<URigVMCollapseNode>(ModelNode))
+			{
+				static const FString CollapseNodeString = TEXT("Collapsed Graph");
+				SubTitle = CollapseNodeString;
+			}
+
 			if (NodeTitle.IsEmpty())
 			{
 				NodeTitle = FText::FromString(ModelNode->GetNodeTitle());
@@ -71,8 +94,19 @@ FText UControlRigGraphNode::GetNodeTitle(ENodeTitleType::Type TitleType) const
 		{
 			NodeTitle = FText::FromString(FString::Printf(TEXT("%s (Deprecated)"), *NodeTitle.ToString()));
 		}
+
+		FullNodeTitle = NodeTitle;
+
+		if(!SubTitle.IsEmpty())
+		{
+			FullNodeTitle = FText::FromString(FString::Printf(TEXT("%s\n%s"), *NodeTitle.ToString(), *SubTitle));
+		}
 	}
 
+	if(TitleType == ENodeTitleType::FullTitle)
+	{
+		return FullNodeTitle;
+	}
 	return NodeTitle;
 }
 
@@ -790,7 +824,7 @@ FText UControlRigGraphNode::GetTooltipText() const
 void UControlRigGraphNode::InvalidateNodeTitle() const
 {
 	NodeTitle = FText();
-
+	FullNodeTitle = FText();
 	NodeTitleDirtied.ExecuteIfBound();
 }
 
diff --git a/Engine/Plugins/Experimental/ControlRig/Source/ControlRigDeveloper/Public/Graph/ControlRigGraphNode.h b/Engine/Plugins/Experimental/ControlRig/Source/ControlRigDeveloper/Public/Graph/ControlRigGraphNode.h
--- a/Engine/Plugins/Experimental/ControlRig/Source/ControlRigDeveloper/Public/Graph/ControlRigGraphNode.h
+++ b/Engine/Plugins/Experimental/ControlRig/Source/ControlRigDeveloper/Public/Graph/ControlRigGraphNode.h
@@ -70,6 +70,9 @@ private:
 	/** The cached node titles */
 	mutable FText NodeTitle;
 
+	/** The cached fulol node title */
+	mutable FText FullNodeTitle;
+
 public:
 
 	DECLARE_DELEGATE(FNodeTitleDirtied);
diff --git a/Engine/Source/Developer/RigVMDeveloper/Private/RigVMModel/RigVMController.cpp b/Engine/Source/Developer/RigVMDeveloper/Private/RigVMModel/RigVMController.cpp
--- a/Engine/Source/Developer/RigVMDeveloper/Private/RigVMModel/RigVMController.cpp
+++ b/Engine/Source/Developer/RigVMDeveloper/Private/RigVMModel/RigVMController.cpp
@@ -68,6 +68,10 @@ void URigVMController::SetGraph(URigVMGraph* InGraph)
 	URigVMGraph* PreviousGraph = GetGraph();
 	if (PreviousGraph)
 	{
+		if(PreviousGraph == InGraph)
+		{
+			return;
+		}
 		PreviousGraph->OnModified().RemoveAll(this);
 	}
 
@@ -3363,6 +3367,9 @@ void URigVMController::SetReferencedFunction(URigVMFunctionReferenceNode* InFunc
 	NewReferences.FunctionReferences.Add(InFunctionRefNode);
 
 	InFunctionRefNode->SetReferencedNode(InNewReferencedNode);
+	
+	FRigVMControllerGraphGuard GraphGuard(this, InFunctionRefNode->GetGraph(), false);
+	GetGraph()->Notify(ERigVMGraphNotifType::NodeReferenceChanged, InFunctionRefNode);
 }
 
 void URigVMController::RefreshFunctionPins(URigVMNode* InNode, bool bNotify)
diff --git a/Engine/Source/Developer/RigVMDeveloper/Public/RigVMModel/RigVMNotifications.h b/Engine/Source/Developer/RigVMDeveloper/Public/RigVMModel/RigVMNotifications.h
--- a/Engine/Source/Developer/RigVMDeveloper/Public/RigVMModel/RigVMNotifications.h
+++ b/Engine/Source/Developer/RigVMDeveloper/Public/RigVMModel/RigVMNotifications.h
@@ -49,6 +49,7 @@ enum class ERigVMGraphNotifType : uint8
 	InteractionBracketCanceled, // A bracket has been canceled (Subject == nullptr)
 	PinBoundVariableChanged, // A pin has been bound or unbound to / from a variable (Subject == URigVMPin)
 	NodeRenamed, // A node has been renamed in the graph (Subject == URigVMNode)
+	NodeReferenceChanged, // A node has changed it's referenced function
 	NodeCategoryChanged, // A node's category has changed (Subject == URigVMNode)
 	NodeKeywordsChanged, // A node's keywords have changed (Subject == URigVMNode)
 	Invalid // The max for this enum (used for guarding)
