Control Rig: Filter out out-of-asset models / graphs when interacting with Control Rig Blueprint

diff --git a/Engine/Plugins/Experimental/ControlRig/Source/ControlRigDeveloper/Private/ControlRigBlueprint.cpp b/Engine/Plugins/Experimental/ControlRig/Source/ControlRigDeveloper/Private/ControlRigBlueprint.cpp
--- a/Engine/Plugins/Experimental/ControlRig/Source/ControlRigDeveloper/Private/ControlRigBlueprint.cpp
+++ b/Engine/Plugins/Experimental/ControlRig/Source/ControlRigDeveloper/Private/ControlRigBlueprint.cpp
@@ -721,6 +721,11 @@ URigVMGraph* UControlRigBlueprint::GetModel(const UEdGraph* InEdGraph) const
 		return Model;
 	}
 
+	if(InEdGraph->GetOutermost() != GetOutermost())
+	{
+		return nullptr;
+	}
+
 #if WITH_EDITORONLY_DATA
 	if (InEdGraph == FunctionLibraryEdGraph)
 	{
@@ -998,6 +1003,11 @@ UEdGraph* UControlRigBlueprint::GetEdGraph(URigVMGraph* InModel) const
 		return nullptr;
 	}
 
+	if(InModel->GetOutermost() != GetOutermost())
+	{
+		return nullptr;
+	}
+
 #if WITH_EDITORONLY_DATA
 	if (InModel == FunctionLibrary)
 	{
