Control Rig: Add SetupUndo flag to Control Setters so we don't spame Undo stack when evaluating Sequencer.

diff --git a/Engine/Plugins/Experimental/ControlRig/Source/ControlRig/Private/ControlRig.cpp b/Engine/Plugins/Experimental/ControlRig/Source/ControlRig/Private/ControlRig.cpp
--- a/Engine/Plugins/Experimental/ControlRig/Source/ControlRig/Private/ControlRig.cpp
+++ b/Engine/Plugins/Experimental/ControlRig/Source/ControlRig/Private/ControlRig.cpp
@@ -1160,7 +1160,7 @@ FTransform UControlRig::GetControlGlobalTransform(const FName& InControlName) co
 	return DynamicHierarchy->GetGlobalTransform(FRigElementKey(InControlName, ERigElementType::Control), false);
 }
 
-bool UControlRig::SetControlGlobalTransform(const FName& InControlName, const FTransform& InGlobalTransform, bool bNotify, const FRigControlModifiedContext& Context)
+bool UControlRig::SetControlGlobalTransform(const FName& InControlName, const FTransform& InGlobalTransform, bool bNotify, const FRigControlModifiedContext& Context, bool bSetupUndo)
 {
 	FTransform GlobalTransform = InGlobalTransform;
 	if (IsSetupModeEnabled())
@@ -1178,7 +1178,7 @@ bool UControlRig::SetControlGlobalTransform(const FName& InControlName, const FT
 		}
 	}
 
-	SetControlValue(InControlName, Value, bNotify, Context);
+	SetControlValue(InControlName, Value, bNotify, Context, bSetupUndo);
 	return true;
 }
 
@@ -1204,7 +1204,7 @@ FRigControlValue UControlRig::GetControlValueFromGlobalTransform(const FName& In
 	return Value;
 }
 
-void UControlRig::SetControlLocalTransform(const FName& InControlName, const FTransform& InLocalTransform, bool bNotify, const FRigControlModifiedContext& Context)
+void UControlRig::SetControlLocalTransform(const FName& InControlName, const FTransform& InLocalTransform, bool bNotify, const FRigControlModifiedContext& Context, bool bSetupUndo)
 {
 	if (FRigControlElement* ControlElement = FindControl(InControlName))
 	{
@@ -1216,7 +1216,7 @@ void UControlRig::SetControlLocalTransform(const FName& InControlName, const FTr
 			OnFilterControl.Broadcast(this, ControlElement, Value);
 			
 		}
-		SetControlValue(InControlName, Value, bNotify, Context);
+		SetControlValue(InControlName, Value, bNotify, Context, bSetupUndo);
 	}
 }
 
diff --git a/Engine/Plugins/Experimental/ControlRig/Source/ControlRig/Private/Rigs/FKControlRig.cpp b/Engine/Plugins/Experimental/ControlRig/Source/ControlRig/Private/Rigs/FKControlRig.cpp
--- a/Engine/Plugins/Experimental/ControlRig/Source/ControlRig/Private/Rigs/FKControlRig.cpp
+++ b/Engine/Plugins/Experimental/ControlRig/Source/ControlRig/Private/Rigs/FKControlRig.cpp
@@ -91,13 +91,13 @@ void UFKControlRig::ExecuteUnits(FRigUnitContext& InOutContext, const FName& InE
 				{
 					case EControlRigFKRigExecuteMode::Replace:
 					{
-						GetHierarchy()->SetCurveValue(CurveElement, CurveValue, false);
+						GetHierarchy()->SetCurveValue(CurveElement, CurveValue, false /*bSetupUndo*/);
 						break;
 					}
 					case EControlRigFKRigExecuteMode::Additive:
 					{
 						const float PreviousValue = GetHierarchy()->GetCurveValue(CurveElement);
-						GetHierarchy()->SetCurveValue(CurveElement, PreviousValue + CurveValue, false);
+						GetHierarchy()->SetCurveValue(CurveElement, PreviousValue + CurveValue, false /*bSetupUndo*/);
 						break;
 					}
 				}
@@ -108,7 +108,9 @@ void UFKControlRig::ExecuteUnits(FRigUnitContext& InOutContext, const FName& InE
 	else if (InEventName == FRigUnit_InverseExecution::EventName)
 	{
 		FRigVMExecuteContext VMContext;
-
+		const bool bNotify = true;
+		const FRigControlModifiedContext Context = FRigControlModifiedContext();
+		const bool bSetupUndo = false;
 		GetHierarchy()->ForEach<FRigBoneElement>([&](FRigBoneElement* BoneElement) -> bool
         {
             const FName ControlName = GetControlName(BoneElement->GetName());
@@ -118,7 +120,7 @@ void UFKControlRig::ExecuteUnits(FRigUnitContext& InOutContext, const FName& InE
 			if (IsControlActive[ControlIndex])
 			{
 				const FEulerTransform EulerTransform(GetHierarchy()->GetTransform(BoneElement, ERigTransformType::CurrentLocal));
-				SetControlValue(ControlName, FRigControlValue::Make(EulerTransform));
+				SetControlValue(ControlName, FRigControlValue::Make(EulerTransform), bNotify, Context, bSetupUndo);
 			}
 			
 			return true;
@@ -133,7 +135,7 @@ void UFKControlRig::ExecuteUnits(FRigUnitContext& InOutContext, const FName& InE
 			if (IsControlActive[ControlIndex])
 			{
 				const float CurveValue = GetHierarchy()->GetCurveValue(CurveElement);
-				SetControlValue(ControlName, FRigControlValue::Make(CurveValue));
+				SetControlValue(ControlName, FRigControlValue::Make(CurveValue), bNotify, Context, bSetupUndo);
 			}
 
 			return true;
@@ -296,11 +298,11 @@ void UFKControlRig::ToggleApplyMode()
 	{
 		ApplyMode = EControlRigFKRigExecuteMode::Additive;
 	}
-
 	if (ApplyMode == EControlRigFKRigExecuteMode::Additive)
 	{
 		FRigControlModifiedContext Context;
 		Context.SetKey = EControlRigSetKey::Never;
+		const bool bSetupUndo = false;
 
 		FTransform ZeroScale = FTransform::Identity;
 		ZeroScale.SetScale3D(FVector::ZeroVector);
@@ -310,11 +312,11 @@ void UFKControlRig::ToggleApplyMode()
         {
             if (ControlElement->Settings.ControlType == ERigControlType::EulerTransform)
 			{
-				SetControlValue<FEulerTransform>(ControlElement->GetName(), EulerZero, true, Context);
+				SetControlValue<FEulerTransform>(ControlElement->GetName(), EulerZero, true, Context, bSetupUndo);
 			}
 			else if (ControlElement->Settings.ControlType == ERigControlType::Float)
 			{
-				SetControlValue<float>(ControlElement->GetName(), 0.f, true, Context);
+				SetControlValue<float>(ControlElement->GetName(), 0.f, true, Context, bSetupUndo);
 			}
 
 			return true;
@@ -324,18 +326,19 @@ void UFKControlRig::ToggleApplyMode()
 	{
 		FRigControlModifiedContext Context;
 		Context.SetKey = EControlRigSetKey::Never;
+		const bool bSetupUndo = false;
 
 		GetHierarchy()->ForEach<FRigControlElement>([&](FRigControlElement* ControlElement) -> bool
         {
             if (ControlElement->Settings.ControlType == ERigControlType::EulerTransform)
 			{
 				const FEulerTransform InitValue = GetHierarchy()->GetControlValue(ControlElement, ERigControlValueType::Initial).Get<FEulerTransform>();
-				SetControlValue<FEulerTransform>(ControlElement->GetName(), InitValue, true, Context);
+				SetControlValue<FEulerTransform>(ControlElement->GetName(), InitValue, true, Context, bSetupUndo);
 			}
 			else if (ControlElement->Settings.ControlType == ERigControlType::Float)
 			{
 				const float InitValue = GetHierarchy()->GetControlValue(ControlElement, ERigControlValueType::Initial).Get<float>();
-				SetControlValue<float>(ControlElement->GetName(), InitValue, true, Context);
+				SetControlValue<float>(ControlElement->GetName(), InitValue, true, Context, bSetupUndo);
 			}
 
 			return true;
diff --git a/Engine/Plugins/Experimental/ControlRig/Source/ControlRig/Private/Sequencer/MovieSceneControlRigParameterTemplate.cpp b/Engine/Plugins/Experimental/ControlRig/Source/ControlRig/Private/Sequencer/MovieSceneControlRigParameterTemplate.cpp
--- a/Engine/Plugins/Experimental/ControlRig/Source/ControlRig/Private/Sequencer/MovieSceneControlRigParameterTemplate.cpp
+++ b/Engine/Plugins/Experimental/ControlRig/Source/ControlRig/Private/Sequencer/MovieSceneControlRigParameterTemplate.cpp
@@ -642,6 +642,7 @@ struct FControlRigParameterPreAnimatedTokenProducer : IMovieScenePreAnimatedToke
 				{
 					if (ControlRig->GetObjectBinding())
 					{
+						const bool bSetupUndo = false;
 						if (USkeletalMeshComponent* SkeletalMeshComponent = Cast<USkeletalMeshComponent>(ControlRig->GetObjectBinding()->GetBoundObject()))
 						{
 							SkeletalMeshRestoreState.RestoreState(SkeletalMeshComponent);
@@ -653,7 +654,7 @@ struct FControlRigParameterPreAnimatedTokenProducer : IMovieScenePreAnimatedToke
 						{
 							if (ControlRig->FindControl(Value.Name))
 							{
-								ControlRig->SetControlValue<float>(Value.Name, Value.Value, true, FRigControlModifiedContext(EControlRigSetKey::Never));
+								ControlRig->SetControlValue<float>(Value.Name, Value.Value, true, FRigControlModifiedContext(EControlRigSetKey::Never,bSetupUndo));
 							}
 						}
 
@@ -661,7 +662,7 @@ struct FControlRigParameterPreAnimatedTokenProducer : IMovieScenePreAnimatedToke
 						{
 							if (ControlRig->FindControl(Value.Name))
 							{
-								ControlRig->SetControlValue<bool>(Value.Name, Value.Value, true, FRigControlModifiedContext(EControlRigSetKey::Never));
+								ControlRig->SetControlValue<bool>(Value.Name, Value.Value, true, FRigControlModifiedContext(EControlRigSetKey::Never,bSetupUndo));
 							}
 						}
 
@@ -669,7 +670,7 @@ struct FControlRigParameterPreAnimatedTokenProducer : IMovieScenePreAnimatedToke
 						{
 							if (ControlRig->FindControl(Value.Name))
 							{
-								ControlRig->SetControlValue<int32>(Value.Name, Value.Value, true, FRigControlModifiedContext(EControlRigSetKey::Never));
+								ControlRig->SetControlValue<int32>(Value.Name, Value.Value, true, FRigControlModifiedContext(EControlRigSetKey::Never,bSetupUndo));
 							}
 						}
 
@@ -677,7 +678,7 @@ struct FControlRigParameterPreAnimatedTokenProducer : IMovieScenePreAnimatedToke
 						{
 							if (ControlRig->FindControl(Value.Name))
 							{
-								ControlRig->SetControlValue<FVector2D>(Value.Name, Value.Value, true, FRigControlModifiedContext(EControlRigSetKey::Never));
+								ControlRig->SetControlValue<FVector2D>(Value.Name, Value.Value, true, FRigControlModifiedContext(EControlRigSetKey::Never,bSetupUndo));
 							}
 						}
 
@@ -685,7 +686,7 @@ struct FControlRigParameterPreAnimatedTokenProducer : IMovieScenePreAnimatedToke
 						{
 							if (ControlRig->FindControl(Value.Name))
 							{
-								ControlRig->SetControlValue<FVector>(Value.Name, Value.Value, true, FRigControlModifiedContext(EControlRigSetKey::Never));
+								ControlRig->SetControlValue<FVector>(Value.Name, Value.Value, true, FRigControlModifiedContext(EControlRigSetKey::Never,bSetupUndo));
 							}
 						}
 
@@ -697,19 +698,19 @@ struct FControlRigParameterPreAnimatedTokenProducer : IMovieScenePreAnimatedToke
 								{
 								case ERigControlType::Transform:
 								{
-									ControlRig->SetControlValue<FTransform>(Value.Name, Value.Value, true, FRigControlModifiedContext(EControlRigSetKey::Never));
+									ControlRig->SetControlValue<FTransform>(Value.Name, Value.Value, true, FRigControlModifiedContext(EControlRigSetKey::Never,bSetupUndo));
 									break;
 								}
 								case ERigControlType::TransformNoScale:
 								{
 									FTransformNoScale NoScale = Value.Value;
-									ControlRig->SetControlValue<FTransformNoScale>(Value.Name, NoScale, true, FRigControlModifiedContext(EControlRigSetKey::Never));
+									ControlRig->SetControlValue<FTransformNoScale>(Value.Name, NoScale, true, FRigControlModifiedContext(EControlRigSetKey::Never,bSetupUndo));
 									break;
 								}
 								case ERigControlType::EulerTransform:
 								{
 									FEulerTransform EulerTransform = Value.Value;
-									ControlRig->SetControlValue<FEulerTransform>(Value.Name, EulerTransform, true, FRigControlModifiedContext(EControlRigSetKey::Never));
+									ControlRig->SetControlValue<FEulerTransform>(Value.Name, EulerTransform, true, FRigControlModifiedContext(EControlRigSetKey::Never,bSetupUndo));
 									break;
 								}
 
@@ -948,6 +949,7 @@ struct FControlRigParameterExecutionToken : IMovieSceneExecutionToken
 
 		if (Section->ControlRig)
 		{
+			const bool bSetupUndo = false;
 			Section->ControlRig->SetAbsoluteTime((float)Context.GetFrameRate().AsSeconds(Context.GetTime()));
 			if (Section->bAdditive == false)
 			{
@@ -958,7 +960,7 @@ struct FControlRigParameterExecutionToken : IMovieSceneExecutionToken
 						FRigControlElement* RigControl = Section->ControlRig->FindControl(BoolNameAndValue.ParameterName);
 						if (RigControl && RigControl->Settings.ControlType == ERigControlType::Bool)
 						{
-							Section->ControlRig->SetControlValue<bool>(BoolNameAndValue.ParameterName, BoolNameAndValue.Value, true, EControlRigSetKey::Never);
+							Section->ControlRig->SetControlValue<bool>(BoolNameAndValue.ParameterName, BoolNameAndValue.Value, true, EControlRigSetKey::Never,bSetupUndo);
 						}
 					}
 				}
@@ -970,7 +972,7 @@ struct FControlRigParameterExecutionToken : IMovieSceneExecutionToken
 						FRigControlElement* RigControl = Section->ControlRig->FindControl(IntegerNameAndValue.ParameterName);
 						if (RigControl && RigControl->Settings.ControlType == ERigControlType::Integer)
 						{
-							Section->ControlRig->SetControlValue<int32>(IntegerNameAndValue.ParameterName, IntegerNameAndValue.Value, true, EControlRigSetKey::Never);
+							Section->ControlRig->SetControlValue<int32>(IntegerNameAndValue.ParameterName, IntegerNameAndValue.Value, true, EControlRigSetKey::Never,bSetupUndo);
 						}
 					}
 				}
@@ -1031,6 +1033,7 @@ struct TControlRigParameterActuatorFloat : TMovieSceneBlendingActuator<FControlR
 		bool bWasDoNotKey = false;
 		if (Section)
 		{
+			const bool bSetupUndo = false;
 			bWasDoNotKey = Section->GetDoNotKey();
 			Section->SetDoNotKey(true);
 
@@ -1039,7 +1042,7 @@ struct TControlRigParameterActuatorFloat : TMovieSceneBlendingActuator<FControlR
 				FRigControlElement* ControlElement = Section->ControlRig->FindControl(ParameterName);
 				if (ControlElement && ControlElement->Settings.ControlType == ERigControlType::Float)
 				{
-					Section->ControlRig->SetControlValue<float>(ParameterName, InFinalValue.Value, true, EControlRigSetKey::Never);
+					Section->ControlRig->SetControlValue<float>(ParameterName, InFinalValue.Value, true, EControlRigSetKey::Never,bSetupUndo);
 				}
 			}
 
@@ -1095,6 +1098,7 @@ struct TControlRigParameterActuatorVector2D : TMovieSceneBlendingActuator<FContr
 		bool bWasDoNotKey = false;
 		if (Section)
 		{
+			const bool bSetupUndo = false;
 			bWasDoNotKey = Section->GetDoNotKey();
 			Section->SetDoNotKey(true);
 
@@ -1103,7 +1107,7 @@ struct TControlRigParameterActuatorVector2D : TMovieSceneBlendingActuator<FContr
 				FRigControlElement* ControlElement = Section->ControlRig->FindControl(ParameterName);
 				if (ControlElement && (ControlElement->Settings.ControlType == ERigControlType::Vector2D))
 				{
-					Section->ControlRig->SetControlValue<FVector2D>(ParameterName, InFinalValue.Value, true, EControlRigSetKey::Never);
+					Section->ControlRig->SetControlValue<FVector2D>(ParameterName, InFinalValue.Value, true, EControlRigSetKey::Never,bSetupUndo);
 				}
 			}
 
@@ -1157,6 +1161,7 @@ struct TControlRigParameterActuatorVector : TMovieSceneBlendingActuator<FControl
 		bool bWasDoNotKey = false;
 		if (Section)
 		{
+			const bool bSetupUndo = false;
 			bWasDoNotKey = Section->GetDoNotKey();
 			Section->SetDoNotKey(true);
 			if (Section->ControlRig && (Section->ControlsToSet.Num() == 0 || Section->ControlsToSet.Contains(ParameterName)))
@@ -1164,7 +1169,7 @@ struct TControlRigParameterActuatorVector : TMovieSceneBlendingActuator<FControl
 				FRigControlElement* ControlElement = Section->ControlRig->FindControl(ParameterName);
 				if (ControlElement && (ControlElement->Settings.ControlType == ERigControlType::Position || ControlElement->Settings.ControlType == ERigControlType::Scale || ControlElement->Settings.ControlType == ERigControlType::Rotator))
 				{
-						Section->ControlRig->SetControlValue<FVector>(ParameterName, InFinalValue.Value, true, EControlRigSetKey::Never);
+						Section->ControlRig->SetControlValue<FVector>(ParameterName, InFinalValue.Value, true, EControlRigSetKey::Never,bSetupUndo);
 				}
 			}
 
@@ -1240,23 +1245,23 @@ struct TControlRigParameterActuatorTransform : TMovieSceneBlendingActuator<FCont
 		{
 			bWasDoNotKey = Section->GetDoNotKey();
 			Section->SetDoNotKey(true);
-		
+			const bool bSetupUndo = false;
 			if (Section->ControlRig && (Section->ControlsToSet.Num() == 0 || Section->ControlsToSet.Contains(ParameterName)))
 			{
 				FRigControlElement* ControlElement = Section->ControlRig->FindControl(ParameterName);
 				if (ControlElement && ControlElement->Settings.ControlType == ERigControlType::Transform)
 				{
-						Section->ControlRig->SetControlValue<FTransform>(ParameterName, InFinalValue.Value,true, EControlRigSetKey::Never);
+						Section->ControlRig->SetControlValue<FTransform>(ParameterName, InFinalValue.Value,true, EControlRigSetKey::Never,bSetupUndo);
 				}
 				else if (ControlElement && ControlElement->Settings.ControlType == ERigControlType::TransformNoScale)
 				{
 					FTransformNoScale NoScale = InFinalValue.Value;
-						Section->ControlRig->SetControlValue<FTransformNoScale>(ParameterName, NoScale, true, EControlRigSetKey::Never);
+						Section->ControlRig->SetControlValue<FTransformNoScale>(ParameterName, NoScale, true, EControlRigSetKey::Never,bSetupUndo);
 				}
 				else if (ControlElement && ControlElement->Settings.ControlType == ERigControlType::EulerTransform)
 				{
 					FEulerTransform Euler = InFinalValue.Value;
-					Section->ControlRig->SetControlValue<FEulerTransform>(ParameterName, Euler, true, EControlRigSetKey::Never);
+					Section->ControlRig->SetControlValue<FEulerTransform>(ParameterName, Euler, true, EControlRigSetKey::Never,bSetupUndo);
 				}
 			}
 
diff --git a/Engine/Plugins/Experimental/ControlRig/Source/ControlRig/Private/Tools/ControlRigPose.cpp b/Engine/Plugins/Experimental/ControlRig/Source/ControlRig/Private/Tools/ControlRigPose.cpp
--- a/Engine/Plugins/Experimental/ControlRig/Source/ControlRig/Private/Tools/ControlRigPose.cpp
+++ b/Engine/Plugins/Experimental/ControlRig/Source/ControlRig/Private/Tools/ControlRigPose.cpp
@@ -33,18 +33,18 @@ void FControlRigControlPose::PastePose(UControlRig* ControlRig, bool bDoKey, boo
 
 }
 
-void FControlRigControlPose::SetControlMirrorTransform(bool bDoLocal, UControlRig* ControlRig, const FName& Name, bool bIsMatched, const FVector& GlobalTranslation, const FQuat& GlobalRotation, const FVector& LocalTranslation, const FQuat& LocalRotation, bool bNotify, const  FRigControlModifiedContext& Context)
+void FControlRigControlPose::SetControlMirrorTransform(bool bDoLocal, UControlRig* ControlRig, const FName& Name, bool bIsMatched, const FVector& GlobalTranslation, const FQuat& GlobalRotation, const FVector& LocalTranslation, const FQuat& LocalRotation, bool bNotify, const  FRigControlModifiedContext&Context,bool bSetupUndo)
 {
 	if (bDoLocal || bIsMatched)
 	{
 		FTransform NewLocalTransform(LocalRotation, LocalTranslation);
-		ControlRig->SetControlLocalTransform(Name, NewLocalTransform, bNotify, Context);
+		ControlRig->SetControlLocalTransform(Name, NewLocalTransform, bNotify,Context,bSetupUndo);
 
 	}
 	else
 	{
 		FTransform NewGlobalTransform(GlobalRotation, GlobalTranslation);
-		ControlRig->SetControlGlobalTransform(Name, NewGlobalTransform, bNotify, Context);
+		ControlRig->SetControlGlobalTransform(Name, NewGlobalTransform, bNotify,Context,bSetupUndo);
 	}	
 }
 
@@ -61,7 +61,9 @@ void FControlRigControlPose::PastePoseInternal(UControlRig* ControlRig, bool bDo
 	TArray<FRigControlElement*> SortedControls;
 	ControlRig->GetControlsInOrder(SortedControls);
 	
-	bool bDoLocal = true;
+	const bool bDoLocal = true;
+	const bool bSetupUndo = false;
+
 	for (FRigControlElement* ControlElement : SortedControls)
 	{
 		if (!ControlRig->IsControlSelected(ControlElement->GetName()))
@@ -87,11 +89,11 @@ void FControlRigControlPose::PastePoseInternal(UControlRig* ControlRig, bool bDo
 					{
 						if (bDoLocal) // -V547  
 						{
-							ControlRig->SetControlLocalTransform(ControlElement->GetName(), CopyRigControl->LocalTransform, true, Context);
+							ControlRig->SetControlLocalTransform(ControlElement->GetName(), CopyRigControl->LocalTransform, true,Context,bSetupUndo);
 						}
 						else
 						{
-							ControlRig->SetControlGlobalTransform(ControlElement->GetName(), CopyRigControl->GlobalTransform, true, Context);
+							ControlRig->SetControlGlobalTransform(ControlElement->GetName(), CopyRigControl->GlobalTransform, true,Context,bSetupUndo);
 						}
 					}
 					else
@@ -102,32 +104,32 @@ void FControlRigControlPose::PastePoseInternal(UControlRig* ControlRig, bool bDo
 						FQuat LocalRotation;
 						bool bIsMatched = MirrorTable.IsMatched(CopyRigControl->Name);
 						MirrorTable.GetMirrorTransform(*CopyRigControl, bDoLocal,bIsMatched, GlobalTranslation, GlobalRotation, LocalTranslation, LocalRotation);
-						SetControlMirrorTransform(bDoLocal,ControlRig, ControlElement->GetName(), bIsMatched, GlobalTranslation, GlobalRotation, LocalTranslation,LocalRotation, true, Context);
+						SetControlMirrorTransform(bDoLocal,ControlRig, ControlElement->GetName(), bIsMatched, GlobalTranslation, GlobalRotation, LocalTranslation,LocalRotation, true,Context,bSetupUndo);
 					}				
 					break;
 				}
 				case ERigControlType::Float:
 				{
 					float Val = CopyRigControl->Value.Get<float>();
-					ControlRig->SetControlValue<float>(ControlElement->GetName(), Val, true, Context);
+					ControlRig->SetControlValue<float>(ControlElement->GetName(), Val, true, Context,bSetupUndo);
 					break;
 				}
 				case ERigControlType::Bool:
 				{
 					bool Val = CopyRigControl->Value.Get<bool>();
-					ControlRig->SetControlValue<bool>(ControlElement->GetName(), Val, true, Context);
+					ControlRig->SetControlValue<bool>(ControlElement->GetName(), Val, true, Context,bSetupUndo);
 					break;
 				}
 				case ERigControlType::Integer:
 				{
 					int32 Val = CopyRigControl->Value.Get<int32>();
-					ControlRig->SetControlValue<int32>(ControlElement->GetName(), Val, true, Context);
+					ControlRig->SetControlValue<int32>(ControlElement->GetName(), Val, true, Context,bSetupUndo);
 					break;
 				}
 				case ERigControlType::Vector2D:
 				{
 					FVector2D Val = CopyRigControl->Value.Get<FVector2D>();
-					ControlRig->SetControlValue<FVector2D>(ControlElement->GetName(), Val, true, Context);
+					ControlRig->SetControlValue<FVector2D>(ControlElement->GetName(), Val, true, Context,bSetupUndo);
 					break;
 				}
 				default:
@@ -159,7 +161,8 @@ void FControlRigControlPose::BlendWithInitialPoses(FControlRigControlPose& Initi
 	TArray<FRigControlElement*> SortedControls;
 	ControlRig->GetControlsInOrder(SortedControls);
 	
-	bool bDoLocal = true;
+	const bool bDoLocal = true;
+	const bool bSetupUndo = false;
 	for (FRigControlElement* ControlElement : SortedControls)
 	{
 		if (!ControlRig->IsControlSelected(ControlElement->GetName()))
@@ -197,7 +200,7 @@ void FControlRigControlPose::BlendWithInitialPoses(FControlRigControlPose& Initi
 								Rotation = FQuat::Slerp(InitialVal.GetRotation(), Val.GetRotation(), BlendValue); //doing slerp here not fast lerp, can be slow this is for content creation
 								Scale = FMath::Lerp(InitialVal.GetScale3D(), Val.GetScale3D(), BlendValue);
 								Val = FTransform(Rotation, Translation, Scale);
-								ControlRig->SetControlLocalTransform(ControlElement->GetName(), Val, bDoKey, Context);
+								ControlRig->SetControlLocalTransform(ControlElement->GetName(), Val, bDoKey,Context,bSetupUndo);
 							}
 							else
 							{
@@ -209,7 +212,7 @@ void FControlRigControlPose::BlendWithInitialPoses(FControlRigControlPose& Initi
 								Rotation = FQuat::Slerp(InitialVal.GetRotation(), Val.GetRotation(), BlendValue); //doing slerp here not fast lerp, can be slow this is for content creation
 								Scale = FMath::Lerp(InitialVal.GetScale3D(), Val.GetScale3D(), BlendValue);
 								Val = FTransform(Rotation, Translation, Scale);
-								ControlRig->SetControlGlobalTransform(ControlElement->GetName(), Val, bDoKey, Context);
+								ControlRig->SetControlGlobalTransform(ControlElement->GetName(), Val, bDoKey,Context,bSetupUndo);
 							}
 						}
 						else
@@ -228,7 +231,7 @@ void FControlRigControlPose::BlendWithInitialPoses(FControlRigControlPose& Initi
 							GlobalRotation = FQuat::Slerp(InitialGlobalRotation, GlobalRotation, BlendValue); //doing slerp here not fast lerp, can be slow this is for content creation
 							LocalTranslation = FMath::Lerp(InitialLocalTranslation, LocalTranslation, BlendValue);
 							LocalRotation = FQuat::Slerp(InitialLocationRotation, LocalRotation, BlendValue); //doing slerp here not fast lerp, can be slow this is for content creation
-							SetControlMirrorTransform(bDoLocal,ControlRig, ControlElement->GetName(), bIsMatched, GlobalTranslation, GlobalRotation,LocalTranslation, LocalRotation, bDoKey, Context);							
+							SetControlMirrorTransform(bDoLocal,ControlRig, ControlElement->GetName(), bIsMatched, GlobalTranslation, GlobalRotation,LocalTranslation, LocalRotation, bDoKey,Context,bSetupUndo);							
 						}
 					}
 				}
diff --git a/Engine/Plugins/Experimental/ControlRig/Source/ControlRig/Public/ControlRig.h b/Engine/Plugins/Experimental/ControlRig/Source/ControlRig/Public/ControlRig.h
--- a/Engine/Plugins/Experimental/ControlRig/Source/ControlRig/Public/ControlRig.h
+++ b/Engine/Plugins/Experimental/ControlRig/Source/ControlRig/Public/ControlRig.h
@@ -264,9 +264,9 @@ public:
 	// Sets the relative value of a Control
 	template<class T>
 	FORCEINLINE_DEBUGGABLE void SetControlValue(const FName& InControlName, T InValue, bool bNotify = true,
-		const FRigControlModifiedContext& Context = FRigControlModifiedContext())
+		const FRigControlModifiedContext& Context = FRigControlModifiedContext(), bool bSetupUndo = true)
 	{
-		SetControlValueImpl(InControlName, FRigControlValue::Make<T>(InValue), bNotify, Context);
+		SetControlValueImpl(InControlName, FRigControlValue::Make<T>(InValue), bNotify, Context, bSetupUndo);
 	}
 
 	// Returns the value of a Control
@@ -278,7 +278,7 @@ public:
 
 	// Sets the relative value of a Control
 	FORCEINLINE_DEBUGGABLE virtual void SetControlValueImpl(const FName& InControlName, const FRigControlValue& InValue, bool bNotify = true,
-		const FRigControlModifiedContext& Context = FRigControlModifiedContext())
+		const FRigControlModifiedContext& Context = FRigControlModifiedContext(), bool bSetupUndo = true)
 	{
 		const FRigElementKey Key(InControlName, ERigElementType::Control);
 
@@ -288,7 +288,7 @@ public:
 			return;
 		}
 
-		DynamicHierarchy->SetControlValue(ControlElement, InValue, ERigControlValueType::Current, true);
+		DynamicHierarchy->SetControlValue(ControlElement, InValue, ERigControlValueType::Current, bSetupUndo);
 
 		if (bNotify && OnControlModified.IsBound())
 		{
@@ -310,11 +310,11 @@ public:
 		--InterRigSyncBracket;
 	}
 
-	bool SetControlGlobalTransform(const FName& InControlName, const FTransform& InGlobalTransform, bool bNotify = true, const FRigControlModifiedContext& Context = FRigControlModifiedContext());
+	bool SetControlGlobalTransform(const FName& InControlName, const FTransform& InGlobalTransform, bool bNotify = true, const FRigControlModifiedContext& Context = FRigControlModifiedContext(), bool bSetupUndo = true);
 
 	virtual FRigControlValue GetControlValueFromGlobalTransform(const FName& InControlName, const FTransform& InGlobalTransform);
 
-	virtual void SetControlLocalTransform(const FName& InControlName, const FTransform& InLocalTransform, bool bNotify = true, const FRigControlModifiedContext& Context = FRigControlModifiedContext());
+	virtual void SetControlLocalTransform(const FName& InControlName, const FTransform& InLocalTransform, bool bNotify = true, const FRigControlModifiedContext& Context = FRigControlModifiedContext(), bool bSetupUndo = true);
 	virtual FTransform GetControlLocalTransform(const FName& InControlName) ;
 
 	virtual UControlRigGizmoLibrary* GetGizmoLibrary() const;
diff --git a/Engine/Plugins/Experimental/ControlRig/Source/ControlRig/Public/Tools/ControlRigPose.h b/Engine/Plugins/Experimental/ControlRig/Source/ControlRig/Public/Tools/ControlRigPose.h
--- a/Engine/Plugins/Experimental/ControlRig/Source/ControlRig/Public/Tools/ControlRigPose.h
+++ b/Engine/Plugins/Experimental/ControlRig/Source/ControlRig/Public/Tools/ControlRigPose.h
@@ -92,7 +92,8 @@ struct CONTROLRIG_API FControlRigControlPose
 
 	void SavePose(UControlRig* ControlRig, bool bUseAll);
 	void PastePose(UControlRig* ControlRig, bool bDoKey, bool bDoMirror);
-	void SetControlMirrorTransform(bool bDoLocalSpace, UControlRig* ControlRig, const FName& Name, bool bIsMatched, const FVector& GlobalTranslation, const FQuat& GlobalRotation,const FVector& LocalTranslation, const FQuat& LocalRotation, bool bNotify, const  FRigControlModifiedContext& Context);
+	void SetControlMirrorTransform(bool bDoLocalSpace, UControlRig* ControlRig, const FName& Name, bool bIsMatched, const FVector& GlobalTranslation, const FQuat& GlobalRotation,const FVector& LocalTranslation,
+		const FQuat& LocalRotation, bool bNotify, const  FRigControlModifiedContext& Context, bool bSetupUndo);
 	void PastePoseInternal(UControlRig* ControlRig, bool bDoKey, bool bDoMirror, const TArray<FRigControlCopy>& ControlsToPaste);
 	void BlendWithInitialPoses(FControlRigControlPose& InitialPose, UControlRig* ControlRig, bool bDoKey, bool bDoMirror, float BlendValue);
 
diff --git a/Engine/Plugins/Experimental/ControlRig/Source/ControlRigEditor/Private/EditMode/ControlRigControlsProxy.cpp b/Engine/Plugins/Experimental/ControlRig/Source/ControlRigEditor/Private/EditMode/ControlRigControlsProxy.cpp
--- a/Engine/Plugins/Experimental/ControlRig/Source/ControlRigEditor/Private/EditMode/ControlRigControlsProxy.cpp
+++ b/Engine/Plugins/Experimental/ControlRig/Source/ControlRigEditor/Private/EditMode/ControlRigControlsProxy.cpp
@@ -151,7 +151,7 @@ void UControlRigTransformControlProxy::PostEditChangeProperty(struct FPropertyCh
 		{
 			//MUST set through ControlRig
 			FControlRigInteractionScope InteractionScope(ControlRig.Get());
-			ControlRig->SetControlValue<FTransform>(ControlName, Transform, true, EControlRigSetKey::DoNotCare);
+			ControlRig->SetControlValue<FTransform>(ControlName, Transform, true, EControlRigSetKey::DoNotCare,false);
 			ControlRig->Evaluate_AnyThread();
 
 		}
@@ -182,7 +182,7 @@ void UControlRigTransformControlProxy::PostEditUndo()
 			CheckEditModeOnSelectionChange(ControlRig.Get());
 		}
 		ControlRig->SelectControl(ControlName, bSelected);
-		ControlRig->SetControlValue<FTransform>(ControlName, Transform, true, EControlRigSetKey::Never);
+		ControlRig->SetControlValue<FTransform>(ControlName, Transform, true, EControlRigSetKey::Never,false);
 	}
 }
 #endif
@@ -192,7 +192,7 @@ void UControlRigTransformControlProxy::SetKey(const IPropertyHandle& KeyedProper
 	FRigControlElement* ControlElement = GetControlElement();
 	if (ControlElement)
 	{
-		ControlRig->SetControlValue<FTransform>(ControlName, Transform, true, EControlRigSetKey::Always);
+		ControlRig->SetControlValue<FTransform>(ControlName, Transform, true, EControlRigSetKey::Always,false);
 	}
 }
 
@@ -207,7 +207,7 @@ void UControlRigTransformNoScaleControlProxy::PostEditChangeProperty(struct FPro
 		{
 			//MUST set through ControlRig
 			FControlRigInteractionScope InteractionScope(ControlRig.Get());
-			ControlRig->SetControlValue<FTransformNoScale>(ControlName, Transform, true, EControlRigSetKey::DoNotCare);
+			ControlRig->SetControlValue<FTransformNoScale>(ControlName, Transform, true, EControlRigSetKey::DoNotCare,false);
 			ControlRig->Evaluate_AnyThread();
 
 		}
@@ -238,7 +238,7 @@ void UControlRigTransformNoScaleControlProxy::PostEditUndo()
 			CheckEditModeOnSelectionChange(ControlRig.Get());
 		}
 		ControlRig->SelectControl(ControlName, bSelected);
-		ControlRig->SetControlValue<FTransformNoScale>(ControlName, Transform, true, EControlRigSetKey::Never);
+		ControlRig->SetControlValue<FTransformNoScale>(ControlName, Transform, true, EControlRigSetKey::Never,false);
 	}
 }
 #endif
@@ -248,7 +248,7 @@ void UControlRigTransformNoScaleControlProxy::SetKey(const IPropertyHandle& Keye
 	FRigControlElement* ControlElement = GetControlElement();
 	if (ControlElement)
 	{
-		ControlRig->SetControlValue<FTransformNoScale>(ControlName, Transform, true, EControlRigSetKey::Always);
+		ControlRig->SetControlValue<FTransformNoScale>(ControlName, Transform, true, EControlRigSetKey::Always,false);
 	}
 }
 
@@ -264,7 +264,7 @@ void UControlRigEulerTransformControlProxy::PostEditChangeProperty(struct FPrope
 		{
 			//MUST set through ControlRig
 			FControlRigInteractionScope InteractionScope(ControlRig.Get());
-			ControlRig->SetControlValue<FEulerTransform>(ControlName, Transform, true, EControlRigSetKey::DoNotCare);
+			ControlRig->SetControlValue<FEulerTransform>(ControlName, Transform, true, EControlRigSetKey::DoNotCare,false);
 			ControlRig->Evaluate_AnyThread();
 
 		}
@@ -297,7 +297,7 @@ void UControlRigEulerTransformControlProxy::PostEditUndo()
 			CheckEditModeOnSelectionChange(ControlRig.Get());
 		}
 		ControlRig->SelectControl(ControlName, bSelected);
-		ControlRig->SetControlValue<FEulerTransform>(ControlName, Transform, true, EControlRigSetKey::Never);
+		ControlRig->SetControlValue<FEulerTransform>(ControlName, Transform, true, EControlRigSetKey::Never,false);
 	}
 }
 #endif
@@ -307,7 +307,7 @@ void UControlRigEulerTransformControlProxy::SetKey(const IPropertyHandle& KeyedP
 	FRigControlElement* ControlElement = GetControlElement();
 	if (ControlElement)
 	{
-		ControlRig->SetControlValue<FEulerTransform>(ControlName, Transform, true, EControlRigSetKey::Always);
+		ControlRig->SetControlValue<FEulerTransform>(ControlName, Transform, true, EControlRigSetKey::Always,false);
 	}
 }
 
@@ -321,7 +321,7 @@ void UControlRigFloatControlProxy::PostEditChangeProperty(struct FPropertyChange
 		{
 			//MUST set through ControlRig
 			FControlRigInteractionScope InteractionScope(ControlRig.Get());
-			ControlRig->SetControlValue<float>(ControlName, Float, true, EControlRigSetKey::DoNotCare);
+			ControlRig->SetControlValue<float>(ControlName, Float, true, EControlRigSetKey::DoNotCare,false);
 			ControlRig->Evaluate_AnyThread();
 
 		}
@@ -352,7 +352,7 @@ void UControlRigFloatControlProxy::PostEditUndo()
 			CheckEditModeOnSelectionChange(ControlRig.Get());
 		}
 		ControlRig->SelectControl(ControlName, bSelected);
-		ControlRig->SetControlValue<float>(ControlName, Float, true, EControlRigSetKey::Never);
+		ControlRig->SetControlValue<float>(ControlName, Float, true, EControlRigSetKey::Never,false);
 	}
 }
 #endif
@@ -363,7 +363,7 @@ void UControlRigFloatControlProxy::SetKey(const IPropertyHandle& KeyedPropertyHa
 	FRigControlElement* ControlElement = GetControlElement();
 	if (ControlElement)
 	{
-		ControlRig->SetControlValue<float>(ControlName, Float, true, EControlRigSetKey::Always);
+		ControlRig->SetControlValue<float>(ControlName, Float, true, EControlRigSetKey::Always,false);
 	}
 }
 
@@ -377,7 +377,7 @@ void UControlRigIntegerControlProxy::PostEditChangeProperty(struct FPropertyChan
 		{
 			//MUST set through ControlRig
 			FControlRigInteractionScope InteractionScope(ControlRig.Get());
-			ControlRig->SetControlValue<int32>(ControlName, Integer, true, EControlRigSetKey::DoNotCare);
+			ControlRig->SetControlValue<int32>(ControlName, Integer, true, EControlRigSetKey::DoNotCare,false);
 			ControlRig->Evaluate_AnyThread();
 		}
 	}
@@ -407,7 +407,7 @@ void UControlRigIntegerControlProxy::PostEditUndo()
 			CheckEditModeOnSelectionChange(ControlRig.Get());
 		}
 		ControlRig->SelectControl(ControlName, bSelected);
-		ControlRig->SetControlValue<int32>(ControlName, Integer, true, EControlRigSetKey::Never);
+		ControlRig->SetControlValue<int32>(ControlName, Integer, true, EControlRigSetKey::Never,false);
 	}
 }
 #endif
@@ -417,7 +417,7 @@ void UControlRigIntegerControlProxy::SetKey(const IPropertyHandle& KeyedProperty
 	FRigControlElement* ControlElement = GetControlElement();
 	if (ControlElement)
 	{
-		ControlRig->SetControlValue<int32>(ControlName, Integer, true, EControlRigSetKey::Always);
+		ControlRig->SetControlValue<int32>(ControlName, Integer, true, EControlRigSetKey::Always,false);
 	}
 }
 
@@ -431,7 +431,7 @@ void UControlRigEnumControlProxy::PostEditChangeProperty(struct FPropertyChanged
 		{
 			//MUST set through ControlRig
 			FControlRigInteractionScope InteractionScope(ControlRig.Get());
-			ControlRig->SetControlValue<int32>(ControlName, Enum.EnumIndex, true, EControlRigSetKey::DoNotCare);
+			ControlRig->SetControlValue<int32>(ControlName, Enum.EnumIndex, true, EControlRigSetKey::DoNotCare,false);
 			ControlRig->Evaluate_AnyThread();
 
 		}
@@ -466,7 +466,7 @@ void UControlRigEnumControlProxy::PostEditUndo()
 			CheckEditModeOnSelectionChange(ControlRig.Get());
 		}
 		ControlRig->SelectControl(ControlName, bSelected);
-		ControlRig->SetControlValue<int32>(ControlName, Enum.EnumIndex, true, EControlRigSetKey::Never);
+		ControlRig->SetControlValue<int32>(ControlName, Enum.EnumIndex, true, EControlRigSetKey::Never,false);
 	}
 }
 #endif
@@ -476,7 +476,7 @@ void UControlRigEnumControlProxy::SetKey(const IPropertyHandle& KeyedPropertyHan
 	FRigControlElement* ControlElement = GetControlElement();
 	if (ControlElement)
 	{
-		ControlRig->SetControlValue<int32>(ControlName, Enum.EnumIndex, true, EControlRigSetKey::Always);
+		ControlRig->SetControlValue<int32>(ControlName, Enum.EnumIndex, true, EControlRigSetKey::Always,false);
 	}
 }
 
@@ -491,7 +491,7 @@ void UControlRigVectorControlProxy::PostEditChangeProperty(struct FPropertyChang
 		{
 			//MUST set through ControlRig
 			FControlRigInteractionScope InteractionScope(ControlRig.Get());
-			ControlRig->SetControlValue<FVector>(ControlName, Vector, true, EControlRigSetKey::DoNotCare);
+			ControlRig->SetControlValue<FVector>(ControlName, Vector, true, EControlRigSetKey::DoNotCare,false);
 			ControlRig->Evaluate_AnyThread();
 
 		}
@@ -522,7 +522,7 @@ void UControlRigVectorControlProxy::PostEditUndo()
 			CheckEditModeOnSelectionChange(ControlRig.Get());
 		}
 		ControlRig->SelectControl(ControlName, bSelected);
-		ControlRig->SetControlValue<FVector>(ControlName, Vector, true, EControlRigSetKey::Never);
+		ControlRig->SetControlValue<FVector>(ControlName, Vector, true, EControlRigSetKey::Never,false);
 	}
 }
 #endif
@@ -533,7 +533,7 @@ void UControlRigVectorControlProxy::SetKey(const IPropertyHandle& KeyedPropertyH
 	FRigControlElement* ControlElement = GetControlElement();
 	if (ControlElement)
 	{
-		ControlRig->SetControlValue<FVector>(ControlName, Vector, true, EControlRigSetKey::Always);
+		ControlRig->SetControlValue<FVector>(ControlName, Vector, true, EControlRigSetKey::Always,false);
 	}
 }
 
@@ -548,7 +548,7 @@ void UControlRigVector2DControlProxy::PostEditChangeProperty(struct FPropertyCha
 		{
 			//MUST set through ControlRig
 			FControlRigInteractionScope InteractionScope(ControlRig.Get());
-			ControlRig->SetControlValue<FVector2D>(ControlName, Vector2D, true, EControlRigSetKey::DoNotCare);
+			ControlRig->SetControlValue<FVector2D>(ControlName, Vector2D, true, EControlRigSetKey::DoNotCare,false);
 			ControlRig->Evaluate_AnyThread();
 
 		}
@@ -579,7 +579,7 @@ void UControlRigVector2DControlProxy::PostEditUndo()
 			CheckEditModeOnSelectionChange(ControlRig.Get());
 		}
 		ControlRig->SelectControl(ControlName, bSelected);
-		ControlRig->SetControlValue<FVector2D>(ControlName, Vector2D, true, EControlRigSetKey::Never);
+		ControlRig->SetControlValue<FVector2D>(ControlName, Vector2D, true, EControlRigSetKey::Never,false);
 	}
 }
 #endif
@@ -590,7 +590,7 @@ void UControlRigVector2DControlProxy::SetKey(const IPropertyHandle& KeyedPropert
 	FRigControlElement* ControlElement = GetControlElement();
 	if (ControlElement)
 	{
-		ControlRig->SetControlValue<FVector2D>(ControlName, Vector2D, true, EControlRigSetKey::Always);
+		ControlRig->SetControlValue<FVector2D>(ControlName, Vector2D, true, EControlRigSetKey::Always,false);
 	}
 }
 
@@ -605,7 +605,7 @@ void UControlRigBoolControlProxy::PostEditChangeProperty(struct FPropertyChanged
 		{
 			//MUST set through ControlRig
 			FControlRigInteractionScope InteractionScope(ControlRig.Get());
-			ControlRig->SetControlValue<bool>(ControlName, Bool, true, EControlRigSetKey::DoNotCare);
+			ControlRig->SetControlValue<bool>(ControlName, Bool, true, EControlRigSetKey::DoNotCare,false);
 			ControlRig->Evaluate_AnyThread();
 		}
 	}
@@ -635,7 +635,7 @@ void UControlRigBoolControlProxy::PostEditUndo()
 			CheckEditModeOnSelectionChange(ControlRig.Get());
 		}
 		ControlRig->SelectControl(ControlName, bSelected);
-		ControlRig->SetControlValue<bool>(ControlName, Bool, true, EControlRigSetKey::Never);
+		ControlRig->SetControlValue<bool>(ControlName, Bool, true, EControlRigSetKey::Never,false);
 	}
 }
 #endif
@@ -646,7 +646,7 @@ void UControlRigBoolControlProxy::SetKey(const IPropertyHandle& KeyedPropertyHan
 	FRigControlElement* ControlElement = GetControlElement();
 	if (ControlElement)
 	{
-		ControlRig->SetControlValue<bool>(ControlName, Bool, true, EControlRigSetKey::Always);
+		ControlRig->SetControlValue<bool>(ControlName, Bool, true, EControlRigSetKey::Always,false);
 	}
 }
 
