Control Rig: Move visibility changes to another event (avoid sequencer to refresh)

diff --git a/Engine/Plugins/Experimental/ControlRig/Source/ControlRig/Private/Rigs/RigHierarchy.cpp b/Engine/Plugins/Experimental/ControlRig/Source/ControlRig/Private/Rigs/RigHierarchy.cpp
--- a/Engine/Plugins/Experimental/ControlRig/Source/ControlRig/Private/Rigs/RigHierarchy.cpp
+++ b/Engine/Plugins/Experimental/ControlRig/Source/ControlRig/Private/Rigs/RigHierarchy.cpp
@@ -1639,7 +1639,7 @@ void URigHierarchy::SetControlVisibility(FRigControlElement* InControlElement, b
 	}
 
 	InControlElement->Settings.bGizmoVisible = bVisibility;
-		Notify(ERigHierarchyNotification::ControlSettingChanged, InControlElement);
+	Notify(ERigHierarchyNotification::ControlVisibilityChanged, InControlElement);
 }
 
 float URigHierarchy::GetCurveValue(FRigCurveElement* InCurveElement) const
diff --git a/Engine/Plugins/Experimental/ControlRig/Source/ControlRig/Public/Rigs/RigHierarchyDefines.h b/Engine/Plugins/Experimental/ControlRig/Source/ControlRig/Public/Rigs/RigHierarchyDefines.h
--- a/Engine/Plugins/Experimental/ControlRig/Source/ControlRig/Public/Rigs/RigHierarchyDefines.h
+++ b/Engine/Plugins/Experimental/ControlRig/Source/ControlRig/Public/Rigs/RigHierarchyDefines.h
@@ -44,6 +44,7 @@ enum class ERigHierarchyNotification : uint8
 	ParentChanged,
 	HierarchyReset,
 	ControlSettingChanged,
+	ControlVisibilityChanged,
 
 	/** MAX - invalid */
 	Max UMETA(Hidden),
