Control Rig: Restoring set initial transform behavior for Controls prior to hierarchy refactor

diff --git a/Engine/Plugins/Experimental/ControlRig/Source/ControlRig/Private/Units/Hierarchy/RigUnit_SetControlOffset.h b/Engine/Plugins/Experimental/ControlRig/Source/ControlRig/Private/Units/Hierarchy/RigUnit_SetControlOffset.h
--- a/Engine/Plugins/Experimental/ControlRig/Source/ControlRig/Private/Units/Hierarchy/RigUnit_SetControlOffset.h
+++ b/Engine/Plugins/Experimental/ControlRig/Source/ControlRig/Private/Units/Hierarchy/RigUnit_SetControlOffset.h
@@ -30,7 +30,7 @@ struct CONTROLRIG_API FRigUnit_SetControlOffset : public FRigUnitMutable
 	/**
 	 * The offse transform to set for the control
 	 */
-	UPROPERTY(meta = (Input, Output))
+	UPROPERTY(meta = (Input))
 	FTransform Offset;
 
 	/**
diff --git a/Engine/Plugins/Experimental/ControlRig/Source/ControlRig/Private/Units/Hierarchy/RigUnit_SetTransform.cpp b/Engine/Plugins/Experimental/ControlRig/Source/ControlRig/Private/Units/Hierarchy/RigUnit_SetTransform.cpp
--- a/Engine/Plugins/Experimental/ControlRig/Source/ControlRig/Private/Units/Hierarchy/RigUnit_SetTransform.cpp
+++ b/Engine/Plugins/Experimental/ControlRig/Source/ControlRig/Private/Units/Hierarchy/RigUnit_SetTransform.cpp
@@ -5,6 +5,7 @@
 #include "Units/RigUnitContext.h"
 #include "Units/Execution/RigUnit_PrepareForExecution.h"
 #include "Units/Execution/RigUnit_InverseExecution.h"
+#include "Units/Hierarchy/RigUnit_SetControlOffset.h"
 
 FString FRigUnit_SetTransform::GetUnitLabel() const
 {
@@ -43,6 +44,21 @@ FRigUnit_SetTransform_Execute()
 				}
 				else
 				{
+					// for controls - set the control offset transform instead
+					if(bInitial && (CachedIndex.GetKey().Type == ERigElementType::Control))
+					{
+						FRigUnit_SetControlOffset::StaticExecute(RigVMExecuteContext, CachedIndex.GetKey().Name, Transform, Space, CachedIndex, ExecuteContext, Context);
+						
+						if (ExecuteContext.EventName == FRigUnit_PrepareForExecution::EventName)
+						{
+							Hierarchy->SetLocalTransformByIndex(CachedIndex, FTransform::Identity, true, bPropagateToChildren);
+							Hierarchy->SetLocalTransformByIndex(CachedIndex, FTransform::Identity, false, bPropagateToChildren);
+
+							
+						}
+						return;
+					}
+					
 					FTransform WeightedTransform = Transform;
 					if (Weight < 1.f - SMALL_NUMBER)
 					{
diff --git a/Engine/Plugins/Experimental/ControlRig/Source/ControlRig/Public/Units/Hierarchy/RigUnit_SetTransform.h b/Engine/Plugins/Experimental/ControlRig/Source/ControlRig/Public/Units/Hierarchy/RigUnit_SetTransform.h
--- a/Engine/Plugins/Experimental/ControlRig/Source/ControlRig/Public/Units/Hierarchy/RigUnit_SetTransform.h
+++ b/Engine/Plugins/Experimental/ControlRig/Source/ControlRig/Public/Units/Hierarchy/RigUnit_SetTransform.h
@@ -7,6 +7,10 @@
 
 /**
  * SetTransform is used to set a single transform on hierarchy.
+ * 
+ * Note: For Controls when setting the initial transform this node
+ * actually sets the Control's offset transform and resets the local
+ * values to (0, 0, 0).
  */
 USTRUCT(meta=(DisplayName="Set Transform", Category="Hierarchy", DocumentationPolicy = "Strict", Keywords="SetBoneTransform,SetControlTransform,SetInitialTransform,SetSpaceTransform", Varying))
 struct CONTROLRIG_API FRigUnit_SetTransform : public FRigUnitMutable
