Control Rig: Improve details panel for Controls / Bones and reduce clutter

diff --git a/Engine/Plugins/Experimental/ControlRig/Source/ControlRig/Public/Rigs/RigHierarchyElements.h b/Engine/Plugins/Experimental/ControlRig/Source/ControlRig/Public/Rigs/RigHierarchyElements.h
--- a/Engine/Plugins/Experimental/ControlRig/Source/ControlRig/Public/Rigs/RigHierarchyElements.h
+++ b/Engine/Plugins/Experimental/ControlRig/Source/ControlRig/Public/Rigs/RigHierarchyElements.h
@@ -541,7 +541,7 @@ public:
 	virtual void Save(FArchive& A, URigHierarchy* Hierarchy, ESerializationPhase SerializationPhase) override;
 	virtual void Load(FArchive& Ar, URigHierarchy* Hierarchy, ESerializationPhase SerializationPhase) override;
 
-	UPROPERTY(BlueprintReadOnly, EditAnywhere, Category = RigElement, meta = (DisplayAfter = "Index", ShowOnlyInnerProperties))
+	UPROPERTY(BlueprintReadOnly, EditAnywhere, Category = RigElement, meta = (DisplayAfter = "Index"))
 	FRigCurrentAndInitialTransform Pose;
 
 protected:
diff --git a/Engine/Plugins/Experimental/ControlRig/Source/ControlRigEditor/Private/ControlRigElementDetails.cpp b/Engine/Plugins/Experimental/ControlRig/Source/ControlRigEditor/Private/ControlRigElementDetails.cpp
--- a/Engine/Plugins/Experimental/ControlRig/Source/ControlRigEditor/Private/ControlRigElementDetails.cpp
+++ b/Engine/Plugins/Experimental/ControlRig/Source/ControlRigEditor/Private/ControlRigElementDetails.cpp
@@ -608,13 +608,119 @@ void FRigBaseElementDetails::OnStructContentsChanged(FProperty* InProperty, cons
 	PropertyUtilities->NotifyFinishedChangingProperties(ChangeEvent);
 }
 
+bool FRigBaseElementDetails::IsSetupModeEnabled() const
+{
+	if(BlueprintBeingCustomized)
+	{
+		if (UControlRig* DebuggedRig = Cast<UControlRig>(BlueprintBeingCustomized->GetObjectBeingDebugged()))
+		{
+			return DebuggedRig->IsSetupModeEnabled();
+		}
+	}
+	return false;
+}
+
 void FRigTransformElementDetails::CustomizeDetails(IDetailLayoutBuilder& DetailBuilder)
 {
 	FRigBaseElementDetails::CustomizeDetails(DetailBuilder);
 
-	IDetailCategoryBuilder& Category = DetailBuilder.EditCategory(TEXT("Pose"));
+	//if(ElementKeyBeingCustomized.Type != ERigElementType::Control)
+	{
+		IDetailCategoryBuilder& Category = DetailBuilder.EditCategory(TEXT("Pose"));
+		if(ElementKeyBeingCustomized.Type == ERigElementType::Control)
+		{
+			Category
+			.InitiallyCollapsed(true);
+		}
+
+		// setup initial global
+		{
+			const TSharedRef<IPropertyHandle> PropertyHandle = DetailBuilder.GetProperty(TEXT("Pose.Initial.Global.Transform"), FRigTransformElement::StaticStruct());
+			Category.AddProperty(PropertyHandle, EPropertyLocation::Advanced).DisplayName(FText::FromString(TEXT("Initial Global")))
+			.IsEnabled(TAttribute<bool>::Create(TAttribute<bool>::FGetter::CreateSP(this, &FRigBaseElementDetails::IsSetupModeEnabled)));
+
+			PoseInitialGlobal.AddHead(FRigControlElement::StaticStruct()->FindPropertyByName(TEXT("Pose")));
+			PoseInitialGlobal.AddTail(FRigCurrentAndInitialTransform::StaticStruct()->FindPropertyByName(TEXT("Initial")));
+			PoseInitialGlobal.AddTail(FRigLocalAndGlobalTransform::StaticStruct()->FindPropertyByName(TEXT("Global")));
+			PoseInitialGlobal.AddTail(FRigComputedTransform::StaticStruct()->FindPropertyByName(TEXT("Transform")));
+			PoseInitialGlobal.SetActiveMemberPropertyNode(PoseInitialGlobal.GetTail()->GetValue());
+
+			const FSimpleDelegate OnTransformChangedDelegate = FSimpleDelegate::CreateStatic(&FRigTransformElementDetails::OnTransformChanged, &PoseInitialGlobal, BlueprintBeingCustomized);
+			PropertyHandle->SetOnPropertyValueChanged(OnTransformChangedDelegate);
+			PropertyHandle->SetOnChildPropertyValueChanged(OnTransformChangedDelegate);
+			PropertyHandle->SetToolTipText(FText::FromString(TEXT("The initial / ref pose global transform in the space of the rig / actor.")));
+		}
+
+		// setup initial local
+		{
+			const TSharedRef<IPropertyHandle> PropertyHandle = DetailBuilder.GetProperty(TEXT("Pose.Initial.Local.Transform"), FRigTransformElement::StaticStruct());
+			Category.AddProperty(PropertyHandle, EPropertyLocation::Advanced).DisplayName(FText::FromString(TEXT("Initial Local")))
+			.IsEnabled(TAttribute<bool>::Create(TAttribute<bool>::FGetter::CreateSP(this, &FRigBaseElementDetails::IsSetupModeEnabled)));
+
+			PoseInitialLocal.AddHead(FRigControlElement::StaticStruct()->FindPropertyByName(TEXT("Pose")));
+			PoseInitialLocal.AddTail(FRigCurrentAndInitialTransform::StaticStruct()->FindPropertyByName(TEXT("Initial")));
+			PoseInitialLocal.AddTail(FRigLocalAndGlobalTransform::StaticStruct()->FindPropertyByName(TEXT("Local")));
+			PoseInitialLocal.AddTail(FRigComputedTransform::StaticStruct()->FindPropertyByName(TEXT("Transform")));
+			PoseInitialLocal.SetActiveMemberPropertyNode(PoseInitialLocal.GetTail()->GetValue());
+
+			const FSimpleDelegate OnTransformChangedDelegate = FSimpleDelegate::CreateStatic(&FRigTransformElementDetails::OnTransformChanged, &PoseInitialLocal, BlueprintBeingCustomized);
+			PropertyHandle->SetOnPropertyValueChanged(OnTransformChangedDelegate);
+			PropertyHandle->SetOnChildPropertyValueChanged(OnTransformChangedDelegate);
+			PropertyHandle->SetToolTipText(FText::FromString(TEXT("The initial / ref pose local transform in the space of the parent.\nFor Controls the initial value is relative to the offset.")));
+		}
+
+		// setup current global
+		{
+			const TSharedRef<IPropertyHandle> PropertyHandle = DetailBuilder.GetProperty(TEXT("Pose.Current.Global.Transform"), FRigTransformElement::StaticStruct());
+			Category.AddProperty(PropertyHandle, EPropertyLocation::Advanced).DisplayName(FText::FromString(TEXT("Current Global")))
+			.IsEnabled(TAttribute<bool>::Create(TAttribute<bool>::FGetter::CreateSP(this, &FRigBaseElementDetails::IsSetupModeEnabled)));
+
+			PoseCurrentGlobal.AddHead(FRigControlElement::StaticStruct()->FindPropertyByName(TEXT("Pose")));
+			PoseCurrentGlobal.AddTail(FRigCurrentAndInitialTransform::StaticStruct()->FindPropertyByName(TEXT("Current")));
+			PoseCurrentGlobal.AddTail(FRigLocalAndGlobalTransform::StaticStruct()->FindPropertyByName(TEXT("Global")));
+			PoseCurrentGlobal.AddTail(FRigComputedTransform::StaticStruct()->FindPropertyByName(TEXT("Transform")));
+			PoseCurrentGlobal.SetActiveMemberPropertyNode(PoseCurrentGlobal.GetTail()->GetValue());
+
+			const FSimpleDelegate OnTransformChangedDelegate = FSimpleDelegate::CreateStatic(&FRigTransformElementDetails::OnTransformChanged, &PoseCurrentGlobal, BlueprintBeingCustomized);
+			PropertyHandle->SetOnPropertyValueChanged(OnTransformChangedDelegate);
+			PropertyHandle->SetOnChildPropertyValueChanged(OnTransformChangedDelegate);
+			PropertyHandle->SetToolTipText(FText::FromString(TEXT("The current global transform in the space of the rig / actor.")));
+		}
+
+		// setup current local
+		{
+			const TSharedRef<IPropertyHandle> PropertyHandle = DetailBuilder.GetProperty(TEXT("Pose.Current.Local.Transform"), FRigTransformElement::StaticStruct());
+			Category.AddProperty(PropertyHandle).DisplayName(FText::FromString(TEXT("Current Local")))
+			.IsEnabled(ElementKeyBeingCustomized.Type != ERigElementType::Control);
+
+			PoseCurrentLocal.AddHead(FRigControlElement::StaticStruct()->FindPropertyByName(TEXT("Pose")));
+			PoseCurrentLocal.AddTail(FRigCurrentAndInitialTransform::StaticStruct()->FindPropertyByName(TEXT("Current")));
+			PoseCurrentLocal.AddTail(FRigLocalAndGlobalTransform::StaticStruct()->FindPropertyByName(TEXT("Local")));
+			PoseCurrentLocal.AddTail(FRigComputedTransform::StaticStruct()->FindPropertyByName(TEXT("Transform")));
+			PoseCurrentLocal.SetActiveMemberPropertyNode(PoseCurrentLocal.GetTail()->GetValue());
+
+			const FSimpleDelegate OnTransformChangedDelegate = FSimpleDelegate::CreateStatic(&FRigTransformElementDetails::OnTransformChanged, &PoseCurrentLocal, BlueprintBeingCustomized);
+			PropertyHandle->SetOnPropertyValueChanged(OnTransformChangedDelegate);
+			PropertyHandle->SetOnChildPropertyValueChanged(OnTransformChangedDelegate);
+			PropertyHandle->SetToolTipText(FText::FromString(TEXT("The current local transform in the space of the parent.\nFor Controls the initial value is relative to the offset.")));
+		}
 
-	Category.InitiallyCollapsed(ElementKeyBeingCustomized.Type != ERigElementType::Control);
+		DetailBuilder.HideProperty(TEXT("Pose"), FRigTransformElement::StaticStruct());
+	}
+}
+
+void FRigTransformElementDetails::OnTransformChanged(FEditPropertyChain* InPropertyChain, UControlRigBlueprint* InBlueprint)
+{
+	if(InBlueprint && InPropertyChain)
+	{
+		if(InPropertyChain->Num() > 1)
+		{
+			FPropertyChangedEvent ChangeEvent(InPropertyChain->GetHead()->GetValue(), EPropertyChangeType::ValueSet);
+			ChangeEvent.SetActiveMemberProperty(InPropertyChain->GetTail()->GetValue());
+			FPropertyChangedChainEvent ChainEvent(*InPropertyChain, ChangeEvent);
+			InBlueprint->BroadcastPostEditChangeChainProperty(ChainEvent);
+		}
+	}
 }
 
 void FRigBoneElementDetails::CustomizeDetails( IDetailLayoutBuilder& DetailBuilder )
@@ -952,13 +1058,18 @@ void FRigControlElementDetails_SetupStructValueWidget(IDetailCategoryBuilder& In
 	TWeakObjectPtr<URigHierarchy> HierarchyPtr = InHierarchy;
 	const FRigElementKey Key = InControlElement->GetKey();
 
-	const TAttribute<bool> EnabledAttribute = TAttribute<bool>::Create(TAttribute<bool>::FGetter::CreateLambda([HierarchyPtr, Key, InValueType]()->bool
+	const TAttribute<bool> EnabledAttribute = TAttribute<bool>::Create(TAttribute<bool>::FGetter::CreateLambda([HierarchyPtr, Key, InValueType, StructToDisplay, ValueStruct]()->bool
     {
         if(HierarchyPtr.IsValid())
         {
             if(FRigControlElement* ControlElement = HierarchyPtr->Find<FRigControlElement>(Key))
             {
-                return ControlElement->Settings.IsValueTypeEnabled(InValueType);
+            	// update the struct with the current control value
+            	uint8* StructMemory = StructToDisplay->GetStructMemory();
+            	const FRigControlValue& CurrentValue = HierarchyPtr->GetControlValue(Key, InValueType);            	
+            	FMemory::Memcpy(StructToDisplay->GetStructMemory(), &CurrentValue.GetRef<T>(), sizeof(T));
+
+            	return ControlElement->Settings.IsValueTypeEnabled(InValueType);
             }
         }
         return false;
@@ -1244,23 +1355,6 @@ void FRigControlElementDetails::CustomizeDetails(IDetailLayoutBuilder& DetailBui
 		}
 	));
 
-	struct Local
-	{
-		static void OnTransformChanged(FEditPropertyChain* InPropertyChain, UControlRigBlueprint* InBlueprint)
-		{
-			if(InBlueprint && InPropertyChain)
-			{
-				if(InPropertyChain->Num() > 1)
-				{
-					FPropertyChangedEvent ChangeEvent(InPropertyChain->GetHead()->GetValue(), EPropertyChangeType::ValueSet);
-					ChangeEvent.SetActiveMemberProperty(InPropertyChain->GetTail()->GetValue());
-					FPropertyChangedChainEvent ChainEvent(*InPropertyChain, ChangeEvent);
-					InBlueprint->BroadcastPostEditChangeChainProperty(ChainEvent);
-				}
-			}
-		}
-	};
-
 	if(ControlElement->Settings.ControlType == ERigControlType::Bool)
 	{
 		DetailBuilder.HideCategory(TEXT("Pose"));
@@ -1271,7 +1365,6 @@ void FRigControlElementDetails::CustomizeDetails(IDetailLayoutBuilder& DetailBui
 	{
 		// setup offset transform
 		{
-			const TSharedRef<IPropertyHandle> OffsetHandle = DetailBuilder.GetProperty(TEXT("Offset"));
 			const TSharedRef<IPropertyHandle> OffsetInitialLocalTransformHandle = DetailBuilder.GetProperty(TEXT("Offset.Initial.Local.Transform"));
 			ControlCategory.AddProperty(OffsetInitialLocalTransformHandle).DisplayName(FText::FromString(TEXT("Offset Transform")));
 			DetailBuilder.HideProperty(TEXT("Offset"));
@@ -1282,14 +1375,14 @@ void FRigControlElementDetails::CustomizeDetails(IDetailLayoutBuilder& DetailBui
 			OffsetPropertyChain.AddTail(FRigComputedTransform::StaticStruct()->FindPropertyByName(TEXT("Transform")));
 			OffsetPropertyChain.SetActiveMemberPropertyNode(OffsetPropertyChain.GetTail()->GetValue());
 
-			const FSimpleDelegate OnTransformChangedDelegate = FSimpleDelegate::CreateStatic(&Local::OnTransformChanged, &OffsetPropertyChain, BlueprintBeingCustomized);
+			const FSimpleDelegate OnTransformChangedDelegate = FSimpleDelegate::CreateStatic(&FRigTransformElementDetails::OnTransformChanged, &OffsetPropertyChain, BlueprintBeingCustomized);
 			OffsetInitialLocalTransformHandle->SetOnPropertyValueChanged(OnTransformChangedDelegate);
 			OffsetInitialLocalTransformHandle->SetOnChildPropertyValueChanged(OnTransformChangedDelegate);
+			OffsetInitialLocalTransformHandle->SetToolTipText(FText::FromString(TEXT("The offset transform is used as a middle man between the parent and the local transform.\nYou can use it to offset a control for visual adjustment.")));
 		}
 
 		// setup gizmo transform
 		{
-			const TSharedRef<IPropertyHandle> GizmoHandle = DetailBuilder.GetProperty(TEXT("Gizmo"));
 			const TSharedRef<IPropertyHandle> GizmoInitialLocalTransformHandle = DetailBuilder.GetProperty(TEXT("Gizmo.Initial.Local.Transform"));
 			GizmoCategory.AddProperty(GizmoInitialLocalTransformHandle).DisplayName(FText::FromString(TEXT("Gizmo Transform")));
 			DetailBuilder.HideProperty(TEXT("Gizmo"));
@@ -1300,9 +1393,10 @@ void FRigControlElementDetails::CustomizeDetails(IDetailLayoutBuilder& DetailBui
 			GizmoPropertyChain.AddTail(FRigComputedTransform::StaticStruct()->FindPropertyByName(TEXT("Transform")));
 			GizmoPropertyChain.SetActiveMemberPropertyNode(GizmoPropertyChain.GetTail()->GetValue());
 
-			const FSimpleDelegate OnTransformChangedDelegate = FSimpleDelegate::CreateStatic(&Local::OnTransformChanged, &GizmoPropertyChain, BlueprintBeingCustomized);
+			const FSimpleDelegate OnTransformChangedDelegate = FSimpleDelegate::CreateStatic(&FRigTransformElementDetails::OnTransformChanged, &GizmoPropertyChain, BlueprintBeingCustomized);
 			GizmoInitialLocalTransformHandle->SetOnPropertyValueChanged(OnTransformChangedDelegate);
 			GizmoInitialLocalTransformHandle->SetOnChildPropertyValueChanged(OnTransformChangedDelegate);
+			GizmoInitialLocalTransformHandle->SetToolTipText(FText::FromString(TEXT("The gizmo transform is used as a transform applied only to the UI element on the screen.\nIt doesn't affect animation.")));
 		}
 	}
 
diff --git a/Engine/Plugins/Experimental/ControlRig/Source/ControlRigEditor/Private/ControlRigElementDetails.h b/Engine/Plugins/Experimental/ControlRig/Source/ControlRigEditor/Private/ControlRigElementDetails.h
--- a/Engine/Plugins/Experimental/ControlRig/Source/ControlRigEditor/Private/ControlRigElementDetails.h
+++ b/Engine/Plugins/Experimental/ControlRig/Source/ControlRigEditor/Private/ControlRigElementDetails.h
@@ -120,6 +120,7 @@ public:
 	void SetName(const FText& InNewText, ETextCommit::Type InCommitType);
 
 	void OnStructContentsChanged(FProperty* InProperty, const TSharedRef<IPropertyUtilities> PropertyUtilities);
+	bool IsSetupModeEnabled() const; 
 
 protected:
 
@@ -134,6 +135,14 @@ public:
 
 	/** IDetailCustomization interface */
 	virtual void CustomizeDetails(IDetailLayoutBuilder& DetailBuilder) override;
+
+protected:
+	static void OnTransformChanged(FEditPropertyChain* InPropertyChain, UControlRigBlueprint* InBlueprint);
+
+	FEditPropertyChain PoseInitialGlobal;
+	FEditPropertyChain PoseInitialLocal;
+	FEditPropertyChain PoseCurrentGlobal;
+	FEditPropertyChain PoseCurrentLocal;
 };
 
 class FRigBoneElementDetails : public FRigTransformElementDetails
@@ -181,6 +190,7 @@ private:
 	static TArray<TSharedPtr<FString>> ControlTypeList;
 	TSharedPtr<FRigInfluenceEntryModifier> InfluenceModifier;
 	TSharedPtr<FStructOnScope> InfluenceModifierStruct;
+	FEditPropertyChain PoseCurrentChain;
 	FEditPropertyChain OffsetPropertyChain;
 	FEditPropertyChain GizmoPropertyChain;
 };
diff --git a/Engine/Plugins/Experimental/ControlRig/Source/ControlRigEditor/Private/Editor/ControlRigEditor.cpp b/Engine/Plugins/Experimental/ControlRig/Source/ControlRigEditor/Private/Editor/ControlRigEditor.cpp
--- a/Engine/Plugins/Experimental/ControlRig/Source/ControlRigEditor/Private/Editor/ControlRigEditor.cpp
+++ b/Engine/Plugins/Experimental/ControlRig/Source/ControlRigEditor/Private/Editor/ControlRigEditor.cpp
@@ -1972,47 +1972,16 @@ void FControlRigEditor::HandleVMCompiledEvent(UBlueprint* InBlueprint, URigVM* I
 void FControlRigEditor::HandleControlRigExecutedEvent(UControlRig* InControlRig, const EControlRigState InState, const FName& InEventName)
 {
 	UpdateGraphCompilerErrors();
-}
-
-
-void FControlRigEditor::CreateEditorModeManager()
-{
-	TSharedPtr<FAssetEditorModeManager> ModeManager = MakeShareable(FModuleManager::LoadModuleChecked<FPersonaModule>("Persona").CreatePersonaEditorModeManager());
-
-	PRAGMA_DISABLE_DEPRECATION_WARNINGS
-	SetAssetEditorModeManager(ModeManager.Get());
-	PRAGMA_ENABLE_DEPRECATION_WARNINGS
-}
-
-
-void FControlRigEditor::Tick(float DeltaTime)
-{
-	FBlueprintEditor::Tick(DeltaTime);
-
-	bool bDrawHierarchyBones = false;
 
-	// tick the control rig in case we don't have skeletal mesh
-	if (UControlRigBlueprint* Blueprint = GetControlRigBlueprint())
+	if (UControlRigBlueprint* ControlRigBP = GetControlRigBlueprint())
 	{
-		if (Blueprint->GetPreviewMesh() == nullptr && 
-			ControlRig != nullptr && 
-			bExecutionControlRig)
-		{
-			// reset transforms here to prevent additive transforms from accumulating to INF
-			ControlRig->GetHierarchy()->ResetPoseToInitial(ERigElementType::Bone);
-			
-			ControlRig->SetDeltaTime(DeltaTime);
-			ControlRig->Evaluate_AnyThread();
-			bDrawHierarchyBones = true;
-		}
-
 		if(RigElementInDetailPanel.IsValid())
 		{
-			URigHierarchy* Hierarchy = Blueprint->Hierarchy; 
+			URigHierarchy* Hierarchy = ControlRigBP->Hierarchy; 
 
 			if(!bSetupModeEnabled)
 			{
-				UControlRig* DebuggedControlRig = Cast<UControlRig>(GetBlueprintObj()->GetObjectBeingDebugged());
+				UControlRig* DebuggedControlRig = Cast<UControlRig>(ControlRigBP->GetObjectBeingDebugged());
 				if (DebuggedControlRig == nullptr)
 				{
 					DebuggedControlRig = ControlRig;
@@ -2049,6 +2018,40 @@ void FControlRigEditor::Tick(float DeltaTime)
 			}
 		}
 	}
+}
+
+void FControlRigEditor::CreateEditorModeManager()
+{
+	TSharedPtr<FAssetEditorModeManager> ModeManager = MakeShareable(FModuleManager::LoadModuleChecked<FPersonaModule>("Persona").CreatePersonaEditorModeManager());
+
+	PRAGMA_DISABLE_DEPRECATION_WARNINGS
+	SetAssetEditorModeManager(ModeManager.Get());
+	PRAGMA_ENABLE_DEPRECATION_WARNINGS
+}
+
+
+void FControlRigEditor::Tick(float DeltaTime)
+{
+	FBlueprintEditor::Tick(DeltaTime);
+
+	bool bDrawHierarchyBones = false;
+
+	// tick the control rig in case we don't have skeletal mesh
+	if (UControlRigBlueprint* Blueprint = GetControlRigBlueprint())
+	{
+		if (Blueprint->GetPreviewMesh() == nullptr && 
+			ControlRig != nullptr && 
+			bExecutionControlRig)
+		{
+			// reset transforms here to prevent additive transforms from accumulating to INF
+			ControlRig->GetHierarchy()->ResetPoseToInitial(ERigElementType::Bone);
+			
+			ControlRig->SetDeltaTime(DeltaTime);
+			ControlRig->Evaluate_AnyThread();
+			bDrawHierarchyBones = true;
+			
+		}
+	}
 
 	if (bDrawHierarchyBones)
 	{
@@ -3595,9 +3598,10 @@ void FControlRigEditor::SynchronizeViewportBoneSelection()
 		TArray<FRigBaseElement*> SelectedBones = RigBlueprint->Hierarchy->GetSelectedElements(ERigElementType::Bone);
 		for (const FRigBaseElement* SelectedBone : SelectedBones)
 		{
-			if(EditorSkelComp->GetReferenceSkeleton().IsValidIndex(SelectedBone->GetSubIndex()))
+ 			const int32 BoneIndex = EditorSkelComp->GetReferenceSkeleton().FindBoneIndex(SelectedBone->GetName());
+			if(BoneIndex != INDEX_NONE)
 			{
-				EditorSkelComp->BonesOfInterest.AddUnique(SelectedBone->GetSubIndex());
+				EditorSkelComp->BonesOfInterest.AddUnique(BoneIndex);
 			}
 		}
 	}
