Control Rig: Support duplicates in collections

diff --git a/Engine/Plugins/Experimental/ControlRig/Source/ControlRig/Private/Rigs/RigHierarchyDefines.cpp b/Engine/Plugins/Experimental/ControlRig/Source/ControlRig/Private/Rigs/RigHierarchyDefines.cpp
--- a/Engine/Plugins/Experimental/ControlRig/Source/ControlRig/Private/Rigs/RigHierarchyDefines.cpp
+++ b/Engine/Plugins/Experimental/ControlRig/Source/ControlRig/Private/Rigs/RigHierarchyDefines.cpp
@@ -172,7 +172,7 @@ FRigElementKeyCollection FRigElementKeyCollection::MakeFromCompleteHierarchy(
 	return Collection.FilterByType(InElementTypes);
 }
 
-FRigElementKeyCollection FRigElementKeyCollection::MakeUnion(const FRigElementKeyCollection& A, const FRigElementKeyCollection& B)
+FRigElementKeyCollection FRigElementKeyCollection::MakeUnion(const FRigElementKeyCollection& A, const FRigElementKeyCollection& B, bool bAllowDuplicates)
 {
 	FRigElementKeyCollection Collection;
 	for (const FRigElementKey& Key : A)
@@ -181,7 +181,14 @@ FRigElementKeyCollection FRigElementKeyCollection::MakeUnion(const FRigElementKe
 	}
 	for (const FRigElementKey& Key : B)
 	{
-		Collection.AddUnique(Key);
+		if(bAllowDuplicates)
+		{
+			Collection.Add(Key);
+		}
+		else
+		{
+			Collection.AddUnique(Key);
+		}
 	}
 	return Collection;
 }
diff --git a/Engine/Plugins/Experimental/ControlRig/Source/ControlRig/Private/Units/Execution/RigUnit_Collection.cpp b/Engine/Plugins/Experimental/ControlRig/Source/ControlRig/Private/Units/Execution/RigUnit_Collection.cpp
--- a/Engine/Plugins/Experimental/ControlRig/Source/ControlRig/Private/Units/Execution/RigUnit_Collection.cpp
+++ b/Engine/Plugins/Experimental/ControlRig/Source/ControlRig/Private/Units/Execution/RigUnit_Collection.cpp
@@ -125,7 +125,14 @@ FRigUnit_CollectionReplaceItems_Execute()
 
 			if (Context.Hierarchy->GetIndex(Key) != INDEX_NONE)
 			{
-				CachedCollection.AddUnique(Key);
+				if(bAllowDuplicates)
+				{
+					CachedCollection.Add(Key);
+				}
+				else
+				{
+					CachedCollection.AddUnique(Key);
+				}
 			}
 			else if(!RemoveInvalidItems)
 			{
@@ -144,7 +151,14 @@ FRigUnit_CollectionItems_Execute()
 	Collection.Reset();
 	for (const FRigElementKey& Key : Items)
 	{
-		Collection.AddUnique(Key);
+		if(bAllowDuplicates)
+		{
+			Collection.Add(Key);
+		}
+		else
+		{
+			Collection.AddUnique(Key);
+		}
 	}
 }
 
@@ -152,7 +166,7 @@ FRigUnit_CollectionUnion_Execute()
 {
 	DECLARE_SCOPE_HIERARCHICAL_COUNTER_RIGUNIT()
 
-	Collection = FRigElementKeyCollection::MakeUnion(A, B);
+	Collection = FRigElementKeyCollection::MakeUnion(A, B, bAllowDuplicates);
 }
 
 FRigUnit_CollectionIntersection_Execute()
diff --git a/Engine/Plugins/Experimental/ControlRig/Source/ControlRig/Private/Units/Execution/RigUnit_Collection.h b/Engine/Plugins/Experimental/ControlRig/Source/ControlRig/Private/Units/Execution/RigUnit_Collection.h
--- a/Engine/Plugins/Experimental/ControlRig/Source/ControlRig/Private/Units/Execution/RigUnit_Collection.h
+++ b/Engine/Plugins/Experimental/ControlRig/Source/ControlRig/Private/Units/Execution/RigUnit_Collection.h
@@ -147,6 +147,7 @@ struct CONTROLRIG_API FRigUnit_CollectionReplaceItems : public FRigUnit_Collecti
 		Old = New = NAME_None;
 		RemoveInvalidItems = false;
 		CachedHierarchyHash = INDEX_NONE;
+		bAllowDuplicates = false;
 	}
 
 	RIGVM_METHOD()
@@ -164,6 +165,9 @@ struct CONTROLRIG_API FRigUnit_CollectionReplaceItems : public FRigUnit_Collecti
 	UPROPERTY(meta = (Input))
 	bool RemoveInvalidItems;
 
+	UPROPERTY(meta = (Input))
+	bool bAllowDuplicates;
+
 	UPROPERTY(meta = (Output))
 	FRigElementKeyCollection Collection;
 
@@ -185,6 +189,7 @@ struct CONTROLRIG_API FRigUnit_CollectionItems : public FRigUnit_CollectionBase
 	FRigUnit_CollectionItems()
 	{
 		Items.Add(FRigElementKey(NAME_None, ERigElementType::Bone));
+		bAllowDuplicates = false;
 	}
 
 	RIGVM_METHOD()
@@ -193,6 +198,9 @@ struct CONTROLRIG_API FRigUnit_CollectionItems : public FRigUnit_CollectionBase
 	UPROPERTY(meta = (Input, ExpandByDefault))
 	TArray<FRigElementKey> Items;
 
+	UPROPERTY(meta = (Input))
+	bool bAllowDuplicates;
+
 	UPROPERTY(meta = (Output))
 	FRigElementKeyCollection Collection;
 };
@@ -208,6 +216,7 @@ struct CONTROLRIG_API FRigUnit_CollectionUnion : public FRigUnit_CollectionBase
 
 	FRigUnit_CollectionUnion()
 	{
+		bAllowDuplicates = false;
 	}
 
 	RIGVM_METHOD()
@@ -219,6 +228,9 @@ struct CONTROLRIG_API FRigUnit_CollectionUnion : public FRigUnit_CollectionBase
 	UPROPERTY(meta = (Input))
 	FRigElementKeyCollection B;
 
+	UPROPERTY(meta = (Input))
+	bool bAllowDuplicates;
+
 	UPROPERTY(meta = (Output))
 	FRigElementKeyCollection Collection;
 };
diff --git a/Engine/Plugins/Experimental/ControlRig/Source/ControlRig/Public/Rigs/RigHierarchyDefines.h b/Engine/Plugins/Experimental/ControlRig/Source/ControlRig/Public/Rigs/RigHierarchyDefines.h
--- a/Engine/Plugins/Experimental/ControlRig/Source/ControlRig/Public/Rigs/RigHierarchyDefines.h
+++ b/Engine/Plugins/Experimental/ControlRig/Source/ControlRig/Public/Rigs/RigHierarchyDefines.h
@@ -1027,7 +1027,7 @@ struct CONTROLRIG_API FRigElementKeyCollection
 		uint8 InElementTypes = (uint8)ERigElementType::All);
 
 	// returns the union between two collections
-	static FRigElementKeyCollection MakeUnion(const FRigElementKeyCollection& A, const FRigElementKeyCollection& B);
+	static FRigElementKeyCollection MakeUnion(const FRigElementKeyCollection& A, const FRigElementKeyCollection& B, bool bAllowDuplicates = false);
 
 	// returns the intersection between two collections
 	static FRigElementKeyCollection MakeIntersection(const FRigElementKeyCollection& A, const FRigElementKeyCollection& B);
