Fixed unity builds and removed dead includes.

#trivial
#rnx

diff --git a/Engine/Plugins/Experimental/ControlRig/Source/ControlRigEditor/Private/Editor/SControlRigFunctionLocalizationWidget.cpp b/Engine/Plugins/Experimental/ControlRig/Source/ControlRigEditor/Private/Editor/SControlRigFunctionLocalizationWidget.cpp
--- a/Engine/Plugins/Experimental/ControlRig/Source/ControlRigEditor/Private/Editor/SControlRigFunctionLocalizationWidget.cpp
+++ b/Engine/Plugins/Experimental/ControlRig/Source/ControlRigEditor/Private/Editor/SControlRigFunctionLocalizationWidget.cpp
@@ -3,17 +3,13 @@
 #include "SControlRigFunctionLocalizationWidget.h"
 #include "Widgets/Layout/SBorder.h"
 #include "Widgets/SBoxPanel.h"
-#include "Widgets/Layout/SScrollBar.h"
-#include "Widgets/Layout/SSpacer.h"
+#include "Widgets/Input/SButton.h"
 #include "Widgets/Input/SCheckBox.h"
 #include "Widgets/Layout/SUniformGridPanel.h"
 #include "DetailLayoutBuilder.h"
-#include "ControlRigEditorStyle.h"
+#include "Editor.h"
 #include "EditorStyleSet.h"
-#include "SlateOptMacros.h"
 #include "HAL/ConsoleManager.h"
-#include "Framework/Application/SlateApplication.h"
-#include "RigVMModel/Nodes/RigVMCollapseNode.h"
 #include "RigVMModel/Nodes/RigVMFunctionReferenceNode.h"
 
 #define LOCTEXT_NAMESPACE "SControlRigFunctionLocalizationWidget"
