Control Rig: Ensure to repopulate graph during async loading

diff --git a/Engine/Plugins/Experimental/ControlRig/Source/ControlRigDeveloper/Private/ControlRigBlueprint.cpp b/Engine/Plugins/Experimental/ControlRig/Source/ControlRigDeveloper/Private/ControlRigBlueprint.cpp
--- a/Engine/Plugins/Experimental/ControlRig/Source/ControlRigDeveloper/Private/ControlRigBlueprint.cpp
+++ b/Engine/Plugins/Experimental/ControlRig/Source/ControlRigDeveloper/Private/ControlRigBlueprint.cpp
@@ -470,7 +470,7 @@ void UControlRigBlueprint::PostLoad()
 	GraphsToDetach.Add(GetModel());
 	GraphsToDetach.Add(GetLocalFunctionLibrary());
 
-	if (!IsInAsyncLoadingThread() || IsRunningCommandlet())
+	ensure(IsInGameThread());
 	{
 		for (URigVMGraph* GraphToDetach : GraphsToDetach)
 		{
diff --git a/Engine/Plugins/Experimental/ControlRig/Source/ControlRigDeveloper/Public/ControlRigBlueprint.h b/Engine/Plugins/Experimental/ControlRig/Source/ControlRigDeveloper/Public/ControlRigBlueprint.h
--- a/Engine/Plugins/Experimental/ControlRig/Source/ControlRigDeveloper/Public/ControlRigBlueprint.h
+++ b/Engine/Plugins/Experimental/ControlRig/Source/ControlRigDeveloper/Public/ControlRigBlueprint.h
@@ -128,6 +128,7 @@ public:
 	virtual void SetObjectBeingDebugged(UObject* NewObject) override;
 	virtual void PreSave(const class ITargetPlatform* TargetPlatform) override;
 	virtual void PostLoad() override;
+	virtual bool IsPostLoadThreadSafe() const override { return false; }
 	virtual void PostTransacted(const FTransactionObjectEvent& TransactionEvent) override;
 	virtual void ReplaceDeprecatedNodes() override;
 
