Control Rig: Replace ensure within PostLoad with Check

diff --git a/Engine/Plugins/Experimental/ControlRig/Source/ControlRigDeveloper/Private/ControlRigBlueprint.cpp b/Engine/Plugins/Experimental/ControlRig/Source/ControlRigDeveloper/Private/ControlRigBlueprint.cpp
--- a/Engine/Plugins/Experimental/ControlRig/Source/ControlRigDeveloper/Private/ControlRigBlueprint.cpp
+++ b/Engine/Plugins/Experimental/ControlRig/Source/ControlRigDeveloper/Private/ControlRigBlueprint.cpp
@@ -470,7 +470,7 @@ void UControlRigBlueprint::PostLoad()
 	GraphsToDetach.Add(GetModel());
 	GraphsToDetach.Add(GetLocalFunctionLibrary());
 
-	ensure(IsInGameThread());
+	check(IsInGameThread());
 	{
 		for (URigVMGraph* GraphToDetach : GraphsToDetach)
 		{
