Control Rig: Fix animation curve support for AnimNode

diff --git a/Engine/Plugins/Experimental/ControlRig/Source/ControlRig/Private/AnimNode_ControlRigBase.cpp b/Engine/Plugins/Experimental/ControlRig/Source/ControlRig/Private/AnimNode_ControlRigBase.cpp
--- a/Engine/Plugins/Experimental/ControlRig/Source/ControlRig/Private/AnimNode_ControlRigBase.cpp
+++ b/Engine/Plugins/Experimental/ControlRig/Source/ControlRig/Private/AnimNode_ControlRigBase.cpp
@@ -127,7 +127,7 @@ void FAnimNode_ControlRigBase::UpdateInput(UControlRig* ControlRig, const FPoseC
 		{
 			const FName& Name = Iter.Key();
 			const uint16 Index = Iter.Value();
-			const FRigElementKey Key(Name, ERigElementType::Bone);
+			const FRigElementKey Key(Name, ERigElementType::Curve);
 
 			ControlRig->GetHierarchy()->SetCurveValue(Key, InOutput.Curve.Get(Index));
 		}
@@ -174,7 +174,7 @@ void FAnimNode_ControlRigBase::UpdateOutput(UControlRig* ControlRig, FPoseContex
 		{
 			const FName& Name = Iter.Key();
 			const uint16 Index = Iter.Value();
-			const FRigElementKey Key(Name, ERigElementType::Bone);
+			const FRigElementKey Key(Name, ERigElementType::Curve);
 
 			const float Value = ControlRig->GetHierarchy()->GetCurveValue(Key);
 			InOutput.Curve.Set(Index, Value);
