Control Rig: Support for minimum / maximum values in details panel

diff --git a/Engine/Plugins/Experimental/ControlRig/Source/ControlRigEditor/Private/ControlRigElementDetails.cpp b/Engine/Plugins/Experimental/ControlRig/Source/ControlRigEditor/Private/ControlRigElementDetails.cpp
--- a/Engine/Plugins/Experimental/ControlRig/Source/ControlRigEditor/Private/ControlRigElementDetails.cpp
+++ b/Engine/Plugins/Experimental/ControlRig/Source/ControlRigEditor/Private/ControlRigElementDetails.cpp
@@ -880,7 +880,7 @@ void FRigControlElementDetails_SetupIntegerValueWidget(IDetailCategoryBuilder& I
             [
                 SNew(SNumericEntryBox<int32>)
                 .Font(FEditorStyle::GetFontStyle(TEXT("MenuItem.Font")))
-                .AllowSpin(true)
+                .AllowSpin(InValueType == ERigControlValueType::Current || InValueType == ERigControlValueType::Initial)
                 .MinSliderValue_Lambda([HierarchyPtr, Key, InValueType]() -> TOptional<int32>
 				{
 				    if(InValueType == ERigControlValueType::Current || InValueType == ERigControlValueType::Initial)
@@ -984,7 +984,7 @@ void FRigControlElementDetails_SetupFloatValueWidget(IDetailCategoryBuilder& InC
         [
             SNew(SNumericEntryBox<float>)
             .Font(FEditorStyle::GetFontStyle(TEXT("MenuItem.Font")))
-            .AllowSpin(true)
+            .AllowSpin(InValueType == ERigControlValueType::Current || InValueType == ERigControlValueType::Initial)
             .Value_Lambda([HierarchyPtr, Key, InValueType]() -> float
             {
             	if(HierarchyPtr.IsValid())
@@ -1401,6 +1401,21 @@ void FRigControlElementDetails::CustomizeDetails(IDetailLayoutBuilder& DetailBui
 	}
 
 	FRigControlElementDetails_SetupValueWidget(ControlCategory, ERigControlValueType::Current, ControlElement, HierarchyBeingCustomized);
+	switch (ControlElement->Settings.ControlType)
+	{
+		case ERigControlType::Bool:
+		case ERigControlType::Float:
+		case ERigControlType::Integer:
+		case ERigControlType::Vector2D:
+		{
+			FRigControlElementDetails_SetupValueWidget(ControlCategory, ERigControlValueType::Initial, ControlElement, HierarchyBeingCustomized);
+			break;
+		}
+		default:
+		{
+			break;
+		}
+	}
 	FRigControlElementDetails_SetupValueWidget(LimitsCategory, ERigControlValueType::Minimum, ControlElement, HierarchyBeingCustomized);
 	FRigControlElementDetails_SetupValueWidget(LimitsCategory, ERigControlValueType::Maximum, ControlElement, HierarchyBeingCustomized);
 
diff --git a/Engine/Plugins/Experimental/ControlRig/Source/ControlRigEditor/Private/Editor/ControlRigEditor.cpp b/Engine/Plugins/Experimental/ControlRig/Source/ControlRigEditor/Private/Editor/ControlRigEditor.cpp
--- a/Engine/Plugins/Experimental/ControlRig/Source/ControlRigEditor/Private/Editor/ControlRigEditor.cpp
+++ b/Engine/Plugins/Experimental/ControlRig/Source/ControlRigEditor/Private/Editor/ControlRigEditor.cpp
@@ -3573,6 +3573,23 @@ void FControlRigEditor::OnHierarchyModified_AnyThread(ERigHierarchyNotification
 				CacheNameLists();
 				break;
 			}
+			case ERigHierarchyNotification::ControlSettingChanged:
+			{
+				if(Key == RigElementInDetailPanel)
+				{
+					UControlRigBlueprint* RigBlueprint = GetControlRigBlueprint();
+					check(RigBlueprint);
+
+					const FRigControlElement* SourceControlElement = Cast<FRigControlElement>(Element);
+					FRigControlElement* TargetControlElement = RigBlueprint->Hierarchy->Find<FRigControlElement>(Key);
+
+					if(SourceControlElement && TargetControlElement)
+					{
+						TargetControlElement->Settings = SourceControlElement->Settings;
+					}
+				}
+				break;
+			}
 			default:
 			{
 				break;
