Control Rig: Support for minimum / maximum values in details panel

diff --git a/Engine/Plugins/Experimental/ControlRig/Source/ControlRig/Private/ControlRig.cpp b/Engine/Plugins/Experimental/ControlRig/Source/ControlRig/Private/ControlRig.cpp
--- a/Engine/Plugins/Experimental/ControlRig/Source/ControlRig/Private/ControlRig.cpp
+++ b/Engine/Plugins/Experimental/ControlRig/Source/ControlRig/Private/ControlRig.cpp
@@ -1311,6 +1311,7 @@ void UControlRig::HandleHierarchyModified(ERigHierarchyNotification InNotificati
 			break;
 		}
 		case ERigHierarchyNotification::ControlSettingChanged:
+		case ERigHierarchyNotification::ControlGizmoTransformChanged:
 		{
 			if(FRigControlElement* ControlElement = Cast<FRigControlElement>((FRigBaseElement*)InElement))
 			{
diff --git a/Engine/Plugins/Experimental/ControlRig/Source/ControlRig/Private/Rigs/RigHierarchy.cpp b/Engine/Plugins/Experimental/ControlRig/Source/ControlRig/Private/Rigs/RigHierarchy.cpp
--- a/Engine/Plugins/Experimental/ControlRig/Source/ControlRig/Private/Rigs/RigHierarchy.cpp
+++ b/Engine/Plugins/Experimental/ControlRig/Source/ControlRig/Private/Rigs/RigHierarchy.cpp
@@ -1337,7 +1337,10 @@ void URigHierarchy::SetControlGizmoTransform(FRigControlElement* InControlElemen
 	}
 #endif
 
-	Notify(ERigHierarchyNotification::ControlSettingChanged, InControlElement);
+	if(IsLocal(InTransformType))
+	{
+		Notify(ERigHierarchyNotification::ControlGizmoTransformChanged, InControlElement);
+	}
 }
 
 FTransform URigHierarchy::GetParentTransform(FRigBaseElement* InElement, const ERigTransformType::Type InTransformType) const
diff --git a/Engine/Plugins/Experimental/ControlRig/Source/ControlRig/Public/Rigs/RigHierarchyDefines.h b/Engine/Plugins/Experimental/ControlRig/Source/ControlRig/Public/Rigs/RigHierarchyDefines.h
--- a/Engine/Plugins/Experimental/ControlRig/Source/ControlRig/Public/Rigs/RigHierarchyDefines.h
+++ b/Engine/Plugins/Experimental/ControlRig/Source/ControlRig/Public/Rigs/RigHierarchyDefines.h
@@ -45,6 +45,7 @@ enum class ERigHierarchyNotification : uint8
 	HierarchyReset,
 	ControlSettingChanged,
 	ControlVisibilityChanged,
+	ControlGizmoTransformChanged,
 
 	/** MAX - invalid */
 	Max UMETA(Hidden),
diff --git a/Engine/Plugins/Experimental/ControlRig/Source/ControlRigEditor/Private/EditMode/ControlRigEditMode.cpp b/Engine/Plugins/Experimental/ControlRig/Source/ControlRigEditor/Private/EditMode/ControlRigEditMode.cpp
--- a/Engine/Plugins/Experimental/ControlRig/Source/ControlRigEditor/Private/EditMode/ControlRigEditMode.cpp
+++ b/Engine/Plugins/Experimental/ControlRig/Source/ControlRigEditor/Private/EditMode/ControlRigEditMode.cpp
@@ -1903,6 +1903,7 @@ void FControlRigEditMode::OnHierarchyModified(ERigHierarchyNotification InNotif,
 		case ERigHierarchyNotification::ElementRenamed:
 		case ERigHierarchyNotification::HierarchyReset:
 		case ERigHierarchyNotification::ControlSettingChanged:
+		case ERigHierarchyNotification::ControlGizmoTransformChanged:
 		{
 			RequestToRecreateGizmoActors();
 			break;
diff --git a/Engine/Plugins/Experimental/ControlRig/Source/ControlRigEditor/Private/Editor/ControlRigEditor.cpp b/Engine/Plugins/Experimental/ControlRig/Source/ControlRigEditor/Private/Editor/ControlRigEditor.cpp
--- a/Engine/Plugins/Experimental/ControlRig/Source/ControlRigEditor/Private/Editor/ControlRigEditor.cpp
+++ b/Engine/Plugins/Experimental/ControlRig/Source/ControlRigEditor/Private/Editor/ControlRigEditor.cpp
@@ -3146,13 +3146,14 @@ void FControlRigEditor::OnBlueprintPropertyChainEvent(FPropertyChangedChainEvent
 			TransformType = ERigTransformType::MakeLocal(TransformType); 
 		}
 
-		URigHierarchy* Hierarchy = DebuggedControlRig->GetHierarchy();
+		URigHierarchy* DebuggedHierarchy = DebuggedControlRig->GetHierarchy();
+		URigHierarchy* Hierarchy = DebuggedHierarchy;
 		if(bSetupModeEnabled || bIsInitial)
 		{
 			Hierarchy = ControlRigBP->Hierarchy;
 		}
 
-		FRigBaseElement* SourceElement = DebuggedControlRig->GetHierarchy()->Find(RigElementInDetailPanel);
+		FRigBaseElement* SourceElement = DebuggedHierarchy->Find(RigElementInDetailPanel);
 		if(SourceElement == nullptr)
 		{
 			SourceElement = Hierarchy->Find(RigElementInDetailPanel);;
@@ -3219,8 +3220,9 @@ void FControlRigEditor::OnBlueprintPropertyChainEvent(FPropertyChangedChainEvent
 				if(FRigControlElement* TargetControlElement = Cast<FRigControlElement>(TargetElement))
 				{
 					const FTransform Transform = SourceControlElement->Gizmo.Get(TransformType);
-					Hierarchy->SetControlGizmoTransform(TargetControlElement, Transform, TransformType, true, true);
-					Hierarchy->SetControlGizmoTransform(TargetControlElement, Transform, ERigTransformType::MakeCurrent(TransformType), true, true);
+					DebuggedHierarchy->SetControlGizmoTransform(SourceControlElement, Transform, TransformType, false, true);
+					DebuggedHierarchy->SetControlGizmoTransform(SourceControlElement, Transform, ERigTransformType::MakeCurrent(TransformType), false, true);
+					return;
 				}
 			}
 		}
@@ -3560,7 +3562,7 @@ void FControlRigEditor::OnHierarchyModified_AnyThread(ERigHierarchyNotification
     		return;
     	}
 
-        const FRigBaseElement* Element = WeakHierarchy.Get()->Find(Key);
+        FRigBaseElement* Element = WeakHierarchy.Get()->Find(Key);
 
 		switch(InNotif)
 		{
@@ -3590,6 +3592,22 @@ void FControlRigEditor::OnHierarchyModified_AnyThread(ERigHierarchyNotification
 				}
 				break;
 			}
+			case ERigHierarchyNotification::ControlGizmoTransformChanged:
+			{
+				if(Key == RigElementInDetailPanel)
+				{
+					UControlRigBlueprint* RigBlueprint = GetControlRigBlueprint();
+					check(RigBlueprint);
+
+					FRigControlElement* SourceControlElement = Cast<FRigControlElement>(Element);
+					if(SourceControlElement)
+					{
+						RigBlueprint->Hierarchy->SetControlGizmoTransform(Key, WeakHierarchy.Get()->GetControlGizmoTransform(SourceControlElement, ERigTransformType::InitialLocal), true);
+						RigBlueprint->Hierarchy->SetControlGizmoTransform(Key, WeakHierarchy.Get()->GetControlGizmoTransform(SourceControlElement, ERigTransformType::CurrentLocal), false);
+					}
+				}
+				break;
+			}
 			default:
 			{
 				break;
diff --git a/Engine/Plugins/Experimental/ControlRig/Source/ControlRigEditor/Private/Editor/SRigHierarchy.cpp b/Engine/Plugins/Experimental/ControlRig/Source/ControlRigEditor/Private/Editor/SRigHierarchy.cpp
--- a/Engine/Plugins/Experimental/ControlRig/Source/ControlRigEditor/Private/Editor/SRigHierarchy.cpp
+++ b/Engine/Plugins/Experimental/ControlRig/Source/ControlRigEditor/Private/Editor/SRigHierarchy.cpp
@@ -1240,6 +1240,7 @@ void SRigHierarchy::HandleSetObjectBeingDebugged(UObject* InObject)
 	if(UControlRig* ControlRig = Cast<UControlRig>(InObject))
 	{
 		ControlRigBeingDebuggedPtr = ControlRig;
+		ControlRig->GetHierarchy()->OnModified().RemoveAll(this);
 		ControlRig->GetHierarchy()->OnModified().AddSP(this, &SRigHierarchy::OnHierarchyModified_AnyThread);
 	}
 
