Control Rig: Warn the user when wiring up a loop iteration pin to a return node

diff --git a/Engine/Plugins/Experimental/ControlRig/Source/ControlRigDeveloper/Private/Graph/ControlRigGraphSchema.cpp b/Engine/Plugins/Experimental/ControlRig/Source/ControlRigDeveloper/Private/Graph/ControlRigGraphSchema.cpp
--- a/Engine/Plugins/Experimental/ControlRig/Source/ControlRigDeveloper/Private/Graph/ControlRigGraphSchema.cpp
+++ b/Engine/Plugins/Experimental/ControlRig/Source/ControlRigDeveloper/Private/Graph/ControlRigGraphSchema.cpp
@@ -26,6 +26,7 @@
 
 #if WITH_EDITOR
 #include "ControlRigEditor/Private/Editor/SControlRigFunctionLocalizationWidget.h"
+#include "Misc/MessageDialog.h"
 #endif
 
 #define LOCTEXT_NAMESPACE "ControlRigGraphSchema"
@@ -144,10 +145,12 @@ void UControlRigGraphSchema::GetContextMenuActions(class UToolMenu* Menu, class
 bool UControlRigGraphSchema::TryCreateConnection(UEdGraphPin* PinA, UEdGraphPin* PinB) const
 {
 #if WITH_EDITOR
+
 	if (GEditor)
 	{
 		GEditor->CancelTransaction(0);
 	}
+	
 #endif
 
 	if (PinA == PinB)
@@ -175,6 +178,82 @@ bool UControlRigGraphSchema::TryCreateConnection(UEdGraphPin* PinA, UEdGraphPin*
 				PinA = PinB;
 				PinB = Temp;
 			}
+
+#if WITH_EDITOR
+
+			// check if we are trying to connect a loop iteration pin to a return
+			if(URigVMGraph* Graph = Controller->GetGraph())
+			{
+				if(URigVMPin* TargetPin = Graph->FindPin(PinB->GetName()))
+				{
+					if(TargetPin->IsExecuteContext() && TargetPin->GetNode()->IsA<URigVMFunctionReturnNode>())
+					{
+						bool bIsInLoopIteration = false;
+						if(URigVMPin* SourcePin = Graph->FindPin(PinA->GetName()))
+						{
+							while(SourcePin)
+							{
+								if(!SourcePin->IsExecuteContext())
+								{
+									break;
+								}
+								URigVMPin* CurrentSourcePin = SourcePin;
+								SourcePin = nullptr;
+								
+								if(URigVMUnitNode* UnitNode = Cast<URigVMUnitNode>(CurrentSourcePin->GetNode()))
+								{
+									TSharedPtr<FStructOnScope> UnitScope = UnitNode->ConstructStructInstance();
+									if(UnitScope.IsValid())
+									{
+										FRigVMStruct* Unit = (FRigVMStruct*)UnitScope->GetStructMemory();
+										if(Unit->IsForLoop())
+										{
+											if(CurrentSourcePin->GetFName() != FRigVMStruct::ForLoopCompletedPinName)
+											{
+												bIsInLoopIteration = true;
+												break;
+											}
+										}
+									}
+									
+								}
+
+								for(URigVMPin* PinOnSourceNode : CurrentSourcePin->GetNode()->GetPins())
+								{
+									if(!PinOnSourceNode->IsExecuteContext())
+									{
+										continue;
+									}
+
+									if(PinOnSourceNode->GetDirection() != ERigVMPinDirection::Input &&
+										PinOnSourceNode->GetDirection() != ERigVMPinDirection::IO)
+									{
+										continue;
+									}
+
+									TArray<URigVMPin*> NextSourcePins = PinOnSourceNode->GetLinkedSourcePins();
+									if(NextSourcePins.Num() > 0)
+									{
+										SourcePin = NextSourcePins[0];
+										break;
+									}
+								}
+							}
+						}
+
+						if(bIsInLoopIteration)
+						{
+							const EAppReturnType::Type Answer = FMessageDialog::Open( EAppMsgType::YesNo, FText::FromString( TEXT("Linking the execute within a loop iteration to the return\nof a function is not recommended.\nAre you sure?") ) );
+							if(Answer == EAppReturnType::No)
+							{
+								return false;
+							}
+						}
+					}
+				}
+			}
+#endif
+			
 			return Controller->AddLink(PinA->GetName(), PinB->GetName());
 		}
 	}
