Control Rig: Node run count feature

diff --git a/Engine/Plugins/Experimental/ControlRig/Source/ControlRigDeveloper/Private/ControlRigBlueprint.cpp b/Engine/Plugins/Experimental/ControlRig/Source/ControlRigDeveloper/Private/ControlRigBlueprint.cpp
--- a/Engine/Plugins/Experimental/ControlRig/Source/ControlRigDeveloper/Private/ControlRigBlueprint.cpp
+++ b/Engine/Plugins/Experimental/ControlRig/Source/ControlRigDeveloper/Private/ControlRigBlueprint.cpp
@@ -660,8 +660,8 @@ void UControlRigBlueprint::HandleReportFromCompiler(EMessageSeverity::Type InSev
 		FNotificationInfo Info(FText::FromString(InMessage));
 		Info.Image = FCoreStyle::Get().GetBrush(TEXT("MessageLog.Error"));
 		Info.bFireAndForget = true;
-		Info.FadeOutDuration = 10.0f;
-		Info.ExpireDuration = 0.0f;
+		Info.FadeOutDuration = 5.0f;
+		Info.ExpireDuration = 8.0f;
 		TSharedPtr<SNotificationItem> NotificationPtr = FSlateNotificationManager::Get().AddNotification(Info);
 		NotificationPtr->SetCompletionState(SNotificationItem::CS_Success);
 #endif
diff --git a/Engine/Plugins/Experimental/ControlRig/Source/ControlRigDeveloper/Private/ControlRigBlueprintCompiler.cpp b/Engine/Plugins/Experimental/ControlRig/Source/ControlRigDeveloper/Private/ControlRigBlueprintCompiler.cpp
--- a/Engine/Plugins/Experimental/ControlRig/Source/ControlRigDeveloper/Private/ControlRigBlueprintCompiler.cpp
+++ b/Engine/Plugins/Experimental/ControlRig/Source/ControlRigDeveloper/Private/ControlRigBlueprintCompiler.cpp
@@ -53,8 +53,8 @@ void FControlRigBlueprintCompilerContext::MarkCompilationFailed(const FString& M
 		FNotificationInfo Info(FText::FromString(Message));
 		Info.Image = FCoreStyle::Get().GetBrush(TEXT("MessageLog.Error"));
 		Info.bFireAndForget = true;
-		Info.FadeOutDuration = 10.0f;
-		Info.ExpireDuration = 0.0f;
+		Info.FadeOutDuration = 5.0f;
+		Info.ExpireDuration = 5.0f;
 		TSharedPtr<SNotificationItem> NotificationPtr = FSlateNotificationManager::Get().AddNotification(Info);
 		NotificationPtr->SetCompletionState(SNotificationItem::CS_Success);
 #endif
diff --git a/Engine/Plugins/Experimental/ControlRig/Source/ControlRigDeveloper/Public/ControlRigBlueprint.h b/Engine/Plugins/Experimental/ControlRig/Source/ControlRigDeveloper/Public/ControlRigBlueprint.h
--- a/Engine/Plugins/Experimental/ControlRig/Source/ControlRigDeveloper/Public/ControlRigBlueprint.h
+++ b/Engine/Plugins/Experimental/ControlRig/Source/ControlRigDeveloper/Public/ControlRigBlueprint.h
@@ -98,6 +98,41 @@ struct CONTROLRIGDEVELOPER_API FControlRigPublicFunctionData
 	TArray<FControlRigPublicFunctionArg> Arguments;
 };
 
+USTRUCT()
+struct CONTROLRIGDEVELOPER_API FRigGraphDisplaySettings
+{
+	GENERATED_BODY();
+
+	FRigGraphDisplaySettings()
+		: bShowNodeRunCounts(false)
+		, NodeRunLowerBound(1)
+		, NodeRunLimit(64)
+	{
+	}
+
+	// When enabled shows the node counts both in the graph view as
+	// we as in the execution stack window.
+	// The number on each node represents how often the node has been run.
+	// Keep in mind when looking at nodes in a function the count
+	// represents the sum of all counts for each node based on all
+	// references of the function currently running.
+	UPROPERTY(EditAnywhere, Category = "Graph Display Settings")
+	bool bShowNodeRunCounts;
+
+	// A lower limit for counts for nodes used for debugging.
+	// Any node lower than this count won't show the run count.
+	UPROPERTY(EditAnywhere, Category = "Graph Display Settings")
+	int32 NodeRunLowerBound;
+
+	// A upper limit for counts for nodes used for debugging.
+	// If a node reaches this count a warning will be issued for the
+	// node and displayed both in the execution stack as well as in the
+	// graph. Setting this to <= 1 disables the warning.
+	// Note: The count limit doesn't apply to functions / collapse nodes.
+	UPROPERTY(EditAnywhere, Category = "Graph Display Settings")
+	int32 NodeRunLimit;
+};
+
 UCLASS(BlueprintType, meta=(IgnoreClassThumbnail))
 class CONTROLRIGDEVELOPER_API UControlRigBlueprint : public UBlueprint, public IInterface_PreviewMeshProvider
 {
@@ -207,6 +242,9 @@ public:
 	bool IsFunctionPublic(const FName& InFunctionName) const;
 	void MarkFunctionPublic(const FName& InFunctionName, bool bIsPublic = true);
 
+	UPROPERTY(EditAnywhere, Category = "User Interface")
+	FRigGraphDisplaySettings RigGraphDisplaySettings;
+
 	UPROPERTY(EditAnywhere, BlueprintReadWrite, Category = "VM")
 	FRigVMCompileSettings VMCompileSettings;
 
diff --git a/Engine/Plugins/Experimental/ControlRig/Source/ControlRigEditor/Private/Editor/ControlRigEditor.cpp b/Engine/Plugins/Experimental/ControlRig/Source/ControlRigEditor/Private/Editor/ControlRigEditor.cpp
--- a/Engine/Plugins/Experimental/ControlRig/Source/ControlRigEditor/Private/Editor/ControlRigEditor.cpp
+++ b/Engine/Plugins/Experimental/ControlRig/Source/ControlRigEditor/Private/Editor/ControlRigEditor.cpp
@@ -1330,8 +1330,8 @@ void FControlRigEditor::Compile()
 				{
 					FNotificationInfo Info(LOCTEXT("ControlRigBlueprintCompilerEmptyRigMessage", "The Control Rig you compiled doesn't do anything. Did you forget to add a Begin_Execution node?"));
 					Info.bFireAndForget = true;
-					Info.FadeOutDuration = 10.0f;
-					Info.ExpireDuration = 0.0f;
+					Info.FadeOutDuration = 5.0f;
+					Info.ExpireDuration = 5.0f;
 					TSharedPtr<SNotificationItem> NotificationPtr = FSlateNotificationManager::Get().AddNotification(Info);
 					NotificationPtr->SetCompletionState(SNotificationItem::CS_Success);
 				}
@@ -1970,8 +1970,6 @@ void FControlRigEditor::HandleVMCompiledEvent(UBlueprint* InBlueprint, URigVM* I
 
 void FControlRigEditor::HandleControlRigExecutedEvent(UControlRig* InControlRig, const EControlRigState InState, const FName& InEventName)
 {
-	UpdateGraphCompilerErrors();
-
 	if (UControlRigBlueprint* ControlRigBP = GetControlRigBlueprint())
 	{
 		if(RigElementInDetailPanel.IsValid())
@@ -2016,7 +2014,78 @@ void FControlRigEditor::HandleControlRigExecutedEvent(UControlRig* InControlRig,
 				Hierarchy->GetControlGizmoTransform(ControlElement, ERigTransformType::InitialLocal);
 			}
 		}
+		
+		if(ControlRigBP->RigGraphDisplaySettings.NodeRunLimit > 1)
+		{
+			if(UControlRig* DebuggedControlRig = Cast<UControlRig>(ControlRigBP->GetObjectBeingDebugged()))
+			{
+				if(URigVM* VM = DebuggedControlRig->GetVM())
+				{
+					bool bFoundLimitWarnings = false;
+					
+					const FRigVMByteCode& ByteCode = VM->GetByteCode();
+					for(int32 InstructionIndex = 0; InstructionIndex < ByteCode.GetNumInstructions(); InstructionIndex++)
+					{
+						const int32 Count = VM->GetInstructionVisitedCount(InstructionIndex);
+						if(Count > ControlRigBP->RigGraphDisplaySettings.NodeRunLimit)
+						{
+							bFoundLimitWarnings = true;
+
+							const FString CallPath = VM->GetByteCode().GetCallPathForInstruction(InstructionIndex); 
+							if(!KnownInstructionLimitWarnings.Contains(CallPath))
+							{
+								const FString Message = FString::Printf(
+                                    TEXT("Instruction has hit the NodeRunLimit\n(ran %d times, limit is %d)\n\nYou can increase the limit in the class settings."),
+                                    Count,
+                                    ControlRigBP->RigGraphDisplaySettings.NodeRunLimit
+                                );
+
+								if(DebuggedControlRig->ControlRigLog)
+								{
+									DebuggedControlRig->ControlRigLog->Entries.Add(
+										FControlRigLog::FLogEntry(EMessageSeverity::Warning, InEventName, InstructionIndex, Message
+									));
+								}
+
+								if(URigVMNode* Subject = Cast<URigVMNode>(VM->GetByteCode().GetSubjectForInstruction(InstructionIndex)))
+								{
+									FNotificationInfo Info(FText::FromString(Message));
+									Info.bFireAndForget = true;
+									Info.FadeOutDuration = 1.0f;
+									Info.ExpireDuration = 5.0f;
+
+									if(UControlRigGraph* EdGraph = Cast<UControlRigGraph>(GetControlRigBlueprint()->GetEdGraph(Subject->GetGraph())))
+									{
+										if(UEdGraphNode* Node = EdGraph->FindNodeForModelNodeName(Subject->GetFName()))
+										{
+											Info.Hyperlink = FSimpleDelegate::CreateLambda([this, Node] ()
+	                                        {
+	                                            JumpToHyperlink(Node, false);
+	                                        });
+									
+											Info.HyperlinkText = FText::FromString(Subject->GetName());
+										}
+									}
+
+									TSharedPtr<SNotificationItem> NotificationPtr = FSlateNotificationManager::Get().AddNotification(Info);
+									NotificationPtr->SetCompletionState(SNotificationItem::CS_Fail);
+								}
+
+								KnownInstructionLimitWarnings.Add(CallPath, Message);
+							}
+						}
+					}
+
+					if(!bFoundLimitWarnings)
+					{
+						KnownInstructionLimitWarnings.Reset();
+					}
+				}
+			}
+		}
 	}
+
+	UpdateGraphCompilerErrors();
 }
 
 void FControlRigEditor::CreateEditorModeManager()
@@ -3667,8 +3736,8 @@ void FControlRigEditor::OnCurveContainerChanged()
 	// notification
 	FNotificationInfo Info(LOCTEXT("CurveContainerChangeHelpMessage", "CurveContainer has been successfully modified."));
 	Info.bFireAndForget = true;
-	Info.FadeOutDuration = 10.0f;
-	Info.ExpireDuration = 0.0f;
+	Info.FadeOutDuration = 5.0f;
+	Info.ExpireDuration = 5.0f;
 
 	TSharedPtr<SNotificationItem> NotificationPtr = FSlateNotificationManager::Get().AddNotification(Info);
 	NotificationPtr->SetCompletionState(SNotificationItem::CS_Success);
@@ -4540,7 +4609,7 @@ void FControlRigEditor::UpdateGraphCompilerErrors()
 					Info.Image = FEditorStyle::GetBrush(TEXT("MessageLog.Error"));
 					Info.bFireAndForget = true;
 					Info.FadeOutDuration = 5.0f;
-					Info.ExpireDuration = 0.0f;
+					Info.ExpireDuration = 5.0f;
 					TSharedPtr<SNotificationItem> NotificationPtr = FSlateNotificationManager::Get().AddNotification(Info);
 					NotificationPtr->SetCompletionState(SNotificationItem::CS_Success);
 				}
@@ -4556,7 +4625,7 @@ void FControlRigEditor::UpdateGraphCompilerErrors()
 					Info.Image = FEditorStyle::GetBrush(TEXT("MessageLog.Warning"));
 					Info.bFireAndForget = true;
 					Info.FadeOutDuration = 5.0f;
-					Info.ExpireDuration = 0.0f;
+					Info.ExpireDuration = 5.0f;
 					TSharedPtr<SNotificationItem> NotificationPtr = FSlateNotificationManager::Get().AddNotification(Info);
 					NotificationPtr->SetCompletionState(SNotificationItem::CS_Success);
 				}
diff --git a/Engine/Plugins/Experimental/ControlRig/Source/ControlRigEditor/Private/Editor/ControlRigEditor.h b/Engine/Plugins/Experimental/ControlRig/Source/ControlRigEditor/Private/Editor/ControlRigEditor.h
--- a/Engine/Plugins/Experimental/ControlRig/Source/ControlRigEditor/Private/Editor/ControlRigEditor.h
+++ b/Engine/Plugins/Experimental/ControlRig/Source/ControlRigEditor/Private/Editor/ControlRigEditor.h
@@ -402,6 +402,7 @@ protected:
 	EControlRigEditorEventQueue LastEventQueue;
 	FString LastDebuggedRig;
 	int32 RigHierarchyTabCount;
+	TMap<FString, FString> KnownInstructionLimitWarnings;
 
 	friend class FControlRigEditorMode;
 	friend class SControlRigStackView;
diff --git a/Engine/Plugins/Experimental/ControlRig/Source/ControlRigEditor/Private/Editor/SControlRigStackView.cpp b/Engine/Plugins/Experimental/ControlRig/Source/ControlRigEditor/Private/Editor/SControlRigStackView.cpp
--- a/Engine/Plugins/Experimental/ControlRig/Source/ControlRigEditor/Private/Editor/SControlRigStackView.cpp
+++ b/Engine/Plugins/Experimental/ControlRig/Source/ControlRigEditor/Private/Editor/SControlRigStackView.cpp
@@ -4,6 +4,7 @@
 #include "Widgets/Layout/SBorder.h"
 #include "Widgets/SBoxPanel.h"
 #include "Widgets/Layout/SScrollBar.h"
+#include "Widgets/Layout/SSpacer.h"
 #include "Widgets/Input/SSearchBox.h"
 #include "ControlRigEditor.h"
 #include "ControlRigEditorStyle.h"
@@ -12,6 +13,7 @@
 #include "ControlRigBlueprintGeneratedClass.h"
 #include "Framework/MultiBox/MultiBoxBuilder.h"
 #include "EditorStyleSet.h"
+#include "DetailLayoutBuilder.h"
 #include "Widgets/Input/SSearchBox.h"
 
 #define LOCTEXT_NAMESPACE "SControlRigStackView"
@@ -19,27 +21,30 @@
 //////////////////////////////////////////////////////////////
 /// FRigStackEntry
 ///////////////////////////////////////////////////////////
-FRigStackEntry::FRigStackEntry(int32 InEntryIndex, ERigStackEntry::Type InEntryType, int32 InInstructionIndex, ERigVMOpCode InOpCode, const FString& InLabel)
+FRigStackEntry::FRigStackEntry(int32 InEntryIndex, ERigStackEntry::Type InEntryType, int32 InInstructionIndex, ERigVMOpCode InOpCode, const FString& InLabel, const FRigVMASTProxy& InProxy)
 	: EntryIndex(InEntryIndex)
 	, EntryType(InEntryType)
 	, InstructionIndex(InInstructionIndex)
+	, CallPath(InProxy.GetCallstack().GetCallPath())
+	, Callstack(InProxy.GetCallstack())
 	, OpCode(InOpCode)
 	, Label(InLabel)
 {
 
 }
 
-TSharedRef<ITableRow> FRigStackEntry::MakeTreeRowWidget(const TSharedRef<STableViewBase>& InOwnerTable, TSharedRef<FRigStackEntry> InEntry, TSharedRef<FUICommandList> InCommandList, TSharedPtr<SControlRigStackView> InStackView)
+TSharedRef<ITableRow> FRigStackEntry::MakeTreeRowWidget(const TSharedRef<STableViewBase>& InOwnerTable, TSharedRef<FRigStackEntry> InEntry, TSharedRef<FUICommandList> InCommandList, TSharedPtr<SControlRigStackView> InStackView, TWeakObjectPtr<UControlRigBlueprint> InBlueprint)
 {
-	return SNew(SRigStackItem, InOwnerTable, InEntry, InCommandList);
+	return SNew(SRigStackItem, InOwnerTable, InEntry, InCommandList, InBlueprint);
 }
 
 //////////////////////////////////////////////////////////////
 /// SRigStackItem
 ///////////////////////////////////////////////////////////
-void SRigStackItem::Construct(const FArguments& InArgs, const TSharedRef<STableViewBase>& OwnerTable, TSharedRef<FRigStackEntry> InStackEntry, TSharedRef<FUICommandList> InCommandList)
+void SRigStackItem::Construct(const FArguments& InArgs, const TSharedRef<STableViewBase>& OwnerTable, TSharedRef<FRigStackEntry> InStackEntry, TSharedRef<FUICommandList> InCommandList, TWeakObjectPtr<UControlRigBlueprint> InBlueprint)
 {
 	WeakStackEntry = InStackEntry;
+	WeakBlueprint = InBlueprint;
 	WeakCommandList = InCommandList;
 
 	TSharedPtr< STextBlock > NumberWidget;
@@ -86,6 +91,7 @@ void SRigStackItem::Construct(const FArguments& InArgs, const TSharedRef<STableV
 			[
 				SAssignNew(NumberWidget, STextBlock)
 				.Text(this, &SRigStackItem::GetIndexText)
+				.Font(IDetailLayoutBuilder::GetDetailFont())
 			]
 			+ SHorizontalBox::Slot()
 			.MaxWidth(22.f)
@@ -95,20 +101,40 @@ void SRigStackItem::Construct(const FArguments& InArgs, const TSharedRef<STableV
 				SNew(SImage)
 				.Image(Icon)
 			]
+			
 			+ SHorizontalBox::Slot()
-			.AutoWidth()
+			.FillWidth(1.f)
+			.MaxWidth(200.f)
 			.VAlign(VAlign_Center)
 			.HAlign(HAlign_Left)
 			[
 				SAssignNew(TextWidget, STextBlock)
 				.Text(this, &SRigStackItem::GetLabelText)
+				.Font(this, &SRigStackItem::GetLabelFont)
 			]
-		], OwnerTable);
+
+			+ SHorizontalBox::Slot()
+			.FillWidth(1.f)
+			.MaxWidth(40.f)
+			[
+				SNew(SSpacer)
+			]
+
+			+ SHorizontalBox::Slot()
+			.MaxWidth(36.f)
+            .VAlign(VAlign_Center)
+            .HAlign(HAlign_Right)
+            [
+                SNew(STextBlock)
+                .Text(this, &SRigStackItem::GetVisitedCountText)
+                .Font(IDetailLayoutBuilder::GetDetailFont())
+            ]
+        ], OwnerTable);
 }
 
 FText SRigStackItem::GetIndexText() const
 {
-	FString IndexStr = FString::FromInt(WeakStackEntry.Pin()->EntryIndex) + TEXT(".");
+	const FString IndexStr = FString::FromInt(WeakStackEntry.Pin()->EntryIndex) + TEXT(".");
 	return FText::FromString(IndexStr);
 }
 
@@ -117,6 +143,40 @@ FText SRigStackItem::GetLabelText() const
 	return (FText::FromString(WeakStackEntry.Pin()->Label));
 }
 
+FSlateFontInfo SRigStackItem::GetLabelFont() const
+{
+	if(GetVisitedCountText().IsEmpty())
+	{
+		return IDetailLayoutBuilder::GetDetailFont();
+	}
+	return IDetailLayoutBuilder::GetDetailFontBold();
+}
+
+FText SRigStackItem::GetVisitedCountText() const
+{
+	if(WeakStackEntry.IsValid() && WeakBlueprint.IsValid())
+	{
+		if(WeakBlueprint->RigGraphDisplaySettings.bShowNodeRunCounts)
+		{
+			if(WeakStackEntry.Pin()->EntryType == ERigStackEntry::Operator)
+			{
+				if(UControlRig* ControlRig = Cast<UControlRig>(WeakBlueprint->GetObjectBeingDebugged()))
+				{
+					if(URigVM* VM = ControlRig->GetVM())
+					{
+						const int32 Count = VM->GetInstructionVisitedCount(WeakStackEntry.Pin()->InstructionIndex);
+						if(Count > 0)
+						{
+							return FText::FromString(FString::FromInt(Count));
+						}
+					}
+				}
+			}
+		}
+	}
+	return FText();
+}
+
 //////////////////////////////////////////////////////////////
 /// SControlRigStackView
 ///////////////////////////////////////////////////////////
@@ -200,10 +260,11 @@ void SControlRigStackView::Construct( const FArguments& InArgs, TSharedRef<FCont
 				SAssignNew(TreeView, STreeView<TSharedPtr<FRigStackEntry>>)
 				.TreeItemsSource(&Operators)
 				.SelectionMode(ESelectionMode::Multi)
-				.OnGenerateRow(this, &SControlRigStackView::MakeTableRowWidget)
+				.OnGenerateRow(this, &SControlRigStackView::MakeTableRowWidget, ControlRigBlueprint)
 				.OnGetChildren(this, &SControlRigStackView::HandleGetChildrenForTree)
 				.OnSelectionChanged(this, &SControlRigStackView::OnSelectionChanged)
 				.OnContextMenuOpening(this, &SControlRigStackView::CreateContextMenu)
+				.OnMouseButtonDoubleClick(this, &SControlRigStackView::OnMouseButtonDoubleClick)
 				.ItemHeight(28)
 			]
 		]
@@ -272,27 +333,30 @@ void SControlRigStackView::OnSelectionChanged(TSharedPtr<FRigStackEntry> Selecti
 		TMap<URigVMGraph*, TArray<FName>> SelectedNodesPerGraph;
 		for (TSharedPtr<FRigStackEntry>& Entry : SelectedItems)
 		{
-			UObject* Subject = ByteCode.GetSubjectForInstruction(Entry->InstructionIndex);
-			URigVMGraph* SubjectGraph = nullptr;
-
-			FName NodeName = NAME_None;
-			if (URigVMNode* Node = Cast<URigVMNode>(Subject))
+			for(int32 StackIndex = 0; StackIndex < Entry->Callstack.Num(); StackIndex++)
 			{
-				NodeName = Node->GetFName();
-				SubjectGraph = Node->GetGraph();
-			}
-			else if (URigVMPin* Pin = Cast<URigVMPin>(Subject))
-			{
-				NodeName = Pin->GetNode()->GetFName();
-				SubjectGraph = Pin->GetGraph();
-			}
+				const UObject* Subject = Entry->Callstack[StackIndex];
+				URigVMGraph* SubjectGraph = nullptr;
 
-			if (NodeName.IsNone() || SubjectGraph == nullptr)
-			{
-				continue;
-			}
+				FName NodeName = NAME_None;
+				if (const URigVMNode* Node = Cast<URigVMNode>(Subject))
+				{
+					NodeName = Node->GetFName();
+					SubjectGraph = Node->GetGraph();
+				}
+				else if (const URigVMPin* Pin = Cast<URigVMPin>(Subject))
+				{
+					NodeName = Pin->GetNode()->GetFName();
+					SubjectGraph = Pin->GetGraph();
+				}
+
+				if (NodeName.IsNone() || SubjectGraph == nullptr)
+				{
+					continue;
+				}
 
-			SelectedNodesPerGraph.FindOrAdd(SubjectGraph).AddUnique(NodeName);
+				SelectedNodesPerGraph.FindOrAdd(SubjectGraph).AddUnique(NodeName);
+			}
 		}
 
 		for (const TPair< URigVMGraph*, TArray<FName> >& Pair : SelectedNodesPerGraph)
@@ -309,9 +373,9 @@ void SControlRigStackView::BindCommands()
 	CommandList->MapAction(Commands.FocusOnSelection, FExecuteAction::CreateSP(this, &SControlRigStackView::HandleFocusOnSelectedGraphNode));
 }
 
-TSharedRef<ITableRow> SControlRigStackView::MakeTableRowWidget(TSharedPtr<FRigStackEntry> InItem, const TSharedRef<STableViewBase>& OwnerTable)
+TSharedRef<ITableRow> SControlRigStackView::MakeTableRowWidget(TSharedPtr<FRigStackEntry> InItem, const TSharedRef<STableViewBase>& OwnerTable, TWeakObjectPtr<UControlRigBlueprint> InBlueprint)
 {
-	return InItem->MakeTreeRowWidget(OwnerTable, InItem.ToSharedRef(), CommandList.ToSharedRef(), SharedThis(this));
+	return InItem->MakeTreeRowWidget(OwnerTable, InItem.ToSharedRef(), CommandList.ToSharedRef(), SharedThis(this), InBlueprint);
 }
 
 void SControlRigStackView::HandleGetChildrenForTree(TSharedPtr<FRigStackEntry> InItem, TArray<TSharedPtr<FRigStackEntry>>& OutChildren)
@@ -323,21 +387,23 @@ void SControlRigStackView::PopulateStackView(URigVM* InVM)
 {
 	if (InVM)
 	{
-		FRigVMInstructionArray Instructions = InVM->GetInstructions();
+		UControlRigBlueprint* Blueprint = ControlRigEditor.Pin()->GetControlRigBlueprint();
+		URigVMGraph* RootGraph = Blueprint->GetModel();
 
+		const FRigVMInstructionArray Instructions = InVM->GetInstructions();
 		const FRigVMByteCode& ByteCode = InVM->GetByteCode();
 
 		// 1. cache information about instructions/nodes, which will be used later 
-		TMap<FString, FString> NodeNameToDisplayName;
-		for (int32 InstructionIndex = 0; InstructionIndex < ByteCode.Num(); InstructionIndex++)
+		TMap<FString, FString> OperandFormatMap;
+		for (int32 InstructionIndex = 0; InstructionIndex < Instructions.Num(); InstructionIndex++)
 		{
-			UObject* Subject = ByteCode.GetSubjectForInstruction(InstructionIndex);
-			if (URigVMNode* Node = Cast<URigVMNode>(Subject))
+			const FRigVMASTProxy Proxy = FRigVMASTProxy::MakeFromCallPath(ByteCode.GetCallPathForInstruction(InstructionIndex), RootGraph);
+			if(URigVMNode* Node = Proxy.GetSubject<URigVMNode>())
 			{
 				FString DisplayName = Node->GetName();
 
 				// only unit nodes among all nodes has StaticExecute() that generates actual instructions
-				if (URigVMUnitNode* UnitNode = Cast<URigVMUnitNode>(Subject))
+				if (URigVMUnitNode* UnitNode = Cast<URigVMUnitNode>(Node))
 				{
 					DisplayName = UnitNode->GetNodeTitle();
 #if WITH_EDITOR
@@ -355,19 +421,19 @@ void SControlRigStackView::PopulateStackView(URigVM* InVM)
 				}
 
 				// this is needed for name replacement later
-				NodeNameToDisplayName.Add(Node->GetName(), DisplayName);
+				OperandFormatMap.Add(Node->GetName(), DisplayName);
 			}
 		}
 
 		// 2. replace raw operand names with NodeTitle.PinName/PropertyName.OffsetName
-		TArray<FString> Labels = InVM->DumpByteCodeAsTextArray(TArray<int32>(), false, [NodeNameToDisplayName](const FString& RegisterName, const FString& RegisterOffsetName)
+		TArray<FString> Labels = InVM->DumpByteCodeAsTextArray(TArray<int32>(), false, [OperandFormatMap](const FString& RegisterName, const FString& RegisterOffsetName)
 		{
 			FString NewRegisterName = RegisterName;
 			FString NodeName;
 			FString PinName;
 			if (RegisterName.Split(TEXT("."), &NodeName, &PinName))
 			{
-				const FString* NodeTitle = NodeNameToDisplayName.Find(NodeName);
+				const FString* NodeTitle = OperandFormatMap.Find(NodeName);
 				NewRegisterName = FString::Printf(TEXT("%s.%s"), NodeTitle ? **NodeTitle : *NodeName, *PinName);
 			}
 			FString OperandLabel;
@@ -385,18 +451,49 @@ void SControlRigStackView::PopulateStackView(URigVM* InVM)
 		for (int32 InstructionIndex = 0; InstructionIndex < Labels.Num(); InstructionIndex++)
 		{
 			FString Label = Labels[InstructionIndex];
+			const FRigVMASTProxy Proxy = FRigVMASTProxy::MakeFromCallPath(ByteCode.GetCallPathForInstruction(InstructionIndex), RootGraph);
 
-			UObject* Subject = ByteCode.GetSubjectForInstruction(InstructionIndex);
-			if(URigVMNode* Node = Cast<URigVMNode>(Subject))
+			if(URigVMNode* Node = Proxy.GetSubject<URigVMNode>())
 			{
-				Label = NodeNameToDisplayName[Node->GetName()];
+				FString Suffix;
+				switch(Instructions[InstructionIndex].OpCode)
+				{
+					case ERigVMOpCode::Copy:
+					case ERigVMOpCode::Zero:
+	                case ERigVMOpCode::BoolFalse:
+	                case ERigVMOpCode::BoolTrue:
+	                case ERigVMOpCode::Increment:
+	                case ERigVMOpCode::Decrement:
+					case ERigVMOpCode::Equals:
+                	case ERigVMOpCode::NotEquals:
+					case ERigVMOpCode::JumpAbsolute:
+                	case ERigVMOpCode::JumpForward:
+                	case ERigVMOpCode::JumpBackward:
+					case ERigVMOpCode::JumpAbsoluteIf:
+                	case ERigVMOpCode::JumpForwardIf:
+                	case ERigVMOpCode::JumpBackwardIf:
+					case ERigVMOpCode::BeginBlock:
+					case ERigVMOpCode::EndBlock:
+					{
+						const FText OpCodeText = StaticEnum<ERigVMOpCode>()->GetDisplayNameTextByValue((int32)Instructions[InstructionIndex].OpCode);
+						Suffix = FString::Printf(TEXT(" - %s"), *OpCodeText.ToString());
+						break;
+					}
+					default:
+					{
+						break;
+					}
+				}
+				
+				Label =  OperandFormatMap[Node->GetName()] + Suffix;
+				// Label = Proxy.GetCallstack().GetCallPath() + Suffix;
 			}
 
 			// add the entry with the new label to the stack view
 			const FRigVMInstruction& Instruction = Instructions[InstructionIndex];
 			if (FilterText.IsEmpty() || Label.Contains(FilterText.ToString()))
 			{
-				TSharedPtr<FRigStackEntry> NewEntry = MakeShared<FRigStackEntry>(InstructionIndex, ERigStackEntry::Operator, InstructionIndex, Instruction.OpCode, Label);
+				TSharedPtr<FRigStackEntry> NewEntry = MakeShared<FRigStackEntry>(InstructionIndex, ERigStackEntry::Operator, InstructionIndex, Instruction.OpCode, Label, Proxy);
 				Operators.Add(NewEntry);
 			}
 		}
@@ -430,19 +527,19 @@ void SControlRigStackView::RefreshTreeView(URigVM* InVM)
 					{
 						case EMessageSeverity::Info:
 						{
-							Operators[LogEntry.InstructionIndex]->Children.Add(MakeShared<FRigStackEntry>(ChildIndex, ERigStackEntry::Info, LogEntry.InstructionIndex, ERigVMOpCode::Invalid, LogEntry.Message));
+							Operators[LogEntry.InstructionIndex]->Children.Add(MakeShared<FRigStackEntry>(ChildIndex, ERigStackEntry::Info, LogEntry.InstructionIndex, ERigVMOpCode::Invalid, LogEntry.Message, FRigVMASTProxy()));
 							break;
 						}
 						case EMessageSeverity::Warning:
 						case EMessageSeverity::PerformanceWarning:
 						{
-							Operators[LogEntry.InstructionIndex]->Children.Add(MakeShared<FRigStackEntry>(ChildIndex, ERigStackEntry::Warning, LogEntry.InstructionIndex, ERigVMOpCode::Invalid, LogEntry.Message));
+							Operators[LogEntry.InstructionIndex]->Children.Add(MakeShared<FRigStackEntry>(ChildIndex, ERigStackEntry::Warning, LogEntry.InstructionIndex, ERigVMOpCode::Invalid, LogEntry.Message, FRigVMASTProxy()));
 							break;
 						}
 						case EMessageSeverity::Error:
 						case EMessageSeverity::CriticalError:
 						{
-							Operators[LogEntry.InstructionIndex]->Children.Add(MakeShared<FRigStackEntry>(ChildIndex, ERigStackEntry::Error, LogEntry.InstructionIndex, ERigVMOpCode::Invalid, LogEntry.Message));
+							Operators[LogEntry.InstructionIndex]->Children.Add(MakeShared<FRigStackEntry>(ChildIndex, ERigStackEntry::Error, LogEntry.InstructionIndex, ERigVMOpCode::Invalid, LogEntry.Message, FRigVMASTProxy()));
 							break;
 						}
 						default:
@@ -545,16 +642,13 @@ void SControlRigStackView::HandleModifiedEvent(ERigVMGraphNotifType InNotifType,
 		case ERigVMGraphNotifType::NodeSelected:
 		case ERigVMGraphNotifType::NodeDeselected:
 		{
-			TArray<int32> InstructionIndices = ByteCode.GetAllInstructionIndicesForSubject(InSubject);
-			for (int32 InstructionIndex : InstructionIndices)
+			for (TSharedPtr<FRigStackEntry>& Operator : Operators)
 			{
-				if (InstructionIndex >= Operators.Num())
+				if(Operator->Callstack.Contains(InSubject))
 				{
-					break;
+					const TGuardValue<bool> SuspendNotifs(bSuspendModelNotifications, true);
+					TreeView->SetItemSelection(Operator, InNotifType == ERigVMGraphNotifType::NodeSelected, ESelectInfo::Direct);
 				}
-
-				TGuardValue<bool> SuspendNotifs(bSuspendModelNotifications, true);
-				TreeView->SetItemSelection(Operators[InstructionIndex], InNotifType == ERigVMGraphNotifType::NodeSelected, ESelectInfo::Direct);
 			}
 			break;
 		}
@@ -610,4 +704,27 @@ void SControlRigStackView::HandlePreviewControlRigUpdated(FControlRigEditor* InE
 {
 }
 
+void SControlRigStackView::OnMouseButtonDoubleClick(TSharedPtr<FRigStackEntry> InItem)
+{
+	if(ControlRigEditor.IsValid() && ControlRigBlueprint.IsValid())
+	{
+		if(UControlRig* ControlRig = Cast<UControlRig>(ControlRigBlueprint->GetObjectBeingDebugged()))
+		{
+			if(URigVM* VM = ControlRig->GetVM())
+			{
+				if(URigVMNode* Subject = Cast<URigVMNode>(VM->GetByteCode().GetSubjectForInstruction(InItem->InstructionIndex)))
+				{
+					if(UControlRigGraph* EdGraph = Cast<UControlRigGraph>(ControlRigBlueprint->GetEdGraph(Subject->GetGraph())))
+					{
+						if(UEdGraphNode* Node = EdGraph->FindNodeForModelNodeName(Subject->GetFName()))
+						{
+							ControlRigEditor.Pin()->JumpToHyperlink(Node, false);
+						}
+					}
+				}
+			}
+		}
+	}
+}
+
 #undef LOCTEXT_NAMESPACE
diff --git a/Engine/Plugins/Experimental/ControlRig/Source/ControlRigEditor/Private/Editor/SControlRigStackView.h b/Engine/Plugins/Experimental/ControlRig/Source/ControlRigEditor/Private/Editor/SControlRigStackView.h
--- a/Engine/Plugins/Experimental/ControlRig/Source/ControlRigEditor/Private/Editor/SControlRigStackView.h
+++ b/Engine/Plugins/Experimental/ControlRig/Source/ControlRigEditor/Private/Editor/SControlRigStackView.h
@@ -29,13 +29,15 @@ namespace ERigStackEntry
 class FRigStackEntry : public TSharedFromThis<FRigStackEntry>
 {
 public:
-	FRigStackEntry(int32 InEntryIndex, ERigStackEntry::Type InEntryType, int32 InInstructionIndex, ERigVMOpCode InOpCode, const FString& InLabel);
+	FRigStackEntry(int32 InEntryIndex, ERigStackEntry::Type InEntryType, int32 InInstructionIndex, ERigVMOpCode InOpCode, const FString& InLabel, const FRigVMASTProxy& InProxy);
 
-	TSharedRef<ITableRow> MakeTreeRowWidget(const TSharedRef<STableViewBase>& InOwnerTable, TSharedRef<FRigStackEntry> InEntry, TSharedRef<FUICommandList> InCommandList, TSharedPtr<SControlRigStackView> InStackView);
+	TSharedRef<ITableRow> MakeTreeRowWidget(const TSharedRef<STableViewBase>& InOwnerTable, TSharedRef<FRigStackEntry> InEntry, TSharedRef<FUICommandList> InCommandList, TSharedPtr<SControlRigStackView> InStackView, TWeakObjectPtr<UControlRigBlueprint> InBlueprint);
 
 	int32 EntryIndex;
 	ERigStackEntry::Type EntryType;
 	int32 InstructionIndex;
+	FString CallPath;
+	FRigVMCallstack Callstack;
 	ERigVMOpCode OpCode;
 	FString Label;
 	TArray<TSharedPtr<FRigStackEntry>> Children;
@@ -46,14 +48,17 @@ class SRigStackItem : public STableRow<TSharedPtr<FRigStackEntry>>
 	SLATE_BEGIN_ARGS(SRigStackItem) {}
 	SLATE_END_ARGS()
 
-	void Construct(const FArguments& InArgs, const TSharedRef<STableViewBase>& OwnerTable, TSharedRef<FRigStackEntry> InStackEntry, TSharedRef<FUICommandList> InCommandList);
+	void Construct(const FArguments& InArgs, const TSharedRef<STableViewBase>& OwnerTable, TSharedRef<FRigStackEntry> InStackEntry, TSharedRef<FUICommandList> InCommandList, TWeakObjectPtr<UControlRigBlueprint> InBlueprint);
 
 private:
 	TWeakPtr<FRigStackEntry> WeakStackEntry;
+	TWeakObjectPtr<UControlRigBlueprint> WeakBlueprint;
 	TWeakPtr<FUICommandList> WeakCommandList;
 
 	FText GetIndexText() const;
 	FText GetLabelText() const;
+	FSlateFontInfo GetLabelFont() const;
+	FText GetVisitedCountText() const;
 };
 
 class SControlRigStackView : public SCompoundWidget
@@ -82,7 +87,7 @@ private:
 	void BindCommands();
 
 	/** Make a row widget for the table */
-	TSharedRef<ITableRow> MakeTableRowWidget(TSharedPtr<FRigStackEntry> InItem, const TSharedRef<STableViewBase>& OwnerTable);
+	TSharedRef<ITableRow> MakeTableRowWidget(TSharedPtr<FRigStackEntry> InItem, const TSharedRef<STableViewBase>& OwnerTable, TWeakObjectPtr<UControlRigBlueprint> InBlueprint);
 
 	/** Get children for the tree */
 	void HandleGetChildrenForTree(TSharedPtr<FRigStackEntry> InItem, TArray<TSharedPtr<FRigStackEntry>>& OutChildren);
@@ -103,6 +108,7 @@ private:
 	void HandleModifiedEvent(ERigVMGraphNotifType InNotifType, URigVMGraph* InGraph, UObject* InSubject);
 	void HandleControlRigInitializedEvent(UControlRig* InControlRig, const EControlRigState InState, const FName& InEventName);
 	void HandlePreviewControlRigUpdated(FControlRigEditor* InEditor);
+	void OnMouseButtonDoubleClick(TSharedPtr<FRigStackEntry> InItem);
 
 	/** Populate the execution stack with descriptive names for each instruction */
 	void PopulateStackView(URigVM* InVM);
diff --git a/Engine/Plugins/Experimental/ControlRig/Source/ControlRigEditor/Private/Graph/SControlRigGraphNode.cpp b/Engine/Plugins/Experimental/ControlRig/Source/ControlRigEditor/Private/Graph/SControlRigGraphNode.cpp
--- a/Engine/Plugins/Experimental/ControlRig/Source/ControlRigEditor/Private/Graph/SControlRigGraphNode.cpp
+++ b/Engine/Plugins/Experimental/ControlRig/Source/ControlRigEditor/Private/Graph/SControlRigGraphNode.cpp
@@ -6,7 +6,6 @@
 #include "Graph/ControlRigGraph.h"
 #include "SGraphPin.h"
 #include "Graph/ControlRigGraphSchema.h"
-#include "Widgets/Layout/SBorder.h"
 #include "Widgets/SBoxPanel.h"
 #include "Widgets/Layout/SSpacer.h"
 #include "Widgets/Input/SButton.h"
@@ -27,6 +26,8 @@
 #include "RigVMModel/RigVMPin.h"
 #include "RigVMCompiler/RigVMCompiler.h"
 #include "IDocumentation.h"
+#include "DetailLayoutBuilder.h"
+#include "EditorStyleSet.h"
 
 #if WITH_EDITOR
 #include "Editor.h"
@@ -57,6 +58,8 @@ void SControlRigGraphNode::Construct( const FArguments& InArgs )
 		return;
 	}
 
+	Blueprint = Cast<UControlRigBlueprint>(FBlueprintEditorUtils::FindBlueprintForNode(this->GraphNode));
+
 	// Re-cache variable info here (unit structure could have changed since last reconstruction, e.g. array add/remove)
 	// and also create missing pins if it hasn't created yet
 	ControlRigGraphNode->AllocateDefaultPins();
@@ -173,6 +176,17 @@ void SControlRigGraphNode::Construct( const FArguments& InArgs )
 		.Image(ImageBrush)
 		.Visibility(EVisibility::Visible);
 
+	static const FSlateColorBrush WhiteBrush(FLinearColor::White);
+	
+	SAssignNew(InstructionCountTextBlockWidget, STextBlock)
+	.Margin(FMargin(2.0f, 2.0f, 2.0f, 1.0f))
+	.Text(this, &SControlRigGraphNode::GetInstructionCountText)
+	.Font(IDetailLayoutBuilder::GetDetailFont())
+	.ColorAndOpacity(FLinearColor::White)
+	.ShadowColorAndOpacity(FLinearColor(0.1f, 0.1f, 0.1f, 1.f))
+	.Visibility(EVisibility::Visible)
+	.ToolTipText(LOCTEXT("NodeHitCountToolTip", "This number represents the hit count for a node.\nFor functions / collapse nodes it represents the sum of all hit counts of contained nodes.\n\nYou can enable / disable the display of the number in the Class Settings\n(Rig Graph Display Settings -> Show Node Run Counts)"));
+
 	ControlRigGraphNode->GetNodeTitleDirtied().BindSP(this, &SControlRigGraphNode::HandleNodeTitleDirtied);
 }
 
@@ -327,9 +341,11 @@ FReply SControlRigGraphNode::OnMouseButtonDoubleClick(const FGeometry& InMyGeome
 		{
 			if (URigVMNode* ModelNode = RigNode->GetModelNode())
 			{
-				UControlRigBlueprint* RigBlueprint = Cast<UControlRigBlueprint>(FBlueprintEditorUtils::FindBlueprintForNode(RigNode));
-				RigBlueprint->BroadcastNodeDoubleClicked(ModelNode);
-				return FReply::Handled();
+				if(Blueprint.IsValid())
+				{
+					Blueprint->BroadcastNodeDoubleClicked(ModelNode);
+					return FReply::Handled();
+				}
 			}
 		}
 	}
@@ -999,7 +1015,7 @@ TArray<FOverlayWidgetInfo> SControlRigGraphNode::GetOverlayWidgets(bool bSelecte
 
 								if (Widgets.Num() == PreviousNumWidgets)
 								{
-									FVector2D ImageSize = VisualDebugIndicatorWidget->GetDesiredSize();
+									const FVector2D ImageSize = VisualDebugIndicatorWidget->GetDesiredSize();
 
 									FOverlayWidgetInfo Info;
 									Info.OverlayOffset = FVector2D(WidgetSize.X - ImageSize.X - 6.f, 6.f);
@@ -1012,6 +1028,26 @@ TArray<FOverlayWidgetInfo> SControlRigGraphNode::GetOverlayWidgets(bool bSelecte
 					}
 				}
 			}
+
+			if(Blueprint.IsValid())
+			{
+				if(Blueprint->RigGraphDisplaySettings.bShowNodeRunCounts)
+				{
+					if(UControlRig* DebuggedControlRig = Cast<UControlRig>(Blueprint->GetObjectBeingDebugged()))
+					{
+						const int32 Count = ModelNode->GetInstructionVisitedCount(DebuggedControlRig->GetVM(), FRigVMASTProxy(), false);
+						if(Count > Blueprint->RigGraphDisplaySettings.NodeRunLowerBound)
+						{
+							const int32 VOffset = bSelected ? -2 : 2;
+							const FVector2D TextSize = InstructionCountTextBlockWidget->GetDesiredSize();
+							FOverlayWidgetInfo Info;
+							Info.OverlayOffset = FVector2D(WidgetSize.X - TextSize.X - 8.f, VOffset - TextSize.Y);
+							Info.Widget = InstructionCountTextBlockWidget;
+							Widgets.Add(Info);
+						}
+					}
+				}
+			}
 		}
 	}
 
@@ -1050,4 +1086,30 @@ void SControlRigGraphNode::HandleNodeTitleDirtied()
 	}
 }
 
+FText SControlRigGraphNode::GetInstructionCountText() const
+{
+	if(Blueprint.IsValid())
+	{
+		if(Blueprint->RigGraphDisplaySettings.bShowNodeRunCounts)
+		{
+			if (UControlRigGraphNode* RigNode = CastChecked<UControlRigGraphNode>(GraphNode, ECastCheckedType::NullAllowed))
+			{
+				if (URigVMNode* ModelNode = RigNode->GetModelNode())
+				{
+					if(UControlRig* DebuggedControlRig = Cast<UControlRig>(Blueprint->GetObjectBeingDebugged()))
+					{
+						const int32 Count = ModelNode->GetInstructionVisitedCount(DebuggedControlRig->GetVM(), FRigVMASTProxy(), true);
+						if(Count > Blueprint->RigGraphDisplaySettings.NodeRunLowerBound)
+						{
+							return FText::FromString(FString::FromInt(Count));
+						}
+					}
+				}
+			}
+		}
+	}
+
+	return FText();
+}
+
 #undef LOCTEXT_NAMESPACE
\ No newline at end of file
diff --git a/Engine/Plugins/Experimental/ControlRig/Source/ControlRigEditor/Private/Graph/SControlRigGraphNode.h b/Engine/Plugins/Experimental/ControlRig/Source/ControlRigEditor/Private/Graph/SControlRigGraphNode.h
--- a/Engine/Plugins/Experimental/ControlRig/Source/ControlRigEditor/Private/Graph/SControlRigGraphNode.h
+++ b/Engine/Plugins/Experimental/ControlRig/Source/ControlRigEditor/Private/Graph/SControlRigGraphNode.h
@@ -7,7 +7,9 @@
 #include "SGraphNode.h"
 #include "Widgets/Views/STreeView.h"
 #include "Widgets/Images/SImage.h"
+#include "Widgets/Text/STextBlock.h"
 #include "RigVMModel/RigVMPin.h"
+#include "ControlRigBlueprint.h"
 
 class UControlRigGraphNode;
 class STableViewBase;
@@ -84,6 +86,8 @@ private:
 
 	void HandleNodeTitleDirtied();
 
+	FText GetInstructionCountText() const;
+
 private:
 	/** Cached widget title area */
 	TSharedPtr<SOverlay> TitleAreaWidget;
@@ -112,10 +116,13 @@ private:
 	int32 NodeErrorType;
 
 	TSharedPtr<SImage> VisualDebugIndicatorWidget;
+	TSharedPtr<STextBlock> InstructionCountTextBlockWidget;
 
 	static const FSlateBrush* CachedImg_CR_Pin_Connected;
 	static const FSlateBrush* CachedImg_CR_Pin_Disconnected;
 
 	/** Cache the node title so we can invalidate it */
 	TSharedPtr<SNodeTitle> NodeTitle;
+
+	TWeakObjectPtr<UControlRigBlueprint> Blueprint;
 };
