Control Rig: Fix build break in SControlRigStackView

diff --git a/Engine/Plugins/Experimental/ControlRig/Source/ControlRigEditor/Private/Editor/SControlRigStackView.cpp b/Engine/Plugins/Experimental/ControlRig/Source/ControlRigEditor/Private/Editor/SControlRigStackView.cpp
--- a/Engine/Plugins/Experimental/ControlRig/Source/ControlRigEditor/Private/Editor/SControlRigStackView.cpp
+++ b/Engine/Plugins/Experimental/ControlRig/Source/ControlRigEditor/Private/Editor/SControlRigStackView.cpp
@@ -264,7 +264,7 @@ void SControlRigStackView::Construct( const FArguments& InArgs, TSharedRef<FCont
 				.OnGetChildren(this, &SControlRigStackView::HandleGetChildrenForTree)
 				.OnSelectionChanged(this, &SControlRigStackView::OnSelectionChanged)
 				.OnContextMenuOpening(this, &SControlRigStackView::CreateContextMenu)
-				.OnMouseButtonDoubleClick(this, &SControlRigStackView::OnMouseButtonDoubleClick)
+				.OnMouseButtonDoubleClick(this, &SControlRigStackView::HandleItemMouseDoubleClick)
 				.ItemHeight(28)
 			]
 		]
@@ -704,7 +704,7 @@ void SControlRigStackView::HandlePreviewControlRigUpdated(FControlRigEditor* InE
 {
 }
 
-void SControlRigStackView::OnMouseButtonDoubleClick(TSharedPtr<FRigStackEntry> InItem)
+void SControlRigStackView::HandleItemMouseDoubleClick(TSharedPtr<FRigStackEntry> InItem)
 {
 	if(ControlRigEditor.IsValid() && ControlRigBlueprint.IsValid())
 	{
diff --git a/Engine/Plugins/Experimental/ControlRig/Source/ControlRigEditor/Private/Editor/SControlRigStackView.h b/Engine/Plugins/Experimental/ControlRig/Source/ControlRigEditor/Private/Editor/SControlRigStackView.h
--- a/Engine/Plugins/Experimental/ControlRig/Source/ControlRigEditor/Private/Editor/SControlRigStackView.h
+++ b/Engine/Plugins/Experimental/ControlRig/Source/ControlRigEditor/Private/Editor/SControlRigStackView.h
@@ -108,7 +108,7 @@ private:
 	void HandleModifiedEvent(ERigVMGraphNotifType InNotifType, URigVMGraph* InGraph, UObject* InSubject);
 	void HandleControlRigInitializedEvent(UControlRig* InControlRig, const EControlRigState InState, const FName& InEventName);
 	void HandlePreviewControlRigUpdated(FControlRigEditor* InEditor);
-	void OnMouseButtonDoubleClick(TSharedPtr<FRigStackEntry> InItem);
+	void HandleItemMouseDoubleClick(TSharedPtr<FRigStackEntry> InItem);
 
 	/** Populate the execution stack with descriptive names for each instruction */
 	void PopulateStackView(URigVM* InVM);
