Control Rig: Follow up on CL 15746111 - improving wording in dialog

diff --git a/Engine/Plugins/Experimental/ControlRig/Source/ControlRigDeveloper/Private/Graph/ControlRigGraphSchema.cpp b/Engine/Plugins/Experimental/ControlRig/Source/ControlRigDeveloper/Private/Graph/ControlRigGraphSchema.cpp
--- a/Engine/Plugins/Experimental/ControlRig/Source/ControlRigDeveloper/Private/Graph/ControlRigGraphSchema.cpp
+++ b/Engine/Plugins/Experimental/ControlRig/Source/ControlRigDeveloper/Private/Graph/ControlRigGraphSchema.cpp
@@ -243,7 +243,7 @@ bool UControlRigGraphSchema::TryCreateConnection(UEdGraphPin* PinA, UEdGraphPin*
 
 						if(bIsInLoopIteration)
 						{
-							const EAppReturnType::Type Answer = FMessageDialog::Open( EAppMsgType::YesNo, FText::FromString( TEXT("Linking the execute within a loop iteration to the return\nof a function is not recommended.\nAre you sure?") ) );
+							const EAppReturnType::Type Answer = FMessageDialog::Open( EAppMsgType::YesNo, FText::FromString( TEXT("Linking a function return within a loop is not recommended.\nAre you sure?") ) );
 							if(Answer == EAppReturnType::No)
 							{
 								return false;
