Control Rig: Defer recreating of gizmos until afterthe user is no longer interacting. This fixes a loop where the gizmos can change when construction scripts are running.

diff --git a/Engine/Plugins/Experimental/ControlRig/Source/ControlRigEditor/Private/EditMode/ControlRigEditMode.cpp b/Engine/Plugins/Experimental/ControlRig/Source/ControlRigEditor/Private/EditMode/ControlRigEditMode.cpp
--- a/Engine/Plugins/Experimental/ControlRig/Source/ControlRigEditor/Private/EditMode/ControlRigEditMode.cpp
+++ b/Engine/Plugins/Experimental/ControlRig/Source/ControlRigEditor/Private/EditMode/ControlRigEditMode.cpp
@@ -46,6 +46,10 @@
 //#include "IPersonaPreviewScene.h"
 //#include "Animation/DebugSkelMeshComponent.h"
 //#include "Persona/Private/AnimationEditorViewportClient.h"
+#include "Framework/Application/SlateApplication.h"
+#include "UnrealEdGlobals.h"
+#include "Editor/UnrealEdEngine.h"
+
 void UControlRigEditModeDelegateHelper::OnPoseInitialized()
 {
 	if (EditMode)
@@ -269,15 +273,13 @@ void FControlRigEditMode::SetObjects_Internal()
 	if (!RuntimeControlRig)
 	{
 		DestroyGizmosActors();
+		SetUpDetailPanel();
 	}
 	else
 	{
 		// create default manipulation layer
-		RecreateGizmoActors();
-		HandleSelectionChanged();
+		RequestToRecreateGizmoActors();
 	}
-
-	SetUpDetailPanel();
 }
 
 bool FControlRigEditMode::UsesToolkits() const
@@ -400,7 +402,8 @@ void FControlRigEditMode::Tick(FEditorViewportClient* ViewportClient, float Delt
 	}
 	RecalcPivotTransform();
 
-	if (bRecreateGizmosRequired)
+	// Defer creation of gizmos if manipulating the viewport
+	if (bRecreateGizmosRequired && !(FSlateApplication::Get().HasAnyMouseCaptor() || GUnrealEd->IsUserInteracting()))
 	{
 		RecreateGizmoActors();
 
diff --git a/Engine/Plugins/Experimental/ControlRig/Source/ControlRigEditor/Private/Sequencer/ControlRigParameterTrackEditor.cpp b/Engine/Plugins/Experimental/ControlRig/Source/ControlRigEditor/Private/Sequencer/ControlRigParameterTrackEditor.cpp
--- a/Engine/Plugins/Experimental/ControlRig/Source/ControlRigEditor/Private/Sequencer/ControlRigParameterTrackEditor.cpp
+++ b/Engine/Plugins/Experimental/ControlRig/Source/ControlRigEditor/Private/Sequencer/ControlRigParameterTrackEditor.cpp
@@ -73,7 +73,8 @@
 #include "Animation/SkeletalMeshActor.h"
 #include "TimerManager.h"
 #include "BakeToControlRigSettings.h"
-
+#include "UnrealEdGlobals.h"
+#include "Editor/UnrealEdEngine.h"
 
 #define LOCTEXT_NAMESPACE "FControlRigParameterTrackEditor"
 
@@ -169,8 +170,11 @@ static USkeleton* AcquireSkeletonFromObjectGuid(const FGuid& Guid, UObject** Obj
 }
 
 FControlRigParameterTrackEditor::FControlRigParameterTrackEditor(TSharedRef<ISequencer> InSequencer)
-	: FKeyframeTrackEditor<UMovieSceneControlRigParameterTrack>(InSequencer), bCurveDisplayTickIsPending(false), bIsDoingSelection(false), bFilterAssetBySkeleton(true), bFilterAssetByAnimatableControls(true)
-
+	: FKeyframeTrackEditor<UMovieSceneControlRigParameterTrack>(InSequencer)
+	, bCurveDisplayTickIsPending(false)
+	, bIsDoingSelection(false)
+	, bFilterAssetBySkeleton(true)
+	, bFilterAssetByAnimatableControls(true)
 {
 	FMovieSceneToolsModule::Get().RegisterAnimationBakeHelper(this);
 
@@ -258,17 +262,28 @@ FControlRigParameterTrackEditor::FControlRigParameterTrackEditor(TSharedRef<ISeq
 								{
 									(*NewControlRig)->ClearControlSelection();
 								}
-								//Force refresh now, not later
-								GetSequencer()->EmptySelection();
-								//Also need to clear these guys out may cause unsure if component is selected
-								if (USelection* SelectedComponents = GEditor->GetSelectedComponents())
+													
+								auto UpdateSelectionDelegate = [this]()
 								{
-									SelectedComponents->DeselectAll();
-								}
-								if (USelection* SelectedActors = GEditor->GetSelectedActors())
-								{
-									SelectedActors->DeselectAll();
-								}
+
+									UE_LOG(LogControlRigEditor, Log, TEXT("UpdateSelectionTimer"));
+									if (!(FSlateApplication::Get().HasAnyMouseCaptor() || GUnrealEd->IsUserInteracting()))
+									{
+										UE_LOG(LogControlRigEditor, Log, TEXT("UpdateSelectionTimer - Sync"));
+										TGuardValue<bool> Guard(bIsDoingSelection, true);
+										GetSequencer()->ExternalSelectionHasChanged();
+
+										if (UpdateSelectionTimerHandle.IsValid())
+										{
+											UE_LOG(LogControlRigEditor, Log, TEXT("UpdateSelectionTimer - Clear"));
+											GEditor->GetTimerManager()->ClearTimer(UpdateSelectionTimerHandle);
+										}
+									}
+
+
+								};
+
+								GEditor->GetTimerManager()->SetTimer(UpdateSelectionTimerHandle, UpdateSelectionDelegate, 0.01f, true);
 							}
 						}
 					}
diff --git a/Engine/Plugins/Experimental/ControlRig/Source/ControlRigEditor/Private/Sequencer/ControlRigParameterTrackEditor.h b/Engine/Plugins/Experimental/ControlRig/Source/ControlRigEditor/Private/Sequencer/ControlRigParameterTrackEditor.h
--- a/Engine/Plugins/Experimental/ControlRig/Source/ControlRigEditor/Private/Sequencer/ControlRigParameterTrackEditor.h
+++ b/Engine/Plugins/Experimental/ControlRig/Source/ControlRigEditor/Private/Sequencer/ControlRigParameterTrackEditor.h
@@ -16,6 +16,7 @@
 #include "AcquiredResources.h"
 #include "MovieSceneToolHelpers.h"
 #include "MovieSceneToolsModule.h"
+#include "Engine/EngineTypes.h"
 
 struct FAssetData;
 class FMenuBuilder;
@@ -196,6 +197,8 @@ private:
 
 	/** Whether or not we should check for Animatable Controls when filtering*/
 	bool bFilterAssetByAnimatableControls;
+
+	FTimerHandle UpdateSelectionTimerHandle;
 };
 
 
