Fixed asserts in ControlRig when counting memory.

diff --git a/Engine/Plugins/Experimental/ControlRig/Source/ControlRig/Private/Rigs/RigHierarchyDefines.cpp b/Engine/Plugins/Experimental/ControlRig/Source/ControlRig/Private/Rigs/RigHierarchyDefines.cpp
--- a/Engine/Plugins/Experimental/ControlRig/Source/ControlRig/Private/Rigs/RigHierarchyDefines.cpp
+++ b/Engine/Plugins/Experimental/ControlRig/Source/ControlRig/Private/Rigs/RigHierarchyDefines.cpp
@@ -9,7 +9,7 @@
 
 void FRigElementKey::Serialize(FArchive& Ar)
 {
-	if (Ar.IsSaving() || Ar.IsObjectReferenceCollector())
+	if (Ar.IsSaving() || Ar.IsObjectReferenceCollector() || Ar.IsCountingMemory())
 	{
 		Save(Ar);
 	}
diff --git a/Engine/Plugins/Experimental/ControlRig/Source/ControlRig/Private/Rigs/RigHierarchyElements.cpp b/Engine/Plugins/Experimental/ControlRig/Source/ControlRig/Private/Rigs/RigHierarchyElements.cpp
--- a/Engine/Plugins/Experimental/ControlRig/Source/ControlRig/Private/Rigs/RigHierarchyElements.cpp
+++ b/Engine/Plugins/Experimental/ControlRig/Source/ControlRig/Private/Rigs/RigHierarchyElements.cpp
@@ -13,7 +13,7 @@ void FRigBaseElement::Serialize(FArchive& Ar, URigHierarchy* Hierarchy, ESeriali
 {
 	Ar.UsingCustomVersion(FAnimObjectVersion::GUID);
 
-	if (Ar.IsSaving() || Ar.IsObjectReferenceCollector())
+	if (Ar.IsSaving() || Ar.IsObjectReferenceCollector() || Ar.IsCountingMemory())
 	{
 		Save(Ar, Hierarchy, SerializationPhase);
 	}
