Control Rig: Ensure cache validity on load

diff --git a/Engine/Plugins/Experimental/ControlRig/Source/ControlRig/Private/Rigs/RigHierarchy.cpp b/Engine/Plugins/Experimental/ControlRig/Source/ControlRig/Private/Rigs/RigHierarchy.cpp
--- a/Engine/Plugins/Experimental/ControlRig/Source/ControlRig/Private/Rigs/RigHierarchy.cpp
+++ b/Engine/Plugins/Experimental/ControlRig/Source/ControlRig/Private/Rigs/RigHierarchy.cpp
@@ -163,7 +163,11 @@ void URigHierarchy::Load(FArchive& Ar)
 	{
 		if(FRigTransformElement* TransformElement = Cast<FRigTransformElement>(Elements[ElementIndex]))
 		{
+#if URIGHIERARCHY_RECURSIVE_DIRTY_PROPAGATION
+			TArray<FRigBaseElement*> CurrentParents = GetParents(TransformElement, false);
+#else
 			TArray<FRigBaseElement*> CurrentParents = GetParents(TransformElement, true);
+#endif
 			for (FRigBaseElement* CurrentParent : CurrentParents)
 			{
 				if(FRigTransformElement* TransformParent = Cast<FRigTransformElement>(CurrentParent))
@@ -213,6 +217,19 @@ void URigHierarchy::Load(FArchive& Ar)
 			}
 		});
 	}
+
+	// load all parts of the hierarchy pose
+	// this ensures cache validity
+	for(int32 ElementIndex = 0; ElementIndex < ElementCount; ElementIndex++)
+	{
+		if(FRigTransformElement* TransformElement = Get<FRigTransformElement>(ElementIndex))
+		{
+			GetTransform(TransformElement, ERigTransformType::InitialLocal);
+			GetTransform(TransformElement, ERigTransformType::InitialGlobal);
+			GetTransform(TransformElement, ERigTransformType::CurrentLocal);
+			GetTransform(TransformElement, ERigTransformType::CurrentGlobal);
+		}
+	}
 }
 
 void URigHierarchy::Reset()
