Control Rig: Ensure to validate local transforms for offset / gizmo transform when changing parents

diff --git a/Engine/Plugins/Experimental/ControlRig/Source/ControlRig/Private/Rigs/RigHierarchy.cpp b/Engine/Plugins/Experimental/ControlRig/Source/ControlRig/Private/Rigs/RigHierarchy.cpp
--- a/Engine/Plugins/Experimental/ControlRig/Source/ControlRig/Private/Rigs/RigHierarchy.cpp
+++ b/Engine/Plugins/Experimental/ControlRig/Source/ControlRig/Private/Rigs/RigHierarchy.cpp
@@ -55,6 +55,10 @@ void URigHierarchy::Save(FArchive& Ar)
 		}
 	}
 
+	// make sure all parts of pose are valid.
+	// this ensures cache validity.
+	ComputeAllTransforms();
+
 	int32 ElementCount = Elements.Num();
 	Ar << ElementCount;
 
@@ -217,19 +221,6 @@ void URigHierarchy::Load(FArchive& Ar)
 			}
 		});
 	}
-
-	// load all parts of the hierarchy pose
-	// this ensures cache validity
-	for(int32 ElementIndex = 0; ElementIndex < ElementCount; ElementIndex++)
-	{
-		if(FRigTransformElement* TransformElement = Get<FRigTransformElement>(ElementIndex))
-		{
-			GetTransform(TransformElement, ERigTransformType::InitialLocal);
-			GetTransform(TransformElement, ERigTransformType::InitialGlobal);
-			GetTransform(TransformElement, ERigTransformType::CurrentLocal);
-			GetTransform(TransformElement, ERigTransformType::CurrentGlobal);
-		}
-	}
 }
 
 void URigHierarchy::Reset()
@@ -2045,3 +2036,23 @@ bool URigHierarchy::ApplyTransformFromStack(const FTransformStackEntry& InEntry,
 	return false;
 }
 
+void URigHierarchy::ComputeAllTransforms()
+{
+	for(int32 ElementIndex = 0; ElementIndex < Elements.Num(); ElementIndex++)
+	{
+		for(int32 TransformTypeIndex = 0; TransformTypeIndex < (int32) ERigTransformType::NumTransformTypes; TransformTypeIndex++)
+		{
+			const ERigTransformType::Type TransformType = (ERigTransformType::Type)TransformTypeIndex; 
+			if(FRigTransformElement* TransformElement = Get<FRigTransformElement>(ElementIndex))
+			{
+				GetTransform(TransformElement, TransformType);
+			}
+			if(FRigControlElement* ControlElement = Get<FRigControlElement>(ElementIndex))
+			{
+				GetControlOffsetTransform(ControlElement, TransformType);
+				GetControlGizmoTransform(ControlElement, TransformType);
+			}
+		}
+	}
+}
+
diff --git a/Engine/Plugins/Experimental/ControlRig/Source/ControlRig/Private/Rigs/RigHierarchyController.cpp b/Engine/Plugins/Experimental/ControlRig/Source/ControlRig/Private/Rigs/RigHierarchyController.cpp
--- a/Engine/Plugins/Experimental/ControlRig/Source/ControlRig/Private/Rigs/RigHierarchyController.cpp
+++ b/Engine/Plugins/Experimental/ControlRig/Source/ControlRig/Private/Rigs/RigHierarchyController.cpp
@@ -688,6 +688,8 @@ FString URigHierarchyController::ExportToText(TArray<FRigElementKey> InKeys) con
 		return FString();
 	}
 
+	Hierarchy->ComputeAllTransforms();
+
 	// sort the keys by traversal order
 	TArray<FRigElementKey> Keys = Hierarchy->SortKeys(InKeys);
 
@@ -1328,6 +1330,14 @@ bool URigHierarchyController::AddParent(FRigBaseElement* InChild, FRigBaseElemen
 		}
 	}
 
+	if(FRigControlElement* ControlElement = Cast<FRigControlElement>(InChild))
+	{
+		Hierarchy->GetControlOffsetTransform(ControlElement, ERigTransformType::CurrentLocal);
+		Hierarchy->GetControlOffsetTransform(ControlElement, ERigTransformType::InitialLocal);
+		Hierarchy->GetControlGizmoTransform(ControlElement, ERigTransformType::CurrentLocal);
+		Hierarchy->GetControlGizmoTransform(ControlElement, ERigTransformType::InitialLocal);
+	}
+
 	if(FRigSingleParentElement* SingleParentElement = Cast<FRigSingleParentElement>(InChild))
 	{
 		if(FRigTransformElement* NewTransformParent = Cast<FRigTransformElement>(InParent))
diff --git a/Engine/Plugins/Experimental/ControlRig/Source/ControlRig/Public/Rigs/RigHierarchy.h b/Engine/Plugins/Experimental/ControlRig/Source/ControlRig/Public/Rigs/RigHierarchy.h
--- a/Engine/Plugins/Experimental/ControlRig/Source/ControlRig/Public/Rigs/RigHierarchy.h
+++ b/Engine/Plugins/Experimental/ControlRig/Source/ControlRig/Public/Rigs/RigHierarchy.h
@@ -2154,6 +2154,11 @@ private:
 	 */
 	bool ApplyTransformFromStack(const FTransformStackEntry& InEntry, bool bUndo);
 
+	/**
+	 * Computes all parts of the pose
+	 */
+	void ComputeAllTransforms();
+
 	/**
 	 * Manages merging transform actions into one during an interaction
 	 */
diff --git a/Engine/Plugins/Experimental/ControlRig/Source/ControlRig/Public/Rigs/RigHierarchyElements.h b/Engine/Plugins/Experimental/ControlRig/Source/ControlRig/Public/Rigs/RigHierarchyElements.h
--- a/Engine/Plugins/Experimental/ControlRig/Source/ControlRig/Public/Rigs/RigHierarchyElements.h
+++ b/Engine/Plugins/Experimental/ControlRig/Source/ControlRig/Public/Rigs/RigHierarchyElements.h
@@ -31,7 +31,8 @@ namespace ERigTransformType
 		InitialLocal,
         CurrentLocal,
         InitialGlobal,
-        CurrentGlobal
+        CurrentGlobal,
+		NumTransformTypes
 	};
 
 	FORCEINLINE ERigTransformType::Type SwapCurrentAndInitial(const Type InTransformType)
