Control Rig: Ensure to reset pose when baking a control rig

diff --git a/Engine/Plugins/Experimental/ControlRig/Source/ControlRig/Private/Sequencer/MovieSceneControlRigParameterSection.cpp b/Engine/Plugins/Experimental/ControlRig/Source/ControlRig/Private/Sequencer/MovieSceneControlRigParameterSection.cpp
--- a/Engine/Plugins/Experimental/ControlRig/Source/ControlRig/Private/Sequencer/MovieSceneControlRigParameterSection.cpp
+++ b/Engine/Plugins/Experimental/ControlRig/Source/ControlRig/Private/Sequencer/MovieSceneControlRigParameterSection.cpp
@@ -2106,11 +2106,16 @@ bool UMovieSceneControlRigParameterSection::LoadAnimSequenceIntoThisSection(UAni
 	const FAnimationCurveData& CurveData = DataModel->GetCurveData();
 	const TArray<FBoneAnimationTrack>& BoneAnimationTracks = DataModel->GetBoneAnimationTracks();
 
+	ControlRig->SetBoneInitialTransformsFromRefSkeleton(Skeleton->GetReferenceSkeleton());
+	ControlRig->RequestSetup();
+
 	for (int32 Index = 0; Index < NumberOfFrames; ++Index)
 	{
 		const float SequenceSecond = AnimSequence->GetTimeAtFrame(Index);
 		FFrameNumber FrameNumber = StartFrame + (FrameRateInFrameNumber * Index);
 
+		ControlRig->GetHierarchy()->ResetPoseToInitial();
+
 		for (const FFloatCurve& Curve : CurveData.FloatCurves)
 		{
 			const float Val = Curve.FloatCurve.Eval(SequenceSecond);
