Control Rig Debugging: Advance time only when Exit operation is reached

diff --git a/Engine/Plugins/Experimental/ControlRig/Source/ControlRig/Private/ControlRig.cpp b/Engine/Plugins/Experimental/ControlRig/Source/ControlRig/Private/ControlRig.cpp
--- a/Engine/Plugins/Experimental/ControlRig/Source/ControlRig/Private/ControlRig.cpp
+++ b/Engine/Plugins/Experimental/ControlRig/Source/ControlRig/Private/ControlRig.cpp
@@ -45,6 +45,7 @@ UControlRig::UControlRig(const FObjectInitializer& ObjectInitializer)
 	, AbsoluteTime(0.0f)
 	, FramesPerSecond(0.0f)
 	, bAccumulateTime(true)
+	, LatestExecutedState(EControlRigState::Invalid)
 #if WITH_EDITOR
 	, ControlRigLog(nullptr)
 	, bEnableControlRigLogging(true)
@@ -67,7 +68,7 @@ UControlRig::UControlRig(const FObjectInitializer& ObjectInitializer)
 	, InteractionBracket(0)
 	, InterRigSyncBracket(0)
 {
-	VM = ObjectInitializer.CreateDefaultSubobject<URigVM>(this, TEXT("VM"));
+	SetVM(ObjectInitializer.CreateDefaultSubobject<URigVM>(this, TEXT("VM")));
 	DynamicHierarchy = ObjectInitializer.CreateDefaultSubobject<URigHierarchy>(this, TEXT("DynamicHierarchy"));
 
 	EventQueue.Add(FRigUnit_BeginExecution::EventName);
@@ -81,6 +82,11 @@ void UControlRig::BeginDestroy()
 	PostSetupEvent.Clear();
 	ExecutedEvent.Clear();
 	SetInteractionRig(nullptr);
+
+	if (VM)
+	{
+		VM->ExecutionReachedExit().RemoveAll(this);
+	}
 }
 
 UWorld* UControlRig::GetWorld() const
@@ -406,7 +412,7 @@ void UControlRig::InstantiateVMFromCDO()
 
 	if (VM == nullptr || VM->GetOuter() != this)
 	{
-		VM = NewObject<URigVM>(this, TEXT("VM"));
+		SetVM(NewObject<URigVM>(this, TEXT("VM")));
 	}
 
 	if (!HasAnyFlags(RF_ClassDefaultObject))
@@ -430,6 +436,8 @@ void UControlRig::Execute(const EControlRigState InState, const FName& InEventNa
 {
 	DECLARE_SCOPE_HIERARCHICAL_COUNTER_FUNC()
 	QUICK_SCOPE_CYCLE_COUNTER(STAT_ControlRig_Execute);
+	
+	LatestExecutedState = InState;
 
 	if (VM)
 	{
@@ -643,14 +651,6 @@ void UControlRig::Execute(const EControlRigState InState, const FName& InEventNa
 			{
 				ExecuteUnits(Context, FRigUnit_BeginExecution::EventName);
 			}
-
-			if (InState != EControlRigState::Init)
-			{
-				if (bAccumulateTime)
-				{
-					AbsoluteTime += DeltaTime;
-				}
-			}
 		}
 		else
 		{
@@ -961,6 +961,24 @@ void UControlRig::SetEventQueue(const TArray<FName>& InEventNames)
 	EventQueue = InEventNames;
 }
 
+void UControlRig::SetVM(URigVM* NewVM)
+{
+	if (VM)
+	{
+		VM->ExecutionReachedExit().RemoveAll(this);
+	}
+	
+	if (NewVM)
+	{
+		if (!NewVM->ExecutionReachedExit().IsBoundToObject(this))
+		{
+			NewVM->ExecutionReachedExit().AddUObject(this, &UControlRig::HandleExecutionReachedExit);
+		}
+	}
+
+	VM = NewVM;
+}
+
 URigVM* UControlRig::GetVM()
 {
 	if (VM == nullptr)
@@ -1146,6 +1164,14 @@ void UControlRig::HandleOnControlModified(UControlRig* Subject, FRigControlEleme
 	}	
 }
 
+void UControlRig::HandleExecutionReachedExit()
+{
+	if (LatestExecutedState != EControlRigState::Init && bAccumulateTime)
+	{
+		AbsoluteTime += DeltaTime;
+	}
+}
+
 bool UControlRig::IsCurveControl(const FRigControlElement* InControlElement) const
 {
 	return InControlElement->Settings.bIsCurve;
diff --git a/Engine/Plugins/Experimental/ControlRig/Source/ControlRig/Private/ControlRigBlueprintGeneratedClass.cpp b/Engine/Plugins/Experimental/ControlRig/Source/ControlRig/Private/ControlRigBlueprintGeneratedClass.cpp
--- a/Engine/Plugins/Experimental/ControlRig/Source/ControlRig/Private/ControlRigBlueprintGeneratedClass.cpp
+++ b/Engine/Plugins/Experimental/ControlRig/Source/ControlRig/Private/ControlRigBlueprintGeneratedClass.cpp
@@ -51,11 +51,11 @@ void UControlRigBlueprintGeneratedClass::Serialize(FArchive& Ar)
 		{
 			if (CDO->VM == nullptr)
 			{
-				CDO->VM = NewObject<URigVM>(CDO, TEXT("VM"));
+				CDO->SetVM(NewObject<URigVM>(CDO, TEXT("VM")));
 			}
 			if (CDO->VM->GetOuter() != CDO)
 			{
-				CDO->VM = NewObject<URigVM>(CDO, TEXT("VM"));
+				CDO->SetVM(NewObject<URigVM>(CDO, TEXT("VM")));
 			}
 			CDO->VM->CopyFrom(VM);
 		}
diff --git a/Engine/Plugins/Experimental/ControlRig/Source/ControlRig/Private/ControlRigComponent.cpp b/Engine/Plugins/Experimental/ControlRig/Source/ControlRig/Private/ControlRigComponent.cpp
--- a/Engine/Plugins/Experimental/ControlRig/Source/ControlRig/Private/ControlRigComponent.cpp
+++ b/Engine/Plugins/Experimental/ControlRig/Source/ControlRig/Private/ControlRigComponent.cpp
@@ -1039,7 +1039,7 @@ UControlRig* UControlRigComponent::SetupControlRigIfRequired()
 	if(ControlRigClass)
 	{
 		ControlRig = NewObject<UControlRig>(this, ControlRigClass);
-		ControlRig->VM = NewObject<URigVM>(ControlRig);
+		ControlRig->SetVM(NewObject<URigVM>(ControlRig));
 
 		SetControlRig(ControlRig);
 
diff --git a/Engine/Plugins/Experimental/ControlRig/Source/ControlRig/Public/ControlRig.h b/Engine/Plugins/Experimental/ControlRig/Source/ControlRig/Public/ControlRig.h
--- a/Engine/Plugins/Experimental/ControlRig/Source/ControlRig/Public/ControlRig.h
+++ b/Engine/Plugins/Experimental/ControlRig/Source/ControlRig/Public/ControlRig.h
@@ -74,6 +74,9 @@ private:
 	/** true if the rig itself should increase the AbsoluteTime */
 	bool bAccumulateTime;
 
+	/** Latest state being processed */
+	EControlRigState LatestExecutedState;
+
 public:
 	virtual void Serialize(FArchive& Ar) override;
 	virtual void PostLoad() override;
@@ -152,7 +155,7 @@ public:
 	template<class T>
 	FORCEINLINE void SetPublicVariableValue(const FName& InVariableName, const T& InValue)
 	{
-		GetPublicVariableByName(InVariableName).SetValue<T>();
+		GetPublicVariableByName(InVariableName).SetValue<T>(InValue);
 	}
 
 	template<class T>
@@ -246,6 +249,9 @@ public:
 	/** Sets the queue of events to run */
 	void SetEventQueue(const TArray<FName>& InEventNames);
 
+	/** Set the VM, and bind its commands */
+	void SetVM(URigVM* VM);
+	
 	URigVM* GetVM();
 
 	/** INodeMappingInterface implementation */
@@ -409,6 +415,8 @@ private:
 	
 	void HandleOnControlModified(UControlRig* Subject, FRigControlElement* Control, const FRigControlModifiedContext& Context);
 
+	void HandleExecutionReachedExit();
+	
 	TArray<FRigVMExternalVariable> GetExternalVariablesImpl(bool bFallbackToBlueprint) const;
 
 	FORCEINLINE FProperty* GetPublicVariableProperty(const FName& InVariableName) const
diff --git a/Engine/Plugins/Experimental/ControlRig/Source/ControlRigDeveloper/Private/ControlRigBlueprint.cpp b/Engine/Plugins/Experimental/ControlRig/Source/ControlRigDeveloper/Private/ControlRigBlueprint.cpp
--- a/Engine/Plugins/Experimental/ControlRig/Source/ControlRigDeveloper/Private/ControlRigBlueprint.cpp
+++ b/Engine/Plugins/Experimental/ControlRig/Source/ControlRigDeveloper/Private/ControlRigBlueprint.cpp
@@ -567,7 +567,7 @@ void UControlRigBlueprint::RecompileVM()
 
 		if (CDO->VM->GetOuter() != CDO)
 		{
-			CDO->VM = NewObject<URigVM>(CDO, TEXT("VM"));
+			CDO->SetVM(NewObject<URigVM>(CDO, TEXT("VM")));
 		}
 
 		if (!HasAnyFlags(RF_Transient | RF_Transactional))
diff --git a/Engine/Source/Runtime/RigVM/Private/RigVMCore/RigVM.cpp b/Engine/Source/Runtime/RigVM/Private/RigVMCore/RigVM.cpp
--- a/Engine/Source/Runtime/RigVM/Private/RigVMCore/RigVM.cpp
+++ b/Engine/Source/Runtime/RigVM/Private/RigVMCore/RigVM.cpp
@@ -1339,6 +1339,7 @@ bool URigVM::Execute(FRigVMMemoryContainerPtrArray Memory, FRigVMFixedArray<void
 			}
 			case ERigVMOpCode::Exit:
 			{
+				ExecutionReachedExit().Broadcast();
 				return true;
 			}
 			case ERigVMOpCode::BeginBlock:
diff --git a/Engine/Source/Runtime/RigVM/Public/RigVMCore/RigVM.h b/Engine/Source/Runtime/RigVM/Public/RigVMCore/RigVM.h
--- a/Engine/Source/Runtime/RigVM/Public/RigVMCore/RigVM.h
+++ b/Engine/Source/Runtime/RigVM/Public/RigVMCore/RigVM.h
@@ -118,6 +118,9 @@ class RIGVM_API URigVM : public UObject
 
 public:
 
+	/** Bindable event for external objects to be notified when the VM reaches an Exit Operation */
+	DECLARE_EVENT(URigVM, FExecutionReachedExitEvent);
+	
 	URigVM();
 	virtual ~URigVM();
 
@@ -501,6 +504,8 @@ public:
 	FString GetOperandLabel(const FRigVMOperand & InOperand, TFunction<FString(const FString& RegisterName, const FString& RegisterOffsetName)> FormatFunction = nullptr) const;
 #endif
 
+	FExecutionReachedExitEvent& ExecutionReachedExit() { return OnExecutionReachedExit; }
+
 private:
 
 	void ResolveFunctionsIfRequired();
@@ -555,5 +560,7 @@ private:
 
 	void CopyDeferredVMIfRequired();
 
+	FExecutionReachedExitEvent OnExecutionReachedExit;
+	
 	friend class URigVMCompiler;
 };
