Control Rig: Added Icon Meta Name to support Per Rig Unit Icon, Added new icons for a few control rig nodes

diff --git a/Engine/Plugins/Experimental/ControlRig/Source/ControlRig/Private/Units/Execution/RigUnit_Collection.h b/Engine/Plugins/Experimental/ControlRig/Source/ControlRig/Private/Units/Execution/RigUnit_Collection.h
--- a/Engine/Plugins/Experimental/ControlRig/Source/ControlRig/Private/Units/Execution/RigUnit_Collection.h
+++ b/Engine/Plugins/Experimental/ControlRig/Source/ControlRig/Private/Units/Execution/RigUnit_Collection.h
@@ -364,7 +364,7 @@ struct CONTROLRIG_API FRigUnit_CollectionItemAtIndex : public FRigUnit_Collectio
 /**
  * Given a collection of items, execute iteratively across all items in a given collection
  */
-USTRUCT(meta=(DisplayName="For Each Item", Keywords="Collection,Loop,Iterate"))
+USTRUCT(meta=(DisplayName="For Each Item", Keywords="Collection,Loop,Iterate", Icon="EditorStyle|GraphEditor.Macro.ForEach_16x"))
 struct CONTROLRIG_API FRigUnit_CollectionLoop : public FRigUnit_CollectionBaseMutable
 {
 	GENERATED_BODY()
diff --git a/Engine/Plugins/Experimental/ControlRig/Source/ControlRig/Private/Units/Execution/RigUnit_ForLoop.h b/Engine/Plugins/Experimental/ControlRig/Source/ControlRig/Private/Units/Execution/RigUnit_ForLoop.h
--- a/Engine/Plugins/Experimental/ControlRig/Source/ControlRig/Private/Units/Execution/RigUnit_ForLoop.h
+++ b/Engine/Plugins/Experimental/ControlRig/Source/ControlRig/Private/Units/Execution/RigUnit_ForLoop.h
@@ -9,7 +9,7 @@
 /**
  * Given a count, execute iteratively until the count is up
  */
-USTRUCT(meta=(DisplayName="For Loop", Category="Execution", TitleColor="1 0 0", NodeColor="0.1 0.1 0.1", Keywords="Iterate"))
+USTRUCT(meta=(DisplayName="For Loop", Category="Execution", TitleColor="1 0 0", NodeColor="0.1 0.1 0.1", Keywords="Iterate", Icon="EditorStyle|GraphEditor.Macro.Loop_16x"))
 struct CONTROLRIG_API FRigUnit_ForLoopCount : public FRigUnitMutable
 {
 	GENERATED_BODY()
diff --git a/Engine/Plugins/Experimental/ControlRig/Source/ControlRig/Private/Units/Execution/RigUnit_SequenceExecution.h b/Engine/Plugins/Experimental/ControlRig/Source/ControlRig/Private/Units/Execution/RigUnit_SequenceExecution.h
--- a/Engine/Plugins/Experimental/ControlRig/Source/ControlRig/Private/Units/Execution/RigUnit_SequenceExecution.h
+++ b/Engine/Plugins/Experimental/ControlRig/Source/ControlRig/Private/Units/Execution/RigUnit_SequenceExecution.h
@@ -9,7 +9,7 @@
 /**
  * Allows for a single execution pulse to trigger a series of events in order.
  */
-USTRUCT(meta=(DisplayName="Sequence", Category="Execution", TitleColor="1 0 0", NodeColor="0.1 0.1 0.1"))
+USTRUCT(meta=(DisplayName="Sequence", Category="Execution", TitleColor="1 0 0", NodeColor="0.1 0.1 0.1", Icon="EditorStyle|GraphEditor.Sequence_16x"))
 struct CONTROLRIG_API FRigUnit_SequenceExecution : public FRigUnit
 {
 	GENERATED_BODY()
diff --git a/Engine/Plugins/Experimental/ControlRig/Source/ControlRigDeveloper/Private/Graph/ControlRigGraphNode.cpp b/Engine/Plugins/Experimental/ControlRig/Source/ControlRigDeveloper/Private/Graph/ControlRigGraphNode.cpp
--- a/Engine/Plugins/Experimental/ControlRig/Source/ControlRigDeveloper/Private/Graph/ControlRigGraphNode.cpp
+++ b/Engine/Plugins/Experimental/ControlRig/Source/ControlRigDeveloper/Private/Graph/ControlRigGraphNode.cpp
@@ -584,7 +584,9 @@ bool UControlRigGraphNode::ShowPaletteIconOnNode() const
 		return ModelNode->IsEvent() ||
 			ModelNode->IsA<URigVMFunctionEntryNode>() ||
 			ModelNode->IsA<URigVMFunctionReturnNode>() ||
-			ModelNode->IsA<URigVMFunctionReferenceNode>();
+			ModelNode->IsA<URigVMFunctionReferenceNode>() ||
+			ModelNode->IsA<URigVMCollapseNode>() ||
+			Cast<URigVMUnitNode>(ModelNode) != nullptr;
 	}
 	return false;
 }
@@ -595,15 +597,72 @@ FSlateIcon UControlRigGraphNode::GetIconAndTint(FLinearColor& OutColor) const
 
 	static FSlateIcon FunctionIcon("EditorStyle", "Kismet.AllClasses.FunctionIcon");
 	static FSlateIcon EventIcon("EditorStyle", "GraphEditor.Event_16x");
+	static FSlateIcon EntryReturnIcon("EditorStyle", "GraphEditor.Default_16x");
+	static FSlateIcon CollapsedNodeIcon("EditorStyle", "GraphEditor.SubGraph_16x");
 
 	if (URigVMNode* ModelNode = GetModelNode())
 	{
-		if (ModelNode->IsEvent() || 
-			ModelNode->IsA<URigVMFunctionEntryNode>() || 
-			ModelNode->IsA<URigVMFunctionReturnNode>())
+		if (ModelNode->IsEvent())
 		{
 			return EventIcon;
 		}
+
+		if (ModelNode->IsA<URigVMFunctionReferenceNode>())
+		{ 
+			return FunctionIcon;
+		}
+		
+		if (ModelNode->IsA<URigVMCollapseNode>())
+		{
+			return CollapsedNodeIcon;
+		}
+
+		if (ModelNode->IsA<URigVMFunctionEntryNode>() || 
+            ModelNode->IsA<URigVMFunctionReturnNode>())
+		{
+			return EntryReturnIcon;
+		}
+
+		if (URigVMUnitNode *UnitNode = Cast<URigVMUnitNode>(ModelNode))
+		{
+			ensure(UnitNode->GetScriptStruct());
+			
+			FString IconPath;
+
+			// icon path format: StyleSetName|StyleName|SmallStyleName|StatusOverlayStyleName
+			// the last two names are optional, see FSlateIcon() for reference
+			UnitNode->GetScriptStruct()->GetStringMetaDataHierarchical(FRigVMStruct::IconMetaName, &IconPath);
+
+			const int32 NumOfIconPathNames = 4;
+			
+			FName IconPathNames[NumOfIconPathNames] = {
+				NAME_None, // StyleSetName
+				NAME_None, // StyleName
+				NAME_None, // SmallStyleName
+				NAME_None  // StatusOverlayStyleName
+			};
+
+			int32 NameIndex = 0;
+
+			while (!IconPath.IsEmpty() && NameIndex < NumOfIconPathNames)
+			{
+				FString Left;
+				FString Right;
+
+				if (!IconPath.Split(TEXT("|"), &Left, &Right))
+				{
+					Left = IconPath;
+				}
+
+				IconPathNames[NameIndex] = FName(*Left);
+
+				NameIndex++;
+				IconPath = Right;
+			}
+
+			return FSlateIcon(IconPathNames[0], IconPathNames[1], IconPathNames[2], IconPathNames[3]);
+			
+		}
 	}
 
 	return FunctionIcon;
diff --git a/Engine/Source/Runtime/RigVM/Private/RigVMCore/RigVMStruct.cpp b/Engine/Source/Runtime/RigVM/Private/RigVMCore/RigVMStruct.cpp
--- a/Engine/Source/Runtime/RigVM/Private/RigVMCore/RigVMStruct.cpp
+++ b/Engine/Source/Runtime/RigVM/Private/RigVMCore/RigVMStruct.cpp
@@ -101,6 +101,7 @@ const FName FRigVMStruct::CustomWidgetMetaName("CustomWidget");
 const FName FRigVMStruct::ConstantMetaName("Constant");
 const FName FRigVMStruct::TitleColorMetaName("TitleColor");
 const FName FRigVMStruct::NodeColorMetaName("NodeColor");
+const FName FRigVMStruct::IconMetaName("Icon");
 const FName FRigVMStruct::KeywordsMetaName("Keywords");
 const FName FRigVMStruct::PrototypeNameMetaName("PrototypeName");
 const FName FRigVMStruct::ExpandPinByDefaultMetaName("ExpandByDefault");
diff --git a/Engine/Source/Runtime/RigVM/Public/RigVMCore/RigVMStruct.h b/Engine/Source/Runtime/RigVM/Public/RigVMCore/RigVMStruct.h
--- a/Engine/Source/Runtime/RigVM/Public/RigVMCore/RigVMStruct.h
+++ b/Engine/Source/Runtime/RigVM/Public/RigVMCore/RigVMStruct.h
@@ -184,6 +184,10 @@ public:
 	static const FName ConstantMetaName;
 	static const FName TitleColorMetaName;
 	static const FName NodeColorMetaName;
+	// icon meta name format: StyleSetName|StyleName|SmallStyleName|StatusOverlayStyleName
+	// the last two names are optional, see FSlateIcon() for reference
+	// Example: Icon="EditorStyle|GraphEditor.Sequence_16x"
+	static const FName IconMetaName;
 	static const FName KeywordsMetaName;
 	static const FName PrototypeNameMetaName;
 	static const FName ExpandPinByDefaultMetaName;
