Control Rig: Changing node header colors for better UX

diff --git a/Engine/Plugins/Experimental/ControlRig/Source/ControlRig/Private/Units/Core/RigUnit_Name.h b/Engine/Plugins/Experimental/ControlRig/Source/ControlRig/Private/Units/Core/RigUnit_Name.h
--- a/Engine/Plugins/Experimental/ControlRig/Source/ControlRig/Private/Units/Core/RigUnit_Name.h
+++ b/Engine/Plugins/Experimental/ControlRig/Source/ControlRig/Private/Units/Core/RigUnit_Name.h
@@ -5,7 +5,7 @@
 #include "Units/RigUnit.h"
 #include "RigUnit_Name.generated.h"
 
-USTRUCT(meta=(Abstract, Category="Core|Name", NodeColor = "0.7 0.05 0.5"))
+USTRUCT(meta=(Abstract, Category="Core|Name", NodeColor = "0.462745, 1,0, 0.329412"))
 struct CONTROLRIG_API FRigUnit_NameBase : public FRigUnit
 {
 	GENERATED_BODY()
diff --git a/Engine/Plugins/Experimental/ControlRig/Source/ControlRig/Private/Units/Execution/RigUnit_BeginExecution.h b/Engine/Plugins/Experimental/ControlRig/Source/ControlRig/Private/Units/Execution/RigUnit_BeginExecution.h
--- a/Engine/Plugins/Experimental/ControlRig/Source/ControlRig/Private/Units/Execution/RigUnit_BeginExecution.h
+++ b/Engine/Plugins/Experimental/ControlRig/Source/ControlRig/Private/Units/Execution/RigUnit_BeginExecution.h
@@ -9,7 +9,7 @@
 /**
  * Event for driving the skeleton hierarchy with variables and rig elements
  */
-USTRUCT(meta=(DisplayName="Forwards Solve", Category="Execution", TitleColor="1 0 0", NodeColor="0.1 0.1 0.1", Keywords="Begin,Update,Tick,Forward,Event"))
+USTRUCT(meta=(DisplayName="Forwards Solve", Category="Execution", NodeColor="1, 0, 0", Keywords="Begin,Update,Tick,Forward,Event"))
 struct CONTROLRIG_API FRigUnit_BeginExecution : public FRigUnit
 {
 	GENERATED_BODY()
diff --git a/Engine/Plugins/Experimental/ControlRig/Source/ControlRig/Private/Units/Execution/RigUnit_Hierarchy.h b/Engine/Plugins/Experimental/ControlRig/Source/ControlRig/Private/Units/Execution/RigUnit_Hierarchy.h
--- a/Engine/Plugins/Experimental/ControlRig/Source/ControlRig/Private/Units/Execution/RigUnit_Hierarchy.h
+++ b/Engine/Plugins/Experimental/ControlRig/Source/ControlRig/Private/Units/Execution/RigUnit_Hierarchy.h
@@ -6,7 +6,7 @@
 #include "ControlRigDefines.h" 
 #include "RigUnit_Hierarchy.generated.h"
 
-USTRUCT(meta = (Abstract, NodeColor = "0.7 0.05 0.5", Category = "Hierarchy"))
+USTRUCT(meta = (Abstract, NodeColor="0.462745, 1,0, 0.329412", Category = "Hierarchy"))
 struct CONTROLRIG_API FRigUnit_HierarchyBase : public FRigUnit
 {
 	GENERATED_BODY()
diff --git a/Engine/Plugins/Experimental/ControlRig/Source/ControlRig/Private/Units/Execution/RigUnit_InverseExecution.h b/Engine/Plugins/Experimental/ControlRig/Source/ControlRig/Private/Units/Execution/RigUnit_InverseExecution.h
--- a/Engine/Plugins/Experimental/ControlRig/Source/ControlRig/Private/Units/Execution/RigUnit_InverseExecution.h
+++ b/Engine/Plugins/Experimental/ControlRig/Source/ControlRig/Private/Units/Execution/RigUnit_InverseExecution.h
@@ -9,7 +9,7 @@
 /**
  * Event for driving elements based off the skeleton hierarchy
  */
-USTRUCT(meta=(DisplayName="Backwards Solve", Category="Inverse", TitleColor="1 0 0", NodeColor="0.1 0.1 0.1", Keywords="Inverse,Reverse,Backwards,Event"))
+USTRUCT(meta=(DisplayName="Backwards Solve", Category="Inverse", NodeColor="1, 0, 0", Keywords="Inverse,Reverse,Backwards,Event"))
 struct CONTROLRIG_API FRigUnit_InverseExecution : public FRigUnit
 {
 	GENERATED_BODY()
diff --git a/Engine/Plugins/Experimental/ControlRig/Source/ControlRig/Private/Units/Execution/RigUnit_PrepareForExecution.h b/Engine/Plugins/Experimental/ControlRig/Source/ControlRig/Private/Units/Execution/RigUnit_PrepareForExecution.h
--- a/Engine/Plugins/Experimental/ControlRig/Source/ControlRig/Private/Units/Execution/RigUnit_PrepareForExecution.h
+++ b/Engine/Plugins/Experimental/ControlRig/Source/ControlRig/Private/Units/Execution/RigUnit_PrepareForExecution.h
@@ -9,7 +9,7 @@
 /**
  * Event to setup elements before all solves
  */
-USTRUCT(meta=(DisplayName="Setup Event", Category="Setup", TitleColor="1 0 0", NodeColor="0.1 0.1 0.1", Keywords="Setup,Init,Fit"))
+USTRUCT(meta=(DisplayName="Setup Event", Category="Setup", NodeColor="1, 0, 0", Keywords="Setup,Init,Fit"))
 struct CONTROLRIG_API FRigUnit_PrepareForExecution : public FRigUnit
 {
 	GENERATED_BODY()
diff --git a/Engine/Plugins/Experimental/ControlRig/Source/ControlRig/Private/Units/Hierarchy/RigUnit_GetControlTransform.h b/Engine/Plugins/Experimental/ControlRig/Source/ControlRig/Private/Units/Hierarchy/RigUnit_GetControlTransform.h
--- a/Engine/Plugins/Experimental/ControlRig/Source/ControlRig/Private/Units/Hierarchy/RigUnit_GetControlTransform.h
+++ b/Engine/Plugins/Experimental/ControlRig/Source/ControlRig/Private/Units/Hierarchy/RigUnit_GetControlTransform.h
@@ -8,7 +8,7 @@
 /**
  * GetControlBool is used to retrieve a single Bool from a hierarchy.
  */
-USTRUCT(meta=(DisplayName="Get Control Bool", Category="Controls", DocumentationPolicy = "Strict", Keywords="GetControlBool", PrototypeName = "GetControlValue", Varying))
+USTRUCT(meta=(DisplayName="Get Control Bool", Category="Controls", DocumentationPolicy = "Strict", Keywords="GetControlBool", PrototypeName = "GetControlValue", NodeColor="0.462745, 1,0, 0.329412",Varying))
 struct CONTROLRIG_API FRigUnit_GetControlBool : public FRigUnit
 {
 	GENERATED_BODY()
@@ -39,7 +39,7 @@ struct CONTROLRIG_API FRigUnit_GetControlBool : public FRigUnit
 /**
  * GetControlFloat is used to retrieve a single Float from a hierarchy.
  */
-USTRUCT(meta=(DisplayName="Get Control Float", Category="Controls", DocumentationPolicy = "Strict", Keywords="GetControlFloat", PrototypeName = "GetControlValue", Varying))
+USTRUCT(meta=(DisplayName="Get Control Float", Category="Controls", DocumentationPolicy = "Strict", Keywords="GetControlFloat", PrototypeName = "GetControlValue", NodeColor="0.462745, 1,0, 0.329412", Varying))
 struct CONTROLRIG_API FRigUnit_GetControlFloat : public FRigUnit
 {
 	GENERATED_BODY()
@@ -80,7 +80,7 @@ struct CONTROLRIG_API FRigUnit_GetControlFloat : public FRigUnit
 /**
  * GetControlFloat is used to retrieve a single Integer from a hierarchy.
  */
-USTRUCT(meta=(DisplayName="Get Control Integer", Category="Controls", DocumentationPolicy = "Strict", Keywords="GetControlInteger,GetControlEnum", PrototypeName = "GetControlValue", Varying))
+USTRUCT(meta=(DisplayName="Get Control Integer", Category="Controls", DocumentationPolicy = "Strict", Keywords="GetControlInteger,GetControlEnum", PrototypeName = "GetControlValue", NodeColor="0.462745, 1,0, 0.329412", Varying))
 struct CONTROLRIG_API FRigUnit_GetControlInteger : public FRigUnit
 {
 	GENERATED_BODY()
@@ -121,7 +121,7 @@ struct CONTROLRIG_API FRigUnit_GetControlInteger : public FRigUnit
 /**
  * GetControlVector2D is used to retrieve a single Vector2D from a hierarchy.
  */
-USTRUCT(meta=(DisplayName="Get Control Vector2D", Category="Controls", DocumentationPolicy = "Strict", Keywords="GetControlVector2D", PrototypeName = "GetControlValue", Varying))
+USTRUCT(meta=(DisplayName="Get Control Vector2D", Category="Controls", DocumentationPolicy = "Strict", Keywords="GetControlVector2D", PrototypeName = "GetControlValue", NodeColor="0.462745, 1,0, 0.329412", Varying))
 struct CONTROLRIG_API FRigUnit_GetControlVector2D : public FRigUnit
 {
 	GENERATED_BODY()
@@ -162,7 +162,7 @@ struct CONTROLRIG_API FRigUnit_GetControlVector2D : public FRigUnit
 /**
  * GetControlVector is used to retrieve a single Vector from a hierarchy, can be used for Controls of type "Position" or "Scale".
  */
-USTRUCT(meta=(DisplayName="Get Control Vector", Category="Controls", DocumentationPolicy = "Strict", Keywords="GetControlVector, GetControlPosition, GetControlScale", PrototypeName = "GetControlValue", Varying))
+USTRUCT(meta=(DisplayName="Get Control Vector", Category="Controls", DocumentationPolicy = "Strict", Keywords="GetControlVector, GetControlPosition, GetControlScale", PrototypeName = "GetControlValue", NodeColor="0.462745, 1,0, 0.329412", Varying))
 struct CONTROLRIG_API FRigUnit_GetControlVector : public FRigUnit
 {
 	GENERATED_BODY()
@@ -211,7 +211,7 @@ struct CONTROLRIG_API FRigUnit_GetControlVector : public FRigUnit
 /**
  * GetControlRotator is used to retrieve a single Rotator from a hierarchy.
  */
-USTRUCT(meta=(DisplayName="Get Control Rotator", Category="Controls", DocumentationPolicy = "Strict", Keywords="GetControlRotator", PrototypeName = "GetControlValue", Varying))
+USTRUCT(meta=(DisplayName="Get Control Rotator", Category="Controls", DocumentationPolicy = "Strict", Keywords="GetControlRotator", PrototypeName = "GetControlValue", NodeColor="0.462745, 1,0, 0.329412", Varying))
 struct CONTROLRIG_API FRigUnit_GetControlRotator : public FRigUnit
 {
 	GENERATED_BODY()
@@ -261,7 +261,7 @@ struct CONTROLRIG_API FRigUnit_GetControlRotator : public FRigUnit
 /**
  * GetControlTransform is used to retrieve a single transform from a hierarchy.
  */
-USTRUCT(meta=(DisplayName="Get Control Transform", Category="Controls", DocumentationPolicy = "Strict", Keywords="GetControlTransform", PrototypeName = "GetControlValue", Varying))
+USTRUCT(meta=(DisplayName="Get Control Transform", Category="Controls", DocumentationPolicy = "Strict", Keywords="GetControlTransform", PrototypeName = "GetControlValue", NodeColor="0.462745, 1,0, 0.329412", Varying))
 struct CONTROLRIG_API FRigUnit_GetControlTransform : public FRigUnit
 {
 	GENERATED_BODY()
diff --git a/Engine/Plugins/Experimental/ControlRig/Source/ControlRig/Private/Units/Hierarchy/RigUnit_GetCurveValue.h b/Engine/Plugins/Experimental/ControlRig/Source/ControlRig/Private/Units/Hierarchy/RigUnit_GetCurveValue.h
--- a/Engine/Plugins/Experimental/ControlRig/Source/ControlRig/Private/Units/Hierarchy/RigUnit_GetCurveValue.h
+++ b/Engine/Plugins/Experimental/ControlRig/Source/ControlRig/Private/Units/Hierarchy/RigUnit_GetCurveValue.h
@@ -8,7 +8,7 @@
 /**
  * GetCurveValue is used to retrieve a single float from a Curve
  */
-USTRUCT(meta=(DisplayName="Get Curve Value", Category="Curve", Keywords="GetCurveValue,float", Varying))
+USTRUCT(meta=(DisplayName="Get Curve Value", Category="Curve", Keywords="GetCurveValue,float", NodeColor="0.462745, 1,0, 0.329412", Varying))
 struct CONTROLRIG_API FRigUnit_GetCurveValue : public FRigUnit
 {
 	GENERATED_BODY()
diff --git a/Engine/Plugins/Experimental/ControlRig/Source/ControlRig/Private/Units/Hierarchy/RigUnit_GetRelativeTransform.h b/Engine/Plugins/Experimental/ControlRig/Source/ControlRig/Private/Units/Hierarchy/RigUnit_GetRelativeTransform.h
--- a/Engine/Plugins/Experimental/ControlRig/Source/ControlRig/Private/Units/Hierarchy/RigUnit_GetRelativeTransform.h
+++ b/Engine/Plugins/Experimental/ControlRig/Source/ControlRig/Private/Units/Hierarchy/RigUnit_GetRelativeTransform.h
@@ -8,7 +8,7 @@
 /**
  * GetRelativeTransform is used to retrieve a single transform from a hierarchy in the space of another transform
  */
-USTRUCT(meta=(DisplayName="Get Relative Transform", Category="Hierarchy", DocumentationPolicy = "Strict", Keywords = "Offset,Local", Varying))
+USTRUCT(meta=(DisplayName="Get Relative Transform", Category="Hierarchy", DocumentationPolicy = "Strict", Keywords = "Offset,Local", NodeColor="0.462745, 1,0, 0.329412", Varying))
 struct CONTROLRIG_API FRigUnit_GetRelativeTransformForItem : public FRigUnit
 {
 	GENERATED_BODY()
diff --git a/Engine/Plugins/Experimental/ControlRig/Source/ControlRig/Private/Units/Hierarchy/RigUnit_GetTransform.h b/Engine/Plugins/Experimental/ControlRig/Source/ControlRig/Private/Units/Hierarchy/RigUnit_GetTransform.h
--- a/Engine/Plugins/Experimental/ControlRig/Source/ControlRig/Private/Units/Hierarchy/RigUnit_GetTransform.h
+++ b/Engine/Plugins/Experimental/ControlRig/Source/ControlRig/Private/Units/Hierarchy/RigUnit_GetTransform.h
@@ -8,7 +8,7 @@
 /**
  * GetTransform is used to retrieve a single transform from a hierarchy.
  */
-USTRUCT(meta=(DisplayName="Get Transform", Category="Hierarchy", DocumentationPolicy = "Strict", Keywords="GetBoneTransform,GetControlTransform,GetInitialTransform,GetSpaceTransform,GetTransform", Varying))
+USTRUCT(meta=(DisplayName="Get Transform", Category="Hierarchy", DocumentationPolicy = "Strict", Keywords="GetBoneTransform,GetControlTransform,GetInitialTransform,GetSpaceTransform,GetTransform", NodeColor="0.462745, 1,0, 0.329412",Varying))
 struct CONTROLRIG_API FRigUnit_GetTransform : public FRigUnit
 {
 	GENERATED_BODY()
diff --git a/Engine/Plugins/Experimental/ControlRig/Source/ControlRig/Private/Units/Hierarchy/RigUnit_OffsetTransform.h b/Engine/Plugins/Experimental/ControlRig/Source/ControlRig/Private/Units/Hierarchy/RigUnit_OffsetTransform.h
--- a/Engine/Plugins/Experimental/ControlRig/Source/ControlRig/Private/Units/Hierarchy/RigUnit_OffsetTransform.h
+++ b/Engine/Plugins/Experimental/ControlRig/Source/ControlRig/Private/Units/Hierarchy/RigUnit_OffsetTransform.h
@@ -8,7 +8,7 @@
 /**
  * Offset Transform is used to add an offset to an existing transform in the hierarchy. The offset is post multiplied.
  */
-USTRUCT(meta=(DisplayName="Offset Transform", Category="Hierarchy", DocumentationPolicy = "Strict", Keywords = "Offset,Relative,AddBoneTransform", Varying))
+USTRUCT(meta=(DisplayName="Offset Transform", Category="Hierarchy", DocumentationPolicy = "Strict", Keywords = "Offset,Relative,AddBoneTransform", NodeColor="0, 0.364706, 1.0", Varying))
 struct CONTROLRIG_API FRigUnit_OffsetTransformForItem : public FRigUnitMutable
 {
 	GENERATED_BODY()
diff --git a/Engine/Plugins/Experimental/ControlRig/Source/ControlRig/Private/Units/Hierarchy/RigUnit_ProjectTransformToNewParent.h b/Engine/Plugins/Experimental/ControlRig/Source/ControlRig/Private/Units/Hierarchy/RigUnit_ProjectTransformToNewParent.h
--- a/Engine/Plugins/Experimental/ControlRig/Source/ControlRig/Private/Units/Hierarchy/RigUnit_ProjectTransformToNewParent.h
+++ b/Engine/Plugins/Experimental/ControlRig/Source/ControlRig/Private/Units/Hierarchy/RigUnit_ProjectTransformToNewParent.h
@@ -13,7 +13,7 @@
  * the pelvis (initial), set the old parent to be the head (initial)
  * and set the new parent to the head as well (not-initial == current).
  */
-USTRUCT(meta=(DisplayName="Project to new Parent", Category="Hierarchy", DocumentationPolicy = "Strict", Keywords = "ProjectTransformToNewParent,Relative,Reparent,Offset", Varying))
+USTRUCT(meta=(DisplayName="Project to new Parent", Category="Hierarchy", DocumentationPolicy = "Strict", Keywords = "ProjectTransformToNewParent,Relative,Reparent,Offset", NodeColor="0.462745, 1,0, 0.329412", Varying))
 struct CONTROLRIG_API FRigUnit_ProjectTransformToNewParent : public FRigUnit
 {
 	GENERATED_BODY()
diff --git a/Engine/Plugins/Experimental/ControlRig/Source/ControlRig/Private/Units/Hierarchy/RigUnit_SendEvent.h b/Engine/Plugins/Experimental/ControlRig/Source/ControlRig/Private/Units/Hierarchy/RigUnit_SendEvent.h
--- a/Engine/Plugins/Experimental/ControlRig/Source/ControlRig/Private/Units/Hierarchy/RigUnit_SendEvent.h
+++ b/Engine/Plugins/Experimental/ControlRig/Source/ControlRig/Private/Units/Hierarchy/RigUnit_SendEvent.h
@@ -8,7 +8,7 @@
 /**
  * SendEvent is used to notify the engine / editor of a change that happend within the Control Rig.
  */
-USTRUCT(meta=(DisplayName="Send Event", Category="Hierarchy", DocumentationPolicy="Strict", Keywords = "SendEvent", PrototypeName = "Event,Notify,Notification"))
+USTRUCT(meta=(DisplayName="Send Event", Category="Hierarchy", DocumentationPolicy="Strict", Keywords = "SendEvent", PrototypeName = "Event,Notify,Notification", NodeColor="1, 0, 0"))
 struct CONTROLRIG_API FRigUnit_SendEvent : public FRigUnitMutable
 {
 	GENERATED_BODY()
diff --git a/Engine/Plugins/Experimental/ControlRig/Source/ControlRig/Private/Units/Hierarchy/RigUnit_SetControlOffset.h b/Engine/Plugins/Experimental/ControlRig/Source/ControlRig/Private/Units/Hierarchy/RigUnit_SetControlOffset.h
--- a/Engine/Plugins/Experimental/ControlRig/Source/ControlRig/Private/Units/Hierarchy/RigUnit_SetControlOffset.h
+++ b/Engine/Plugins/Experimental/ControlRig/Source/ControlRig/Private/Units/Hierarchy/RigUnit_SetControlOffset.h
@@ -8,7 +8,7 @@
 /**
  * SetControlOffset is used to perform a change in the hierarchy by setting a single control's transform.
  */
-USTRUCT(meta=(DisplayName="Set Control Offset", Category="Setup", DocumentationPolicy="Strict", Keywords = "SetControlOffset,Initial,InitialTransform,SetInitialTransform,SetInitialControlTransform", PrototypeName = "SetControlValue"))
+USTRUCT(meta=(DisplayName="Set Control Offset", Category="Setup", DocumentationPolicy="Strict", Keywords = "SetControlOffset,Initial,InitialTransform,SetInitialTransform,SetInitialControlTransform", PrototypeName = "SetControlValue", NodeColor="0, 0.364706, 1.0"))
 struct CONTROLRIG_API FRigUnit_SetControlOffset : public FRigUnitMutable
 {
 	GENERATED_BODY()
diff --git a/Engine/Plugins/Experimental/ControlRig/Source/ControlRig/Private/Units/Hierarchy/RigUnit_SetControlVisibility.h b/Engine/Plugins/Experimental/ControlRig/Source/ControlRig/Private/Units/Hierarchy/RigUnit_SetControlVisibility.h
--- a/Engine/Plugins/Experimental/ControlRig/Source/ControlRig/Private/Units/Hierarchy/RigUnit_SetControlVisibility.h
+++ b/Engine/Plugins/Experimental/ControlRig/Source/ControlRig/Private/Units/Hierarchy/RigUnit_SetControlVisibility.h
@@ -8,7 +8,7 @@
 /**
  * SetControlVisibility is used to change the gizmo visibility on a control at runtime
  */
-USTRUCT(meta=(DisplayName="Set Control Visibility", Category="Hierarchy", DocumentationPolicy="Strict", Keywords = "SetControlVisibility,Visibility,Hide,Show,Hidden,Visible,SetGizmoVisibility", PrototypeName = "SetControlVisibility"))
+USTRUCT(meta=(DisplayName="Set Control Visibility", Category="Hierarchy", DocumentationPolicy="Strict", Keywords = "SetControlVisibility,Visibility,Hide,Show,Hidden,Visible,SetGizmoVisibility", PrototypeName = "SetControlVisibility", NodeColor="0, 0.364706, 1.0"))
 struct CONTROLRIG_API FRigUnit_SetControlVisibility : public FRigUnitMutable
 {
 	GENERATED_BODY()
diff --git a/Engine/Plugins/Experimental/ControlRig/Source/ControlRig/Private/Units/Hierarchy/RigUnit_SetRelativeTransform.h b/Engine/Plugins/Experimental/ControlRig/Source/ControlRig/Private/Units/Hierarchy/RigUnit_SetRelativeTransform.h
--- a/Engine/Plugins/Experimental/ControlRig/Source/ControlRig/Private/Units/Hierarchy/RigUnit_SetRelativeTransform.h
+++ b/Engine/Plugins/Experimental/ControlRig/Source/ControlRig/Private/Units/Hierarchy/RigUnit_SetRelativeTransform.h
@@ -8,7 +8,7 @@
 /**
  * SetRelativeTransform is used to set a single transform from a hierarchy in the space of another item
  */
-USTRUCT(meta=(DisplayName="Set Relative Transform", Category="Hierarchy", DocumentationPolicy = "Strict", Keywords = "Offset,Local", Varying))
+USTRUCT(meta=(DisplayName="Set Relative Transform", Category="Hierarchy", DocumentationPolicy = "Strict", Keywords = "Offset,Local", NodeColor="0, 0.364706, 1.0", Varying))
 struct CONTROLRIG_API FRigUnit_SetRelativeTransformForItem : public FRigUnitMutable
 {
 	GENERATED_BODY()
diff --git a/Engine/Plugins/Experimental/ControlRig/Source/ControlRig/Private/Units/Hierarchy/RigUnit_WorldSpace.h b/Engine/Plugins/Experimental/ControlRig/Source/ControlRig/Private/Units/Hierarchy/RigUnit_WorldSpace.h
--- a/Engine/Plugins/Experimental/ControlRig/Source/ControlRig/Private/Units/Hierarchy/RigUnit_WorldSpace.h
+++ b/Engine/Plugins/Experimental/ControlRig/Source/ControlRig/Private/Units/Hierarchy/RigUnit_WorldSpace.h
@@ -8,7 +8,7 @@
 /**
  * Converts a transform from rig (global) space to world space
  */
-USTRUCT(meta=(DisplayName="To World", MenuDescSuffix="(Transform)", Category="Hierarchy", DocumentationPolicy = "Strict", Keywords="Global,Local,World,Actor,ComponentSpace,FromRig", Varying))
+USTRUCT(meta=(DisplayName="To World", MenuDescSuffix="(Transform)", Category="Hierarchy", DocumentationPolicy = "Strict", Keywords="Global,Local,World,Actor,ComponentSpace,FromRig", NodeColor="0.462745, 1,0, 0.329412", Varying))
 struct CONTROLRIG_API FRigUnit_ToWorldSpace_Transform : public FRigUnit
 {
 	GENERATED_BODY()
@@ -37,7 +37,7 @@ struct CONTROLRIG_API FRigUnit_ToWorldSpace_Transform : public FRigUnit
 /**
  * Converts a transform from world space to rig (global) space
  */
-USTRUCT(meta=(DisplayName="From World", MenuDescSuffix="(Transform)", Category="Hierarchy", DocumentationPolicy = "Strict", Keywords="Global,Local,World,Actor,ComponentSpace,ToRig", Varying))
+USTRUCT(meta=(DisplayName="From World", MenuDescSuffix="(Transform)", Category="Hierarchy", DocumentationPolicy = "Strict", Keywords="Global,Local,World,Actor,ComponentSpace,ToRig", NodeColor="0.462745, 1,0, 0.329412", Varying))
 struct CONTROLRIG_API FRigUnit_ToRigSpace_Transform : public FRigUnit
 {
 	GENERATED_BODY()
@@ -66,7 +66,7 @@ struct CONTROLRIG_API FRigUnit_ToRigSpace_Transform : public FRigUnit
 /**
  * Converts a position / location from rig (global) space to world space
  */
-USTRUCT(meta=(DisplayName="To World", MenuDescSuffix="(Vector)", Category="Hierarchy", DocumentationPolicy = "Strict", Keywords="Global,Local,World,Actor,ComponentSpace,FromRig", Varying))
+USTRUCT(meta=(DisplayName="To World", MenuDescSuffix="(Vector)", Category="Hierarchy", DocumentationPolicy = "Strict", Keywords="Global,Local,World,Actor,ComponentSpace,FromRig", NodeColor="0.462745, 1,0, 0.329412", Varying))
 struct CONTROLRIG_API FRigUnit_ToWorldSpace_Location : public FRigUnit
 {
 	GENERATED_BODY()
@@ -95,7 +95,7 @@ struct CONTROLRIG_API FRigUnit_ToWorldSpace_Location : public FRigUnit
 /**
  * Converts a position / location from world space to rig (global) space
  */
-USTRUCT(meta=(DisplayName="From World", MenuDescSuffix="(Vector)", Category="Hierarchy", DocumentationPolicy = "Strict", Keywords="Global,Local,World,Actor,ComponentSpace,ToRig", Varying))
+USTRUCT(meta=(DisplayName="From World", MenuDescSuffix="(Vector)", Category="Hierarchy", DocumentationPolicy = "Strict", Keywords="Global,Local,World,Actor,ComponentSpace,ToRig", NodeColor="0.462745, 1,0, 0.329412", Varying))
 struct CONTROLRIG_API FRigUnit_ToRigSpace_Location : public FRigUnit
 {
 	GENERATED_BODY()
@@ -124,7 +124,7 @@ struct CONTROLRIG_API FRigUnit_ToRigSpace_Location : public FRigUnit
 /**
  * Converts a rotation from rig (global) space to world space
  */
-USTRUCT(meta=(DisplayName="To World", MenuDescSuffix="(Quat)", Category="Hierarchy", DocumentationPolicy = "Strict", Keywords="Global,Local,World,Actor,ComponentSpace,FromRig", Varying))
+USTRUCT(meta=(DisplayName="To World", MenuDescSuffix="(Quat)", Category="Hierarchy", DocumentationPolicy = "Strict", Keywords="Global,Local,World,Actor,ComponentSpace,FromRig", NodeColor="0.462745, 1,0, 0.329412", Varying))
 struct CONTROLRIG_API FRigUnit_ToWorldSpace_Rotation : public FRigUnit
 {
 	GENERATED_BODY()
@@ -153,7 +153,7 @@ struct CONTROLRIG_API FRigUnit_ToWorldSpace_Rotation : public FRigUnit
 /**
  * Converts a rotation from world space to rig (global) space
  */
-USTRUCT(meta=(DisplayName="From World", MenuDescSuffix="(Quat)", Category="Hierarchy", DocumentationPolicy = "Strict", Keywords="Global,Local,World,Actor,ComponentSpace,ToRig", Varying))
+USTRUCT(meta=(DisplayName="From World", MenuDescSuffix="(Quat)", Category="Hierarchy", DocumentationPolicy = "Strict", Keywords="Global,Local,World,Actor,ComponentSpace,ToRig", NodeColor="0.462745, 1,0, 0.329412", Varying))
 struct CONTROLRIG_API FRigUnit_ToRigSpace_Rotation : public FRigUnit
 {
 	GENERATED_BODY()
diff --git a/Engine/Plugins/Experimental/ControlRig/Source/ControlRig/Public/Units/Hierarchy/RigUnit_SetTransform.h b/Engine/Plugins/Experimental/ControlRig/Source/ControlRig/Public/Units/Hierarchy/RigUnit_SetTransform.h
--- a/Engine/Plugins/Experimental/ControlRig/Source/ControlRig/Public/Units/Hierarchy/RigUnit_SetTransform.h
+++ b/Engine/Plugins/Experimental/ControlRig/Source/ControlRig/Public/Units/Hierarchy/RigUnit_SetTransform.h
@@ -12,7 +12,7 @@
  * actually sets the Control's offset transform and resets the local
  * values to (0, 0, 0).
  */
-USTRUCT(meta=(DisplayName="Set Transform", Category="Hierarchy", DocumentationPolicy = "Strict", Keywords="SetBoneTransform,SetControlTransform,SetInitialTransform,SetSpaceTransform", Varying))
+USTRUCT(meta=(DisplayName="Set Transform", Category="Hierarchy", DocumentationPolicy = "Strict", Keywords="SetBoneTransform,SetControlTransform,SetInitialTransform,SetSpaceTransform", NodeColor="0, 0.364706, 1.0", Varying))
 struct CONTROLRIG_API FRigUnit_SetTransform : public FRigUnitMutable
 {
 	GENERATED_BODY()
@@ -85,7 +85,7 @@ struct CONTROLRIG_API FRigUnit_SetTransform : public FRigUnitMutable
 /**
  * SetTranslation is used to set a single translation on hierarchy.
  */
-USTRUCT(meta=(DisplayName="Set Translation", Category="Hierarchy", DocumentationPolicy = "Strict", Keywords="SetBoneTranslation,SetControlTranslation,SetInitialTranslation,SetSpaceTranslation,SetBoneLocation,SetControlLocation,SetInitialLocation,SetSpaceLocation,SetBonePosition,SetControlPosition,SetInitialPosition,SetSpacePosition,SetTranslation,SetLocation,SetPosition", Varying))
+USTRUCT(meta=(DisplayName="Set Translation", Category="Hierarchy", DocumentationPolicy = "Strict", Keywords="SetBoneTranslation,SetControlTranslation,SetInitialTranslation,SetSpaceTranslation,SetBoneLocation,SetControlLocation,SetInitialLocation,SetSpaceLocation,SetBonePosition,SetControlPosition,SetInitialPosition,SetSpacePosition,SetTranslation,SetLocation,SetPosition", NodeColor="0, 0.364706, 1.0", Varying))
 struct CONTROLRIG_API FRigUnit_SetTranslation : public FRigUnitMutable
 {
 	GENERATED_BODY()
@@ -149,7 +149,7 @@ struct CONTROLRIG_API FRigUnit_SetTranslation : public FRigUnitMutable
 /**
  * SetRotation is used to set a single rotation on hierarchy.
  */
-USTRUCT(meta=(DisplayName="Set Rotation", Category="Hierarchy", DocumentationPolicy = "Strict", Keywords="SetBoneRotation,SetControlRotation,SetInitialRotation,SetSpaceRotation,SetBoneOrientation,SetControlOrientation,SetInitialOrientation,SetSpaceOrientation,SetRotation,SetOrientation", Varying))
+USTRUCT(meta=(DisplayName="Set Rotation", Category="Hierarchy", DocumentationPolicy = "Strict", Keywords="SetBoneRotation,SetControlRotation,SetInitialRotation,SetSpaceRotation,SetBoneOrientation,SetControlOrientation,SetInitialOrientation,SetSpaceOrientation,SetRotation,SetOrientation", NodeColor="0, 0.364706, 1.0", Varying))
 struct CONTROLRIG_API FRigUnit_SetRotation : public FRigUnitMutable
 {
 	GENERATED_BODY()
@@ -213,7 +213,7 @@ struct CONTROLRIG_API FRigUnit_SetRotation : public FRigUnitMutable
 /**
  * SetScale is used to set a single scale on hierarchy.
  */
-USTRUCT(meta=(DisplayName="Set Scale", Category="Hierarchy", DocumentationPolicy = "Strict", Keywords="SetBoneScale,SetControlScale,SetInitialScale,SetSpaceScale,SetScale", Varying))
+USTRUCT(meta=(DisplayName="Set Scale", Category="Hierarchy", DocumentationPolicy = "Strict", Keywords="SetBoneScale,SetControlScale,SetInitialScale,SetSpaceScale,SetScale", NodeColor="0, 0.364706, 1.0", Varying))
 struct FRigUnit_SetScale : public FRigUnitMutable
 {
 	GENERATED_BODY()
diff --git a/Engine/Plugins/Experimental/ControlRig/Source/ControlRig/Public/Units/Highlevel/RigUnit_HighlevelBase.h b/Engine/Plugins/Experimental/ControlRig/Source/ControlRig/Public/Units/Highlevel/RigUnit_HighlevelBase.h
--- a/Engine/Plugins/Experimental/ControlRig/Source/ControlRig/Public/Units/Highlevel/RigUnit_HighlevelBase.h
+++ b/Engine/Plugins/Experimental/ControlRig/Source/ControlRig/Public/Units/Highlevel/RigUnit_HighlevelBase.h
@@ -12,13 +12,13 @@ enum class EControlRigVectorKind : uint8
 	Location
 };
 
-USTRUCT(meta=(Abstract, Category="Debug", NodeColor = "0.4 0.05 0.4"))
+USTRUCT(meta=(Abstract, Category="Debug", NodeColor = "0.462745, 1,0, 0.329412"))
 struct FRigUnit_HighlevelBase : public FRigUnit
 {
 	GENERATED_BODY()
 };
 
-USTRUCT(meta=(Abstract, Category="Debug", NodeColor = "0.4 0.05 0.4"))
+USTRUCT(meta=(Abstract, Category="Debug", NodeColor = "0, 0.364706, 1.0"))
 struct FRigUnit_HighlevelBaseMutable : public FRigUnitMutable
 {
 	GENERATED_BODY()
diff --git a/Engine/Source/Developer/RigVMDeveloper/Private/RigVMModel/RigVMNode.cpp b/Engine/Source/Developer/RigVMDeveloper/Private/RigVMModel/RigVMNode.cpp
--- a/Engine/Source/Developer/RigVMDeveloper/Private/RigVMModel/RigVMNode.cpp
+++ b/Engine/Source/Developer/RigVMDeveloper/Private/RigVMModel/RigVMNode.cpp
@@ -11,7 +11,7 @@ URigVMNode::URigVMNode()
 : UObject()
 , Position(FVector2D::ZeroVector)
 , Size(FVector2D::ZeroVector)
-, NodeColor(FLinearColor::Black)
+, NodeColor(FLinearColor::White)
 , GetSliceContextBracket(0)
 {
 
