Control Rig: Avoid notifies during class generation - speeds up load quite dramatically

diff --git a/Engine/Plugins/Experimental/ControlRig/Source/ControlRigDeveloper/Private/ControlRigBlueprint.cpp b/Engine/Plugins/Experimental/ControlRig/Source/ControlRigDeveloper/Private/ControlRigBlueprint.cpp
--- a/Engine/Plugins/Experimental/ControlRig/Source/ControlRigDeveloper/Private/ControlRigBlueprint.cpp
+++ b/Engine/Plugins/Experimental/ControlRig/Source/ControlRigDeveloper/Private/ControlRigBlueprint.cpp
@@ -144,6 +144,17 @@ UClass* UControlRigBlueprint::GetBlueprintClass() const
 	return UControlRigBlueprintGeneratedClass::StaticClass();
 }
 
+UClass* UControlRigBlueprint::RegenerateClass(UClass* ClassToRegenerate, UObject* PreviousCDO)
+{
+	UClass* Result = nullptr;
+	{
+		TGuardValue<bool> NotificationGuard(bSuspendAllNotifications, true);
+		Result = Super::RegenerateClass(ClassToRegenerate, PreviousCDO);
+	}
+	PropagateHierarchyFromBPToInstances();
+	return Result;
+}
+
 void UControlRigBlueprint::LoadModulesRequiredForCompilation() 
 {
 }
@@ -452,6 +463,8 @@ void UControlRigBlueprint::PostLoad()
 		}
 	}
 
+	PropagateHierarchyFromBPToInstances();
+	
 	// remove all non-controlrig-graphs
 	TArray<UEdGraph*> NewUberGraphPages;
 	for (UEdGraph* Graph : UbergraphPages)
diff --git a/Engine/Plugins/Experimental/ControlRig/Source/ControlRigDeveloper/Public/ControlRigBlueprint.h b/Engine/Plugins/Experimental/ControlRig/Source/ControlRigDeveloper/Public/ControlRigBlueprint.h
--- a/Engine/Plugins/Experimental/ControlRig/Source/ControlRigDeveloper/Public/ControlRigBlueprint.h
+++ b/Engine/Plugins/Experimental/ControlRig/Source/ControlRigDeveloper/Public/ControlRigBlueprint.h
@@ -155,6 +155,7 @@ public:
 
 	// UBlueprint interface
 	virtual UClass* GetBlueprintClass() const override;
+	virtual UClass* RegenerateClass(UClass* ClassToRegenerate, UObject* PreviousCDO) override;
 	virtual bool SupportedByDefaultBlueprintFactory() const override { return false; }
 	virtual bool IsValidForBytecodeOnlyRecompile() const override { return false; }
 	virtual void LoadModulesRequiredForCompilation() override;
