Control Rig: Fix cache validity bug with Controls

diff --git a/Engine/Plugins/Experimental/ControlRig/Source/ControlRig/Private/Rigs/RigHierarchy.cpp b/Engine/Plugins/Experimental/ControlRig/Source/ControlRig/Private/Rigs/RigHierarchy.cpp
--- a/Engine/Plugins/Experimental/ControlRig/Source/ControlRig/Private/Rigs/RigHierarchy.cpp
+++ b/Engine/Plugins/Experimental/ControlRig/Source/ControlRig/Private/Rigs/RigHierarchy.cpp
@@ -1821,9 +1821,19 @@ void URigHierarchy::PropagateDirtyFlags(FRigTransformElement* InTransformElement
 		for(const FRigTransformElement::FElementToDirty& ElementToDirty : InTransformElement->ElementsToDirty)
 		{
 #if URIGHIERARCHY_RECURSIVE_DIRTY_PROPAGATION
-			if(ElementToDirty.Element->Pose.IsDirty(TypeToDirty))
+			if(FRigMultiParentElement* MultiParentElement = Cast<FRigMultiParentElement>(ElementToDirty.Element))
 			{
-				continue;
+				if(MultiParentElement->Parent.IsDirty(TypeToDirty))
+				{
+					continue;
+				}
+			}
+			else
+			{
+				if(ElementToDirty.Element->Pose.IsDirty(TypeToDirty))
+				{
+					continue;
+				}
 			}
 #else
 
@@ -1854,11 +1864,21 @@ void URigHierarchy::PropagateDirtyFlags(FRigTransformElement* InTransformElement
 		{
 #if URIGHIERARCHY_RECURSIVE_DIRTY_PROPAGATION
 
-			if(ElementToDirty.Element->Pose.IsDirty(TypeToDirty))
+			if(FRigMultiParentElement* MultiParentElement = Cast<FRigMultiParentElement>(ElementToDirty.Element))
 			{
-				continue;
+				if(MultiParentElement->Parent.IsDirty(TypeToDirty))
+				{
+					continue;
+				}
 			}
-			
+			else
+			{
+				if(ElementToDirty.Element->Pose.IsDirty(TypeToDirty))
+				{
+					continue;
+				}
+			}
+						
 #else
 
 			if(!bAffectChildren && ElementToDirty.HierarchyDistance > 1)
