Control Rig: Add SetupUndo flag to Control Setters so we don't spame Undo stack when evaluating Sequencer. Typo from last submit

diff --git a/Engine/Plugins/Experimental/ControlRig/Source/ControlRig/Private/Sequencer/MovieSceneControlRigParameterTemplate.cpp b/Engine/Plugins/Experimental/ControlRig/Source/ControlRig/Private/Sequencer/MovieSceneControlRigParameterTemplate.cpp
--- a/Engine/Plugins/Experimental/ControlRig/Source/ControlRig/Private/Sequencer/MovieSceneControlRigParameterTemplate.cpp
+++ b/Engine/Plugins/Experimental/ControlRig/Source/ControlRig/Private/Sequencer/MovieSceneControlRigParameterTemplate.cpp
@@ -654,7 +654,7 @@ struct FControlRigParameterPreAnimatedTokenProducer : IMovieScenePreAnimatedToke
 						{
 							if (ControlRig->FindControl(Value.Name))
 							{
-								ControlRig->SetControlValue<float>(Value.Name, Value.Value, true, FRigControlModifiedContext(EControlRigSetKey::Never,bSetupUndo));
+								ControlRig->SetControlValue<float>(Value.Name, Value.Value, true, FRigControlModifiedContext(EControlRigSetKey::Never),bSetupUndo);
 							}
 						}
 
@@ -662,7 +662,7 @@ struct FControlRigParameterPreAnimatedTokenProducer : IMovieScenePreAnimatedToke
 						{
 							if (ControlRig->FindControl(Value.Name))
 							{
-								ControlRig->SetControlValue<bool>(Value.Name, Value.Value, true, FRigControlModifiedContext(EControlRigSetKey::Never,bSetupUndo));
+								ControlRig->SetControlValue<bool>(Value.Name, Value.Value, true, FRigControlModifiedContext(EControlRigSetKey::Never),bSetupUndo);
 							}
 						}
 
@@ -670,7 +670,7 @@ struct FControlRigParameterPreAnimatedTokenProducer : IMovieScenePreAnimatedToke
 						{
 							if (ControlRig->FindControl(Value.Name))
 							{
-								ControlRig->SetControlValue<int32>(Value.Name, Value.Value, true, FRigControlModifiedContext(EControlRigSetKey::Never,bSetupUndo));
+								ControlRig->SetControlValue<int32>(Value.Name, Value.Value, true, FRigControlModifiedContext(EControlRigSetKey::Never),bSetupUndo);
 							}
 						}
 
@@ -678,7 +678,7 @@ struct FControlRigParameterPreAnimatedTokenProducer : IMovieScenePreAnimatedToke
 						{
 							if (ControlRig->FindControl(Value.Name))
 							{
-								ControlRig->SetControlValue<FVector2D>(Value.Name, Value.Value, true, FRigControlModifiedContext(EControlRigSetKey::Never,bSetupUndo));
+								ControlRig->SetControlValue<FVector2D>(Value.Name, Value.Value, true, FRigControlModifiedContext(EControlRigSetKey::Never),bSetupUndo);
 							}
 						}
 
@@ -686,7 +686,7 @@ struct FControlRigParameterPreAnimatedTokenProducer : IMovieScenePreAnimatedToke
 						{
 							if (ControlRig->FindControl(Value.Name))
 							{
-								ControlRig->SetControlValue<FVector>(Value.Name, Value.Value, true, FRigControlModifiedContext(EControlRigSetKey::Never,bSetupUndo));
+								ControlRig->SetControlValue<FVector>(Value.Name, Value.Value, true, FRigControlModifiedContext(EControlRigSetKey::Never),bSetupUndo);
 							}
 						}
 
@@ -698,19 +698,19 @@ struct FControlRigParameterPreAnimatedTokenProducer : IMovieScenePreAnimatedToke
 								{
 								case ERigControlType::Transform:
 								{
-									ControlRig->SetControlValue<FTransform>(Value.Name, Value.Value, true, FRigControlModifiedContext(EControlRigSetKey::Never,bSetupUndo));
+									ControlRig->SetControlValue<FTransform>(Value.Name, Value.Value, true, FRigControlModifiedContext(EControlRigSetKey::Never),bSetupUndo);
 									break;
 								}
 								case ERigControlType::TransformNoScale:
 								{
 									FTransformNoScale NoScale = Value.Value;
-									ControlRig->SetControlValue<FTransformNoScale>(Value.Name, NoScale, true, FRigControlModifiedContext(EControlRigSetKey::Never,bSetupUndo));
+									ControlRig->SetControlValue<FTransformNoScale>(Value.Name, NoScale, true, FRigControlModifiedContext(EControlRigSetKey::Never),bSetupUndo);
 									break;
 								}
 								case ERigControlType::EulerTransform:
 								{
 									FEulerTransform EulerTransform = Value.Value;
-									ControlRig->SetControlValue<FEulerTransform>(Value.Name, EulerTransform, true, FRigControlModifiedContext(EControlRigSetKey::Never,bSetupUndo));
+									ControlRig->SetControlValue<FEulerTransform>(Value.Name, EulerTransform, true, FRigControlModifiedContext(EControlRigSetKey::Never),bSetupUndo);
 									break;
 								}
 
