Sequencer: Control Rig: Python Exposure

diff --git a/Engine/Plugins/Experimental/ControlRig/ControlRig.uplugin b/Engine/Plugins/Experimental/ControlRig/ControlRig.uplugin
--- a/Engine/Plugins/Experimental/ControlRig/ControlRig.uplugin
+++ b/Engine/Plugins/Experimental/ControlRig/ControlRig.uplugin
@@ -23,6 +23,10 @@
 		{
 			"Name": "PropertyAccessEditor",
 			"Enabled": true
+		},
+		{
+			"Name": "SequencerScripting",
+			"Enabled": true
 		}
 	],
 	"Modules" :
diff --git a/Engine/Plugins/Experimental/ControlRig/Source/ControlRig/Public/ControlRig.h b/Engine/Plugins/Experimental/ControlRig/Source/ControlRig/Public/ControlRig.h
--- a/Engine/Plugins/Experimental/ControlRig/Source/ControlRig/Public/ControlRig.h
+++ b/Engine/Plugins/Experimental/ControlRig/Source/ControlRig/Public/ControlRig.h
@@ -327,9 +327,13 @@ public:
 	virtual void CreateRigControlsForCurveContainer();
 	virtual void GetControlsInOrder(TArray<FRigControlElement*>& SortedControls) const;
 
+	UFUNCTION(BlueprintCallable, Category = "Control Rig")
 	virtual void SelectControl(const FName& InControlName, bool bSelect = true);
+	UFUNCTION(BlueprintCallable, Category = "Control Rig")
 	virtual bool ClearControlSelection();
+	UFUNCTION(BlueprintPure, Category = "Control Rig")
 	virtual TArray<FName> CurrentControlSelection() const;
+	UFUNCTION(BlueprintPure, Category = "Control Rig")
 	virtual bool IsControlSelected(const FName& InControlName)const;
 
 	// Returns true if this manipulatable subject is currently
diff --git a/Engine/Plugins/Experimental/ControlRig/Source/ControlRigEditor/ControlRigEditor.Build.cs b/Engine/Plugins/Experimental/ControlRig/Source/ControlRigEditor/ControlRigEditor.Build.cs
--- a/Engine/Plugins/Experimental/ControlRig/Source/ControlRigEditor/ControlRigEditor.Build.cs
+++ b/Engine/Plugins/Experimental/ControlRig/Source/ControlRigEditor/ControlRigEditor.Build.cs
@@ -76,6 +76,7 @@ namespace UnrealBuildTool.Rules
                     "RigVMDeveloper",
 					"AnimationEditor",
 					"MessageLog",
+                    "SequencerScripting",
 					"PropertyAccessEditor",
 					"KismetWidgets",
 				}
diff --git a/Engine/Plugins/Experimental/ControlRig/Source/ControlRigEditor/Private/ControlRigSequencerEditorLibrary.cpp b/Engine/Plugins/Experimental/ControlRig/Source/ControlRigEditor/Private/ControlRigSequencerEditorLibrary.cpp
new file mode 100644
--- /dev/null
+++ b/Engine/Plugins/Experimental/ControlRig/Source/ControlRigEditor/Private/ControlRigSequencerEditorLibrary.cpp
@@ -0,0 +1,1536 @@
+// Copyright Epic Games, Inc. All Rights Reserved.
+
+#include "ControlRigSequencerEditorLibrary.h"
+#include "LevelSequence.h"
+#include "Sequencer/MovieSceneControlRigParameterTrack.h"
+#include "Sequencer/MovieSceneControlRigParameterSection.h"
+#include "Components/SkeletalMeshComponent.h"
+#include "Editor.h"
+#include "Editor/UnrealEdEngine.h"
+#include "UnrealEdGlobals.h"
+#include "ISequencer.h"
+#include "ControlRigEditorModule.h"
+#include "Channels/FloatChannelCurveModel.h"
+#include "TransformNoScale.h"
+#include "ControlRigComponent.h"
+#include "MovieSceneToolHelpers.h"
+#include "Rigs/FKControlRig.h"
+#include "ControlRig/Private/Units/Execution/RigUnit_InverseExecution.h"
+#include "ControlRig.h"
+#include "ControlRigEditMode.h"
+#include "EditorModeManager.h"
+#include "Engine/Selection.h"
+#include "ControlRigObjectBinding.h"
+#include "Engine/SCS_Node.h"
+#include "ILevelSequenceEditorToolkit.h"
+#include "ISequencer.h"
+#include "Tools/ControlRigTweener.h"
+#include "LevelSequencePlayer.h"
+#include "LevelSequence.h"
+#include "LevelSequenceActor.h"
+#include "Logging/LogMacros.h"
+#include "ControlRig/Private/Units/Execution/RigUnit_InverseExecution.h"
+#include "Tracks/MovieSceneSkeletalAnimationTrack.h"
+#include "Exporters/AnimSeqExportOption.h"
+#include "MovieSceneTimeHelpers.h"
+
+#define LOCTEXT_NAMESPACE "ControlrigSequencerEditorLibrary"
+
+TArray<UControlRig*> UControlRigSequencerEditorLibrary::GetVisibleControlRigs()
+{
+	TArray<UControlRig*> ControlRigs;
+	FControlRigEditMode* ControlRigEditMode = static_cast<FControlRigEditMode*>(GLevelEditorModeTools().GetActiveMode(FControlRigEditMode::ModeName));
+	if (ControlRigEditMode && ControlRigEditMode->GetControlRig(true) && ControlRigEditMode->GetControlRig(true)->GetObjectBinding() &&\
+		ControlRigEditMode->GetControlRig(true)->GetObjectBinding()->GetBoundObject())
+	{
+		ControlRigs.Add(ControlRigEditMode->GetControlRig(true));
+	}
+	return ControlRigs;
+}
+
+TArray<FControlRigSequencerBindingProxy> UControlRigSequencerEditorLibrary::GetControlRigs(ULevelSequence* LevelSequence)
+{
+	TArray<FControlRigSequencerBindingProxy> ControlRigBindingProxies;
+	if (LevelSequence)
+	{
+		UMovieScene* MovieScene = LevelSequence->GetMovieScene();
+		if (MovieScene)
+		{
+			const TArray<FMovieSceneBinding>& Bindings = MovieScene->GetBindings();
+			for (const FMovieSceneBinding& Binding : Bindings)
+			{
+				TArray<UMovieSceneTrack*> Tracks = MovieScene->FindTracks(UMovieSceneControlRigParameterTrack::StaticClass(), Binding.GetObjectGuid(), NAME_None);
+				for (UMovieSceneTrack* AnyOleTrack : Tracks)
+				{
+					UMovieSceneControlRigParameterTrack* Track = Cast<UMovieSceneControlRigParameterTrack>(AnyOleTrack);
+					if (Track && Track->GetControlRig())
+					{
+						FControlRigSequencerBindingProxy BindingProxy;
+						BindingProxy.ControlRig = Track->GetControlRig();
+						BindingProxy.Proxy.BindingID = Binding.GetObjectGuid();
+						BindingProxy.Proxy.Sequence = LevelSequence;
+						ControlRigBindingProxies.Add(BindingProxy);
+					}
+				}
+			}
+		}
+	}
+	return ControlRigBindingProxies;
+}
+
+static void AcquireSkeletonAndSkelMeshCompFromObject(UObject* BoundObject, USkeleton** OutSkeleton, USkeletalMeshComponent** OutSkeletalMeshComponent)
+{
+	*OutSkeletalMeshComponent = nullptr;
+	*OutSkeleton = nullptr;
+	if (AActor* Actor = Cast<AActor>(BoundObject))
+	{
+		for (UActorComponent* Component : Actor->GetComponents())
+		{
+			USkeletalMeshComponent* SkeletalMeshComp = Cast<USkeletalMeshComponent>(Component);
+			if (SkeletalMeshComp)
+			{
+				*OutSkeletalMeshComponent = SkeletalMeshComp;
+				if (SkeletalMeshComp->SkeletalMesh && SkeletalMeshComp->SkeletalMesh->GetSkeleton())
+				{
+					*OutSkeleton = SkeletalMeshComp->SkeletalMesh->GetSkeleton();
+				}
+				return;
+			}
+		}
+
+		AActor* ActorCDO = Cast<AActor>(Actor->GetClass()->GetDefaultObject());
+		if (ActorCDO)
+		{
+			for (UActorComponent* Component : ActorCDO->GetComponents())
+			{
+				USkeletalMeshComponent* SkeletalMeshComp = Cast<USkeletalMeshComponent>(Component);
+				if (SkeletalMeshComp)
+				{
+					*OutSkeletalMeshComponent = SkeletalMeshComp;
+					if (SkeletalMeshComp->SkeletalMesh && SkeletalMeshComp->SkeletalMesh->GetSkeleton())
+					{
+						*OutSkeleton = SkeletalMeshComp->SkeletalMesh->GetSkeleton();
+					}
+					return;
+				}
+			}
+		}
+
+		UBlueprintGeneratedClass* ActorBlueprintGeneratedClass = Cast<UBlueprintGeneratedClass>(Actor->GetClass());
+		if (ActorBlueprintGeneratedClass)
+		{
+			const TArray<USCS_Node*>& ActorBlueprintNodes = ActorBlueprintGeneratedClass->SimpleConstructionScript->GetAllNodes();
+
+			for (USCS_Node* Node : ActorBlueprintNodes)
+			{
+				if (Node->ComponentClass->IsChildOf(USkeletalMeshComponent::StaticClass()))
+				{
+					USkeletalMeshComponent* SkeletalMeshComp = Cast<USkeletalMeshComponent>(Node->GetActualComponentTemplate(ActorBlueprintGeneratedClass));
+					if (SkeletalMeshComp)
+					{
+						*OutSkeletalMeshComponent = SkeletalMeshComp;
+						if (SkeletalMeshComp->SkeletalMesh && SkeletalMeshComp->SkeletalMesh->GetSkeleton())
+						{
+							*OutSkeleton = SkeletalMeshComp->SkeletalMesh->GetSkeleton();
+						}
+					}
+				}
+			}
+		}
+	}
+	else if (USkeletalMeshComponent* SkeletalMeshComponent = Cast<USkeletalMeshComponent>(BoundObject))
+	{
+		*OutSkeletalMeshComponent = SkeletalMeshComponent;
+		if (SkeletalMeshComponent->SkeletalMesh && SkeletalMeshComponent->SkeletalMesh->GetSkeleton())
+		{
+			*OutSkeleton = SkeletalMeshComponent->SkeletalMesh->GetSkeleton();
+		}
+	}
+}
+
+static TSharedPtr<ISequencer> GetSequencerFromAsset(ULevelSequence* LevelSequence)
+{
+	IAssetEditorInstance* AssetEditor = GEditor->GetEditorSubsystem<UAssetEditorSubsystem>()->FindEditorForAsset(LevelSequence, false);
+	ILevelSequenceEditorToolkit* LevelSequenceEditor = static_cast<ILevelSequenceEditorToolkit*>(AssetEditor);
+	TSharedPtr<ISequencer> Sequencer = LevelSequenceEditor ? LevelSequenceEditor->GetSequencer() : nullptr;
+	if (Sequencer.IsValid() == false)
+	{
+		UE_LOG(LogControlRig, Error, TEXT("Can not open Sequencer for the LevelSequence %s"), *(LevelSequence->GetPathName()));
+	}
+	return Sequencer;
+}
+
+static UMovieSceneControlRigParameterTrack* AddControlRig(ULevelSequence* LevelSequence,const UClass* InClass, UObject* BoundActor, FGuid ObjectBinding, UControlRig* InExistingControlRig)
+{
+	FSlateApplication::Get().DismissAllMenus();
+
+	if (InClass && InClass->IsChildOf(UControlRig::StaticClass()) && LevelSequence && LevelSequence->GetMovieScene())
+	{
+		UMovieScene* OwnerMovieScene = LevelSequence->GetMovieScene();
+		TSharedPtr<ISequencer> SharedSequencer = GetSequencerFromAsset(LevelSequence);
+		ISequencer* Sequencer = nullptr; // will be valid  if we have a ISequencer AND it's focused.
+		if (SharedSequencer.IsValid() && SharedSequencer->GetFocusedMovieSceneSequence() == LevelSequence)
+		{
+			Sequencer = SharedSequencer.Get();
+		}
+		LevelSequence->Modify();
+		OwnerMovieScene->Modify();
+		UMovieSceneControlRigParameterTrack* Track = Cast<UMovieSceneControlRigParameterTrack>(OwnerMovieScene->AddTrack(UMovieSceneControlRigParameterTrack::StaticClass(), ObjectBinding));
+		if (Track)
+		{
+			FString ObjectName = InClass->GetName(); //GetDisplayNameText().ToString();
+			ObjectName.RemoveFromEnd(TEXT("_C"));
+
+			bool bSequencerOwnsControlRig = false;
+			UControlRig* ControlRig = InExistingControlRig;
+			if (ControlRig == nullptr)
+			{
+				ControlRig = NewObject<UControlRig>(Track, InClass, FName(*ObjectName), RF_Transactional);
+				bSequencerOwnsControlRig = true;
+			}
+
+			ControlRig->Modify();
+			ControlRig->SetObjectBinding(MakeShared<FControlRigObjectBinding>());
+			ControlRig->GetObjectBinding()->BindToObject(BoundActor);
+			ControlRig->GetDataSourceRegistry()->RegisterDataSource(UControlRig::OwnerComponent, ControlRig->GetObjectBinding()->GetBoundObject());
+			ControlRig->Initialize();
+			ControlRig->Evaluate_AnyThread();
+
+
+			Track->Modify();
+			UMovieSceneSection* NewSection = Track->CreateControlRigSection(0, ControlRig, bSequencerOwnsControlRig);
+			NewSection->Modify();
+
+			//mz todo need to have multiple rigs with same class
+			Track->SetTrackName(FName(*ObjectName));
+			Track->SetDisplayName(FText::FromString(ObjectName));
+
+			if (Sequencer)
+			{
+				Sequencer->NotifyMovieSceneDataChanged(EMovieSceneDataChangeType::MovieSceneStructureItemsChanged);
+				Sequencer->EmptySelection();
+				Sequencer->SelectSection(NewSection);
+				Sequencer->ThrobSectionSelection();
+				Sequencer->NotifyMovieSceneDataChanged(EMovieSceneDataChangeType::MovieSceneStructureItemAdded);
+				Sequencer->ObjectImplicitlyAdded(ControlRig);
+			}
+
+			FControlRigEditMode* ControlRigEditMode = static_cast<FControlRigEditMode*>(GLevelEditorModeTools().GetActiveMode(FControlRigEditMode::ModeName));
+			if (!ControlRigEditMode)
+			{
+				GLevelEditorModeTools().ActivateMode(FControlRigEditMode::ModeName);
+				ControlRigEditMode = static_cast<FControlRigEditMode*>(GLevelEditorModeTools().GetActiveMode(FControlRigEditMode::ModeName));
+
+			}
+			if (ControlRigEditMode)
+			{
+				ControlRigEditMode->SetObjects(ControlRig, nullptr, SharedSequencer);
+			}
+			return Track;
+		}
+	}
+	return nullptr;
+}
+
+
+UMovieSceneTrack* UControlRigSequencerEditorLibrary::FindOrCreateControlRigTrack(ULevelSequence* LevelSequence, const UClass* ControlRigClass, const FSequencerBindingProxy& InBinding)
+{
+	UMovieScene* MovieScene = InBinding.Sequence ? InBinding.Sequence->GetMovieScene() : nullptr;
+	UMovieSceneTrack* BaseTrack = nullptr;
+	if (LevelSequence && MovieScene && InBinding.BindingID.IsValid())
+	{
+		const TArray<FMovieSceneBinding>& Bindings = MovieScene->GetBindings();
+		for (const FMovieSceneBinding& Binding : Bindings)
+		{
+			if (Binding.GetObjectGuid() == InBinding.BindingID)
+			{
+				TArray<UMovieSceneTrack*> Tracks = MovieScene->FindTracks(UMovieSceneControlRigParameterTrack::StaticClass(), Binding.GetObjectGuid(), NAME_None);
+				for (UMovieSceneTrack* AnyOleTrack : Tracks)
+				{
+					UMovieSceneControlRigParameterTrack* Track = Cast<UMovieSceneControlRigParameterTrack>(AnyOleTrack);
+					if (Track && Track->GetControlRig() && Track->GetControlRig()->GetClass() == ControlRigClass)
+					{
+						return Track;
+					}
+				}
+
+				TArray<UObject*, TInlineAllocator<1>> Result;
+				UObject* Context = nullptr;
+				LevelSequence->LocateBoundObjects(InBinding.BindingID, Context, Result);
+				if (Result.Num() > 0 && Result[0])
+				{
+					UObject* BoundObject = Result[0];
+					USkeleton* Skeleton = nullptr;
+					USkeletalMeshComponent* SkeletalMeshComponent = nullptr;
+					AcquireSkeletonAndSkelMeshCompFromObject(BoundObject, &Skeleton, &SkeletalMeshComponent);
+
+					UControlRig* ControlRig = nullptr;
+						
+
+					if (Skeleton && SkeletalMeshComponent)
+					{
+						UMovieSceneControlRigParameterTrack* Track = AddControlRig(LevelSequence, ControlRigClass, SkeletalMeshComponent, InBinding.BindingID, nullptr);
+
+						if (Track)
+						{
+							BaseTrack = Track;								
+						}
+					}
+				}
+			}
+		}
+	}
+	return BaseTrack;
+}
+
+TArray<UMovieSceneTrack*> UControlRigSequencerEditorLibrary::FindOrCreateControlRigComponentTrack(ULevelSequence* LevelSequence, const FSequencerBindingProxy& InBinding)
+{
+	TArray< UMovieSceneTrack*> Tracks;
+	TArray<UObject*, TInlineAllocator<1>> Result;
+
+	UObject* Context = nullptr;
+	LevelSequence->LocateBoundObjects(InBinding.BindingID, Context, Result);
+	if (Result.Num() > 0 && Result[0])
+	{
+		UObject* BoundObject = Result[0];
+		
+		if (AActor* BoundActor = Cast<AActor>(BoundObject))
+		{
+			TArray<UControlRigComponent*> ControlRigComponents;
+			BoundActor->GetComponents<UControlRigComponent>(ControlRigComponents);
+			for (UControlRigComponent* ControlRigComponent : ControlRigComponents)
+			{
+				if (UControlRig* CR = ControlRigComponent->GetControlRig())
+				{
+					UMovieSceneControlRigParameterTrack* Track = AddControlRig(LevelSequence,CR->GetClass(), BoundActor, InBinding.BindingID, CR);
+					Tracks.Add(Track);
+				}
+			}
+		}
+	}
+	return Tracks;
+}
+
+
+bool UControlRigSequencerEditorLibrary::TweenControlRig(ULevelSequence* LevelSequence, UControlRig* ControlRig, float TweenValue)
+{
+	TWeakPtr<ISequencer> WeakSequencer = GetSequencerFromAsset(LevelSequence);
+	if (WeakSequencer.IsValid() && WeakSequencer.Pin()->GetFocusedMovieSceneSequence() == LevelSequence
+		&& ControlRig && LevelSequence->GetMovieScene())
+	{
+		FControlsToTween ControlsToTween;
+		LevelSequence->GetMovieScene()->Modify();
+		TArray<UControlRig*> SelectedControlRigs;
+		SelectedControlRigs.Add(ControlRig);
+		ControlsToTween.Setup(SelectedControlRigs, WeakSequencer);
+		ControlsToTween.Blend(WeakSequencer, TweenValue);
+		return true;
+	}
+	return false;
+}
+
+
+bool UControlRigSequencerEditorLibrary::SnapControlRig(FFrameNumber StartFrame, FFrameNumber EndFrame, const FControlRigSnapperSelection& ChildrenToSnap,
+	const FControlRigSnapperSelection& ParentToSnap, const UControlRigSnapSettings* SnapSettings)
+{
+	FControlRigSnapper Snapper;
+	return Snapper.SnapIt(StartFrame, EndFrame, ChildrenToSnap, ParentToSnap,SnapSettings);
+}
+
+FTransform UControlRigSequencerEditorLibrary::GetActorWorldTransform(ULevelSequence* LevelSequence, AActor* Actor, FFrameNumber Frame)
+{
+	TArray<FFrameNumber> Frames;
+	Frames.Add(Frame);
+	TArray<FTransform> Transforms = GetActorWorldTransforms(LevelSequence, Actor, Frames);
+	if (Transforms.Num() == 1)
+	{
+		return Transforms[0];
+	}
+	return FTransform::Identity;
+}
+
+TArray<FTransform> UControlRigSequencerEditorLibrary::GetActorWorldTransforms(ULevelSequence* LevelSequence,AActor* Actor, const TArray<FFrameNumber>& Frames)
+{
+	TWeakPtr<ISequencer> WeakSequencer = GetSequencerFromAsset(LevelSequence);
+	TArray<FTransform> OutWorldTransforms;
+	if (WeakSequencer.IsValid() && Actor)
+	{
+		FActorForWorldTransforms Actors;
+		Actors.Actor = Actor;
+		MovieSceneToolHelpers::GetActorWorldTransforms(WeakSequencer.Pin().Get(), Actors, Frames, OutWorldTransforms); 
+	}
+	return OutWorldTransforms;
+
+}
+
+FTransform UControlRigSequencerEditorLibrary::GetSkeletalMeshComponentWorldTransform(ULevelSequence* LevelSequence,USkeletalMeshComponent* SkeletalMeshComponent, FFrameNumber Frame, FName SocketName)
+{
+	TArray<FFrameNumber> Frames;
+	Frames.Add(Frame);
+	TArray<FTransform> Transforms = GetSkeletalMeshComponentWorldTransforms(LevelSequence, SkeletalMeshComponent, Frames);
+	if (Transforms.Num() == 1)
+	{
+		return Transforms[0];
+	}
+	return FTransform::Identity;
+}
+
+TArray<FTransform> UControlRigSequencerEditorLibrary::GetSkeletalMeshComponentWorldTransforms(ULevelSequence* LevelSequence,USkeletalMeshComponent* SkeletalMeshComponent, const TArray<FFrameNumber>& Frames, FName SocketName)
+{
+	TWeakPtr<ISequencer> WeakSequencer = GetSequencerFromAsset(LevelSequence);
+	TArray<FTransform> OutWorldTransforms;
+	if (WeakSequencer.IsValid() && SkeletalMeshComponent)
+	{
+		FActorForWorldTransforms Actors;
+		AActor* Actor = SkeletalMeshComponent->GetTypedOuter<AActor>();
+		if (Actor)
+		{
+			Actors.Actor = Actor;
+			Actors.Component = SkeletalMeshComponent;
+			Actors.SocketName = SocketName;
+			MovieSceneToolHelpers::GetActorWorldTransforms(WeakSequencer.Pin().Get(), Actors, Frames, OutWorldTransforms);
+		}
+	}
+	return OutWorldTransforms;
+}
+
+FTransform UControlRigSequencerEditorLibrary::GetControlRigWorldTransform(ULevelSequence* LevelSequence,UControlRig* ControlRig, FName ControlName, FFrameNumber Frame)
+{
+	TArray<FFrameNumber> Frames;
+	Frames.Add(Frame);
+	TArray<FTransform> Transforms = GetControlRigWorldTransforms(LevelSequence, ControlRig, ControlName, Frames);
+	if (Transforms.Num() == 1)
+	{
+		return Transforms[0];
+	}
+	return FTransform::Identity;
+}
+
+TArray<FTransform> UControlRigSequencerEditorLibrary::GetControlRigWorldTransforms(ULevelSequence* LevelSequence,UControlRig* ControlRig, FName ControlName, const TArray<FFrameNumber>& Frames)
+{
+	TWeakPtr<ISequencer> WeakSequencer = GetSequencerFromAsset(LevelSequence);
+	TArray<FTransform> OutWorldTransforms;
+	if (WeakSequencer.IsValid() && ControlRig)
+	{
+		if (TSharedPtr<IControlRigObjectBinding> ObjectBinding = ControlRig->GetObjectBinding())
+		{
+			USceneComponent* Component = Cast<USceneComponent>(ObjectBinding->GetBoundObject());
+			if (Component)
+			{
+				AActor* Actor = Component->GetTypedOuter< AActor >();
+				if (Actor)
+				{
+					ControlRig->Modify();
+					TArray<FTransform> ControlRigParentWorldTransforms;
+					FActorForWorldTransforms ControlRigActorSelection;
+					ControlRigActorSelection.Actor = Actor;
+					MovieSceneToolHelpers::GetActorWorldTransforms(WeakSequencer.Pin().Get(), ControlRigActorSelection, Frames, ControlRigParentWorldTransforms);
+					FControlRigSnapper Snapper;
+					Snapper.GetControlRigControlTransforms(WeakSequencer.Pin().Get(), ControlRig, ControlName, Frames, ControlRigParentWorldTransforms, OutWorldTransforms);	
+				}
+			}
+		}
+	}
+	return OutWorldTransforms;
+}
+
+
+static void LocalSetControlRigWorldTransforms(ULevelSequence* LevelSequence,UControlRig* ControlRig, FName ControlName, EControlRigSetKey SetKey, const TArray<FFrameNumber>& Frames, const TArray<FTransform>& WorldTransforms)
+{
+	TWeakPtr<ISequencer> WeakSequencer = GetSequencerFromAsset(LevelSequence);
+	if (WeakSequencer.IsValid() && ControlRig)
+	{
+		if (TSharedPtr<IControlRigObjectBinding> ObjectBinding = ControlRig->GetObjectBinding())
+		{
+			USceneComponent* Component = Cast<USceneComponent>(ObjectBinding->GetBoundObject());
+			if (Component)
+			{
+				AActor* Actor = Component->GetTypedOuter< AActor >();
+				if (Actor)
+				{
+					UMovieScene* MovieScene = WeakSequencer.Pin().Get()->GetFocusedMovieSceneSequence()->GetMovieScene();
+					MovieScene->Modify();
+					FFrameRate TickResolution = MovieScene->GetTickResolution();
+					FRigControlModifiedContext Context;
+					Context.SetKey = SetKey;
+
+					ControlRig->Modify();
+					TArray<FTransform> ControlRigParentWorldTransforms;
+					FActorForWorldTransforms ControlRigActorSelection;
+					ControlRigActorSelection.Actor = Actor;
+					MovieSceneToolHelpers::GetActorWorldTransforms(WeakSequencer.Pin().Get(), ControlRigActorSelection, Frames, ControlRigParentWorldTransforms);
+					FControlRigSnapper Snapper;
+					
+					TArray<FFrameNumber> OneFrame;
+					OneFrame.SetNum(1);
+					TArray<FTransform> CurrentControlRigTransform, CurrentParentWorldTransform;
+					CurrentControlRigTransform.SetNum(1);
+					CurrentParentWorldTransform.SetNum(1);
+					
+					for (int32 Index = 0; Index < WorldTransforms.Num(); ++Index)
+					{
+						OneFrame[0] = Frames[Index];
+						CurrentParentWorldTransform[0] = ControlRigParentWorldTransforms[Index];
+						//this will evaluate at the current frame which we want
+						Snapper.GetControlRigControlTransforms(WeakSequencer.Pin().Get(), ControlRig, ControlName, OneFrame, CurrentParentWorldTransform, CurrentControlRigTransform);
+							
+						const FFrameNumber& FrameNumber = Frames[Index];
+						Context.LocalTime = TickResolution.AsSeconds(FFrameTime(FrameNumber));
+						FTransform GlobalTransform = WorldTransforms[Index].GetRelativeTransform(ControlRigParentWorldTransforms[Index]);
+						ControlRig->SetControlGlobalTransform(ControlName, GlobalTransform, true, Context);
+					}
+				}
+			}
+		}
+	}
+}
+
+void UControlRigSequencerEditorLibrary::SetControlRigWorldTransform(ULevelSequence* LevelSequence, UControlRig* ControlRig, FName ControlName, FFrameNumber Frame, const FTransform& WorldTransform, bool bSetKey)
+{
+	EControlRigSetKey SetKey = bSetKey ? EControlRigSetKey::Always : EControlRigSetKey::DoNotCare;
+	TArray<FFrameNumber> Frames;
+	Frames.Add(Frame);
+	TArray<FTransform> WorldTransforms;
+	WorldTransforms.Add(WorldTransform);
+
+	LocalSetControlRigWorldTransforms(LevelSequence, ControlRig, ControlName, EControlRigSetKey::Always, Frames, WorldTransforms);
+
+}
+
+void UControlRigSequencerEditorLibrary::SetControlRigWorldTransforms(ULevelSequence* LevelSequence, UControlRig* ControlRig, FName ControlName, const TArray<FFrameNumber>& Frames, const TArray<FTransform>& WorldTransforms)
+{
+	LocalSetControlRigWorldTransforms(LevelSequence, ControlRig, ControlName, EControlRigSetKey::Always, Frames, WorldTransforms);
+}
+
+bool UControlRigSequencerEditorLibrary::BakeToControlRig(UWorld* World, ULevelSequence* LevelSequence, UClass* InClass, UAnimSeqExportOption* ExportOptions, bool bReduceKeys, float Tolerance,
+	const FSequencerBindingProxy& Binding)
+{
+	UMovieScene* MovieScene = LevelSequence->GetMovieScene();
+	if (Binding.Sequence != LevelSequence)
+	{
+		UE_LOG(LogControlRig, Error, TEXT("Baking: Binding.Sequence different"));
+		return false;
+	}
+	//get level sequence if one exists...
+	TWeakPtr<ISequencer> WeakSequencer = GetSequencerFromAsset(LevelSequence);
+
+	ALevelSequenceActor* OutActor;
+	FMovieSceneSequencePlaybackSettings Settings;
+	FLevelSequenceCameraSettings CameraSettings;
+	FMovieSceneSequenceIDRef Template = MovieSceneSequenceID::Root;
+	FMovieSceneSequenceTransform RootToLocalTransform;
+	ULevelSequencePlayer* Player = ULevelSequencePlayer::CreateLevelSequencePlayer(World, LevelSequence, Settings, OutActor);
+	Player->Initialize(LevelSequence, World->PersistentLevel, Settings, CameraSettings);
+	Player->State.AssignSequence(MovieSceneSequenceID::Root, *LevelSequence, *Player);
+
+	bool bResult = false;
+	const FScopedTransaction Transaction(LOCTEXT("BakeToControlRig_Transaction", "Bake To Control Rig"));
+	{
+		FSpawnableRestoreState SpawnableRestoreState(MovieScene);
+
+		if (SpawnableRestoreState.bWasChanged)
+		{
+			// Evaluate at the beginning of the subscene time to ensure that spawnables are created before export
+			Player->SetPlaybackPosition(FMovieSceneSequencePlaybackParams(UE::MovieScene::DiscreteInclusiveLower(MovieScene->GetPlaybackRange()).Value, EUpdatePositionMethod::Play));
+		}
+		TArrayView<TWeakObjectPtr<>>  Result = Player->FindBoundObjects(Binding.BindingID, Template);
+	
+		if (Result.Num() > 0 && Result[0].IsValid())
+		{
+
+			UObject* BoundObject = Result[0].Get();
+			USkeleton* Skeleton = nullptr;
+			USkeletalMeshComponent* SkeletalMeshComp = nullptr;
+			AcquireSkeletonAndSkelMeshCompFromObject(BoundObject, &Skeleton, &SkeletalMeshComp);
+			if (SkeletalMeshComp && SkeletalMeshComp->SkeletalMesh && SkeletalMeshComp->SkeletalMesh->GetSkeleton())
+			{
+				UAnimSequence* TempAnimSequence = NewObject<UAnimSequence>(GetTransientPackage(), NAME_None);
+				TempAnimSequence->SetSkeleton(Skeleton);
+				bResult = MovieSceneToolHelpers::ExportToAnimSequence(TempAnimSequence, ExportOptions, MovieScene, Player, SkeletalMeshComp, Template, RootToLocalTransform);
+				if (bResult == false)
+				{
+					TempAnimSequence->MarkPendingKill();
+					World->DestroyActor(OutActor);
+					return false;
+				}
+
+				MovieScene->Modify();
+				TArray<UMovieSceneTrack*> Tracks = MovieScene->FindTracks(UMovieSceneControlRigParameterTrack::StaticClass(), Binding.BindingID, NAME_None);
+				UMovieSceneControlRigParameterTrack* Track = nullptr;
+				for (UMovieSceneTrack* AnyOleTrack : Tracks)
+				{
+					UMovieSceneControlRigParameterTrack* ValidTrack = Cast<UMovieSceneControlRigParameterTrack>(AnyOleTrack);
+					if (ValidTrack)
+					{
+						Track = ValidTrack;
+						Track->Modify();
+						for (UMovieSceneSection* Section : Track->GetAllSections())
+						{
+							Section->SetIsActive(false);
+						}
+					}
+				}
+				if(Track == nullptr)
+				{
+					Track = Cast<UMovieSceneControlRigParameterTrack>(MovieScene->AddTrack(UMovieSceneControlRigParameterTrack::StaticClass(), Binding.BindingID));
+					{
+						Track->Modify();
+					}
+				}
+
+				if (Track)
+				{
+					FString ObjectName = InClass->GetName();
+					ObjectName.RemoveFromEnd(TEXT("_C"));
+					UControlRig* ControlRig = NewObject<UControlRig>(Track, InClass, FName(*ObjectName), RF_Transactional);
+					if (InClass != UFKControlRig::StaticClass() && !ControlRig->SupportsEvent(FRigUnit_InverseExecution::EventName))
+					{
+						TempAnimSequence->MarkPendingKill();
+						MovieScene->RemoveTrack(*Track);
+						World->DestroyActor(OutActor);
+						return false;
+					}
+					FControlRigEditMode* ControlRigEditMode = nullptr;
+					if (WeakSequencer.IsValid())
+					{
+						ControlRigEditMode = static_cast<FControlRigEditMode*>(GLevelEditorModeTools().GetActiveMode(FControlRigEditMode::ModeName));
+						if (!ControlRigEditMode)
+						{
+							GLevelEditorModeTools().ActivateMode(FControlRigEditMode::ModeName);
+							ControlRigEditMode = static_cast<FControlRigEditMode*>(GLevelEditorModeTools().GetActiveMode(FControlRigEditMode::ModeName));
+
+						}
+						else
+						{
+							UControlRig* OldControlRig = ControlRigEditMode->GetControlRig(false);
+							if (OldControlRig)
+							{
+								WeakSequencer.Pin()->ObjectImplicitlyRemoved(OldControlRig);
+							}
+						}
+					}
+
+					ControlRig->Modify();
+					ControlRig->SetObjectBinding(MakeShared<FControlRigObjectBinding>());
+					ControlRig->GetObjectBinding()->BindToObject(SkeletalMeshComp);
+					ControlRig->GetDataSourceRegistry()->RegisterDataSource(UControlRig::OwnerComponent, ControlRig->GetObjectBinding()->GetBoundObject());
+					ControlRig->Initialize();
+					ControlRig->Evaluate_AnyThread();
+
+					bool bSequencerOwnsControlRig = true;
+					UMovieSceneSection* NewSection = Track->CreateControlRigSection(0, ControlRig, bSequencerOwnsControlRig);
+					UMovieSceneControlRigParameterSection* ParamSection = Cast<UMovieSceneControlRigParameterSection>(NewSection);
+
+					//mz todo need to have multiple rigs with same class
+					Track->SetTrackName(FName(*ObjectName));
+					Track->SetDisplayName(FText::FromString(ObjectName));
+
+					if (WeakSequencer.IsValid())
+					{
+						WeakSequencer.Pin()->EmptySelection();
+						WeakSequencer.Pin()->SelectSection(NewSection);
+						WeakSequencer.Pin()->ThrobSectionSelection();
+						WeakSequencer.Pin()->NotifyMovieSceneDataChanged(EMovieSceneDataChangeType::MovieSceneStructureItemAdded);
+					}
+					ParamSection->LoadAnimSequenceIntoThisSection(TempAnimSequence, MovieScene, Skeleton,
+						bReduceKeys, Tolerance);
+
+					//Turn Off Any Skeletal Animation Tracks
+					UMovieSceneSkeletalAnimationTrack* SkelTrack = Cast<UMovieSceneSkeletalAnimationTrack>(MovieScene->FindTrack(UMovieSceneSkeletalAnimationTrack::StaticClass(), Binding.BindingID, NAME_None));
+					if (SkelTrack)
+					{
+						SkelTrack->Modify();
+						//can't just turn off the track so need to mute the sections
+						const TArray<UMovieSceneSection*>& Sections = SkelTrack->GetAllSections();
+						for (UMovieSceneSection* Section : Sections)
+						{
+							if (Section)
+							{
+								Section->TryModify();
+								Section->SetIsActive(false);
+							}
+						}
+					}
+					//Finish Setup
+					if (ControlRigEditMode)
+					{
+						ControlRigEditMode->SetObjects(ControlRig, nullptr, WeakSequencer.Pin());
+					}
+
+					TempAnimSequence->MarkPendingKill();
+					if (WeakSequencer.IsValid())
+					{
+						WeakSequencer.Pin()->ObjectImplicitlyAdded(ControlRig);
+						WeakSequencer.Pin()->NotifyMovieSceneDataChanged(EMovieSceneDataChangeType::MovieSceneStructureItemAdded);
+					}
+					bResult = true;
+				}
+			}
+		}
+	}
+
+	Player->Stop();
+	World->DestroyActor(OutActor);
+
+	return bResult;
+}
+
+bool UControlRigSequencerEditorLibrary::LoadAnimSequenceIntoControlRigSection(UMovieSceneSection* MovieSceneSection, UAnimSequence* AnimSequence,  USkeleton* Skeleton,
+	FFrameNumber InStartFrame,bool bKeyReduce, float Tolerance)
+{
+	if (MovieSceneSection == nullptr || AnimSequence == nullptr || Skeleton == nullptr)
+	{
+		return false;;
+	}
+	UMovieScene* MovieScene = MovieSceneSection->GetTypedOuter<UMovieScene>();
+	if (MovieScene == nullptr)
+	{
+		return false;
+	}
+	if (UMovieSceneControlRigParameterSection* Section = Cast<UMovieSceneControlRigParameterSection>(MovieSceneSection))
+	{
+		return Section->LoadAnimSequenceIntoThisSection(AnimSequence, MovieScene, Skeleton, bKeyReduce, Tolerance, InStartFrame);
+	}
+	return false;
+}
+
+static bool LocalGetControlRigControlValues(IMovieScenePlayer* Player, UMovieSceneSequence* MovieSceneSequence, FMovieSceneSequenceIDRef Template, FMovieSceneSequenceTransform& RootToLocalTransform,
+	UControlRig* ControlRig, const FName& ControlName,
+	const TArray<FFrameNumber>& Frames, TArray<FRigControlValue>& OutValues)
+{
+	if (Player == nullptr || MovieSceneSequence == nullptr || ControlRig == nullptr)
+	{
+		return false;
+	}
+	if (ControlRig->FindControl(ControlName) == nullptr)
+	{
+		UE_LOG(LogControlRig, Error, TEXT("Can not find Control %s"), *(ControlName.ToString()));
+		return false;
+	}
+	if (UMovieScene* MovieScene = MovieSceneSequence->GetMovieScene())
+	{
+
+		FFrameRate TickResolution = MovieScene->GetTickResolution();
+
+		OutValues.SetNum(Frames.Num());
+		for (int32 Index = 0; Index < Frames.Num(); ++Index)
+		{
+			const FFrameNumber& FrameNumber = Frames[Index];
+			FFrameTime GlobalTime(FrameNumber);
+
+			FMovieSceneContext Context = FMovieSceneContext(FMovieSceneEvaluationRange(GlobalTime, TickResolution), Player->GetPlaybackStatus()).SetHasJumped(true);
+
+			Player->GetEvaluationTemplate().Evaluate(Context, *Player);
+			ControlRig->Evaluate_AnyThread();
+			OutValues[Index] = ControlRig->GetControlValue(ControlName);
+		}
+	}
+	return true;
+}
+
+static bool GetControlRigValues(ISequencer* Sequencer, UControlRig* ControlRig, const FName& ControlName,
+	const TArray<FFrameNumber>& Frames,  TArray<FRigControlValue>& OutValues)
+{
+	if (Sequencer->GetFocusedMovieSceneSequence())
+	{
+		FMovieSceneSequenceIDRef Template = Sequencer->GetFocusedTemplateID();
+		FMovieSceneSequenceTransform RootToLocalTransform;
+		return LocalGetControlRigControlValues(Sequencer, Sequencer->GetFocusedMovieSceneSequence(), Template, RootToLocalTransform,
+			ControlRig, ControlName, Frames, OutValues);
+
+	}
+	return false;
+}
+
+static bool GetControlRigValue(ISequencer* Sequencer, UControlRig* ControlRig, const FName& ControlName,
+	const FFrameNumber Frame, FRigControlValue& OutValue)
+{
+	if (Sequencer->GetFocusedMovieSceneSequence())
+	{
+		TArray<FFrameNumber> Frames;
+		Frames.Add(Frame);
+		TArray<FRigControlValue> OutValues;
+		FMovieSceneSequenceIDRef Template = Sequencer->GetFocusedTemplateID();
+		FMovieSceneSequenceTransform RootToLocalTransform;
+		bool bVal = LocalGetControlRigControlValues(Sequencer, Sequencer->GetFocusedMovieSceneSequence(), Template, RootToLocalTransform,
+			ControlRig, ControlName, Frames, OutValues);
+		if (bVal)
+		{
+			OutValue = OutValues[0];
+		}
+		return bVal;
+	}
+	return false;
+}
+
+static bool GetControlRigValues(UWorld* World, ULevelSequence* LevelSequence, UControlRig* ControlRig, const FName& ControlName,
+	const TArray<FFrameNumber>& Frames, TArray<FRigControlValue>& OutValues)
+{
+	if (LevelSequence)
+	{
+		ALevelSequenceActor* OutActor;
+		FMovieSceneSequencePlaybackSettings Settings;
+		FLevelSequenceCameraSettings CameraSettings;
+		FMovieSceneSequenceIDRef Template = MovieSceneSequenceID::Root;
+		FMovieSceneSequenceTransform RootToLocalTransform;
+		ULevelSequencePlayer* Player = ULevelSequencePlayer::CreateLevelSequencePlayer(World, LevelSequence, Settings, OutActor);
+		Player->Initialize(LevelSequence, World->PersistentLevel, Settings, CameraSettings);
+		Player->State.AssignSequence(MovieSceneSequenceID::Root, *LevelSequence, *Player);
+		return LocalGetControlRigControlValues(Player, LevelSequence, Template, RootToLocalTransform,
+			ControlRig, ControlName, Frames, OutValues);
+
+	}
+	return false;
+}
+
+
+float UControlRigSequencerEditorLibrary::GetLocalControlRigFloat(ULevelSequence* LevelSequence,UControlRig* ControlRig, FName ControlName, FFrameNumber Frame)
+{
+	float Value = 0.0f;
+	TWeakPtr<ISequencer> WeakSequencer = GetSequencerFromAsset(LevelSequence);
+
+	if (WeakSequencer.IsValid() && ControlRig)
+	{
+		FRigControlValue RigValue;
+		if (GetControlRigValue(WeakSequencer.Pin().Get(), ControlRig, ControlName, Frame, RigValue))
+		{
+			Value = RigValue.Get<float>();
+		}
+	}
+	return Value;
+}
+
+TArray<float> UControlRigSequencerEditorLibrary::GetLocalControlRigFloats(ULevelSequence* LevelSequence,UControlRig* ControlRig, FName ControlName, const TArray<FFrameNumber>& Frames)
+{
+	TArray<float> Values;
+	TWeakPtr<ISequencer> WeakSequencer = GetSequencerFromAsset(LevelSequence);
+
+	if (WeakSequencer.IsValid() && ControlRig)
+	{
+		TArray<FRigControlValue> RigValues;
+		if (GetControlRigValues(WeakSequencer.Pin().Get(), ControlRig, ControlName, Frames, RigValues))
+		{
+			Values.Reserve(RigValues.Num());
+			for (const FRigControlValue& RigValue : RigValues)
+			{
+				float Value = RigValue.Get<float>();
+				Values.Add(Value);
+			}
+		}
+	}
+	return Values;
+}
+
+void UControlRigSequencerEditorLibrary::SetLocalControlRigFloat(ULevelSequence* LevelSequence, UControlRig* ControlRig, FName ControlName, FFrameNumber Frame, float Value, bool bSetKey)
+{
+	if (LevelSequence == nullptr || ControlRig == nullptr)
+	{
+		return;
+	}
+	if (UMovieScene* MovieScene = LevelSequence->GetMovieScene())
+	{
+		FFrameRate TickResolution = MovieScene->GetTickResolution();
+		FRigControlModifiedContext Context;
+		Context.SetKey = bSetKey ? EControlRigSetKey::Always : EControlRigSetKey::DoNotCare;
+		Context.LocalTime = TickResolution.AsSeconds(FFrameTime(Frame));
+		ControlRig->SetControlValue<float>(ControlName, Value, true, FRigControlModifiedContext(EControlRigSetKey::Never, false));
+	}
+}
+
+void UControlRigSequencerEditorLibrary::SetLocalControlRigFloats(ULevelSequence* LevelSequence, UControlRig* ControlRig, FName ControlName, const TArray<FFrameNumber>& Frames, const TArray<float> Values)
+{
+	if (LevelSequence == nullptr || ControlRig == nullptr || (Frames.Num() != Values.Num()))
+	{
+		return;
+	}
+	if (UMovieScene* MovieScene = LevelSequence->GetMovieScene())
+	{
+		FFrameRate TickResolution = MovieScene->GetTickResolution();
+		FRigControlModifiedContext Context;
+		Context.SetKey = EControlRigSetKey::Always;
+		for (int32 Index = 0; Index < Frames.Num(); ++Index)
+		{
+			FFrameNumber Frame = Frames[Index];
+			float  Value = Values[Index];
+			Context.LocalTime = TickResolution.AsSeconds(FFrameTime(Frame));
+			ControlRig->SetControlValue<float>(ControlName, Value, true, FRigControlModifiedContext(EControlRigSetKey::Never, false));
+		}
+	}
+}
+
+
+bool UControlRigSequencerEditorLibrary::GetLocalControlRigBool(ULevelSequence* LevelSequence, UControlRig* ControlRig, FName ControlName, FFrameNumber Frame)
+{
+	bool Value = true;
+	TWeakPtr<ISequencer> WeakSequencer = GetSequencerFromAsset(LevelSequence);
+
+	if (WeakSequencer.IsValid() && ControlRig)
+	{
+		FRigControlValue RigValue;
+		if (GetControlRigValue(WeakSequencer.Pin().Get(), ControlRig, ControlName, Frame, RigValue))
+		{
+			Value = RigValue.Get<bool>();
+		}
+	}
+	return Value;
+}
+
+TArray<bool> UControlRigSequencerEditorLibrary::GetLocalControlRigBools(ULevelSequence* LevelSequence, UControlRig* ControlRig, FName ControlName, const TArray<FFrameNumber>& Frames)
+{
+	TArray<bool> Values;
+	TWeakPtr<ISequencer> WeakSequencer = GetSequencerFromAsset(LevelSequence);
+
+	if (WeakSequencer.IsValid() && ControlRig)
+	{
+		TArray<FRigControlValue> RigValues;
+		if (GetControlRigValues(WeakSequencer.Pin().Get(), ControlRig, ControlName, Frames, RigValues))
+		{
+			Values.Reserve(RigValues.Num());
+			for (const FRigControlValue& RigValue : RigValues)
+			{
+				bool Value = RigValue.Get<bool>();
+				Values.Add(Value);
+			}
+		}
+	}
+	return Values;
+}
+
+void UControlRigSequencerEditorLibrary::SetLocalControlRigBool(ULevelSequence* LevelSequence, UControlRig* ControlRig, FName ControlName, FFrameNumber Frame, bool Value, bool bSetKey)
+{
+	if (LevelSequence == nullptr || ControlRig == nullptr)
+	{
+		return;
+	}
+	if (UMovieScene* MovieScene = LevelSequence->GetMovieScene())
+	{
+		FFrameRate TickResolution = MovieScene->GetTickResolution();
+		FRigControlModifiedContext Context;
+		Context.SetKey = bSetKey ? EControlRigSetKey::Always : EControlRigSetKey::DoNotCare;
+		Context.LocalTime = TickResolution.AsSeconds(FFrameTime(Frame));
+		ControlRig->SetControlValue<bool>(ControlName, Value, true, FRigControlModifiedContext(EControlRigSetKey::Never, false));
+	}
+}
+
+void UControlRigSequencerEditorLibrary::SetLocalControlRigBools(ULevelSequence* LevelSequence, UControlRig* ControlRig, FName ControlName, const TArray<FFrameNumber>& Frames, const TArray<bool> Values)
+{
+	if (LevelSequence == nullptr || ControlRig == nullptr || (Frames.Num() != Values.Num()))
+	{
+		return;
+	}
+	if (UMovieScene* MovieScene = LevelSequence->GetMovieScene())
+	{
+		FFrameRate TickResolution = MovieScene->GetTickResolution();
+		FRigControlModifiedContext Context;
+		Context.SetKey = EControlRigSetKey::Always;
+		for (int32 Index = 0; Index < Frames.Num(); ++Index)
+		{
+			FFrameNumber Frame = Frames[Index];
+			bool Value = Values[Index];
+			Context.LocalTime = TickResolution.AsSeconds(FFrameTime(Frame));
+			ControlRig->SetControlValue<bool>(ControlName, Value, true, FRigControlModifiedContext(EControlRigSetKey::Never, false));
+		}
+	}
+}
+
+int32 UControlRigSequencerEditorLibrary::GetLocalControlRigInt(ULevelSequence* LevelSequence, UControlRig* ControlRig, FName ControlName, FFrameNumber Frame)
+{
+	int32 Value = 0;
+	TWeakPtr<ISequencer> WeakSequencer = GetSequencerFromAsset(LevelSequence);
+
+	if (WeakSequencer.IsValid() && ControlRig)
+	{
+		FRigControlValue RigValue;
+		if (GetControlRigValue(WeakSequencer.Pin().Get(), ControlRig, ControlName, Frame, RigValue))
+		{
+			Value = RigValue.Get<int32>();
+		}
+	}
+	return Value;
+}
+
+TArray<int32> UControlRigSequencerEditorLibrary::GetLocalControlRigInts(ULevelSequence* LevelSequence, UControlRig* ControlRig, FName ControlName, const TArray<FFrameNumber>& Frames)
+{
+	TArray<int32> Values;
+	TWeakPtr<ISequencer> WeakSequencer = GetSequencerFromAsset(LevelSequence);
+
+	if (WeakSequencer.IsValid() && ControlRig)
+	{
+		TArray<FRigControlValue> RigValues;
+		if (GetControlRigValues(WeakSequencer.Pin().Get(), ControlRig, ControlName, Frames, RigValues))
+		{
+			Values.Reserve(RigValues.Num());
+			for (const FRigControlValue& RigValue : RigValues)
+			{
+				int32 Value = RigValue.Get<int32>();
+				Values.Add(Value);
+			}
+		}
+	}
+	return Values;
+}
+
+void UControlRigSequencerEditorLibrary::SetLocalControlRigInt(ULevelSequence* LevelSequence, UControlRig* ControlRig, FName ControlName, FFrameNumber Frame, int32 Value, bool bSetKey)
+{
+	if (LevelSequence == nullptr || ControlRig == nullptr)
+	{
+		return;
+	}
+	if (UMovieScene* MovieScene = LevelSequence->GetMovieScene())
+	{
+		FFrameRate TickResolution = MovieScene->GetTickResolution();
+		FRigControlModifiedContext Context;
+		Context.SetKey = bSetKey ? EControlRigSetKey::Always : EControlRigSetKey::DoNotCare;
+		Context.LocalTime = TickResolution.AsSeconds(FFrameTime(Frame));
+		ControlRig->SetControlValue<int32>(ControlName, Value, true, FRigControlModifiedContext(EControlRigSetKey::Never, false));
+	}
+}
+
+void UControlRigSequencerEditorLibrary::SetLocalControlRigInts(ULevelSequence* LevelSequence, UControlRig* ControlRig, FName ControlName, const TArray<FFrameNumber>& Frames, const TArray<int32> Values)
+{
+	if (LevelSequence == nullptr || ControlRig == nullptr || (Frames.Num() != Values.Num()))
+	{
+		return;
+	}
+	if (UMovieScene* MovieScene = LevelSequence->GetMovieScene())
+	{
+		FFrameRate TickResolution = MovieScene->GetTickResolution();
+		FRigControlModifiedContext Context;
+		Context.SetKey = EControlRigSetKey::Always;
+		for (int32 Index = 0; Index < Frames.Num(); ++Index)
+		{
+			FFrameNumber Frame = Frames[Index];
+			int32 Value = Values[Index];
+			Context.LocalTime = TickResolution.AsSeconds(FFrameTime(Frame));
+			ControlRig->SetControlValue<int32>(ControlName, Value, true, FRigControlModifiedContext(EControlRigSetKey::Never, false));
+		}
+	}
+}
+
+
+FVector2D UControlRigSequencerEditorLibrary::GetLocalControlRigVector2D(ULevelSequence* LevelSequence, UControlRig* ControlRig, FName ControlName, FFrameNumber Frame)
+{
+	FVector2D Value;
+	TWeakPtr<ISequencer> WeakSequencer = GetSequencerFromAsset(LevelSequence);
+
+	if (WeakSequencer.IsValid() && ControlRig)
+	{
+		FRigControlValue RigValue;
+		if (GetControlRigValue(WeakSequencer.Pin().Get(), ControlRig, ControlName, Frame, RigValue))
+		{
+			Value = RigValue.Get<FVector2D>();
+		}
+	}
+	return Value;
+}
+
+TArray<FVector2D> UControlRigSequencerEditorLibrary::GetLocalControlRigVector2Ds(ULevelSequence* LevelSequence, UControlRig* ControlRig, FName ControlName, const TArray<FFrameNumber>& Frames)
+{
+	TArray<FVector2D> Values;
+	TWeakPtr<ISequencer> WeakSequencer = GetSequencerFromAsset(LevelSequence);
+
+	if (WeakSequencer.IsValid() && ControlRig)
+	{
+		TArray<FRigControlValue> RigValues;
+		if (GetControlRigValues(WeakSequencer.Pin().Get(), ControlRig, ControlName, Frames, RigValues))
+		{
+			Values.Reserve(RigValues.Num());
+			for (const FRigControlValue& RigValue : RigValues)
+			{
+				FVector2D Value = RigValue.Get<FVector2D>();
+				Values.Add(Value);
+			}
+		}
+	}
+	return Values;
+
+}
+
+void UControlRigSequencerEditorLibrary::SetLocalControlRigVector2D(ULevelSequence* LevelSequence, UControlRig* ControlRig, FName ControlName, FFrameNumber Frame, FVector2D Value, bool bSetKey)
+{
+	if (LevelSequence == nullptr || ControlRig == nullptr)
+	{
+		return;
+	}
+	if (UMovieScene* MovieScene = LevelSequence->GetMovieScene())
+	{
+		FFrameRate TickResolution = MovieScene->GetTickResolution();
+		FRigControlModifiedContext Context;
+		Context.SetKey = bSetKey ? EControlRigSetKey::Always : EControlRigSetKey::DoNotCare;
+		Context.LocalTime = TickResolution.AsSeconds(FFrameTime(Frame));
+		ControlRig->SetControlValue<FVector2D>(ControlName, Value, true, FRigControlModifiedContext(EControlRigSetKey::Never, false));
+	}
+}
+
+void UControlRigSequencerEditorLibrary::SetLocalControlRigVector2Ds(ULevelSequence* LevelSequence, UControlRig* ControlRig, FName ControlName, const TArray<FFrameNumber>& Frames, const TArray<FVector2D> Values)
+{
+	if (LevelSequence == nullptr || ControlRig == nullptr || (Frames.Num() != Values.Num()))
+	{
+		return;
+	}
+	if (UMovieScene* MovieScene = LevelSequence->GetMovieScene())
+	{
+		FFrameRate TickResolution = MovieScene->GetTickResolution();
+		FRigControlModifiedContext Context;
+		Context.SetKey = EControlRigSetKey::Always;
+		for (int32 Index = 0; Index < Frames.Num(); ++Index)
+		{
+			FFrameNumber Frame = Frames[Index];
+			FVector2D Value = Values[Index];
+			Context.LocalTime = TickResolution.AsSeconds(FFrameTime(Frame));
+			ControlRig->SetControlValue<FVector2D>(ControlName, Value, true, FRigControlModifiedContext(EControlRigSetKey::Never, false));
+		}
+	}
+}
+
+
+FVector UControlRigSequencerEditorLibrary::GetLocalControlRigPosition(ULevelSequence* LevelSequence, UControlRig* ControlRig, FName ControlName, FFrameNumber Frame)
+{
+	FVector Value;
+	TWeakPtr<ISequencer> WeakSequencer = GetSequencerFromAsset(LevelSequence);
+
+	if (WeakSequencer.IsValid() && ControlRig)
+	{
+		FRigControlValue RigValue;
+		if (GetControlRigValue(WeakSequencer.Pin().Get(), ControlRig, ControlName, Frame, RigValue))
+		{
+			Value = RigValue.Get<FVector>();
+		}
+	}
+	return Value;
+}
+
+TArray<FVector> UControlRigSequencerEditorLibrary::GetLocalControlRigPositions(ULevelSequence* LevelSequence, UControlRig* ControlRig, FName ControlName, const TArray<FFrameNumber>& Frames)
+{
+	TArray<FVector> Values;
+	TWeakPtr<ISequencer> WeakSequencer = GetSequencerFromAsset(LevelSequence);
+
+	if (WeakSequencer.IsValid() && ControlRig)
+	{
+		TArray<FRigControlValue> RigValues;
+		if (GetControlRigValues(WeakSequencer.Pin().Get(), ControlRig, ControlName, Frames, RigValues))
+		{
+			Values.Reserve(RigValues.Num());
+			for (const FRigControlValue& RigValue : RigValues)
+			{
+				FVector Value = RigValue.Get<FVector>();
+				Values.Add(Value);
+			}
+		}
+	}
+	return Values;
+}
+
+void UControlRigSequencerEditorLibrary::SetLocalControlRigPosition(ULevelSequence* LevelSequence, UControlRig* ControlRig, FName ControlName, FFrameNumber Frame, FVector Value, bool bSetKey)
+{
+	if (LevelSequence == nullptr || ControlRig == nullptr)
+	{
+		return;
+	}
+	if (UMovieScene* MovieScene = LevelSequence->GetMovieScene())
+	{
+		FFrameRate TickResolution = MovieScene->GetTickResolution();
+		FRigControlModifiedContext Context;
+		Context.SetKey = bSetKey ? EControlRigSetKey::Always : EControlRigSetKey::DoNotCare;
+		Context.LocalTime = TickResolution.AsSeconds(FFrameTime(Frame));
+		ControlRig->SetControlValue<FVector>(ControlName, Value, true, FRigControlModifiedContext(EControlRigSetKey::Never, false));
+	}
+}
+
+void UControlRigSequencerEditorLibrary::SetLocalControlRigPositions(ULevelSequence* LevelSequence, UControlRig* ControlRig, FName ControlName, const TArray<FFrameNumber>& Frames, const TArray<FVector> Values)
+{
+	if (LevelSequence == nullptr || ControlRig == nullptr || (Frames.Num() != Values.Num()))
+	{
+		return;
+	}
+	if (UMovieScene* MovieScene = LevelSequence->GetMovieScene())
+	{
+		FFrameRate TickResolution = MovieScene->GetTickResolution();
+		FRigControlModifiedContext Context;
+		Context.SetKey = EControlRigSetKey::Always;
+		for (int32 Index = 0; Index < Frames.Num(); ++Index)
+		{
+			FFrameNumber Frame = Frames[Index];
+			FVector Value = Values[Index];
+			Context.LocalTime = TickResolution.AsSeconds(FFrameTime(Frame));
+			ControlRig->SetControlValue<FVector>(ControlName, Value, true, FRigControlModifiedContext(EControlRigSetKey::Never, false));
+		}
+	}
+}
+
+
+FRotator UControlRigSequencerEditorLibrary::GetLocalControlRigRotator(ULevelSequence* LevelSequence, UControlRig* ControlRig, FName ControlName, FFrameNumber Frame)
+{
+	FRotator Value;
+	TWeakPtr<ISequencer> WeakSequencer = GetSequencerFromAsset(LevelSequence);
+
+	if (WeakSequencer.IsValid() && ControlRig)
+	{
+		FRigControlValue RigValue;
+		if (GetControlRigValue(WeakSequencer.Pin().Get(), ControlRig, ControlName, Frame, RigValue))
+		{
+			Value = RigValue.Get<FRotator>();
+		}
+	}
+	return Value;
+}
+
+TArray<FRotator> UControlRigSequencerEditorLibrary::GetLocalControlRigRotators(ULevelSequence* LevelSequence, UControlRig* ControlRig, FName ControlName, const TArray<FFrameNumber>& Frames)
+{
+	TArray<FRotator> Values;
+	TWeakPtr<ISequencer> WeakSequencer = GetSequencerFromAsset(LevelSequence);
+
+	if (WeakSequencer.IsValid() && ControlRig)
+	{
+		TArray<FRigControlValue> RigValues;
+		if (GetControlRigValues(WeakSequencer.Pin().Get(), ControlRig, ControlName, Frames, RigValues))
+		{
+			Values.Reserve(RigValues.Num());
+			for (const FRigControlValue& RigValue : RigValues)
+			{
+				FRotator Value = RigValue.Get<FRotator>();
+				Values.Add(Value);
+			}
+		}
+	}
+	return Values;
+}
+
+void UControlRigSequencerEditorLibrary::SetLocalControlRigRotator(ULevelSequence* LevelSequence, UControlRig* ControlRig, FName ControlName, FFrameNumber Frame, FRotator Value, bool bSetKey)
+{
+	if (LevelSequence == nullptr || ControlRig == nullptr)
+	{
+		return;
+	}
+	if (UMovieScene* MovieScene = LevelSequence->GetMovieScene())
+	{
+		FFrameRate TickResolution = MovieScene->GetTickResolution();
+		FRigControlModifiedContext Context;
+		Context.SetKey = bSetKey ? EControlRigSetKey::Always : EControlRigSetKey::DoNotCare;
+		Context.LocalTime = TickResolution.AsSeconds(FFrameTime(Frame));
+		ControlRig->SetControlValue<FRotator>(ControlName, Value, true, FRigControlModifiedContext(EControlRigSetKey::Never, false));
+	}
+}
+
+void UControlRigSequencerEditorLibrary::SetLocalControlRigRotators(ULevelSequence* LevelSequence, UControlRig* ControlRig, FName ControlName, const TArray<FFrameNumber>& Frames, const TArray<FRotator> Values)
+{
+	if (LevelSequence == nullptr || ControlRig == nullptr || (Frames.Num() != Values.Num()))
+	{
+		return;
+	}
+	if (UMovieScene* MovieScene = LevelSequence->GetMovieScene())
+	{
+		FFrameRate TickResolution = MovieScene->GetTickResolution();
+		FRigControlModifiedContext Context;
+		Context.SetKey = EControlRigSetKey::Always;
+		for (int32 Index = 0; Index < Frames.Num(); ++Index)
+		{
+			FFrameNumber Frame = Frames[Index];
+			FRotator Value = Values[Index];
+			Context.LocalTime = TickResolution.AsSeconds(FFrameTime(Frame));
+			ControlRig->SetControlValue<FRotator>(ControlName, Value, true, FRigControlModifiedContext(EControlRigSetKey::Never, false));
+		}
+	}
+}
+
+
+FVector UControlRigSequencerEditorLibrary::GetLocalControlRigScale(ULevelSequence* LevelSequence, UControlRig* ControlRig, FName ControlName, FFrameNumber Frame)
+{
+	FVector Value;
+	TWeakPtr<ISequencer> WeakSequencer = GetSequencerFromAsset(LevelSequence);
+
+	if (WeakSequencer.IsValid() && ControlRig)
+	{
+		FRigControlValue RigValue;
+		if (GetControlRigValue(WeakSequencer.Pin().Get(), ControlRig, ControlName, Frame, RigValue))
+		{
+			Value = RigValue.Get<FVector>();
+		}
+	}
+	return Value;
+}
+
+TArray<FVector>UControlRigSequencerEditorLibrary::GetLocalControlRigScales(ULevelSequence* LevelSequence, UControlRig* ControlRig, FName ControlName, const TArray<FFrameNumber>& Frames)
+{
+	TArray<FVector> Values;
+	TWeakPtr<ISequencer> WeakSequencer = GetSequencerFromAsset(LevelSequence);
+
+	if (WeakSequencer.IsValid() && ControlRig)
+	{
+		TArray<FRigControlValue> RigValues;
+		if (GetControlRigValues(WeakSequencer.Pin().Get(), ControlRig, ControlName, Frames, RigValues))
+		{
+			Values.Reserve(RigValues.Num());
+			for (const FRigControlValue& RigValue : RigValues)
+			{
+				FVector Value = RigValue.Get<FVector>();
+				Values.Add(Value);
+			}
+		}
+	}
+	return Values;
+}
+
+void UControlRigSequencerEditorLibrary::SetLocalControlRigScale(ULevelSequence* LevelSequence, UControlRig* ControlRig, FName ControlName, FFrameNumber Frame, FVector Value, bool bSetKey)
+{
+	if (LevelSequence == nullptr || ControlRig == nullptr)
+	{
+		return;
+	}
+	if (UMovieScene* MovieScene = LevelSequence->GetMovieScene())
+	{
+		FFrameRate TickResolution = MovieScene->GetTickResolution();
+		FRigControlModifiedContext Context;
+		Context.SetKey = bSetKey ? EControlRigSetKey::Always : EControlRigSetKey::DoNotCare;
+		Context.LocalTime = TickResolution.AsSeconds(FFrameTime(Frame));
+		ControlRig->SetControlValue<FVector>(ControlName, Value, true, FRigControlModifiedContext(EControlRigSetKey::Never, false));
+	}
+}
+
+void UControlRigSequencerEditorLibrary::SetLocalControlRigScales(ULevelSequence* LevelSequence, UControlRig* ControlRig, FName ControlName, const TArray<FFrameNumber>& Frames, const TArray<FVector> Values)
+{
+	if (LevelSequence == nullptr || ControlRig == nullptr || (Frames.Num() != Values.Num()))
+	{
+		return;
+	}
+	if (UMovieScene* MovieScene = LevelSequence->GetMovieScene())
+	{
+		FFrameRate TickResolution = MovieScene->GetTickResolution();
+		FRigControlModifiedContext Context;
+		Context.SetKey = EControlRigSetKey::Always;
+		for (int32 Index = 0; Index < Frames.Num(); ++Index)
+		{
+			FFrameNumber Frame = Frames[Index];
+			FVector Value = Values[Index];
+			Context.LocalTime = TickResolution.AsSeconds(FFrameTime(Frame));
+			ControlRig->SetControlValue<FVector>(ControlName, Value, true, FRigControlModifiedContext(EControlRigSetKey::Never, false));
+		}
+	}
+}
+
+
+FEulerTransform UControlRigSequencerEditorLibrary::GetLocalControlRigEulerTransform(ULevelSequence* LevelSequence, UControlRig* ControlRig, FName ControlName, FFrameNumber Frame)
+{
+	FEulerTransform Value = FEulerTransform::Identity;
+	TWeakPtr<ISequencer> WeakSequencer = GetSequencerFromAsset(LevelSequence);
+
+	if (WeakSequencer.IsValid() && ControlRig)
+	{
+		FRigControlValue RigValue;
+		if (GetControlRigValue(WeakSequencer.Pin().Get(), ControlRig, ControlName, Frame, RigValue))
+		{
+			Value = RigValue.Get<FEulerTransform>();
+		}
+	}
+	return Value;
+}
+
+TArray<FEulerTransform> UControlRigSequencerEditorLibrary::GetLocalControlRigEulerTransforms(ULevelSequence* LevelSequence, UControlRig* ControlRig, FName ControlName, const TArray<FFrameNumber>& Frames)
+{
+	TArray<FEulerTransform> Values;
+	TWeakPtr<ISequencer> WeakSequencer = GetSequencerFromAsset(LevelSequence);
+
+	if (WeakSequencer.IsValid() && ControlRig)
+	{
+		TArray<FRigControlValue> RigValues;
+		if (GetControlRigValues(WeakSequencer.Pin().Get(), ControlRig, ControlName, Frames, RigValues))
+		{
+			Values.Reserve(RigValues.Num());
+			for (const FRigControlValue& RigValue : RigValues)
+			{
+				FEulerTransform Value = RigValue.Get<FEulerTransform>();
+				Values.Add(Value);
+			}
+		}
+	}
+	return Values;
+}
+
+void UControlRigSequencerEditorLibrary::SetLocalControlRigEulerTransform(ULevelSequence* LevelSequence, UControlRig* ControlRig, FName ControlName, FFrameNumber Frame, FEulerTransform Value, bool bSetKey)
+{
+	if (LevelSequence == nullptr || ControlRig == nullptr)
+	{
+		return;
+	}
+	if (UMovieScene* MovieScene = LevelSequence->GetMovieScene())
+	{
+		FFrameRate TickResolution = MovieScene->GetTickResolution();
+		FRigControlModifiedContext Context;
+		Context.SetKey = bSetKey ? EControlRigSetKey::Always : EControlRigSetKey::DoNotCare;
+		Context.LocalTime = TickResolution.AsSeconds(FFrameTime(Frame));
+		ControlRig->SetControlValue<FEulerTransform>(ControlName, Value, true, FRigControlModifiedContext(EControlRigSetKey::Never, false));
+	}
+}
+
+void UControlRigSequencerEditorLibrary::SetLocalControlRigEulerTransforms(ULevelSequence* LevelSequence, UControlRig* ControlRig, FName ControlName, const TArray<FFrameNumber>& Frames, const TArray<FEulerTransform> Values)
+{
+	if (LevelSequence == nullptr || ControlRig == nullptr || (Frames.Num() != Values.Num()))
+	{
+		return;
+	}
+	if (UMovieScene* MovieScene = LevelSequence->GetMovieScene())
+	{
+		FFrameRate TickResolution = MovieScene->GetTickResolution();
+		FRigControlModifiedContext Context;
+		Context.SetKey = EControlRigSetKey::Always;
+		for (int32 Index = 0; Index < Frames.Num(); ++Index)
+		{
+			FFrameNumber Frame = Frames[Index];
+			FEulerTransform Value = Values[Index];
+			Context.LocalTime = TickResolution.AsSeconds(FFrameTime(Frame));
+			ControlRig->SetControlValue<FEulerTransform>(ControlName, Value, true, FRigControlModifiedContext(EControlRigSetKey::Never, false));
+		}
+	}
+}
+
+
+FTransformNoScale UControlRigSequencerEditorLibrary::GetLocalControlRigTransformNoScale(ULevelSequence* LevelSequence, UControlRig* ControlRig, FName ControlName, FFrameNumber Frame)
+{
+	FTransformNoScale Value = FTransformNoScale::Identity;
+	TWeakPtr<ISequencer> WeakSequencer = GetSequencerFromAsset(LevelSequence);
+
+	if (WeakSequencer.IsValid() && ControlRig)
+	{
+		FRigControlValue RigValue;
+		if (GetControlRigValue(WeakSequencer.Pin().Get(), ControlRig, ControlName, Frame, RigValue))
+		{
+			Value = RigValue.Get<FTransformNoScale>();
+		}
+	}
+	return Value;
+}
+
+TArray<FTransformNoScale> UControlRigSequencerEditorLibrary::GetLocalControlRigTransformNoScales(ULevelSequence* LevelSequence, UControlRig* ControlRig, FName ControlName, const TArray<FFrameNumber>& Frames)
+{
+	TArray<FTransformNoScale> Values;
+	TWeakPtr<ISequencer> WeakSequencer = GetSequencerFromAsset(LevelSequence);
+
+	if (WeakSequencer.IsValid() && ControlRig)
+	{
+		TArray<FRigControlValue> RigValues;
+		if (GetControlRigValues(WeakSequencer.Pin().Get(), ControlRig, ControlName, Frames, RigValues))
+		{
+			Values.Reserve(RigValues.Num());
+			for (const FRigControlValue& RigValue : RigValues)
+			{
+				FTransformNoScale Value = RigValue.Get<FTransformNoScale>();
+				Values.Add(Value);
+			}
+		}
+	}
+	return Values;
+}
+
+void UControlRigSequencerEditorLibrary::SetLocalControlRigTransformNoScale(ULevelSequence* LevelSequence, UControlRig* ControlRig, FName ControlName, FFrameNumber Frame, FTransformNoScale Value, bool bSetKey)
+{
+	if (LevelSequence == nullptr || ControlRig == nullptr)
+	{
+		return;
+	}
+	if (UMovieScene* MovieScene = LevelSequence->GetMovieScene())
+	{
+		FFrameRate TickResolution = MovieScene->GetTickResolution();
+		FRigControlModifiedContext Context;
+		Context.SetKey = bSetKey ? EControlRigSetKey::Always : EControlRigSetKey::DoNotCare;
+		Context.LocalTime = TickResolution.AsSeconds(FFrameTime(Frame));
+		ControlRig->SetControlValue<FTransformNoScale>(ControlName, Value, true, FRigControlModifiedContext(EControlRigSetKey::Never, false));
+	}
+}
+
+void UControlRigSequencerEditorLibrary::SetLocalControlRigTransformNoScales(ULevelSequence* LevelSequence, UControlRig* ControlRig, FName ControlName, const TArray<FFrameNumber>& Frames, const TArray<FTransformNoScale> Values)
+{
+	if (LevelSequence == nullptr || ControlRig == nullptr || (Frames.Num() != Values.Num()))
+	{
+		return;
+	}
+	if (UMovieScene* MovieScene = LevelSequence->GetMovieScene())
+	{
+		FFrameRate TickResolution = MovieScene->GetTickResolution();
+		FRigControlModifiedContext Context;
+		Context.SetKey = EControlRigSetKey::Always;
+		for (int32 Index = 0; Index < Frames.Num(); ++Index)
+		{
+			FFrameNumber Frame = Frames[Index];
+			FTransformNoScale Value = Values[Index];
+			Context.LocalTime = TickResolution.AsSeconds(FFrameTime(Frame));
+			ControlRig->SetControlValue<FTransformNoScale>(ControlName, Value, true, FRigControlModifiedContext(EControlRigSetKey::Never, false));
+		}
+	}
+}
+
+
+FTransform UControlRigSequencerEditorLibrary::GetLocalControlRigTransform(ULevelSequence* LevelSequence, UControlRig* ControlRig, FName ControlName, FFrameNumber Frame)
+{
+	FTransform Value = FTransform::Identity;
+	TWeakPtr<ISequencer> WeakSequencer = GetSequencerFromAsset(LevelSequence);
+	if (WeakSequencer.IsValid() && ControlRig)
+	{
+		FRigControlValue RigValue;
+		if (GetControlRigValue(WeakSequencer.Pin().Get(), ControlRig, ControlName, Frame, RigValue))
+		{
+			Value = RigValue.Get<FTransform>();
+		}
+	}
+	return Value;
+}
+
+TArray<FTransform> UControlRigSequencerEditorLibrary::GetLocalControlRigTransforms(ULevelSequence* LevelSequence, UControlRig* ControlRig, FName ControlName, const TArray<FFrameNumber>& Frames)
+{
+	TArray<FTransform> Values;
+	TWeakPtr<ISequencer> WeakSequencer = GetSequencerFromAsset(LevelSequence);
+
+	if (WeakSequencer.IsValid() && ControlRig)
+	{
+		TArray<FRigControlValue> RigValues;
+		if (GetControlRigValues(WeakSequencer.Pin().Get(), ControlRig, ControlName, Frames, RigValues))
+		{
+			Values.Reserve(RigValues.Num());
+			for (const FRigControlValue& RigValue : RigValues)
+			{
+				FTransform Value = RigValue.Get<FTransform>();
+				Values.Add(Value);
+			}
+		}
+	}
+	return Values;
+}
+
+void UControlRigSequencerEditorLibrary::SetLocalControlRigTransform(ULevelSequence* LevelSequence, UControlRig* ControlRig, FName ControlName, FFrameNumber Frame, FTransform Value, bool bSetKey)
+{
+	if (LevelSequence == nullptr || ControlRig == nullptr)
+	{
+		return;
+	}
+	if (UMovieScene* MovieScene = LevelSequence->GetMovieScene())
+	{
+		FFrameRate TickResolution = MovieScene->GetTickResolution();
+		FRigControlModifiedContext Context;
+		Context.SetKey = bSetKey ? EControlRigSetKey::Always : EControlRigSetKey::DoNotCare;
+		Context.LocalTime = TickResolution.AsSeconds(FFrameTime(Frame));
+		ControlRig->SetControlValue<FTransform>(ControlName, Value, true, FRigControlModifiedContext(EControlRigSetKey::Never, false));
+	}
+}
+
+void UControlRigSequencerEditorLibrary::SetLocalControlRigTransforms(ULevelSequence* LevelSequence, UControlRig* ControlRig, FName ControlName, const TArray<FFrameNumber>& Frames, const TArray<FTransform> Values)
+{
+	if (LevelSequence == nullptr || ControlRig == nullptr || (Frames.Num() != Values.Num()))
+	{
+		return;
+	}
+	if (UMovieScene* MovieScene = LevelSequence->GetMovieScene())
+	{
+		FFrameRate TickResolution = MovieScene->GetTickResolution();
+		FRigControlModifiedContext Context;
+		Context.SetKey = EControlRigSetKey::Always;
+		for (int32 Index = 0; Index < Frames.Num(); ++Index)
+		{
+			FFrameNumber Frame = Frames[Index];
+			FTransform Value = Values[Index];
+			Context.LocalTime = TickResolution.AsSeconds(FFrameTime(Frame));
+			ControlRig->SetControlValue<FTransform>(ControlName, Value, true, FRigControlModifiedContext(EControlRigSetKey::Never, false));
+		}
+	}
+}
+
+
+#undef LOCTEXT_NAMESPACE
diff --git a/Engine/Plugins/Experimental/ControlRig/Source/ControlRigEditor/Private/EditMode/ControlRigEditMode.cpp b/Engine/Plugins/Experimental/ControlRig/Source/ControlRigEditor/Private/EditMode/ControlRigEditMode.cpp
--- a/Engine/Plugins/Experimental/ControlRig/Source/ControlRigEditor/Private/EditMode/ControlRigEditMode.cpp
+++ b/Engine/Plugins/Experimental/ControlRig/Source/ControlRigEditor/Private/EditMode/ControlRigEditMode.cpp
@@ -1186,13 +1186,15 @@ bool FControlRigEditMode::InputDelta(FEditorViewportClient* InViewportClient, FV
 	const EAxisList::Type CurrentAxis = InViewportClient->GetCurrentWidgetAxis();
 	const ECoordSystem CoordSystem = InViewportClient->GetWidgetCoordSystemSpace();
 
-	if (InteractionScope != nullptr && bMouseButtonDown && !bCtrlDown && !bShiftDown && CurrentAxis != EAxisList::None)
-	{
-		const bool bDoRotation = !Rot.IsZero() && (WidgetMode == UE::Widget::WM_Rotate || WidgetMode == UE::Widget::WM_TranslateRotateZ);
-		const bool bDoTranslation = !Drag.IsZero() && (WidgetMode == UE::Widget::WM_Translate || WidgetMode == UE::Widget::WM_TranslateRotateZ);
-		const bool bDoScale = !Scale.IsZero() && WidgetMode == UE::Widget::WM_Scale;
+	const bool bDoRotation = !Rot.IsZero() && (WidgetMode == UE::Widget::WM_Rotate || WidgetMode == UE::Widget::WM_TranslateRotateZ);
+	const bool bDoTranslation = !Drag.IsZero() && (WidgetMode == UE::Widget::WM_Translate || WidgetMode == UE::Widget::WM_TranslateRotateZ);
+	const bool bDoScale = !Scale.IsZero() && WidgetMode == UE::Widget::WM_Scale;
+
 
-		if (AreRigElementsSelected(FRigElementTypeHelper::ToMask(ERigElementType::Control)))
+	if (InteractionScope != nullptr && bMouseButtonDown && !bCtrlDown && !bShiftDown && CurrentAxis != EAxisList::None
+		&& (bDoRotation || bDoTranslation || bDoScale))
+	{
+			if (AreRigElementsSelected(FRigElementTypeHelper::ToMask(ERigElementType::Control)))
 		{
 			FTransform ComponentTransform = GetHostingSceneComponentTransform();
 			bool bDoLocal = (CoordSystem == ECoordSystem::COORD_Local && Settings && Settings->bLocalTransformsInEachLocalSpace);
diff --git a/Engine/Plugins/Experimental/ControlRig/Source/ControlRigEditor/Private/EditMode/SControlRigSnapper.cpp b/Engine/Plugins/Experimental/ControlRig/Source/ControlRigEditor/Private/EditMode/SControlRigSnapper.cpp
--- a/Engine/Plugins/Experimental/ControlRig/Source/ControlRigEditor/Private/EditMode/SControlRigSnapper.cpp
+++ b/Engine/Plugins/Experimental/ControlRig/Source/ControlRigEditor/Private/EditMode/SControlRigSnapper.cpp
@@ -26,6 +26,7 @@
 #include "Widgets/Text/STextBlock.h"
 #include "Widgets/Views/STableRow.h"
 #include "Widgets/Input/SButton.h"
+#include "Tools/ControlRigSnapSettings.h"
 
 #define LOCTEXT_NAMESPACE "ControlRigSnapper"
 
@@ -374,7 +375,8 @@ FReply SControlRigSnapper::OnEndFrameClicked()
 
 FReply SControlRigSnapper::OnSnapAnimationClicked()
 {
-	Snapper.SnapIt(StartFrame, EndFrame, ActorToSnap, ParentToSnap);
+	const UControlRigSnapSettings* SnapSettings = GetDefault<UControlRigSnapSettings>();
+	Snapper.SnapIt(StartFrame, EndFrame, ActorToSnap, ParentToSnap,SnapSettings);
 	return FReply::Handled();
 }
 
diff --git a/Engine/Plugins/Experimental/ControlRig/Source/ControlRigEditor/Private/Sequencer/ControlRigParameterTrackEditor.cpp b/Engine/Plugins/Experimental/ControlRig/Source/ControlRigEditor/Private/Sequencer/ControlRigParameterTrackEditor.cpp
--- a/Engine/Plugins/Experimental/ControlRig/Source/ControlRigEditor/Private/Sequencer/ControlRigParameterTrackEditor.cpp
+++ b/Engine/Plugins/Experimental/ControlRig/Source/ControlRigEditor/Private/Sequencer/ControlRigParameterTrackEditor.cpp
@@ -361,6 +361,14 @@ void FControlRigParameterTrackEditor::ObjectImplicitlyAdded(UObject* InObject)
 	}
 }
 
+void FControlRigParameterTrackEditor::ObjectImplicitlyRemoved(UObject* InObject)
+{
+	UControlRig* ControlRig = Cast<UControlRig>(InObject);
+	if (ControlRig)
+	{
+		UnbindControlRig(ControlRig);
+	}
+}
 
 void FControlRigParameterTrackEditor::OnRelease()
 {
@@ -1298,8 +1306,8 @@ void FControlRigParameterTrackEditor::AddControlRigFromComponent(FGuid InGuid)
 			{
 				AddControlRig(CR->GetClass(), BoundActor, InGuid, CR);
 			}
-			
 		}
+			
 	}
 }
 
@@ -3280,3 +3288,5 @@ void FControlRigParameterSection::OnAnimationAssetEnterPressedForFK(const TArray
 }
 
 #undef LOCTEXT_NAMESPACE
+
+
diff --git a/Engine/Plugins/Experimental/ControlRig/Source/ControlRigEditor/Private/Sequencer/ControlRigParameterTrackEditor.h b/Engine/Plugins/Experimental/ControlRig/Source/ControlRigEditor/Private/Sequencer/ControlRigParameterTrackEditor.h
--- a/Engine/Plugins/Experimental/ControlRig/Source/ControlRigEditor/Private/Sequencer/ControlRigParameterTrackEditor.h
+++ b/Engine/Plugins/Experimental/ControlRig/Source/ControlRigEditor/Private/Sequencer/ControlRigParameterTrackEditor.h
@@ -61,6 +61,7 @@ public:
 	virtual void OnAddTransformKeysForSelectedObjects(EMovieSceneTransformChannel Channel);
 	virtual bool HasTransformKeyOverridePriority() const override;
 	virtual void ObjectImplicitlyAdded(UObject* InObject)  override;
+	virtual void ObjectImplicitlyRemoved(UObject* InObject)  override;
 	virtual void BuildTrackContextMenu(FMenuBuilder& MenuBuilder, UMovieSceneTrack* InTrack) override;
 	virtual bool HandleAssetAdded(UObject* Asset, const FGuid& TargetObjectGuid) override;
 
diff --git a/Engine/Plugins/Experimental/ControlRig/Source/ControlRigEditor/Private/Tools/ControlRigSnapper.cpp b/Engine/Plugins/Experimental/ControlRig/Source/ControlRigEditor/Private/Tools/ControlRigSnapper.cpp
--- a/Engine/Plugins/Experimental/ControlRig/Source/ControlRigEditor/Private/Tools/ControlRigSnapper.cpp
+++ b/Engine/Plugins/Experimental/ControlRig/Source/ControlRigEditor/Private/Tools/ControlRigSnapper.cpp
@@ -27,6 +27,8 @@
 #include "Components/SkeletalMeshComponent.h"
 #include "ScopedTransaction.h"
 #include "MovieSceneToolHelpers.h"
+#include "LevelSequencePlayer.h"
+#include "LevelSequenceActor.h"
 
 #define LOCTEXT_NAMESPACE "ControlRigSnapper"
 
@@ -109,8 +111,9 @@ TWeakPtr<ISequencer> FControlRigSnapper::GetSequencer()
 	return WeakSequencer;
 }
 
-static bool GetControlRigControlTransforms(ISequencer* Sequencer,  UControlRig* ControlRig, const FName& ControlName, 
-	const TArray<FFrameNumber> &Frames, const TArray<FTransform>& ParentTransforms,TArray<FTransform>& OutTransforms)
+static bool LocalGetControlRigControlTransforms(IMovieScenePlayer* Player, UMovieSceneSequence* MovieSceneSequence, FMovieSceneSequenceIDRef Template, FMovieSceneSequenceTransform& RootToLocalTransform,
+	UControlRig* ControlRig, const FName& ControlName,
+	const TArray<FFrameNumber>& Frames, const TArray<FTransform>& ParentTransforms, TArray<FTransform>& OutTransforms)
 {
 	if (Frames.Num() > ParentTransforms.Num())
 	{
@@ -122,31 +125,63 @@ static bool GetControlRigControlTransforms(ISequencer* Sequencer,  UControlRig*
 		UE_LOG(LogControlRig, Error, TEXT("Can not find Control %s"), *(ControlName.ToString()));
 		return false;
 	}
-	if (UMovieScene* MovieScene = Sequencer->GetFocusedMovieSceneSequence()->GetMovieScene())
+	if (UMovieScene* MovieScene = MovieSceneSequence->GetMovieScene())
 	{
 
-		FMovieSceneSequenceIDRef Template = Sequencer->GetFocusedTemplateID();
-		FMovieSceneSequenceTransform RootToLocalTransform;
-
 		FFrameRate TickResolution = MovieScene->GetTickResolution();
 		FFrameRate DisplayRate = MovieScene->GetDisplayRate();
 
 		OutTransforms.SetNum(Frames.Num());
-		for (int32 Index = 0; Index <Frames.Num(); ++Index)
+		for (int32 Index = 0; Index < Frames.Num(); ++Index)
 		{
 			const FFrameNumber& FrameNumber = Frames[Index];
 			FFrameTime GlobalTime(FrameNumber);
 
-			FMovieSceneContext Context = FMovieSceneContext(FMovieSceneEvaluationRange(GlobalTime, TickResolution), Sequencer->GetPlaybackStatus()).SetHasJumped(true);
+			FMovieSceneContext Context = FMovieSceneContext(FMovieSceneEvaluationRange(GlobalTime, TickResolution), Player->GetPlaybackStatus()).SetHasJumped(true);
 
-			Sequencer->GetEvaluationTemplate().Evaluate(Context, *Sequencer);
+			Player->GetEvaluationTemplate().Evaluate(Context, *Player);
 			ControlRig->Evaluate_AnyThread();
-			OutTransforms[Index] =  ControlRig->GetControlGlobalTransform(ControlName) * ParentTransforms[Index];
+			OutTransforms[Index] = ControlRig->GetControlGlobalTransform(ControlName) * ParentTransforms[Index];
 		}
 	}
 	return true;
 }
 
+bool FControlRigSnapper::GetControlRigControlTransforms(ISequencer* Sequencer,  UControlRig* ControlRig, const FName& ControlName,
+	const TArray<FFrameNumber> &Frames, const TArray<FTransform>& ParentTransforms,TArray<FTransform>& OutTransforms)
+{
+	if (Sequencer->GetFocusedMovieSceneSequence())
+	{
+		FMovieSceneSequenceIDRef Template = Sequencer->GetFocusedTemplateID();
+		FMovieSceneSequenceTransform RootToLocalTransform;
+		return LocalGetControlRigControlTransforms(Sequencer, Sequencer->GetFocusedMovieSceneSequence(), Template, RootToLocalTransform,
+			ControlRig, ControlName, Frames, ParentTransforms, OutTransforms);
+	
+	}
+	return false;
+}
+
+bool FControlRigSnapper::GetControlRigControlTransforms(UWorld* World,ULevelSequence* LevelSequence, UControlRig* ControlRig, const FName& ControlName,
+	const TArray<FFrameNumber>& Frames, const TArray<FTransform>& ParentTransforms, TArray<FTransform>& OutTransforms)
+{
+	if (LevelSequence)
+	{
+		ALevelSequenceActor* OutActor;
+		FMovieSceneSequencePlaybackSettings Settings;
+		FLevelSequenceCameraSettings CameraSettings;
+		FMovieSceneSequenceIDRef Template = MovieSceneSequenceID::Root;
+		FMovieSceneSequenceTransform RootToLocalTransform;
+		ULevelSequencePlayer* Player = ULevelSequencePlayer::CreateLevelSequencePlayer(World, LevelSequence, Settings, OutActor);
+		Player->Initialize(LevelSequence, World->PersistentLevel, Settings, CameraSettings);
+		Player->State.AssignSequence(MovieSceneSequenceID::Root, *LevelSequence, *Player);
+		bool Success = LocalGetControlRigControlTransforms(Player, LevelSequence, Template, RootToLocalTransform,
+			ControlRig, ControlName, Frames, ParentTransforms, OutTransforms);
+		World->DestroyActor(OutActor);
+		return Success;
+	}
+	return false;
+}
+
 //Matinee version of this doesn't actually set the key it only adds..sigh
 static void SetOrAddKey(TMovieSceneChannelData<FMovieSceneFloatValue>& ChannelData, FFrameNumber Time, float Value)
 {
@@ -176,9 +211,9 @@ struct FGuidAndActor
 	FGuidAndActor(FGuid InGuid, AActor* InActor) : Guid(InGuid), Actor(InActor) {};
 	FGuid Guid;
 	AActor* Actor;
-	
+
 	bool SetTransform(ISequencer* Sequencer, const TArray<FFrameNumber>& Frames,
-		const TArray<FTransform>& WorldTransformsToSnapTo)
+		const TArray<FTransform>& WorldTransformsToSnapTo, const UControlRigSnapSettings* SnapSettings)
 	{
 		if (!Sequencer || !Sequencer->GetFocusedMovieSceneSequence())
 		{
@@ -224,7 +259,6 @@ struct FGuidAndActor
 				Transform = FTransform::Identity;
 			}
 		}
-		const UControlRigSnapSettings* SnapSettings = GetDefault<UControlRigSnapSettings>();
 
 		TArrayView<FMovieSceneFloatChannel*> Channels = TransformSection->GetChannelProxy().GetChannels<FMovieSceneFloatChannel>();
 		for (int32 Index = 0; Index < Frames.Num(); ++Index)
@@ -392,9 +426,10 @@ static bool CalculateWorldTransformsFromParents(ISequencer* Sequencer, const FCo
 				MovieSceneToolHelpers::GetActorWorldTransforms(Sequencer, ActorSelection, Frames, ParentTransforms);
 
 				//just do first for now but may do average later.
+				FControlRigSnapper Snapper;
 				for (const FName& Name : ControlRigAndSelection.ControlNames)
 				{
-					GetControlRigControlTransforms(Sequencer, ControlRig, Name, Frames, ParentTransforms, OutParentWorldTransforms);
+					Snapper.GetControlRigControlTransforms(Sequencer, ControlRig, Name, Frames, ParentTransforms, OutParentWorldTransforms);
 					return true;
 				}
 			}
@@ -410,13 +445,11 @@ static bool CalculateWorldTransformsFromParents(ISequencer* Sequencer, const FCo
 
 
 bool FControlRigSnapper::SnapIt(FFrameNumber StartFrame, FFrameNumber EndFrame,const FControlRigSnapperSelection& ActorToSnap,
-	const FControlRigSnapperSelection& ParentToSnap)
+	const FControlRigSnapperSelection& ParentToSnap, const UControlRigSnapSettings* SnapSettings)
 {
 	TWeakPtr<ISequencer> InSequencer = GetSequencer();
 	if (InSequencer.IsValid() && InSequencer.Pin()->GetFocusedMovieSceneSequence() && ActorToSnap.IsValid())
 	{
-		const UControlRigSnapSettings* SnapSettings = GetDefault<UControlRigSnapSettings>();
-
 		FScopedTransaction ScopedTransaction(LOCTEXT("SnapAnimation", "Snap Animation"));
 
 		ISequencer* Sequencer = InSequencer.Pin().Get();
@@ -487,7 +520,7 @@ bool FControlRigSnapper::SnapIt(FFrameNumber StartFrame, FFrameNumber EndFrame,c
 				}
 
 			}
-			GuidActor.SetTransform(Sequencer, Frames, WorldTransformToSnap);
+			GuidActor.SetTransform(Sequencer, Frames, WorldTransformToSnap,SnapSettings);
 		}
 
 		//Now do Control Rigs
@@ -519,7 +552,7 @@ bool FControlRigSnapper::SnapIt(FFrameNumber StartFrame, FFrameNumber EndFrame,c
 					FActorForWorldTransforms ControlRigActorSelection;
 					ControlRigActorSelection.Actor = Actor;
 					MovieSceneToolHelpers::GetActorWorldTransforms(Sequencer, ControlRigActorSelection, Frames, ControlRigParentWorldTransforms);
-
+					FControlRigSnapper Snapper;
 					for (const FName& Name : ControlRigAndSelection.ControlNames)
 					{
 						TArray<FFrameNumber> OneFrame;
@@ -531,7 +564,7 @@ bool FControlRigSnapper::SnapIt(FFrameNumber StartFrame, FFrameNumber EndFrame,c
 						{
 							OneFrame[0] = Frames[0];
 							CurrentParentWorldTransform[0] = ControlRigParentWorldTransforms[0];
-							GetControlRigControlTransforms(Sequencer, ControlRig, Name, OneFrame, CurrentParentWorldTransform, CurrentControlRigTransform);
+							Snapper.GetControlRigControlTransforms(Sequencer, ControlRig, Name, OneFrame, CurrentParentWorldTransform, CurrentControlRigTransform);
 							if (bSnapToFirstFrameNotParents)
 							{
 								for (FTransform& Transform : WorldTransformToSnap)
@@ -579,8 +612,6 @@ bool FControlRigSnapper::SnapIt(FFrameNumber StartFrame, FFrameNumber EndFrame,c
 							const FFrameNumber& FrameNumber = Frames[Index];
 							Context.LocalTime = TickResolution.AsSeconds(FFrameTime(FrameNumber));
 							FTransform GlobalTransform = WorldTransformToSnap[Index].GetRelativeTransform(ControlRigParentWorldTransforms[Index]);
-							//TODO need interrogator to get control rig values OR use different functon
-							GlobalTransform.SetScale3D(FVector(1.0f, 1.0f, 1.0f)); //note this is scale is off...
 							ControlRig->SetControlGlobalTransform(Name, GlobalTransform, true, Context);
 						}
 					}
diff --git a/Engine/Plugins/Experimental/ControlRig/Source/ControlRigEditor/Public/ControlRigSequencerEditorLibrary.h b/Engine/Plugins/Experimental/ControlRig/Source/ControlRigEditor/Public/ControlRigSequencerEditorLibrary.h
new file mode 100644
--- /dev/null
+++ b/Engine/Plugins/Experimental/ControlRig/Source/ControlRigEditor/Public/ControlRigSequencerEditorLibrary.h
@@ -0,0 +1,681 @@
+// Copyright Epic Games, Inc. All Rights Reserved.
+
+/**
+ *
+ * Control Rig Sequencer Exposure
+ */
+
+#pragma once
+
+#include "CoreMinimal.h"
+#include "UObject/ObjectMacros.h"
+#include "SequencerBindingProxy.h"
+#include "ControlRig.h"
+#include "Tools/ControlRigSnapper.h"
+#include "TransformNoScale.h"
+#include "EulerTransform.h"
+#include "Tools/ControlRigSnapSettings.h"
+#include "ControlRigSequencerEditorLibrary.generated.h"
+
+class ULevelSequence;
+
+USTRUCT(BlueprintType)
+struct FControlRigSequencerBindingProxy
+{
+	GENERATED_BODY()
+
+	FControlRigSequencerBindingProxy()
+		: ControlRig(nullptr)
+	{}
+
+	FControlRigSequencerBindingProxy(const FSequencerBindingProxy& InProxy, UControlRig* InControlRig)
+		: Proxy(InProxy)
+		, ControlRig(InControlRig)
+	{}
+
+	UPROPERTY()
+	FSequencerBindingProxy Proxy;
+
+	UPROPERTY(BlueprintReadOnly, Category = ControlRig)
+	UControlRig* ControlRig;
+};
+
+/**
+* This is a set of helper functions to access various parts of the Sequencer and Control Rig API via Python and Blueprints.
+*/
+UCLASS(meta=(Transient, ScriptName="ControlRigSequencerLibrary"))
+class CONTROLRIGEDITOR_API UControlRigSequencerEditorLibrary : public UBlueprintFunctionLibrary
+{
+
+	public:
+	GENERATED_BODY()
+
+public:
+
+	/**
+	* Get all of the visible control rigs in the level
+	* @return returns list of visible Control Rigs
+	*/
+	UFUNCTION(BlueprintCallable, Category = "Editor Scripting | Sequencer Tools | Control Rig")
+	static TArray<UControlRig*> GetVisibleControlRigs();
+
+	/**
+	* Get all of the control rigs and their bindings in the level sequence
+	* @param LevelSequence The movie scene sequence to look for Control Rigs.
+	* @return returns list of Control Rigs in the level sequence.
+	*/
+	UFUNCTION(BlueprintCallable, Category = "Editor Scripting | Sequencer Tools | Control Rig")
+	static TArray<FControlRigSequencerBindingProxy> GetControlRigs(ULevelSequence* LevelSequence);
+	
+	/**
+	* Find or create a Control Rig track of a specific class based upon the binding
+	* @param LevelSequence The LevelSequence to find or create
+	* @param ControlRigClass The class of the Control Rig
+	* @param InBinding The binding (actor or component binding) to find or create the Control Rig track
+	* @return returns Return the found or created track
+	*/
+	UFUNCTION(BlueprintCallable, Category = "Editor Scripting | Sequencer Tools | Control Rig")
+	static UMovieSceneTrack* FindOrCreateControlRigTrack(ULevelSequence* LevelSequence, const UClass* ControlRigClass, const FSequencerBindingProxy& InBinding);
+
+	/**
+	* Find or create a Control Rig Component
+	* @param LevelSequence The LevelSequence to find or create
+	* @param InBinding The binding (actor or component binding) to find or create the Control Rig tracks
+	* @return returns Find array of component Control Rigs that were found or created
+	*/
+	UFUNCTION(BlueprintCallable, Category = "Editor Scripting | Sequencer Tools | Control Rig")
+	static TArray<UMovieSceneTrack*> FindOrCreateControlRigComponentTrack(ULevelSequence* LevelSequence, const FSequencerBindingProxy& InBinding);
+	
+	/**
+	* Load anim sequence into this control rig section
+	* @param MovieSceneSection The MovieSceneSectionto load into
+	* @param AnimSequence The Sequence to load
+	* @param MovieScene The MovieScene getting loaded into
+	* @param Skeleton The Skeleton of the skelmesh component getting loaded into.
+	* @param InStartFrame Frame to insert the animation
+	* @param bKeyReduce If true do key reduction based upon Tolerance, if false don't
+	* @param Tolerance If reducing keys, tolerance about which keys will be removed, smaller tolerance, more keys usually.
+
+	* @return returns True if successful, False otherwise
+	*/
+	UFUNCTION(BlueprintCallable, Category = "Editor Scripting | Sequencer Tools | Control Rig")
+	static bool LoadAnimSequenceIntoControlRigSection(UMovieSceneSection* MovieSceneSection, UAnimSequence* AnimSequence, USkeleton* Skeleton,
+		FFrameNumber InStartFrame, bool bKeyReduce = false, float Tolerance = 0.001f);
+
+	/**
+	* Bake the current animation in the binding to a Control Rig track
+	* @param World The active world
+	* @param LevelSequence The LevelSequence we are baking
+	* @param ControlRigClass The class of the Control Rig
+	* @param ExportOptions Export options for creating an animation sequence
+	* @param bKeyReduce If true do key reduction based upon Tolerance, if false don't
+	* @param Tolerance If reducing keys, tolerance about which keys will be removed, smaller tolerance, more keys usually.
+	* @param Binding The binding upon which to bake
+	* @return returns True if successful, False otherwise
+	*/
+	UFUNCTION(BlueprintCallable, Category = "Editor Scripting | Sequencer Tools | Control Rig")
+	static bool BakeToControlRig(UWorld* World, ULevelSequence* LevelSequence, UClass* ControlRigClass, UAnimSeqExportOption* ExportOptions, bool bReduceKeys, float Tolerance,
+			const FSequencerBindingProxy& Binding);
+
+	/**
+	* Peform a Tween operation on the current active sequencer time(must be visible).
+	* @param LevelSequence The LevelSequence that's loaded in the editor
+	* @param ControlRig The Control Rig to tween.
+	* @param TweenValue The tween value to use, range from -1(blend to previous) to 1(blend to next)
+	*/
+	UFUNCTION(BlueprintCallable, Category = "Editor Scripting | Sequencer Tools | Control Rig")
+	static bool TweenControlRig(ULevelSequence* LevelSequence, UControlRig* ControlRig, float TweenValue);
+
+	/**
+	* Peform a Snap operation to snap the children to the parent. 
+	* @param StartFrame Beginning of the snap
+	* @param EndFrame End of the snap
+	* @param ChildrenToSnap The children objects that snap and get keys set onto. They need to live in an active Sequencer in the level editor
+	* @param ParentToSnap The parent object to snap relative to. If animated, it needs to live in an active Sequencer in the level editor
+	* @param SnapSettings Settings to use
+	* @param Returns True if successful, 
+	*/
+	UFUNCTION(BlueprintCallable, Category = "Editor Scripting | Sequencer Tools | Control Rig")
+	static bool SnapControlRig(FFrameNumber StartFrame, FFrameNumber EndFrame, const FControlRigSnapperSelection& ChildrenToSnap,
+		const FControlRigSnapperSelection& ParentToSnap, const UControlRigSnapSettings* SnapSettings);
+	
+	/**
+	* Get Actors World Transform at a specific time
+	* @param LevelSequence Active Sequence to get transform for
+	* @param Actor The actor
+	* @param Frame Time to get the transform
+	* @return Returns World Transform
+	*/
+	UFUNCTION(BlueprintCallable, Category = "Editor Scripting | Sequencer Tools | Control Rig")
+	static FTransform GetActorWorldTransform(ULevelSequence* LevelSequence,AActor* Actor, FFrameNumber Frame);
+
+	/**
+	* Get Actors World Transforms at specific times
+	* @param LevelSequence Active Sequence to get transform for
+	* @param Actor The actor
+	* @param Frames Times to get the transform
+	* @return Returns World Transforms
+	*/
+	UFUNCTION(BlueprintCallable, Category = "Editor Scripting | Sequencer Tools | Control Rig")
+	static TArray<FTransform> GetActorWorldTransforms(ULevelSequence* LevelSequence, AActor* Actor, const TArray<FFrameNumber>& Frames);
+
+	/**
+	* Get SkeletalMeshComponent World Transform at a specific time
+	* @param LevelSequence Active Sequence to get transform for
+	* @param SkeletalMeshComponent The SkeletalMeshComponent
+	* @param Frame Time to get the transform
+	* @return Returns World Transform
+	*/
+	UFUNCTION(BlueprintCallable, Category = "Editor Scripting | Sequencer Tools | Control Rig")
+	static FTransform GetSkeletalMeshComponentWorldTransform(ULevelSequence* LevelSequence, USkeletalMeshComponent* SkeletalMeshComponent, FFrameNumber Frame, FName SocketName = NAME_None);		
+
+	/**
+	* Get SkeletalMeshComponents World Transforms at specific times
+	* @param LevelSequence Active Sequence to get transform for
+	* @param SkeletalMeshComponent The SkeletalMeshComponent
+	* @param Frames Times to get the transform
+	* @return Returns World Transforms
+	*/
+	UFUNCTION(BlueprintCallable, Category = "Editor Scripting | Sequencer Tools | Control Rig")
+	static TArray<FTransform> GetSkeletalMeshComponentWorldTransforms(ULevelSequence* LevelSequence, USkeletalMeshComponent* SkeletalMeshComponent, const TArray<FFrameNumber>& Frames, FName SocketName = NAME_None);
+
+	/**
+	* Get ControlRig Control's World Transform at a specific time
+	* @param LevelSequence Active Sequence to get transform for
+	* @param ControlRig The ControlRig
+	* @param ControlName Name of the Control
+	* @param Frame Time to get the transform
+	* @return Returns World Transform
+	*/
+	UFUNCTION(BlueprintCallable, Category = "Editor Scripting | Sequencer Tools | Control Rig")
+	static FTransform GetControlRigWorldTransform(ULevelSequence* LevelSequence, UControlRig* ControlRig, FName ControlName, FFrameNumber Frame);
+
+	/**
+	* Get ControlRig Control's World Transforms at specific times
+	* @param LevelSequence Active Sequence to get transform for
+	* @param ControlRig The ControlRig
+	* @param ControlName Name of the Control 
+	* @param Frames Times to get the transform
+	* @return Returns World Transforms
+	*/
+	UFUNCTION(BlueprintCallable, Category = "Editor Scripting | Sequencer Tools | Control Rig")
+	static TArray<FTransform> GetControlRigWorldTransforms(ULevelSequence* LevelSequence, UControlRig* ControlRig, FName ControlName, const TArray<FFrameNumber>& Frames);
+
+	/**
+	* Set ControlRig Control's World Transform at a specific time
+	* @param LevelSequence Active Sequence to set transforms for. Must be loaded in Level Editor.
+	* @param ControlRig The ControlRig
+	* @param ControlName Name of the Control
+	* @param Frame Time to set the transform
+	* @oaram WorldTransform World Transform to set
+	* @param bSetKey Whether or not to set a key.
+	*/
+	UFUNCTION(BlueprintCallable, Category = "Editor Scripting | Sequencer Tools | Control Rig")
+	static void SetControlRigWorldTransform(ULevelSequence* LevelSequence, UControlRig* ControlRig, FName ControlName, FFrameNumber Frame, const FTransform& WorldTransform, bool bSetKey = true);
+
+	/**
+	* Set ControlRig Control's World Transforms at a specific times.
+	* @param LevelSequence Active Sequence to set transforms for. Must be loaded in Level Editor.
+	* @param ControlRig The ControlRig
+	* @param ControlName Name of the Control
+	* @param Frame Time to set the transform
+	* @oaram WorldTransform World Transform to set
+	* @param bSetKey Whether or not to set a key.
+	*/
+	UFUNCTION(BlueprintCallable, Category = "Editor Scripting | Sequencer Tools | Control Rig")
+	static void SetControlRigWorldTransforms(ULevelSequence* LevelSequence, UControlRig* ControlRig, FName ControlName, const TArray<FFrameNumber>& Frames, const TArray<FTransform>& WorldTransforms);
+	
+	/**
+	* Get ControlRig Control's float value at a specific time
+	* @param LevelSequence Active Sequence to get value for
+	* @param ControlRig The ControlRig
+	* @param ControlName Name of the Control, should be a float control
+	* @param Frame Time to get the value
+	* @return Returns Value at that time
+	*/
+	UFUNCTION(BlueprintCallable, Category = "Editor Scripting | Sequencer Tools | Control Rig")
+	static float GetLocalControlRigFloat(ULevelSequence* LevelSequence, UControlRig* ControlRig, FName ControlName, FFrameNumber Frame);
+
+	/**
+	* Get ControlRig Control's float values at specific times
+	* @param LevelSequence Active Sequence to get value for
+	* @param ControlRig The ControlRig
+	* @param ControlName Name of the Control, should be a float control
+	* @param Frames Times to get the values
+	* @return Returns Values at those times
+	*/
+	UFUNCTION(BlueprintCallable, Category = "Editor Scripting | Sequencer Tools | Control Rig")
+	static TArray<float> GetLocalControlRigFloats(ULevelSequence* LevelSequence, UControlRig* ControlRig, FName ControlName, const TArray<FFrameNumber>& Frames);
+
+	/**
+	* Set ControlRig Control's float value at specific time
+	* @param LevelSequence Active Sequence to set value on
+	* @param ControlRig The ControlRig
+	* @param ControlName Name of the Control, should be a float control
+	* @param Frame Time to set the value
+	* @param Value to set
+	* @param bSetKey If True set a key, if not just set the value
+	*/
+	UFUNCTION(BlueprintCallable, Category = "Editor Scripting | Sequencer Tools | Control Rig")
+	static void SetLocalControlRigFloat(ULevelSequence* LevelSequence, UControlRig* ControlRig, FName ControlName, FFrameNumber Frame, float Value, bool bSetKey = true);
+
+	/**
+	* Set ControlRig Control's float values at specific times
+	* @param LevelSequence Active Sequence to set value on
+	* @param ControlRig The ControlRig
+	* @param ControlName Name of the Control, should be a float control
+	* @param Frames Times to set the values
+	* @param Values to set at those times
+	*/
+	UFUNCTION(BlueprintCallable, Category = "Editor Scripting | Sequencer Tools | Control Rig")
+	static void SetLocalControlRigFloats(ULevelSequence* LevelSequence, UControlRig* ControlRig, FName ControlName, const TArray<FFrameNumber>& Frames, const TArray<float> Values);
+
+	/**
+	* Get ControlRig Control's bool value at a specific time
+	* @param LevelSequence Active Sequence to get value for
+	* @param ControlRig The ControlRig
+	* @param ControlName Name of the Control, should be a bool control
+	* @param Frame Time to get the value
+	* @return Returns Value at that time
+	*/
+	UFUNCTION(BlueprintCallable, Category = "Editor Scripting | Sequencer Tools | Control Rig")
+	static bool GetLocalControlRigBool(ULevelSequence* LevelSequence, UControlRig* ControlRig, FName ControlName, FFrameNumber Frame);
+
+	/**
+	* Get ControlRig Control's bool values at specific times
+	* @param LevelSequence Active Sequence to get value for
+	* @param ControlRig The ControlRig
+	* @param ControlName Name of the Control, should be a bool control
+	* @param Frames Times to get the values
+	* @return Returns Values at those times
+	*/
+	UFUNCTION(BlueprintCallable, Category = "Editor Scripting | Sequencer Tools | Control Rig")
+	static TArray<bool> GetLocalControlRigBools(ULevelSequence* LevelSequence, UControlRig* ControlRig, FName ControlName, const TArray<FFrameNumber>& Frames);
+
+	/**
+	* Set ControlRig Control's bool value at specific time
+	* @param LevelSequence Active Sequence to set value on
+	* @param ControlRig The ControlRig
+	* @param ControlName Name of the Control, should be a bool control
+	* @param Frame Time to set the value
+	* @param Value to set
+	* @param bSetKey If True set a key, if not just set the value
+	*/
+	UFUNCTION(BlueprintCallable, Category = "Editor Scripting | Sequencer Tools | Control Rig")
+	static void SetLocalControlRigBool(ULevelSequence* LevelSequence, UControlRig* ControlRig, FName ControlName, FFrameNumber Frame, bool Value, bool bSetKey = true);
+
+	/**
+	* Set ControlRig Control's bool values at specific times
+	* @param LevelSequence Active Sequence to set value on
+	* @param ControlRig The ControlRig
+	* @param ControlName Name of the Control, should be a bool control
+	* @param Frames Times to set the values
+	* @param Values to set at those times
+*/
+	UFUNCTION(BlueprintCallable, Category = "Editor Scripting | Sequencer Tools | Control Rig")
+	static void SetLocalControlRigBools(ULevelSequence* LevelSequence, UControlRig* ControlRig, FName ControlName, const TArray<FFrameNumber>& Frames, const TArray<bool> Values);
+
+	/**
+	* Get ControlRig Control's integer value at a specific time
+	* @param LevelSequence Active Sequence to get value for
+	* @param ControlRig The ControlRig
+	* @param ControlName Name of the Control, should be a integer control
+	* @param Frame Time to get the value
+	* @return Returns Value at that time
+	*/
+	UFUNCTION(BlueprintCallable, Category = "Editor Scripting | Sequencer Tools | Control Rig")
+	static int32 GetLocalControlRigInt(ULevelSequence* LevelSequence, UControlRig* ControlRig, FName ControlName, FFrameNumber Frame);
+
+	/**
+	* Get ControlRig Control's integer values at specific times
+	* @param LevelSequence Active Sequence to get value for
+	* @param ControlRig The ControlRig
+	* @param ControlName Name of the Control, should be a intteger control
+	* @param Frames Times to get the values
+	* @return Returns Values at those times
+	*/
+	UFUNCTION(BlueprintCallable, Category = "Editor Scripting | Sequencer Tools | Control Rig")
+	static TArray<int32> GetLocalControlRigInts(ULevelSequence* LevelSequence, UControlRig* ControlRig, FName ControlName, const TArray<FFrameNumber>& Frames);
+
+	/**
+	* Set ControlRig Control's int value at specific time
+	* @param LevelSequence Active Sequence to set value on
+	* @param ControlRig The ControlRig
+	* @param ControlName Name of the Control, should be a int control
+	* @param Frame Time to set the value
+	* @param Value to set
+	* @param bSetKey If True set a key, if not just set the value
+	*/
+	UFUNCTION(BlueprintCallable, Category = "Editor Scripting | Sequencer Tools | Control Rig")
+	static void SetLocalControlRigInt(ULevelSequence* LevelSequence, UControlRig* ControlRig, FName ControlName, FFrameNumber Frame, int32 Value, bool bSetKey = true);
+
+
+	/**
+	* Set ControlRig Control's int values at specific times
+	* @param LevelSequence Active Sequence to set value on
+	* @param ControlRig The ControlRig
+	* @param ControlName Name of the Control, should be a int control
+	* @param Frames Times to set the values
+	* @param Values to set at those times
+	*/
+	UFUNCTION(BlueprintCallable, Category = "Editor Scripting | Sequencer Tools | Control Rig")
+	static void SetLocalControlRigInts(ULevelSequence* LevelSequence, UControlRig* ControlRig, FName ControlName, const TArray<FFrameNumber>& Frames, const TArray<int32> Values);
+
+	/**
+	* Get ControlRig Control's Vector2D value at a specific time
+	* @param LevelSequence Active Sequence to get value for
+	* @param ControlRig The ControlRig
+	* @param ControlName Name of the Control, should be a Vector2D control
+	* @param Frame Time to get the value
+	* @return Returns Value at that time
+	*/
+	UFUNCTION(BlueprintCallable, Category = "Editor Scripting | Sequencer Tools | Control Rig")
+	static FVector2D GetLocalControlRigVector2D(ULevelSequence* LevelSequence, UControlRig* ControlRig, FName ControlName, FFrameNumber Frame);
+
+	/**
+	* Get ControlRig Control's Vector2D values at specific times
+	* @param LevelSequence Active Sequence to get value for
+	* @param ControlRig The ControlRig
+	* @param ControlName Name of the Control, should be a Vector2D control
+	* @param Frames Times to get the values
+	* @return Returns Values at those times
+	*/
+	UFUNCTION(BlueprintCallable, Category = "Editor Scripting | Sequencer Tools | Control Rig")
+	static TArray<FVector2D> GetLocalControlRigVector2Ds(ULevelSequence* LevelSequence, UControlRig* ControlRig, FName ControlName, const TArray<FFrameNumber>& Frames);
+
+	/**
+	* Set ControlRig Control's Vector2D value at specific time
+	* @param LevelSequence Active Sequence to set value on
+	* @param ControlRig The ControlRig
+	* @param ControlName Name of the Control, should be a Vector2D control
+	* @param Frame Time to set the value
+	* @param Value to set
+	* @param bSetKey If True set a key, if not just set the value
+	*/
+	UFUNCTION(BlueprintCallable, Category = "Editor Scripting | Sequencer Tools | Control Rig")
+	static void SetLocalControlRigVector2D(ULevelSequence* LevelSequence, UControlRig* ControlRig, FName ControlName, FFrameNumber Frame, FVector2D Value, bool bSetKey = true);
+
+
+	/**
+	* Set ControlRig Control's Vector2D values at specific times
+	* @param LevelSequence Active Sequence to set value on
+	* @param ControlRig The ControlRig
+	* @param ControlName Name of the Control, should be a Vector2D control
+	* @param Frames Times to set the values
+	* @param Values to set at those times
+	*/
+	UFUNCTION(BlueprintCallable, Category = "Editor Scripting | Sequencer Tools | Control Rig")
+	static void SetLocalControlRigVector2Ds(ULevelSequence* LevelSequence, UControlRig* ControlRig, FName ControlName, const TArray<FFrameNumber>& Frames, const TArray<FVector2D> Values);
+
+	/**
+	* Get ControlRig Control's Position value at a specific time
+	* @param LevelSequence Active Sequence to get value for
+	* @param ControlRig The ControlRig
+	* @param ControlName Name of the Control, should be a Position control
+	* @param Frame Time to get the value
+	* @return Returns Value at that time
+	*/
+	UFUNCTION(BlueprintCallable, Category = "Editor Scripting | Sequencer Tools | Control Rig")
+	static FVector GetLocalControlRigPosition(ULevelSequence* LevelSequence, UControlRig* ControlRig, FName ControlName, FFrameNumber Frame);
+
+	/**
+	* Get ControlRig Control's Position values at specific times
+	* @param LevelSequence Active Sequence to get value for
+	* @param ControlRig The ControlRig
+	* @param ControlName Name of the Control, should be a Position control
+	* @param Frames Times to get the values
+	* @return Returns Values at those times
+	*/
+	UFUNCTION(BlueprintCallable, Category = "Editor Scripting | Sequencer Tools | Control Rig")
+	static TArray<FVector> GetLocalControlRigPositions(ULevelSequence* LevelSequence, UControlRig* ControlRig, FName ControlName, const TArray<FFrameNumber>& Frames);
+
+	/**
+	* Set ControlRig Control's Position value at specific time
+	* @param LevelSequence Active Sequence to set value on
+	* @param ControlRig The ControlRig
+	* @param ControlName Name of the Control, should be a Position control
+	* @param Frame Time to set the value
+	* @param Value to set
+	* @param bSetKey If True set a key, if not just set the value
+	*/
+	UFUNCTION(BlueprintCallable, Category = "Editor Scripting | Sequencer Tools | Control Rig")
+	static void SetLocalControlRigPosition(ULevelSequence* LevelSequence, UControlRig* ControlRig, FName ControlName, FFrameNumber Frame, FVector Value, bool bSetKey = true);
+
+	/**
+	* Set ControlRig Control's Position values at specific times
+	* @param LevelSequence Active Sequence to set value on
+	* @param ControlRig The ControlRig
+	* @param ControlName Name of the Control, should be a Position control
+	* @param Frames Times to set the values
+	* @param Values to set at those times
+	*/
+	UFUNCTION(BlueprintCallable, Category = "Editor Scripting | Sequencer Tools | Control Rig")
+	static void SetLocalControlRigPositions(ULevelSequence* LevelSequence, UControlRig* ControlRig, FName ControlName, const TArray<FFrameNumber>& Frames, const TArray<FVector> Values);
+
+	/**
+	* Get ControlRig Control's Rotator value at a specific time
+	* @param LevelSequence Active Sequence to get value for
+	* @param ControlRig The ControlRig
+	* @param ControlName Name of the Control, should be a Rotator control
+	* @param Frame Time to get the value
+	* @return Returns Value at that time
+	*/
+	UFUNCTION(BlueprintCallable, Category = "Editor Scripting | Sequencer Tools | Control Rig")
+	static FRotator GetLocalControlRigRotator(ULevelSequence* LevelSequence, UControlRig* ControlRig, FName ControlName, FFrameNumber Frame);
+
+	/**
+	* Get ControlRig Control's Rotator values at specific times
+	* @param LevelSequence Active Sequence to get value for
+	* @param ControlRig The ControlRig
+	* @param ControlName Name of the Control, should be a Rotator control
+	* @param Frames Times to get the values
+	* @return Returns Values at those times
+	*/
+	UFUNCTION(BlueprintCallable, Category = "Editor Scripting | Sequencer Tools | Control Rig")
+	static TArray<FRotator> GetLocalControlRigRotators(ULevelSequence* LevelSequence, UControlRig* ControlRig, FName ControlName, const TArray<FFrameNumber>& Frames);
+
+	/**
+	* Set ControlRig Control's Rotator value at specific time
+	* @param LevelSequence Active Sequence to set value on
+	* @param ControlRig The ControlRig
+	* @param ControlName Name of the Control, should be a Rotator control
+	* @param Frame Time to set the value
+	* @param Value to set
+	* @param bSetKey If True set a key, if not just set the value
+	*/
+	UFUNCTION(BlueprintCallable, Category = "Editor Scripting | Sequencer Tools | Control Rig")
+	static void SetLocalControlRigRotator(ULevelSequence* LevelSequence, UControlRig* ControlRig, FName ControlName, FFrameNumber Frame, FRotator Value, bool bSetKey = true);
+
+	/**
+	* Set ControlRig Control's Rotator values at specific times
+	* @param LevelSequence Active Sequence to set value on
+	* @param ControlRig The ControlRig
+	* @param ControlName Name of the Control, should be a Rotator control
+	* @param Frames Times to set the values
+	* @param Values to set at those times
+	*/
+	UFUNCTION(BlueprintCallable, Category = "Editor Scripting | Sequencer Tools | Control Rig")
+	static void SetLocalControlRigRotators(ULevelSequence* LevelSequence, UControlRig* ControlRig, FName ControlName, const TArray<FFrameNumber>& Frames, const TArray<FRotator> Values);
+
+	/**
+	* Get ControlRig Control's Scale value at a specific time
+	* @param LevelSequence Active Sequence to get value for
+	* @param ControlRig The ControlRig
+	* @param ControlName Name of the Control, should be a Scale control
+	* @param Frame Time to get the value
+	* @return Returns Value at that time
+	*/
+	UFUNCTION(BlueprintCallable, Category = "Editor Scripting | Sequencer Tools | Control Rig")
+	static FVector GetLocalControlRigScale(ULevelSequence* LevelSequence, UControlRig* ControlRig, FName ControlName, FFrameNumber Frame);
+
+	/**
+	* Get ControlRig Control's Scale values at specific times
+	* @param LevelSequence Active Sequence to get value for
+	* @param ControlRig The ControlRig
+	* @param ControlName Name of the Control, should be a Scale control
+	* @param Frames Times to get the values
+	* @return Returns Values at those times
+	*/
+	UFUNCTION(BlueprintCallable, Category = "Editor Scripting | Sequencer Tools | Control Rig")
+	static TArray<FVector> GetLocalControlRigScales(ULevelSequence* LevelSequence, UControlRig* ControlRig, FName ControlName, const TArray<FFrameNumber>& Frames);
+
+	/**
+	* Set ControlRig Control's Scale value at specific time
+	* @param LevelSequence Active Sequence to set value on
+	* @param ControlRig The ControlRig
+	* @param ControlName Name of the Control, should be a Scale control
+	* @param Frame Time to set the value
+	* @param Value to set
+	* @param bSetKey If True set a key, if not just set the value
+*/
+	UFUNCTION(BlueprintCallable, Category = "Editor Scripting | Sequencer Tools | Control Rig")
+	static void SetLocalControlRigScale(ULevelSequence* LevelSequence, UControlRig* ControlRig, FName ControlName, FFrameNumber Frame, FVector Value, bool bSetKey = true);
+
+	/**
+	* Set ControlRig Control's Scale values at specific times
+	* @param LevelSequence Active Sequence to set value on
+	* @param ControlRig The ControlRig
+	* @param ControlName Name of the Control, should be a Scale control
+	* @param Frames Times to set the values
+	* @param Values to set at those times
+	*/
+	UFUNCTION(BlueprintCallable, Category = "Editor Scripting | Sequencer Tools | Control Rig")
+	static void SetLocalControlRigScales(ULevelSequence* LevelSequence, UControlRig* ControlRig, FName ControlName, const TArray<FFrameNumber>& Frames, const TArray<FVector> Values);
+
+
+	/**
+	* Get ControlRig Control's EulerTransform value at a specific time
+	* @param LevelSequence Active Sequence to get value for
+	* @param ControlRig The ControlRig
+	* @param ControlName Name of the Control, should be a EulerTransfom control
+	* @param Frame Time to get the value
+	* @return Returns Value at that time
+	*/
+	UFUNCTION(BlueprintCallable, Category = "Editor Scripting | Sequencer Tools | Control Rig")
+	static FEulerTransform GetLocalControlRigEulerTransform(ULevelSequence* LevelSequence, UControlRig* ControlRig, FName ControlName, FFrameNumber Frame);
+
+	/**
+	* Get ControlRig Control's EulerTransform values at specific times
+	* @param LevelSequence Active Sequence to get value for
+	* @param ControlRig The ControlRig
+	* @param ControlName Name of the Control, should be a EulerTransform control
+	* @param Frames Times to get the values
+	* @return Returns Values at those times
+	*/
+	UFUNCTION(BlueprintCallable, Category = "Editor Scripting | Sequencer Tools | Control Rig")
+	static TArray<FEulerTransform> GetLocalControlRigEulerTransforms(ULevelSequence* LevelSequence, UControlRig* ControlRig, FName ControlName, const TArray<FFrameNumber>& Frames);
+
+	/**
+	* Set ControlRig Control's EulerTransform value at specific time
+	* @param LevelSequence Active Sequence to set value on
+	* @param ControlRig The ControlRig
+	* @param ControlName Name of the Control, should be a EulerTransform control
+	* @param Frame Time to set the value
+	* @param Value to set
+	* @param bSetKey If True set a key, if not just set the value
+	*/
+	UFUNCTION(BlueprintCallable, Category = "Editor Scripting | Sequencer Tools | Control Rig")
+	static void SetLocalControlRigEulerTransform(ULevelSequence* LevelSequence, UControlRig* ControlRig, FName ControlName, FFrameNumber Frame, FEulerTransform Value, bool bSetKey = true);
+
+	/**
+	* Set ControlRig Control's EulerTransform values at specific times
+	* @param LevelSequence Active Sequence to set value on
+	* @param ControlRig The ControlRig
+	* @param ControlName Name of the Control, should be a EulerTransform control
+	* @param Frames Times to set the values
+	* @param Values to set at those times
+	*/
+	UFUNCTION(BlueprintCallable, Category = "Editor Scripting | Sequencer Tools | Control Rig")
+	static void SetLocalControlRigEulerTransforms(ULevelSequence* LevelSequence, UControlRig* ControlRig, FName ControlName, const TArray<FFrameNumber>& Frames, const TArray<FEulerTransform> Values);
+
+	/**
+	* Get ControlRig Control's TransformNoScale value at a specific time
+	* @param LevelSequence Active Sequence to get value for
+	* @param ControlRig The ControlRig
+	* @param ControlName Name of the Control, should be a TransformNoScale control
+	* @param Frame Time to get the value
+	* @return Returns Value at that time
+	*/
+	UFUNCTION(BlueprintCallable, Category = "Editor Scripting | Sequencer Tools | Control Rig")
+	static FTransformNoScale GetLocalControlRigTransformNoScale(ULevelSequence* LevelSequence, UControlRig* ControlRig, FName ControlName, FFrameNumber Frame);
+
+	/**
+	* Get ControlRig Control's TransformNoScale values at specific times
+	* @param LevelSequence Active Sequence to get value for
+	* @param ControlRig The ControlRig
+	* @param ControlName Name of the Control, should be a TransformNoScale control
+	* @param Frames Times to get the values
+	* @return Returns Values at those times
+	*/
+	UFUNCTION(BlueprintCallable, Category = "Editor Scripting | Sequencer Tools | Control Rig")
+	static TArray<FTransformNoScale> GetLocalControlRigTransformNoScales(ULevelSequence* LevelSequence, UControlRig* ControlRig, FName ControlName, const TArray<FFrameNumber>& Frames);
+
+	/**
+	* Set ControlRig Control's TransformNoScale value at specific time
+	* @param LevelSequence Active Sequence to set value on
+	* @param ControlRig The ControlRig
+	* @param ControlName Name of the Control, should be a TransformNoScale control
+	* @param Frame Time to set the value
+	* @param Value to set
+	* @param bSetKey If True set a key, if not just set the value
+	*/
+	UFUNCTION(BlueprintCallable, Category = "Editor Scripting | Sequencer Tools | Control Rig")
+	static void SetLocalControlRigTransformNoScale(ULevelSequence* LevelSequence, UControlRig* ControlRig, FName ControlName, FFrameNumber Frame, FTransformNoScale Value, bool bSetKey = true);
+
+	/**
+	* Set ControlRig Control's TransformNoScale values at specific times
+	* @param LevelSequence Active Sequence to set value on
+	* @param ControlRig The ControlRig
+	* @param ControlName Name of the Control, should be a TransformNoScale control
+	* @param Frames Times to set the values
+	* @param Values to set at those times
+	*/
+	UFUNCTION(BlueprintCallable, Category = "Editor Scripting | Sequencer Tools | Control Rig")
+	static void SetLocalControlRigTransformNoScales(ULevelSequence* LevelSequence, UControlRig* ControlRig, FName ControlName, const TArray<FFrameNumber>& Frames, const TArray<FTransformNoScale> Values);
+
+	/**
+	* Get ControlRig Control's Transform value at a specific time
+	* @param LevelSequence Active Sequence to get value for
+	* @param ControlRig The ControlRig
+	* @param ControlName Name of the Control, should be a Transform control
+	* @param Frame Time to get the value
+	* @return Returns Value at that time
+	*/
+	UFUNCTION(BlueprintCallable, Category = "Editor Scripting | Sequencer Tools | Control Rig")
+	static FTransform GetLocalControlRigTransform(ULevelSequence* LevelSequence, UControlRig* ControlRig, FName ControlName, FFrameNumber Frame);
+
+	/**
+	* Get ControlRig Control's Transform values at specific times
+	* @param LevelSequence Active Sequence to get value for
+	* @param ControlRig The ControlRig
+	* @param ControlName Name of the Control, should be a Transform control
+	* @param Frames Times to get the values
+	* @return Returns Values at those times
+	*/
+	UFUNCTION(BlueprintCallable, Category = "Editor Scripting | Sequencer Tools | Control Rig")
+	static TArray<FTransform> GetLocalControlRigTransforms(ULevelSequence* LevelSequence, UControlRig* ControlRig, FName ControlName, const TArray<FFrameNumber>& Frames);
+
+	/**
+	* Set ControlRig Control's Transform value at specific time
+	* @param LevelSequence Active Sequence to set value on
+	* @param ControlRig The ControlRig
+	* @param ControlName Name of the Control, should be a Transform control
+	* @param Frame Time to set the value
+	* @param Value to set 
+	* @param bSetKey If True set a key, if not just set the value
+	*/
+	UFUNCTION(BlueprintCallable, Category = "Editor Scripting | Sequencer Tools | Control Rig")
+	static void SetLocalControlRigTransform(ULevelSequence* LevelSequence, UControlRig* ControlRig, FName ControlName, FFrameNumber Frame, FTransform Value, bool bSetKey = true);
+
+	/**
+	* Set ControlRig Control's Transform values at specific times
+	* @param LevelSequence Active Sequence to set value on
+	* @param ControlRig The ControlRig
+	* @param ControlName Name of the Control, should be a Transform control
+	* @param Frames Times to set the values
+	* @param Values to set at those times
+	*/
+	UFUNCTION(BlueprintCallable, Category = "Editor Scripting | Sequencer Tools | Control Rig")
+	static void SetLocalControlRigTransforms(ULevelSequence* LevelSequence, UControlRig* ControlRig, FName ControlName, const TArray<FFrameNumber>& Frames, const TArray<FTransform> Values);
+};
+
diff --git a/Engine/Plugins/Experimental/ControlRig/Source/ControlRigEditor/Public/Tools/ControlRigSnapper.h b/Engine/Plugins/Experimental/ControlRig/Source/ControlRigEditor/Public/Tools/ControlRigSnapper.h
--- a/Engine/Plugins/Experimental/ControlRig/Source/ControlRigEditor/Public/Tools/ControlRigSnapper.h
+++ b/Engine/Plugins/Experimental/ControlRig/Source/ControlRigEditor/Public/Tools/ControlRigSnapper.h
@@ -12,6 +12,8 @@ class AActor;
 class UControlRig;
 class ISequencer;
 class UMovieScene;
+class ULevelSequence;
+class UControlRigSnapSettings;
 
 //Specification containing a Control Rig and a list of selected Controls we use to get World Transforms for Snapping.
 USTRUCT(BlueprintType)
@@ -58,9 +60,37 @@ struct FControlRigSnapper
 	* @param EndFrame End of the snap interval.
 	* @param ChildrenToSnap The children to snap over the interval. Will set a key per frame on them.
 	* @param ParentToSnap  The parent to snap to.
+	* @param SnapSettings The settings to use.
 	*/
 	bool SnapIt(FFrameNumber StartFrame, FFrameNumber EndFrame, const FControlRigSnapperSelection& ChildrenToSnap,
-		const FControlRigSnapperSelection& ParentToSnap);
+		const FControlRigSnapperSelection& ParentToSnap, const UControlRigSnapSettings* SnapSettings);
+
+
+	/**
+	*  Get the ControlRig Transforms for current Sequencer
+	* @param Sequencer Sequencer Evaluating
+	* @param ControlRig The Control Rig to evaluate
+	* @param ControlName The name of the Control to get.
+	* @param Frames  The times to get the transforms.
+	* @param ParentTransforms List of Parent Transforms for each time that the resulting transforms will be concatenated to.
+	* @param OutTransforms  Results
+	*/
+	bool GetControlRigControlTransforms(ISequencer* Sequencer, UControlRig* ControlRig, const FName& ControlName,
+		const TArray<FFrameNumber>& Frames, const TArray<FTransform>& ParentTransforms, TArray<FTransform>& OutTransforms);
+
+
+	/**
+	*  Get the ControlRig Transforms from specified MovieSceneSequence. This will create a new player to evalaute so can be used on MovieSceneSequence that isn't in the scene
+	* @param World Active World
+	* @param LevelSequence LevelSequence Evaluating
+	* @param ControlRig The Control Rig to evaluate
+	* @param ControlName The name of the Control to get.
+	* @param Frames  The times to get the transforms.
+	* @param ParentTransforms List of Parent Transforms for each time that the resulting transforms will be concatenated to.
+	* @param OutTransforms  Results
+	*/
+	bool GetControlRigControlTransforms(UWorld* World, ULevelSequence* LevelSequence,
+		UControlRig* ControlRig, const FName& ControlName,const TArray<FFrameNumber>& Frames, const TArray<FTransform>& ParentTransforms, TArray<FTransform>& OutTransforms);
 
 };
 
diff --git a/Engine/Source/Editor/MovieSceneTools/Private/MovieSceneToolHelpers.cpp b/Engine/Source/Editor/MovieSceneTools/Private/MovieSceneToolHelpers.cpp
--- a/Engine/Source/Editor/MovieSceneTools/Private/MovieSceneToolHelpers.cpp
+++ b/Engine/Source/Editor/MovieSceneTools/Private/MovieSceneToolHelpers.cpp
@@ -3979,7 +3979,12 @@ static void GetSequencerActorWorldTransforms(ISequencer* Sequencer, const FActor
 	if (AActor* Actor = ActorSelection.Actor.Get())
 	{
 
-		USkeletalMeshComponent* SkelMeshComp = MovieSceneToolHelpers::AcquireSkeletalMeshFromObject(Actor);
+		USkeletalMeshComponent* SkelMeshComp = ActorSelection.Component.IsValid() ? Cast<USkeletalMeshComponent>(ActorSelection.Component.Get()) : nullptr;
+
+		if (!SkelMeshComp)
+		{
+			SkelMeshComp = MovieSceneToolHelpers::AcquireSkeletalMeshFromObject(Actor);
+		}
 
 		if (ActorSelection.SocketName != NAME_None && SkelMeshComp)
 		{
@@ -4028,14 +4033,20 @@ static void GetSequencerActorWorldTransforms(ISequencer* Sequencer, const FActor
 						}
 					}
 
-					SkelMeshComp->TickAnimation(0.03f, false);
-					SkelMeshComp->RefreshBoneTransforms();
-					SkelMeshComp->RefreshSlaveComponents();
-					SkelMeshComp->UpdateComponentToWorld();
-					SkelMeshComp->FinalizeBoneTransform();
-					SkelMeshComp->MarkRenderTransformDirty();
-					SkelMeshComp->MarkRenderDynamicDataDirty();
+					TArray<USkeletalMeshComponent*> MeshComps;
+					ActorSelection.Actor->GetComponents<USkeletalMeshComponent>(MeshComps, true);
 
+					for (USkeletalMeshComponent* MeshComp : MeshComps)
+					{
+						MeshComp->TickAnimation(0.03f, false);
+						MeshComp->RefreshBoneTransforms();
+						MeshComp->RefreshSlaveComponents();
+						MeshComp->UpdateComponentToWorld();
+						MeshComp->FinalizeBoneTransform();
+						MeshComp->MarkRenderTransformDirty();
+						MeshComp->MarkRenderDynamicDataDirty();
+					}
+					
 					OutTransforms[Index] = SkelMeshComp->GetSocketTransform(ActorSelection.SocketName);// GetSocketTransofrm is world space in theory*OutTransforms[Index];
 
 				}
diff --git a/Engine/Source/Editor/MovieSceneTools/Public/MovieSceneToolHelpers.h b/Engine/Source/Editor/MovieSceneTools/Public/MovieSceneToolHelpers.h
--- a/Engine/Source/Editor/MovieSceneTools/Public/MovieSceneToolHelpers.h
+++ b/Engine/Source/Editor/MovieSceneTools/Public/MovieSceneToolHelpers.h
@@ -509,16 +509,16 @@ public:
 	* @return Returns the USkeletalMeshComponent if one is found
 	*/
 	static USkeletalMeshComponent* AcquireSkeletalMeshFromObject(UObject* BoundObject);
+	
 	/*
 	*  Get an actors word transforms at the specified times
-	* @param Sequencer Sequencer to evaluation
+	* @param Sequencer Sequencer to evaluate
     * @param ActorForWorldTransforms The actor and possible component and socket that we want to get the world transforms for.
 	* @param Frames The times we want to get the world transforms
 	* @param OutWorldTransforms The calculated world transforms, one for each specified frame.
 	*/
 	static void GetActorWorldTransforms(ISequencer* Sequencer, const FActorForWorldTransforms& Actors, const TArray<FFrameNumber>& Frames, TArray<FTransform>& OutWorldTransforms);
 
-
 };
 
 // Helper to make spawnables persist throughout the export process and then restore properly afterwards
diff --git a/Engine/Source/Editor/Sequencer/Private/Sequencer.cpp b/Engine/Source/Editor/Sequencer/Private/Sequencer.cpp
--- a/Engine/Source/Editor/Sequencer/Private/Sequencer.cpp
+++ b/Engine/Source/Editor/Sequencer/Private/Sequencer.cpp
@@ -9692,6 +9692,14 @@ void FSequencer::ObjectImplicitlyAdded(UObject* InObject) const
 	}
 }
 
+void FSequencer::ObjectImplicitlyRemoved(UObject* InObject) const
+{
+	for (int32 i = 0; i < TrackEditors.Num(); ++i)
+	{
+		TrackEditors[i]->ObjectImplicitlyRemoved(InObject);
+	}
+}
+
 void FSequencer::SetFilterOn(const FText& InName, bool bOn)
 {
 	SequencerWidget->SetFilterOn(InName, bOn);
diff --git a/Engine/Source/Editor/Sequencer/Private/Sequencer.h b/Engine/Source/Editor/Sequencer/Private/Sequencer.h
--- a/Engine/Source/Editor/Sequencer/Private/Sequencer.h
+++ b/Engine/Source/Editor/Sequencer/Private/Sequencer.h
@@ -845,6 +845,8 @@ public:
 	virtual bool IsReadOnly() const override;
 	virtual void ExternalSelectionHasChanged() override { SynchronizeSequencerSelectionWithExternalSelection(); }
 	virtual void ObjectImplicitlyAdded(UObject* InObject) const override;
+	virtual void ObjectImplicitlyRemoved(UObject* InObject) const override;
+
 	/** Access the user-supplied settings object */
 	virtual USequencerSettings* GetSequencerSettings() const override { return Settings; }
 	virtual void SetSequencerSettings(USequencerSettings* InSettings) override { Settings = InSettings; }
diff --git a/Engine/Source/Editor/Sequencer/Public/ISequencer.h b/Engine/Source/Editor/Sequencer/Public/ISequencer.h
--- a/Engine/Source/Editor/Sequencer/Public/ISequencer.h
+++ b/Engine/Source/Editor/Sequencer/Public/ISequencer.h
@@ -673,6 +673,12 @@ public:
 	*/
 	virtual void ObjectImplicitlyAdded(UObject* InObject) const = 0;
 
+	/**
+	* Specify that an object was implicitly removed. We will notify the track editors that it was
+	@InObject Object that was removed that was part of a track/binding but not the real binding
+	*/
+	virtual void ObjectImplicitlyRemoved(UObject* InObject) const = 0;
+
 public:
 	/**
 	*    Turn on/off the filter with the specified name
diff --git a/Engine/Source/Editor/Sequencer/Public/ISequencerTrackEditor.h b/Engine/Source/Editor/Sequencer/Public/ISequencerTrackEditor.h
--- a/Engine/Source/Editor/Sequencer/Public/ISequencerTrackEditor.h
+++ b/Engine/Source/Editor/Sequencer/Public/ISequencerTrackEditor.h
@@ -252,6 +252,11 @@ public:
 	*/
 	virtual void ObjectImplicitlyAdded(UObject* InObject)  {}
 
+	/**
+	* Handle this object being implicitly removed
+	*/
+	virtual void ObjectImplicitlyRemoved(UObject* InObject) {}
+
 public:
 
 	/** Virtual destructor. */
diff --git a/Engine/Source/Runtime/MovieScene/Private/MovieScene.cpp b/Engine/Source/Runtime/MovieScene/Private/MovieScene.cpp
--- a/Engine/Source/Runtime/MovieScene/Private/MovieScene.cpp
+++ b/Engine/Source/Runtime/MovieScene/Private/MovieScene.cpp
@@ -881,6 +881,31 @@ UMovieSceneTrack* UMovieScene::FindTrack(TSubclassOf<UMovieSceneTrack> TrackClas
 	return nullptr;
 }
 
+TArray<UMovieSceneTrack*> UMovieScene::FindTracks(TSubclassOf<UMovieSceneTrack> TrackClass, const FGuid& ObjectGuid, const FName& TrackName) const
+{
+	check(ObjectGuid.IsValid());
+	TArray<UMovieSceneTrack*> MovieSceneTracks;
+	for (const auto& Binding : ObjectBindings)
+	{
+		if (Binding.GetObjectGuid() != ObjectGuid)
+		{
+			continue;
+		}
+
+		for (const auto& Track : Binding.GetTracks())
+		{
+			if (TrackClass.GetDefaultObject() == nullptr || Track->GetClass() == TrackClass)
+			{
+				if (TrackName == NAME_None || Track->GetTrackName() == TrackName)
+				{
+					MovieSceneTracks.Add(Track);
+				}
+			}
+		}
+	}
+
+	return MovieSceneTracks;
+}
 
 UMovieSceneTrack* UMovieScene::AddTrack( TSubclassOf<UMovieSceneTrack> TrackClass, const FGuid& ObjectGuid )
 {
diff --git a/Engine/Source/Runtime/MovieScene/Public/MovieScene.h b/Engine/Source/Runtime/MovieScene/Public/MovieScene.h
--- a/Engine/Source/Runtime/MovieScene/Public/MovieScene.h
+++ b/Engine/Source/Runtime/MovieScene/Public/MovieScene.h
@@ -554,7 +554,7 @@ public:
 	 * @param ObjectGuid The runtime object guid that the track is bound to.
 	 * @param TrackName The name of the track to differentiate the one we are searching for from other tracks of the same class (optional).
 	 * @return The found track or nullptr if one does not exist.
-	 * @see AddTrack, RemoveTrack
+	 * @see AddTrack, RemoveTrack, FindTracks
 	 */
 	UMovieSceneTrack* FindTrack(TSubclassOf<UMovieSceneTrack> TrackClass, const FGuid& ObjectGuid, const FName& TrackName = NAME_None) const;
 	
@@ -565,7 +565,7 @@ public:
 	 * @param ObjectGuid The runtime object guid that the track is bound to.
 	 * @param TrackName The name of the track to differentiate the one we are searching for from other tracks of the same class (optional).
 	 * @return The found track or nullptr if one does not exist.
-	 * @see AddTrack, RemoveTrack
+	 * @see AddTrack, RemoveTrack, FindTracks
 	 */
 	template<typename TrackClass>
 	TrackClass* FindTrack(const FGuid& ObjectGuid, const FName& TrackName = NAME_None) const
@@ -573,6 +573,17 @@ public:
 		return Cast<TrackClass>(FindTrack(TrackClass::StaticClass(), ObjectGuid, TrackName));
 	}
 
+	/**
+	 * Find all tracks of a given class.
+	 *
+	 * @param TrackClass The class of the track to find.
+	 * @param ObjectGuid The runtime object guid that the track is bound to.
+	 * @param TrackName The name of the track to differentiate the one we are searching for from other tracks of the same class (optional).
+	 * @return The found tracks or an empty array if none exist
+	 * @see AddTrack, RemoveTrack, FindTrack
+	 */
+	TArray<UMovieSceneTrack*> FindTracks(TSubclassOf<UMovieSceneTrack> TrackClass, const FGuid& ObjectGuid, const FName& TrackName = NAME_None) const;
+
 	/**
 	 * Removes a track.
 	 *
