Control Rig: Changed pin color for item pin, added support for external plugins to register pin colors for custom struct type

diff --git a/Engine/Plugins/Experimental/ControlRig/Source/ControlRigDeveloper/Private/ControlRigDeveloper.cpp b/Engine/Plugins/Experimental/ControlRig/Source/ControlRigDeveloper/Private/ControlRigDeveloper.cpp
--- a/Engine/Plugins/Experimental/ControlRig/Source/ControlRigDeveloper/Private/ControlRigDeveloper.cpp
+++ b/Engine/Plugins/Experimental/ControlRig/Source/ControlRigDeveloper/Private/ControlRigDeveloper.cpp
@@ -21,6 +21,16 @@ public:
 	FControlRigBlueprintCompiler ControlRigBlueprintCompiler;
 
 	static TSharedPtr<FKismetCompilerContext> GetControlRigCompiler(UBlueprint* BP, FCompilerResultsLog& InMessageLog, const FKismetCompilerOptions& InCompileOptions);
+
+	virtual void RegisterPinTypeColor(UStruct* Struct, const FLinearColor Color) override;
+
+	virtual void UnRegisterPinTypeColor(UStruct* Struct) override;
+	
+	virtual const FLinearColor* FindPinTypeColor(UStruct* Struct) const override;
+
+private:
+
+	TMap<UStruct*, FLinearColor> PinTypeColorMap;
 };
 
 void FControlRigDeveloperModule::StartupModule()
@@ -55,6 +65,23 @@ TSharedPtr<FKismetCompilerContext> FControlRigDeveloperModule::GetControlRigComp
 	return TSharedPtr<FKismetCompilerContext>(new FControlRigBlueprintCompilerContext(BP, InMessageLog, InCompileOptions));
 }
 
+void FControlRigDeveloperModule::RegisterPinTypeColor(UStruct* Struct, const FLinearColor Color)
+{
+	if (Struct)
+	{ 
+		PinTypeColorMap.FindOrAdd(Struct) = Color;
+	}
+}
+
+void FControlRigDeveloperModule::UnRegisterPinTypeColor(UStruct* Struct)
+{
+	PinTypeColorMap.Remove(Struct);
+}
+
+const FLinearColor* FControlRigDeveloperModule::FindPinTypeColor(UStruct* Struct) const
+{
+	return PinTypeColorMap.Find(Struct);
+}
 
 IMPLEMENT_MODULE(FControlRigDeveloperModule, ControlRigDeveloper)
 
diff --git a/Engine/Plugins/Experimental/ControlRig/Source/ControlRigDeveloper/Private/Graph/ControlRigGraphSchema.cpp b/Engine/Plugins/Experimental/ControlRig/Source/ControlRigDeveloper/Private/Graph/ControlRigGraphSchema.cpp
--- a/Engine/Plugins/Experimental/ControlRig/Source/ControlRigDeveloper/Private/Graph/ControlRigGraphSchema.cpp
+++ b/Engine/Plugins/Experimental/ControlRig/Source/ControlRigDeveloper/Private/Graph/ControlRigGraphSchema.cpp
@@ -15,6 +15,7 @@
 #include "ControlRig.h"
 #include "ControlRigBlueprint.h"
 #include "ControlRigBlueprintGeneratedClass.h"
+#include "ControlRigDeveloper.h"
 #include "Widgets/Notifications/SNotificationList.h"
 #include "Framework/Notifications/NotificationManager.h"
 #include "EulerTransform.h"
@@ -340,6 +341,20 @@ FLinearColor UControlRigGraphSchema::GetPinTypeColor(const FEdGraphPinType& PinT
 			{
 				return FLinearColor::White;
 			}
+			
+			if (Struct == FRigElementKey::StaticStruct() || Struct == FRigElementKeyCollection::StaticStruct())
+			{
+				return FLinearColor(0.0, 0.6588, 0.9490);
+			}
+
+			// external types can register their own colors, check if there are any
+			if (IControlRigDeveloperModule* Module = FModuleManager::GetModulePtr<IControlRigDeveloperModule>("ControlRigDeveloper"))
+			{
+				if (const FLinearColor* Color = Module->FindPinTypeColor(Struct))
+				{
+					return *Color;
+				}
+			}
 		}
 	}
 	return GetDefault<UEdGraphSchema_K2>()->GetPinTypeColor(PinType);
@@ -1019,4 +1034,4 @@ void UControlRigGraphSchema::EndGraphNodeInteraction(UEdGraphNode* InNode) const
 #endif
 }
 
-#undef LOCTEXT_NAMESPACE
\ No newline at end of file
+#undef LOCTEXT_NAMESPACE
diff --git a/Engine/Plugins/Experimental/ControlRig/Source/ControlRigDeveloper/Public/ControlRigDeveloper.h b/Engine/Plugins/Experimental/ControlRig/Source/ControlRigDeveloper/Public/ControlRigDeveloper.h
--- a/Engine/Plugins/Experimental/ControlRig/Source/ControlRigDeveloper/Public/ControlRigDeveloper.h
+++ b/Engine/Plugins/Experimental/ControlRig/Source/ControlRigDeveloper/Public/ControlRigDeveloper.h
@@ -10,5 +10,8 @@ DECLARE_LOG_CATEGORY_EXTERN(LogControlRigDeveloper, Log, All);
 class IControlRigDeveloperModule : public IModuleInterface
 {
 public:
+	virtual void RegisterPinTypeColor(UStruct* Struct, const FLinearColor Color) = 0;
+	virtual void UnRegisterPinTypeColor(UStruct* Struct) = 0;
+	virtual const FLinearColor* FindPinTypeColor(UStruct* Struct) const = 0;
 };
 
diff --git a/Engine/Plugins/Experimental/ControlRigSpline/ControlRigSpline.uplugin b/Engine/Plugins/Experimental/ControlRigSpline/ControlRigSpline.uplugin
--- a/Engine/Plugins/Experimental/ControlRigSpline/ControlRigSpline.uplugin
+++ b/Engine/Plugins/Experimental/ControlRigSpline/ControlRigSpline.uplugin
@@ -27,6 +27,12 @@
 			"Type": "Runtime",
 			"LoadingPhase": "Default",
 			"WhitelistPlatforms": [ "Win64" ]
+		},
+		{
+			"Name": "ControlRigSplineEditor",
+			"Type": "Editor",
+			"LoadingPhase": "PostEngineInit",
+			"WhitelistPlatforms": [ "Win64" ]
 		}
 	]
 }
\ No newline at end of file
diff --git a/Engine/Plugins/Experimental/ControlRigSpline/Source/ControlRigSplineEditor/ControlRigSplineEditor.build.cs b/Engine/Plugins/Experimental/ControlRigSpline/Source/ControlRigSplineEditor/ControlRigSplineEditor.build.cs
new file mode 100644
--- /dev/null
+++ b/Engine/Plugins/Experimental/ControlRigSpline/Source/ControlRigSplineEditor/ControlRigSplineEditor.build.cs
@@ -0,0 +1,25 @@
+// Copyright Epic Games, Inc. All Rights Reserved.
+
+namespace UnrealBuildTool.Rules
+{
+	public class ControlRigSplineEditor : ModuleRules
+	{
+		public ControlRigSplineEditor(ReadOnlyTargetRules Target) : base(Target)
+		{
+			bEnableUndefinedIdentifierWarnings = false;
+
+			PublicDefinitions.Add("TINYSPLINE_FLOAT_PRECISION=1");
+			PrivateIncludePaths.Add("ControlRigSpline/ThirdParty/TinySpline");
+
+			PublicDependencyModuleNames.AddRange(
+			new string[]
+			{
+				"Core",
+				"CoreUObject",
+				"Engine",
+				"ControlRigDeveloper",
+				"ControlRigSpline"
+			});
+		}
+	}
+}
diff --git a/Engine/Plugins/Experimental/ControlRigSpline/Source/ControlRigSplineEditor/Private/ControlRigSplineEditorModule.cpp b/Engine/Plugins/Experimental/ControlRigSpline/Source/ControlRigSplineEditor/Private/ControlRigSplineEditorModule.cpp
new file mode 100644
--- /dev/null
+++ b/Engine/Plugins/Experimental/ControlRigSpline/Source/ControlRigSplineEditor/Private/ControlRigSplineEditorModule.cpp
@@ -0,0 +1,15 @@
+// Copyright Epic Games, Inc. All Rights Reserved.
+
+#include "ControlRigSplineEditorModule.h"
+#include "ControlRigDeveloper/Public/ControlRigDeveloper.h"
+#include "ControlRigSplineTypes.h"
+
+void FControlRigSplineEditorModule::StartupModule()
+{
+}
+
+void FControlRigSplineEditorModule::ShutdownModule()
+{
+}
+
+IMPLEMENT_MODULE(FControlRigSplineEditorModule, ControlRigSpline)
diff --git a/Engine/Plugins/Experimental/ControlRigSpline/Source/ControlRigSplineEditor/Public/ControlRigSplineEditorModule.h b/Engine/Plugins/Experimental/ControlRigSpline/Source/ControlRigSplineEditor/Public/ControlRigSplineEditorModule.h
new file mode 100644
--- /dev/null
+++ b/Engine/Plugins/Experimental/ControlRigSpline/Source/ControlRigSplineEditor/Public/ControlRigSplineEditorModule.h
@@ -0,0 +1,13 @@
+// Copyright Epic Games, Inc. All Rights Reserved.
+#pragma once
+#include "CoreMinimal.h"
+#include "Modules/ModuleInterface.h"
+#include "Modules/ModuleManager.h"
+
+class FControlRigSplineEditorModule : public IModuleInterface
+{
+	// IModuleInterface interface
+	virtual void StartupModule() override;
+	virtual void ShutdownModule() override;
+};
+
