Nounity fixes

diff --git a/Engine/Plugins/Experimental/ControlRig/Source/ControlRigEditor/Private/ControlRigSequencerEditorLibrary.cpp b/Engine/Plugins/Experimental/ControlRig/Source/ControlRigEditor/Private/ControlRigSequencerEditorLibrary.cpp
--- a/Engine/Plugins/Experimental/ControlRig/Source/ControlRigEditor/Private/ControlRigSequencerEditorLibrary.cpp
+++ b/Engine/Plugins/Experimental/ControlRig/Source/ControlRigEditor/Private/ControlRigSequencerEditorLibrary.cpp
@@ -33,6 +33,7 @@
 #include "Tracks/MovieSceneSkeletalAnimationTrack.h"
 #include "Exporters/AnimSeqExportOption.h"
 #include "MovieSceneTimeHelpers.h"
+#include "ScopedTransaction.h"
 
 #define LOCTEXT_NAMESPACE "ControlrigSequencerEditorLibrary"
 
