Control Rig: Undo changes on SetControlOffset rig unit

diff --git a/Engine/Plugins/Experimental/ControlRig/Source/ControlRig/Private/Units/Hierarchy/RigUnit_SetControlOffset.h b/Engine/Plugins/Experimental/ControlRig/Source/ControlRig/Private/Units/Hierarchy/RigUnit_SetControlOffset.h
--- a/Engine/Plugins/Experimental/ControlRig/Source/ControlRig/Private/Units/Hierarchy/RigUnit_SetControlOffset.h
+++ b/Engine/Plugins/Experimental/ControlRig/Source/ControlRig/Private/Units/Hierarchy/RigUnit_SetControlOffset.h
@@ -30,7 +30,7 @@ struct CONTROLRIG_API FRigUnit_SetControlOffset : public FRigUnitMutable
 	/**
 	 * The offse transform to set for the control
 	 */
-	UPROPERTY(meta = (Input))
+	UPROPERTY(meta = (Input, Output))
 	FTransform Offset;
 
 	/**
diff --git a/Engine/Plugins/Experimental/ControlRig/Source/ControlRig/Private/Units/Hierarchy/RigUnit_SetTransform.cpp b/Engine/Plugins/Experimental/ControlRig/Source/ControlRig/Private/Units/Hierarchy/RigUnit_SetTransform.cpp
--- a/Engine/Plugins/Experimental/ControlRig/Source/ControlRig/Private/Units/Hierarchy/RigUnit_SetTransform.cpp
+++ b/Engine/Plugins/Experimental/ControlRig/Source/ControlRig/Private/Units/Hierarchy/RigUnit_SetTransform.cpp
@@ -47,14 +47,13 @@ FRigUnit_SetTransform_Execute()
 					// for controls - set the control offset transform instead
 					if(bInitial && (CachedIndex.GetKey().Type == ERigElementType::Control))
 					{
-						FRigUnit_SetControlOffset::StaticExecute(RigVMExecuteContext, CachedIndex.GetKey().Name, Transform, Space, CachedIndex, ExecuteContext, Context);
+						FTransform TransformMutable = Transform;
+						FRigUnit_SetControlOffset::StaticExecute(RigVMExecuteContext, CachedIndex.GetKey().Name, TransformMutable, Space, CachedIndex, ExecuteContext, Context);
 						
 						if (ExecuteContext.EventName == FRigUnit_PrepareForExecution::EventName)
 						{
 							Hierarchy->SetLocalTransformByIndex(CachedIndex, FTransform::Identity, true, bPropagateToChildren);
 							Hierarchy->SetLocalTransformByIndex(CachedIndex, FTransform::Identity, false, bPropagateToChildren);
-
-							
 						}
 						return;
 					}
