Sequencer: Control Rig: Make sure we setup correctly after doing an intialize
#trivial

diff --git a/Engine/Plugins/Experimental/ControlRig/Source/ControlRig/Private/Sequencer/MovieSceneControlRigParameterSection.cpp b/Engine/Plugins/Experimental/ControlRig/Source/ControlRig/Private/Sequencer/MovieSceneControlRigParameterSection.cpp
--- a/Engine/Plugins/Experimental/ControlRig/Source/ControlRig/Private/Sequencer/MovieSceneControlRigParameterSection.cpp
+++ b/Engine/Plugins/Experimental/ControlRig/Source/ControlRig/Private/Sequencer/MovieSceneControlRigParameterSection.cpp
@@ -2108,6 +2108,7 @@ bool UMovieSceneControlRigParameterSection::LoadAnimSequenceIntoThisSection(UAni
 
 	ControlRig->SetBoneInitialTransformsFromRefSkeleton(Skeleton->GetReferenceSkeleton());
 	ControlRig->RequestSetup();
+	ControlRig->Evaluate_AnyThread();
 
 	for (int32 Index = 0; Index < NumberOfFrames; ++Index)
 	{
diff --git a/Engine/Plugins/Experimental/ControlRig/Source/ControlRig/Private/Sequencer/MovieSceneControlRigParameterTemplate.cpp b/Engine/Plugins/Experimental/ControlRig/Source/ControlRig/Private/Sequencer/MovieSceneControlRigParameterTemplate.cpp
--- a/Engine/Plugins/Experimental/ControlRig/Source/ControlRig/Private/Sequencer/MovieSceneControlRigParameterTemplate.cpp
+++ b/Engine/Plugins/Experimental/ControlRig/Source/ControlRig/Private/Sequencer/MovieSceneControlRigParameterTemplate.cpp
@@ -562,8 +562,9 @@ void FControlRigBindingHelper::BindToSequencerInstance(UControlRig* ControlRig)
 				AnimInstance->RecalcRequiredBones();
 				AnimInstance->AddControlRigTrack(ControlRig->GetUniqueID(), ControlRig);
 				ControlRig->Initialize();
-
+				ControlRig->RequestInit();
 				ControlRig->SetBoneInitialTransformsFromSkeletalMesh(SkeletalMeshComponent->SkeletalMesh);
+				ControlRig->Evaluate_AnyThread();
 			}
 		}
 	}
diff --git a/Engine/Plugins/Experimental/ControlRig/Source/ControlRigEditor/Private/Sequencer/ControlRigParameterTrackEditor.cpp b/Engine/Plugins/Experimental/ControlRig/Source/ControlRigEditor/Private/Sequencer/ControlRigParameterTrackEditor.cpp
--- a/Engine/Plugins/Experimental/ControlRig/Source/ControlRigEditor/Private/Sequencer/ControlRigParameterTrackEditor.cpp
+++ b/Engine/Plugins/Experimental/ControlRig/Source/ControlRigEditor/Private/Sequencer/ControlRigParameterTrackEditor.cpp
@@ -942,6 +942,8 @@ void FControlRigParameterTrackEditor::BakeToControlRig(UClass* InClass, FGuid Ob
 					ControlRig->GetObjectBinding()->BindToObject(BoundActor);
 					ControlRig->GetDataSourceRegistry()->RegisterDataSource(UControlRig::OwnerComponent, ControlRig->GetObjectBinding()->GetBoundObject());
 					ControlRig->Initialize();
+					ControlRig->RequestInit();
+					ControlRig->SetBoneInitialTransformsFromRefSkeleton(Skeleton->GetReferenceSkeleton());
 					ControlRig->Evaluate_AnyThread();
 
 					UMovieSceneSection* NewSection = Track->CreateControlRigSection(0, ControlRig, bSequencerOwnsControlRig);
