Control Rig: Ported Control Bone Transform fix to UE5

diff --git a/Engine/Plugins/Experimental/ControlRig/Source/ControlRig/Private/AnimNode_ControlRigBase.cpp b/Engine/Plugins/Experimental/ControlRig/Source/ControlRig/Private/AnimNode_ControlRigBase.cpp
--- a/Engine/Plugins/Experimental/ControlRig/Source/ControlRig/Private/AnimNode_ControlRigBase.cpp
+++ b/Engine/Plugins/Experimental/ControlRig/Source/ControlRig/Private/AnimNode_ControlRigBase.cpp
@@ -118,6 +118,11 @@ void FAnimNode_ControlRigBase::UpdateInput(UControlRig* ControlRig, const FPoseC
 
 			ControlRig->GetHierarchy()->SetGlobalTransform(Key, ComponentTransform, false);
 		}
+
+		
+#if WITH_EDITOR
+		ControlRig->ApplyTransformOverrideForUserCreatedBones();
+#endif
 	}
 
 	if (InputSettings.bUpdateCurves)
diff --git a/Engine/Plugins/Experimental/ControlRig/Source/ControlRig/Private/ControlRig.cpp b/Engine/Plugins/Experimental/ControlRig/Source/ControlRig/Private/ControlRig.cpp
--- a/Engine/Plugins/Experimental/ControlRig/Source/ControlRig/Private/ControlRig.cpp
+++ b/Engine/Plugins/Experimental/ControlRig/Source/ControlRig/Private/ControlRig.cpp
@@ -1471,7 +1471,7 @@ FName UControlRig::RemoveTransientControl(URigVMPin* InPin)
 	const FName ControlName = GetNameForTransientControl(InPin);
 	if(FRigControlElement* ControlElement = FindControl(ControlName))
 	{
-		//Hierarchy->Notify(ERigHierarchyNotification::ElementDeselected, ControlElement);
+		DynamicHierarchy->Notify(ERigHierarchyNotification::ElementDeselected, ControlElement);
 		if(Controller->RemoveElement(ControlElement))
 		{
 			return ControlName;
@@ -1554,11 +1554,27 @@ FName UControlRig::AddTransientControl(const FRigElementKey& InElement)
 
 	if (InElement.Type == ERigElementType::Bone)
 	{
-		if(FRigBoneElement* BoneElement = DynamicHierarchy->Find<FRigBoneElement>(InElement))
+		// don't allow transient control to modify forward mode poses when we
+		// already switched to the setup mode
+		if (!IsSetupModeEnabled())
 		{
-			if (PreviewInstance && BoneElement->BoneType == ERigBoneType::Imported)
+			if(FRigBoneElement* BoneElement = DynamicHierarchy->Find<FRigBoneElement>(InElement))
 			{
-				PreviewInstance->ModifyBone(InElement.Name);
+				// add a modify bone AnimNode internally that the transient control controls for imported bones only
+				// for user created bones, refer to UControlRig::TransformOverrideForUserCreatedBones 
+				if (BoneElement->BoneType == ERigBoneType::Imported)
+				{ 
+					if (PreviewInstance)
+					{
+						PreviewInstance->ModifyBone(InElement.Name);
+					}
+				}
+				else if (BoneElement->BoneType == ERigBoneType::User)
+				{
+					// add an empty entry, which will be given the correct value in
+					// SetTransientControlValue(InElement);
+					TransformOverrideForUserCreatedBones.FindOrAdd(InElement.Name);
+				}
 			}
 		}
 	}
@@ -1590,18 +1606,43 @@ bool UControlRig::SetTransientControlValue(const FRigElementKey& InElement)
 	{
 		if (InElement.Type == ERigElementType::Bone)
 		{
-			const FTransform LocalTransform = DynamicHierarchy->GetLocalTransform(InElement);
-			DynamicHierarchy->SetTransform(ControlElement, LocalTransform, ERigTransformType::InitialLocal, true, false);
-			DynamicHierarchy->SetTransform(ControlElement, LocalTransform, ERigTransformType::CurrentLocal, true, false);
-
-			if (PreviewInstance)
+			if (IsSetupModeEnabled())
+			{
+				// need to get initial because that is what setup mode uses
+				// specifically, when user change the initial from the details panel
+				// this will allow the transient control to react to that change
+				const FTransform InitialLocalTransform = DynamicHierarchy->GetInitialLocalTransform(InElement);
+				DynamicHierarchy->SetTransform(ControlElement, InitialLocalTransform, ERigTransformType::InitialLocal, true, false);
+				DynamicHierarchy->SetTransform(ControlElement, InitialLocalTransform, ERigTransformType::CurrentLocal, true, false);
+			}
+			else
 			{
-				if (FAnimNode_ModifyBone* Modify = PreviewInstance->FindModifiedBone(InElement.Name))
+				const FTransform LocalTransform = DynamicHierarchy->GetLocalTransform(InElement);
+				DynamicHierarchy->SetTransform(ControlElement, LocalTransform, ERigTransformType::InitialLocal, true, false);
+				DynamicHierarchy->SetTransform(ControlElement, LocalTransform, ERigTransformType::CurrentLocal, true, false);
+
+				if (FRigBoneElement* BoneElement = DynamicHierarchy->Find<FRigBoneElement>(InElement))
 				{
-					Modify->Translation = LocalTransform.GetTranslation();
-					Modify->Rotation = LocalTransform.GetRotation().Rotator();
-					Modify->TranslationSpace = EBoneControlSpace::BCS_ParentBoneSpace;
-					Modify->RotationSpace = EBoneControlSpace::BCS_ParentBoneSpace;
+					if (BoneElement->BoneType == ERigBoneType::Imported)
+					{
+						if (PreviewInstance)
+						{
+							if (FAnimNode_ModifyBone* Modify = PreviewInstance->FindModifiedBone(InElement.Name))
+							{
+								Modify->Translation = LocalTransform.GetTranslation();
+								Modify->Rotation = LocalTransform.GetRotation().Rotator();
+								Modify->TranslationSpace = EBoneControlSpace::BCS_ParentBoneSpace;
+								Modify->RotationSpace = EBoneControlSpace::BCS_ParentBoneSpace;
+							}
+						}	
+					}
+					else if (BoneElement->BoneType == ERigBoneType::User)
+					{
+						if (FTransform* TransformOverride = TransformOverrideForUserCreatedBones.Find(InElement.Name))
+						{
+							*TransformOverride = LocalTransform;
+						}	
+					}
 				}
 			}
 		}
@@ -1638,7 +1679,7 @@ FName UControlRig::RemoveTransientControl(const FRigElementKey& InElement)
 	const FName ControlName = GetNameForTransientControl(InElement);
 	if(FRigControlElement* ControlElement = FindControl(ControlName))
 	{
-		//Hierarchy->Notify(ERigHierarchyNotification::ElementDeselected, ControlElement);
+		DynamicHierarchy->Notify(ERigHierarchyNotification::ElementDeselected, ControlElement);
 		if(Controller->RemoveElement(ControlElement))
 		{
 			return ControlName;
@@ -1732,6 +1773,19 @@ void UControlRig::ClearTransientControls()
 	}
 }
 
+void UControlRig::ApplyTransformOverrideForUserCreatedBones()
+{
+	if(DynamicHierarchy == nullptr)
+	{
+		return;
+	}
+	
+	for (const auto& Entry : TransformOverrideForUserCreatedBones)
+	{
+		DynamicHierarchy->SetLocalTransform(FRigElementKey(Entry.Key, ERigElementType::Bone), Entry.Value, false);
+	}
+}
+
 #endif
 
 void UControlRig::HandleHierarchyEvent(URigHierarchy* InHierarchy, const FRigEventContext& InEvent)
diff --git a/Engine/Plugins/Experimental/ControlRig/Source/ControlRig/Public/ControlRig.h b/Engine/Plugins/Experimental/ControlRig/Source/ControlRig/Public/ControlRig.h
--- a/Engine/Plugins/Experimental/ControlRig/Source/ControlRig/Public/ControlRig.h
+++ b/Engine/Plugins/Experimental/ControlRig/Source/ControlRig/Public/ControlRig.h
@@ -500,8 +500,17 @@ private:
 
 	UAnimPreviewInstance* PreviewInstance;
 
+	// this is needed because PreviewInstance->ModifyBone(...) cannot modify user created bones,
+	TMap<FName, FTransform> TransformOverrideForUserCreatedBones;
+	
+public:
+	
+	void ApplyTransformOverrideForUserCreatedBones();
+	
 #endif
 
+private: 
+
 	void HandleHierarchyEvent(URigHierarchy* InHierarchy, const FRigEventContext& InEvent);
 	FRigEventDelegate RigEventDelegate;
 
diff --git a/Engine/Plugins/Experimental/ControlRig/Source/ControlRigDeveloper/Private/ControlRigBlueprint.cpp b/Engine/Plugins/Experimental/ControlRig/Source/ControlRigDeveloper/Private/ControlRigBlueprint.cpp
--- a/Engine/Plugins/Experimental/ControlRig/Source/ControlRigDeveloper/Private/ControlRigBlueprint.cpp
+++ b/Engine/Plugins/Experimental/ControlRig/Source/ControlRigDeveloper/Private/ControlRigBlueprint.cpp
@@ -1481,20 +1481,45 @@ FName UControlRigBlueprint::AddTransientControl(const FRigElementKey& InElement)
 		ValueScope = MakeUnique<FControlValueScope>(this);
 	}
 
-	// for now we only allow one pin control at the same time
-	ClearTransientControls();
-
 	UControlRigBlueprintGeneratedClass* RigClass = GetControlRigBlueprintGeneratedClass();
 	UControlRig* CDO = Cast<UControlRig>(RigClass->GetDefaultObject(true /* create if needed */));
 
 	FName ReturnName = NAME_None;
 	TArray<UObject*> ArchetypeInstances;
 	CDO->GetArchetypeInstances(ArchetypeInstances);
+
+	// hierarchy transforms will be reset when ClearTransientControls() is called,
+	// so to retain any bone transform modifications we have to save them
+	TMap<UObject*, FTransform> SavedElementLocalTransforms;
 	for (UObject* ArchetypeInstance : ArchetypeInstances)
 	{
 		UControlRig* InstancedControlRig = Cast<UControlRig>(ArchetypeInstance);
 		if (InstancedControlRig)
 		{
+			if (InstancedControlRig->DynamicHierarchy)
+			{ 
+				SavedElementLocalTransforms.FindOrAdd(InstancedControlRig) = InstancedControlRig->DynamicHierarchy->GetLocalTransform(InElement);
+			}
+		}
+	}
+
+	// for now we only allow one pin control at the same time
+	ClearTransientControls();
+	
+	for (UObject* ArchetypeInstance : ArchetypeInstances)
+	{
+		UControlRig* InstancedControlRig = Cast<UControlRig>(ArchetypeInstance);
+		if (InstancedControlRig)
+		{
+			// restore the element transforms so that transient controls are created at the right place
+			if (const FTransform* SavedTransform = SavedElementLocalTransforms.Find(InstancedControlRig))
+			{
+				if (InstancedControlRig->DynamicHierarchy)
+				{ 
+					InstancedControlRig->DynamicHierarchy->SetLocalTransform(InElement, *SavedTransform);
+				}
+			}
+			
 			FName ControlName = InstancedControlRig->AddTransientControl(InElement);
 			if (ReturnName == NAME_None)
 			{
diff --git a/Engine/Plugins/Experimental/ControlRig/Source/ControlRigEditor/Private/Editor/ControlRigEditor.cpp b/Engine/Plugins/Experimental/ControlRig/Source/ControlRigEditor/Private/Editor/ControlRigEditor.cpp
--- a/Engine/Plugins/Experimental/ControlRig/Source/ControlRigEditor/Private/Editor/ControlRigEditor.cpp
+++ b/Engine/Plugins/Experimental/ControlRig/Source/ControlRigEditor/Private/Editor/ControlRigEditor.cpp
@@ -746,15 +746,18 @@ void FControlRigEditor::ToggleSetupMode()
 		}
 
 		RigBlueprint->Validator->SetControlRig(ControlRig);
-
+		
+		// need to clear selection before remove transient control
+		// because active selection will trigger transient control recreation after removal	
+		PreviousSelection = RigBlueprint->Hierarchy->GetSelectedKeys();
+		RigBlueprint->HierarchyController->ClearSelection();
+		
 		// need to copy here since the removal changes the iterator
 		if (ControlRig)
 		{
 			RigBlueprint->ClearTransientControls();
 		}
 
-		PreviousSelection = RigBlueprint->Hierarchy->GetSelectedKeys();
-		RigBlueprint->HierarchyController->ClearSelection();
 	}
 
 	if (ControlRig)
@@ -1324,6 +1327,16 @@ void FControlRigEditor::Compile()
 			bSetupModeEnabled = false;
 		}
 
+		// clear transient controls such that we don't leave
+		// a phantom gizmo in the viewport
+		// have to do this before compile() because during compile
+		// a new control rig instance is created without the transient controls
+		// so clear is never called for old transient controls
+		RigBlueprint->ClearTransientControls();
+
+		// default to always reset all bone modifications 
+		ResetAllBoneModification(); 
+		
 		{
 			FBlueprintEditor::Compile();
 		}
@@ -1377,11 +1390,6 @@ void FControlRigEditor::Compile()
 			}
 		}
 
-		if (PreviewInstance)
-		{
-			PreviewInstance->ResetModifiedBone();
-		}
-
 		if (UControlRigSettings::Get()->bResetControlTransformsOnCompile)
 		{
 			RigBlueprint->Hierarchy->ForEach<FRigControlElement>([RigBlueprint](FRigControlElement* ControlElement) -> bool
@@ -3728,6 +3736,72 @@ void FControlRigEditor::SynchronizeViewportBoneSelection()
 	}
 }
 
+void FControlRigEditor::UpdateBoneModification(FName BoneName, const FTransform& LocalTransform)
+{
+	if (ControlRig)
+	{ 
+		if (PreviewInstance)
+		{ 
+			if (FAnimNode_ModifyBone* Modify = PreviewInstance->FindModifiedBone(BoneName))
+			{
+				Modify->Translation = LocalTransform.GetTranslation();
+				Modify->Rotation = LocalTransform.GetRotation().Rotator();
+				Modify->TranslationSpace = EBoneControlSpace::BCS_ParentBoneSpace;
+				Modify->RotationSpace = EBoneControlSpace::BCS_ParentBoneSpace; 
+			}
+		}
+		
+		TMap<FName, FTransform>* TransformOverrideMap = &ControlRig->TransformOverrideForUserCreatedBones;
+		if (UControlRig* DebuggedControlRig = Cast<UControlRig>(GetBlueprintObj()->GetObjectBeingDebugged()))
+		{
+			TransformOverrideMap = &DebuggedControlRig->TransformOverrideForUserCreatedBones;
+		}
+
+		if (FTransform* Transform = TransformOverrideMap->Find(BoneName))
+		{
+			*Transform = LocalTransform;
+		}
+	}
+}
+
+void FControlRigEditor::RemoveBoneModification(FName BoneName)
+{
+	if (ControlRig)
+	{
+		if (PreviewInstance)
+		{
+			PreviewInstance->RemoveBoneModification(BoneName);
+		}
+
+		TMap<FName, FTransform>* TransformOverrideMap = &ControlRig->TransformOverrideForUserCreatedBones;
+		if (UControlRig* DebuggedControlRig = Cast<UControlRig>(GetBlueprintObj()->GetObjectBeingDebugged()))
+		{
+			TransformOverrideMap = &DebuggedControlRig->TransformOverrideForUserCreatedBones;
+		}
+
+		TransformOverrideMap->Remove(BoneName);
+	}
+}
+
+void FControlRigEditor::ResetAllBoneModification()
+{
+	if (ControlRig)
+	{
+		if (PreviewInstance)
+		{
+			PreviewInstance->ResetModifiedBone();
+		}
+
+		TMap<FName, FTransform>* TransformOverrideMap = &ControlRig->TransformOverrideForUserCreatedBones;
+		if (UControlRig* DebuggedControlRig = Cast<UControlRig>(GetBlueprintObj()->GetObjectBeingDebugged()))
+		{
+			TransformOverrideMap = &DebuggedControlRig->TransformOverrideForUserCreatedBones;
+		}
+
+		TransformOverrideMap->Reset();
+	}
+}
+
 FControlRigEditorEditMode* FControlRigEditor::GetEditMode() const
 {
 	return static_cast<FControlRigEditorEditMode*>(GetEditorModeManager().GetActiveMode(FControlRigEditorEditMode::ModeName));
@@ -4365,16 +4439,9 @@ void FControlRigEditor::HandleOnControlModified(UControlRig* Subject, FRigContro
 					Blueprint->Hierarchy->SetInitialLocalTransform(ElementKey, Transform);
 					Hierarchy->SetInitialLocalTransform(ElementKey, Transform);
 				}
-
-				if (PreviewInstance)
-				{
-					if (FAnimNode_ModifyBone* Modify = PreviewInstance->FindModifiedBone(ElementKey.Name))
-					{
-						Modify->Translation = Transform.GetTranslation();
-						Modify->Rotation = Transform.GetRotation().Rotator();
-						Modify->TranslationSpace = EBoneControlSpace::BCS_ParentBoneSpace;
-						Modify->RotationSpace = EBoneControlSpace::BCS_ParentBoneSpace;
-					}
+				else
+				{ 
+					UpdateBoneModification(ElementKey.Name, Transform);
 				}
 			}
 			else if (ElementKey.Type == ERigElementType::Null)
diff --git a/Engine/Plugins/Experimental/ControlRig/Source/ControlRigEditor/Private/Editor/ControlRigEditor.h b/Engine/Plugins/Experimental/ControlRig/Source/ControlRigEditor/Private/Editor/ControlRigEditor.h
--- a/Engine/Plugins/Experimental/ControlRig/Source/ControlRigEditor/Private/Editor/ControlRigEditor.h
+++ b/Engine/Plugins/Experimental/ControlRig/Source/ControlRigEditor/Private/Editor/ControlRigEditor.h
@@ -180,6 +180,15 @@ protected:
 	void OnControlsSettingsChanged();
 
 	void SynchronizeViewportBoneSelection();
+	
+	// update the cached modification value
+	void UpdateBoneModification(FName BoneName, const FTransform& Transform);
+
+	// remove a single bone modification across all instance
+	void RemoveBoneModification(FName BoneName);
+
+	// reset all bone modification across all instance
+	void ResetAllBoneModification();
 
 	// FBlueprintEditor Interface
 	virtual void CreateDefaultCommands() override;
diff --git a/Engine/Plugins/Experimental/ControlRig/Source/ControlRigEditor/Private/Editor/SRigHierarchy.cpp b/Engine/Plugins/Experimental/ControlRig/Source/ControlRigEditor/Private/Editor/SRigHierarchy.cpp
--- a/Engine/Plugins/Experimental/ControlRig/Source/ControlRigEditor/Private/Editor/SRigHierarchy.cpp
+++ b/Engine/Plugins/Experimental/ControlRig/Source/ControlRigEditor/Private/Editor/SRigHierarchy.cpp
@@ -2372,6 +2372,9 @@ void SRigHierarchy::HandleResetTransform(bool bSelectionOnly)
 				if (!bSelectionOnly)
 				{
 					KeysToReset = DebuggedHierarchy->GetAllKeys(true, ERigElementType::Control);
+
+					// Bone Transforms can also be modified, support reset for them as well
+					KeysToReset.Append(DebuggedHierarchy->GetAllKeys(true, ERigElementType::Bone));
 				}
 
 				for (FRigElementKey Key : KeysToReset)
@@ -2379,6 +2382,12 @@ void SRigHierarchy::HandleResetTransform(bool bSelectionOnly)
 					const FTransform InitialTransform = GetHierarchy()->GetInitialLocalTransform(Key);
 					GetHierarchy()->SetLocalTransform(Key, InitialTransform, false, true, true);
 					DebuggedHierarchy->SetLocalTransform(Key, InitialTransform, false, true, true);
+
+					if (Key.Type == ERigElementType::Bone)
+					{
+						Blueprint->RemoveTransientControl(Key);
+						ControlRigEditor.Pin()->RemoveBoneModification(Key.Name); 
+					}
 				}
 			}
 		}
@@ -2486,11 +2495,6 @@ void SRigHierarchy::HandleControlBoneOrSpaceTransform()
 		return;
 	}
 
-	if(!DebuggedControlRig->IsSetupModeEnabled())
-	{
-		return;
-	}
-
 	TArray<FRigElementKey> SelectedKeys = GetSelectedKeys();
 	for (const FRigElementKey& SelectedKey : SelectedKeys)
 	{
@@ -2679,15 +2683,6 @@ void SRigHierarchy::HandleSetInitialTransformFromClosestBone()
                         SelectedKey.Type == ERigElementType::Bone)
 					{
 						FTransform InitialTransform = LocalTransform;
-						if (ControlRigEditor.Pin()->PreviewInstance)
-						{
-							if (FAnimNode_ModifyBone* ModifyBone = ControlRigEditor.Pin()->PreviewInstance->FindModifiedBone(SelectedKey.Name))
-							{
-								InitialTransform.SetTranslation(ModifyBone->Translation);
-								InitialTransform.SetRotation(FQuat(ModifyBone->Rotation));
-								InitialTransform.SetScale3D(ModifyBone->Scale);
-							}
-						}
 
 						if(FRigTransformElement* TransformElement = GetHierarchy()->Find<FRigTransformElement>(SelectedKey))
 						{
