Control Rig: Recompile when preview mesh changed

diff --git a/Engine/Plugins/Experimental/ControlRig/Source/ControlRigEditor/Private/Editor/ControlRigEditor.cpp b/Engine/Plugins/Experimental/ControlRig/Source/ControlRigEditor/Private/Editor/ControlRigEditor.cpp
--- a/Engine/Plugins/Experimental/ControlRig/Source/ControlRigEditor/Private/Editor/ControlRigEditor.cpp
+++ b/Engine/Plugins/Experimental/ControlRig/Source/ControlRigEditor/Private/Editor/ControlRigEditor.cpp
@@ -2891,7 +2891,8 @@ void FControlRigEditor::HandlePreviewMeshChanged(USkeletalMesh* InOldSkeletalMes
 		if (UControlRigBlueprint* ControlRigBP = GetControlRigBlueprint())
 		{
 			ControlRigBP->SetPreviewMesh(InNewSkeletalMesh);
-
+			UpdateControlRig();
+			
 			if(UControlRig* DebuggedControlRig = Cast<UControlRig>(GetBlueprintObj()->GetObjectBeingDebugged()))
 			{
 				DebuggedControlRig->GetHierarchy()->Notify(ERigHierarchyNotification::HierarchyReset, nullptr);
