Unblocking robomerge

Sequencer:Control Rig: New icons for Snap/Pose/Tween/CreatePose/SelectOnly

diff --git a/Engine/Plugins/Experimental/ControlRig/Source/ControlRig/Private/Tools/ControlRigPose.cpp b/Engine/Plugins/Experimental/ControlRig/Source/ControlRig/Private/Tools/ControlRigPose.cpp
--- a/Engine/Plugins/Experimental/ControlRig/Source/ControlRig/Private/Tools/ControlRigPose.cpp
+++ b/Engine/Plugins/Experimental/ControlRig/Source/ControlRig/Private/Tools/ControlRigPose.cpp
@@ -74,69 +74,66 @@ void FControlRigControlPose::PastePoseInternal(UControlRig* ControlRig, bool bDo
 		FRigControlCopy* CopyRigControl = MirrorTable.GetControl(*this, ControlElement->GetName());
 		if (CopyRigControl)
 		{
-			if (CopyRigControl->ControlType == ControlElement->Settings.ControlType)
+			switch (ControlElement->Settings.ControlType)
 			{
-				switch (ControlElement->Settings.ControlType)
-				{
-				case ERigControlType::Transform:
-				case ERigControlType::TransformNoScale:
-				case ERigControlType::EulerTransform:
-				case ERigControlType::Position:
-				case ERigControlType::Scale:
-				case ERigControlType::Rotator:
+			case ERigControlType::Transform:
+			case ERigControlType::TransformNoScale:
+			case ERigControlType::EulerTransform:
+			case ERigControlType::Position:
+			case ERigControlType::Scale:
+			case ERigControlType::Rotator:
+			{
+				if (bDoMirror == false)
 				{
-					if (bDoMirror == false)
+					if (bDoLocal) // -V547  
 					{
-						if (bDoLocal) // -V547  
-						{
-							ControlRig->SetControlLocalTransform(ControlElement->GetName(), CopyRigControl->LocalTransform, true,Context,bSetupUndo);
-						}
-						else
-						{
-							ControlRig->SetControlGlobalTransform(ControlElement->GetName(), CopyRigControl->GlobalTransform, true,Context,bSetupUndo);
-						}
+						ControlRig->SetControlLocalTransform(ControlElement->GetName(), CopyRigControl->LocalTransform, true,Context,bSetupUndo);
 					}
 					else
 					{
-						FVector GlobalTranslation;
-						FQuat GlobalRotation;
-						FVector LocalTranslation;
-						FQuat LocalRotation;
-						bool bIsMatched = MirrorTable.IsMatched(CopyRigControl->Name);
-						MirrorTable.GetMirrorTransform(*CopyRigControl, bDoLocal,bIsMatched, GlobalTranslation, GlobalRotation, LocalTranslation, LocalRotation);
-						SetControlMirrorTransform(bDoLocal,ControlRig, ControlElement->GetName(), bIsMatched, GlobalTranslation, GlobalRotation, LocalTranslation,LocalRotation, true,Context,bSetupUndo);
-					}				
-					break;
-				}
-				case ERigControlType::Float:
-				{
-					float Val = CopyRigControl->Value.Get<float>();
-					ControlRig->SetControlValue<float>(ControlElement->GetName(), Val, true, Context,bSetupUndo);
-					break;
-				}
-				case ERigControlType::Bool:
-				{
-					bool Val = CopyRigControl->Value.Get<bool>();
-					ControlRig->SetControlValue<bool>(ControlElement->GetName(), Val, true, Context,bSetupUndo);
-					break;
-				}
-				case ERigControlType::Integer:
-				{
-					int32 Val = CopyRigControl->Value.Get<int32>();
-					ControlRig->SetControlValue<int32>(ControlElement->GetName(), Val, true, Context,bSetupUndo);
-					break;
+						ControlRig->SetControlGlobalTransform(ControlElement->GetName(), CopyRigControl->GlobalTransform, true,Context,bSetupUndo);
+					}
 				}
-				case ERigControlType::Vector2D:
+				else
 				{
-					FVector2D Val = CopyRigControl->Value.Get<FVector2D>();
-					ControlRig->SetControlValue<FVector2D>(ControlElement->GetName(), Val, true, Context,bSetupUndo);
-					break;
-				}
-				default:
-					//TODO add log
-					break;
-				};
+					FVector GlobalTranslation;
+					FQuat GlobalRotation;
+					FVector LocalTranslation;
+					FQuat LocalRotation;
+					bool bIsMatched = MirrorTable.IsMatched(CopyRigControl->Name);
+					MirrorTable.GetMirrorTransform(*CopyRigControl, bDoLocal,bIsMatched, GlobalTranslation, GlobalRotation, LocalTranslation, LocalRotation);
+					SetControlMirrorTransform(bDoLocal,ControlRig, ControlElement->GetName(), bIsMatched, GlobalTranslation, GlobalRotation, LocalTranslation,LocalRotation, true,Context,bSetupUndo);
+				}				
+				break;
 			}
+			case ERigControlType::Float:
+			{
+				float Val = CopyRigControl->Value.Get<float>();
+				ControlRig->SetControlValue<float>(ControlElement->GetName(), Val, true, Context,bSetupUndo);
+				break;
+			}
+			case ERigControlType::Bool:
+			{
+				bool Val = CopyRigControl->Value.Get<bool>();
+				ControlRig->SetControlValue<bool>(ControlElement->GetName(), Val, true, Context,bSetupUndo);
+				break;
+			}
+			case ERigControlType::Integer:
+			{
+				int32 Val = CopyRigControl->Value.Get<int32>();
+				ControlRig->SetControlValue<int32>(ControlElement->GetName(), Val, true, Context,bSetupUndo);
+				break;
+			}
+			case ERigControlType::Vector2D:
+			{
+				FVector2D Val = CopyRigControl->Value.Get<FVector2D>();
+				ControlRig->SetControlValue<FVector2D>(ControlElement->GetName(), Val, true, Context,bSetupUndo);
+				break;
+			}
+			default:
+				//TODO add log
+				break;
+			};
 		}
 	}
 }
@@ -172,68 +169,64 @@ void FControlRigControlPose::BlendWithInitialPoses(FControlRigControlPose& Initi
 		FRigControlCopy* CopyRigControl = MirrorTable.GetControl(*this, ControlElement->GetName());
 		if (CopyRigControl)
 		{
-			if (CopyRigControl->ControlType == ControlElement->Settings.ControlType)
+			FRigControlCopy* InitialFound = nullptr;
+			int32* Index = InitialPose.CopyOfControlsNameToIndex.Find(CopyRigControl->Name);
+			if (Index)
 			{
-
-				FRigControlCopy* InitialFound = nullptr;
-				int32* Index = InitialPose.CopyOfControlsNameToIndex.Find(CopyRigControl->Name);
-				if (Index)
-				{
-					InitialFound = &(InitialPose.CopyOfControls[*Index]);
-				}
-				if (InitialFound && InitialFound->ControlType == CopyRigControl->ControlType)
+				InitialFound = &(InitialPose.CopyOfControls[*Index]);
+			}
+			if (InitialFound && InitialFound->ControlType == CopyRigControl->ControlType)
+			{
+				if ((CopyRigControl->ControlType == ERigControlType::Transform || CopyRigControl->ControlType == ERigControlType::EulerTransform ||
+					CopyRigControl->ControlType == ERigControlType::TransformNoScale || CopyRigControl->ControlType == ERigControlType::Position ||
+					CopyRigControl->ControlType == ERigControlType::Rotator || CopyRigControl->ControlType == ERigControlType::Scale
+					))
 				{
-					if ((CopyRigControl->ControlType == ERigControlType::Transform || CopyRigControl->ControlType == ERigControlType::EulerTransform ||
-						CopyRigControl->ControlType == ERigControlType::TransformNoScale || CopyRigControl->ControlType == ERigControlType::Position ||
-						CopyRigControl->ControlType == ERigControlType::Rotator || CopyRigControl->ControlType == ERigControlType::Scale
-						))
+					if (bDoMirror == false)
 					{
-						if (bDoMirror == false)
+						if (bDoLocal == true)    // -V547  
 						{
-							if (bDoLocal == true)    // -V547  
-							{
-								FTransform Val = CopyRigControl->LocalTransform;
-								FTransform InitialVal = InitialFound->LocalTransform;
-								FVector Translation, Scale;
-								FQuat Rotation;
-								Translation = FMath::Lerp(InitialVal.GetTranslation(), Val.GetTranslation(), BlendValue);
-								Rotation = FQuat::Slerp(InitialVal.GetRotation(), Val.GetRotation(), BlendValue); //doing slerp here not fast lerp, can be slow this is for content creation
-								Scale = FMath::Lerp(InitialVal.GetScale3D(), Val.GetScale3D(), BlendValue);
-								Val = FTransform(Rotation, Translation, Scale);
-								ControlRig->SetControlLocalTransform(ControlElement->GetName(), Val, bDoKey,Context,bSetupUndo);
-							}
-							else
-							{
-								FTransform Val = CopyRigControl->GlobalTransform;
-								FTransform InitialVal = InitialFound->GlobalTransform;
-								FVector Translation, Scale;
-								FQuat Rotation;
-								Translation = FMath::Lerp(InitialVal.GetTranslation(), Val.GetTranslation(), BlendValue);
-								Rotation = FQuat::Slerp(InitialVal.GetRotation(), Val.GetRotation(), BlendValue); //doing slerp here not fast lerp, can be slow this is for content creation
-								Scale = FMath::Lerp(InitialVal.GetScale3D(), Val.GetScale3D(), BlendValue);
-								Val = FTransform(Rotation, Translation, Scale);
-								ControlRig->SetControlGlobalTransform(ControlElement->GetName(), Val, bDoKey,Context,bSetupUndo);
-							}
+							FTransform Val = CopyRigControl->LocalTransform;
+							FTransform InitialVal = InitialFound->LocalTransform;
+							FVector Translation, Scale;
+							FQuat Rotation;
+							Translation = FMath::Lerp(InitialVal.GetTranslation(), Val.GetTranslation(), BlendValue);
+							Rotation = FQuat::Slerp(InitialVal.GetRotation(), Val.GetRotation(), BlendValue); //doing slerp here not fast lerp, can be slow this is for content creation
+							Scale = FMath::Lerp(InitialVal.GetScale3D(), Val.GetScale3D(), BlendValue);
+							Val = FTransform(Rotation, Translation, Scale);
+							ControlRig->SetControlLocalTransform(ControlElement->GetName(), Val, bDoKey,Context,bSetupUndo);
 						}
 						else
 						{
-							FVector GlobalTranslation;
-							FQuat GlobalRotation;
-							FVector LocalTranslation;
-							FQuat LocalRotation;
-							bool bIsMatched = MirrorTable.IsMatched(CopyRigControl->Name);
-							MirrorTable.GetMirrorTransform(*CopyRigControl,bDoLocal, bIsMatched, GlobalTranslation, GlobalRotation, LocalTranslation, LocalRotation);
-							FVector InitialTranslation = InitialFound->GlobalTransform.GetTranslation();
-							FQuat InitialGlobalRotation = InitialFound->GlobalTransform.GetRotation();
-							FVector InitialLocalTranslation = InitialFound->LocalTransform.GetTranslation();
-							FQuat InitialLocationRotation = InitialFound->LocalTransform.GetRotation();
-							GlobalTranslation = FMath::Lerp(InitialTranslation, GlobalTranslation, BlendValue);
-							GlobalRotation = FQuat::Slerp(InitialGlobalRotation, GlobalRotation, BlendValue); //doing slerp here not fast lerp, can be slow this is for content creation
-							LocalTranslation = FMath::Lerp(InitialLocalTranslation, LocalTranslation, BlendValue);
-							LocalRotation = FQuat::Slerp(InitialLocationRotation, LocalRotation, BlendValue); //doing slerp here not fast lerp, can be slow this is for content creation
-							SetControlMirrorTransform(bDoLocal,ControlRig, ControlElement->GetName(), bIsMatched, GlobalTranslation, GlobalRotation,LocalTranslation, LocalRotation, bDoKey,Context,bSetupUndo);							
+							FTransform Val = CopyRigControl->GlobalTransform;
+							FTransform InitialVal = InitialFound->GlobalTransform;
+							FVector Translation, Scale;
+							FQuat Rotation;
+							Translation = FMath::Lerp(InitialVal.GetTranslation(), Val.GetTranslation(), BlendValue);
+							Rotation = FQuat::Slerp(InitialVal.GetRotation(), Val.GetRotation(), BlendValue); //doing slerp here not fast lerp, can be slow this is for content creation
+							Scale = FMath::Lerp(InitialVal.GetScale3D(), Val.GetScale3D(), BlendValue);
+							Val = FTransform(Rotation, Translation, Scale);
+							ControlRig->SetControlGlobalTransform(ControlElement->GetName(), Val, bDoKey,Context,bSetupUndo);
 						}
 					}
+					else
+					{
+						FVector GlobalTranslation;
+						FQuat GlobalRotation;
+						FVector LocalTranslation;
+						FQuat LocalRotation;
+						bool bIsMatched = MirrorTable.IsMatched(CopyRigControl->Name);
+						MirrorTable.GetMirrorTransform(*CopyRigControl,bDoLocal, bIsMatched, GlobalTranslation, GlobalRotation, LocalTranslation, LocalRotation);
+						FVector InitialTranslation = InitialFound->GlobalTransform.GetTranslation();
+						FQuat InitialGlobalRotation = InitialFound->GlobalTransform.GetRotation();
+						FVector InitialLocalTranslation = InitialFound->LocalTransform.GetTranslation();
+						FQuat InitialLocationRotation = InitialFound->LocalTransform.GetRotation();
+						GlobalTranslation = FMath::Lerp(InitialTranslation, GlobalTranslation, BlendValue);
+						GlobalRotation = FQuat::Slerp(InitialGlobalRotation, GlobalRotation, BlendValue); //doing slerp here not fast lerp, can be slow this is for content creation
+						LocalTranslation = FMath::Lerp(InitialLocalTranslation, LocalTranslation, BlendValue);
+						LocalRotation = FQuat::Slerp(InitialLocationRotation, LocalRotation, BlendValue); //doing slerp here not fast lerp, can be slow this is for content creation
+						SetControlMirrorTransform(bDoLocal,ControlRig, ControlElement->GetName(), bIsMatched, GlobalTranslation, GlobalRotation,LocalTranslation, LocalRotation, bDoKey,Context,bSetupUndo);							
+					}
 				}
 			}
 		}
diff --git a/Engine/Plugins/Experimental/ControlRig/Source/ControlRigEditor/Private/ControlRigEditorStyle.h b/Engine/Plugins/Experimental/ControlRig/Source/ControlRigEditor/Private/ControlRigEditorStyle.h
--- a/Engine/Plugins/Experimental/ControlRig/Source/ControlRigEditor/Private/ControlRigEditorStyle.h
+++ b/Engine/Plugins/Experimental/ControlRig/Source/ControlRigEditor/Private/ControlRigEditorStyle.h
@@ -49,6 +49,16 @@ public:
 			Set("ControlRig.ReImportFromRigSequence.Small", new IMAGE_BRUSH("Slate/ReImportRigSequence_16x", Icon16x16));
 		}
 
+		//Tool Styles
+		{
+			Set("ControlRig.OnlySelectControls", new IMAGE_BRUSH_SVG("Slate/AnimationSelectOnlyControlRig", Icon16x16));
+			Set("ControlRig.SnapperTool", new IMAGE_BRUSH_SVG("Slate/AnimationControlRigSnapper", Icon16x16));
+			Set("ControlRig.PoseTool", new IMAGE_BRUSH_SVG("Slate/AnimationPoses", Icon16x16));
+			Set("ControlRig.CreatePose", new IMAGE_BRUSH_SVG("Slate/AnimationCreatePose", Icon16x16));
+			Set("ControlRig.TweenTool", new IMAGE_BRUSH_SVG("Slate/AnimationTweens", Icon16x16));
+
+		}
+
 		// Control Rig Editor styles
 		{
 			Set("ControlRig.TabIcon", new IMAGE_BRUSH("Slate/ControlRigTab_16x", Icon16x16));
diff --git a/Engine/Plugins/Experimental/ControlRig/Source/ControlRigEditor/Private/EditMode/SControlRigBaseListWidget.cpp b/Engine/Plugins/Experimental/ControlRig/Source/ControlRigEditor/Private/EditMode/SControlRigBaseListWidget.cpp
--- a/Engine/Plugins/Experimental/ControlRig/Source/ControlRigEditor/Private/EditMode/SControlRigBaseListWidget.cpp
+++ b/Engine/Plugins/Experimental/ControlRig/Source/ControlRigEditor/Private/EditMode/SControlRigBaseListWidget.cpp
@@ -273,6 +273,7 @@ void SControlRigPoseAnimSelectionToolbar::Construct(const FArguments& InArgs)
 	OwningControlRigWidget = InArgs._OwningControlRigWidget;
 
 	FToolBarBuilder ToolbarBuilder(TSharedPtr<const FUICommandList>(), FMultiBoxCustomization::None, TSharedPtr<FExtender>(), true);
+	ToolbarBuilder.SetLabelVisibility(EVisibility::Visible);
 	FUIAction CreatePoseDialog(
 		FExecuteAction::CreateRaw(this, &SControlRigPoseAnimSelectionToolbar::MakeControlRigAssetDialog, FControlRigAssetType::ControlRigPose,false));
 	/*
@@ -290,9 +291,9 @@ void SControlRigPoseAnimSelectionToolbar::Construct(const FArguments& InArgs)
 	{
 		ToolbarBuilder.AddToolBarButton(CreatePoseDialog,
 			NAME_None,
-			LOCTEXT("CreatePose", "Create Pose From Selection"),
+			LOCTEXT("CreatePose", "Create Pose"),
 			LOCTEXT("CreatePoseTooltip", "Create pose asset from selection."),
-			FSlateIcon(),
+			FSlateIcon(TEXT("ControlRigEditorStyle"), TEXT("ControlRig.CreatePose")),
 			EUserInterfaceActionType::Button
 		);
 		/*
diff --git a/Engine/Plugins/Experimental/ControlRig/Source/ControlRigEditor/Private/EditMode/SControlRigEditModeTools.cpp b/Engine/Plugins/Experimental/ControlRig/Source/ControlRigEditor/Private/EditMode/SControlRigEditModeTools.cpp
--- a/Engine/Plugins/Experimental/ControlRig/Source/ControlRigEditor/Private/EditMode/SControlRigEditModeTools.cpp
+++ b/Engine/Plugins/Experimental/ControlRig/Source/ControlRigEditor/Private/EditMode/SControlRigEditModeTools.cpp
@@ -28,6 +28,7 @@
 #include "SControlRigTweenWidget.h"
 #include "IControlRigEditorModule.h"
 #include "Framework/Docking/TabManager.h"
+#include "ControlRigEditorStyle.h"
 
 #define LOCTEXT_NAMESPACE "ControlRigRootCustomization"
 
@@ -524,9 +525,9 @@ void SControlRigEditModeTools::CustomizeToolBarPalette(FToolBarBuilder& ToolBarB
 	})
 		),
 		NAME_None,
-		LOCTEXT("OnlySelectControls", "Only Select Controls"),
+		LOCTEXT("OnlySelectControls", "Select"),
 		LOCTEXT("OnlySelectControlsTooltip", "Only Select Control Rig Controls"),
-		FSlateIcon(FEditorStyle::GetStyleSetName(), "FoliageEditMode.SetPaint"), //todo get icon
+		FSlateIcon(TEXT("ControlRigEditorStyle"), TEXT("ControlRig.OnlySelectControls")),
 		EUserInterfaceActionType::ToggleButton
 		);
 	ToolBarBuilder.AddSeparator();
@@ -537,7 +538,7 @@ void SControlRigEditModeTools::CustomizeToolBarPalette(FToolBarBuilder& ToolBarB
 		NAME_None,
 		LOCTEXT("Poses", "Poses"),
 		LOCTEXT("PosesTooltip", "Show Poses"),
-		FSlateIcon(FEditorStyle::GetStyleSetName(), "FoliageEditMode.Settings"), //MZ todo replace with correct icon
+		FSlateIcon(TEXT("ControlRigEditorStyle"), TEXT("ControlRig.PoseTool")),
 		EUserInterfaceActionType::Button
 	);
 	ToolBarBuilder.AddSeparator();
@@ -548,7 +549,7 @@ void SControlRigEditModeTools::CustomizeToolBarPalette(FToolBarBuilder& ToolBarB
 		NAME_None,
 		LOCTEXT("Tweens", "Tweens"),
 		LOCTEXT("TweensTooltip", "Create Tweens"),
-		FSlateIcon(FEditorStyle::GetStyleSetName(), "SkeletonTree.BlendProfile"), //MZ todo replace with correct icon
+		FSlateIcon(TEXT("ControlRigEditorStyle"), TEXT("ControlRig.TweenTool")),
 		EUserInterfaceActionType::Button
 	);
 
@@ -557,8 +558,8 @@ void SControlRigEditModeTools::CustomizeToolBarPalette(FToolBarBuilder& ToolBarB
 		FExecuteAction::CreateSP(this, &SControlRigEditModeTools::MakeSnapperDialog),
 		NAME_None,
 		LOCTEXT("Snapper", "Snapper"),
-		LOCTEXT("SnapperTooltip", "Snap Control to another over a set of frames"),
-		FSlateIcon(FEditorStyle::GetStyleSetName(), "Persona.Tabs.AnimSlotManager"), //MZ todo replace with correct icon
+		LOCTEXT("SnapperTooltip", "Snap child objects to a parent object over a set of frames"),
+		FSlateIcon(TEXT("ControlRigEditorStyle"), TEXT("ControlRig.SnapperTool")),
 		EUserInterfaceActionType::Button
 	);
 	/*
