Control Rig: Only set curves on anim instance that have really changed

diff --git a/Engine/Plugins/Experimental/ControlRig/Source/ControlRig/Private/AnimNode_ControlRigBase.cpp b/Engine/Plugins/Experimental/ControlRig/Source/ControlRig/Private/AnimNode_ControlRigBase.cpp
--- a/Engine/Plugins/Experimental/ControlRig/Source/ControlRig/Private/AnimNode_ControlRigBase.cpp
+++ b/Engine/Plugins/Experimental/ControlRig/Source/ControlRig/Private/AnimNode_ControlRigBase.cpp
@@ -181,8 +181,15 @@ void FAnimNode_ControlRigBase::UpdateOutput(UControlRig* ControlRig, FPoseContex
 			const uint16 Index = Iter.Value();
 			const FRigElementKey Key(Name, ERigElementType::Curve);
 
+			const float PreviousValue = InOutput.Curve.Get(Index);
 			const float Value = ControlRig->GetHierarchy()->GetCurveValue(Key);
-			InOutput.Curve.Set(Index, Value);
+
+			if(!FMath::IsNearlyEqual(PreviousValue, Value))
+			{
+				// this causes a side effect of marking the curve as "valid"
+				// so only apply it for curves that have really changed
+				InOutput.Curve.Set(Index, Value);
+			}
 		}
 	}
 }
