Sequencer: Control Rig: Need to use retargeted SkeletalMeshComponent refpose to set CR's bone initial transofrms instead of the RefSkeleton pose.
#frosty FROST-1408

diff --git a/Engine/Plugins/Experimental/ControlRig/Source/ControlRig/Private/ControlRig.cpp b/Engine/Plugins/Experimental/ControlRig/Source/ControlRig/Private/ControlRig.cpp
--- a/Engine/Plugins/Experimental/ControlRig/Source/ControlRig/Private/ControlRig.cpp
+++ b/Engine/Plugins/Experimental/ControlRig/Source/ControlRig/Private/ControlRig.cpp
@@ -12,6 +12,9 @@
 #include "Units/Execution/RigUnit_BeginExecution.h"
 #include "Units/Execution/RigUnit_PrepareForExecution.h"
 #include "Units/Execution/RigUnit_InverseExecution.h"
+#include "Components/SkeletalMeshComponent.h"
+#include "Animation/AnimInstance.h"
+#include "Animation/AnimationPoseData.h"
 #if WITH_EDITOR
 #include "ControlRigModule.h"
 #include "Modules/ModuleManager.h"
@@ -2328,6 +2331,42 @@ FRigVMExternalVariable UControlRig::GetExternalVariableFromDescription(const FBP
 
 #endif // WITH_EDITOR
 
+void UControlRig::SetBoneInitialTransformsFromSkeletalMeshComponent(USkeletalMeshComponent* InSkelMeshComp)
+{
+	check(InSkelMeshComp);
+	check(DynamicHierarchy);
+	if (InSkelMeshComp->GetAnimInstance() == nullptr)
+	{
+		SetBoneInitialTransformsFromSkeletalMesh(InSkelMeshComp->SkeletalMesh);
+		return;
+	}
+	else
+	{
+		FMemMark Mark(FMemStack::Get());
+		FCompactPose OutPose;
+		OutPose.ResetToRefPose(InSkelMeshComp->GetAnimInstance()->GetRequiredBones());
+
+		DynamicHierarchy->ForEach<FRigBoneElement>([this, &OutPose, InSkelMeshComp](FRigBoneElement* BoneElement) -> bool
+			{
+				if (BoneElement->BoneType == ERigBoneType::Imported)
+				{
+					int32 MeshIndex = OutPose.GetBoneContainer().GetPoseBoneIndexForBoneName(BoneElement->GetName());
+					if (MeshIndex != INDEX_NONE)
+					{
+						FCompactPoseBoneIndex CPIndex = OutPose.GetBoneContainer().MakeCompactPoseIndex(FMeshPoseBoneIndex(MeshIndex));
+						if (CPIndex != INDEX_NONE)
+						{
+							FTransform LocalInitialTransform = OutPose.GetRefPose(CPIndex);
+							DynamicHierarchy->SetTransform(BoneElement, LocalInitialTransform, ERigTransformType::InitialLocal, true, false);
+						}
+					}
+				}
+				return true;
+			});
+		bResetInitialTransformsBeforeSetup = false;
+	}
+}
+
 void UControlRig::SetBoneInitialTransformsFromSkeletalMesh(USkeletalMesh* InSkeletalMesh)
 {
 	check(InSkeletalMesh);
diff --git a/Engine/Plugins/Experimental/ControlRig/Source/ControlRig/Private/Sequencer/MovieSceneControlRigParameterSection.cpp b/Engine/Plugins/Experimental/ControlRig/Source/ControlRig/Private/Sequencer/MovieSceneControlRigParameterSection.cpp
--- a/Engine/Plugins/Experimental/ControlRig/Source/ControlRig/Private/Sequencer/MovieSceneControlRigParameterSection.cpp
+++ b/Engine/Plugins/Experimental/ControlRig/Source/ControlRig/Private/Sequencer/MovieSceneControlRigParameterSection.cpp
@@ -14,6 +14,7 @@
 #include "Misc/ScopedSlowTask.h"
 #include "MovieSceneTimeHelpers.h"
 #include "Animation/AnimSequenceHelpers.h"
+#include "Components/SkeletalMeshComponent.h"
 
 #define LOCTEXT_NAMESPACE "MovieSceneControlParameterRigSection"
 
@@ -2063,9 +2064,13 @@ void UMovieSceneControlRigParameterSection::RecordControlRigKey(FFrameNumber Fra
 	}
 }
 
-bool UMovieSceneControlRigParameterSection::LoadAnimSequenceIntoThisSection(UAnimSequence* AnimSequence, UMovieScene* MovieScene,USkeleton* Skeleton,
+bool UMovieSceneControlRigParameterSection::LoadAnimSequenceIntoThisSection(UAnimSequence* AnimSequence, UMovieScene* MovieScene,USkeletalMeshComponent* SkelMeshComp,
 	bool bKeyReduce, float Tolerance, FFrameNumber InStartFrame)
 {
+	if (SkelMeshComp == nullptr || SkelMeshComp->SkeletalMesh == nullptr || SkelMeshComp->SkeletalMesh->GetSkeleton() == nullptr)
+	{
+		return false;
+	}
 	UFKControlRig* AutoRig = Cast<UFKControlRig>(ControlRig);
 	if (!AutoRig && !ControlRig->SupportsEvent(FRigUnit_InverseExecution::EventName))
 	{
@@ -2126,7 +2131,7 @@ bool UMovieSceneControlRigParameterSection::LoadAnimSequenceIntoThisSection(UAni
 	const FAnimationCurveData& CurveData = DataModel->GetCurveData();
 	const TArray<FBoneAnimationTrack>& BoneAnimationTracks = DataModel->GetBoneAnimationTracks();
 
-	ControlRig->SetBoneInitialTransformsFromRefSkeleton(Skeleton->GetReferenceSkeleton());
+	ControlRig->SetBoneInitialTransformsFromSkeletalMeshComponent(SkelMeshComp);
 	ControlRig->RequestSetup();
 	ControlRig->Evaluate_AnyThread();
 
@@ -2147,7 +2152,7 @@ bool UMovieSceneControlRigParameterSection::LoadAnimSequenceIntoThisSection(UAni
 		{
 			const FBoneAnimationTrack& AnimationTrack = BoneAnimationTracks[TrackIndex];
 
-			FName BoneName = Skeleton->GetReferenceSkeleton().GetBoneName(AnimationTrack.BoneTreeIndex);
+			FName BoneName = SkelMeshComp->SkeletalMesh->GetSkeleton()->GetReferenceSkeleton().GetBoneName(AnimationTrack.BoneTreeIndex);
 			FTransform BoneTransform;
 			UE::Anim::GetBoneTransformFromModel(DataModel, BoneTransform, TrackIndex, SequenceSecond, EAnimInterpolationType::Linear);
 			SourceHierarchy->SetLocalTransform(FRigElementKey(BoneName, ERigElementType::Bone), BoneTransform);
diff --git a/Engine/Plugins/Experimental/ControlRig/Source/ControlRig/Private/Sequencer/MovieSceneControlRigParameterTemplate.cpp b/Engine/Plugins/Experimental/ControlRig/Source/ControlRig/Private/Sequencer/MovieSceneControlRigParameterTemplate.cpp
--- a/Engine/Plugins/Experimental/ControlRig/Source/ControlRig/Private/Sequencer/MovieSceneControlRigParameterTemplate.cpp
+++ b/Engine/Plugins/Experimental/ControlRig/Source/ControlRig/Private/Sequencer/MovieSceneControlRigParameterTemplate.cpp
@@ -563,7 +563,7 @@ void FControlRigBindingHelper::BindToSequencerInstance(UControlRig* ControlRig)
 				AnimInstance->AddControlRigTrack(ControlRig->GetUniqueID(), ControlRig);
 				ControlRig->Initialize();
 				ControlRig->RequestInit();
-				ControlRig->SetBoneInitialTransformsFromSkeletalMesh(SkeletalMeshComponent->SkeletalMesh);
+				ControlRig->SetBoneInitialTransformsFromSkeletalMeshComponent(SkeletalMeshComponent);
 				ControlRig->Evaluate_AnyThread();
 			}
 		}
diff --git a/Engine/Plugins/Experimental/ControlRig/Source/ControlRig/Public/ControlRig.h b/Engine/Plugins/Experimental/ControlRig/Source/ControlRig/Public/ControlRig.h
--- a/Engine/Plugins/Experimental/ControlRig/Source/ControlRig/Public/ControlRig.h
+++ b/Engine/Plugins/Experimental/ControlRig/Source/ControlRig/Public/ControlRig.h
@@ -32,6 +32,7 @@
 class IControlRigObjectBinding;
 class UScriptStruct;
 class USkeletalMesh;
+class USkeletalMeshComponent;
 
 struct FReferenceSkeleton;
 struct FRigUnit;
@@ -375,6 +376,10 @@ public:
 	FControlRigExecuteEvent& OnExecuted_AnyThread() { return ExecutedEvent; }
 	FRigEventDelegate& OnRigEvent_AnyThread() { return RigEventDelegate; }
 
+	// Setup the initial transform / ref pose of the bones based upon skeletal mesh component
+	// Unlike the other two functions below this uses the current refpose instead of the RefSkeleton pose.
+	void SetBoneInitialTransformsFromSkeletalMeshComponent(USkeletalMeshComponent* InSkelMeshComp);
+
 	// Setup the initial transforms / ref pose of the bones based on a skeletal mesh
 	void SetBoneInitialTransformsFromSkeletalMesh(USkeletalMesh* InSkeletalMesh);
 
diff --git a/Engine/Plugins/Experimental/ControlRig/Source/ControlRig/Public/Sequencer/MovieSceneControlRigParameterSection.h b/Engine/Plugins/Experimental/ControlRig/Source/ControlRig/Public/Sequencer/MovieSceneControlRigParameterSection.h
--- a/Engine/Plugins/Experimental/ControlRig/Source/ControlRig/Public/Sequencer/MovieSceneControlRigParameterSection.h
+++ b/Engine/Plugins/Experimental/ControlRig/Source/ControlRig/Public/Sequencer/MovieSceneControlRigParameterSection.h
@@ -19,6 +19,7 @@
 #include "MovieSceneControlRigParameterSection.generated.h"
 
 class UAnimSequence;
+class USkeletalMeshComponent;
 
 struct CONTROLRIG_API FControlRigBindingHelper
 {
@@ -214,9 +215,9 @@ public:
 	//Function to save control rig key when recording.
 	void RecordControlRigKey(FFrameNumber FrameNumber, bool bSetDefault, bool bDoAutoKey);
 
-	//Function to load an Anim Sequence into this section. It will automatically reszie to the section size.
+	//Function to load an Anim Sequence into this section. It will automatically resize to the section size.
 	//Will return false if fails or is canceled
-	virtual bool LoadAnimSequenceIntoThisSection(UAnimSequence* Sequence, UMovieScene* MovieScene, USkeleton* Skeleton,
+	virtual bool LoadAnimSequenceIntoThisSection(UAnimSequence* Sequence, UMovieScene* MovieScene, USkeletalMeshComponent* SkelMeshComp,
 		bool bKeyReduce, float Tolerance, FFrameNumber InStartFrame = 0);
 #endif
 	const TArray<bool>& GetControlsMask() const
diff --git a/Engine/Plugins/Experimental/ControlRig/Source/ControlRigEditor/Private/ControlRigEditorModule.cpp b/Engine/Plugins/Experimental/ControlRig/Source/ControlRigEditor/Private/ControlRigEditorModule.cpp
--- a/Engine/Plugins/Experimental/ControlRig/Source/ControlRigEditor/Private/ControlRigEditorModule.cpp
+++ b/Engine/Plugins/Experimental/ControlRig/Source/ControlRigEditor/Private/ControlRigEditorModule.cpp
@@ -790,8 +790,8 @@ void FControlRigEditorModule::BakeToControlRig(UClass* ControlRigClass, UAnimSeq
 			UMovieSceneControlRigParameterTrack* Track = MovieScene->AddTrack<UMovieSceneControlRigParameterTrack>(ActorTrackGuid);
 			if (Track)
 			{
-				USkeletalMesh* SkeletalMesh = MeshActor->GetSkeletalMeshComponent()->SkeletalMesh;
-				USkeleton* Skeleton = SkeletalMesh->GetSkeleton();
+				USkeletalMeshComponent* SkelMeshComp = MeshActor->GetSkeletalMeshComponent();
+				USkeletalMesh* SkeletalMesh = SkelMeshComp->SkeletalMesh;
 
 				FString ObjectName = (ControlRigClass->GetName());
 				ObjectName.RemoveFromEnd(TEXT("_C"));
@@ -813,12 +813,12 @@ void FControlRigEditorModule::BakeToControlRig(UClass* ControlRigClass, UAnimSeq
 				UMovieSceneControlRigParameterSection* ParamSection = Cast<UMovieSceneControlRigParameterSection>(NewSection);
 			
 				FBakeToControlDelegate BakeCallback = FBakeToControlDelegate::CreateLambda([this, WeakSequencer, LevelSequence, 
-					AnimSequence, MovieScene, ControlRig, ParamSection,ActorTrackGuid, Skeleton]
+					AnimSequence, MovieScene, ControlRig, ParamSection,ActorTrackGuid, SkelMeshComp]
 				(bool bKeyReduce, float KeyReduceTolerance)
 				{
 					if (ParamSection)
 					{
-						ParamSection->LoadAnimSequenceIntoThisSection(AnimSequence, MovieScene, Skeleton, bKeyReduce,
+						ParamSection->LoadAnimSequenceIntoThisSection(AnimSequence, MovieScene, SkelMeshComp, bKeyReduce,
 							KeyReduceTolerance);
 					}
 					WeakSequencer.Pin()->EmptySelection();
diff --git a/Engine/Plugins/Experimental/ControlRig/Source/ControlRigEditor/Private/ControlRigSequencerEditorLibrary.cpp b/Engine/Plugins/Experimental/ControlRig/Source/ControlRigEditor/Private/ControlRigSequencerEditorLibrary.cpp
--- a/Engine/Plugins/Experimental/ControlRig/Source/ControlRigEditor/Private/ControlRigSequencerEditorLibrary.cpp
+++ b/Engine/Plugins/Experimental/ControlRig/Source/ControlRigEditor/Private/ControlRigSequencerEditorLibrary.cpp
@@ -368,7 +368,7 @@ FTransform UControlRigSequencerEditorLibrary::GetSkeletalMeshComponentWorldTrans
 {
 	TArray<FFrameNumber> Frames;
 	Frames.Add(Frame);
-	TArray<FTransform> Transforms = GetSkeletalMeshComponentWorldTransforms(LevelSequence, SkeletalMeshComponent, Frames);
+	TArray<FTransform> Transforms = GetSkeletalMeshComponentWorldTransforms(LevelSequence, SkeletalMeshComponent, Frames,SocketName);
 	if (Transforms.Num() == 1)
 	{
 		return Transforms[0];
@@ -515,30 +515,44 @@ bool UControlRigSequencerEditorLibrary::BakeToControlRig(UWorld* World, ULevelSe
 	//get level sequence if one exists...
 	TWeakPtr<ISequencer> WeakSequencer = GetSequencerFromAsset(LevelSequence);
 
-	ALevelSequenceActor* OutActor;
+	ALevelSequenceActor* OutActor = nullptr;
 	FMovieSceneSequencePlaybackSettings Settings;
 	FLevelSequenceCameraSettings CameraSettings;
 	FMovieSceneSequenceIDRef Template = MovieSceneSequenceID::Root;
 	FMovieSceneSequenceTransform RootToLocalTransform;
-	ULevelSequencePlayer* Player = ULevelSequencePlayer::CreateLevelSequencePlayer(World, LevelSequence, Settings, OutActor);
-	Player->Initialize(LevelSequence, World->PersistentLevel, Settings, CameraSettings);
-	Player->State.AssignSequence(MovieSceneSequenceID::Root, *LevelSequence, *Player);
+	IMovieScenePlayer* Player = nullptr;
+	ULevelSequencePlayer* LevelPlayer = nullptr;
+	if (WeakSequencer.IsValid())
+	{
+		Player = WeakSequencer.Pin().Get();
+	}
+	else
+	{
+		Player = LevelPlayer = ULevelSequencePlayer::CreateLevelSequencePlayer(World, LevelSequence, Settings, OutActor);
+		LevelPlayer->Initialize(LevelSequence, World->PersistentLevel, Settings, CameraSettings);
+		LevelPlayer->State.AssignSequence(MovieSceneSequenceID::Root, *LevelSequence, *Player);
+	}
+	if (Player == nullptr)
+	{
+		UE_LOG(LogControlRig, Error, TEXT("Baking: Problem Setting up Player"));
+		return false;
+	}
+	
 
 	bool bResult = false;
 	const FScopedTransaction Transaction(LOCTEXT("BakeToControlRig_Transaction", "Bake To Control Rig"));
 	{
 		FSpawnableRestoreState SpawnableRestoreState(MovieScene);
 
-		if (SpawnableRestoreState.bWasChanged)
+		if (LevelPlayer && SpawnableRestoreState.bWasChanged)
 		{
 			// Evaluate at the beginning of the subscene time to ensure that spawnables are created before export
-			Player->SetPlaybackPosition(FMovieSceneSequencePlaybackParams(UE::MovieScene::DiscreteInclusiveLower(MovieScene->GetPlaybackRange()).Value, EUpdatePositionMethod::Play));
+			LevelPlayer->SetPlaybackPosition(FMovieSceneSequencePlaybackParams(UE::MovieScene::DiscreteInclusiveLower(MovieScene->GetPlaybackRange()).Value, EUpdatePositionMethod::Play));
 		}
 		TArrayView<TWeakObjectPtr<>>  Result = Player->FindBoundObjects(Binding.BindingID, Template);
 	
 		if (Result.Num() > 0 && Result[0].IsValid())
 		{
-
 			UObject* BoundObject = Result[0].Get();
 			USkeleton* Skeleton = nullptr;
 			USkeletalMeshComponent* SkeletalMeshComp = nullptr;
@@ -551,7 +565,10 @@ bool UControlRigSequencerEditorLibrary::BakeToControlRig(UWorld* World, ULevelSe
 				if (bResult == false)
 				{
 					TempAnimSequence->MarkPendingKill();
-					World->DestroyActor(OutActor);
+					if (OutActor)
+					{
+						World->DestroyActor(OutActor);
+					}
 					return false;
 				}
 
@@ -588,7 +605,10 @@ bool UControlRigSequencerEditorLibrary::BakeToControlRig(UWorld* World, ULevelSe
 					{
 						TempAnimSequence->MarkPendingKill();
 						MovieScene->RemoveTrack(*Track);
-						World->DestroyActor(OutActor);
+						if (OutActor)
+						{
+							World->DestroyActor(OutActor);
+						}
 						return false;
 					}
 					FControlRigEditMode* ControlRigEditMode = nullptr;
@@ -633,7 +653,7 @@ bool UControlRigSequencerEditorLibrary::BakeToControlRig(UWorld* World, ULevelSe
 						WeakSequencer.Pin()->ThrobSectionSelection();
 						WeakSequencer.Pin()->NotifyMovieSceneDataChanged(EMovieSceneDataChangeType::MovieSceneStructureItemAdded);
 					}
-					ParamSection->LoadAnimSequenceIntoThisSection(TempAnimSequence, MovieScene, Skeleton,
+					ParamSection->LoadAnimSequenceIntoThisSection(TempAnimSequence, MovieScene, SkeletalMeshComp,
 						bReduceKeys, Tolerance);
 
 					//Turn Off Any Skeletal Animation Tracks
@@ -670,18 +690,23 @@ bool UControlRigSequencerEditorLibrary::BakeToControlRig(UWorld* World, ULevelSe
 		}
 	}
 
-	Player->Stop();
-	World->DestroyActor(OutActor);
-
+	if (LevelPlayer)
+	{
+		LevelPlayer->Stop();
+	}
+	if (OutActor)
+	{
+		World->DestroyActor(OutActor);
+	}
 	return bResult;
 }
 
-bool UControlRigSequencerEditorLibrary::LoadAnimSequenceIntoControlRigSection(UMovieSceneSection* MovieSceneSection, UAnimSequence* AnimSequence,  USkeleton* Skeleton,
+bool UControlRigSequencerEditorLibrary::LoadAnimSequenceIntoControlRigSection(UMovieSceneSection* MovieSceneSection, UAnimSequence* AnimSequence, USkeletalMeshComponent* SkelMeshComp,
 	FFrameNumber InStartFrame,bool bKeyReduce, float Tolerance)
 {
-	if (MovieSceneSection == nullptr || AnimSequence == nullptr || Skeleton == nullptr)
+	if (MovieSceneSection == nullptr || AnimSequence == nullptr || SkelMeshComp == nullptr)
 	{
-		return false;;
+		return false;
 	}
 	UMovieScene* MovieScene = MovieSceneSection->GetTypedOuter<UMovieScene>();
 	if (MovieScene == nullptr)
@@ -690,7 +715,7 @@ bool UControlRigSequencerEditorLibrary::LoadAnimSequenceIntoControlRigSection(UM
 	}
 	if (UMovieSceneControlRigParameterSection* Section = Cast<UMovieSceneControlRigParameterSection>(MovieSceneSection))
 	{
-		return Section->LoadAnimSequenceIntoThisSection(AnimSequence, MovieScene, Skeleton, bKeyReduce, Tolerance, InStartFrame);
+		return Section->LoadAnimSequenceIntoThisSection(AnimSequence, MovieScene, SkelMeshComp, bKeyReduce, Tolerance, InStartFrame);
 	}
 	return false;
 }
@@ -769,7 +794,7 @@ static bool GetControlRigValues(UWorld* World, ULevelSequence* LevelSequence, UC
 {
 	if (LevelSequence)
 	{
-		ALevelSequenceActor* OutActor;
+		ALevelSequenceActor* OutActor = nullptr;
 		FMovieSceneSequencePlaybackSettings Settings;
 		FLevelSequenceCameraSettings CameraSettings;
 		FMovieSceneSequenceIDRef Template = MovieSceneSequenceID::Root;
diff --git a/Engine/Plugins/Experimental/ControlRig/Source/ControlRigEditor/Private/Sequencer/ControlRigParameterTrackEditor.cpp b/Engine/Plugins/Experimental/ControlRig/Source/ControlRigEditor/Private/Sequencer/ControlRigParameterTrackEditor.cpp
--- a/Engine/Plugins/Experimental/ControlRig/Source/ControlRigEditor/Private/Sequencer/ControlRigParameterTrackEditor.cpp
+++ b/Engine/Plugins/Experimental/ControlRig/Source/ControlRigEditor/Private/Sequencer/ControlRigParameterTrackEditor.cpp
@@ -841,7 +841,6 @@ void FControlRigParameterTrackEditor::BakeToControlRig(UClass* InClass, FGuid Ob
 		UMovieSceneSequence* OwnerSequence = GetSequencer()->GetFocusedMovieSceneSequence();
 		UMovieScene* OwnerMovieScene = OwnerSequence->GetMovieScene();
 		{
-
 			UAnimSequence* TempAnimSequence = NewObject<UAnimSequence>(GetTransientPackage(), NAME_None);
 			TempAnimSequence->SetSkeleton(Skeleton);
 			const TSharedPtr<ISequencer> ParentSequencer = GetSequencer();
@@ -937,7 +936,6 @@ void FControlRigParameterTrackEditor::BakeToControlRig(UClass* InClass, FGuid Ob
 						}
 					}
 
-
 					bool bSequencerOwnsControlRig = true;
 
 					ControlRig->Modify();
@@ -946,7 +944,7 @@ void FControlRigParameterTrackEditor::BakeToControlRig(UClass* InClass, FGuid Ob
 					ControlRig->GetDataSourceRegistry()->RegisterDataSource(UControlRig::OwnerComponent, ControlRig->GetObjectBinding()->GetBoundObject());
 					ControlRig->Initialize();
 					ControlRig->RequestInit();
-					ControlRig->SetBoneInitialTransformsFromRefSkeleton(Skeleton->GetReferenceSkeleton());
+					ControlRig->SetBoneInitialTransformsFromSkeletalMeshComponent(SkelMeshComp);
 					ControlRig->Evaluate_AnyThread();
 
 					UMovieSceneSection* NewSection = Track->CreateControlRigSection(0, ControlRig, bSequencerOwnsControlRig);
@@ -960,7 +958,7 @@ void FControlRigParameterTrackEditor::BakeToControlRig(UClass* InClass, FGuid Ob
 					GetSequencer()->SelectSection(NewSection);
 					GetSequencer()->ThrobSectionSelection();
 					GetSequencer()->NotifyMovieSceneDataChanged(EMovieSceneDataChangeType::MovieSceneStructureItemAdded);
-					ParamSection->LoadAnimSequenceIntoThisSection(TempAnimSequence, OwnerMovieScene, Skeleton,
+					ParamSection->LoadAnimSequenceIntoThisSection(TempAnimSequence, OwnerMovieScene, SkelMeshComp,
 						BakeSettings->bReduceKeys, BakeSettings->Tolerance);
 
 					//Turn Off Any Skeletal Animation Tracks
@@ -3216,16 +3214,17 @@ void FControlRigParameterSection::OnAnimationAssetSelectedForFK(const FAssetData
 	{
 		UAnimSequence* AnimSequence = Cast<UAnimSequence>(AssetData.GetAsset());
 		UObject* BoundObject = nullptr;
-		USkeleton* Skeleton = AcquireSkeletonFromObjectGuid(ObjectBinding, &BoundObject, SequencerPtr);
+		AcquireSkeletonFromObjectGuid(ObjectBinding, &BoundObject, SequencerPtr);
+		USkeletalMeshComponent* SkelMeshComp = AcquireSkeletalMeshFromObject(BoundObject, SequencerPtr);
 
-		if (AnimSequence && Skeleton && AnimSequence->GetDataModel()->GetNumBoneTracks() > 0)
+		if (AnimSequence && SkelMeshComp && AnimSequence->GetDataModel()->GetNumBoneTracks() > 0)
 		{
 
 			FScopedTransaction Transaction(LOCTEXT("BakeAnimation_Transaction", "Bake Animation To FK Control Rig"));
 			Section->Modify();
 			UMovieScene* MovieScene = SequencerPtr->GetFocusedMovieSceneSequence()->GetMovieScene();
 			FFrameNumber StartFrame = SequencerPtr->GetLocalTime().Time.GetFrame();
-			if (!Section->LoadAnimSequenceIntoThisSection(AnimSequence, MovieScene, Skeleton, false, 0.1f, StartFrame))
+			if (!Section->LoadAnimSequenceIntoThisSection(AnimSequence, MovieScene, SkelMeshComp, false, 0.1f, StartFrame))
 			{
 				Transaction.Cancel();
 			}
diff --git a/Engine/Plugins/Experimental/ControlRig/Source/ControlRigEditor/Public/ControlRigSequencerEditorLibrary.h b/Engine/Plugins/Experimental/ControlRig/Source/ControlRigEditor/Public/ControlRigSequencerEditorLibrary.h
--- a/Engine/Plugins/Experimental/ControlRig/Source/ControlRigEditor/Public/ControlRigSequencerEditorLibrary.h
+++ b/Engine/Plugins/Experimental/ControlRig/Source/ControlRigEditor/Public/ControlRigSequencerEditorLibrary.h
@@ -91,7 +91,7 @@ public:
 	* @param MovieSceneSection The MovieSceneSectionto load into
 	* @param AnimSequence The Sequence to load
 	* @param MovieScene The MovieScene getting loaded into
-	* @param Skeleton The Skeleton of the skelmesh component getting loaded into.
+	* @param SkelMeshComponent The Skeletal Mesh component getting loaded into.
 	* @param InStartFrame Frame to insert the animation
 	* @param bKeyReduce If true do key reduction based upon Tolerance, if false don't
 	* @param Tolerance If reducing keys, tolerance about which keys will be removed, smaller tolerance, more keys usually.
@@ -99,7 +99,7 @@ public:
 	* @return returns True if successful, False otherwise
 	*/
 	UFUNCTION(BlueprintCallable, Category = "Editor Scripting | Sequencer Tools | Control Rig")
-	static bool LoadAnimSequenceIntoControlRigSection(UMovieSceneSection* MovieSceneSection, UAnimSequence* AnimSequence, USkeleton* Skeleton,
+	static bool LoadAnimSequenceIntoControlRigSection(UMovieSceneSection* MovieSceneSection, UAnimSequence* AnimSequence, USkeletalMeshComponent* SkelMeshComp,
 		FFrameNumber InStartFrame, bool bKeyReduce = false, float Tolerance = 0.001f);
 
 	/**
