Control Rig: Fix up FK rig for new hierarchy

diff --git a/Engine/Plugins/Experimental/ControlRig/Source/ControlRig/Private/Rigs/FKControlRig.cpp b/Engine/Plugins/Experimental/ControlRig/Source/ControlRig/Private/Rigs/FKControlRig.cpp
--- a/Engine/Plugins/Experimental/ControlRig/Source/ControlRig/Private/Rigs/FKControlRig.cpp
+++ b/Engine/Plugins/Experimental/ControlRig/Source/ControlRig/Private/Rigs/FKControlRig.cpp
@@ -29,17 +29,6 @@ FName UFKControlRig::GetControlName(const FName& InBoneName)
 	return NAME_None;
 }
 
-FName UFKControlRig::GetNullName(const FName& InBoneName)
-{
-	if (InBoneName != NAME_None)
-	{
-		return FName(*(InBoneName.ToString() + TEXT("_NULL")));
-	}
-
-	// if bone name is coming as none, we don't append
-	return NAME_None;
-}
-
 void UFKControlRig::ExecuteUnits(FRigUnitContext& InOutContext, const FName& InEventName)
 {
 	if (InOutContext.State != EControlRigState::Update)
@@ -58,18 +47,23 @@ void UFKControlRig::ExecuteUnits(FRigUnitContext& InOutContext, const FName& InE
 			const int32 ControlIndex = GetHierarchy()->GetIndex(ControlKey);
 			if (IsControlActive[ControlIndex])
 			{
-				const FTransform Transform = GetHierarchy()->GetLocalTransform(ControlIndex);
+				FRigControlElement* Control = GetHierarchy()->Get<FRigControlElement>(ControlIndex);
 				switch (ApplyMode)
 				{
 					case EControlRigFKRigExecuteMode::Replace:
 					{
+						const FTransform LocalTransform = GetHierarchy()->GetLocalTransform(ControlIndex);
+						const FTransform OffsetTransform = Control->Offset.Initial.Local.Transform;
+						const FTransform Transform = LocalTransform * OffsetTransform;
 						GetHierarchy()->SetTransform(BoneElement, Transform, ERigTransformType::CurrentLocal, true, false);
 						break;
 					}
 					case EControlRigFKRigExecuteMode::Additive:
 					{
+						const FTransform LocalTransform = GetHierarchy()->GetLocalTransform(ControlIndex);
 						const FTransform PreviousTransform = GetHierarchy()->GetTransform(BoneElement, ERigTransformType::CurrentLocal);
-						GetHierarchy()->SetTransform(BoneElement, Transform * PreviousTransform, ERigTransformType::CurrentLocal, true, false);
+						const FTransform Transform = LocalTransform * PreviousTransform;
+						GetHierarchy()->SetTransform(BoneElement, Transform, ERigTransformType::CurrentLocal, true, false);
 						break;
 					}
 				}
@@ -228,31 +222,25 @@ void UFKControlRig::CreateRigElements(const FReferenceSkeleton& InReferenceSkele
 	GetHierarchy()->ForEach<FRigBoneElement>([&](FRigBoneElement* BoneElement) -> bool
 	{
 		const FName BoneName = BoneElement->GetName();
-		const int32 ParentIndex = GetHierarchy()->GetFirstParent(BoneElement->GetIndex());
-		const FName NullName = GetNullName(BoneName);// name conflict?
 		const FName ControlName = GetControlName(BoneName); // name conflict?
+		FRigElementKey ParentKey = GetHierarchy()->GetFirstParent(BoneElement->GetKey());
 
-		FRigElementKey NullKey;
-
-		FTransform LocalTransform;
-		if (ParentIndex != INDEX_NONE)
+		FTransform OffsetTransform;
+		if (ParentKey.IsValid())
 		{
 			FTransform GlobalTransform = GetHierarchy()->GetGlobalTransform(BoneElement->GetIndex());
-			FTransform ParentTransform = GetHierarchy()->GetGlobalTransform(ParentIndex);
-			LocalTransform = GlobalTransform.GetRelativeTransform(ParentTransform);
-			NullKey = Controller->AddNull(NullName, GetHierarchy()->GetKey(ParentIndex), FTransform::Identity, false, false);
+			FTransform ParentTransform = GetHierarchy()->GetGlobalTransform(ParentKey);
+			OffsetTransform = GlobalTransform.GetRelativeTransform(ParentTransform);
 		}
 		else
 		{
-			LocalTransform = GetHierarchy()->GetLocalTransform(BoneElement->GetIndex());
-			NullKey = Controller->AddNull(NullName, FRigElementKey(), FTransform::Identity, true, false);
+			OffsetTransform = GetHierarchy()->GetLocalTransform(BoneElement->GetIndex());
 		}
 
 		FRigControlSettings Settings;
 		Settings.ControlType = ERigControlType::EulerTransform;
 		Settings.DisplayName = BoneName;
-		const FRigElementKey ControlKey = Controller->AddControl(ControlName, NullKey, Settings, FRigControlValue::Make(FEulerTransform::Identity), FTransform::Identity, FTransform::Identity, false);
-		GetHierarchy()->SetLocalTransform(ControlKey, LocalTransform, true, true, false);
+		Controller->AddControl(ControlName, ParentKey, Settings, FRigControlValue::Make(FEulerTransform::Identity), OffsetTransform, FTransform::Identity, false);
 
 		return true;
 	});
diff --git a/Engine/Plugins/Experimental/ControlRig/Source/ControlRig/Public/Rigs/FKControlRig.h b/Engine/Plugins/Experimental/ControlRig/Source/ControlRig/Public/Rigs/FKControlRig.h
--- a/Engine/Plugins/Experimental/ControlRig/Source/ControlRig/Public/Rigs/FKControlRig.h
+++ b/Engine/Plugins/Experimental/ControlRig/Source/ControlRig/Public/Rigs/FKControlRig.h
@@ -46,8 +46,6 @@ public:
 
 	// utility function to 
 	static FName GetControlName(const FName& InBoneName);
-	static FName GetNullName(const FName& InBoneName);
-
 
 	TArray<FName> GetControlNames();
 	bool GetControlActive(int32 Index) const;
