Control Rig: Sequencer: Fix on duplicating.
#trivial

diff --git a/Engine/Plugins/Experimental/ControlRig/Source/ControlRig/Private/Sequencer/MovieSceneControlRigParameterSection.cpp b/Engine/Plugins/Experimental/ControlRig/Source/ControlRig/Private/Sequencer/MovieSceneControlRigParameterSection.cpp
--- a/Engine/Plugins/Experimental/ControlRig/Source/ControlRig/Private/Sequencer/MovieSceneControlRigParameterSection.cpp
+++ b/Engine/Plugins/Experimental/ControlRig/Source/ControlRig/Private/Sequencer/MovieSceneControlRigParameterSection.cpp
@@ -1097,13 +1097,19 @@ void UMovieSceneControlRigParameterSection::ReconstructChannelProxy(bool bForce)
 	if(bForce || ControlsMask!= OldControlsMask)
 	{
 		ControlChannelMap.Empty();
-
 		OldControlsMask = ControlsMask;
 		// Need to create the channels in sorted orders
 		if (ControlRig)
 		{
 			TArray<FRigControlElement*> SortedControls;
 			ControlRig->GetControlsInOrder(SortedControls);
+			if (ControlsMask.Num() != SortedControls.Num())
+			{
+				TArray<bool> OnArray;
+				OnArray.Init(true, ControlRig->AvailableControls().Num());
+				SetControlsMask(OnArray);
+				OldControlsMask = ControlsMask;
+			}
 			
 			int32 ControlIndex = 0; 
 			int32 TotalIndex = 0; 
