Sequencer: Control Rig: Snapper:  Entering frames now works.
#trivial

diff --git a/Engine/Plugins/Experimental/ControlRig/Source/ControlRigEditor/Private/EditMode/SControlRigSnapper.cpp b/Engine/Plugins/Experimental/ControlRig/Source/ControlRigEditor/Private/EditMode/SControlRigSnapper.cpp
--- a/Engine/Plugins/Experimental/ControlRig/Source/ControlRigEditor/Private/EditMode/SControlRigSnapper.cpp
+++ b/Engine/Plugins/Experimental/ControlRig/Source/ControlRigEditor/Private/EditMode/SControlRigSnapper.cpp
@@ -256,6 +256,19 @@ void SControlRigSnapper::Construct(const FArguments& InArgs)
 							[
 								SNew(SEditableTextBox)
 								.ToolTipText(LOCTEXT("GetStartFrameTooltip", "Set first frame to snap"))
+								.OnTextCommitted_Lambda([this](const FText& InText, ETextCommit::Type TextCommitType)
+									{
+										TWeakPtr<ISequencer> Sequencer = Snapper.GetSequencer();
+										if (Sequencer.IsValid() && Sequencer.Pin()->GetFocusedMovieSceneSequence())
+										{
+
+											TOptional<double> NewFrameTime = Sequencer.Pin()->GetNumericTypeInterface()->FromString(InText.ToString(), 0);
+											if (NewFrameTime.IsSet())
+											{
+												StartFrame = FFrameNumber((int32)NewFrameTime.GetValue());
+											}
+										}
+									})
 								.Text_Lambda([this]()
 									{
 										return GetStartFrameToSnapText();
@@ -276,6 +289,19 @@ void SControlRigSnapper::Construct(const FArguments& InArgs)
 								[
 									SNew(SEditableTextBox)
 									.ToolTipText(LOCTEXT("GetEndFrameTooltip", "Set end frame to snap"))
+									.OnTextCommitted_Lambda([this](const FText& InText, ETextCommit::Type TextCommitType)
+									{
+										TWeakPtr<ISequencer> Sequencer = Snapper.GetSequencer();
+										if (Sequencer.IsValid() && Sequencer.Pin()->GetFocusedMovieSceneSequence())
+										{
+
+											TOptional<double> NewFrameTime = Sequencer.Pin()->GetNumericTypeInterface()->FromString(InText.ToString(), 0);
+											if (NewFrameTime.IsSet())
+											{
+												EndFrame = FFrameNumber((int32)NewFrameTime.GetValue());
+											}
+										}
+									})
 									.Text_Lambda([this]()
 										{
 											return GetEndFrameToSnapText();
diff --git a/Engine/Source/Editor/MovieSceneTools/Private/MovieSceneToolHelpers.cpp b/Engine/Source/Editor/MovieSceneTools/Private/MovieSceneToolHelpers.cpp
--- a/Engine/Source/Editor/MovieSceneTools/Private/MovieSceneToolHelpers.cpp
+++ b/Engine/Source/Editor/MovieSceneTools/Private/MovieSceneToolHelpers.cpp
@@ -4146,14 +4146,13 @@ static void GetNonSequencerActorWorldTransforms(ISequencer* Sequencer, const FAc
 			}
 			else
 			{
-				WorldTransform = WorldTransform * Actor->GetActorTransform();
+				WorldTransform = WorldTransform * Actor->GetRootComponent()->GetRelativeTransform();
 			}
-
+			SocketName = Actor->GetAttachParentSocketName();
 			Actor = Actor->GetAttachParentActor();
 			NewActorSelection.Actor = Actor;
 			if (Actor)
 			{
-				SocketName = Actor->GetAttachParentSocketName();
 				NewActorSelection.SocketName = SocketName;
 			}
 		} while (Actor);
@@ -4178,3 +4177,4 @@ void MovieSceneToolHelpers::GetActorWorldTransforms(ISequencer* Sequencer, const
 		GetNonSequencerActorWorldTransforms(Sequencer, ActorSelection, Frames, OutWorldTransforms);
 	}
 }
+
