Control Rig: Dependency / Dependent API

diff --git a/Engine/Plugins/Experimental/ControlRig/Source/ControlRigDeveloper/Private/ControlRigBlueprint.cpp b/Engine/Plugins/Experimental/ControlRig/Source/ControlRigDeveloper/Private/ControlRigBlueprint.cpp
--- a/Engine/Plugins/Experimental/ControlRig/Source/ControlRigDeveloper/Private/ControlRigBlueprint.cpp
+++ b/Engine/Plugins/Experimental/ControlRig/Source/ControlRigDeveloper/Private/ControlRigBlueprint.cpp
@@ -24,6 +24,7 @@
 #include "Units/Hierarchy/RigUnit_SetBoneTransform.h"
 #include "Async/TaskGraphInterfaces.h"
 #include "Misc/CoreDelegates.h"
+#include "AssetRegistryModule.h"
 
 #if WITH_EDITOR
 #include "IControlRigEditorModule.h"
@@ -106,6 +107,10 @@ UControlRigBlueprint::UControlRigBlueprint(const FObjectInitializer& ObjectIniti
 	HierarchyController->OnModified().AddUObject(this, &UControlRigBlueprint::HandleHierarchyModified);
 }
 
+UControlRigBlueprint::UControlRigBlueprint()
+{
+}
+
 void UControlRigBlueprint::InitializeModelIfRequired(bool bRecompileVM)
 {
 	DECLARE_SCOPE_HIERARCHICAL_COUNTER_FUNC()
@@ -526,9 +531,7 @@ void UControlRigBlueprint::PostLoad()
 	}
 
 	// perform backwards compat value upgrades
-	TArray<URigVMGraph*> GraphsToValidate;
-	GraphsToValidate.Add(GetModel());
-	GraphsToValidate.Add(GetLocalFunctionLibrary());
+	TArray<URigVMGraph*> GraphsToValidate = GetAllModels();
 	for (int32 GraphIndex = 0; GraphIndex < GraphsToValidate.Num(); GraphIndex++)
 	{
 		URigVMGraph* GraphToValidate = GraphsToValidate[GraphIndex];
@@ -539,11 +542,6 @@ void UControlRigBlueprint::PostLoad()
 			
 		for(URigVMNode* Node : GraphToValidate->GetNodes())
 		{
-			if(URigVMCollapseNode* CollapseNode = Cast<URigVMCollapseNode>(Node))
-			{
-				GraphsToValidate.AddUnique(CollapseNode->GetContainedGraph());
-			}
-
 			TArray<URigVMPin*> Pins = Node->GetAllPinsRecursively();
 			for(URigVMPin* Pin : Pins)
 			{
@@ -560,6 +558,20 @@ void UControlRigBlueprint::PostLoad()
 					}
 				}
 			}
+
+			if(URigVMFunctionReferenceNode* FunctionReferenceNode = Cast<URigVMFunctionReferenceNode>(Node))
+			{
+				if(URigVMLibraryNode* DependencyNode = FunctionReferenceNode->GetReferencedNode())
+				{
+					if(UControlRigBlueprint* DependencyBlueprint = DependencyNode->GetTypedOuter<UControlRigBlueprint>())
+					{
+						if(DependencyBlueprint != this)
+						{
+							DependencyBlueprint->GetLocalFunctionLibrary()->UpdateReferencesForReferenceNode(FunctionReferenceNode);
+						}
+					}
+				}
+			}
 		}
 	}	
 
@@ -1181,6 +1193,115 @@ void UControlRigBlueprint::MarkFunctionPublic(const FName& InFunctionName, bool
 	}
 }
 
+TArray<UControlRigBlueprint*> UControlRigBlueprint::GetDependencies(bool bRecursive) const
+{
+	TArray<UControlRigBlueprint*> Dependencies;
+
+	TArray<URigVMGraph*> Graphs = GetAllModels();
+	for(URigVMGraph* Graph : Graphs)
+	{
+		for(URigVMNode* Node : Graph->GetNodes())
+		{
+			if(URigVMFunctionReferenceNode* FunctionReferenceNode = Cast<URigVMFunctionReferenceNode>(Node))
+			{
+				if(URigVMLibraryNode* LibraryNode = FunctionReferenceNode->GetReferencedNode())
+				{
+					if(UControlRigBlueprint* DependencyBlueprint = LibraryNode->GetTypedOuter<UControlRigBlueprint>())
+					{
+						if(DependencyBlueprint != this)
+						{
+							if(!Dependencies.Contains(DependencyBlueprint))
+							{
+								Dependencies.Add(DependencyBlueprint);
+
+								if(bRecursive)
+								{
+									TArray<UControlRigBlueprint*> ChildDependencies = DependencyBlueprint->GetDependencies(true);
+									for(UControlRigBlueprint* ChildDependency : ChildDependencies)
+									{
+										Dependencies.AddUnique(ChildDependency);
+									}
+								}
+							}
+						}
+					}
+				}
+			}
+		}
+	}
+
+	return Dependencies;
+}
+
+TArray<FAssetData> UControlRigBlueprint::GetDependentAssets() const
+{
+	TArray<FAssetData> Dependents;
+	TArray<FName> AssetPaths;
+
+	if(FunctionLibrary)
+	{
+		const FAssetRegistryModule& AssetRegistryModule = FModuleManager::LoadModuleChecked<FAssetRegistryModule>(TEXT("AssetRegistry"));
+		TArray<URigVMLibraryNode*> Functions = FunctionLibrary->GetFunctions();
+		for(URigVMLibraryNode* Function : Functions)
+		{
+			const FName FunctionName = Function->GetFName();
+			if(IsFunctionPublic(FunctionName))
+			{
+				TArray<TSoftObjectPtr<URigVMFunctionReferenceNode>> References = FunctionLibrary->GetReferencesForFunction(FunctionName);
+				for(const TSoftObjectPtr<URigVMFunctionReferenceNode>& Reference : References)
+				{
+					const FName AssetPath = Reference.ToSoftObjectPath().GetAssetPathName();
+					if(AssetPath.ToString().StartsWith(TEXT("/Engine/Transient")))
+					{
+						continue;
+					}
+					
+					if(!AssetPaths.Contains(AssetPath))
+					{
+						AssetPaths.Add(AssetPath);
+
+						const FAssetData AssetData = AssetRegistryModule.Get().GetAssetByObjectPath(*AssetPath.ToString());
+						if(AssetData.IsValid())
+						{
+							Dependents.Add(AssetData);
+						}
+					}
+				}
+			}
+		}
+	}
+
+	return Dependents;
+}
+
+TArray<UControlRigBlueprint*> UControlRigBlueprint::GetDependentBlueprints(bool bRecursive) const
+{
+	TArray<FAssetData> Assets = GetDependentAssets();
+	TArray<UControlRigBlueprint*> Dependents;
+
+	for(const FAssetData& Asset : Assets)
+	{
+		if(UControlRigBlueprint* Dependent = Cast<UControlRigBlueprint>(Asset.GetAsset()))
+		{
+			if(!Dependents.Contains(Dependent))
+			{
+				Dependents.Add(Dependent);
+
+				if(bRecursive)
+				{
+					TArray<UControlRigBlueprint*> ParentDependents = Dependent->GetDependentBlueprints(true);
+					for(UControlRigBlueprint* ParentDependent : ParentDependents)
+					{
+						Dependents.AddUnique(ParentDependent);
+					}
+				}
+			}
+		}
+	}
+
+	return Dependents;
+}
+
 void UControlRigBlueprint::GetTypeActions(FBlueprintActionDatabaseRegistrar& ActionRegistrar) const
 {
 	DECLARE_SCOPE_HIERARCHICAL_COUNTER_FUNC()
diff --git a/Engine/Plugins/Experimental/ControlRig/Source/ControlRigDeveloper/Public/ControlRigBlueprint.h b/Engine/Plugins/Experimental/ControlRig/Source/ControlRigDeveloper/Public/ControlRigBlueprint.h
--- a/Engine/Plugins/Experimental/ControlRig/Source/ControlRigDeveloper/Public/ControlRigBlueprint.h
+++ b/Engine/Plugins/Experimental/ControlRig/Source/ControlRigDeveloper/Public/ControlRigBlueprint.h
@@ -249,6 +249,22 @@ public:
 	bool IsFunctionPublic(const FName& InFunctionName) const;
 	void MarkFunctionPublic(const FName& InFunctionName, bool bIsPublic = true);
 
+	// Returns a list of dependencies of this blueprint.
+	// Dependencies are blueprints that contain functions used in this blueprint
+	TArray<UControlRigBlueprint*> GetDependencies(bool bRecursive = false) const;
+
+	// Returns a list of dependents as unresolved soft object pointers.
+	// A dependent is a blueprint which uses a function defined in this blueprint.
+	// This function is not recursive, since it avoids opening the asset.
+	// Use GetDependentBlueprints as an alternative.
+	TArray<FAssetData> GetDependentAssets() const;
+
+	// Returns a list of dependents as resolved blueprints.
+	// A dependent is a blueprint which uses a function defined in this blueprint.
+	// This function loads the dependent assets and can introduce a large cost
+	// depending on the size / count of assets in the project.
+	TArray<UControlRigBlueprint*> GetDependentBlueprints(bool bRecursive = false) const;
+
 	UPROPERTY(EditAnywhere, Category = "User Interface")
 	FRigGraphDisplaySettings RigGraphDisplaySettings;
 
diff --git a/Engine/Source/Developer/RigVMDeveloper/Private/RigVMModel/RigVMController.cpp b/Engine/Source/Developer/RigVMDeveloper/Private/RigVMModel/RigVMController.cpp
--- a/Engine/Source/Developer/RigVMDeveloper/Private/RigVMModel/RigVMController.cpp
+++ b/Engine/Source/Developer/RigVMDeveloper/Private/RigVMModel/RigVMController.cpp
@@ -1791,6 +1791,7 @@ TArray<FName> URigVMController::ImportNodesFromText(const FString& InText, bool
 					if (URigVMLibraryNode* FunctionDefinition = FunctionRefNode->GetReferencedNode())
 					{
 						FunctionLibrary->FunctionReferences.FindOrAdd(FunctionDefinition).FunctionReferences.Add(FunctionRefNode);
+						FunctionLibrary->MarkPackageDirty();
 					}
 				}
 			}
@@ -3390,6 +3391,7 @@ void URigVMController::SetReferencedFunction(URigVMFunctionReferenceNode* InFunc
 	URigVMFunctionLibrary* NewLibrary = Cast<URigVMFunctionLibrary>(InNewReferencedNode->GetOuter());
 	FRigVMFunctionReferenceArray& NewReferences = NewLibrary->FunctionReferences.FindOrAdd(InNewReferencedNode);
 	NewReferences.FunctionReferences.Add(InFunctionRefNode);
+	NewLibrary->MarkPackageDirty();
 
 	InFunctionRefNode->SetReferencedNode(InNewReferencedNode);
 	
@@ -6214,6 +6216,7 @@ URigVMFunctionReferenceNode* URigVMController::AddFunctionReferenceNode(URigVMLi
 	if (URigVMFunctionLibrary* FunctionLibrary = InFunctionDefinition->GetLibrary())
 	{
 		FunctionLibrary->FunctionReferences.FindOrAdd(InFunctionDefinition).FunctionReferences.Add(FunctionRefNode);
+		FunctionLibrary->MarkPackageDirty();
 	}
 
 	for (URigVMPin* SourcePin : InFunctionDefinition->Pins)
diff --git a/Engine/Source/Developer/RigVMDeveloper/Private/RigVMModel/RigVMFunctionLibrary.cpp b/Engine/Source/Developer/RigVMDeveloper/Private/RigVMModel/RigVMFunctionLibrary.cpp
--- a/Engine/Source/Developer/RigVMDeveloper/Private/RigVMModel/RigVMFunctionLibrary.cpp
+++ b/Engine/Source/Developer/RigVMDeveloper/Private/RigVMModel/RigVMFunctionLibrary.cpp
@@ -86,6 +86,42 @@ TArray< FString > URigVMFunctionLibrary::GetReferencePathsForFunction(const FNam
 	return Result;
 }
 
+void URigVMFunctionLibrary::UpdateReferencesForReferenceNode(URigVMFunctionReferenceNode* InReferenceNode)
+{
+	if(InReferenceNode->GetOutermost() == GetTransientPackage())
+	{
+		return;
+	}
+	
+	if(URigVMLibraryNode* Function = InReferenceNode->GetReferencedNode())
+	{
+		if(Function->GetOuter() == this)
+		{
+			FRigVMFunctionReferenceArray* ReferencesEntry = FunctionReferences.Find(Function);
+			if (ReferencesEntry == nullptr)
+			{
+				Modify();
+				FunctionReferences.Add(Function);
+				ReferencesEntry = FunctionReferences.Find(Function);
+			}
+
+			const FString ReferenceNodePathName = InReferenceNode->GetPathName();
+			for (int32 ReferenceIndex = 0; ReferenceIndex < ReferencesEntry->Num(); ReferenceIndex++)
+			{
+				const TSoftObjectPtr<URigVMFunctionReferenceNode>& Reference = ReferencesEntry->operator [](ReferenceIndex);
+				if(Reference.ToString() == ReferenceNodePathName)
+				{
+					return;
+				}
+			}
+
+			Modify();
+			ReferencesEntry->FunctionReferences.Add(InReferenceNode);
+			MarkPackageDirty();
+		}
+	}
+}
+
 URigVMLibraryNode* URigVMFunctionLibrary::FindPreviouslyLocalizedFunction(URigVMLibraryNode* InFunctionToLocalize)
 {
 	if(InFunctionToLocalize == nullptr)
diff --git a/Engine/Source/Developer/RigVMDeveloper/Public/RigVMModel/RigVMFunctionLibrary.h b/Engine/Source/Developer/RigVMDeveloper/Public/RigVMModel/RigVMFunctionLibrary.h
--- a/Engine/Source/Developer/RigVMDeveloper/Public/RigVMModel/RigVMFunctionLibrary.h
+++ b/Engine/Source/Developer/RigVMDeveloper/Public/RigVMModel/RigVMFunctionLibrary.h
@@ -63,6 +63,9 @@ public:
 	UFUNCTION(BlueprintCallable, Category = RigVMGraph)
 	TArray< FString > GetReferencePathsForFunction(const FName& InFunctionName);
 
+	// Update the references list for a given reference node
+	void UpdateReferencesForReferenceNode(URigVMFunctionReferenceNode* InReferenceNode);
+
 	// Returns a function that has been previously localized based on the provided function to localize.
 	// We maintain meta data on what functions have been created locally based on which other ones,
 	// and use this method to avoid redundant localizations.
diff --git a/Engine/Source/Developer/RigVMDeveloper/Public/RigVMModel/RigVMGraph.h b/Engine/Source/Developer/RigVMDeveloper/Public/RigVMModel/RigVMGraph.h
--- a/Engine/Source/Developer/RigVMDeveloper/Public/RigVMModel/RigVMGraph.h
+++ b/Engine/Source/Developer/RigVMDeveloper/Public/RigVMModel/RigVMGraph.h
@@ -7,6 +7,8 @@
 #include "RigVMNotifications.h"
 #include "Nodes/RigVMVariableNode.h"
 #include "Nodes/RigVMParameterNode.h"
+#include "Nodes/RigVMFunctionEntryNode.h"
+#include "Nodes/RigVMFunctionReturnNode.h"
 #include "RigVMCompiler/RigVMAST.h"
 #include "RigVMGraph.generated.h"
 
