Control Rig: Avoid allocating new objects when IsGarbageCollecting()

diff --git a/Engine/Plugins/Experimental/ControlRig/Source/ControlRig/Private/Rigs/RigHierarchy.cpp b/Engine/Plugins/Experimental/ControlRig/Source/ControlRig/Private/Rigs/RigHierarchy.cpp
--- a/Engine/Plugins/Experimental/ControlRig/Source/ControlRig/Private/Rigs/RigHierarchy.cpp
+++ b/Engine/Plugins/Experimental/ControlRig/Source/ControlRig/Private/Rigs/RigHierarchy.cpp
@@ -1046,10 +1046,13 @@ URigHierarchyController* URigHierarchy::GetController(bool bCreateIfNeeded)
 	{
 		if(UObject* Outer = GetOuter())
 		{
-			URigHierarchyController* Controller = NewObject<URigHierarchyController>(Outer);
-			Controller->SetHierarchy(this);
-			LastControllerPtr = Controller;
-			return Controller;
+			if(!IsGarbageCollecting())
+			{
+				URigHierarchyController* Controller = NewObject<URigHierarchyController>(Outer);
+				Controller->SetHierarchy(this);
+				LastControllerPtr = Controller;
+				return Controller;
+			}
 		}
 	}
 	return nullptr;
diff --git a/Engine/Plugins/Experimental/ControlRig/Source/ControlRig/Public/ControlRig.h b/Engine/Plugins/Experimental/ControlRig/Source/ControlRig/Public/ControlRig.h
--- a/Engine/Plugins/Experimental/ControlRig/Source/ControlRig/Public/ControlRig.h
+++ b/Engine/Plugins/Experimental/ControlRig/Source/ControlRig/Public/ControlRig.h
@@ -200,7 +200,7 @@ public:
 			DynamicHierarchy = nullptr;
 		}
 		
-		if(DynamicHierarchy == nullptr)
+		if(DynamicHierarchy == nullptr && !IsGarbageCollecting())
 		{
 			DynamicHierarchy = NewObject<URigHierarchy>(this);
 			if (!HasAnyFlags(RF_ClassDefaultObject))
