Control Rig: Fix recursive setting for Children rig unit node + add unit tests

diff --git a/Engine/Plugins/Experimental/ControlRig/Source/ControlRig/Private/Rigs/RigHierarchyDefines.cpp b/Engine/Plugins/Experimental/ControlRig/Source/ControlRig/Private/Rigs/RigHierarchyDefines.cpp
--- a/Engine/Plugins/Experimental/ControlRig/Source/ControlRig/Private/Rigs/RigHierarchyDefines.cpp
+++ b/Engine/Plugins/Experimental/ControlRig/Source/ControlRig/Private/Rigs/RigHierarchyDefines.cpp
@@ -90,7 +90,10 @@ FRigElementKeyCollection FRigElementKeyCollection::MakeFromChildren(
 				const int32 PreviousSize = Collection.Num();
 				if(PreviousSize == Collection.AddUnique(Child))
 				{
-					ParentKeys.Add(Child);
+					if(bRecursive)
+					{
+						ParentKeys.Add(Child);
+					}
 				}
 			}
 		}
diff --git a/Engine/Plugins/Experimental/ControlRig/Source/ControlRig/Private/Units/Execution/RigUnit_Collection.cpp b/Engine/Plugins/Experimental/ControlRig/Source/ControlRig/Private/Units/Execution/RigUnit_Collection.cpp
--- a/Engine/Plugins/Experimental/ControlRig/Source/ControlRig/Private/Units/Execution/RigUnit_Collection.cpp
+++ b/Engine/Plugins/Experimental/ControlRig/Source/ControlRig/Private/Units/Execution/RigUnit_Collection.cpp
@@ -4,6 +4,10 @@
 #include "Units/Execution/RigUnit_Item.h"
 #include "Units/RigUnitContext.h"
 
+#if WITH_EDITOR
+#include "Units/RigUnitTest.h"
+#endif
+
 FRigUnit_CollectionChain_Execute()
 {
     DECLARE_SCOPE_HIERARCHICAL_COUNTER_RIGUNIT()
@@ -98,6 +102,41 @@ FRigUnit_CollectionChildren_Execute()
 	Collection = CachedCollection;
 }
 
+IMPLEMENT_RIGUNIT_AUTOMATION_TEST(FRigUnit_CollectionChildren)
+{
+	const FRigElementKey Root = Controller->AddBone(TEXT("Root"), FRigElementKey(), FTransform(FVector(0.f, 0.f, 0.f)), true, ERigBoneType::User);
+	const FRigElementKey BoneA = Controller->AddBone(TEXT("BoneA"), Root, FTransform(FVector(0.f, 0.f, 0.f)), true, ERigBoneType::User);
+	const FRigElementKey BoneB = Controller->AddBone(TEXT("BoneB"), BoneA, FTransform(FVector(0.f, 0.f, 0.f)), true, ERigBoneType::User);
+	const FRigElementKey BoneC = Controller->AddBone(TEXT("BoneC"), Root, FTransform(FVector(0.f, 0.f, 0.f)), true, ERigBoneType::User);
+
+	Unit.Parent = Root;
+	Unit.bIncludeParent = false;
+	Unit.bRecursive = false;
+	Execute();
+	AddErrorIfFalse(Unit.Collection.Num() == 2, TEXT("unexpected result"));
+	AddErrorIfFalse(Unit.Collection[0] == BoneA, TEXT("unexpected result"));
+	AddErrorIfFalse(Unit.Collection[1] == BoneC, TEXT("unexpected result"));
+
+	Unit.bIncludeParent = true;
+	Unit.bRecursive = false;
+	Execute();
+	AddErrorIfFalse(Unit.Collection.Num() == 3, TEXT("unexpected result"));
+	AddErrorIfFalse(Unit.Collection[0] == Root, TEXT("unexpected result"));
+	AddErrorIfFalse(Unit.Collection[1] == BoneA, TEXT("unexpected result"));
+	AddErrorIfFalse(Unit.Collection[2] == BoneC, TEXT("unexpected result"));
+
+	Unit.bIncludeParent = true;
+	Unit.bRecursive = true;
+	Execute();
+	AddErrorIfFalse(Unit.Collection.Num() == 4, TEXT("unexpected result"));
+	AddErrorIfFalse(Unit.Collection[0] == Root, TEXT("unexpected result"));
+	AddErrorIfFalse(Unit.Collection[1] == BoneA, TEXT("unexpected result"));
+	AddErrorIfFalse(Unit.Collection[2] == BoneC, TEXT("unexpected result"));
+	AddErrorIfFalse(Unit.Collection[3] == BoneB, TEXT("unexpected result"));
+
+	return true;
+}
+
 FRigUnit_CollectionReplaceItems_Execute()
 {
 	DECLARE_SCOPE_HIERARCHICAL_COUNTER_RIGUNIT()
