Control Rig: Fix build break - follow up on CL 16050398

diff --git a/Engine/Plugins/Experimental/ControlRig/Source/ControlRig/Private/Units/Execution/RigUnit_Collection.cpp b/Engine/Plugins/Experimental/ControlRig/Source/ControlRig/Private/Units/Execution/RigUnit_Collection.cpp
--- a/Engine/Plugins/Experimental/ControlRig/Source/ControlRig/Private/Units/Execution/RigUnit_Collection.cpp
+++ b/Engine/Plugins/Experimental/ControlRig/Source/ControlRig/Private/Units/Execution/RigUnit_Collection.cpp
@@ -102,6 +102,8 @@ FRigUnit_CollectionChildren_Execute()
 	Collection = CachedCollection;
 }
 
+#if WITH_EDITOR
+
 IMPLEMENT_RIGUNIT_AUTOMATION_TEST(FRigUnit_CollectionChildren)
 {
 	const FRigElementKey Root = Controller->AddBone(TEXT("Root"), FRigElementKey(), FTransform(FVector(0.f, 0.f, 0.f)), true, ERigBoneType::User);
@@ -137,6 +139,8 @@ IMPLEMENT_RIGUNIT_AUTOMATION_TEST(FRigUnit_CollectionChildren)
 	return true;
 }
 
+#endif
+
 FRigUnit_CollectionReplaceItems_Execute()
 {
 	DECLARE_SCOPE_HIERARCHICAL_COUNTER_RIGUNIT()
