Control Rig: Fix rotation limits on CCDIK rig unit

diff --git a/Engine/Plugins/Experimental/ControlRig/Source/ControlRig/Private/Units/Highlevel/Hierarchy/RigUnit_CCDIK.cpp b/Engine/Plugins/Experimental/ControlRig/Source/ControlRig/Private/Units/Highlevel/Hierarchy/RigUnit_CCDIK.cpp
--- a/Engine/Plugins/Experimental/ControlRig/Source/ControlRig/Private/Units/Highlevel/Hierarchy/RigUnit_CCDIK.cpp
+++ b/Engine/Plugins/Experimental/ControlRig/Source/ControlRig/Private/Units/Highlevel/Hierarchy/RigUnit_CCDIK.cpp
@@ -63,8 +63,11 @@ FRigUnit_CCDIKPerItem_Execute()
 	TArray<float>& RotationLimitsPerItem = WorkData.RotationLimitsPerItem;
 	FCachedRigElement& CachedEffector = WorkData.CachedEffector;
 
-	if (Context.State == EControlRigState::Init ||
-		RotationLimits.Num() != RotationLimitIndex.Num())
+	if ((Context.State == EControlRigState::Init) ||
+		(
+			(RotationLimits.Num() != RotationLimitIndex.Num()) &&
+			(RotationLimitIndex.Num() > 0))
+		)
 	{
 		CachedItems.Reset();
 		RotationLimitIndex.Reset();
