ControlRig: Fix pole vector control not working when leg extended

diff --git a/Engine/Plugins/Experimental/ControlRig/Source/ControlRig/Private/Math/ControlRigMathLibrary.cpp b/Engine/Plugins/Experimental/ControlRig/Source/ControlRig/Private/Math/ControlRigMathLibrary.cpp
--- a/Engine/Plugins/Experimental/ControlRig/Source/ControlRig/Private/Math/ControlRigMathLibrary.cpp
+++ b/Engine/Plugins/Experimental/ControlRig/Source/ControlRig/Private/Math/ControlRigMathLibrary.cpp
@@ -436,7 +436,7 @@ void FControlRigMathLibrary::SolveBasicTwoBoneIK(FTransform& BoneA, FTransform&
 
 		if (SecondaryAxisWeight > SMALL_NUMBER)
 		{
-			FVector Target2 = BoneB.GetLocation() - (Effector.GetLocation() + BoneA.GetLocation()) * 0.5f;
+			FVector Target2 = BoneBLocation - (Effector.GetLocation() + BoneA.GetLocation()) * 0.5f;
 			if (!Target2.IsNearlyZero() && !Axis.IsNearlyZero())
 			{
 				Target2 = Target2 - FVector::DotProduct(Target2, Target1) * Target1;
@@ -455,7 +455,7 @@ void FControlRigMathLibrary::SolveBasicTwoBoneIK(FTransform& BoneA, FTransform&
 	}
 
 	Axis = BoneB.TransformVectorNoScale(PrimaryAxis);
-	Target1 = Effector.GetLocation() - BoneB.GetLocation();
+	Target1 = Effector.GetLocation() - BoneBLocation;
 	if (!Target1.IsNearlyZero() && !Axis.IsNearlyZero())
 	{
 		Target1 = Target1.GetSafeNormal();
@@ -465,7 +465,7 @@ void FControlRigMathLibrary::SolveBasicTwoBoneIK(FTransform& BoneA, FTransform&
 		if (SecondaryAxisWeight > SMALL_NUMBER)
 		{
 			Axis = BoneB.TransformVectorNoScale(SecondaryAxis);
-			FVector Target2 = BoneB.GetLocation() - (Effector.GetLocation() + BoneA.GetLocation()) * 0.5f;
+			FVector Target2 = BoneBLocation - (Effector.GetLocation() + BoneA.GetLocation()) * 0.5f;
 			if (!Target2.IsNearlyZero() && !Axis.IsNearlyZero())
 			{
 				Target2 = Target2 - FVector::DotProduct(Target2, Target1) * Target1;
