Control Rig: Sequencer: Python test and some API tweaks/fixes including what time space you are in, defaults to display rate.

diff --git a/Engine/Plugins/Experimental/ControlRig/Source/ControlRigEditor/Private/ControlRigSequencerEditorLibrary.cpp b/Engine/Plugins/Experimental/ControlRig/Source/ControlRigEditor/Private/ControlRigSequencerEditorLibrary.cpp
--- a/Engine/Plugins/Experimental/ControlRig/Source/ControlRigEditor/Private/ControlRigSequencerEditorLibrary.cpp
+++ b/Engine/Plugins/Experimental/ControlRig/Source/ControlRigEditor/Private/ControlRigSequencerEditorLibrary.cpp
@@ -266,8 +266,6 @@ UMovieSceneTrack* UControlRigSequencerEditorLibrary::FindOrCreateControlRigTrack
 					AcquireSkeletonAndSkelMeshCompFromObject(BoundObject, &Skeleton, &SkeletalMeshComponent);
 
 					UControlRig* ControlRig = nullptr;
-						
-
 					if (Skeleton && SkeletalMeshComponent)
 					{
 						UMovieSceneControlRigParameterTrack* Track = AddControlRig(LevelSequence, ControlRigClass, SkeletalMeshComponent, InBinding.BindingID, nullptr);
@@ -288,13 +286,16 @@ TArray<UMovieSceneTrack*> UControlRigSequencerEditorLibrary::FindOrCreateControl
 {
 	TArray< UMovieSceneTrack*> Tracks;
 	TArray<UObject*, TInlineAllocator<1>> Result;
-
+	if (LevelSequence == nullptr || LevelSequence->GetMovieScene() == nullptr)
+	{
+		return Tracks;
+	}
+	UMovieScene* MovieScene = LevelSequence->GetMovieScene();
 	UObject* Context = nullptr;
 	LevelSequence->LocateBoundObjects(InBinding.BindingID, Context, Result);
 	if (Result.Num() > 0 && Result[0])
 	{
 		UObject* BoundObject = Result[0];
-		
 		if (AActor* BoundActor = Cast<AActor>(BoundObject))
 		{
 			TArray<UControlRigComponent*> ControlRigComponents;
@@ -303,8 +304,27 @@ TArray<UMovieSceneTrack*> UControlRigSequencerEditorLibrary::FindOrCreateControl
 			{
 				if (UControlRig* CR = ControlRigComponent->GetControlRig())
 				{
-					UMovieSceneControlRigParameterTrack* Track = AddControlRig(LevelSequence,CR->GetClass(), BoundActor, InBinding.BindingID, CR);
-					Tracks.Add(Track);
+					UMovieSceneControlRigParameterTrack* GoodTrack = nullptr;
+					if (FMovieSceneBinding* Binding = MovieScene->FindBinding(InBinding.BindingID))
+					{
+						for (UMovieSceneTrack* Track : Binding->GetTracks())
+						{
+							if (UMovieSceneControlRigParameterTrack* ControlRigParameterTrack = Cast<UMovieSceneControlRigParameterTrack>(Track))
+							{
+								if (ControlRigParameterTrack->GetControlRig() == CR)
+								{
+									GoodTrack = ControlRigParameterTrack;
+									break;
+								}
+							}
+						}
+					}
+
+					if (GoodTrack == nullptr)
+					{
+						GoodTrack = AddControlRig(LevelSequence, CR->GetClass(), BoundActor, InBinding.BindingID, CR);
+					}
+					Tracks.Add(GoodTrack);
 				}
 			}
 		}
@@ -331,26 +351,50 @@ bool UControlRigSequencerEditorLibrary::TweenControlRig(ULevelSequence* LevelSeq
 }
 
 
-bool UControlRigSequencerEditorLibrary::SnapControlRig(FFrameNumber StartFrame, FFrameNumber EndFrame, const FControlRigSnapperSelection& ChildrenToSnap,
-	const FControlRigSnapperSelection& ParentToSnap, const UControlRigSnapSettings* SnapSettings)
+bool UControlRigSequencerEditorLibrary::SnapControlRig(ULevelSequence* LevelSequence, FFrameNumber StartFrame, FFrameNumber EndFrame, const FControlRigSnapperSelection& ChildrenToSnap,
+	const FControlRigSnapperSelection& ParentToSnap, const UControlRigSnapSettings* SnapSettings, ESequenceTimeUnit TimeUnit)
 {
+	if (LevelSequence == nullptr || LevelSequence->GetMovieScene() == nullptr)
+	{
+		return false;
+	}
 	FControlRigSnapper Snapper;
+	if (TimeUnit == ESequenceTimeUnit::DisplayRate)
+	{
+		StartFrame = FFrameRate::TransformTime(FFrameTime(StartFrame, 0), LevelSequence->GetMovieScene()->GetDisplayRate(), LevelSequence->GetMovieScene()->GetTickResolution()).RoundToFrame();
+		EndFrame = FFrameRate::TransformTime(FFrameTime(EndFrame, 0), LevelSequence->GetMovieScene()->GetDisplayRate(), LevelSequence->GetMovieScene()->GetTickResolution()).RoundToFrame();
+
+	}
 	return Snapper.SnapIt(StartFrame, EndFrame, ChildrenToSnap, ParentToSnap,SnapSettings);
 }
 
-FTransform UControlRigSequencerEditorLibrary::GetActorWorldTransform(ULevelSequence* LevelSequence, AActor* Actor, FFrameNumber Frame)
+FTransform UControlRigSequencerEditorLibrary::GetActorWorldTransform(ULevelSequence* LevelSequence, AActor* Actor, FFrameNumber Frame, ESequenceTimeUnit TimeUnit)
 {
 	TArray<FFrameNumber> Frames;
 	Frames.Add(Frame);
-	TArray<FTransform> Transforms = GetActorWorldTransforms(LevelSequence, Actor, Frames);
+	TArray<FTransform> Transforms = GetActorWorldTransforms(LevelSequence, Actor, Frames, TimeUnit);
 	if (Transforms.Num() == 1)
 	{
 		return Transforms[0];
 	}
 	return FTransform::Identity;
 }
+static void ConvertFramesToTickResolution(UMovieScene* MovieScene, const TArray<FFrameNumber>& InFrames, ESequenceTimeUnit TimeUnit, TArray<FFrameNumber>& OutFrames)
+{
+	OutFrames.SetNum(InFrames.Num());
+	int32 Index = 0;
+	for (const FFrameNumber& Frame : InFrames)
+	{
+		FFrameTime FrameTime(Frame);
+		if (TimeUnit == ESequenceTimeUnit::DisplayRate)
+		{
+			FrameTime = FFrameRate::TransformTime(FrameTime, MovieScene->GetDisplayRate(), MovieScene->GetTickResolution());
+		}
+		OutFrames[Index++] = FrameTime.RoundToFrame();
+	}
+}
 
-TArray<FTransform> UControlRigSequencerEditorLibrary::GetActorWorldTransforms(ULevelSequence* LevelSequence,AActor* Actor, const TArray<FFrameNumber>& Frames)
+TArray<FTransform> UControlRigSequencerEditorLibrary::GetActorWorldTransforms(ULevelSequence* LevelSequence,AActor* Actor, const TArray<FFrameNumber>& InFrames, ESequenceTimeUnit TimeUnit)
 {
 	TWeakPtr<ISequencer> WeakSequencer = GetSequencerFromAsset(LevelSequence);
 	TArray<FTransform> OutWorldTransforms;
@@ -358,17 +402,20 @@ TArray<FTransform> UControlRigSequencerEditorLibrary::GetActorWorldTransforms(UL
 	{
 		FActorForWorldTransforms Actors;
 		Actors.Actor = Actor;
+		TArray<FFrameNumber> Frames;
+		ConvertFramesToTickResolution(LevelSequence->GetMovieScene(), InFrames, TimeUnit, Frames);
 		MovieSceneToolHelpers::GetActorWorldTransforms(WeakSequencer.Pin().Get(), Actors, Frames, OutWorldTransforms); 
 	}
 	return OutWorldTransforms;
 
 }
 
-FTransform UControlRigSequencerEditorLibrary::GetSkeletalMeshComponentWorldTransform(ULevelSequence* LevelSequence,USkeletalMeshComponent* SkeletalMeshComponent, FFrameNumber Frame, FName SocketName)
+FTransform UControlRigSequencerEditorLibrary::GetSkeletalMeshComponentWorldTransform(ULevelSequence* LevelSequence,USkeletalMeshComponent* SkeletalMeshComponent, FFrameNumber Frame, 
+	ESequenceTimeUnit TimeUnit, FName SocketName)
 {
 	TArray<FFrameNumber> Frames;
 	Frames.Add(Frame);
-	TArray<FTransform> Transforms = GetSkeletalMeshComponentWorldTransforms(LevelSequence, SkeletalMeshComponent, Frames,SocketName);
+	TArray<FTransform> Transforms = GetSkeletalMeshComponentWorldTransforms(LevelSequence, SkeletalMeshComponent, Frames, TimeUnit,SocketName);
 	if (Transforms.Num() == 1)
 	{
 		return Transforms[0];
@@ -376,7 +423,8 @@ FTransform UControlRigSequencerEditorLibrary::GetSkeletalMeshComponentWorldTrans
 	return FTransform::Identity;
 }
 
-TArray<FTransform> UControlRigSequencerEditorLibrary::GetSkeletalMeshComponentWorldTransforms(ULevelSequence* LevelSequence,USkeletalMeshComponent* SkeletalMeshComponent, const TArray<FFrameNumber>& Frames, FName SocketName)
+TArray<FTransform> UControlRigSequencerEditorLibrary::GetSkeletalMeshComponentWorldTransforms(ULevelSequence* LevelSequence,USkeletalMeshComponent* SkeletalMeshComponent, const TArray<FFrameNumber>& InFrames,
+	ESequenceTimeUnit TimeUnit, FName SocketName)
 {
 	TWeakPtr<ISequencer> WeakSequencer = GetSequencerFromAsset(LevelSequence);
 	TArray<FTransform> OutWorldTransforms;
@@ -389,17 +437,20 @@ TArray<FTransform> UControlRigSequencerEditorLibrary::GetSkeletalMeshComponentWo
 			Actors.Actor = Actor;
 			Actors.Component = SkeletalMeshComponent;
 			Actors.SocketName = SocketName;
+			TArray<FFrameNumber> Frames;
+			ConvertFramesToTickResolution(LevelSequence->GetMovieScene(), InFrames, TimeUnit, Frames);
 			MovieSceneToolHelpers::GetActorWorldTransforms(WeakSequencer.Pin().Get(), Actors, Frames, OutWorldTransforms);
 		}
 	}
 	return OutWorldTransforms;
 }
 
-FTransform UControlRigSequencerEditorLibrary::GetControlRigWorldTransform(ULevelSequence* LevelSequence,UControlRig* ControlRig, FName ControlName, FFrameNumber Frame)
+FTransform UControlRigSequencerEditorLibrary::GetControlRigWorldTransform(ULevelSequence* LevelSequence,UControlRig* ControlRig, FName ControlName, FFrameNumber Frame,
+	ESequenceTimeUnit TimeUnit)
 {
 	TArray<FFrameNumber> Frames;
 	Frames.Add(Frame);
-	TArray<FTransform> Transforms = GetControlRigWorldTransforms(LevelSequence, ControlRig, ControlName, Frames);
+	TArray<FTransform> Transforms = GetControlRigWorldTransforms(LevelSequence, ControlRig, ControlName, Frames, TimeUnit);
 	if (Transforms.Num() == 1)
 	{
 		return Transforms[0];
@@ -407,7 +458,8 @@ FTransform UControlRigSequencerEditorLibrary::GetControlRigWorldTransform(ULevel
 	return FTransform::Identity;
 }
 
-TArray<FTransform> UControlRigSequencerEditorLibrary::GetControlRigWorldTransforms(ULevelSequence* LevelSequence,UControlRig* ControlRig, FName ControlName, const TArray<FFrameNumber>& Frames)
+TArray<FTransform> UControlRigSequencerEditorLibrary::GetControlRigWorldTransforms(ULevelSequence* LevelSequence,UControlRig* ControlRig, FName ControlName, const TArray<FFrameNumber>& InFrames,
+	ESequenceTimeUnit TimeUnit)
 {
 	TWeakPtr<ISequencer> WeakSequencer = GetSequencerFromAsset(LevelSequence);
 	TArray<FTransform> OutWorldTransforms;
@@ -425,6 +477,8 @@ TArray<FTransform> UControlRigSequencerEditorLibrary::GetControlRigWorldTransfor
 					TArray<FTransform> ControlRigParentWorldTransforms;
 					FActorForWorldTransforms ControlRigActorSelection;
 					ControlRigActorSelection.Actor = Actor;
+					TArray<FFrameNumber> Frames;
+					ConvertFramesToTickResolution(LevelSequence->GetMovieScene(), InFrames, TimeUnit, Frames);
 					MovieSceneToolHelpers::GetActorWorldTransforms(WeakSequencer.Pin().Get(), ControlRigActorSelection, Frames, ControlRigParentWorldTransforms);
 					FControlRigSnapper Snapper;
 					Snapper.GetControlRigControlTransforms(WeakSequencer.Pin().Get(), ControlRig, ControlName, Frames, ControlRigParentWorldTransforms, OutWorldTransforms);	
@@ -436,7 +490,8 @@ TArray<FTransform> UControlRigSequencerEditorLibrary::GetControlRigWorldTransfor
 }
 
 
-static void LocalSetControlRigWorldTransforms(ULevelSequence* LevelSequence,UControlRig* ControlRig, FName ControlName, EControlRigSetKey SetKey, const TArray<FFrameNumber>& Frames, const TArray<FTransform>& WorldTransforms)
+static void LocalSetControlRigWorldTransforms(ULevelSequence* LevelSequence,UControlRig* ControlRig, FName ControlName, EControlRigSetKey SetKey, const TArray<FFrameNumber>& InFrames, 
+	const TArray<FTransform>& WorldTransforms, ESequenceTimeUnit TimeUnit)
 {
 	TWeakPtr<ISequencer> WeakSequencer = GetSequencerFromAsset(LevelSequence);
 	if (WeakSequencer.IsValid() && ControlRig)
@@ -449,6 +504,9 @@ static void LocalSetControlRigWorldTransforms(ULevelSequence* LevelSequence,UCon
 				AActor* Actor = Component->GetTypedOuter< AActor >();
 				if (Actor)
 				{
+					TArray<FFrameNumber> Frames;
+					ConvertFramesToTickResolution(LevelSequence->GetMovieScene(), InFrames, TimeUnit, Frames);
+
 					UMovieScene* MovieScene = WeakSequencer.Pin().Get()->GetFocusedMovieSceneSequence()->GetMovieScene();
 					MovieScene->Modify();
 					FFrameRate TickResolution = MovieScene->GetTickResolution();
@@ -486,7 +544,8 @@ static void LocalSetControlRigWorldTransforms(ULevelSequence* LevelSequence,UCon
 	}
 }
 
-void UControlRigSequencerEditorLibrary::SetControlRigWorldTransform(ULevelSequence* LevelSequence, UControlRig* ControlRig, FName ControlName, FFrameNumber Frame, const FTransform& WorldTransform, bool bSetKey)
+void UControlRigSequencerEditorLibrary::SetControlRigWorldTransform(ULevelSequence* LevelSequence, UControlRig* ControlRig, FName ControlName, FFrameNumber Frame, const FTransform& WorldTransform,
+	ESequenceTimeUnit TimeUnit, bool bSetKey)
 {
 	EControlRigSetKey SetKey = bSetKey ? EControlRigSetKey::Always : EControlRigSetKey::DoNotCare;
 	TArray<FFrameNumber> Frames;
@@ -494,13 +553,14 @@ void UControlRigSequencerEditorLibrary::SetControlRigWorldTransform(ULevelSequen
 	TArray<FTransform> WorldTransforms;
 	WorldTransforms.Add(WorldTransform);
 
-	LocalSetControlRigWorldTransforms(LevelSequence, ControlRig, ControlName, EControlRigSetKey::Always, Frames, WorldTransforms);
+	LocalSetControlRigWorldTransforms(LevelSequence, ControlRig, ControlName, EControlRigSetKey::Always, Frames, WorldTransforms, TimeUnit);
 
 }
 
-void UControlRigSequencerEditorLibrary::SetControlRigWorldTransforms(ULevelSequence* LevelSequence, UControlRig* ControlRig, FName ControlName, const TArray<FFrameNumber>& Frames, const TArray<FTransform>& WorldTransforms)
+void UControlRigSequencerEditorLibrary::SetControlRigWorldTransforms(ULevelSequence* LevelSequence, UControlRig* ControlRig, FName ControlName, const TArray<FFrameNumber>& Frames, const TArray<FTransform>& WorldTransforms,
+	ESequenceTimeUnit TimeUnit)
 {
-	LocalSetControlRigWorldTransforms(LevelSequence, ControlRig, ControlName, EControlRigSetKey::Always, Frames, WorldTransforms);
+	LocalSetControlRigWorldTransforms(LevelSequence, ControlRig, ControlName, EControlRigSetKey::Always, Frames, WorldTransforms,TimeUnit);
 }
 
 bool UControlRigSequencerEditorLibrary::BakeToControlRig(UWorld* World, ULevelSequence* LevelSequence, UClass* InClass, UAnimSeqExportOption* ExportOptions, bool bReduceKeys, float Tolerance,
@@ -538,6 +598,7 @@ bool UControlRigSequencerEditorLibrary::BakeToControlRig(UWorld* World, ULevelSe
 		return false;
 	}
 	
+	
 
 	bool bResult = false;
 	const FScopedTransaction Transaction(LOCTEXT("BakeToControlRig_Transaction", "Bake To Control Rig"));
@@ -636,6 +697,8 @@ bool UControlRigSequencerEditorLibrary::BakeToControlRig(UWorld* World, ULevelSe
 					ControlRig->GetObjectBinding()->BindToObject(SkeletalMeshComp);
 					ControlRig->GetDataSourceRegistry()->RegisterDataSource(UControlRig::OwnerComponent, ControlRig->GetObjectBinding()->GetBoundObject());
 					ControlRig->Initialize();
+					ControlRig->RequestInit();
+					ControlRig->SetBoneInitialTransformsFromSkeletalMeshComponent(SkeletalMeshComp);
 					ControlRig->Evaluate_AnyThread();
 
 					bool bSequencerOwnsControlRig = true;
@@ -702,7 +765,7 @@ bool UControlRigSequencerEditorLibrary::BakeToControlRig(UWorld* World, ULevelSe
 }
 
 bool UControlRigSequencerEditorLibrary::LoadAnimSequenceIntoControlRigSection(UMovieSceneSection* MovieSceneSection, UAnimSequence* AnimSequence, USkeletalMeshComponent* SkelMeshComp,
-	FFrameNumber InStartFrame,bool bKeyReduce, float Tolerance)
+	FFrameNumber InStartFrame, ESequenceTimeUnit TimeUnit,bool bKeyReduce, float Tolerance)
 {
 	if (MovieSceneSection == nullptr || AnimSequence == nullptr || SkelMeshComp == nullptr)
 	{
@@ -715,14 +778,18 @@ bool UControlRigSequencerEditorLibrary::LoadAnimSequenceIntoControlRigSection(UM
 	}
 	if (UMovieSceneControlRigParameterSection* Section = Cast<UMovieSceneControlRigParameterSection>(MovieSceneSection))
 	{
+		if (TimeUnit == ESequenceTimeUnit::DisplayRate)
+		{
+			InStartFrame = FFrameRate::TransformTime(FFrameTime(InStartFrame, 0), MovieScene->GetDisplayRate(),MovieScene->GetTickResolution()).RoundToFrame();
+		}
 		return Section->LoadAnimSequenceIntoThisSection(AnimSequence, MovieScene, SkelMeshComp, bKeyReduce, Tolerance, InStartFrame);
 	}
 	return false;
 }
 
 static bool LocalGetControlRigControlValues(IMovieScenePlayer* Player, UMovieSceneSequence* MovieSceneSequence, FMovieSceneSequenceIDRef Template, FMovieSceneSequenceTransform& RootToLocalTransform,
-	UControlRig* ControlRig, const FName& ControlName,
-	const TArray<FFrameNumber>& Frames, TArray<FRigControlValue>& OutValues)
+	UControlRig* ControlRig, const FName& ControlName, ESequenceTimeUnit TimeUnit,
+	const TArray<FFrameNumber>& InFrames, TArray<FRigControlValue>& OutValues)
 {
 	if (Player == nullptr || MovieSceneSequence == nullptr || ControlRig == nullptr)
 	{
@@ -737,12 +804,17 @@ static bool LocalGetControlRigControlValues(IMovieScenePlayer* Player, UMovieSce
 	{
 
 		FFrameRate TickResolution = MovieScene->GetTickResolution();
-
+		TArray<FFrameNumber> Frames;
+		ConvertFramesToTickResolution(MovieSceneSequence->GetMovieScene(), InFrames, TimeUnit, Frames);
 		OutValues.SetNum(Frames.Num());
 		for (int32 Index = 0; Index < Frames.Num(); ++Index)
 		{
 			const FFrameNumber& FrameNumber = Frames[Index];
 			FFrameTime GlobalTime(FrameNumber);
+			if (TimeUnit == ESequenceTimeUnit::DisplayRate)
+			{
+				GlobalTime = FFrameRate::TransformTime(GlobalTime, MovieScene->GetDisplayRate(), MovieScene->GetTickResolution());
+			}
 
 			FMovieSceneContext Context = FMovieSceneContext(FMovieSceneEvaluationRange(GlobalTime, TickResolution), Player->GetPlaybackStatus()).SetHasJumped(true);
 
@@ -754,7 +826,7 @@ static bool LocalGetControlRigControlValues(IMovieScenePlayer* Player, UMovieSce
 	return true;
 }
 
-static bool GetControlRigValues(ISequencer* Sequencer, UControlRig* ControlRig, const FName& ControlName,
+static bool GetControlRigValues(ISequencer* Sequencer, UControlRig* ControlRig, const FName& ControlName, ESequenceTimeUnit TimeUnit,
 	const TArray<FFrameNumber>& Frames,  TArray<FRigControlValue>& OutValues)
 {
 	if (Sequencer->GetFocusedMovieSceneSequence())
@@ -762,13 +834,13 @@ static bool GetControlRigValues(ISequencer* Sequencer, UControlRig* ControlRig,
 		FMovieSceneSequenceIDRef Template = Sequencer->GetFocusedTemplateID();
 		FMovieSceneSequenceTransform RootToLocalTransform;
 		return LocalGetControlRigControlValues(Sequencer, Sequencer->GetFocusedMovieSceneSequence(), Template, RootToLocalTransform,
-			ControlRig, ControlName, Frames, OutValues);
+			ControlRig, ControlName,TimeUnit, Frames, OutValues);
 
 	}
 	return false;
 }
 
-static bool GetControlRigValue(ISequencer* Sequencer, UControlRig* ControlRig, const FName& ControlName,
+static bool GetControlRigValue(ISequencer* Sequencer, UControlRig* ControlRig, const FName& ControlName, ESequenceTimeUnit TimeUnit,
 	const FFrameNumber Frame, FRigControlValue& OutValue)
 {
 	if (Sequencer->GetFocusedMovieSceneSequence())
@@ -779,7 +851,7 @@ static bool GetControlRigValue(ISequencer* Sequencer, UControlRig* ControlRig, c
 		FMovieSceneSequenceIDRef Template = Sequencer->GetFocusedTemplateID();
 		FMovieSceneSequenceTransform RootToLocalTransform;
 		bool bVal = LocalGetControlRigControlValues(Sequencer, Sequencer->GetFocusedMovieSceneSequence(), Template, RootToLocalTransform,
-			ControlRig, ControlName, Frames, OutValues);
+			ControlRig, ControlName,TimeUnit, Frames, OutValues);
 		if (bVal)
 		{
 			OutValue = OutValues[0];
@@ -789,7 +861,7 @@ static bool GetControlRigValue(ISequencer* Sequencer, UControlRig* ControlRig, c
 	return false;
 }
 
-static bool GetControlRigValues(UWorld* World, ULevelSequence* LevelSequence, UControlRig* ControlRig, const FName& ControlName,
+static bool GetControlRigValues(UWorld* World, ULevelSequence* LevelSequence, UControlRig* ControlRig, const FName& ControlName, ESequenceTimeUnit TimeUnit,
 	const TArray<FFrameNumber>& Frames, TArray<FRigControlValue>& OutValues)
 {
 	if (LevelSequence)
@@ -803,22 +875,30 @@ static bool GetControlRigValues(UWorld* World, ULevelSequence* LevelSequence, UC
 		Player->Initialize(LevelSequence, World->PersistentLevel, Settings, CameraSettings);
 		Player->State.AssignSequence(MovieSceneSequenceID::Root, *LevelSequence, *Player);
 		return LocalGetControlRigControlValues(Player, LevelSequence, Template, RootToLocalTransform,
-			ControlRig, ControlName, Frames, OutValues);
+			ControlRig, ControlName,TimeUnit, Frames, OutValues);
 
 	}
 	return false;
 }
 
 
-float UControlRigSequencerEditorLibrary::GetLocalControlRigFloat(ULevelSequence* LevelSequence,UControlRig* ControlRig, FName ControlName, FFrameNumber Frame)
+float UControlRigSequencerEditorLibrary::GetLocalControlRigFloat(ULevelSequence* LevelSequence,UControlRig* ControlRig, FName ControlName, FFrameNumber Frame, ESequenceTimeUnit TimeUnit)
 {
 	float Value = 0.0f;
 	TWeakPtr<ISequencer> WeakSequencer = GetSequencerFromAsset(LevelSequence);
 
 	if (WeakSequencer.IsValid() && ControlRig)
 	{
+		if (FRigControlElement* Element = ControlRig->FindControl(ControlName))
+		{
+			if (Element->Settings.ControlType != ERigControlType::Float)
+			{
+				UE_LOG(LogControlRig, Error, TEXT("Control Rig Wrong Type"));
+				return Value;
+			}
+		}
 		FRigControlValue RigValue;
-		if (GetControlRigValue(WeakSequencer.Pin().Get(), ControlRig, ControlName, Frame, RigValue))
+		if (GetControlRigValue(WeakSequencer.Pin().Get(), ControlRig, ControlName, TimeUnit, Frame, RigValue))
 		{
 			Value = RigValue.Get<float>();
 		}
@@ -826,15 +906,24 @@ float UControlRigSequencerEditorLibrary::GetLocalControlRigFloat(ULevelSequence*
 	return Value;
 }
 
-TArray<float> UControlRigSequencerEditorLibrary::GetLocalControlRigFloats(ULevelSequence* LevelSequence,UControlRig* ControlRig, FName ControlName, const TArray<FFrameNumber>& Frames)
+TArray<float> UControlRigSequencerEditorLibrary::GetLocalControlRigFloats(ULevelSequence* LevelSequence,UControlRig* ControlRig, FName ControlName, const TArray<FFrameNumber>& Frames, ESequenceTimeUnit TimeUnit)
 {
 	TArray<float> Values;
 	TWeakPtr<ISequencer> WeakSequencer = GetSequencerFromAsset(LevelSequence);
 
 	if (WeakSequencer.IsValid() && ControlRig)
 	{
+		if (FRigControlElement* Element = ControlRig->FindControl(ControlName))
+		{
+			if (Element->Settings.ControlType != ERigControlType::Float)
+			{
+				UE_LOG(LogControlRig, Error, TEXT("Control Rig Wrong Type"));
+				return Values;
+			}
+		}
+
 		TArray<FRigControlValue> RigValues;
-		if (GetControlRigValues(WeakSequencer.Pin().Get(), ControlRig, ControlName, Frames, RigValues))
+		if (GetControlRigValues(WeakSequencer.Pin().Get(), ControlRig, ControlName,TimeUnit, Frames, RigValues))
 		{
 			Values.Reserve(RigValues.Num());
 			for (const FRigControlValue& RigValue : RigValues)
@@ -847,28 +936,49 @@ TArray<float> UControlRigSequencerEditorLibrary::GetLocalControlRigFloats(ULevel
 	return Values;
 }
 
-void UControlRigSequencerEditorLibrary::SetLocalControlRigFloat(ULevelSequence* LevelSequence, UControlRig* ControlRig, FName ControlName, FFrameNumber Frame, float Value, bool bSetKey)
+void UControlRigSequencerEditorLibrary::SetLocalControlRigFloat(ULevelSequence* LevelSequence, UControlRig* ControlRig, FName ControlName, FFrameNumber Frame, float Value, ESequenceTimeUnit TimeUnit, bool bSetKey)
 {
 	if (LevelSequence == nullptr || ControlRig == nullptr)
 	{
 		return;
 	}
+	if (FRigControlElement* Element = ControlRig->FindControl(ControlName))
+	{
+		if (Element->Settings.ControlType != ERigControlType::Float)
+		{
+			UE_LOG(LogControlRig, Error, TEXT("Control Rig Wrong Type"));
+			return;
+		}
+	}
 	if (UMovieScene* MovieScene = LevelSequence->GetMovieScene())
 	{
 		FFrameRate TickResolution = MovieScene->GetTickResolution();
+		if (TimeUnit == ESequenceTimeUnit::DisplayRate)
+		{
+			Frame = FFrameRate::TransformTime(FFrameTime(Frame, 0), MovieScene->GetDisplayRate(), MovieScene->GetTickResolution()).RoundToFrame();
+		}
 		FRigControlModifiedContext Context;
 		Context.SetKey = bSetKey ? EControlRigSetKey::Always : EControlRigSetKey::DoNotCare;
 		Context.LocalTime = TickResolution.AsSeconds(FFrameTime(Frame));
-		ControlRig->SetControlValue<float>(ControlName, Value, true, FRigControlModifiedContext(EControlRigSetKey::Never, false));
+		ControlRig->SetControlValue<float>(ControlName, Value, true, Context);
 	}
 }
 
-void UControlRigSequencerEditorLibrary::SetLocalControlRigFloats(ULevelSequence* LevelSequence, UControlRig* ControlRig, FName ControlName, const TArray<FFrameNumber>& Frames, const TArray<float> Values)
+void UControlRigSequencerEditorLibrary::SetLocalControlRigFloats(ULevelSequence* LevelSequence, UControlRig* ControlRig, FName ControlName, const TArray<FFrameNumber>& Frames, const TArray<float> Values,
+	ESequenceTimeUnit TimeUnit)
 {
 	if (LevelSequence == nullptr || ControlRig == nullptr || (Frames.Num() != Values.Num()))
 	{
 		return;
 	}
+	if (FRigControlElement* Element = ControlRig->FindControl(ControlName))
+	{
+		if (Element->Settings.ControlType != ERigControlType::Float)
+		{
+			UE_LOG(LogControlRig, Error, TEXT("Control Rig Wrong Type"));
+			return;
+		}
+	}
 	if (UMovieScene* MovieScene = LevelSequence->GetMovieScene())
 	{
 		FFrameRate TickResolution = MovieScene->GetTickResolution();
@@ -877,23 +987,35 @@ void UControlRigSequencerEditorLibrary::SetLocalControlRigFloats(ULevelSequence*
 		for (int32 Index = 0; Index < Frames.Num(); ++Index)
 		{
 			FFrameNumber Frame = Frames[Index];
+			if (TimeUnit == ESequenceTimeUnit::DisplayRate)
+			{
+				Frame = FFrameRate::TransformTime(FFrameTime(Frame, 0), MovieScene->GetDisplayRate(), MovieScene->GetTickResolution()).RoundToFrame();
+			}
 			float  Value = Values[Index];
 			Context.LocalTime = TickResolution.AsSeconds(FFrameTime(Frame));
-			ControlRig->SetControlValue<float>(ControlName, Value, true, FRigControlModifiedContext(EControlRigSetKey::Never, false));
+			ControlRig->SetControlValue<float>(ControlName, Value, true, Context);
 		}
 	}
 }
 
 
-bool UControlRigSequencerEditorLibrary::GetLocalControlRigBool(ULevelSequence* LevelSequence, UControlRig* ControlRig, FName ControlName, FFrameNumber Frame)
+bool UControlRigSequencerEditorLibrary::GetLocalControlRigBool(ULevelSequence* LevelSequence, UControlRig* ControlRig, FName ControlName, FFrameNumber Frame, ESequenceTimeUnit TimeUnit)
 {
 	bool Value = true;
 	TWeakPtr<ISequencer> WeakSequencer = GetSequencerFromAsset(LevelSequence);
 
 	if (WeakSequencer.IsValid() && ControlRig)
 	{
+		if (FRigControlElement* Element = ControlRig->FindControl(ControlName))
+		{
+			if (Element->Settings.ControlType != ERigControlType::Bool)
+			{
+				UE_LOG(LogControlRig, Error, TEXT("Control Rig Wrong Type"));
+				return Value;
+			}
+		}
 		FRigControlValue RigValue;
-		if (GetControlRigValue(WeakSequencer.Pin().Get(), ControlRig, ControlName, Frame, RigValue))
+		if (GetControlRigValue(WeakSequencer.Pin().Get(), ControlRig, ControlName,TimeUnit, Frame, RigValue))
 		{
 			Value = RigValue.Get<bool>();
 		}
@@ -901,15 +1023,23 @@ bool UControlRigSequencerEditorLibrary::GetLocalControlRigBool(ULevelSequence* L
 	return Value;
 }
 
-TArray<bool> UControlRigSequencerEditorLibrary::GetLocalControlRigBools(ULevelSequence* LevelSequence, UControlRig* ControlRig, FName ControlName, const TArray<FFrameNumber>& Frames)
+TArray<bool> UControlRigSequencerEditorLibrary::GetLocalControlRigBools(ULevelSequence* LevelSequence, UControlRig* ControlRig, FName ControlName, const TArray<FFrameNumber>& Frames, ESequenceTimeUnit TimeUnit)
 {
 	TArray<bool> Values;
 	TWeakPtr<ISequencer> WeakSequencer = GetSequencerFromAsset(LevelSequence);
 
 	if (WeakSequencer.IsValid() && ControlRig)
 	{
+		if (FRigControlElement* Element = ControlRig->FindControl(ControlName))
+		{
+			if (Element->Settings.ControlType != ERigControlType::Bool)
+			{
+				UE_LOG(LogControlRig, Error, TEXT("Control Rig Wrong Type"));
+				return Values;
+			}
+		}
 		TArray<FRigControlValue> RigValues;
-		if (GetControlRigValues(WeakSequencer.Pin().Get(), ControlRig, ControlName, Frames, RigValues))
+		if (GetControlRigValues(WeakSequencer.Pin().Get(), ControlRig, ControlName,TimeUnit, Frames, RigValues))
 		{
 			Values.Reserve(RigValues.Num());
 			for (const FRigControlValue& RigValue : RigValues)
@@ -922,7 +1052,7 @@ TArray<bool> UControlRigSequencerEditorLibrary::GetLocalControlRigBools(ULevelSe
 	return Values;
 }
 
-void UControlRigSequencerEditorLibrary::SetLocalControlRigBool(ULevelSequence* LevelSequence, UControlRig* ControlRig, FName ControlName, FFrameNumber Frame, bool Value, bool bSetKey)
+void UControlRigSequencerEditorLibrary::SetLocalControlRigBool(ULevelSequence* LevelSequence, UControlRig* ControlRig, FName ControlName, FFrameNumber Frame, bool Value, ESequenceTimeUnit TimeUnit, bool bSetKey)
 {
 	if (LevelSequence == nullptr || ControlRig == nullptr)
 	{
@@ -930,15 +1060,28 @@ void UControlRigSequencerEditorLibrary::SetLocalControlRigBool(ULevelSequence* L
 	}
 	if (UMovieScene* MovieScene = LevelSequence->GetMovieScene())
 	{
+		if (FRigControlElement* Element = ControlRig->FindControl(ControlName))
+		{
+			if (Element->Settings.ControlType != ERigControlType::Bool)
+			{
+				UE_LOG(LogControlRig, Error, TEXT("Control Rig Wrong Type"));
+				return;
+			}
+		}
 		FFrameRate TickResolution = MovieScene->GetTickResolution();
+		if (TimeUnit == ESequenceTimeUnit::DisplayRate)
+		{
+			Frame = FFrameRate::TransformTime(FFrameTime(Frame, 0), MovieScene->GetDisplayRate(), MovieScene->GetTickResolution()).RoundToFrame();
+		}
 		FRigControlModifiedContext Context;
 		Context.SetKey = bSetKey ? EControlRigSetKey::Always : EControlRigSetKey::DoNotCare;
 		Context.LocalTime = TickResolution.AsSeconds(FFrameTime(Frame));
-		ControlRig->SetControlValue<bool>(ControlName, Value, true, FRigControlModifiedContext(EControlRigSetKey::Never, false));
+		ControlRig->SetControlValue<bool>(ControlName, Value, true, Context);
 	}
 }
 
-void UControlRigSequencerEditorLibrary::SetLocalControlRigBools(ULevelSequence* LevelSequence, UControlRig* ControlRig, FName ControlName, const TArray<FFrameNumber>& Frames, const TArray<bool> Values)
+void UControlRigSequencerEditorLibrary::SetLocalControlRigBools(ULevelSequence* LevelSequence, UControlRig* ControlRig, FName ControlName, const TArray<FFrameNumber>& Frames,
+	const TArray<bool> Values, ESequenceTimeUnit TimeUnit)
 {
 	if (LevelSequence == nullptr || ControlRig == nullptr || (Frames.Num() != Values.Num()))
 	{
@@ -946,28 +1089,48 @@ void UControlRigSequencerEditorLibrary::SetLocalControlRigBools(ULevelSequence*
 	}
 	if (UMovieScene* MovieScene = LevelSequence->GetMovieScene())
 	{
+		if (FRigControlElement* Element = ControlRig->FindControl(ControlName))
+		{
+			if (Element->Settings.ControlType != ERigControlType::Bool)
+			{
+				UE_LOG(LogControlRig, Error, TEXT("Control Rig Wrong Type"));
+				return;
+			}
+		}
 		FFrameRate TickResolution = MovieScene->GetTickResolution();
 		FRigControlModifiedContext Context;
 		Context.SetKey = EControlRigSetKey::Always;
 		for (int32 Index = 0; Index < Frames.Num(); ++Index)
 		{
 			FFrameNumber Frame = Frames[Index];
+			if (TimeUnit == ESequenceTimeUnit::DisplayRate)
+			{
+				Frame = FFrameRate::TransformTime(FFrameTime(Frame, 0), MovieScene->GetDisplayRate(), MovieScene->GetTickResolution()).RoundToFrame();
+			}
 			bool Value = Values[Index];
 			Context.LocalTime = TickResolution.AsSeconds(FFrameTime(Frame));
-			ControlRig->SetControlValue<bool>(ControlName, Value, true, FRigControlModifiedContext(EControlRigSetKey::Never, false));
+			ControlRig->SetControlValue<bool>(ControlName, Value, true, Context);
 		}
 	}
 }
 
-int32 UControlRigSequencerEditorLibrary::GetLocalControlRigInt(ULevelSequence* LevelSequence, UControlRig* ControlRig, FName ControlName, FFrameNumber Frame)
+int32 UControlRigSequencerEditorLibrary::GetLocalControlRigInt(ULevelSequence* LevelSequence, UControlRig* ControlRig, FName ControlName, FFrameNumber Frame, ESequenceTimeUnit TimeUnit)
 {
 	int32 Value = 0;
 	TWeakPtr<ISequencer> WeakSequencer = GetSequencerFromAsset(LevelSequence);
 
 	if (WeakSequencer.IsValid() && ControlRig)
 	{
+		if (FRigControlElement* Element = ControlRig->FindControl(ControlName))
+		{
+			if (Element->Settings.ControlType != ERigControlType::Integer)
+			{
+				UE_LOG(LogControlRig, Error, TEXT("Control Rig Wrong Type"));
+				return Value;
+			}
+		}
 		FRigControlValue RigValue;
-		if (GetControlRigValue(WeakSequencer.Pin().Get(), ControlRig, ControlName, Frame, RigValue))
+		if (GetControlRigValue(WeakSequencer.Pin().Get(), ControlRig, ControlName,TimeUnit, Frame, RigValue))
 		{
 			Value = RigValue.Get<int32>();
 		}
@@ -975,15 +1138,23 @@ int32 UControlRigSequencerEditorLibrary::GetLocalControlRigInt(ULevelSequence* L
 	return Value;
 }
 
-TArray<int32> UControlRigSequencerEditorLibrary::GetLocalControlRigInts(ULevelSequence* LevelSequence, UControlRig* ControlRig, FName ControlName, const TArray<FFrameNumber>& Frames)
+TArray<int32> UControlRigSequencerEditorLibrary::GetLocalControlRigInts(ULevelSequence* LevelSequence, UControlRig* ControlRig, FName ControlName, const TArray<FFrameNumber>& Frames, ESequenceTimeUnit TimeUnit)
 {
 	TArray<int32> Values;
 	TWeakPtr<ISequencer> WeakSequencer = GetSequencerFromAsset(LevelSequence);
 
 	if (WeakSequencer.IsValid() && ControlRig)
 	{
+		if (FRigControlElement* Element = ControlRig->FindControl(ControlName))
+		{
+			if (Element->Settings.ControlType != ERigControlType::Integer)
+			{
+				UE_LOG(LogControlRig, Error, TEXT("Control Rig Wrong Type"));
+				return Values;
+			}
+		}
 		TArray<FRigControlValue> RigValues;
-		if (GetControlRigValues(WeakSequencer.Pin().Get(), ControlRig, ControlName, Frames, RigValues))
+		if (GetControlRigValues(WeakSequencer.Pin().Get(), ControlRig, ControlName, TimeUnit,Frames, RigValues))
 		{
 			Values.Reserve(RigValues.Num());
 			for (const FRigControlValue& RigValue : RigValues)
@@ -996,28 +1167,49 @@ TArray<int32> UControlRigSequencerEditorLibrary::GetLocalControlRigInts(ULevelSe
 	return Values;
 }
 
-void UControlRigSequencerEditorLibrary::SetLocalControlRigInt(ULevelSequence* LevelSequence, UControlRig* ControlRig, FName ControlName, FFrameNumber Frame, int32 Value, bool bSetKey)
+void UControlRigSequencerEditorLibrary::SetLocalControlRigInt(ULevelSequence* LevelSequence, UControlRig* ControlRig, FName ControlName, FFrameNumber Frame, int32 Value, ESequenceTimeUnit TimeUnit,bool bSetKey)
 {
 	if (LevelSequence == nullptr || ControlRig == nullptr)
 	{
 		return;
 	}
+	if (FRigControlElement* Element = ControlRig->FindControl(ControlName))
+	{
+		if (Element->Settings.ControlType != ERigControlType::Integer)
+		{
+			UE_LOG(LogControlRig, Error, TEXT("Control Rig Wrong Type"));
+			return;
+		}
+	}
 	if (UMovieScene* MovieScene = LevelSequence->GetMovieScene())
 	{
 		FFrameRate TickResolution = MovieScene->GetTickResolution();
+		if (TimeUnit == ESequenceTimeUnit::DisplayRate)
+		{
+			Frame = FFrameRate::TransformTime(FFrameTime(Frame, 0), MovieScene->GetDisplayRate(), MovieScene->GetTickResolution()).RoundToFrame();
+		}
 		FRigControlModifiedContext Context;
 		Context.SetKey = bSetKey ? EControlRigSetKey::Always : EControlRigSetKey::DoNotCare;
 		Context.LocalTime = TickResolution.AsSeconds(FFrameTime(Frame));
-		ControlRig->SetControlValue<int32>(ControlName, Value, true, FRigControlModifiedContext(EControlRigSetKey::Never, false));
+		ControlRig->SetControlValue<int32>(ControlName, Value, true, Context);
 	}
 }
 
-void UControlRigSequencerEditorLibrary::SetLocalControlRigInts(ULevelSequence* LevelSequence, UControlRig* ControlRig, FName ControlName, const TArray<FFrameNumber>& Frames, const TArray<int32> Values)
+void UControlRigSequencerEditorLibrary::SetLocalControlRigInts(ULevelSequence* LevelSequence, UControlRig* ControlRig, FName ControlName, const TArray<FFrameNumber>& Frames,
+	const TArray<int32> Values , ESequenceTimeUnit TimeUnit)
 {
 	if (LevelSequence == nullptr || ControlRig == nullptr || (Frames.Num() != Values.Num()))
 	{
 		return;
 	}
+	if (FRigControlElement* Element = ControlRig->FindControl(ControlName))
+	{
+		if (Element->Settings.ControlType != ERigControlType::Integer)
+		{
+			UE_LOG(LogControlRig, Error, TEXT("Control Rig Wrong Type"));
+			return;
+		}
+	}
 	if (UMovieScene* MovieScene = LevelSequence->GetMovieScene())
 	{
 		FFrameRate TickResolution = MovieScene->GetTickResolution();
@@ -1026,23 +1218,35 @@ void UControlRigSequencerEditorLibrary::SetLocalControlRigInts(ULevelSequence* L
 		for (int32 Index = 0; Index < Frames.Num(); ++Index)
 		{
 			FFrameNumber Frame = Frames[Index];
+			if (TimeUnit == ESequenceTimeUnit::DisplayRate)
+			{
+				Frame = FFrameRate::TransformTime(FFrameTime(Frame, 0), MovieScene->GetDisplayRate(), MovieScene->GetTickResolution()).RoundToFrame();
+			}
 			int32 Value = Values[Index];
 			Context.LocalTime = TickResolution.AsSeconds(FFrameTime(Frame));
-			ControlRig->SetControlValue<int32>(ControlName, Value, true, FRigControlModifiedContext(EControlRigSetKey::Never, false));
+			ControlRig->SetControlValue<int32>(ControlName, Value, true, Context);
 		}
 	}
 }
 
 
-FVector2D UControlRigSequencerEditorLibrary::GetLocalControlRigVector2D(ULevelSequence* LevelSequence, UControlRig* ControlRig, FName ControlName, FFrameNumber Frame)
+FVector2D UControlRigSequencerEditorLibrary::GetLocalControlRigVector2D(ULevelSequence* LevelSequence, UControlRig* ControlRig, FName ControlName, FFrameNumber Frame, ESequenceTimeUnit TimeUnit)
 {
 	FVector2D Value;
 	TWeakPtr<ISequencer> WeakSequencer = GetSequencerFromAsset(LevelSequence);
 
 	if (WeakSequencer.IsValid() && ControlRig)
 	{
+		if (FRigControlElement* Element = ControlRig->FindControl(ControlName))
+		{
+			if (Element->Settings.ControlType != ERigControlType::Vector2D)
+			{
+				UE_LOG(LogControlRig, Error, TEXT("Control Rig Wrong Type"));
+				return Value;
+			}
+		}
 		FRigControlValue RigValue;
-		if (GetControlRigValue(WeakSequencer.Pin().Get(), ControlRig, ControlName, Frame, RigValue))
+		if (GetControlRigValue(WeakSequencer.Pin().Get(), ControlRig, ControlName, TimeUnit, Frame, RigValue))
 		{
 			Value = RigValue.Get<FVector2D>();
 		}
@@ -1050,15 +1254,23 @@ FVector2D UControlRigSequencerEditorLibrary::GetLocalControlRigVector2D(ULevelSe
 	return Value;
 }
 
-TArray<FVector2D> UControlRigSequencerEditorLibrary::GetLocalControlRigVector2Ds(ULevelSequence* LevelSequence, UControlRig* ControlRig, FName ControlName, const TArray<FFrameNumber>& Frames)
+TArray<FVector2D> UControlRigSequencerEditorLibrary::GetLocalControlRigVector2Ds(ULevelSequence* LevelSequence, UControlRig* ControlRig, FName ControlName, const TArray<FFrameNumber>& Frames, ESequenceTimeUnit TimeUnit)
 {
 	TArray<FVector2D> Values;
 	TWeakPtr<ISequencer> WeakSequencer = GetSequencerFromAsset(LevelSequence);
 
 	if (WeakSequencer.IsValid() && ControlRig)
 	{
+		if (FRigControlElement* Element = ControlRig->FindControl(ControlName))
+		{
+			if (Element->Settings.ControlType != ERigControlType::Vector2D)
+			{
+				UE_LOG(LogControlRig, Error, TEXT("Control Rig Wrong Type"));
+				return Values;
+			}
+		}
 		TArray<FRigControlValue> RigValues;
-		if (GetControlRigValues(WeakSequencer.Pin().Get(), ControlRig, ControlName, Frames, RigValues))
+		if (GetControlRigValues(WeakSequencer.Pin().Get(), ControlRig, ControlName, TimeUnit, Frames, RigValues))
 		{
 			Values.Reserve(RigValues.Num());
 			for (const FRigControlValue& RigValue : RigValues)
@@ -1072,28 +1284,49 @@ TArray<FVector2D> UControlRigSequencerEditorLibrary::GetLocalControlRigVector2Ds
 
 }
 
-void UControlRigSequencerEditorLibrary::SetLocalControlRigVector2D(ULevelSequence* LevelSequence, UControlRig* ControlRig, FName ControlName, FFrameNumber Frame, FVector2D Value, bool bSetKey)
+void UControlRigSequencerEditorLibrary::SetLocalControlRigVector2D(ULevelSequence* LevelSequence, UControlRig* ControlRig, FName ControlName, FFrameNumber Frame, FVector2D Value, ESequenceTimeUnit TimeUnit,bool bSetKey)
 {
 	if (LevelSequence == nullptr || ControlRig == nullptr)
 	{
 		return;
 	}
+	if (FRigControlElement* Element = ControlRig->FindControl(ControlName))
+	{
+		if (Element->Settings.ControlType != ERigControlType::Vector2D)
+		{
+			UE_LOG(LogControlRig, Error, TEXT("Control Rig Wrong Type"));
+			return;
+		}
+	}
 	if (UMovieScene* MovieScene = LevelSequence->GetMovieScene())
 	{
 		FFrameRate TickResolution = MovieScene->GetTickResolution();
+		if (TimeUnit == ESequenceTimeUnit::DisplayRate)
+		{
+			Frame = FFrameRate::TransformTime(FFrameTime(Frame, 0), MovieScene->GetDisplayRate(), MovieScene->GetTickResolution()).RoundToFrame();
+		}
 		FRigControlModifiedContext Context;
 		Context.SetKey = bSetKey ? EControlRigSetKey::Always : EControlRigSetKey::DoNotCare;
 		Context.LocalTime = TickResolution.AsSeconds(FFrameTime(Frame));
-		ControlRig->SetControlValue<FVector2D>(ControlName, Value, true, FRigControlModifiedContext(EControlRigSetKey::Never, false));
+		ControlRig->SetControlValue<FVector2D>(ControlName, Value, true, Context);
 	}
 }
 
-void UControlRigSequencerEditorLibrary::SetLocalControlRigVector2Ds(ULevelSequence* LevelSequence, UControlRig* ControlRig, FName ControlName, const TArray<FFrameNumber>& Frames, const TArray<FVector2D> Values)
+void UControlRigSequencerEditorLibrary::SetLocalControlRigVector2Ds(ULevelSequence* LevelSequence, UControlRig* ControlRig, FName ControlName, const TArray<FFrameNumber>& Frames, 
+	const TArray<FVector2D> Values, ESequenceTimeUnit TimeUnit)
 {
 	if (LevelSequence == nullptr || ControlRig == nullptr || (Frames.Num() != Values.Num()))
 	{
 		return;
 	}
+	if (FRigControlElement* Element = ControlRig->FindControl(ControlName))
+	{
+		if (Element->Settings.ControlType != ERigControlType::Vector2D)
+		{
+			UE_LOG(LogControlRig, Error, TEXT("Control Rig Wrong Type"));
+			return;
+		}
+	}
 	if (UMovieScene* MovieScene = LevelSequence->GetMovieScene())
 	{
 		FFrameRate TickResolution = MovieScene->GetTickResolution();
@@ -1102,23 +1335,35 @@ void UControlRigSequencerEditorLibrary::SetLocalControlRigVector2Ds(ULevelSequen
 		for (int32 Index = 0; Index < Frames.Num(); ++Index)
 		{
 			FFrameNumber Frame = Frames[Index];
+			if (TimeUnit == ESequenceTimeUnit::DisplayRate)
+			{
+				Frame = FFrameRate::TransformTime(FFrameTime(Frame, 0), MovieScene->GetDisplayRate(), MovieScene->GetTickResolution()).RoundToFrame();
+			}
 			FVector2D Value = Values[Index];
 			Context.LocalTime = TickResolution.AsSeconds(FFrameTime(Frame));
-			ControlRig->SetControlValue<FVector2D>(ControlName, Value, true, FRigControlModifiedContext(EControlRigSetKey::Never, false));
+			ControlRig->SetControlValue<FVector2D>(ControlName, Value, true, Context);
 		}
 	}
 }
 
 
-FVector UControlRigSequencerEditorLibrary::GetLocalControlRigPosition(ULevelSequence* LevelSequence, UControlRig* ControlRig, FName ControlName, FFrameNumber Frame)
+FVector UControlRigSequencerEditorLibrary::GetLocalControlRigPosition(ULevelSequence* LevelSequence, UControlRig* ControlRig, FName ControlName, FFrameNumber Frame, ESequenceTimeUnit TimeUnit)
 {
 	FVector Value;
 	TWeakPtr<ISequencer> WeakSequencer = GetSequencerFromAsset(LevelSequence);
 
 	if (WeakSequencer.IsValid() && ControlRig)
 	{
+		if (FRigControlElement* Element = ControlRig->FindControl(ControlName))
+		{
+			if (Element->Settings.ControlType != ERigControlType::Position)
+			{
+				UE_LOG(LogControlRig, Error, TEXT("Control Rig Wrong Type"));
+				return Value;
+			}
+		}
 		FRigControlValue RigValue;
-		if (GetControlRigValue(WeakSequencer.Pin().Get(), ControlRig, ControlName, Frame, RigValue))
+		if (GetControlRigValue(WeakSequencer.Pin().Get(), ControlRig, ControlName, TimeUnit, Frame, RigValue))
 		{
 			Value = RigValue.Get<FVector>();
 		}
@@ -1126,15 +1371,23 @@ FVector UControlRigSequencerEditorLibrary::GetLocalControlRigPosition(ULevelSequ
 	return Value;
 }
 
-TArray<FVector> UControlRigSequencerEditorLibrary::GetLocalControlRigPositions(ULevelSequence* LevelSequence, UControlRig* ControlRig, FName ControlName, const TArray<FFrameNumber>& Frames)
+TArray<FVector> UControlRigSequencerEditorLibrary::GetLocalControlRigPositions(ULevelSequence* LevelSequence, UControlRig* ControlRig, FName ControlName, const TArray<FFrameNumber>& Frames, ESequenceTimeUnit TimeUnit)
 {
 	TArray<FVector> Values;
 	TWeakPtr<ISequencer> WeakSequencer = GetSequencerFromAsset(LevelSequence);
 
 	if (WeakSequencer.IsValid() && ControlRig)
 	{
+		if (FRigControlElement* Element = ControlRig->FindControl(ControlName))
+		{
+			if (Element->Settings.ControlType != ERigControlType::Position)
+			{
+				UE_LOG(LogControlRig, Error, TEXT("Control Rig Wrong Type"));
+				return Values;
+			}
+		}
 		TArray<FRigControlValue> RigValues;
-		if (GetControlRigValues(WeakSequencer.Pin().Get(), ControlRig, ControlName, Frames, RigValues))
+		if (GetControlRigValues(WeakSequencer.Pin().Get(), ControlRig, ControlName, TimeUnit, Frames, RigValues))
 		{
 			Values.Reserve(RigValues.Num());
 			for (const FRigControlValue& RigValue : RigValues)
@@ -1147,28 +1400,49 @@ TArray<FVector> UControlRigSequencerEditorLibrary::GetLocalControlRigPositions(U
 	return Values;
 }
 
-void UControlRigSequencerEditorLibrary::SetLocalControlRigPosition(ULevelSequence* LevelSequence, UControlRig* ControlRig, FName ControlName, FFrameNumber Frame, FVector Value, bool bSetKey)
+void UControlRigSequencerEditorLibrary::SetLocalControlRigPosition(ULevelSequence* LevelSequence, UControlRig* ControlRig, FName ControlName, FFrameNumber Frame, FVector Value, ESequenceTimeUnit TimeUnit,bool bSetKey)
 {
 	if (LevelSequence == nullptr || ControlRig == nullptr)
 	{
 		return;
 	}
+	if (FRigControlElement* Element = ControlRig->FindControl(ControlName))
+	{
+		if (Element->Settings.ControlType != ERigControlType::Position)
+		{
+			UE_LOG(LogControlRig, Error, TEXT("Control Rig Wrong Type"));
+			return;
+		}
+	}
 	if (UMovieScene* MovieScene = LevelSequence->GetMovieScene())
 	{
 		FFrameRate TickResolution = MovieScene->GetTickResolution();
+		if (TimeUnit == ESequenceTimeUnit::DisplayRate)
+		{
+			Frame = FFrameRate::TransformTime(FFrameTime(Frame, 0), MovieScene->GetDisplayRate(), MovieScene->GetTickResolution()).RoundToFrame();
+		}
 		FRigControlModifiedContext Context;
 		Context.SetKey = bSetKey ? EControlRigSetKey::Always : EControlRigSetKey::DoNotCare;
 		Context.LocalTime = TickResolution.AsSeconds(FFrameTime(Frame));
-		ControlRig->SetControlValue<FVector>(ControlName, Value, true, FRigControlModifiedContext(EControlRigSetKey::Never, false));
+		ControlRig->SetControlValue<FVector>(ControlName, Value, true, Context);
 	}
 }
 
-void UControlRigSequencerEditorLibrary::SetLocalControlRigPositions(ULevelSequence* LevelSequence, UControlRig* ControlRig, FName ControlName, const TArray<FFrameNumber>& Frames, const TArray<FVector> Values)
+void UControlRigSequencerEditorLibrary::SetLocalControlRigPositions(ULevelSequence* LevelSequence, UControlRig* ControlRig, FName ControlName, const TArray<FFrameNumber>& Frames, 
+	const TArray<FVector> Values, ESequenceTimeUnit TimeUnit)
 {
 	if (LevelSequence == nullptr || ControlRig == nullptr || (Frames.Num() != Values.Num()))
 	{
 		return;
 	}
+	if (FRigControlElement* Element = ControlRig->FindControl(ControlName))
+	{
+		if (Element->Settings.ControlType != ERigControlType::Position)
+		{
+			UE_LOG(LogControlRig, Error, TEXT("Control Rig Wrong Type"));
+			return;
+		}
+	}
 	if (UMovieScene* MovieScene = LevelSequence->GetMovieScene())
 	{
 		FFrameRate TickResolution = MovieScene->GetTickResolution();
@@ -1177,23 +1451,35 @@ void UControlRigSequencerEditorLibrary::SetLocalControlRigPositions(ULevelSequen
 		for (int32 Index = 0; Index < Frames.Num(); ++Index)
 		{
 			FFrameNumber Frame = Frames[Index];
+			if (TimeUnit == ESequenceTimeUnit::DisplayRate)
+			{
+				Frame = FFrameRate::TransformTime(FFrameTime(Frame, 0), MovieScene->GetDisplayRate(), MovieScene->GetTickResolution()).RoundToFrame();
+			}
 			FVector Value = Values[Index];
 			Context.LocalTime = TickResolution.AsSeconds(FFrameTime(Frame));
-			ControlRig->SetControlValue<FVector>(ControlName, Value, true, FRigControlModifiedContext(EControlRigSetKey::Never, false));
+			ControlRig->SetControlValue<FVector>(ControlName, Value, true, Context);
 		}
 	}
 }
 
 
-FRotator UControlRigSequencerEditorLibrary::GetLocalControlRigRotator(ULevelSequence* LevelSequence, UControlRig* ControlRig, FName ControlName, FFrameNumber Frame)
+FRotator UControlRigSequencerEditorLibrary::GetLocalControlRigRotator(ULevelSequence* LevelSequence, UControlRig* ControlRig, FName ControlName, FFrameNumber Frame, ESequenceTimeUnit TimeUnit)
 {
 	FRotator Value;
 	TWeakPtr<ISequencer> WeakSequencer = GetSequencerFromAsset(LevelSequence);
 
 	if (WeakSequencer.IsValid() && ControlRig)
 	{
+		if (FRigControlElement* Element = ControlRig->FindControl(ControlName))
+		{
+			if (Element->Settings.ControlType != ERigControlType::Rotator)
+			{
+				UE_LOG(LogControlRig, Error, TEXT("Control Rig Wrong Type"));
+				return Value;
+			}
+		}
 		FRigControlValue RigValue;
-		if (GetControlRigValue(WeakSequencer.Pin().Get(), ControlRig, ControlName, Frame, RigValue))
+		if (GetControlRigValue(WeakSequencer.Pin().Get(), ControlRig, ControlName,TimeUnit, Frame, RigValue))
 		{
 			Value = RigValue.Get<FRotator>();
 		}
@@ -1201,15 +1487,23 @@ FRotator UControlRigSequencerEditorLibrary::GetLocalControlRigRotator(ULevelSequ
 	return Value;
 }
 
-TArray<FRotator> UControlRigSequencerEditorLibrary::GetLocalControlRigRotators(ULevelSequence* LevelSequence, UControlRig* ControlRig, FName ControlName, const TArray<FFrameNumber>& Frames)
+TArray<FRotator> UControlRigSequencerEditorLibrary::GetLocalControlRigRotators(ULevelSequence* LevelSequence, UControlRig* ControlRig, FName ControlName, const TArray<FFrameNumber>& Frames, ESequenceTimeUnit TimeUnit)
 {
 	TArray<FRotator> Values;
 	TWeakPtr<ISequencer> WeakSequencer = GetSequencerFromAsset(LevelSequence);
 
 	if (WeakSequencer.IsValid() && ControlRig)
 	{
+		if (FRigControlElement* Element = ControlRig->FindControl(ControlName))
+		{
+			if (Element->Settings.ControlType != ERigControlType::Rotator)
+			{
+				UE_LOG(LogControlRig, Error, TEXT("Control Rig Wrong Type"));
+				return Values;
+			}
+		}
 		TArray<FRigControlValue> RigValues;
-		if (GetControlRigValues(WeakSequencer.Pin().Get(), ControlRig, ControlName, Frames, RigValues))
+		if (GetControlRigValues(WeakSequencer.Pin().Get(), ControlRig, ControlName, TimeUnit, Frames, RigValues))
 		{
 			Values.Reserve(RigValues.Num());
 			for (const FRigControlValue& RigValue : RigValues)
@@ -1222,28 +1516,49 @@ TArray<FRotator> UControlRigSequencerEditorLibrary::GetLocalControlRigRotators(U
 	return Values;
 }
 
-void UControlRigSequencerEditorLibrary::SetLocalControlRigRotator(ULevelSequence* LevelSequence, UControlRig* ControlRig, FName ControlName, FFrameNumber Frame, FRotator Value, bool bSetKey)
+void UControlRigSequencerEditorLibrary::SetLocalControlRigRotator(ULevelSequence* LevelSequence, UControlRig* ControlRig, FName ControlName, FFrameNumber Frame, FRotator Value, ESequenceTimeUnit TimeUnit, bool bSetKey)
 {
 	if (LevelSequence == nullptr || ControlRig == nullptr)
 	{
 		return;
 	}
+	if (FRigControlElement* Element = ControlRig->FindControl(ControlName))
+	{
+		if (Element->Settings.ControlType != ERigControlType::Rotator)
+		{
+			UE_LOG(LogControlRig, Error, TEXT("Control Rig Wrong Type"));
+			return;
+		}
+	}
 	if (UMovieScene* MovieScene = LevelSequence->GetMovieScene())
 	{
 		FFrameRate TickResolution = MovieScene->GetTickResolution();
+		if (TimeUnit == ESequenceTimeUnit::DisplayRate)
+		{
+			Frame = FFrameRate::TransformTime(FFrameTime(Frame, 0), MovieScene->GetDisplayRate(), MovieScene->GetTickResolution()).RoundToFrame();
+		}
 		FRigControlModifiedContext Context;
 		Context.SetKey = bSetKey ? EControlRigSetKey::Always : EControlRigSetKey::DoNotCare;
 		Context.LocalTime = TickResolution.AsSeconds(FFrameTime(Frame));
-		ControlRig->SetControlValue<FRotator>(ControlName, Value, true, FRigControlModifiedContext(EControlRigSetKey::Never, false));
+		ControlRig->SetControlValue<FRotator>(ControlName, Value, true, Context);
 	}
 }
 
-void UControlRigSequencerEditorLibrary::SetLocalControlRigRotators(ULevelSequence* LevelSequence, UControlRig* ControlRig, FName ControlName, const TArray<FFrameNumber>& Frames, const TArray<FRotator> Values)
+void UControlRigSequencerEditorLibrary::SetLocalControlRigRotators(ULevelSequence* LevelSequence, UControlRig* ControlRig, FName ControlName, const TArray<FFrameNumber>& Frames,
+	const TArray<FRotator> Values, ESequenceTimeUnit TimeUnit)
 {
 	if (LevelSequence == nullptr || ControlRig == nullptr || (Frames.Num() != Values.Num()))
 	{
 		return;
 	}
+	if (FRigControlElement* Element = ControlRig->FindControl(ControlName))
+	{
+		if (Element->Settings.ControlType != ERigControlType::Rotator)
+		{
+			UE_LOG(LogControlRig, Error, TEXT("Control Rig Wrong Type"));
+			return;
+		}
+	}
 	if (UMovieScene* MovieScene = LevelSequence->GetMovieScene())
 	{
 		FFrameRate TickResolution = MovieScene->GetTickResolution();
@@ -1252,23 +1567,35 @@ void UControlRigSequencerEditorLibrary::SetLocalControlRigRotators(ULevelSequenc
 		for (int32 Index = 0; Index < Frames.Num(); ++Index)
 		{
 			FFrameNumber Frame = Frames[Index];
+			if (TimeUnit == ESequenceTimeUnit::DisplayRate)
+			{
+				Frame = FFrameRate::TransformTime(FFrameTime(Frame, 0), MovieScene->GetDisplayRate(), MovieScene->GetTickResolution()).RoundToFrame();
+			}
 			FRotator Value = Values[Index];
 			Context.LocalTime = TickResolution.AsSeconds(FFrameTime(Frame));
-			ControlRig->SetControlValue<FRotator>(ControlName, Value, true, FRigControlModifiedContext(EControlRigSetKey::Never, false));
+			ControlRig->SetControlValue<FRotator>(ControlName, Value, true, Context);
 		}
 	}
 }
 
 
-FVector UControlRigSequencerEditorLibrary::GetLocalControlRigScale(ULevelSequence* LevelSequence, UControlRig* ControlRig, FName ControlName, FFrameNumber Frame)
+FVector UControlRigSequencerEditorLibrary::GetLocalControlRigScale(ULevelSequence* LevelSequence, UControlRig* ControlRig, FName ControlName, FFrameNumber Frame,  ESequenceTimeUnit TimeUnit)
 {
 	FVector Value;
 	TWeakPtr<ISequencer> WeakSequencer = GetSequencerFromAsset(LevelSequence);
 
 	if (WeakSequencer.IsValid() && ControlRig)
 	{
+		if (FRigControlElement* Element = ControlRig->FindControl(ControlName))
+		{
+			if (Element->Settings.ControlType != ERigControlType::Scale)
+			{
+				UE_LOG(LogControlRig, Error, TEXT("Control Rig Wrong Type"));
+				return Value;
+			}
+		}
 		FRigControlValue RigValue;
-		if (GetControlRigValue(WeakSequencer.Pin().Get(), ControlRig, ControlName, Frame, RigValue))
+		if (GetControlRigValue(WeakSequencer.Pin().Get(), ControlRig, ControlName,TimeUnit, Frame, RigValue))
 		{
 			Value = RigValue.Get<FVector>();
 		}
@@ -1276,15 +1603,23 @@ FVector UControlRigSequencerEditorLibrary::GetLocalControlRigScale(ULevelSequenc
 	return Value;
 }
 
-TArray<FVector>UControlRigSequencerEditorLibrary::GetLocalControlRigScales(ULevelSequence* LevelSequence, UControlRig* ControlRig, FName ControlName, const TArray<FFrameNumber>& Frames)
+TArray<FVector>UControlRigSequencerEditorLibrary::GetLocalControlRigScales(ULevelSequence* LevelSequence, UControlRig* ControlRig, FName ControlName, const TArray<FFrameNumber>& Frames, ESequenceTimeUnit TimeUnit)
 {
 	TArray<FVector> Values;
 	TWeakPtr<ISequencer> WeakSequencer = GetSequencerFromAsset(LevelSequence);
 
 	if (WeakSequencer.IsValid() && ControlRig)
 	{
+		if (FRigControlElement* Element = ControlRig->FindControl(ControlName))
+		{
+			if (Element->Settings.ControlType != ERigControlType::Scale)
+			{
+				UE_LOG(LogControlRig, Error, TEXT("Control Rig Wrong Type"));
+				return Values;
+			}
+		}
 		TArray<FRigControlValue> RigValues;
-		if (GetControlRigValues(WeakSequencer.Pin().Get(), ControlRig, ControlName, Frames, RigValues))
+		if (GetControlRigValues(WeakSequencer.Pin().Get(), ControlRig, ControlName,TimeUnit, Frames, RigValues))
 		{
 			Values.Reserve(RigValues.Num());
 			for (const FRigControlValue& RigValue : RigValues)
@@ -1297,28 +1632,49 @@ TArray<FVector>UControlRigSequencerEditorLibrary::GetLocalControlRigScales(ULeve
 	return Values;
 }
 
-void UControlRigSequencerEditorLibrary::SetLocalControlRigScale(ULevelSequence* LevelSequence, UControlRig* ControlRig, FName ControlName, FFrameNumber Frame, FVector Value, bool bSetKey)
+void UControlRigSequencerEditorLibrary::SetLocalControlRigScale(ULevelSequence* LevelSequence, UControlRig* ControlRig, FName ControlName, FFrameNumber Frame, FVector Value, ESequenceTimeUnit TimeUnit, bool bSetKey)
 {
 	if (LevelSequence == nullptr || ControlRig == nullptr)
 	{
 		return;
 	}
+	if (FRigControlElement* Element = ControlRig->FindControl(ControlName))
+	{
+		if (Element->Settings.ControlType != ERigControlType::Scale)
+		{
+			UE_LOG(LogControlRig, Error, TEXT("Control Rig Wrong Type"));
+			return;
+		}
+	}
 	if (UMovieScene* MovieScene = LevelSequence->GetMovieScene())
 	{
 		FFrameRate TickResolution = MovieScene->GetTickResolution();
+		if (TimeUnit == ESequenceTimeUnit::DisplayRate)
+		{
+			Frame = FFrameRate::TransformTime(FFrameTime(Frame, 0), MovieScene->GetDisplayRate(), MovieScene->GetTickResolution()).RoundToFrame();
+		}
 		FRigControlModifiedContext Context;
 		Context.SetKey = bSetKey ? EControlRigSetKey::Always : EControlRigSetKey::DoNotCare;
 		Context.LocalTime = TickResolution.AsSeconds(FFrameTime(Frame));
-		ControlRig->SetControlValue<FVector>(ControlName, Value, true, FRigControlModifiedContext(EControlRigSetKey::Never, false));
+		ControlRig->SetControlValue<FVector>(ControlName, Value, true, Context);
 	}
 }
 
-void UControlRigSequencerEditorLibrary::SetLocalControlRigScales(ULevelSequence* LevelSequence, UControlRig* ControlRig, FName ControlName, const TArray<FFrameNumber>& Frames, const TArray<FVector> Values)
+void UControlRigSequencerEditorLibrary::SetLocalControlRigScales(ULevelSequence* LevelSequence, UControlRig* ControlRig, FName ControlName, const TArray<FFrameNumber>& Frames, 
+	const TArray<FVector> Values, ESequenceTimeUnit TimeUnit)
 {
 	if (LevelSequence == nullptr || ControlRig == nullptr || (Frames.Num() != Values.Num()))
 	{
 		return;
 	}
+	if (FRigControlElement* Element = ControlRig->FindControl(ControlName))
+	{
+		if (Element->Settings.ControlType != ERigControlType::Scale)
+		{
+			UE_LOG(LogControlRig, Error, TEXT("Control Rig Wrong Type"));
+			return;
+		}
+	}
 	if (UMovieScene* MovieScene = LevelSequence->GetMovieScene())
 	{
 		FFrameRate TickResolution = MovieScene->GetTickResolution();
@@ -1327,23 +1683,35 @@ void UControlRigSequencerEditorLibrary::SetLocalControlRigScales(ULevelSequence*
 		for (int32 Index = 0; Index < Frames.Num(); ++Index)
 		{
 			FFrameNumber Frame = Frames[Index];
+			if (TimeUnit == ESequenceTimeUnit::DisplayRate)
+			{
+				Frame = FFrameRate::TransformTime(FFrameTime(Frame, 0), MovieScene->GetDisplayRate(), MovieScene->GetTickResolution()).RoundToFrame();
+			}
 			FVector Value = Values[Index];
 			Context.LocalTime = TickResolution.AsSeconds(FFrameTime(Frame));
-			ControlRig->SetControlValue<FVector>(ControlName, Value, true, FRigControlModifiedContext(EControlRigSetKey::Never, false));
+			ControlRig->SetControlValue<FVector>(ControlName, Value, true, Context);
 		}
 	}
 }
 
 
-FEulerTransform UControlRigSequencerEditorLibrary::GetLocalControlRigEulerTransform(ULevelSequence* LevelSequence, UControlRig* ControlRig, FName ControlName, FFrameNumber Frame)
+FEulerTransform UControlRigSequencerEditorLibrary::GetLocalControlRigEulerTransform(ULevelSequence* LevelSequence, UControlRig* ControlRig, FName ControlName, FFrameNumber Frame, ESequenceTimeUnit TimeUnit)
 {
 	FEulerTransform Value = FEulerTransform::Identity;
 	TWeakPtr<ISequencer> WeakSequencer = GetSequencerFromAsset(LevelSequence);
 
 	if (WeakSequencer.IsValid() && ControlRig)
 	{
+		if (FRigControlElement* Element = ControlRig->FindControl(ControlName))
+		{
+			if (Element->Settings.ControlType != ERigControlType::EulerTransform)
+			{
+				UE_LOG(LogControlRig, Error, TEXT("Control Rig Wrong Type"));
+				return Value;
+			}
+		}
 		FRigControlValue RigValue;
-		if (GetControlRigValue(WeakSequencer.Pin().Get(), ControlRig, ControlName, Frame, RigValue))
+		if (GetControlRigValue(WeakSequencer.Pin().Get(), ControlRig, ControlName, TimeUnit, Frame, RigValue))
 		{
 			Value = RigValue.Get<FEulerTransform>();
 		}
@@ -1351,15 +1719,23 @@ FEulerTransform UControlRigSequencerEditorLibrary::GetLocalControlRigEulerTransf
 	return Value;
 }
 
-TArray<FEulerTransform> UControlRigSequencerEditorLibrary::GetLocalControlRigEulerTransforms(ULevelSequence* LevelSequence, UControlRig* ControlRig, FName ControlName, const TArray<FFrameNumber>& Frames)
+TArray<FEulerTransform> UControlRigSequencerEditorLibrary::GetLocalControlRigEulerTransforms(ULevelSequence* LevelSequence, UControlRig* ControlRig, FName ControlName, const TArray<FFrameNumber>& Frames, ESequenceTimeUnit TimeUnit)
 {
 	TArray<FEulerTransform> Values;
 	TWeakPtr<ISequencer> WeakSequencer = GetSequencerFromAsset(LevelSequence);
 
 	if (WeakSequencer.IsValid() && ControlRig)
 	{
+		if (FRigControlElement* Element = ControlRig->FindControl(ControlName))
+		{
+			if (Element->Settings.ControlType != ERigControlType::EulerTransform)
+			{
+				UE_LOG(LogControlRig, Error, TEXT("Control Rig Wrong Type"));
+				return Values;
+			}
+		}
 		TArray<FRigControlValue> RigValues;
-		if (GetControlRigValues(WeakSequencer.Pin().Get(), ControlRig, ControlName, Frames, RigValues))
+		if (GetControlRigValues(WeakSequencer.Pin().Get(), ControlRig, ControlName, TimeUnit, Frames, RigValues))
 		{
 			Values.Reserve(RigValues.Num());
 			for (const FRigControlValue& RigValue : RigValues)
@@ -1372,28 +1748,50 @@ TArray<FEulerTransform> UControlRigSequencerEditorLibrary::GetLocalControlRigEul
 	return Values;
 }
 
-void UControlRigSequencerEditorLibrary::SetLocalControlRigEulerTransform(ULevelSequence* LevelSequence, UControlRig* ControlRig, FName ControlName, FFrameNumber Frame, FEulerTransform Value, bool bSetKey)
+void UControlRigSequencerEditorLibrary::SetLocalControlRigEulerTransform(ULevelSequence* LevelSequence, UControlRig* ControlRig, FName ControlName, FFrameNumber Frame, FEulerTransform Value,
+	ESequenceTimeUnit TimeUnit, bool bSetKey)
 {
 	if (LevelSequence == nullptr || ControlRig == nullptr)
 	{
 		return;
 	}
+	if (FRigControlElement* Element = ControlRig->FindControl(ControlName))
+	{
+		if (Element->Settings.ControlType != ERigControlType::EulerTransform)
+		{
+			UE_LOG(LogControlRig, Error, TEXT("Control Rig Wrong Type"));
+			return;
+		}
+	}
 	if (UMovieScene* MovieScene = LevelSequence->GetMovieScene())
 	{
 		FFrameRate TickResolution = MovieScene->GetTickResolution();
+		if (TimeUnit == ESequenceTimeUnit::DisplayRate)
+		{
+			Frame = FFrameRate::TransformTime(FFrameTime(Frame, 0), MovieScene->GetDisplayRate(), MovieScene->GetTickResolution()).RoundToFrame();
+		}
 		FRigControlModifiedContext Context;
 		Context.SetKey = bSetKey ? EControlRigSetKey::Always : EControlRigSetKey::DoNotCare;
 		Context.LocalTime = TickResolution.AsSeconds(FFrameTime(Frame));
-		ControlRig->SetControlValue<FEulerTransform>(ControlName, Value, true, FRigControlModifiedContext(EControlRigSetKey::Never, false));
+		ControlRig->SetControlValue<FEulerTransform>(ControlName, Value, true, Context);
 	}
 }
 
-void UControlRigSequencerEditorLibrary::SetLocalControlRigEulerTransforms(ULevelSequence* LevelSequence, UControlRig* ControlRig, FName ControlName, const TArray<FFrameNumber>& Frames, const TArray<FEulerTransform> Values)
+void UControlRigSequencerEditorLibrary::SetLocalControlRigEulerTransforms(ULevelSequence* LevelSequence, UControlRig* ControlRig, FName ControlName, const TArray<FFrameNumber>& Frames,
+	const TArray<FEulerTransform> Values, ESequenceTimeUnit TimeUnit)
 {
 	if (LevelSequence == nullptr || ControlRig == nullptr || (Frames.Num() != Values.Num()))
 	{
 		return;
 	}
+	if (FRigControlElement* Element = ControlRig->FindControl(ControlName))
+	{
+		if (Element->Settings.ControlType != ERigControlType::EulerTransform)
+		{
+			UE_LOG(LogControlRig, Error, TEXT("Control Rig Wrong Type"));
+			return;
+		}
+	}
 	if (UMovieScene* MovieScene = LevelSequence->GetMovieScene())
 	{
 		FFrameRate TickResolution = MovieScene->GetTickResolution();
@@ -1402,23 +1800,35 @@ void UControlRigSequencerEditorLibrary::SetLocalControlRigEulerTransforms(ULevel
 		for (int32 Index = 0; Index < Frames.Num(); ++Index)
 		{
 			FFrameNumber Frame = Frames[Index];
+			if (TimeUnit == ESequenceTimeUnit::DisplayRate)
+			{
+				Frame = FFrameRate::TransformTime(FFrameTime(Frame, 0), MovieScene->GetDisplayRate(), MovieScene->GetTickResolution()).RoundToFrame();
+			}
 			FEulerTransform Value = Values[Index];
 			Context.LocalTime = TickResolution.AsSeconds(FFrameTime(Frame));
-			ControlRig->SetControlValue<FEulerTransform>(ControlName, Value, true, FRigControlModifiedContext(EControlRigSetKey::Never, false));
+			ControlRig->SetControlValue<FEulerTransform>(ControlName, Value, true, Context);
 		}
 	}
 }
 
 
-FTransformNoScale UControlRigSequencerEditorLibrary::GetLocalControlRigTransformNoScale(ULevelSequence* LevelSequence, UControlRig* ControlRig, FName ControlName, FFrameNumber Frame)
+FTransformNoScale UControlRigSequencerEditorLibrary::GetLocalControlRigTransformNoScale(ULevelSequence* LevelSequence, UControlRig* ControlRig, FName ControlName, FFrameNumber Frame, ESequenceTimeUnit TimeUnit)
 {
 	FTransformNoScale Value = FTransformNoScale::Identity;
 	TWeakPtr<ISequencer> WeakSequencer = GetSequencerFromAsset(LevelSequence);
 
 	if (WeakSequencer.IsValid() && ControlRig)
 	{
+		if (FRigControlElement* Element = ControlRig->FindControl(ControlName))
+		{
+			if (Element->Settings.ControlType != ERigControlType::TransformNoScale)
+			{
+				UE_LOG(LogControlRig, Error, TEXT("Control Rig Wrong Type"));
+				return Value;
+			}
+		}
 		FRigControlValue RigValue;
-		if (GetControlRigValue(WeakSequencer.Pin().Get(), ControlRig, ControlName, Frame, RigValue))
+		if (GetControlRigValue(WeakSequencer.Pin().Get(), ControlRig, ControlName,TimeUnit, Frame, RigValue))
 		{
 			Value = RigValue.Get<FTransformNoScale>();
 		}
@@ -1426,15 +1836,23 @@ FTransformNoScale UControlRigSequencerEditorLibrary::GetLocalControlRigTransform
 	return Value;
 }
 
-TArray<FTransformNoScale> UControlRigSequencerEditorLibrary::GetLocalControlRigTransformNoScales(ULevelSequence* LevelSequence, UControlRig* ControlRig, FName ControlName, const TArray<FFrameNumber>& Frames)
+TArray<FTransformNoScale> UControlRigSequencerEditorLibrary::GetLocalControlRigTransformNoScales(ULevelSequence* LevelSequence, UControlRig* ControlRig, FName ControlName, const TArray<FFrameNumber>& Frames, ESequenceTimeUnit TimeUnit)
 {
 	TArray<FTransformNoScale> Values;
 	TWeakPtr<ISequencer> WeakSequencer = GetSequencerFromAsset(LevelSequence);
 
 	if (WeakSequencer.IsValid() && ControlRig)
 	{
+		if (FRigControlElement* Element = ControlRig->FindControl(ControlName))
+		{
+			if (Element->Settings.ControlType != ERigControlType::TransformNoScale)
+			{
+				UE_LOG(LogControlRig, Error, TEXT("Control Rig Wrong Type"));
+				return Values;
+			}
+		}
 		TArray<FRigControlValue> RigValues;
-		if (GetControlRigValues(WeakSequencer.Pin().Get(), ControlRig, ControlName, Frames, RigValues))
+		if (GetControlRigValues(WeakSequencer.Pin().Get(), ControlRig, ControlName,TimeUnit, Frames, RigValues))
 		{
 			Values.Reserve(RigValues.Num());
 			for (const FRigControlValue& RigValue : RigValues)
@@ -1447,28 +1865,50 @@ TArray<FTransformNoScale> UControlRigSequencerEditorLibrary::GetLocalControlRigT
 	return Values;
 }
 
-void UControlRigSequencerEditorLibrary::SetLocalControlRigTransformNoScale(ULevelSequence* LevelSequence, UControlRig* ControlRig, FName ControlName, FFrameNumber Frame, FTransformNoScale Value, bool bSetKey)
+void UControlRigSequencerEditorLibrary::SetLocalControlRigTransformNoScale(ULevelSequence* LevelSequence, UControlRig* ControlRig, FName ControlName, FFrameNumber Frame, FTransformNoScale Value,
+	ESequenceTimeUnit TimeUnit,bool bSetKey)
 {
 	if (LevelSequence == nullptr || ControlRig == nullptr)
 	{
 		return;
 	}
+	if (FRigControlElement* Element = ControlRig->FindControl(ControlName))
+	{
+		if (Element->Settings.ControlType != ERigControlType::TransformNoScale)
+		{
+			UE_LOG(LogControlRig, Error, TEXT("Control Rig Wrong Type"));
+			return;
+		}
+	}
 	if (UMovieScene* MovieScene = LevelSequence->GetMovieScene())
 	{
 		FFrameRate TickResolution = MovieScene->GetTickResolution();
+		if (TimeUnit == ESequenceTimeUnit::DisplayRate)
+		{
+			Frame = FFrameRate::TransformTime(FFrameTime(Frame, 0), MovieScene->GetDisplayRate(), MovieScene->GetTickResolution()).RoundToFrame();
+		}
 		FRigControlModifiedContext Context;
 		Context.SetKey = bSetKey ? EControlRigSetKey::Always : EControlRigSetKey::DoNotCare;
 		Context.LocalTime = TickResolution.AsSeconds(FFrameTime(Frame));
-		ControlRig->SetControlValue<FTransformNoScale>(ControlName, Value, true, FRigControlModifiedContext(EControlRigSetKey::Never, false));
+		ControlRig->SetControlValue<FTransformNoScale>(ControlName, Value, true, Context);
 	}
 }
 
-void UControlRigSequencerEditorLibrary::SetLocalControlRigTransformNoScales(ULevelSequence* LevelSequence, UControlRig* ControlRig, FName ControlName, const TArray<FFrameNumber>& Frames, const TArray<FTransformNoScale> Values)
+void UControlRigSequencerEditorLibrary::SetLocalControlRigTransformNoScales(ULevelSequence* LevelSequence, UControlRig* ControlRig, FName ControlName, const TArray<FFrameNumber>& Frames, 
+	const TArray<FTransformNoScale> Values, ESequenceTimeUnit TimeUnit)
 {
 	if (LevelSequence == nullptr || ControlRig == nullptr || (Frames.Num() != Values.Num()))
 	{
 		return;
 	}
+	if (FRigControlElement* Element = ControlRig->FindControl(ControlName))
+	{
+		if (Element->Settings.ControlType != ERigControlType::TransformNoScale)
+		{
+			UE_LOG(LogControlRig, Error, TEXT("Control Rig Wrong Type"));
+			return;
+		}
+	}
 	if (UMovieScene* MovieScene = LevelSequence->GetMovieScene())
 	{
 		FFrameRate TickResolution = MovieScene->GetTickResolution();
@@ -1477,22 +1917,35 @@ void UControlRigSequencerEditorLibrary::SetLocalControlRigTransformNoScales(ULev
 		for (int32 Index = 0; Index < Frames.Num(); ++Index)
 		{
 			FFrameNumber Frame = Frames[Index];
+			if (TimeUnit == ESequenceTimeUnit::DisplayRate)
+			{
+				Frame = FFrameRate::TransformTime(FFrameTime(Frame, 0), MovieScene->GetDisplayRate(), MovieScene->GetTickResolution()).RoundToFrame();
+			}
 			FTransformNoScale Value = Values[Index];
 			Context.LocalTime = TickResolution.AsSeconds(FFrameTime(Frame));
-			ControlRig->SetControlValue<FTransformNoScale>(ControlName, Value, true, FRigControlModifiedContext(EControlRigSetKey::Never, false));
+			ControlRig->SetControlValue<FTransformNoScale>(ControlName, Value, true, Context);
 		}
 	}
 }
 
 
-FTransform UControlRigSequencerEditorLibrary::GetLocalControlRigTransform(ULevelSequence* LevelSequence, UControlRig* ControlRig, FName ControlName, FFrameNumber Frame)
+FTransform UControlRigSequencerEditorLibrary::GetLocalControlRigTransform(ULevelSequence* LevelSequence, UControlRig* ControlRig, FName ControlName, FFrameNumber Frame, ESequenceTimeUnit TimeUnit)
 {
 	FTransform Value = FTransform::Identity;
 	TWeakPtr<ISequencer> WeakSequencer = GetSequencerFromAsset(LevelSequence);
 	if (WeakSequencer.IsValid() && ControlRig)
 	{
+		if (FRigControlElement* Element = ControlRig->FindControl(ControlName))
+		{
+			if (Element->Settings.ControlType != ERigControlType::Transform)
+			{
+				UE_LOG(LogControlRig, Error, TEXT("Control Rig Wrong Type"));
+				return Value;
+			}
+		}
+		
 		FRigControlValue RigValue;
-		if (GetControlRigValue(WeakSequencer.Pin().Get(), ControlRig, ControlName, Frame, RigValue))
+		if (GetControlRigValue(WeakSequencer.Pin().Get(), ControlRig, ControlName, TimeUnit, Frame, RigValue))
 		{
 			Value = RigValue.Get<FTransform>();
 		}
@@ -1500,15 +1953,23 @@ FTransform UControlRigSequencerEditorLibrary::GetLocalControlRigTransform(ULevel
 	return Value;
 }
 
-TArray<FTransform> UControlRigSequencerEditorLibrary::GetLocalControlRigTransforms(ULevelSequence* LevelSequence, UControlRig* ControlRig, FName ControlName, const TArray<FFrameNumber>& Frames)
+TArray<FTransform> UControlRigSequencerEditorLibrary::GetLocalControlRigTransforms(ULevelSequence* LevelSequence, UControlRig* ControlRig, FName ControlName, const TArray<FFrameNumber>& Frames, ESequenceTimeUnit TimeUnit)
 {
 	TArray<FTransform> Values;
 	TWeakPtr<ISequencer> WeakSequencer = GetSequencerFromAsset(LevelSequence);
 
 	if (WeakSequencer.IsValid() && ControlRig)
 	{
+		if (FRigControlElement* Element = ControlRig->FindControl(ControlName))
+		{
+			if (Element->Settings.ControlType != ERigControlType::Transform)
+			{
+				UE_LOG(LogControlRig, Error, TEXT("Control Rig Wrong Type"));
+				return Values;
+			}
+		}
 		TArray<FRigControlValue> RigValues;
-		if (GetControlRigValues(WeakSequencer.Pin().Get(), ControlRig, ControlName, Frames, RigValues))
+		if (GetControlRigValues(WeakSequencer.Pin().Get(), ControlRig, ControlName, TimeUnit, Frames, RigValues))
 		{
 			Values.Reserve(RigValues.Num());
 			for (const FRigControlValue& RigValue : RigValues)
@@ -1521,28 +1982,50 @@ TArray<FTransform> UControlRigSequencerEditorLibrary::GetLocalControlRigTransfor
 	return Values;
 }
 
-void UControlRigSequencerEditorLibrary::SetLocalControlRigTransform(ULevelSequence* LevelSequence, UControlRig* ControlRig, FName ControlName, FFrameNumber Frame, FTransform Value, bool bSetKey)
+void UControlRigSequencerEditorLibrary::SetLocalControlRigTransform(ULevelSequence* LevelSequence, UControlRig* ControlRig, FName ControlName, FFrameNumber Frame, FTransform Value,
+	ESequenceTimeUnit TimeUnit,bool bSetKey)
 {
 	if (LevelSequence == nullptr || ControlRig == nullptr)
 	{
 		return;
 	}
+	if (FRigControlElement* Element = ControlRig->FindControl(ControlName))
+	{
+		if (Element->Settings.ControlType != ERigControlType::Transform)
+		{
+			UE_LOG(LogControlRig, Error, TEXT("Control Rig Wrong Type"));
+			return;
+		}
+	}
 	if (UMovieScene* MovieScene = LevelSequence->GetMovieScene())
 	{
 		FFrameRate TickResolution = MovieScene->GetTickResolution();
+		if (TimeUnit == ESequenceTimeUnit::DisplayRate)
+		{
+			Frame = FFrameRate::TransformTime(FFrameTime(Frame, 0), MovieScene->GetDisplayRate(), MovieScene->GetTickResolution()).RoundToFrame();
+		}
 		FRigControlModifiedContext Context;
 		Context.SetKey = bSetKey ? EControlRigSetKey::Always : EControlRigSetKey::DoNotCare;
 		Context.LocalTime = TickResolution.AsSeconds(FFrameTime(Frame));
-		ControlRig->SetControlValue<FTransform>(ControlName, Value, true, FRigControlModifiedContext(EControlRigSetKey::Never, false));
+		ControlRig->SetControlValue<FTransform>(ControlName, Value, true, Context);
 	}
 }
 
-void UControlRigSequencerEditorLibrary::SetLocalControlRigTransforms(ULevelSequence* LevelSequence, UControlRig* ControlRig, FName ControlName, const TArray<FFrameNumber>& Frames, const TArray<FTransform> Values)
+void UControlRigSequencerEditorLibrary::SetLocalControlRigTransforms(ULevelSequence* LevelSequence, UControlRig* ControlRig, FName ControlName, const TArray<FFrameNumber>& Frames, 
+	const TArray<FTransform> Values, ESequenceTimeUnit TimeUnit)
 {
 	if (LevelSequence == nullptr || ControlRig == nullptr || (Frames.Num() != Values.Num()))
 	{
 		return;
 	}
+	if (FRigControlElement* Element = ControlRig->FindControl(ControlName))
+	{
+		if (Element->Settings.ControlType != ERigControlType::Transform)
+		{
+			UE_LOG(LogControlRig, Error, TEXT("Control Rig Wrong Type"));
+			return;
+		}
+	}
 	if (UMovieScene* MovieScene = LevelSequence->GetMovieScene())
 	{
 		FFrameRate TickResolution = MovieScene->GetTickResolution();
@@ -1551,9 +2034,13 @@ void UControlRigSequencerEditorLibrary::SetLocalControlRigTransforms(ULevelSeque
 		for (int32 Index = 0; Index < Frames.Num(); ++Index)
 		{
 			FFrameNumber Frame = Frames[Index];
+			if (TimeUnit == ESequenceTimeUnit::DisplayRate)
+			{
+				Frame = FFrameRate::TransformTime(FFrameTime(Frame, 0), MovieScene->GetDisplayRate(), MovieScene->GetTickResolution()).RoundToFrame();
+			}
 			FTransform Value = Values[Index];
 			Context.LocalTime = TickResolution.AsSeconds(FFrameTime(Frame));
-			ControlRig->SetControlValue<FTransform>(ControlName, Value, true, FRigControlModifiedContext(EControlRigSetKey::Never, false));
+			ControlRig->SetControlValue<FTransform>(ControlName, Value, true, Context);
 		}
 	}
 }
diff --git a/Engine/Plugins/Experimental/ControlRig/Source/ControlRigEditor/Private/Tools/ControlRigSnapper.cpp b/Engine/Plugins/Experimental/ControlRig/Source/ControlRigEditor/Private/Tools/ControlRigSnapper.cpp
--- a/Engine/Plugins/Experimental/ControlRig/Source/ControlRigEditor/Private/Tools/ControlRigSnapper.cpp
+++ b/Engine/Plugins/Experimental/ControlRig/Source/ControlRigEditor/Private/Tools/ControlRigSnapper.cpp
@@ -453,7 +453,7 @@ bool FControlRigSnapper::SnapIt(FFrameNumber StartFrame, FFrameNumber EndFrame,c
 		FScopedTransaction ScopedTransaction(LOCTEXT("SnapAnimation", "Snap Animation"));
 
 		ISequencer* Sequencer = InSequencer.Pin().Get();
-
+		Sequencer->ForceEvaluate(); // force an evaluate so any control rig get's binding setup
 		UMovieScene* MovieScene = Sequencer->GetFocusedMovieSceneSequence()->GetMovieScene();
 		MovieScene->Modify();
 		
@@ -619,7 +619,7 @@ bool FControlRigSnapper::SnapIt(FFrameNumber StartFrame, FFrameNumber EndFrame,c
 			}
 		}
 	}
-	return false;
+	return true;
 }
 
 
diff --git a/Engine/Plugins/Experimental/ControlRig/Source/ControlRigEditor/Public/ControlRigSequencerEditorLibrary.h b/Engine/Plugins/Experimental/ControlRig/Source/ControlRigEditor/Public/ControlRigSequencerEditorLibrary.h
--- a/Engine/Plugins/Experimental/ControlRig/Source/ControlRigEditor/Public/ControlRigSequencerEditorLibrary.h
+++ b/Engine/Plugins/Experimental/ControlRig/Source/ControlRigEditor/Public/ControlRigSequencerEditorLibrary.h
@@ -15,6 +15,7 @@
 #include "TransformNoScale.h"
 #include "EulerTransform.h"
 #include "Tools/ControlRigSnapSettings.h"
+#include "SequenceTimeUnit.h"
 #include "ControlRigSequencerEditorLibrary.generated.h"
 
 class ULevelSequence;
@@ -93,6 +94,7 @@ public:
 	* @param MovieScene The MovieScene getting loaded into
 	* @param SkelMeshComponent The Skeletal Mesh component getting loaded into.
 	* @param InStartFrame Frame to insert the animation
+	* @param TimeUnit Unit for all frame and time values, either in display rate or tick resolution
 	* @param bKeyReduce If true do key reduction based upon Tolerance, if false don't
 	* @param Tolerance If reducing keys, tolerance about which keys will be removed, smaller tolerance, more keys usually.
 
@@ -100,7 +102,7 @@ public:
 	*/
 	UFUNCTION(BlueprintCallable, Category = "Editor Scripting | Sequencer Tools | Control Rig")
 	static bool LoadAnimSequenceIntoControlRigSection(UMovieSceneSection* MovieSceneSection, UAnimSequence* AnimSequence, USkeletalMeshComponent* SkelMeshComp,
-		FFrameNumber InStartFrame, bool bKeyReduce = false, float Tolerance = 0.001f);
+		FFrameNumber InStartFrame, ESequenceTimeUnit TimeUnit = ESequenceTimeUnit::DisplayRate, bool bKeyReduce = false, float Tolerance = 0.001f);
 
 	/**
 	* Bake the current animation in the binding to a Control Rig track
@@ -128,56 +130,66 @@ public:
 
 	/**
 	* Peform a Snap operation to snap the children to the parent. 
+	* @param LevelSequence Active Sequence to snap
 	* @param StartFrame Beginning of the snap
 	* @param EndFrame End of the snap
 	* @param ChildrenToSnap The children objects that snap and get keys set onto. They need to live in an active Sequencer in the level editor
 	* @param ParentToSnap The parent object to snap relative to. If animated, it needs to live in an active Sequencer in the level editor
 	* @param SnapSettings Settings to use
+	* @param TimeUnit Unit for frame values, either in display rate or tick resolution
 	* @param Returns True if successful, 
 	*/
 	UFUNCTION(BlueprintCallable, Category = "Editor Scripting | Sequencer Tools | Control Rig")
-	static bool SnapControlRig(FFrameNumber StartFrame, FFrameNumber EndFrame, const FControlRigSnapperSelection& ChildrenToSnap,
-		const FControlRigSnapperSelection& ParentToSnap, const UControlRigSnapSettings* SnapSettings);
+	static bool SnapControlRig(ULevelSequence* LevelSequence, FFrameNumber StartFrame, FFrameNumber EndFrame, const FControlRigSnapperSelection& ChildrenToSnap,
+		const FControlRigSnapperSelection& ParentToSnap, const UControlRigSnapSettings* SnapSettings, ESequenceTimeUnit TimeUnit = ESequenceTimeUnit::DisplayRate);
 	
 	/**
 	* Get Actors World Transform at a specific time
 	* @param LevelSequence Active Sequence to get transform for
 	* @param Actor The actor
 	* @param Frame Time to get the transform
+	* @param TimeUnit Unit for frame values, either in display rate or tick resolution
 	* @return Returns World Transform
 	*/
 	UFUNCTION(BlueprintCallable, Category = "Editor Scripting | Sequencer Tools | Control Rig")
-	static FTransform GetActorWorldTransform(ULevelSequence* LevelSequence,AActor* Actor, FFrameNumber Frame);
+	static FTransform GetActorWorldTransform(ULevelSequence* LevelSequence,AActor* Actor, FFrameNumber Frame, ESequenceTimeUnit TimeUnit = ESequenceTimeUnit::DisplayRate);
 
 	/**
 	* Get Actors World Transforms at specific times
 	* @param LevelSequence Active Sequence to get transform for
 	* @param Actor The actor
 	* @param Frames Times to get the transform
+	* @param TimeUnit Unit for frame values, either in display rate or tick resolution
 	* @return Returns World Transforms
 	*/
 	UFUNCTION(BlueprintCallable, Category = "Editor Scripting | Sequencer Tools | Control Rig")
-	static TArray<FTransform> GetActorWorldTransforms(ULevelSequence* LevelSequence, AActor* Actor, const TArray<FFrameNumber>& Frames);
+	static TArray<FTransform> GetActorWorldTransforms(ULevelSequence* LevelSequence, AActor* Actor, const TArray<FFrameNumber>& Frames, ESequenceTimeUnit TimeUnit = ESequenceTimeUnit::DisplayRate);
 
 	/**
 	* Get SkeletalMeshComponent World Transform at a specific time
 	* @param LevelSequence Active Sequence to get transform for
 	* @param SkeletalMeshComponent The SkeletalMeshComponent
 	* @param Frame Time to get the transform
+	* @param TimeUnit Unit for frame values, either in display rate or tick resolution
+	* @param SocketName Optional name of the socket
 	* @return Returns World Transform
 	*/
 	UFUNCTION(BlueprintCallable, Category = "Editor Scripting | Sequencer Tools | Control Rig")
-	static FTransform GetSkeletalMeshComponentWorldTransform(ULevelSequence* LevelSequence, USkeletalMeshComponent* SkeletalMeshComponent, FFrameNumber Frame, FName SocketName = NAME_None);		
+	static FTransform GetSkeletalMeshComponentWorldTransform(ULevelSequence* LevelSequence, USkeletalMeshComponent* SkeletalMeshComponent, FFrameNumber Frame,
+		ESequenceTimeUnit TimeUnit = ESequenceTimeUnit::DisplayRate,FName SocketName = NAME_None);
 
 	/**
 	* Get SkeletalMeshComponents World Transforms at specific times
 	* @param LevelSequence Active Sequence to get transform for
 	* @param SkeletalMeshComponent The SkeletalMeshComponent
 	* @param Frames Times to get the transform
+	* @param TimeUnit Unit for frame values, either in display rate or tick resolution
+	* @param SocketName Optional name of the socket
 	* @return Returns World Transforms
 	*/
 	UFUNCTION(BlueprintCallable, Category = "Editor Scripting | Sequencer Tools | Control Rig")
-	static TArray<FTransform> GetSkeletalMeshComponentWorldTransforms(ULevelSequence* LevelSequence, USkeletalMeshComponent* SkeletalMeshComponent, const TArray<FFrameNumber>& Frames, FName SocketName = NAME_None);
+	static TArray<FTransform> GetSkeletalMeshComponentWorldTransforms(ULevelSequence* LevelSequence, USkeletalMeshComponent* SkeletalMeshComponent, const TArray<FFrameNumber>& Frames, 
+		ESequenceTimeUnit TimeUnit = ESequenceTimeUnit::DisplayRate,FName SocketName = NAME_None);
 
 	/**
 	* Get ControlRig Control's World Transform at a specific time
@@ -185,10 +197,12 @@ public:
 	* @param ControlRig The ControlRig
 	* @param ControlName Name of the Control
 	* @param Frame Time to get the transform
+	* @param TimeUnit Unit for frame values, either in display rate or tick resolution
 	* @return Returns World Transform
 	*/
 	UFUNCTION(BlueprintCallable, Category = "Editor Scripting | Sequencer Tools | Control Rig")
-	static FTransform GetControlRigWorldTransform(ULevelSequence* LevelSequence, UControlRig* ControlRig, FName ControlName, FFrameNumber Frame);
+	static FTransform GetControlRigWorldTransform(ULevelSequence* LevelSequence, UControlRig* ControlRig, FName ControlName, FFrameNumber Frame,
+		ESequenceTimeUnit TimeUnit = ESequenceTimeUnit::DisplayRate);
 
 	/**
 	* Get ControlRig Control's World Transforms at specific times
@@ -196,10 +210,12 @@ public:
 	* @param ControlRig The ControlRig
 	* @param ControlName Name of the Control 
 	* @param Frames Times to get the transform
+	* @param TimeUnit Unit for frame values, either in display rate or tick resolution
 	* @return Returns World Transforms
 	*/
 	UFUNCTION(BlueprintCallable, Category = "Editor Scripting | Sequencer Tools | Control Rig")
-	static TArray<FTransform> GetControlRigWorldTransforms(ULevelSequence* LevelSequence, UControlRig* ControlRig, FName ControlName, const TArray<FFrameNumber>& Frames);
+	static TArray<FTransform> GetControlRigWorldTransforms(ULevelSequence* LevelSequence, UControlRig* ControlRig, FName ControlName, const TArray<FFrameNumber>& Frames,
+		ESequenceTimeUnit TimeUnit = ESequenceTimeUnit::DisplayRate);
 
 	/**
 	* Set ControlRig Control's World Transform at a specific time
@@ -208,22 +224,25 @@ public:
 	* @param ControlName Name of the Control
 	* @param Frame Time to set the transform
 	* @oaram WorldTransform World Transform to set
+	* @param TimeUnit Unit for frame values, either in display rate or tick resolution
 	* @param bSetKey Whether or not to set a key.
 	*/
 	UFUNCTION(BlueprintCallable, Category = "Editor Scripting | Sequencer Tools | Control Rig")
-	static void SetControlRigWorldTransform(ULevelSequence* LevelSequence, UControlRig* ControlRig, FName ControlName, FFrameNumber Frame, const FTransform& WorldTransform, bool bSetKey = true);
+	static void SetControlRigWorldTransform(ULevelSequence* LevelSequence, UControlRig* ControlRig, FName ControlName, FFrameNumber Frame, const FTransform& WorldTransform,
+		ESequenceTimeUnit TimeUnit = ESequenceTimeUnit::DisplayRate, bool bSetKey = true);
 
 	/**
 	* Set ControlRig Control's World Transforms at a specific times.
 	* @param LevelSequence Active Sequence to set transforms for. Must be loaded in Level Editor.
 	* @param ControlRig The ControlRig
 	* @param ControlName Name of the Control
-	* @param Frame Time to set the transform
+	* @param Frames Times to set the transform
 	* @oaram WorldTransform World Transform to set
-	* @param bSetKey Whether or not to set a key.
+	* @param TimeUnit Unit for frame values, either in display rate or tick resolution
 	*/
 	UFUNCTION(BlueprintCallable, Category = "Editor Scripting | Sequencer Tools | Control Rig")
-	static void SetControlRigWorldTransforms(ULevelSequence* LevelSequence, UControlRig* ControlRig, FName ControlName, const TArray<FFrameNumber>& Frames, const TArray<FTransform>& WorldTransforms);
+	static void SetControlRigWorldTransforms(ULevelSequence* LevelSequence, UControlRig* ControlRig, FName ControlName, const TArray<FFrameNumber>& Frames, 
+		 const TArray<FTransform>& WorldTransforms, ESequenceTimeUnit TimeUnit = ESequenceTimeUnit::DisplayRate);
 	
 	/**
 	* Get ControlRig Control's float value at a specific time
@@ -231,10 +250,12 @@ public:
 	* @param ControlRig The ControlRig
 	* @param ControlName Name of the Control, should be a float control
 	* @param Frame Time to get the value
+	* @param TimeUnit Unit for frame values, either in display rate or tick resolution
 	* @return Returns Value at that time
 	*/
 	UFUNCTION(BlueprintCallable, Category = "Editor Scripting | Sequencer Tools | Control Rig")
-	static float GetLocalControlRigFloat(ULevelSequence* LevelSequence, UControlRig* ControlRig, FName ControlName, FFrameNumber Frame);
+	static float GetLocalControlRigFloat(ULevelSequence* LevelSequence, UControlRig* ControlRig, FName ControlName, FFrameNumber Frame,
+		ESequenceTimeUnit TimeUnit = ESequenceTimeUnit::DisplayRate);
 
 	/**
 	* Get ControlRig Control's float values at specific times
@@ -242,10 +263,12 @@ public:
 	* @param ControlRig The ControlRig
 	* @param ControlName Name of the Control, should be a float control
 	* @param Frames Times to get the values
+	* @param TimeUnit Unit for frame values, either in display rate or tick resolution
 	* @return Returns Values at those times
 	*/
 	UFUNCTION(BlueprintCallable, Category = "Editor Scripting | Sequencer Tools | Control Rig")
-	static TArray<float> GetLocalControlRigFloats(ULevelSequence* LevelSequence, UControlRig* ControlRig, FName ControlName, const TArray<FFrameNumber>& Frames);
+	static TArray<float> GetLocalControlRigFloats(ULevelSequence* LevelSequence, UControlRig* ControlRig, FName ControlName, const TArray<FFrameNumber>& Frames,
+		ESequenceTimeUnit TimeUnit = ESequenceTimeUnit::DisplayRate);
 
 	/**
 	* Set ControlRig Control's float value at specific time
@@ -254,10 +277,12 @@ public:
 	* @param ControlName Name of the Control, should be a float control
 	* @param Frame Time to set the value
 	* @param Value to set
+	* @param TimeUnit Unit for frame values, either in display rate or tick resolution
 	* @param bSetKey If True set a key, if not just set the value
 	*/
 	UFUNCTION(BlueprintCallable, Category = "Editor Scripting | Sequencer Tools | Control Rig")
-	static void SetLocalControlRigFloat(ULevelSequence* LevelSequence, UControlRig* ControlRig, FName ControlName, FFrameNumber Frame, float Value, bool bSetKey = true);
+	static void SetLocalControlRigFloat(ULevelSequence* LevelSequence, UControlRig* ControlRig, FName ControlName, FFrameNumber Frame, float Value,
+		ESequenceTimeUnit TimeUnit = ESequenceTimeUnit::DisplayRate,bool bSetKey = true);
 
 	/**
 	* Set ControlRig Control's float values at specific times
@@ -266,9 +291,11 @@ public:
 	* @param ControlName Name of the Control, should be a float control
 	* @param Frames Times to set the values
 	* @param Values to set at those times
+	* @param TimeUnit Unit for frame values, either in display rate or tick resolution
 	*/
 	UFUNCTION(BlueprintCallable, Category = "Editor Scripting | Sequencer Tools | Control Rig")
-	static void SetLocalControlRigFloats(ULevelSequence* LevelSequence, UControlRig* ControlRig, FName ControlName, const TArray<FFrameNumber>& Frames, const TArray<float> Values);
+	static void SetLocalControlRigFloats(ULevelSequence* LevelSequence, UControlRig* ControlRig, FName ControlName, const TArray<FFrameNumber>& Frames, const TArray<float> Values,
+		ESequenceTimeUnit TimeUnit = ESequenceTimeUnit::DisplayRate);
 
 	/**
 	* Get ControlRig Control's bool value at a specific time
@@ -276,10 +303,12 @@ public:
 	* @param ControlRig The ControlRig
 	* @param ControlName Name of the Control, should be a bool control
 	* @param Frame Time to get the value
+	* @param TimeUnit Unit for frame values, either in display rate or tick resolution
 	* @return Returns Value at that time
 	*/
 	UFUNCTION(BlueprintCallable, Category = "Editor Scripting | Sequencer Tools | Control Rig")
-	static bool GetLocalControlRigBool(ULevelSequence* LevelSequence, UControlRig* ControlRig, FName ControlName, FFrameNumber Frame);
+	static bool GetLocalControlRigBool(ULevelSequence* LevelSequence, UControlRig* ControlRig, FName ControlName, FFrameNumber Frame,
+		ESequenceTimeUnit TimeUnit = ESequenceTimeUnit::DisplayRate);
 
 	/**
 	* Get ControlRig Control's bool values at specific times
@@ -287,10 +316,12 @@ public:
 	* @param ControlRig The ControlRig
 	* @param ControlName Name of the Control, should be a bool control
 	* @param Frames Times to get the values
+	* @param TimeUnit Unit for frame values, either in display rate or tick resolution
 	* @return Returns Values at those times
 	*/
 	UFUNCTION(BlueprintCallable, Category = "Editor Scripting | Sequencer Tools | Control Rig")
-	static TArray<bool> GetLocalControlRigBools(ULevelSequence* LevelSequence, UControlRig* ControlRig, FName ControlName, const TArray<FFrameNumber>& Frames);
+	static TArray<bool> GetLocalControlRigBools(ULevelSequence* LevelSequence, UControlRig* ControlRig, FName ControlName, const TArray<FFrameNumber>& Frames,
+		ESequenceTimeUnit TimeUnit = ESequenceTimeUnit::DisplayRate);
 
 	/**
 	* Set ControlRig Control's bool value at specific time
@@ -299,10 +330,12 @@ public:
 	* @param ControlName Name of the Control, should be a bool control
 	* @param Frame Time to set the value
 	* @param Value to set
+	* @param TimeUnit Unit for frame values, either in display rate or tick resolution
 	* @param bSetKey If True set a key, if not just set the value
 	*/
 	UFUNCTION(BlueprintCallable, Category = "Editor Scripting | Sequencer Tools | Control Rig")
-	static void SetLocalControlRigBool(ULevelSequence* LevelSequence, UControlRig* ControlRig, FName ControlName, FFrameNumber Frame, bool Value, bool bSetKey = true);
+	static void SetLocalControlRigBool(ULevelSequence* LevelSequence, UControlRig* ControlRig, FName ControlName, FFrameNumber Frame, bool Value,
+		ESequenceTimeUnit TimeUnit = ESequenceTimeUnit::DisplayRate, bool bSetKey = true);
 
 	/**
 	* Set ControlRig Control's bool values at specific times
@@ -311,9 +344,11 @@ public:
 	* @param ControlName Name of the Control, should be a bool control
 	* @param Frames Times to set the values
 	* @param Values to set at those times
+	* @param TimeUnit Unit for frame values, either in display rate or tick resolution
 */
 	UFUNCTION(BlueprintCallable, Category = "Editor Scripting | Sequencer Tools | Control Rig")
-	static void SetLocalControlRigBools(ULevelSequence* LevelSequence, UControlRig* ControlRig, FName ControlName, const TArray<FFrameNumber>& Frames, const TArray<bool> Values);
+	static void SetLocalControlRigBools(ULevelSequence* LevelSequence, UControlRig* ControlRig, FName ControlName, const TArray<FFrameNumber>& Frames,
+		const TArray<bool> Values, ESequenceTimeUnit TimeUnit = ESequenceTimeUnit::DisplayRate);
 
 	/**
 	* Get ControlRig Control's integer value at a specific time
@@ -321,10 +356,12 @@ public:
 	* @param ControlRig The ControlRig
 	* @param ControlName Name of the Control, should be a integer control
 	* @param Frame Time to get the value
+	* @param TimeUnit Unit for frame values, either in display rate or tick resolution
 	* @return Returns Value at that time
 	*/
 	UFUNCTION(BlueprintCallable, Category = "Editor Scripting | Sequencer Tools | Control Rig")
-	static int32 GetLocalControlRigInt(ULevelSequence* LevelSequence, UControlRig* ControlRig, FName ControlName, FFrameNumber Frame);
+	static int32 GetLocalControlRigInt(ULevelSequence* LevelSequence, UControlRig* ControlRig, FName ControlName, FFrameNumber Frame, 
+		ESequenceTimeUnit TimeUnit = ESequenceTimeUnit::DisplayRate);
 
 	/**
 	* Get ControlRig Control's integer values at specific times
@@ -332,10 +369,12 @@ public:
 	* @param ControlRig The ControlRig
 	* @param ControlName Name of the Control, should be a intteger control
 	* @param Frames Times to get the values
+	* @param TimeUnit Unit for frame values, either in display rate or tick resolution
 	* @return Returns Values at those times
 	*/
 	UFUNCTION(BlueprintCallable, Category = "Editor Scripting | Sequencer Tools | Control Rig")
-	static TArray<int32> GetLocalControlRigInts(ULevelSequence* LevelSequence, UControlRig* ControlRig, FName ControlName, const TArray<FFrameNumber>& Frames);
+	static TArray<int32> GetLocalControlRigInts(ULevelSequence* LevelSequence, UControlRig* ControlRig, FName ControlName, const TArray<FFrameNumber>& Frames,
+		ESequenceTimeUnit TimeUnit = ESequenceTimeUnit::DisplayRate);
 
 	/**
 	* Set ControlRig Control's int value at specific time
@@ -344,10 +383,12 @@ public:
 	* @param ControlName Name of the Control, should be a int control
 	* @param Frame Time to set the value
 	* @param Value to set
+	* @param TimeUnit Unit for frame values, either in display rate or tick resolution
 	* @param bSetKey If True set a key, if not just set the value
 	*/
 	UFUNCTION(BlueprintCallable, Category = "Editor Scripting | Sequencer Tools | Control Rig")
-	static void SetLocalControlRigInt(ULevelSequence* LevelSequence, UControlRig* ControlRig, FName ControlName, FFrameNumber Frame, int32 Value, bool bSetKey = true);
+	static void SetLocalControlRigInt(ULevelSequence* LevelSequence, UControlRig* ControlRig, FName ControlName, FFrameNumber Frame, int32 Value,
+		ESequenceTimeUnit TimeUnit = ESequenceTimeUnit::DisplayRate,bool bSetKey = true);
 
 
 	/**
@@ -357,9 +398,11 @@ public:
 	* @param ControlName Name of the Control, should be a int control
 	* @param Frames Times to set the values
 	* @param Values to set at those times
+	* @param TimeUnit Unit for frame values, either in display rate or tick resolution
 	*/
 	UFUNCTION(BlueprintCallable, Category = "Editor Scripting | Sequencer Tools | Control Rig")
-	static void SetLocalControlRigInts(ULevelSequence* LevelSequence, UControlRig* ControlRig, FName ControlName, const TArray<FFrameNumber>& Frames, const TArray<int32> Values);
+	static void SetLocalControlRigInts(ULevelSequence* LevelSequence, UControlRig* ControlRig, FName ControlName, const TArray<FFrameNumber>& Frames, const TArray<int32> Values,
+		ESequenceTimeUnit TimeUnit = ESequenceTimeUnit::DisplayRate);
 
 	/**
 	* Get ControlRig Control's Vector2D value at a specific time
@@ -367,10 +410,12 @@ public:
 	* @param ControlRig The ControlRig
 	* @param ControlName Name of the Control, should be a Vector2D control
 	* @param Frame Time to get the value
+	* @param TimeUnit Unit for frame values, either in display rate or tick resolution
 	* @return Returns Value at that time
 	*/
 	UFUNCTION(BlueprintCallable, Category = "Editor Scripting | Sequencer Tools | Control Rig")
-	static FVector2D GetLocalControlRigVector2D(ULevelSequence* LevelSequence, UControlRig* ControlRig, FName ControlName, FFrameNumber Frame);
+	static FVector2D GetLocalControlRigVector2D(ULevelSequence* LevelSequence, UControlRig* ControlRig, FName ControlName, FFrameNumber Frame,
+		ESequenceTimeUnit TimeUnit = ESequenceTimeUnit::DisplayRate);
 
 	/**
 	* Get ControlRig Control's Vector2D values at specific times
@@ -378,10 +423,12 @@ public:
 	* @param ControlRig The ControlRig
 	* @param ControlName Name of the Control, should be a Vector2D control
 	* @param Frames Times to get the values
+	* @param TimeUnit Unit for frame values, either in display rate or tick resolution
 	* @return Returns Values at those times
 	*/
 	UFUNCTION(BlueprintCallable, Category = "Editor Scripting | Sequencer Tools | Control Rig")
-	static TArray<FVector2D> GetLocalControlRigVector2Ds(ULevelSequence* LevelSequence, UControlRig* ControlRig, FName ControlName, const TArray<FFrameNumber>& Frames);
+	static TArray<FVector2D> GetLocalControlRigVector2Ds(ULevelSequence* LevelSequence, UControlRig* ControlRig, FName ControlName, const TArray<FFrameNumber>& Frames,
+		ESequenceTimeUnit TimeUnit = ESequenceTimeUnit::DisplayRate);
 
 	/**
 	* Set ControlRig Control's Vector2D value at specific time
@@ -390,10 +437,12 @@ public:
 	* @param ControlName Name of the Control, should be a Vector2D control
 	* @param Frame Time to set the value
 	* @param Value to set
+	* @param TimeUnit Unit for frame values, either in display rate or tick resolution
 	* @param bSetKey If True set a key, if not just set the value
 	*/
 	UFUNCTION(BlueprintCallable, Category = "Editor Scripting | Sequencer Tools | Control Rig")
-	static void SetLocalControlRigVector2D(ULevelSequence* LevelSequence, UControlRig* ControlRig, FName ControlName, FFrameNumber Frame, FVector2D Value, bool bSetKey = true);
+	static void SetLocalControlRigVector2D(ULevelSequence* LevelSequence, UControlRig* ControlRig, FName ControlName, FFrameNumber Frame, FVector2D Value,
+		ESequenceTimeUnit TimeUnit = ESequenceTimeUnit::DisplayRate,bool bSetKey = true);
 
 
 	/**
@@ -403,9 +452,11 @@ public:
 	* @param ControlName Name of the Control, should be a Vector2D control
 	* @param Frames Times to set the values
 	* @param Values to set at those times
+	* @param TimeUnit Unit for frame values, either in display rate or tick resolution
 	*/
 	UFUNCTION(BlueprintCallable, Category = "Editor Scripting | Sequencer Tools | Control Rig")
-	static void SetLocalControlRigVector2Ds(ULevelSequence* LevelSequence, UControlRig* ControlRig, FName ControlName, const TArray<FFrameNumber>& Frames, const TArray<FVector2D> Values);
+	static void SetLocalControlRigVector2Ds(ULevelSequence* LevelSequence, UControlRig* ControlRig, FName ControlName, const TArray<FFrameNumber>& Frames, const TArray<FVector2D> Values,
+		ESequenceTimeUnit TimeUnit = ESequenceTimeUnit::DisplayRate);
 
 	/**
 	* Get ControlRig Control's Position value at a specific time
@@ -413,10 +464,12 @@ public:
 	* @param ControlRig The ControlRig
 	* @param ControlName Name of the Control, should be a Position control
 	* @param Frame Time to get the value
+	* @param TimeUnit Unit for frame values, either in display rate or tick resolution
 	* @return Returns Value at that time
 	*/
 	UFUNCTION(BlueprintCallable, Category = "Editor Scripting | Sequencer Tools | Control Rig")
-	static FVector GetLocalControlRigPosition(ULevelSequence* LevelSequence, UControlRig* ControlRig, FName ControlName, FFrameNumber Frame);
+	static FVector GetLocalControlRigPosition(ULevelSequence* LevelSequence, UControlRig* ControlRig, FName ControlName, FFrameNumber Frame,
+		ESequenceTimeUnit TimeUnit = ESequenceTimeUnit::DisplayRate);
 
 	/**
 	* Get ControlRig Control's Position values at specific times
@@ -424,10 +477,12 @@ public:
 	* @param ControlRig The ControlRig
 	* @param ControlName Name of the Control, should be a Position control
 	* @param Frames Times to get the values
+	* @param TimeUnit Unit for frame values, either in display rate or tick resolution
 	* @return Returns Values at those times
 	*/
 	UFUNCTION(BlueprintCallable, Category = "Editor Scripting | Sequencer Tools | Control Rig")
-	static TArray<FVector> GetLocalControlRigPositions(ULevelSequence* LevelSequence, UControlRig* ControlRig, FName ControlName, const TArray<FFrameNumber>& Frames);
+	static TArray<FVector> GetLocalControlRigPositions(ULevelSequence* LevelSequence, UControlRig* ControlRig, FName ControlName, const TArray<FFrameNumber>& Frames,
+		ESequenceTimeUnit TimeUnit = ESequenceTimeUnit::DisplayRate);
 
 	/**
 	* Set ControlRig Control's Position value at specific time
@@ -436,10 +491,12 @@ public:
 	* @param ControlName Name of the Control, should be a Position control
 	* @param Frame Time to set the value
 	* @param Value to set
+	* @param TimeUnit Unit for frame values, either in display rate or tick resolution
 	* @param bSetKey If True set a key, if not just set the value
 	*/
 	UFUNCTION(BlueprintCallable, Category = "Editor Scripting | Sequencer Tools | Control Rig")
-	static void SetLocalControlRigPosition(ULevelSequence* LevelSequence, UControlRig* ControlRig, FName ControlName, FFrameNumber Frame, FVector Value, bool bSetKey = true);
+	static void SetLocalControlRigPosition(ULevelSequence* LevelSequence, UControlRig* ControlRig, FName ControlName, FFrameNumber Frame, FVector Value,
+		ESequenceTimeUnit TimeUnit = ESequenceTimeUnit::DisplayRate, bool bSetKey = true);
 
 	/**
 	* Set ControlRig Control's Position values at specific times
@@ -448,9 +505,11 @@ public:
 	* @param ControlName Name of the Control, should be a Position control
 	* @param Frames Times to set the values
 	* @param Values to set at those times
+	* @param TimeUnit Unit for frame values, either in display rate or tick resolution
 	*/
 	UFUNCTION(BlueprintCallable, Category = "Editor Scripting | Sequencer Tools | Control Rig")
-	static void SetLocalControlRigPositions(ULevelSequence* LevelSequence, UControlRig* ControlRig, FName ControlName, const TArray<FFrameNumber>& Frames, const TArray<FVector> Values);
+	static void SetLocalControlRigPositions(ULevelSequence* LevelSequence, UControlRig* ControlRig, FName ControlName, const TArray<FFrameNumber>& Frames, const TArray<FVector> Values,
+		ESequenceTimeUnit TimeUnit = ESequenceTimeUnit::DisplayRate);
 
 	/**
 	* Get ControlRig Control's Rotator value at a specific time
@@ -458,10 +517,12 @@ public:
 	* @param ControlRig The ControlRig
 	* @param ControlName Name of the Control, should be a Rotator control
 	* @param Frame Time to get the value
+	* @param TimeUnit Unit for frame values, either in display rate or tick resolution
 	* @return Returns Value at that time
 	*/
 	UFUNCTION(BlueprintCallable, Category = "Editor Scripting | Sequencer Tools | Control Rig")
-	static FRotator GetLocalControlRigRotator(ULevelSequence* LevelSequence, UControlRig* ControlRig, FName ControlName, FFrameNumber Frame);
+	static FRotator GetLocalControlRigRotator(ULevelSequence* LevelSequence, UControlRig* ControlRig, FName ControlName, FFrameNumber Frame,
+		ESequenceTimeUnit TimeUnit = ESequenceTimeUnit::DisplayRate);
 
 	/**
 	* Get ControlRig Control's Rotator values at specific times
@@ -469,10 +530,12 @@ public:
 	* @param ControlRig The ControlRig
 	* @param ControlName Name of the Control, should be a Rotator control
 	* @param Frames Times to get the values
+	* @param TimeUnit Unit for frame values, either in display rate or tick resolution
 	* @return Returns Values at those times
 	*/
 	UFUNCTION(BlueprintCallable, Category = "Editor Scripting | Sequencer Tools | Control Rig")
-	static TArray<FRotator> GetLocalControlRigRotators(ULevelSequence* LevelSequence, UControlRig* ControlRig, FName ControlName, const TArray<FFrameNumber>& Frames);
+	static TArray<FRotator> GetLocalControlRigRotators(ULevelSequence* LevelSequence, UControlRig* ControlRig, FName ControlName, const TArray<FFrameNumber>& Frames,
+		ESequenceTimeUnit TimeUnit = ESequenceTimeUnit::DisplayRate);
 
 	/**
 	* Set ControlRig Control's Rotator value at specific time
@@ -481,10 +544,12 @@ public:
 	* @param ControlName Name of the Control, should be a Rotator control
 	* @param Frame Time to set the value
 	* @param Value to set
+	* @param TimeUnit Unit for frame values, either in display rate or tick resolution
 	* @param bSetKey If True set a key, if not just set the value
 	*/
 	UFUNCTION(BlueprintCallable, Category = "Editor Scripting | Sequencer Tools | Control Rig")
-	static void SetLocalControlRigRotator(ULevelSequence* LevelSequence, UControlRig* ControlRig, FName ControlName, FFrameNumber Frame, FRotator Value, bool bSetKey = true);
+	static void SetLocalControlRigRotator(ULevelSequence* LevelSequence, UControlRig* ControlRig, FName ControlName, FFrameNumber Frame, FRotator Value, 
+		ESequenceTimeUnit TimeUnit = ESequenceTimeUnit::DisplayRate, bool bSetKey = true);
 
 	/**
 	* Set ControlRig Control's Rotator values at specific times
@@ -493,9 +558,11 @@ public:
 	* @param ControlName Name of the Control, should be a Rotator control
 	* @param Frames Times to set the values
 	* @param Values to set at those times
+	* @param TimeUnit Unit for frame values, either in display rate or tick resolution
 	*/
 	UFUNCTION(BlueprintCallable, Category = "Editor Scripting | Sequencer Tools | Control Rig")
-	static void SetLocalControlRigRotators(ULevelSequence* LevelSequence, UControlRig* ControlRig, FName ControlName, const TArray<FFrameNumber>& Frames, const TArray<FRotator> Values);
+	static void SetLocalControlRigRotators(ULevelSequence* LevelSequence, UControlRig* ControlRig, FName ControlName, const TArray<FFrameNumber>& Frames, const TArray<FRotator> Values,
+		ESequenceTimeUnit TimeUnit = ESequenceTimeUnit::DisplayRate);
 
 	/**
 	* Get ControlRig Control's Scale value at a specific time
@@ -503,10 +570,12 @@ public:
 	* @param ControlRig The ControlRig
 	* @param ControlName Name of the Control, should be a Scale control
 	* @param Frame Time to get the value
+	* @param TimeUnit Unit for frame values, either in display rate or tick resolution
 	* @return Returns Value at that time
 	*/
 	UFUNCTION(BlueprintCallable, Category = "Editor Scripting | Sequencer Tools | Control Rig")
-	static FVector GetLocalControlRigScale(ULevelSequence* LevelSequence, UControlRig* ControlRig, FName ControlName, FFrameNumber Frame);
+	static FVector GetLocalControlRigScale(ULevelSequence* LevelSequence, UControlRig* ControlRig, FName ControlName, FFrameNumber Frame,
+		ESequenceTimeUnit TimeUnit = ESequenceTimeUnit::DisplayRate);
 
 	/**
 	* Get ControlRig Control's Scale values at specific times
@@ -514,10 +583,12 @@ public:
 	* @param ControlRig The ControlRig
 	* @param ControlName Name of the Control, should be a Scale control
 	* @param Frames Times to get the values
+	* @param TimeUnit Unit for frame values, either in display rate or tick resolution
 	* @return Returns Values at those times
 	*/
 	UFUNCTION(BlueprintCallable, Category = "Editor Scripting | Sequencer Tools | Control Rig")
-	static TArray<FVector> GetLocalControlRigScales(ULevelSequence* LevelSequence, UControlRig* ControlRig, FName ControlName, const TArray<FFrameNumber>& Frames);
+	static TArray<FVector> GetLocalControlRigScales(ULevelSequence* LevelSequence, UControlRig* ControlRig, FName ControlName, const TArray<FFrameNumber>& Frames,
+		ESequenceTimeUnit TimeUnit = ESequenceTimeUnit::DisplayRate);
 
 	/**
 	* Set ControlRig Control's Scale value at specific time
@@ -526,10 +597,12 @@ public:
 	* @param ControlName Name of the Control, should be a Scale control
 	* @param Frame Time to set the value
 	* @param Value to set
+	* @param TimeUnit Unit for frame values, either in display rate or tick resolution
 	* @param bSetKey If True set a key, if not just set the value
 */
 	UFUNCTION(BlueprintCallable, Category = "Editor Scripting | Sequencer Tools | Control Rig")
-	static void SetLocalControlRigScale(ULevelSequence* LevelSequence, UControlRig* ControlRig, FName ControlName, FFrameNumber Frame, FVector Value, bool bSetKey = true);
+	static void SetLocalControlRigScale(ULevelSequence* LevelSequence, UControlRig* ControlRig, FName ControlName, FFrameNumber Frame, FVector Value, 
+		ESequenceTimeUnit TimeUnit = ESequenceTimeUnit::DisplayRate, bool bSetKey = true);
 
 	/**
 	* Set ControlRig Control's Scale values at specific times
@@ -538,9 +611,11 @@ public:
 	* @param ControlName Name of the Control, should be a Scale control
 	* @param Frames Times to set the values
 	* @param Values to set at those times
+	* @param TimeUnit Unit for frame values, either in display rate or tick resolution
 	*/
 	UFUNCTION(BlueprintCallable, Category = "Editor Scripting | Sequencer Tools | Control Rig")
-	static void SetLocalControlRigScales(ULevelSequence* LevelSequence, UControlRig* ControlRig, FName ControlName, const TArray<FFrameNumber>& Frames, const TArray<FVector> Values);
+	static void SetLocalControlRigScales(ULevelSequence* LevelSequence, UControlRig* ControlRig, FName ControlName, const TArray<FFrameNumber>& Frames, const TArray<FVector> Values,
+		ESequenceTimeUnit TimeUnit = ESequenceTimeUnit::DisplayRate);
 
 
 	/**
@@ -549,10 +624,12 @@ public:
 	* @param ControlRig The ControlRig
 	* @param ControlName Name of the Control, should be a EulerTransfom control
 	* @param Frame Time to get the value
+	* @param TimeUnit Unit for frame values, either in display rate or tick resolution
 	* @return Returns Value at that time
 	*/
 	UFUNCTION(BlueprintCallable, Category = "Editor Scripting | Sequencer Tools | Control Rig")
-	static FEulerTransform GetLocalControlRigEulerTransform(ULevelSequence* LevelSequence, UControlRig* ControlRig, FName ControlName, FFrameNumber Frame);
+	static FEulerTransform GetLocalControlRigEulerTransform(ULevelSequence* LevelSequence, UControlRig* ControlRig, FName ControlName, FFrameNumber Frame,
+		ESequenceTimeUnit TimeUnit = ESequenceTimeUnit::DisplayRate);
 
 	/**
 	* Get ControlRig Control's EulerTransform values at specific times
@@ -560,10 +637,12 @@ public:
 	* @param ControlRig The ControlRig
 	* @param ControlName Name of the Control, should be a EulerTransform control
 	* @param Frames Times to get the values
+	* @param TimeUnit Unit for frame values, either in display rate or tick resolution
 	* @return Returns Values at those times
 	*/
 	UFUNCTION(BlueprintCallable, Category = "Editor Scripting | Sequencer Tools | Control Rig")
-	static TArray<FEulerTransform> GetLocalControlRigEulerTransforms(ULevelSequence* LevelSequence, UControlRig* ControlRig, FName ControlName, const TArray<FFrameNumber>& Frames);
+	static TArray<FEulerTransform> GetLocalControlRigEulerTransforms(ULevelSequence* LevelSequence, UControlRig* ControlRig, FName ControlName, const TArray<FFrameNumber>& Frames,
+		ESequenceTimeUnit TimeUnit = ESequenceTimeUnit::DisplayRate);
 
 	/**
 	* Set ControlRig Control's EulerTransform value at specific time
@@ -572,10 +651,12 @@ public:
 	* @param ControlName Name of the Control, should be a EulerTransform control
 	* @param Frame Time to set the value
 	* @param Value to set
+	* @param TimeUnit Unit for frame values, either in display rate or tick resolution
 	* @param bSetKey If True set a key, if not just set the value
 	*/
 	UFUNCTION(BlueprintCallable, Category = "Editor Scripting | Sequencer Tools | Control Rig")
-	static void SetLocalControlRigEulerTransform(ULevelSequence* LevelSequence, UControlRig* ControlRig, FName ControlName, FFrameNumber Frame, FEulerTransform Value, bool bSetKey = true);
+	static void SetLocalControlRigEulerTransform(ULevelSequence* LevelSequence, UControlRig* ControlRig, FName ControlName, FFrameNumber Frame, FEulerTransform Value,
+		ESequenceTimeUnit TimeUnit = ESequenceTimeUnit::DisplayRate, bool bSetKey = true);
 
 	/**
 	* Set ControlRig Control's EulerTransform values at specific times
@@ -584,9 +665,11 @@ public:
 	* @param ControlName Name of the Control, should be a EulerTransform control
 	* @param Frames Times to set the values
 	* @param Values to set at those times
+	* @param TimeUnit Unit for frame values, either in display rate or tick resolution
 	*/
 	UFUNCTION(BlueprintCallable, Category = "Editor Scripting | Sequencer Tools | Control Rig")
-	static void SetLocalControlRigEulerTransforms(ULevelSequence* LevelSequence, UControlRig* ControlRig, FName ControlName, const TArray<FFrameNumber>& Frames, const TArray<FEulerTransform> Values);
+	static void SetLocalControlRigEulerTransforms(ULevelSequence* LevelSequence, UControlRig* ControlRig, FName ControlName, const TArray<FFrameNumber>& Frames, const TArray<FEulerTransform> Values,
+		ESequenceTimeUnit TimeUnit = ESequenceTimeUnit::DisplayRate);
 
 	/**
 	* Get ControlRig Control's TransformNoScale value at a specific time
@@ -594,10 +677,12 @@ public:
 	* @param ControlRig The ControlRig
 	* @param ControlName Name of the Control, should be a TransformNoScale control
 	* @param Frame Time to get the value
+	* @param TimeUnit Unit for frame values, either in display rate or tick resolution
 	* @return Returns Value at that time
 	*/
 	UFUNCTION(BlueprintCallable, Category = "Editor Scripting | Sequencer Tools | Control Rig")
-	static FTransformNoScale GetLocalControlRigTransformNoScale(ULevelSequence* LevelSequence, UControlRig* ControlRig, FName ControlName, FFrameNumber Frame);
+	static FTransformNoScale GetLocalControlRigTransformNoScale(ULevelSequence* LevelSequence, UControlRig* ControlRig, FName ControlName, FFrameNumber Frame,
+		ESequenceTimeUnit TimeUnit = ESequenceTimeUnit::DisplayRate);
 
 	/**
 	* Get ControlRig Control's TransformNoScale values at specific times
@@ -605,10 +690,12 @@ public:
 	* @param ControlRig The ControlRig
 	* @param ControlName Name of the Control, should be a TransformNoScale control
 	* @param Frames Times to get the values
+	* @param TimeUnit Unit for frame values, either in display rate or tick resolution
 	* @return Returns Values at those times
 	*/
 	UFUNCTION(BlueprintCallable, Category = "Editor Scripting | Sequencer Tools | Control Rig")
-	static TArray<FTransformNoScale> GetLocalControlRigTransformNoScales(ULevelSequence* LevelSequence, UControlRig* ControlRig, FName ControlName, const TArray<FFrameNumber>& Frames);
+	static TArray<FTransformNoScale> GetLocalControlRigTransformNoScales(ULevelSequence* LevelSequence, UControlRig* ControlRig, FName ControlName, const TArray<FFrameNumber>& Frames,
+		ESequenceTimeUnit TimeUnit = ESequenceTimeUnit::DisplayRate);
 
 	/**
 	* Set ControlRig Control's TransformNoScale value at specific time
@@ -617,10 +704,12 @@ public:
 	* @param ControlName Name of the Control, should be a TransformNoScale control
 	* @param Frame Time to set the value
 	* @param Value to set
+	* @param TimeUnit Unit for frame values, either in display rate or tick resolution
 	* @param bSetKey If True set a key, if not just set the value
 	*/
 	UFUNCTION(BlueprintCallable, Category = "Editor Scripting | Sequencer Tools | Control Rig")
-	static void SetLocalControlRigTransformNoScale(ULevelSequence* LevelSequence, UControlRig* ControlRig, FName ControlName, FFrameNumber Frame, FTransformNoScale Value, bool bSetKey = true);
+	static void SetLocalControlRigTransformNoScale(ULevelSequence* LevelSequence, UControlRig* ControlRig, FName ControlName, FFrameNumber Frame, FTransformNoScale Value, 
+		ESequenceTimeUnit TimeUnit = ESequenceTimeUnit::DisplayRate, bool bSetKey = true);
 
 	/**
 	* Set ControlRig Control's TransformNoScale values at specific times
@@ -629,9 +718,11 @@ public:
 	* @param ControlName Name of the Control, should be a TransformNoScale control
 	* @param Frames Times to set the values
 	* @param Values to set at those times
+	* @param TimeUnit Unit for frame values, either in display rate or tick resolution
 	*/
 	UFUNCTION(BlueprintCallable, Category = "Editor Scripting | Sequencer Tools | Control Rig")
-	static void SetLocalControlRigTransformNoScales(ULevelSequence* LevelSequence, UControlRig* ControlRig, FName ControlName, const TArray<FFrameNumber>& Frames, const TArray<FTransformNoScale> Values);
+	static void SetLocalControlRigTransformNoScales(ULevelSequence* LevelSequence, UControlRig* ControlRig, FName ControlName, const TArray<FFrameNumber>& Frames, const TArray<FTransformNoScale> Values,
+		ESequenceTimeUnit TimeUnit = ESequenceTimeUnit::DisplayRate);
 
 	/**
 	* Get ControlRig Control's Transform value at a specific time
@@ -639,10 +730,12 @@ public:
 	* @param ControlRig The ControlRig
 	* @param ControlName Name of the Control, should be a Transform control
 	* @param Frame Time to get the value
+	* @param TimeUnit Unit for frame values, either in display rate or tick resolution
 	* @return Returns Value at that time
 	*/
 	UFUNCTION(BlueprintCallable, Category = "Editor Scripting | Sequencer Tools | Control Rig")
-	static FTransform GetLocalControlRigTransform(ULevelSequence* LevelSequence, UControlRig* ControlRig, FName ControlName, FFrameNumber Frame);
+	static FTransform GetLocalControlRigTransform(ULevelSequence* LevelSequence, UControlRig* ControlRig, FName ControlName, FFrameNumber Frame, 
+		ESequenceTimeUnit TimeUnit = ESequenceTimeUnit::DisplayRate);
 
 	/**
 	* Get ControlRig Control's Transform values at specific times
@@ -650,10 +743,12 @@ public:
 	* @param ControlRig The ControlRig
 	* @param ControlName Name of the Control, should be a Transform control
 	* @param Frames Times to get the values
+	* @param TimeUnit Unit for frame values, either in display rate or tick resolution
 	* @return Returns Values at those times
 	*/
 	UFUNCTION(BlueprintCallable, Category = "Editor Scripting | Sequencer Tools | Control Rig")
-	static TArray<FTransform> GetLocalControlRigTransforms(ULevelSequence* LevelSequence, UControlRig* ControlRig, FName ControlName, const TArray<FFrameNumber>& Frames);
+	static TArray<FTransform> GetLocalControlRigTransforms(ULevelSequence* LevelSequence, UControlRig* ControlRig, FName ControlName, const TArray<FFrameNumber>& Frames,
+		ESequenceTimeUnit TimeUnit = ESequenceTimeUnit::DisplayRate);
 
 	/**
 	* Set ControlRig Control's Transform value at specific time
@@ -662,10 +757,12 @@ public:
 	* @param ControlName Name of the Control, should be a Transform control
 	* @param Frame Time to set the value
 	* @param Value to set 
+	* @param TimeUnit Unit for frame values, either in display rate or tick resolution
 	* @param bSetKey If True set a key, if not just set the value
 	*/
 	UFUNCTION(BlueprintCallable, Category = "Editor Scripting | Sequencer Tools | Control Rig")
-	static void SetLocalControlRigTransform(ULevelSequence* LevelSequence, UControlRig* ControlRig, FName ControlName, FFrameNumber Frame, FTransform Value, bool bSetKey = true);
+	static void SetLocalControlRigTransform(ULevelSequence* LevelSequence, UControlRig* ControlRig, FName ControlName, FFrameNumber Frame, FTransform Value, 
+		ESequenceTimeUnit TimeUnit = ESequenceTimeUnit::DisplayRate, bool bSetKey = true);
 
 	/**
 	* Set ControlRig Control's Transform values at specific times
@@ -674,8 +771,10 @@ public:
 	* @param ControlName Name of the Control, should be a Transform control
 	* @param Frames Times to set the values
 	* @param Values to set at those times
+	* @param TimeUnit Unit for frame values, either in display rate or tick resolution
 	*/
 	UFUNCTION(BlueprintCallable, Category = "Editor Scripting | Sequencer Tools | Control Rig")
-	static void SetLocalControlRigTransforms(ULevelSequence* LevelSequence, UControlRig* ControlRig, FName ControlName, const TArray<FFrameNumber>& Frames, const TArray<FTransform> Values);
+	static void SetLocalControlRigTransforms(ULevelSequence* LevelSequence, UControlRig* ControlRig, FName ControlName, const TArray<FFrameNumber>& Frames, const TArray<FTransform> Values,
+		ESequenceTimeUnit TimeUnit = ESequenceTimeUnit::DisplayRate);
 };
 
