Removed unnecessary forward declarations of FMobilityCustomization.

diff --git a/Engine/Plugins/Experimental/ControlRig/Source/ControlRigEditor/Private/ControlRigComponentDetails.h b/Engine/Plugins/Experimental/ControlRig/Source/ControlRigEditor/Private/ControlRigComponentDetails.h
--- a/Engine/Plugins/Experimental/ControlRig/Source/ControlRigEditor/Private/ControlRigComponentDetails.h
+++ b/Engine/Plugins/Experimental/ControlRig/Source/ControlRigEditor/Private/ControlRigComponentDetails.h
@@ -5,7 +5,6 @@
 #include "CoreMinimal.h"
 #include "IDetailCustomization.h"
 
-class FMobilityCustomization;
 class IDetailLayoutBuilder;
 
 class FControlRigComponentDetails : public IDetailCustomization
diff --git a/Engine/Plugins/Experimental/ControlRig/Source/ControlRigEditor/Private/ControlRigDrawingDetails.h b/Engine/Plugins/Experimental/ControlRig/Source/ControlRigEditor/Private/ControlRigDrawingDetails.h
--- a/Engine/Plugins/Experimental/ControlRig/Source/ControlRigEditor/Private/ControlRigDrawingDetails.h
+++ b/Engine/Plugins/Experimental/ControlRig/Source/ControlRigEditor/Private/ControlRigDrawingDetails.h
@@ -7,7 +7,6 @@
 #include "ControlRigBlueprint.h"
 #include "ControlRigDrawingDetails.generated.h"
 
-class FMobilityCustomization;
 class IDetailLayoutBuilder;
 
 USTRUCT()
diff --git a/Engine/Plugins/Experimental/ControlRig/Source/ControlRigEditor/Private/ControlRigGraphDetails.h b/Engine/Plugins/Experimental/ControlRig/Source/ControlRigEditor/Private/ControlRigGraphDetails.h
--- a/Engine/Plugins/Experimental/ControlRig/Source/ControlRigEditor/Private/ControlRigGraphDetails.h
+++ b/Engine/Plugins/Experimental/ControlRig/Source/ControlRigEditor/Private/ControlRigGraphDetails.h
@@ -13,7 +13,6 @@
 #include "Graph/SControlRigGraphNode.h"
 #include "Widgets/Colors/SColorBlock.h"
 
-class FMobilityCustomization;
 class IDetailLayoutBuilder;
 
 class FControlRigArgumentGroupLayout : public IDetailCustomNodeBuilder, public TSharedFromThis<FControlRigArgumentGroupLayout>
