Control Rig Debugging: Breakpoints

diff --git a/Engine/Plugins/Experimental/ControlRig/Source/ControlRig/Private/ControlRig.cpp b/Engine/Plugins/Experimental/ControlRig/Source/ControlRig/Private/ControlRig.cpp
--- a/Engine/Plugins/Experimental/ControlRig/Source/ControlRig/Private/ControlRig.cpp
+++ b/Engine/Plugins/Experimental/ControlRig/Source/ControlRig/Private/ControlRig.cpp
@@ -70,6 +70,9 @@ UControlRig::UControlRig(const FObjectInitializer& ObjectInitializer)
 	, PostSetupBracket(0)
 	, InteractionBracket(0)
 	, InterRigSyncBracket(0)
+#if WITH_EDITOR
+	, VMSnapshotBeforeExecution(nullptr)
+#endif
 {
 	SetVM(ObjectInitializer.CreateDefaultSubobject<URigVM>(this, TEXT("VM")));
 	DynamicHierarchy = ObjectInitializer.CreateDefaultSubobject<URigHierarchy>(this, TEXT("DynamicHierarchy"));
@@ -417,6 +420,13 @@ void UControlRig::InstantiateVMFromCDO()
 	{
 		SetVM(NewObject<URigVM>(this, TEXT("VM")));
 	}
+	
+#if WITH_EDITOR
+	if (VMSnapshotBeforeExecution == nullptr || VMSnapshotBeforeExecution->GetOuter() != this)
+	{
+		VMSnapshotBeforeExecution = NewObject<URigVM>(this);
+	}
+#endif
 
 	if (!HasAnyFlags(RF_ClassDefaultObject))
 	{
@@ -468,6 +478,15 @@ void UControlRig::Execute(const EControlRigState InState, const FName& InEventNa
 #endif
 	}
 
+#if WITH_EDITOR
+	if (UControlRig* CDO = GetClass()->GetDefaultObject<UControlRig>())
+	{
+		// Copy the breakpoints. This will not override the state of the breakpoints
+		DebugInfo.SetBreakpoints(CDO->DebugInfo.GetBreakpoints());
+	}
+	VM->SetDebugInfo(&DebugInfo);
+#endif
+
 	bool bJustRanInit = false;
 	if (bRequiresInitExecution)
 	{
@@ -943,6 +962,16 @@ void UControlRig::ExecuteUnits(FRigUnitContext& InOutContext, const FName& InEve
 		}
 		else
 		{
+#if WITH_EDITOR
+			if (VM->GetHaltedAtInstruction() != INDEX_NONE)
+			{
+				VM->CopyFrom(VMSnapshotBeforeExecution, false, false, false, true, true);	
+			}
+			else
+			{
+				VMSnapshotBeforeExecution->CopyFrom(VM, false, false, false, true, true);
+			}
+#endif
 			VM->Execute(FRigVMMemoryContainerPtrArray(LocalMemory, 2), AdditionalArguments, InEventName);
 		}
 	}
@@ -1169,6 +1198,11 @@ void UControlRig::HandleOnControlModified(UControlRig* Subject, FRigControlEleme
 
 void UControlRig::HandleExecutionReachedExit()
 {
+#if WITH_EDITOR
+	VMSnapshotBeforeExecution->CopyFrom(VM, false, false, false, true, true);
+	DebugInfo.ResetState();
+#endif
+	
 	if (LatestExecutedState != EControlRigState::Init && bAccumulateTime)
 	{
 		AbsoluteTime += DeltaTime;
@@ -2329,6 +2363,17 @@ FRigVMExternalVariable UControlRig::GetExternalVariableFromDescription(const FBP
 	return GetExternalVariableFromPinType(InVariableDescription.VarName, InVariableDescription.VarType, bIsPublic, bIsReadOnly);
 }
 
+void UControlRig::AddBreakpoint(int32 InstructionIndex)
+{
+	DebugInfo.AddBreakpoint(InstructionIndex);
+}
+
+void UControlRig::ResumeExecution()
+{
+	VM->CopyFrom(VMSnapshotBeforeExecution, false, false, false, true, true);
+	VM->ResumeExecution();
+}
+
 #endif // WITH_EDITOR
 
 void UControlRig::SetBoneInitialTransformsFromSkeletalMeshComponent(USkeletalMeshComponent* InSkelMeshComp)
diff --git a/Engine/Plugins/Experimental/ControlRig/Source/ControlRig/Public/ControlRig.h b/Engine/Plugins/Experimental/ControlRig/Source/ControlRig/Public/ControlRig.h
--- a/Engine/Plugins/Experimental/ControlRig/Source/ControlRig/Public/ControlRig.h
+++ b/Engine/Plugins/Experimental/ControlRig/Source/ControlRig/Public/ControlRig.h
@@ -546,11 +546,6 @@ public:
 	UFUNCTION(BlueprintSetter)
 	void SetInteractionRigClass(TSubclassOf<UControlRig> InInteractionRigClass);
 
-	// UObject interface
-#if WITH_EDITOR
-	virtual void PreEditChange(FProperty* PropertyAboutToChange) override;
-	virtual void PostEditChangeProperty(FPropertyChangedEvent& PropertyChangedEvent) override;
-#endif
 
 public:
 	//~ Begin IInterface_AssetUserData Interface
@@ -638,6 +633,36 @@ protected:
 
 	TArray<FRigElementKey> QueuedModifiedControls;
 
+private:
+
+	UPROPERTY(transient)
+	URigVM* VMSnapshotBeforeExecution;
+	
+#if WITH_EDITOR	
+
+	FRigVMDebugInfo DebugInfo;
+	
+public:	
+	
+	/** Adds a breakpoint in the VM at the InstructionIndex */
+	void AddBreakpoint(int32 InstructionIndex);
+
+	/** Reloads the VM from the SnapshotVM, and increments activation hit on the halted breakpoint */
+	void ResumeExecution();
+	
+	FRigVMDebugInfo& GetDebugInfo() { return DebugInfo; }
+
+	URigVM* GetVMSnapshot() { return VMSnapshotBeforeExecution; }
+
+	void SetVMSnapshot(URigVM* InVMSnapshot) { VMSnapshotBeforeExecution = InVMSnapshot; }
+#endif	
+
+	// UObject interface
+	#if WITH_EDITOR
+	virtual void PreEditChange(FProperty* PropertyAboutToChange) override;
+	virtual void PostEditChangeProperty(FPropertyChangedEvent& PropertyChangedEvent) override;
+#endif
+
 	friend class FControlRigBlueprintCompilerContext;
 	friend struct FRigHierarchyRef;
 	friend class FControlRigEditor;
diff --git a/Engine/Plugins/Experimental/ControlRig/Source/ControlRigDeveloper/Private/ControlRigBlueprint.cpp b/Engine/Plugins/Experimental/ControlRig/Source/ControlRigDeveloper/Private/ControlRigBlueprint.cpp
--- a/Engine/Plugins/Experimental/ControlRig/Source/ControlRigDeveloper/Private/ControlRigBlueprint.cpp
+++ b/Engine/Plugins/Experimental/ControlRig/Source/ControlRigDeveloper/Private/ControlRigBlueprint.cpp
@@ -618,6 +618,13 @@ void UControlRigBlueprint::RecompileVM()
 			CDO->SetVM(NewObject<URigVM>(CDO, TEXT("VM")));
 		}
 
+#if WITH_EDITOR
+		if (CDO->GetVMSnapshot() == nullptr || CDO->GetVMSnapshot()->GetOuter() != CDO)
+		{
+			CDO->SetVMSnapshot(NewObject<URigVM>(CDO));
+		}
+#endif
+
 		if (!HasAnyFlags(RF_Transient | RF_Transactional))
 		{
 			CDO->Modify(false);
@@ -672,6 +679,10 @@ void UControlRigBlueprint::RecompileVM()
 		bVMRecompilationRequired = false;
 		VMRecompilationBracket = 0;
 		VMCompiledEvent.Broadcast(this, CDO->VM);
+
+#if WITH_EDITOR
+		RefreshControlRigBreakpoints();
+#endif
 	}
 }
 
@@ -785,6 +796,204 @@ void UControlRigBlueprint::HandleReportFromCompiler(EMessageSeverity::Type InSev
 	}
 }
 
+#if WITH_EDITOR
+
+void UControlRigBlueprint::ClearBreakpoints()
+{
+	Breakpoints.Empty();
+	RefreshControlRigBreakpoints();
+}
+
+bool UControlRigBlueprint::AddBreakpoint(const FString& InBreakpointNodePath)
+{
+	URigVMLibraryNode* FunctionNode = nullptr;
+	
+	// Find the node in the graph
+	URigVMNode* BreakpointNode = GetModel()->FindNode(InBreakpointNodePath);
+	if (BreakpointNode == nullptr)
+	{
+		// If we cannot find the node, it might be because it is inside a function
+		FString FunctionName = InBreakpointNodePath, Right;
+		URigVMNode::SplitNodePathAtStart(InBreakpointNodePath, FunctionName, Right);
+
+		// Look inside the local function library
+		if (URigVMLibraryNode* LibraryNode = GetLocalFunctionLibrary()->FindFunction(FName(FunctionName)))
+		{
+			BreakpointNode = LibraryNode->GetContainedGraph()->FindNode(Right);
+			FunctionNode = LibraryNode;
+		}
+	}
+
+	return AddBreakpoint(BreakpointNode, FunctionNode);
+}
+
+bool UControlRigBlueprint::AddBreakpoint(URigVMNode* InBreakpointNode, URigVMLibraryNode* LibraryNode)
+{
+	if (InBreakpointNode == nullptr)
+	{
+		return false;
+	}
+
+	bool bSuccess = true;
+	if (LibraryNode)
+	{
+		// If the breakpoint node is inside a library node, find all references to the library node
+		TArray<TSoftObjectPtr<URigVMFunctionReferenceNode>> References = LibraryNode->GetLibrary()->GetReferencesForFunction(LibraryNode->GetFName());
+		for (TSoftObjectPtr<URigVMFunctionReferenceNode> Reference : References)
+		{
+			if (!Reference.IsValid())
+			{
+				continue;
+			}
+
+			UControlRigBlueprint* ReferenceBlueprint = Reference->GetTypedOuter<UControlRigBlueprint>();
+
+			// If the reference is not inside another function, add a breakpoint in the blueprint containing the
+			// reference, without a function specified
+			bool bIsInsideFunction = Reference->GetRootGraph()->IsA<URigVMFunctionLibrary>();
+			if(!bIsInsideFunction)
+			{
+				bSuccess &= ReferenceBlueprint->AddBreakpoint(InBreakpointNode);
+			}
+			else
+			{
+				// Otherwise, we need to add breakpoints to all the blueprints that reference this
+				// function (when the blueprint graph is flattened)
+				
+				// Get all the functions containing this reference
+				URigVMNode* Node = Reference.Get();
+				while (Node->GetGraph() != ReferenceBlueprint->GetLocalFunctionLibrary())
+				{
+					if (URigVMLibraryNode* ParentLibraryNode = Cast<URigVMLibraryNode>(Node->GetGraph()->GetOuter()))
+					{
+						// Recursively add breakpoints to the reference blueprint, specifying the parent function
+						bSuccess &= ReferenceBlueprint->AddBreakpoint(InBreakpointNode, ParentLibraryNode);
+					}
+
+					Node = Cast<URigVMNode>(Node->GetGraph()->GetOuter());
+				}
+			}
+		}
+	}
+	else
+	{
+		if (!Breakpoints.Contains(InBreakpointNode))
+		{
+			// Add the breakpoint to the VM
+			bSuccess = AddBreakpointToControlRig(InBreakpointNode);
+		}
+	}
+
+	return bSuccess;
+}
+
+bool UControlRigBlueprint::AddBreakpointToControlRig(URigVMNode* InBreakpointNode)
+{
+	UControlRigBlueprintGeneratedClass* RigClass = GetControlRigBlueprintGeneratedClass();
+	UControlRig* CDO = Cast<UControlRig>(RigClass->GetDefaultObject(false));
+	const FRigVMByteCode* ByteCode = GetController()->GetCurrentByteCode();
+	TSet<FString> AddedCallpaths;
+
+	if (CDO && ByteCode)
+	{
+		FRigVMInstructionArray Instructions = ByteCode->GetInstructions();	
+
+		// For each instruction, see if the node is in the callpath
+		// Only add one breakpoint for each callpath related to this node (i.e. if a node produces multiple
+		// instructions, only add a breakpoint to the first instruction)
+		for (int32 i = 0; i< Instructions.Num(); ++i)
+		{
+			const FRigVMASTProxy Proxy = FRigVMASTProxy::MakeFromCallPath(ByteCode->GetCallPathForInstruction(i), GetModel());
+			if (Proxy.GetCallstack().Contains(InBreakpointNode))
+			{
+				// Find the callpath related to the breakpoint node
+				FRigVMASTProxy BreakpointProxy = Proxy;
+				while(BreakpointProxy.GetSubject() != InBreakpointNode)
+				{
+					BreakpointProxy = BreakpointProxy.GetParent();
+				}
+				const FString& BreakpointCallPath = BreakpointProxy.GetCallstack().GetCallPath();
+
+				// Only add this callpath breakpoint once
+				if (!AddedCallpaths.Contains(BreakpointCallPath))
+				{
+					AddedCallpaths.Add(BreakpointCallPath);
+					CDO->AddBreakpoint(i);
+				}
+			}
+		}
+	}
+
+	if (AddedCallpaths.Num() > 0)
+	{
+		Breakpoints.AddUnique(InBreakpointNode);
+		return true;
+	}
+	
+	return false;
+}
+
+bool UControlRigBlueprint::RemoveBreakpoint(const FString& InBreakpointNodePath)
+{
+	// Find the node in the graph
+	URigVMNode* BreakpointNode = GetModel()->FindNode(InBreakpointNodePath);
+	if (BreakpointNode == nullptr)
+	{
+		// If we cannot find the node, it might be because it is inside a function
+		FString FunctionName = InBreakpointNodePath, Right;
+		URigVMNode::SplitNodePathAtStart(InBreakpointNodePath, FunctionName, Right);
+
+		// Look inside the local function library
+		if (URigVMLibraryNode* LibraryNode = GetLocalFunctionLibrary()->FindFunction(FName(FunctionName)))
+		{
+			BreakpointNode = LibraryNode->GetContainedGraph()->FindNode(Right);
+		}
+	}
+
+	bool bSuccess = RemoveBreakpoint(BreakpointNode);
+
+	// Remove the breakpoint from all the loaded dependent blueprints
+	TArray<UControlRigBlueprint*> DependentBlueprints = GetDependentBlueprints(true, true);
+	DependentBlueprints.Remove(this);
+	for (UControlRigBlueprint* Dependent : DependentBlueprints)
+	{
+		bSuccess &= Dependent->RemoveBreakpoint(BreakpointNode);
+	}
+	return bSuccess;
+}
+
+bool UControlRigBlueprint::RemoveBreakpoint(URigVMNode* InBreakpointNode)
+{
+	if (Breakpoints.Contains(InBreakpointNode))
+	{
+		Breakpoints.Remove(InBreakpointNode);
+
+		// Multiple breakpoint nodes might set a breakpoint to the same instruction. When we remove
+		// one of the breakpoint nodes, we do not want to remove the instruction breakpoint if there
+		// is another breakpoint node addressing it. For that reason, we just recompute all the
+		// breakpoint instructions.
+		// Refreshing breakpoints in the control rig will keep the state it had before.
+		RefreshControlRigBreakpoints();
+		return true;
+	}
+
+	return false;
+}
+
+
+void UControlRigBlueprint::RefreshControlRigBreakpoints()
+{
+	UControlRigBlueprintGeneratedClass* RigClass = GetControlRigBlueprintGeneratedClass();
+	UControlRig* CDO = Cast<UControlRig>(RigClass->GetDefaultObject(false));
+	CDO->GetDebugInfo().Clear();
+	for (URigVMNode* Node : Breakpoints)
+	{
+		AddBreakpointToControlRig(Node);
+	}
+}
+
+#endif
+
 void UControlRigBlueprint::RequestControlRigInit()
 {
 	UControlRigBlueprintGeneratedClass* RigClass = GetControlRigBlueprintGeneratedClass();
@@ -1298,25 +1507,28 @@ TArray<FAssetData> UControlRigBlueprint::GetDependentAssets() const
 	return Dependents;
 }
 
-TArray<UControlRigBlueprint*> UControlRigBlueprint::GetDependentBlueprints(bool bRecursive) const
+TArray<UControlRigBlueprint*> UControlRigBlueprint::GetDependentBlueprints(bool bRecursive, bool bOnlyLoaded) const
 {
 	TArray<FAssetData> Assets = GetDependentAssets();
 	TArray<UControlRigBlueprint*> Dependents;
 
 	for(const FAssetData& Asset : Assets)
 	{
-		if(UControlRigBlueprint* Dependent = Cast<UControlRigBlueprint>(Asset.GetAsset()))
+		if (!bOnlyLoaded || Asset.IsAssetLoaded())
 		{
-			if(!Dependents.Contains(Dependent))
+			if(UControlRigBlueprint* Dependent = Cast<UControlRigBlueprint>(Asset.GetAsset()))
 			{
-				Dependents.Add(Dependent);
-
-				if(bRecursive)
+				if(!Dependents.Contains(Dependent))
 				{
-					TArray<UControlRigBlueprint*> ParentDependents = Dependent->GetDependentBlueprints(true);
-					for(UControlRigBlueprint* ParentDependent : ParentDependents)
+					Dependents.Add(Dependent);
+
+					if(bRecursive && Dependent != this)
 					{
-						Dependents.AddUnique(ParentDependent);
+						TArray<UControlRigBlueprint*> ParentDependents = Dependent->GetDependentBlueprints(true);
+						for(UControlRigBlueprint* ParentDependent : ParentDependents)
+						{
+							Dependents.AddUnique(ParentDependent);
+						}
 					}
 				}
 			}
@@ -2259,6 +2471,14 @@ void UControlRigBlueprint::HandleModifiedEvent(ERigVMGraphNotifType InNotifType,
 			case ERigVMGraphNotifType::NodeAdded:
 			case ERigVMGraphNotifType::NodeRemoved:
 			{
+				if (InNotifType == ERigVMGraphNotifType::NodeRemoved)
+				{
+					if (URigVMNode* RigVMNode = Cast<URigVMNode>(InSubject))
+					{
+						RemoveBreakpoint(RigVMNode);
+					}
+				}
+					
 				if (URigVMCollapseNode* CollapseNode = Cast<URigVMCollapseNode>(InSubject))
 				{
 					if (InNotifType == ERigVMGraphNotifType::NodeAdded)
diff --git a/Engine/Plugins/Experimental/ControlRig/Source/ControlRigDeveloper/Public/ControlRigBlueprint.h b/Engine/Plugins/Experimental/ControlRig/Source/ControlRigDeveloper/Public/ControlRigBlueprint.h
--- a/Engine/Plugins/Experimental/ControlRig/Source/ControlRigDeveloper/Public/ControlRigBlueprint.h
+++ b/Engine/Plugins/Experimental/ControlRig/Source/ControlRigDeveloper/Public/ControlRigBlueprint.h
@@ -261,9 +261,9 @@ public:
 
 	// Returns a list of dependents as resolved blueprints.
 	// A dependent is a blueprint which uses a function defined in this blueprint.
-	// This function loads the dependent assets and can introduce a large cost
+	// If bOnlyLoaded is false, this function loads the dependent assets and can introduce a large cost
 	// depending on the size / count of assets in the project.
-	TArray<UControlRigBlueprint*> GetDependentBlueprints(bool bRecursive = false) const;
+	TArray<UControlRigBlueprint*> GetDependentBlueprints(bool bRecursive = false, bool bOnlyLoaded = false) const;
 
 	UPROPERTY(EditAnywhere, Category = "User Interface")
 	FRigGraphDisplaySettings RigGraphDisplaySettings;
@@ -533,6 +533,38 @@ private:
 	bool RemoveEdGraphForCollapseNode(URigVMCollapseNode* InNode, bool bNotify = false);
 	void HandleReportFromCompiler(EMessageSeverity::Type InSeverity, UObject* InSubject, const FString& InMessage);
 
+#if WITH_EDITOR
+private:
+	TArray<URigVMNode*> Breakpoints;
+
+public:
+
+	/** Removes all the breakpoints from the blueprint and the VM */
+	void ClearBreakpoints();
+
+	/** Adds a breakpoint to all loaded blueprints which use the node indicated by InBreakpointNodePath 
+	  * If the node is inside a public function, it will add a breakpoint to all blueprints calling this function. */
+	bool AddBreakpoint(const FString& InBreakpointNodePath);
+
+	/** Adds a breakpoint to all loaded blueprints which use the InBreakpointNode. 
+	  * If LibraryNode is not null, it indicates that the library uses the InBreapointNode, and the function will add
+	  * breakpoints to any other loaded blueprint that references this library. */
+	bool AddBreakpoint(URigVMNode* InBreakpointNode, URigVMLibraryNode* LibraryNode = nullptr);
+
+	/** Adds a breakpoint to the first instruction of each callpath related to the InBreakpointNode */
+	bool AddBreakpointToControlRig(URigVMNode* InBreakpointNode);
+
+	/** Removes the given breakpoint from all the loaded blueprints that use this node, and recomputes all breakpoints
+	  * in the VM  */
+	bool RemoveBreakpoint(const FString& InBreakpointNodePath);
+	bool RemoveBreakpoint(URigVMNode* InBreakpointNode);
+
+	/** Recomputes the instruction breakpoints given the node breakpoints in the blueprint */
+	void RefreshControlRigBreakpoints();
+
+#endif
+
+private:
 	bool bDirtyDuringLoad;
 	bool bErrorsDuringCompilation;
 
diff --git a/Engine/Source/Runtime/RigVM/Private/RigVMCore/RigVM.cpp b/Engine/Source/Runtime/RigVM/Private/RigVMCore/RigVM.cpp
--- a/Engine/Source/Runtime/RigVM/Private/RigVMCore/RigVM.cpp
+++ b/Engine/Source/Runtime/RigVM/Private/RigVMCore/RigVM.cpp
@@ -65,6 +65,10 @@ URigVM::URigVM()
     : WorkMemoryPtr(&WorkMemoryStorage)
     , LiteralMemoryPtr(&LiteralMemoryStorage)
     , ByteCodePtr(&ByteCodeStorage)
+#if WITH_EDITOR
+	, DebugInfo(nullptr)
+	, HaltedAtInstruction(INDEX_NONE)
+#endif
     , FunctionNamesPtr(&FunctionNamesStorage)
     , FunctionsPtr(&FunctionsStorage)
 #if WITH_EDITOR
@@ -209,14 +213,11 @@ void URigVM::CopyFrom(URigVM* InVM, bool bDeferCopy, bool bReferenceLiteralMemor
 
 	if(InVM->WorkMemoryPtr == &InVM->WorkMemoryStorage)
 	{
+		WorkMemoryStorage = InVM->WorkMemoryStorage;
 		if (bCopyDynamicRegisters)
 		{
 			WorkMemoryStorage.CopyRegisters(InVM->WorkMemoryStorage);
 		}
-		else
-		{
-			WorkMemoryStorage = InVM->WorkMemoryStorage;			
-		}
 		WorkMemoryPtr = &WorkMemoryStorage;
 	}
 	else
@@ -325,6 +326,27 @@ TArray<FName> URigVM::GetEntryNames() const
 	return EntryNames;
 }
 
+#if WITH_EDITOR
+bool URigVM::ResumeExecution()
+{
+	HaltedAtInstruction = INDEX_NONE;
+	if (FRigVMBreakpoint* Breakpoint = DebugInfo->FindBreakpoint(Context.InstructionIndex))
+	{
+		DebugInfo->IncrementBreakpointActivationOnHit(Context.InstructionIndex);
+		return true;
+	}
+
+	return false;
+}
+
+bool URigVM::ResumeExecution(FRigVMMemoryContainerPtrArray Memory, FRigVMFixedArray<void*> AdditionalArguments, const FName& InEntryName)
+{
+	ResumeExecution();
+	return Execute(Memory, AdditionalArguments, InEntryName);
+}
+
+#endif
+
 const TArray<FRigVMParameter>& URigVM::GetParameters() const
 {
 	return Parameters;
@@ -731,8 +753,35 @@ bool URigVM::Initialize(FRigVMMemoryContainerPtrArray Memory, FRigVMFixedArray<v
 
 	CacheDebugWatchRegisters();
 	
+#if WITH_EDITOR
+	if (DebugInfo)
+	{
+		DebugInfo->StartExecution();
+		HaltedAtInstruction = INDEX_NONE;
+	}
+#endif
+
 	while (Instructions.IsValidIndex(Context.InstructionIndex))
 	{
+#if WITH_EDITOR
+		if (DebugInfo)
+		{
+			if (FRigVMBreakpoint* Breakpoint = DebugInfo->FindBreakpoint(Context.InstructionIndex))
+			{
+				if (DebugInfo->IsActive(Context.InstructionIndex))
+				{
+					HaltedAtInstruction = Context.InstructionIndex;
+					ExecutionHalted().Broadcast(Context.InstructionIndex);
+					return true;
+				}
+				else
+				{
+					DebugInfo->HitBreakpoint(Context.InstructionIndex);
+				}
+			}
+		}
+#endif
+		
 		const FRigVMInstruction& Instruction = Instructions[Context.InstructionIndex];
 
 		switch (Instruction.OpCode)
@@ -1012,9 +1061,34 @@ bool URigVM::Execute(FRigVMMemoryContainerPtrArray Memory, FRigVMFixedArray<void
 		Context.InstructionIndex = (uint16)ByteCode.GetEntry(EntryIndex).InstructionIndex;
 	}
 
+#if WITH_EDITOR
+	if (DebugInfo)
+	{
+		DebugInfo->StartExecution();
+		HaltedAtInstruction = INDEX_NONE;
+	}
+#endif
+
 	while (Instructions.IsValidIndex(Context.InstructionIndex))
 	{
 #if WITH_EDITOR
+		if (DebugInfo)
+		{
+			if (FRigVMBreakpoint* Breakpoint = DebugInfo->FindBreakpoint(Context.InstructionIndex))
+			{
+				if (DebugInfo->IsActive(Context.InstructionIndex))
+				{
+					HaltedAtInstruction = Context.InstructionIndex;
+					ExecutionHalted().Broadcast(Context.InstructionIndex);
+					return true;
+				}
+				else
+				{
+					DebugInfo->HitBreakpoint(Context.InstructionIndex);
+				}
+			}
+		}
+		
 		InstructionVisitedDuringLastRun[Context.InstructionIndex]++;
 		InstructionVisitOrder.Add(Context.InstructionIndex);
 #endif
diff --git a/Engine/Source/Runtime/RigVM/Private/RigVMCore/RigVMMemory.cpp b/Engine/Source/Runtime/RigVM/Private/RigVMCore/RigVMMemory.cpp
--- a/Engine/Source/Runtime/RigVM/Private/RigVMCore/RigVMMemory.cpp
+++ b/Engine/Source/Runtime/RigVM/Private/RigVMCore/RigVMMemory.cpp
@@ -1186,6 +1186,12 @@ bool FRigVMMemoryContainer::Copy(
 	uint8* TargetPtr = GetDataPtr(Target, InTargetRegisterOffset, InTargetSliceIndex, false);
 	uint16 NumBytes = Target.GetNumBytesPerSlice();
 
+	// TODO: preventing crash when SourcePtr is null
+	if (SourcePtr == nullptr)
+	{
+		return false;
+	}
+	
 	if (Source.IsNestedDynamic())
 	{
 #if WITH_EDITOR
diff --git a/Engine/Source/Runtime/RigVM/Public/RigVMCore/RigVM.h b/Engine/Source/Runtime/RigVM/Public/RigVMCore/RigVM.h
--- a/Engine/Source/Runtime/RigVM/Public/RigVMCore/RigVM.h
+++ b/Engine/Source/Runtime/RigVM/Public/RigVMCore/RigVM.h
@@ -8,6 +8,9 @@
 #include "RigVMRegistry.h"
 #include "RigVMByteCode.h"
 #include "RigVMStatistics.h"
+#if WITH_EDITOR
+#include "RigVMDebugInfo.h"
+#endif
 #include "RigVM.generated.h"
 
 // The type of parameter for a VM
@@ -120,7 +123,10 @@ public:
 
 	/** Bindable event for external objects to be notified when the VM reaches an Exit Operation */
 	DECLARE_EVENT(URigVM, FExecutionReachedExitEvent);
-	
+#if WITH_EDITOR
+	DECLARE_EVENT_OneParam(URigVM, FExecutionHaltedEvent, int32);
+#endif
+
 	URigVM();
 	virtual ~URigVM();
 
@@ -213,6 +219,10 @@ public:
 	FORCEINLINE const TArray<int32> GetInstructionVisitOrder() const { return InstructionVisitOrder; }
 
 	FORCEINLINE const void SetFirstEntryEventInEventQueue(const FName& InFirstEventName) { FirstEntryEventInQueue = InFirstEventName; }
+
+	bool ResumeExecution(FRigVMMemoryContainerPtrArray Memory, FRigVMFixedArray<void*> AdditionalArguments, const FName& InEntryName = NAME_None);
+
+	bool ResumeExecution();
 #endif
 
 	// Returns the parameters of the VM
@@ -505,6 +515,13 @@ public:
 #endif
 
 	FExecutionReachedExitEvent& ExecutionReachedExit() { return OnExecutionReachedExit; }
+#if WITH_EDITOR
+	FExecutionHaltedEvent& ExecutionHalted() { return OnExecutionHalted; }
+
+	void SetDebugInfo(FRigVMDebugInfo* InDebugInfo) { DebugInfo = InDebugInfo; }
+
+	int32 GetHaltedAtInstruction() const { return HaltedAtInstruction; }
+#endif
 
 private:
 
@@ -519,6 +536,11 @@ private:
 	UPROPERTY(transient)
 	FRigVMExecuteContext Context;
 
+#if WITH_EDITOR
+	FRigVMDebugInfo* DebugInfo;
+	int32 HaltedAtInstruction;
+#endif
+
 	UPROPERTY()
 	TArray<FName> FunctionNamesStorage;
 	TArray<FName>* FunctionNamesPtr;
@@ -570,6 +592,10 @@ private:
 	void CopyDeferredVMIfRequired();
 
 	FExecutionReachedExitEvent OnExecutionReachedExit;
+#if WITH_EDITOR
+	FExecutionHaltedEvent OnExecutionHalted;
+#endif
+
 	
 	friend class URigVMCompiler;
 };
diff --git a/Engine/Source/Runtime/RigVM/Public/RigVMCore/RigVMDebugInfo.h b/Engine/Source/Runtime/RigVM/Public/RigVMCore/RigVMDebugInfo.h
new file mode 100644
--- /dev/null
+++ b/Engine/Source/Runtime/RigVM/Public/RigVMCore/RigVMDebugInfo.h
@@ -0,0 +1,126 @@
+﻿// Copyright Epic Games, Inc. All Rights Reserved.
+
+#pragma once
+
+#include "CoreMinimal.h"
+#include "RigVMDebugInfo.generated.h"
+
+
+USTRUCT()
+struct RIGVM_API FRigVMBreakpoint
+{
+	GENERATED_BODY()
+	
+	FORCEINLINE FRigVMBreakpoint()
+	: bIsActive(true), bIsTemporary(false) 
+	{
+
+	}
+
+	bool bIsActive;
+	bool bIsTemporary;
+};
+
+USTRUCT()
+struct RIGVM_API FRigVMDebugInfo
+{
+	GENERATED_BODY()
+
+	FORCEINLINE FRigVMDebugInfo()
+	{
+	}
+
+	FORCEINLINE void ResetState()
+	{
+		BreakpointHits.Empty();
+		BreakpointActivationOnHit.Empty();
+	}
+
+	FORCEINLINE void StartExecution()
+	{
+		BreakpointHits.Empty();
+	}
+
+	FORCEINLINE void Clear()
+	{
+		Breakpoints.Empty();
+		// Do not remove state
+	}
+
+	FRigVMBreakpoint* FindBreakpoint(const uint16 InstructionIndex) 
+	{
+		return Breakpoints.Find(InstructionIndex);
+	}
+	
+	FRigVMBreakpoint* AddBreakpoint(const uint16 InstructionIndex)
+	{
+		if (!Breakpoints.Contains(InstructionIndex))
+		{
+			if (FRigVMBreakpoint* NewBP = &Breakpoints.Add(InstructionIndex))
+			{
+				// Do not override the state if it already exists
+				BreakpointActivationOnHit.FindOrAdd(InstructionIndex, 0);
+				BreakpointHits.FindOrAdd(InstructionIndex, 0);
+				return NewBP;
+			}
+		}
+		return nullptr;
+	}
+
+	TMap<uint16 , FRigVMBreakpoint> GetBreakpoints() const { return Breakpoints; }
+
+	void SetBreakpoints(const TMap<uint16 , FRigVMBreakpoint>& InBreakpoints) { Breakpoints = InBreakpoints; }
+
+	bool HasBreakpoint(const uint16 InstructionIndex) { return Breakpoints.Contains(InstructionIndex); }
+
+	bool IsActive(const uint16 InstructionIndex)
+	{
+		if (FRigVMBreakpoint* BP = Breakpoints.Find(InstructionIndex))
+		{
+			if (BP->bIsActive)
+			{
+				uint16 Hits = 0;
+				uint16 OnHit = 0;
+				if (BreakpointHits.Contains(InstructionIndex))
+				{
+					Hits = BreakpointHits.FindChecked(InstructionIndex);
+				}
+				if (BreakpointActivationOnHit.Contains(InstructionIndex))
+				{
+					OnHit = BreakpointActivationOnHit.FindChecked(InstructionIndex);
+				}
+				return Hits == OnHit;
+			}				
+		}
+		return false;		
+	}
+
+	void HitBreakpoint(const uint16 InstructionIndex)
+	{
+		if (BreakpointHits.Contains(InstructionIndex))
+		{
+			BreakpointHits[InstructionIndex]++;
+		}
+		else
+		{
+			BreakpointHits.Add(InstructionIndex, 1);
+		}
+	}
+
+	void IncrementBreakpointActivationOnHit(const uint16 InstructionIndex)
+	{
+		if (BreakpointActivationOnHit.Contains(InstructionIndex))
+		{
+			BreakpointActivationOnHit[InstructionIndex]++;
+		}
+		else
+		{
+			BreakpointActivationOnHit.Add(InstructionIndex, 1);
+		}
+	}
+
+private:
+	TMap<uint16 , FRigVMBreakpoint> Breakpoints;
+	TMap<uint16 , uint16> BreakpointActivationOnHit; // After how many instruction executions, this breakpoint becomes active
+	TMap<uint16 , uint16> BreakpointHits; // How many times this instruction has been executed
+};
