Control Rig: Create the snapshot VM lazily whenever a breakpoint is added

diff --git a/Engine/Plugins/Experimental/ControlRig/Source/ControlRig/Private/ControlRig.cpp b/Engine/Plugins/Experimental/ControlRig/Source/ControlRig/Private/ControlRig.cpp
--- a/Engine/Plugins/Experimental/ControlRig/Source/ControlRig/Private/ControlRig.cpp
+++ b/Engine/Plugins/Experimental/ControlRig/Source/ControlRig/Private/ControlRig.cpp
@@ -421,13 +421,6 @@ void UControlRig::InstantiateVMFromCDO()
 		SetVM(NewObject<URigVM>(this, TEXT("VM")));
 	}
 	
-#if WITH_EDITOR
-	if (VMSnapshotBeforeExecution == nullptr || VMSnapshotBeforeExecution->GetOuter() != this)
-	{
-		VMSnapshotBeforeExecution = NewObject<URigVM>(this);
-	}
-#endif
-
 	if (!HasAnyFlags(RF_ClassDefaultObject))
 	{
 		UControlRig* CDO = GetClass()->GetDefaultObject<UControlRig>();
@@ -963,13 +956,16 @@ void UControlRig::ExecuteUnits(FRigUnitContext& InOutContext, const FName& InEve
 		else
 		{
 #if WITH_EDITOR
-			if (VM->GetHaltedAtInstruction() != INDEX_NONE)
-			{
-				VM->CopyFrom(VMSnapshotBeforeExecution, false, false, false, true, true);	
-			}
-			else
+			if(URigVM* SnapShotVM = GetSnapshotVM(false)) // don't create it for normal runs
 			{
-				VMSnapshotBeforeExecution->CopyFrom(VM, false, false, false, true, true);
+				if (VM->GetHaltedAtInstruction() != INDEX_NONE)
+				{
+					VM->CopyFrom(SnapShotVM, false, false, false, true, true);	
+				}
+				else
+				{
+					SnapShotVM->CopyFrom(VM, false, false, false, true, true);
+				}
 			}
 #endif
 			VM->Execute(FRigVMMemoryContainerPtrArray(LocalMemory, 2), AdditionalArguments, InEventName);
@@ -1199,7 +1195,10 @@ void UControlRig::HandleOnControlModified(UControlRig* Subject, FRigControlEleme
 void UControlRig::HandleExecutionReachedExit()
 {
 #if WITH_EDITOR
-	VMSnapshotBeforeExecution->CopyFrom(VM, false, false, false, true, true);
+	if(URigVM* SnapShotVM = GetSnapshotVM())
+	{
+		SnapShotVM->CopyFrom(VM, false, false, false, true, true);
+	}
 	DebugInfo.ResetState();
 #endif
 	
@@ -2363,14 +2362,41 @@ FRigVMExternalVariable UControlRig::GetExternalVariableFromDescription(const FBP
 	return GetExternalVariableFromPinType(InVariableDescription.VarName, InVariableDescription.VarType, bIsPublic, bIsReadOnly);
 }
 
+URigVM* UControlRig::GetSnapshotVM(bool bCreateIfNeeded)
+{
+#if WITH_EDITOR
+	if(VMSnapshotBeforeExecution != nullptr)
+	{
+		if(VMSnapshotBeforeExecution->GetOuter() != this)
+		{
+			VMSnapshotBeforeExecution = nullptr;
+		}
+	}
+	if ((VMSnapshotBeforeExecution == nullptr) && bCreateIfNeeded)
+	{
+		VMSnapshotBeforeExecution = NewObject<URigVM>(this);
+	}
+	return VMSnapshotBeforeExecution;
+#else
+	return nullptr;
+#endif
+}
+
 void UControlRig::AddBreakpoint(int32 InstructionIndex)
 {
-	DebugInfo.AddBreakpoint(InstructionIndex);
+	// this makes sure that the snapshot exists
+	if(URigVM* SnapShotVM = GetSnapshotVM())
+	{
+		DebugInfo.AddBreakpoint(InstructionIndex);
+	}
 }
 
 void UControlRig::ResumeExecution()
 {
-	VM->CopyFrom(VMSnapshotBeforeExecution, false, false, false, true, true);
+	if(URigVM* SnapShotVM = GetSnapshotVM())
+	{
+		VM->CopyFrom(SnapShotVM, false, false, false, true, true);
+	}
 	VM->ResumeExecution();
 }
 
diff --git a/Engine/Plugins/Experimental/ControlRig/Source/ControlRig/Public/ControlRig.h b/Engine/Plugins/Experimental/ControlRig/Source/ControlRig/Public/ControlRig.h
--- a/Engine/Plugins/Experimental/ControlRig/Source/ControlRig/Public/ControlRig.h
+++ b/Engine/Plugins/Experimental/ControlRig/Source/ControlRig/Public/ControlRig.h
@@ -635,8 +635,12 @@ protected:
 
 private:
 
+#if WITH_EDITORONLY_DATA
+
 	UPROPERTY(transient)
 	URigVM* VMSnapshotBeforeExecution;
+
+#endif
 	
 #if WITH_EDITOR	
 
@@ -644,6 +648,9 @@ private:
 	
 public:	
 	
+	/** Creates the snapshot VM if required and returns it */
+	URigVM* GetSnapshotVM(bool bCreateIfNeeded = true);
+
 	/** Adds a breakpoint in the VM at the InstructionIndex */
 	void AddBreakpoint(int32 InstructionIndex);
 
