Setting RF_Transient on VMSnapshotBeforeExecution when it is created.

diff --git a/Engine/Plugins/Experimental/ControlRig/Source/ControlRig/Private/ControlRig.cpp b/Engine/Plugins/Experimental/ControlRig/Source/ControlRig/Private/ControlRig.cpp
--- a/Engine/Plugins/Experimental/ControlRig/Source/ControlRig/Private/ControlRig.cpp
+++ b/Engine/Plugins/Experimental/ControlRig/Source/ControlRig/Private/ControlRig.cpp
@@ -2374,7 +2374,7 @@ URigVM* UControlRig::GetSnapshotVM(bool bCreateIfNeeded)
 	}
 	if ((VMSnapshotBeforeExecution == nullptr) && bCreateIfNeeded)
 	{
-		VMSnapshotBeforeExecution = NewObject<URigVM>(this);
+		VMSnapshotBeforeExecution = NewObject<URigVM>(this, NAME_None, RF_Transient);
 	}
 	return VMSnapshotBeforeExecution;
 #else
