Temporarily disabled VMSnapshotBeforeExecution to unblock Frosty cooking

diff --git a/Engine/Plugins/Experimental/ControlRig/Source/ControlRig/Private/ControlRig.cpp b/Engine/Plugins/Experimental/ControlRig/Source/ControlRig/Private/ControlRig.cpp
--- a/Engine/Plugins/Experimental/ControlRig/Source/ControlRig/Private/ControlRig.cpp
+++ b/Engine/Plugins/Experimental/ControlRig/Source/ControlRig/Private/ControlRig.cpp
@@ -2364,7 +2364,8 @@ FRigVMExternalVariable UControlRig::GetExternalVariableFromDescription(const FBP
 
 URigVM* UControlRig::GetSnapshotVM(bool bCreateIfNeeded)
 {
-#if WITH_EDITOR
+// temporarily disabled #todo
+#if 0 || WITH_EDITOR
 	if(VMSnapshotBeforeExecution != nullptr)
 	{
 		if(VMSnapshotBeforeExecution->GetOuter() != this)
