Found and fixed another codepath that was setting VMSnapshotBeforeExecution without RF_Transient

diff --git a/Engine/Plugins/Experimental/ControlRig/Source/ControlRigDeveloper/Private/ControlRigBlueprint.cpp b/Engine/Plugins/Experimental/ControlRig/Source/ControlRigDeveloper/Private/ControlRigBlueprint.cpp
--- a/Engine/Plugins/Experimental/ControlRig/Source/ControlRigDeveloper/Private/ControlRigBlueprint.cpp
+++ b/Engine/Plugins/Experimental/ControlRig/Source/ControlRigDeveloper/Private/ControlRigBlueprint.cpp
@@ -621,7 +621,7 @@ void UControlRigBlueprint::RecompileVM()
 #if WITH_EDITOR
 		if (CDO->GetVMSnapshot() == nullptr || CDO->GetVMSnapshot()->GetOuter() != CDO)
 		{
-			CDO->SetVMSnapshot(NewObject<URigVM>(CDO));
+			CDO->SetVMSnapshot(NewObject<URigVM>(CDO, NAME_None, RF_Transient));
 		}
 #endif
 
