ControlRig: Fix VM Snapshot producing crash when saving RigVM

diff --git a/Engine/Plugins/Experimental/ControlRig/Source/ControlRig/Private/ControlRig.cpp b/Engine/Plugins/Experimental/ControlRig/Source/ControlRig/Private/ControlRig.cpp
--- a/Engine/Plugins/Experimental/ControlRig/Source/ControlRig/Private/ControlRig.cpp
+++ b/Engine/Plugins/Experimental/ControlRig/Source/ControlRig/Private/ControlRig.cpp
@@ -1107,6 +1107,15 @@ void UControlRig::PostLoad()
 			DynamicHierarchy->SetFlags(DynamicHierarchy->GetFlags() | RF_Public | RF_DefaultSubObject);
 		}
 	}
+
+#if WITH_EDITORONLY_DATA
+	if (VMSnapshotBeforeExecution)
+	{
+		// Some VMSnapshots might have been created without the Transient flag.
+		// Assets from that version require the snapshot to be flagged as below.
+		VMSnapshotBeforeExecution->SetFlags(VMSnapshotBeforeExecution->GetFlags() | RF_Transient);
+	}
+#endif
 }
 
 TArray<FRigControlElement*> UControlRig::AvailableControls() const
@@ -2364,8 +2373,7 @@ FRigVMExternalVariable UControlRig::GetExternalVariableFromDescription(const FBP
 
 URigVM* UControlRig::GetSnapshotVM(bool bCreateIfNeeded)
 {
-// temporarily disabled #todo
-#if 0 && WITH_EDITOR
+#if WITH_EDITOR
 	if(VMSnapshotBeforeExecution != nullptr)
 	{
 		if(VMSnapshotBeforeExecution->GetOuter() != this)
diff --git a/Engine/Plugins/Experimental/ControlRig/Source/ControlRig/Public/ControlRig.h b/Engine/Plugins/Experimental/ControlRig/Source/ControlRig/Public/ControlRig.h
--- a/Engine/Plugins/Experimental/ControlRig/Source/ControlRig/Public/ControlRig.h
+++ b/Engine/Plugins/Experimental/ControlRig/Source/ControlRig/Public/ControlRig.h
@@ -648,9 +648,6 @@ private:
 	
 public:	
 	
-	/** Creates the snapshot VM if required and returns it */
-	URigVM* GetSnapshotVM(bool bCreateIfNeeded = true);
-
 	/** Adds a breakpoint in the VM at the InstructionIndex */
 	void AddBreakpoint(int32 InstructionIndex);
 
@@ -659,9 +656,8 @@ public:
 	
 	FRigVMDebugInfo& GetDebugInfo() { return DebugInfo; }
 
-	URigVM* GetVMSnapshot() { return VMSnapshotBeforeExecution; }
-
-	void SetVMSnapshot(URigVM* InVMSnapshot) { VMSnapshotBeforeExecution = InVMSnapshot; }
+	/** Creates the snapshot VM if required and returns it */
+	URigVM* GetSnapshotVM(bool bCreateIfNeeded = true);
 #endif	
 
 	// UObject interface
diff --git a/Engine/Plugins/Experimental/ControlRig/Source/ControlRigDeveloper/Private/ControlRigBlueprint.cpp b/Engine/Plugins/Experimental/ControlRig/Source/ControlRigDeveloper/Private/ControlRigBlueprint.cpp
--- a/Engine/Plugins/Experimental/ControlRig/Source/ControlRigDeveloper/Private/ControlRigBlueprint.cpp
+++ b/Engine/Plugins/Experimental/ControlRig/Source/ControlRigDeveloper/Private/ControlRigBlueprint.cpp
@@ -618,13 +618,6 @@ void UControlRigBlueprint::RecompileVM()
 			CDO->SetVM(NewObject<URigVM>(CDO, TEXT("VM")));
 		}
 
-#if WITH_EDITOR
-		if (CDO->GetVMSnapshot() == nullptr || CDO->GetVMSnapshot()->GetOuter() != CDO)
-		{
-			CDO->SetVMSnapshot(NewObject<URigVM>(CDO, NAME_None, RF_Transient));
-		}
-#endif
-
 		if (!HasAnyFlags(RF_Transient | RF_Transactional))
 		{
 			CDO->Modify(false);
