Temporarily removed a control rig delegate binding to avoid crashes on PIE shutown. Needs a better long-term solution.

diff --git a/Engine/Plugins/Experimental/ControlRig/Source/ControlRig/Private/ControlRig.cpp b/Engine/Plugins/Experimental/ControlRig/Source/ControlRig/Private/ControlRig.cpp
--- a/Engine/Plugins/Experimental/ControlRig/Source/ControlRig/Private/ControlRig.cpp
+++ b/Engine/Plugins/Experimental/ControlRig/Source/ControlRig/Private/ControlRig.cpp
@@ -998,10 +998,10 @@ void UControlRig::SetVM(URigVM* NewVM)
 	
 	if (NewVM)
 	{
-		if (!NewVM->ExecutionReachedExit().IsBoundToObject(this))
-		{
-			NewVM->ExecutionReachedExit().AddUObject(this, &UControlRig::HandleExecutionReachedExit);
-		}
+// 		if (!NewVM->ExecutionReachedExit().IsBoundToObject(this))
+// 		{
+// 			NewVM->ExecutionReachedExit().AddUObject(this, &UControlRig::HandleExecutionReachedExit);
+// 		}
 	}
 
 	VM = NewVM;
