Control Rig: Catch poses which are not valid (due to missing ref skeletons on skeletal mesh assets, for example)

diff --git a/Engine/Plugins/Experimental/ControlRig/Source/ControlRig/Private/ControlRig.cpp b/Engine/Plugins/Experimental/ControlRig/Source/ControlRig/Private/ControlRig.cpp
--- a/Engine/Plugins/Experimental/ControlRig/Source/ControlRig/Private/ControlRig.cpp
+++ b/Engine/Plugins/Experimental/ControlRig/Source/ControlRig/Private/ControlRig.cpp
@@ -2425,6 +2425,10 @@ void UControlRig::SetBoneInitialTransformsFromSkeletalMeshComponent(USkeletalMes
 		FMemMark Mark(FMemStack::Get());
 		FCompactPose OutPose;
 		OutPose.ResetToRefPose(InSkelMeshComp->GetAnimInstance()->GetRequiredBones());
+		if(!OutPose.GetBoneContainer().IsValid())
+		{
+			return;
+		}
 
 		DynamicHierarchy->ForEach<FRigBoneElement>([this, &OutPose, InSkelMeshComp](FRigBoneElement* BoneElement) -> bool
 			{
