Contorl Rig: Ensure parent relationships when importing content from Pre 5.0

diff --git a/Engine/Plugins/Experimental/ControlRig/Source/ControlRigDeveloper/Private/ControlRigBlueprint.cpp b/Engine/Plugins/Experimental/ControlRig/Source/ControlRigDeveloper/Private/ControlRigBlueprint.cpp
--- a/Engine/Plugins/Experimental/ControlRig/Source/ControlRigDeveloper/Private/ControlRigBlueprint.cpp
+++ b/Engine/Plugins/Experimental/ControlRig/Source/ControlRigDeveloper/Private/ControlRigBlueprint.cpp
@@ -405,13 +405,28 @@ void UControlRigBlueprint::PostLoad()
 		
 		Hierarchy->Reset();
 
+		TMap<FRigElementKey, FRigElementKey> KeyMap;;
+
 		for(const FRigBone& Bone : HierarchyContainer_DEPRECATED.BoneHierarchy)
 		{
-			HierarchyController->AddBone(Bone.Name, Bone.GetParentElementKey(true), Bone.InitialTransform, true, Bone.Type, false);
+			const FRigElementKey OriginalParentKey = Bone.GetParentElementKey(true);
+			const FRigElementKey* ParentKey = nullptr;
+			if(OriginalParentKey.IsValid())
+			{
+				ParentKey = KeyMap.Find(OriginalParentKey);
+			}
+			if(ParentKey == nullptr)
+			{
+				ParentKey = &OriginalParentKey;
+			}
+
+			const FRigElementKey Key = HierarchyController->AddBone(Bone.Name, *ParentKey, Bone.InitialTransform, true, Bone.Type, false);
+			KeyMap.Add(Bone.GetElementKey(), Key);
 		}
 		for(const FRigSpace& Space : HierarchyContainer_DEPRECATED.SpaceHierarchy)
 		{
-			HierarchyController->AddNull(Space.Name, FRigElementKey(), Space.InitialTransform, false, false);
+			const FRigElementKey Key = HierarchyController->AddNull(Space.Name, FRigElementKey(), Space.InitialTransform, false, false);
+			KeyMap.Add(Space.GetElementKey(), Key);
 		}
 		for(const FRigControl& Control : HierarchyContainer_DEPRECATED.ControlHierarchy)
 		{
@@ -439,7 +454,7 @@ void UControlRigBlueprint::PostLoad()
 				InitialValue.SetFromTransform(InitialValue.Storage_DEPRECATED, Settings.ControlType, Settings.PrimaryAxis);
 			}
 			
-			HierarchyController->AddControl(
+			const FRigElementKey Key = HierarchyController->AddControl(
 				Control.Name,
 				FRigElementKey(),
 				Settings,
@@ -447,32 +462,55 @@ void UControlRigBlueprint::PostLoad()
 				Control.OffsetTransform,
 				Control.GizmoTransform,
 				false);
+
+			KeyMap.Add(Control.GetElementKey(), Key);
 		}
 		
 		for(const FRigCurve& Curve : HierarchyContainer_DEPRECATED.CurveContainer)
 		{
-			HierarchyController->AddCurve(Curve.Name, Curve.Value, false);
+			const FRigElementKey Key = HierarchyController->AddCurve(Curve.Name, Curve.Value, false);
+			KeyMap.Add(Curve.GetElementKey(), Key);
 		}
 
 		for(const FRigSpace& Space : HierarchyContainer_DEPRECATED.SpaceHierarchy)
 		{
-			FRigElementKey ParentKey = Space.GetParentElementKey(true);
-			if(ParentKey.IsValid())
+			const FRigElementKey OriginalParentKey = Space.GetParentElementKey(true);
+			if(OriginalParentKey.IsValid())
 			{
+				FRigElementKey ParentKey;
+				if(const FRigElementKey* ParentKeyPtr = KeyMap.Find(OriginalParentKey))
+				{
+					ParentKey = *ParentKeyPtr;
+				}
 				HierarchyController->SetParent(Space.GetElementKey(), ParentKey, false, false);
 			}
 		}
 
 		for(const FRigControl& Control : HierarchyContainer_DEPRECATED.ControlHierarchy)
 		{
-			FRigElementKey ParentKey = Control.GetParentElementKey(true);
+			FRigElementKey OriginalParentKey = Control.GetParentElementKey(true);
 			const FRigElementKey SpaceKey = Control.GetSpaceElementKey(true);
-			ParentKey = SpaceKey.IsValid() ? SpaceKey : ParentKey;
-			if(ParentKey.IsValid())
+			OriginalParentKey = SpaceKey.IsValid() ? SpaceKey : OriginalParentKey;
+			if(OriginalParentKey.IsValid())
 			{
+				FRigElementKey ParentKey;
+				if(const FRigElementKey* ParentKeyPtr = KeyMap.Find(OriginalParentKey))
+				{
+					ParentKey = *ParentKeyPtr;
+				}
 				HierarchyController->SetParent(Control.GetElementKey(), ParentKey, false, false);
 			}
 		}
+
+		for(const TPair<FRigElementKey, FRigElementKey>& Pair: KeyMap)
+		{
+			if(Pair.Key != Pair.Value)
+			{
+				check(Pair.Key.Type == Pair.Value.Type);
+				const FText TypeLabel = StaticEnum<ERigElementType>()->GetDisplayNameTextByValue((int64)Pair.Key.Type);
+				GetController()->ReportAndNotifyErrorf(TEXT("%s '%s' was renamed to '%s' during load."), *TypeLabel.ToString(), *Pair.Key.Name.ToString(), *Pair.Value.Name.ToString());
+			}
+		}
 	}
 
 	PropagateHierarchyFromBPToInstances();
