Control Rig: When refreshing hierarchy also refresh curves

diff --git a/Engine/Plugins/Experimental/ControlRig/Source/ControlRigEditor/Private/Editor/SRigHierarchy.cpp b/Engine/Plugins/Experimental/ControlRig/Source/ControlRigEditor/Private/Editor/SRigHierarchy.cpp
--- a/Engine/Plugins/Experimental/ControlRig/Source/ControlRigEditor/Private/Editor/SRigHierarchy.cpp
+++ b/Engine/Plugins/Experimental/ControlRig/Source/ControlRigEditor/Private/Editor/SRigHierarchy.cpp
@@ -1492,6 +1492,7 @@ void SRigHierarchy::RefreshHierarchy(const FAssetData& InAssetData)
 		check(Controller);
 		
 		Controller->ImportBones(RefSkeleton, NAME_None, true, true, bSelectBones, true);
+		Controller->ImportCurves(Mesh->GetSkeleton(), NAME_None, false, true);
 	}
 
 	ControlRigBlueprint->PropagateHierarchyFromBPToInstances();
