Control Rig: Lower error message for name validation / fixing to a in editor / non-cook warning only

diff --git a/Engine/Plugins/Experimental/ControlRig/Source/ControlRigDeveloper/Private/ControlRigBlueprint.cpp b/Engine/Plugins/Experimental/ControlRig/Source/ControlRigDeveloper/Private/ControlRigBlueprint.cpp
--- a/Engine/Plugins/Experimental/ControlRig/Source/ControlRigDeveloper/Private/ControlRigBlueprint.cpp
+++ b/Engine/Plugins/Experimental/ControlRig/Source/ControlRigDeveloper/Private/ControlRigBlueprint.cpp
@@ -502,16 +502,21 @@ void UControlRigBlueprint::PostLoad()
 			}
 		}
 
-		for(const TPair<FRigElementKey, FRigElementKey>& Pair: KeyMap)
+#if WITH_EDITOR
+		if(!IsRunningCommandlet()) // don't show warnings like this if we are cooking
 		{
-			if(Pair.Key != Pair.Value)
+			for(const TPair<FRigElementKey, FRigElementKey>& Pair: KeyMap)
 			{
-				check(Pair.Key.Type == Pair.Value.Type);
-				const FText TypeLabel = StaticEnum<ERigElementType>()->GetDisplayNameTextByValue((int64)Pair.Key.Type);
-				GetController()->ReportAndNotifyErrorf(TEXT("%s '%s' was renamed to '%s' during load."), *TypeLabel.ToString(), *Pair.Key.Name.ToString(), *Pair.Value.Name.ToString());
+				if(Pair.Key != Pair.Value)
+				{
+					check(Pair.Key.Type == Pair.Value.Type);
+					const FText TypeLabel = StaticEnum<ERigElementType>()->GetDisplayNameTextByValue((int64)Pair.Key.Type);
+					GetController()->ReportWarningf(TEXT("%s '%s' was renamed to '%s' during load (fixing invalid name)."), *TypeLabel.ToString(), *Pair.Key.Name.ToString(), *Pair.Value.Name.ToString());
+				}
 			}
 		}
 	}
+#endif
 
 	PropagateHierarchyFromBPToInstances();
 	
