ControlRig: Fix crash of ControlRig referencing non-existing VM Snapshot

diff --git a/Engine/Plugins/Experimental/ControlRig/Source/ControlRig/Private/ControlRig.cpp b/Engine/Plugins/Experimental/ControlRig/Source/ControlRig/Private/ControlRig.cpp
--- a/Engine/Plugins/Experimental/ControlRig/Source/ControlRig/Private/ControlRig.cpp
+++ b/Engine/Plugins/Experimental/ControlRig/Source/ControlRig/Private/ControlRig.cpp
@@ -70,7 +70,7 @@ UControlRig::UControlRig(const FObjectInitializer& ObjectInitializer)
 	, PostSetupBracket(0)
 	, InteractionBracket(0)
 	, InterRigSyncBracket(0)
-#if WITH_EDITOR
+#if WITH_EDITORONLY_DATA
 	, VMSnapshotBeforeExecution(nullptr)
 #endif
 {
@@ -93,6 +93,13 @@ void UControlRig::BeginDestroy()
 	{
 		VM->ExecutionReachedExit().RemoveAll(this);
 	}
+
+#if WITH_EDITORONLY_DATA
+	if (VMSnapshotBeforeExecution)
+	{
+		VMSnapshotBeforeExecution = nullptr;
+	}
+#endif
 }
 
 UWorld* UControlRig::GetWorld() const
@@ -476,6 +483,12 @@ void UControlRig::Execute(const EControlRigState InState, const FName& InEventNa
 	{
 		// Copy the breakpoints. This will not override the state of the breakpoints
 		DebugInfo.SetBreakpoints(CDO->DebugInfo.GetBreakpoints());
+
+		// If there are any breakpoints, create the Snapshot VM if it hasn't been created yet
+		if (DebugInfo.GetBreakpoints().Num() > 0)
+		{
+			GetSnapshotVM();
+		}
 	}
 	VM->SetDebugInfo(&DebugInfo);
 #endif
@@ -490,6 +503,7 @@ void UControlRig::Execute(const EControlRigState InState, const FName& InEventNa
 			Execute(EControlRigState::Init, InEventName);
 			bJustRanInit = true;
 		}
+		
 	}
 
 	FRigUnitContext Context;
@@ -998,10 +1012,10 @@ void UControlRig::SetVM(URigVM* NewVM)
 	
 	if (NewVM)
 	{
-// 		if (!NewVM->ExecutionReachedExit().IsBoundToObject(this))
-// 		{
-// 			NewVM->ExecutionReachedExit().AddUObject(this, &UControlRig::HandleExecutionReachedExit);
-// 		}
+		if (!NewVM->ExecutionReachedExit().IsBoundToObject(this))
+		{
+			NewVM->ExecutionReachedExit().AddUObject(this, &UControlRig::HandleExecutionReachedExit);
+		}
 	}
 
 	VM = NewVM;
@@ -1204,7 +1218,7 @@ void UControlRig::HandleOnControlModified(UControlRig* Subject, FRigControlEleme
 void UControlRig::HandleExecutionReachedExit()
 {
 #if WITH_EDITOR
-	if(URigVM* SnapShotVM = GetSnapshotVM())
+	if(URigVM* SnapShotVM = GetSnapshotVM(false))
 	{
 		SnapShotVM->CopyFrom(VM, false, false, false, true, true);
 	}
@@ -2374,16 +2388,9 @@ FRigVMExternalVariable UControlRig::GetExternalVariableFromDescription(const FBP
 URigVM* UControlRig::GetSnapshotVM(bool bCreateIfNeeded)
 {
 #if WITH_EDITOR
-	if(VMSnapshotBeforeExecution != nullptr)
-	{
-		if(VMSnapshotBeforeExecution->GetOuter() != this)
-		{
-			VMSnapshotBeforeExecution = nullptr;
-		}
-	}
 	if ((VMSnapshotBeforeExecution == nullptr) && bCreateIfNeeded)
 	{
-		VMSnapshotBeforeExecution = NewObject<URigVM>(this, NAME_None, RF_Transient);
+		VMSnapshotBeforeExecution = NewObject<URigVM>(GetTransientPackage(), NAME_None, RF_Transient);
 	}
 	return VMSnapshotBeforeExecution;
 #else
diff --git a/Engine/Source/Runtime/RigVM/Private/RigVMCore/RigVM.cpp b/Engine/Source/Runtime/RigVM/Private/RigVMCore/RigVM.cpp
--- a/Engine/Source/Runtime/RigVM/Private/RigVMCore/RigVM.cpp
+++ b/Engine/Source/Runtime/RigVM/Private/RigVMCore/RigVM.cpp
@@ -84,6 +84,11 @@ URigVM::URigVM()
 URigVM::~URigVM()
 {
 	Reset();
+
+	ExecutionReachedExit().Clear();
+#if WITH_EDITOR
+	ExecutionHalted().Clear();
+#endif
 }
 
 void URigVM::Serialize(FArchive& Ar)
