Control Rig Debugging: Breakpoints UI.

diff --git a/Engine/Plugins/Experimental/ControlRig/Source/ControlRig/Private/ControlRig.cpp b/Engine/Plugins/Experimental/ControlRig/Source/ControlRig/Private/ControlRig.cpp
--- a/Engine/Plugins/Experimental/ControlRig/Source/ControlRig/Private/ControlRig.cpp
+++ b/Engine/Plugins/Experimental/ControlRig/Source/ControlRig/Private/ControlRig.cpp
@@ -503,7 +503,6 @@ void UControlRig::Execute(const EControlRigState InState, const FName& InEventNa
 			Execute(EControlRigState::Init, InEventName);
 			bJustRanInit = true;
 		}
-		
 	}
 
 	FRigUnitContext Context;
@@ -2398,17 +2397,14 @@ URigVM* UControlRig::GetSnapshotVM(bool bCreateIfNeeded)
 #endif
 }
 
-void UControlRig::AddBreakpoint(int32 InstructionIndex)
+void UControlRig::AddBreakpoint(int32 InstructionIndex, URigVMNode* InNode)
 {
-	// this makes sure that the snapshot exists
-	if(URigVM* SnapShotVM = GetSnapshotVM())
-	{
-		DebugInfo.AddBreakpoint(InstructionIndex);
-	}
+	DebugInfo.AddBreakpoint(InstructionIndex, InNode);
 }
 
 void UControlRig::ResumeExecution()
 {
+	// this makes sure that the snapshot exists
 	if(URigVM* SnapShotVM = GetSnapshotVM())
 	{
 		VM->CopyFrom(SnapShotVM, false, false, false, true, true);
diff --git a/Engine/Plugins/Experimental/ControlRig/Source/ControlRig/Public/ControlRig.h b/Engine/Plugins/Experimental/ControlRig/Source/ControlRig/Public/ControlRig.h
--- a/Engine/Plugins/Experimental/ControlRig/Source/ControlRig/Public/ControlRig.h
+++ b/Engine/Plugins/Experimental/ControlRig/Source/ControlRig/Public/ControlRig.h
@@ -546,6 +546,11 @@ public:
 	UFUNCTION(BlueprintSetter)
 	void SetInteractionRigClass(TSubclassOf<UControlRig> InInteractionRigClass);
 
+	// UObject interface
+#if WITH_EDITOR
+	virtual void PreEditChange(FProperty* PropertyAboutToChange) override;
+	virtual void PostEditChangeProperty(FPropertyChangedEvent& PropertyChangedEvent) override;
+#endif
 
 public:
 	//~ Begin IInterface_AssetUserData Interface
@@ -649,7 +654,7 @@ private:
 public:	
 	
 	/** Adds a breakpoint in the VM at the InstructionIndex */
-	void AddBreakpoint(int32 InstructionIndex);
+	void AddBreakpoint(int32 InstructionIndex, URigVMNode* InNode);
 
 	/** Reloads the VM from the SnapshotVM, and increments activation hit on the halted breakpoint */
 	void ResumeExecution();
@@ -660,11 +665,6 @@ public:
 	URigVM* GetSnapshotVM(bool bCreateIfNeeded = true);
 #endif	
 
-	// UObject interface
-	#if WITH_EDITOR
-	virtual void PreEditChange(FProperty* PropertyAboutToChange) override;
-	virtual void PostEditChangeProperty(FPropertyChangedEvent& PropertyChangedEvent) override;
-#endif
 
 	friend class FControlRigBlueprintCompilerContext;
 	friend struct FRigHierarchyRef;
diff --git a/Engine/Plugins/Experimental/ControlRig/Source/ControlRigDeveloper/Private/ControlRigBlueprint.cpp b/Engine/Plugins/Experimental/ControlRig/Source/ControlRigDeveloper/Private/ControlRigBlueprint.cpp
--- a/Engine/Plugins/Experimental/ControlRig/Source/ControlRigDeveloper/Private/ControlRigBlueprint.cpp
+++ b/Engine/Plugins/Experimental/ControlRig/Source/ControlRigDeveloper/Private/ControlRigBlueprint.cpp
@@ -836,7 +836,7 @@ void UControlRigBlueprint::HandleReportFromCompiler(EMessageSeverity::Type InSev
 
 void UControlRigBlueprint::ClearBreakpoints()
 {
-	Breakpoints.Empty();
+	RigVMBreakpointNodes.Empty();
 	RefreshControlRigBreakpoints();
 }
 
@@ -913,7 +913,7 @@ bool UControlRigBlueprint::AddBreakpoint(URigVMNode* InBreakpointNode, URigVMLib
 	}
 	else
 	{
-		if (!Breakpoints.Contains(InBreakpointNode))
+		if (!RigVMBreakpointNodes.Contains(InBreakpointNode))
 		{
 			// Add the breakpoint to the VM
 			bSuccess = AddBreakpointToControlRig(InBreakpointNode);
@@ -954,7 +954,7 @@ bool UControlRigBlueprint::AddBreakpointToControlRig(URigVMNode* InBreakpointNod
 				if (!AddedCallpaths.Contains(BreakpointCallPath))
 				{
 					AddedCallpaths.Add(BreakpointCallPath);
-					CDO->AddBreakpoint(i);
+					CDO->AddBreakpoint(i, InBreakpointNode);
 				}
 			}
 		}
@@ -962,7 +962,7 @@ bool UControlRigBlueprint::AddBreakpointToControlRig(URigVMNode* InBreakpointNod
 
 	if (AddedCallpaths.Num() > 0)
 	{
-		Breakpoints.AddUnique(InBreakpointNode);
+		RigVMBreakpointNodes.AddUnique(InBreakpointNode);
 		return true;
 	}
 	
@@ -1000,9 +1000,9 @@ bool UControlRigBlueprint::RemoveBreakpoint(const FString& InBreakpointNodePath)
 
 bool UControlRigBlueprint::RemoveBreakpoint(URigVMNode* InBreakpointNode)
 {
-	if (Breakpoints.Contains(InBreakpointNode))
+	if (RigVMBreakpointNodes.Contains(InBreakpointNode))
 	{
-		Breakpoints.Remove(InBreakpointNode);
+		RigVMBreakpointNodes.Remove(InBreakpointNode);
 
 		// Multiple breakpoint nodes might set a breakpoint to the same instruction. When we remove
 		// one of the breakpoint nodes, we do not want to remove the instruction breakpoint if there
@@ -1022,7 +1022,7 @@ void UControlRigBlueprint::RefreshControlRigBreakpoints()
 	UControlRigBlueprintGeneratedClass* RigClass = GetControlRigBlueprintGeneratedClass();
 	UControlRig* CDO = Cast<UControlRig>(RigClass->GetDefaultObject(false));
 	CDO->GetDebugInfo().Clear();
-	for (URigVMNode* Node : Breakpoints)
+	for (URigVMNode* Node : RigVMBreakpointNodes)
 	{
 		AddBreakpointToControlRig(Node);
 	}
diff --git a/Engine/Plugins/Experimental/ControlRig/Source/ControlRigDeveloper/Private/Graph/ControlRigGraph.cpp b/Engine/Plugins/Experimental/ControlRig/Source/ControlRigDeveloper/Private/Graph/ControlRigGraph.cpp
--- a/Engine/Plugins/Experimental/ControlRig/Source/ControlRigDeveloper/Private/Graph/ControlRigGraph.cpp
+++ b/Engine/Plugins/Experimental/ControlRig/Source/ControlRigDeveloper/Private/Graph/ControlRigGraph.cpp
@@ -657,39 +657,17 @@ int32 UControlRigGraph::GetInstructionIndex(UControlRigGraphNode* InNode)
 				return InstructionIndex;
 			}
 
-			TArray<URigVMNode*> NodesToCheck;
-			if (URigVMLibraryNode* LibraryNode = Cast<URigVMLibraryNode>(InModelNode))
+			FRigVMInstructionArray Instructions = InByteCode->GetInstructions();
+			for (int32 i = 0; i < Instructions.Num(); ++i)
 			{
-				NodesToCheck.Append(LibraryNode->GetContainedNodes());
-			}
-
-			if (InModelNode->IsA<URigVMFunctionReturnNode>() ||
-				InModelNode->IsA<URigVMRerouteNode>())
-			{
-				NodesToCheck.Append(InModelNode->GetLinkedSourceNodes());
-			}
-
-			if (InModelNode->IsA<URigVMFunctionEntryNode>() ||
-				InModelNode->IsA<URigVMRerouteNode>())
-			{
-				NodesToCheck.Append(InModelNode->GetLinkedTargetNodes());
-			}
-
-			int32 MinimumInstructionIndex = INDEX_NONE;
-			for (URigVMNode* NodeToCheck : NodesToCheck)
-			{
-				int32 ContainedInstructionIndex = GetInstructionIndex(NodeToCheck, InByteCode, Indices);
-				if (ContainedInstructionIndex != INDEX_NONE)
+				const FRigVMASTProxy Proxy = FRigVMASTProxy::MakeFromCallPath(InByteCode->GetCallPathForInstruction(i), InModelNode->GetRootGraph());
+				if (Proxy.GetCallstack().Contains(InModelNode))
 				{
-					if (ContainedInstructionIndex < MinimumInstructionIndex || MinimumInstructionIndex == INDEX_NONE)
-					{
-						MinimumInstructionIndex = ContainedInstructionIndex;
-					}
+					Indices.FindOrAdd(InModelNode) = i;
+					return i;
 				}
 			}
-
-			Indices.FindOrAdd(InModelNode) = MinimumInstructionIndex;
-			return MinimumInstructionIndex;
+			return INDEX_NONE;
 		}
 	};
 
diff --git a/Engine/Plugins/Experimental/ControlRig/Source/ControlRigDeveloper/Public/ControlRigBlueprint.h b/Engine/Plugins/Experimental/ControlRig/Source/ControlRigDeveloper/Public/ControlRigBlueprint.h
--- a/Engine/Plugins/Experimental/ControlRig/Source/ControlRigDeveloper/Public/ControlRigBlueprint.h
+++ b/Engine/Plugins/Experimental/ControlRig/Source/ControlRigDeveloper/Public/ControlRigBlueprint.h
@@ -535,7 +535,7 @@ private:
 
 #if WITH_EDITOR
 private:
-	TArray<URigVMNode*> Breakpoints;
+	TArray<URigVMNode*> RigVMBreakpointNodes;
 
 public:
 
diff --git a/Engine/Plugins/Experimental/ControlRig/Source/ControlRigEditor/Private/ControlRigBlueprintCommands.cpp b/Engine/Plugins/Experimental/ControlRig/Source/ControlRigEditor/Private/ControlRigBlueprintCommands.cpp
--- a/Engine/Plugins/Experimental/ControlRig/Source/ControlRigEditor/Private/ControlRigBlueprintCommands.cpp
+++ b/Engine/Plugins/Experimental/ControlRig/Source/ControlRigEditor/Private/ControlRigBlueprintCommands.cpp
@@ -14,6 +14,7 @@ void FControlRigBlueprintCommands::RegisterCommands()
 	UI_COMMAND(UpdateEvent, "Forwards Solve", "Run the normal update graph", EUserInterfaceActionType::Button, FInputChord());
 	UI_COMMAND(InverseEvent, "Backwards Solve", "Run the inverse graph", EUserInterfaceActionType::Button, FInputChord());
 	UI_COMMAND(InverseAndUpdateEvent, "Backwards and Forwards", "Run the inverse graph followed by the update graph", EUserInterfaceActionType::Button, FInputChord());
+	UI_COMMAND(ResumeExecution, "Resume", "Resumes execution after being halted at a breakpoint.", EUserInterfaceActionType::Button, FInputChord(EKeys::F5));
 }
 
 #undef LOCTEXT_NAMESPACE
diff --git a/Engine/Plugins/Experimental/ControlRig/Source/ControlRigEditor/Private/ControlRigBlueprintCommands.h b/Engine/Plugins/Experimental/ControlRig/Source/ControlRigEditor/Private/ControlRigBlueprintCommands.h
--- a/Engine/Plugins/Experimental/ControlRig/Source/ControlRigEditor/Private/ControlRigBlueprintCommands.h
+++ b/Engine/Plugins/Experimental/ControlRig/Source/ControlRigEditor/Private/ControlRigBlueprintCommands.h
@@ -42,6 +42,9 @@ public:
 	/** Run the inverse graph followed by the update graph */
 	TSharedPtr< FUICommandInfo > InverseAndUpdateEvent;
 
+	/** Resume the execution of the graph when halted at a breakpoint */
+	TSharedPtr< FUICommandInfo > ResumeExecution;
+
 	/**
 	 * Initialize commands
 	 */
diff --git a/Engine/Plugins/Experimental/ControlRig/Source/ControlRigEditor/Private/ControlRigEditorModule.cpp b/Engine/Plugins/Experimental/ControlRig/Source/ControlRigEditor/Private/ControlRigEditorModule.cpp
--- a/Engine/Plugins/Experimental/ControlRig/Source/ControlRigEditor/Private/ControlRigEditorModule.cpp
+++ b/Engine/Plugins/Experimental/ControlRig/Source/ControlRigEditor/Private/ControlRigEditorModule.cpp
@@ -1893,6 +1893,42 @@ void FControlRigEditorModule::GetContextMenuActions(const UControlRigGraphSchema
 						));
 					}
 
+					
+					FToolMenuSection& DebugSection = Menu->AddSection("EdGraphSchemaDebug", LOCTEXT("DebugHeader", "Debug"));
+					URigVMNode* Node = RigNode->GetModelNode();
+					if (Node->HasBreakpoint())
+					{
+						
+						DebugSection.AddMenuEntry(
+                        "Remove Breakpoint",
+                        LOCTEXT("RemoveBreakpoint", "Remove Breakpoint"),
+                        LOCTEXT("RemoveBreakpoint_Tooltip", "Removes a breakpoint to the graph at this node"),
+                        FSlateIcon(),
+                        FUIAction(FExecuteAction::CreateLambda([Controller, Node, RigBlueprint]()
+                        {
+                            if (RigBlueprint->RemoveBreakpoint(Node))
+                            {
+                                Node->SetHasBreakpoint(false);
+                            }
+                        })));
+					}
+					else
+					{
+						DebugSection.AddMenuEntry(
+                        "Add Breakpoint",
+                        LOCTEXT("AddBreakpoint", "Add Breakpoint"),
+                        LOCTEXT("AddBreakpoint_Tooltip", "Adds a breakpoint to the graph at this node"),
+                        FSlateIcon(),
+                        FUIAction(FExecuteAction::CreateLambda([Controller, Node, RigBlueprint]()
+                        {
+                            if (RigBlueprint->AddBreakpoint(Node))
+                            {
+                                Node->SetHasBreakpoint(true);
+                            }
+                        })));
+					}
+					
+
 					FToolMenuSection& OrganizationSection = Menu->AddSection("EdGraphSchemaOrganization", LOCTEXT("OrganizationHeader", "Organization"));
 					OrganizationSection.AddMenuEntry(
 						"Collapse Nodes",
diff --git a/Engine/Plugins/Experimental/ControlRig/Source/ControlRigEditor/Private/Editor/ControlRigEditor.cpp b/Engine/Plugins/Experimental/ControlRig/Source/ControlRigEditor/Private/Editor/ControlRigEditor.cpp
--- a/Engine/Plugins/Experimental/ControlRig/Source/ControlRigEditor/Private/Editor/ControlRigEditor.cpp
+++ b/Engine/Plugins/Experimental/ControlRig/Source/ControlRigEditor/Private/Editor/ControlRigEditor.cpp
@@ -429,6 +429,25 @@ void FControlRigEditor::BindCommands()
 		FControlRigBlueprintCommands::Get().InverseAndUpdateEvent,
 		FExecuteAction::CreateSP(this, &FControlRigEditor::SetEventQueue, EControlRigEditorEventQueue::InverseAndUpdate),
 		FCanExecuteAction());
+
+	GetToolkitCommands()->MapAction(
+		FControlRigBlueprintCommands::Get().ResumeExecution,
+		FExecuteAction::CreateLambda([this]()
+        {
+            if (ControlRig)
+            {
+	            ControlRig->ResumeExecution();
+            }
+        }),			
+		FIsActionChecked::CreateLambda([this]()
+		{
+			if (ControlRig && ControlRig->GetVM())
+			{
+				return ControlRig->GetVM()->GetHaltedAtInstruction() != INDEX_NONE;
+			}
+			return false;
+	}));
+
 }
 
 void FControlRigEditor::ToggleExecuteGraph()
@@ -558,6 +577,11 @@ void FControlRigEditor::FillToolbar(FToolBarBuilder& ToolbarBuilder)
 			NAME_None, TAttribute<FText>(), TAttribute<FText>(), FSlateIcon(FControlRigEditorStyle::Get().GetStyleSetName(), "ControlRig.AutoCompileGraph"));
 
 		ToolbarBuilder.AddWidget(SNew(SBlueprintEditorSelectedDebugObjectWidget, SharedThis(this)));
+
+		ToolbarBuilder.AddSeparator();
+		ToolbarBuilder.AddToolBarButton(FControlRigBlueprintCommands::Get().ResumeExecution,
+			NAME_None, TAttribute<FText>(), TAttribute<FText>(), FSlateIcon(FControlRigEditorStyle::Get().GetStyleSetName(), "ControlRig.ResumeExecution"));
+
 	}
 	ToolbarBuilder.EndSection();
 }
@@ -1320,6 +1344,10 @@ void FControlRigEditor::Compile()
 		{
 			ControlRig->OnInitialized_AnyThread().Clear();
 			ControlRig->OnExecuted_AnyThread().Clear();
+			if (ControlRig->GetVM())
+			{
+				ControlRig->GetVM()->ExecutionHalted().RemoveAll(this);
+			}
 		}
 
 		if (bSetupModeEnabled)
@@ -2115,6 +2143,28 @@ void FControlRigEditor::HandleControlRigExecutedEvent(UControlRig* InControlRig,
 	UpdateGraphCompilerErrors();
 }
 
+void FControlRigEditor::HandleControlRigExecutionHalted(const int32 InstructionIndex, URigVMNode* InNode)
+{
+	if (InNode == nullptr)
+	{
+		return;
+	}
+	
+	if (UControlRigBlueprint* Blueprint = GetControlRigBlueprint())
+	{
+		if (Blueprint->GetAllModels().Contains(InNode->GetGraph()))
+		{
+			if(UControlRigGraph* EdGraph = Cast<UControlRigGraph>(Blueprint->GetEdGraph(InNode->GetGraph())))
+			{
+				if(UEdGraphNode* EdNode = EdGraph->FindNodeForModelNodeName(InNode->GetFName()))
+				{
+					JumpToHyperlink(EdNode, false);
+				}
+			}
+		}
+	}
+}
+
 void FControlRigEditor::CreateEditorModeManager()
 {
 	TSharedPtr<FAssetEditorModeManager> ModeManager = MakeShareable(FModuleManager::LoadModuleChecked<FPersonaModule>("Persona").CreatePersonaEditorModeManager());
@@ -2835,6 +2885,11 @@ void FControlRigEditor::UpdateControlRig()
 			ControlRig->OnExecuted_AnyThread().AddSP(this, &FControlRigEditor::HandleControlRigExecutedEvent);
 			ControlRig->RequestInit();
 			ControlRig->ControlModified().AddSP(this, &FControlRigEditor::HandleOnControlModified);
+
+			if (ControlRig->GetVM())
+			{
+				ControlRig->GetVM()->ExecutionHalted().AddSP(this, &FControlRigEditor::HandleControlRigExecutionHalted);
+			}
 		}
 	}
 }
diff --git a/Engine/Plugins/Experimental/ControlRig/Source/ControlRigEditor/Private/Editor/ControlRigEditor.h b/Engine/Plugins/Experimental/ControlRig/Source/ControlRigEditor/Private/Editor/ControlRigEditor.h
--- a/Engine/Plugins/Experimental/ControlRig/Source/ControlRigEditor/Private/Editor/ControlRigEditor.h
+++ b/Engine/Plugins/Experimental/ControlRig/Source/ControlRigEditor/Private/Editor/ControlRigEditor.h
@@ -212,6 +212,7 @@ protected:
 	void HandleModifiedEvent(ERigVMGraphNotifType InNotifType, URigVMGraph* InGraph, UObject* InSubject);
 	void HandleVMCompiledEvent(UBlueprint* InBlueprint, URigVM* InVM);
 	void HandleControlRigExecutedEvent(UControlRig* InControlRig, const EControlRigState InState, const FName& InEventName);
+	void HandleControlRigExecutionHalted(const int32 InstructionIndex, URigVMNode* InNode);
 
 	// FBaseToolKit overrides
 	void CreateEditorModeManager() override;
diff --git a/Engine/Plugins/Experimental/ControlRig/Source/ControlRigEditor/Private/Editor/SControlRigStackView.cpp b/Engine/Plugins/Experimental/ControlRig/Source/ControlRigEditor/Private/Editor/SControlRigStackView.cpp
--- a/Engine/Plugins/Experimental/ControlRig/Source/ControlRigEditor/Private/Editor/SControlRigStackView.cpp
+++ b/Engine/Plugins/Experimental/ControlRig/Source/ControlRigEditor/Private/Editor/SControlRigStackView.cpp
@@ -548,6 +548,8 @@ void SControlRigStackView::RefreshTreeView(URigVM* InVM)
 						}
 					}
 				}
+
+				InVM->ExecutionHalted().AddSP(this, &SControlRigStackView::HandleExecutionHalted);
 			}
 		}
 	}
@@ -616,6 +618,15 @@ void SControlRigStackView::OnVMCompiled(UBlueprint* InCompiledBlueprint, URigVM*
 	}
 }
 
+void SControlRigStackView::HandleExecutionHalted(const int32 HaltedAtInstruction, URigVMNode* InNode)
+{	
+	if (HaltedAtInstruction != INDEX_NONE && Operators.Num() > HaltedAtInstruction)
+	{
+		TreeView->SetSelection(Operators[HaltedAtInstruction]);
+		TreeView->SetScrollOffset(FMath::Max(HaltedAtInstruction-5, 0));
+	}
+}
+
 void SControlRigStackView::OnFilterTextChanged(const FText& SearchText)
 {
 	FilterText = SearchText;
diff --git a/Engine/Plugins/Experimental/ControlRig/Source/ControlRigEditor/Private/Editor/SControlRigStackView.h b/Engine/Plugins/Experimental/ControlRig/Source/ControlRigEditor/Private/Editor/SControlRigStackView.h
--- a/Engine/Plugins/Experimental/ControlRig/Source/ControlRigEditor/Private/Editor/SControlRigStackView.h
+++ b/Engine/Plugins/Experimental/ControlRig/Source/ControlRigEditor/Private/Editor/SControlRigStackView.h
@@ -97,6 +97,9 @@ private:
 
 	void OnVMCompiled(UBlueprint* InCompiledBlueprint, URigVM* InCompiledVM);
 
+	//* Focus on the instruction when the execution is halted */
+	void HandleExecutionHalted(const int32 HaltedAtInstruction, URigVMNode* InNode);
+
 	/** Search box widget */
 	TSharedPtr<SSearchBox> FilterBox;
 	FText FilterText;
diff --git a/Engine/Plugins/Experimental/ControlRig/Source/ControlRigEditor/Private/Graph/SControlRigGraphNode.cpp b/Engine/Plugins/Experimental/ControlRig/Source/ControlRigEditor/Private/Graph/SControlRigGraphNode.cpp
--- a/Engine/Plugins/Experimental/ControlRig/Source/ControlRigEditor/Private/Graph/SControlRigGraphNode.cpp
+++ b/Engine/Plugins/Experimental/ControlRig/Source/ControlRigEditor/Private/Graph/SControlRigGraphNode.cpp
@@ -1050,6 +1050,27 @@ FReply SControlRigGraphNode::HandleAddArrayElement(URigVMPin* InItem)
 	return FReply::Handled();
 }
 
+/** Populate the brushes array with any overlay brushes to render */
+void SControlRigGraphNode::GetOverlayBrushes(bool bSelected, const FVector2D WidgetSize, TArray<FOverlayBrushInfo>& Brushes) const
+{
+	UControlRigGraphNode* RigGraphNode = Cast<UControlRigGraphNode>(GraphNode);
+
+	const URigVMNode* VMNode = RigGraphNode->GetModelNode();
+	const bool bHasBreakpoint = VMNode->HasBreakpoint();
+	if (bHasBreakpoint)
+	{
+		FOverlayBrushInfo BreakpointOverlayInfo;
+
+		BreakpointOverlayInfo.Brush = FEditorStyle::GetBrush(TEXT("Kismet.DebuggerOverlay.Breakpoint.EnabledAndValid"));
+		if (BreakpointOverlayInfo.Brush != NULL)
+		{
+			BreakpointOverlayInfo.OverlayOffset -= BreakpointOverlayInfo.Brush->ImageSize / 2.f;
+		}
+
+		Brushes.Add(BreakpointOverlayInfo);
+	}
+}
+
 void SControlRigGraphNode::GetNodeInfoPopups(FNodeInfoContext* Context, TArray<FGraphInformationPopupInfo>& Popups) const
 {
 	FKismetNodeInfoContext* K2Context = (FKismetNodeInfoContext*)Context;
diff --git a/Engine/Plugins/Experimental/ControlRig/Source/ControlRigEditor/Private/Graph/SControlRigGraphNode.h b/Engine/Plugins/Experimental/ControlRig/Source/ControlRigEditor/Private/Graph/SControlRigGraphNode.h
--- a/Engine/Plugins/Experimental/ControlRig/Source/ControlRigEditor/Private/Graph/SControlRigGraphNode.h
+++ b/Engine/Plugins/Experimental/ControlRig/Source/ControlRigEditor/Private/Graph/SControlRigGraphNode.h
@@ -48,6 +48,7 @@ public:
 
 	virtual TSharedRef<SWidget> CreateNodeContentArea() override;
 	virtual TSharedPtr<SGraphPin> GetHoveredPin( const FGeometry& MyGeometry, const FPointerEvent& MouseEvent ) const override;
+	virtual void GetOverlayBrushes(bool bSelected, const FVector2D WidgetSize, TArray<FOverlayBrushInfo>& Brushes) const override;
 	virtual void GetNodeInfoPopups(FNodeInfoContext* Context, TArray<FGraphInformationPopupInfo>& Popups) const override;
 	virtual TArray<FOverlayWidgetInfo> GetOverlayWidgets(bool bSelected, const FVector2D& WidgetSize) const override;
 
diff --git a/Engine/Source/Developer/RigVMDeveloper/Public/RigVMModel/RigVMNode.h b/Engine/Source/Developer/RigVMDeveloper/Public/RigVMModel/RigVMNode.h
--- a/Engine/Source/Developer/RigVMDeveloper/Public/RigVMModel/RigVMNode.h
+++ b/Engine/Source/Developer/RigVMDeveloper/Public/RigVMModel/RigVMNode.h
@@ -203,6 +203,11 @@ public:
 	// Returns the number of visited / run instructions for this node
 	virtual int32 GetInstructionVisitedCount(URigVM* InVM, const FRigVMASTProxy& InProxy = FRigVMASTProxy(), bool bConsolidatePerNode = false) const; 
 
+	UFUNCTION(BlueprintCallable, Category = RigVMNode)
+	bool HasBreakpoint() const { return bHasBreakpoint; }
+
+	UFUNCTION(BlueprintCallable, Category = RigVMNode)
+	void SetHasBreakpoint(const bool bValue) { bHasBreakpoint = bValue; }
 private:
 
 	static const FString NodeColorName;
@@ -230,6 +235,9 @@ protected:
 	UPROPERTY(transient)
 	FName PreviousName;
 
+	UPROPERTY(transient)
+	bool bHasBreakpoint;
+
 private:
 
 	UPROPERTY()
diff --git a/Engine/Source/Runtime/RigVM/Private/RigVMCore/RigVM.cpp b/Engine/Source/Runtime/RigVM/Private/RigVMCore/RigVM.cpp
--- a/Engine/Source/Runtime/RigVM/Private/RigVMCore/RigVM.cpp
+++ b/Engine/Source/Runtime/RigVM/Private/RigVMCore/RigVM.cpp
@@ -775,8 +775,11 @@ bool URigVM::Initialize(FRigVMMemoryContainerPtrArray Memory, FRigVMFixedArray<v
 			{
 				if (DebugInfo->IsActive(Context.InstructionIndex))
 				{
-					HaltedAtInstruction = Context.InstructionIndex;
-					ExecutionHalted().Broadcast(Context.InstructionIndex);
+					if (HaltedAtInstruction != Context.InstructionIndex)
+					{
+						HaltedAtInstruction = Context.InstructionIndex;
+						ExecutionHalted().Broadcast(Context.InstructionIndex, Breakpoint->Subject);
+					}
 					return true;
 				}
 				else
@@ -1008,6 +1011,14 @@ bool URigVM::Initialize(FRigVMMemoryContainerPtrArray Memory, FRigVMFixedArray<v
 		Context.InstructionIndex++;
 	}
 
+#if WITH_EDITOR
+	if (HaltedAtInstruction != INDEX_NONE)
+	{
+		HaltedAtInstruction = INDEX_NONE;
+		ExecutionHalted().Broadcast(INDEX_NONE, nullptr);
+	}
+#endif
+	
 	return true;
 }
 
@@ -1070,7 +1081,6 @@ bool URigVM::Execute(FRigVMMemoryContainerPtrArray Memory, FRigVMFixedArray<void
 	if (DebugInfo)
 	{
 		DebugInfo->StartExecution();
-		HaltedAtInstruction = INDEX_NONE;
 	}
 #endif
 
@@ -1083,8 +1093,11 @@ bool URigVM::Execute(FRigVMMemoryContainerPtrArray Memory, FRigVMFixedArray<void
 			{
 				if (DebugInfo->IsActive(Context.InstructionIndex))
 				{
-					HaltedAtInstruction = Context.InstructionIndex;
-					ExecutionHalted().Broadcast(Context.InstructionIndex);
+					if (HaltedAtInstruction != Context.InstructionIndex)
+					{
+						HaltedAtInstruction = Context.InstructionIndex;
+						ExecutionHalted().Broadcast(Context.InstructionIndex, Breakpoint->Subject);
+					}
 					return true;
 				}
 				else
@@ -1434,6 +1447,13 @@ bool URigVM::Execute(FRigVMMemoryContainerPtrArray Memory, FRigVMFixedArray<void
 			case ERigVMOpCode::Exit:
 			{
 				ExecutionReachedExit().Broadcast();
+#if WITH_EDITOR					
+				if (HaltedAtInstruction != INDEX_NONE)
+				{
+					HaltedAtInstruction = INDEX_NONE;
+					ExecutionHalted().Broadcast(INDEX_NONE, nullptr);
+				}
+#endif
 				return true;
 			}
 			case ERigVMOpCode::BeginBlock:
@@ -1458,6 +1478,14 @@ bool URigVM::Execute(FRigVMMemoryContainerPtrArray Memory, FRigVMFixedArray<void
 		}
 	}
 
+#if WITH_EDITOR
+	if (HaltedAtInstruction != INDEX_NONE)
+	{
+		HaltedAtInstruction = INDEX_NONE;
+		ExecutionHalted().Broadcast(INDEX_NONE, nullptr);
+	}
+#endif
+
 	return true;
 }
 
diff --git a/Engine/Source/Runtime/RigVM/Public/RigVMCore/RigVM.h b/Engine/Source/Runtime/RigVM/Public/RigVMCore/RigVM.h
--- a/Engine/Source/Runtime/RigVM/Public/RigVMCore/RigVM.h
+++ b/Engine/Source/Runtime/RigVM/Public/RigVMCore/RigVM.h
@@ -124,7 +124,7 @@ public:
 	/** Bindable event for external objects to be notified when the VM reaches an Exit Operation */
 	DECLARE_EVENT(URigVM, FExecutionReachedExitEvent);
 #if WITH_EDITOR
-	DECLARE_EVENT_OneParam(URigVM, FExecutionHaltedEvent, int32);
+	DECLARE_EVENT_TwoParams(URigVM, FExecutionHaltedEvent, int32, URigVMNode*);
 #endif
 
 	URigVM();
diff --git a/Engine/Source/Runtime/RigVM/Public/RigVMCore/RigVMDebugInfo.h b/Engine/Source/Runtime/RigVM/Public/RigVMCore/RigVMDebugInfo.h
--- a/Engine/Source/Runtime/RigVM/Public/RigVMCore/RigVMDebugInfo.h
+++ b/Engine/Source/Runtime/RigVM/Public/RigVMCore/RigVMDebugInfo.h
@@ -5,20 +5,29 @@
 #include "CoreMinimal.h"
 #include "RigVMDebugInfo.generated.h"
 
+class URigVMNode;
 
 USTRUCT()
 struct RIGVM_API FRigVMBreakpoint
 {
 	GENERATED_BODY()
+
+	FRigVMBreakpoint()
+	: bIsActive(true)
+	, Subject(nullptr)
+	{
+
+	}
 	
-	FORCEINLINE FRigVMBreakpoint()
-	: bIsActive(true), bIsTemporary(false) 
+	FRigVMBreakpoint(URigVMNode* InNode)
+	: bIsActive(true)
+	, Subject(InNode)
 	{
 
 	}
 
 	bool bIsActive;
-	bool bIsTemporary;
+	URigVMNode* Subject;
 };
 
 USTRUCT()
@@ -52,11 +61,11 @@ struct RIGVM_API FRigVMDebugInfo
 		return Breakpoints.Find(InstructionIndex);
 	}
 	
-	FRigVMBreakpoint* AddBreakpoint(const uint16 InstructionIndex)
+	FRigVMBreakpoint* AddBreakpoint(const uint16 InstructionIndex, URigVMNode* InNode)
 	{
 		if (!Breakpoints.Contains(InstructionIndex))
 		{
-			if (FRigVMBreakpoint* NewBP = &Breakpoints.Add(InstructionIndex))
+			if (FRigVMBreakpoint* NewBP = &Breakpoints.Add(InstructionIndex, InNode))
 			{
 				// Do not override the state if it already exists
 				BreakpointActivationOnHit.FindOrAdd(InstructionIndex, 0);
