ControlRig: Breakpoints UI icons

diff --git a/Engine/Plugins/Experimental/ControlRig/Source/ControlRigEditor/Private/ControlRigEditorStyle.h b/Engine/Plugins/Experimental/ControlRig/Source/ControlRigEditor/Private/ControlRigEditorStyle.h
--- a/Engine/Plugins/Experimental/ControlRig/Source/ControlRigEditor/Private/ControlRigEditorStyle.h
+++ b/Engine/Plugins/Experimental/ControlRig/Source/ControlRigEditor/Private/ControlRigEditorStyle.h
@@ -5,6 +5,7 @@
 #include "Styling/SlateStyle.h"
 #include "Styling/SlateStyleRegistry.h"
 #include "Styling/SlateStyleMacros.h"
+#include "Styling/StyleColors.h"
 
 #define TTF_FONT(RelativePath, ...) FSlateFontInfo(RootToContentDir(RelativePath, TEXT(".ttf")), __VA_ARGS__)
 #define OTF_FONT(RelativePath, ...) FSlateFontInfo(RootToContentDir(RelativePath, TEXT(".otf")), __VA_ARGS__)
@@ -88,6 +89,8 @@ public:
 			Set("ControlRig.Bug.Normal", new IMAGE_BRUSH("Slate/ControlRig_Bug_28x", Icon14x14));
 			Set("ControlRig.Bug.Open", new IMAGE_BRUSH("Slate/ControlRig_BugOpen_28x", Icon14x14));
 			Set("ControlRig.Bug.Solid", new IMAGE_BRUSH("Slate/ControlRig_BugSolid_28x", Icon14x14));
+
+			Set("ControlRig.ResumeExecution", new IMAGE_BRUSH_SVG("Slate/simulate", Icon40x40));
 		}
 
 		// Graph styles
diff --git a/Engine/Plugins/Experimental/ControlRig/Source/ControlRigEditor/Private/Editor/ControlRigEditor.cpp b/Engine/Plugins/Experimental/ControlRig/Source/ControlRigEditor/Private/Editor/ControlRigEditor.cpp
--- a/Engine/Plugins/Experimental/ControlRig/Source/ControlRigEditor/Private/Editor/ControlRigEditor.cpp
+++ b/Engine/Plugins/Experimental/ControlRig/Source/ControlRigEditor/Private/Editor/ControlRigEditor.cpp
@@ -108,6 +108,7 @@ FControlRigEditor::FControlRigEditor()
 	, LastEventQueue(EControlRigEditorEventQueue::Setup)
 	, LastDebuggedRig()
 	, RigHierarchyTabCount(0)
+	, HaltedAtNode(nullptr)
 {
 }
 
@@ -2145,10 +2146,17 @@ void FControlRigEditor::HandleControlRigExecutedEvent(UControlRig* InControlRig,
 
 void FControlRigEditor::HandleControlRigExecutionHalted(const int32 InstructionIndex, URigVMNode* InNode)
 {
+	if (HaltedAtNode != nullptr)
+	{
+		HaltedAtNode->SetExecutionIsHaltedAtThisNode(false);
+	}
+	HaltedAtNode = InNode;
+
 	if (InNode == nullptr)
 	{
 		return;
-	}
+	}	
+	InNode->SetExecutionIsHaltedAtThisNode(true);
 	
 	if (UControlRigBlueprint* Blueprint = GetControlRigBlueprint())
 	{
diff --git a/Engine/Plugins/Experimental/ControlRig/Source/ControlRigEditor/Private/Editor/ControlRigEditor.h b/Engine/Plugins/Experimental/ControlRig/Source/ControlRigEditor/Private/Editor/ControlRigEditor.h
--- a/Engine/Plugins/Experimental/ControlRig/Source/ControlRigEditor/Private/Editor/ControlRigEditor.h
+++ b/Engine/Plugins/Experimental/ControlRig/Source/ControlRigEditor/Private/Editor/ControlRigEditor.h
@@ -414,6 +414,8 @@ protected:
 	int32 RigHierarchyTabCount;
 	TMap<FString, FString> KnownInstructionLimitWarnings;
 
+	URigVMNode* HaltedAtNode;
+
 	friend class FControlRigEditorMode;
 	friend class SControlRigStackView;
 	friend class SRigHierarchy;
diff --git a/Engine/Plugins/Experimental/ControlRig/Source/ControlRigEditor/Private/Graph/SControlRigGraphNode.cpp b/Engine/Plugins/Experimental/ControlRig/Source/ControlRigEditor/Private/Graph/SControlRigGraphNode.cpp
--- a/Engine/Plugins/Experimental/ControlRig/Source/ControlRigEditor/Private/Graph/SControlRigGraphNode.cpp
+++ b/Engine/Plugins/Experimental/ControlRig/Source/ControlRigEditor/Private/Graph/SControlRigGraphNode.cpp
@@ -1069,6 +1069,25 @@ void SControlRigGraphNode::GetOverlayBrushes(bool bSelected, const FVector2D Wid
 
 		Brushes.Add(BreakpointOverlayInfo);
 	}
+
+	// Paint red arrow pointing at breakpoint node that caused a halt in execution
+	{
+		FOverlayBrushInfo IPOverlayInfo;
+		if (VMNode->ExecutionIsHaltedAtThisNode())
+		{
+			IPOverlayInfo.Brush = FEditorStyle::GetBrush( TEXT("Kismet.DebuggerOverlay.InstructionPointerBreakpoint") );
+			if (IPOverlayInfo.Brush != NULL)
+			{
+				float Overlap = 10.f;
+				IPOverlayInfo.OverlayOffset.X = (WidgetSize.X/2.f) - (IPOverlayInfo.Brush->ImageSize.X/2.f);
+				IPOverlayInfo.OverlayOffset.Y = (Overlap - IPOverlayInfo.Brush->ImageSize.Y);
+			}
+
+			IPOverlayInfo.AnimationEnvelope = FVector2D(0.f, 10.f);
+
+			Brushes.Add(IPOverlayInfo);
+		}
+	}
 }
 
 void SControlRigGraphNode::GetNodeInfoPopups(FNodeInfoContext* Context, TArray<FGraphInformationPopupInfo>& Popups) const
diff --git a/Engine/Source/Developer/RigVMDeveloper/Public/RigVMModel/RigVMNode.h b/Engine/Source/Developer/RigVMDeveloper/Public/RigVMModel/RigVMNode.h
--- a/Engine/Source/Developer/RigVMDeveloper/Public/RigVMModel/RigVMNode.h
+++ b/Engine/Source/Developer/RigVMDeveloper/Public/RigVMModel/RigVMNode.h
@@ -208,6 +208,12 @@ public:
 
 	UFUNCTION(BlueprintCallable, Category = RigVMNode)
 	void SetHasBreakpoint(const bool bValue) { bHasBreakpoint = bValue; }
+
+	UFUNCTION(BlueprintCallable, Category = RigVMNode)
+	bool ExecutionIsHaltedAtThisNode() const { return bHaltedAtThisNode; }
+
+	UFUNCTION(BlueprintCallable, Category = RigVMNode)
+	void SetExecutionIsHaltedAtThisNode(const bool bValue) { bHaltedAtThisNode = bValue; }
 private:
 
 	static const FString NodeColorName;
@@ -238,6 +244,9 @@ protected:
 	UPROPERTY(transient)
 	bool bHasBreakpoint;
 
+	UPROPERTY(transient)
+	bool bHaltedAtThisNode;
+
 private:
 
 	UPROPERTY()
