Control Rig: New Parent/Position/Scale/Rotation Constraint Node

diff --git a/Engine/Plugins/Experimental/ControlRig/Source/ControlRig/Private/Units/Highlevel/Hierarchy/RigUnit_TransformConstraint.cpp b/Engine/Plugins/Experimental/ControlRig/Source/ControlRig/Private/Units/Highlevel/Hierarchy/RigUnit_TransformConstraint.cpp
--- a/Engine/Plugins/Experimental/ControlRig/Source/ControlRig/Private/Units/Highlevel/Hierarchy/RigUnit_TransformConstraint.cpp
+++ b/Engine/Plugins/Experimental/ControlRig/Source/ControlRig/Private/Units/Highlevel/Hierarchy/RigUnit_TransformConstraint.cpp
@@ -163,4 +163,780 @@ void FRigUnit_TransformConstraintPerItem::AddConstraintData(const FRigVMFixedArr
 	}
 
 	OutConstraintDataToTargets.Add(NewIndex, TargetIndex);
-}
\ No newline at end of file
+}
+
+FRigUnit_ParentConstraint_Execute()
+ {
+ 	URigHierarchy* Hierarchy = ExecuteContext.Hierarchy;	
+	if (Hierarchy)
+	{
+		if (!Child.IsValid())
+		{
+			return;
+		}
+		
+		FTransform ChildInitialGlobalTransform = Hierarchy->GetInitialGlobalTransform(Child);
+		FTransform ChildCurrentGlobalTransform = Hierarchy->GetGlobalTransform(Child, false);
+
+		// calculate total weight
+		float OverallWeight = 0;
+		for (const FConstraintParent& Parent : Parents)
+		{
+			const float ClampedWeight = FMath::Max(Parent.Weight, 0.f);
+
+			if (ClampedWeight < KINDA_SMALL_NUMBER)
+			{
+				continue;
+			}
+
+			if (!Parent.Item.IsValid())
+			{
+				continue;
+			}
+
+			OverallWeight += ClampedWeight;
+		}
+
+		if (OverallWeight > KINDA_SMALL_NUMBER)
+		{
+			const float WeightNormalizer = 1.0f / OverallWeight;
+
+			FTransform MixedGlobalTransform = FTransform::Identity;
+
+			// initial rotation needs to be (0,0,0,0) instead of (0,0,0,1) due to Quaternion Blending Math
+			MixedGlobalTransform.SetRotation(FQuat(0.f, 0.f, 0.f, 0.f));
+			MixedGlobalTransform.SetScale3D(FVector::ZeroVector);
+
+			float AccumulatedWeight = 0.0f;
+
+			for (const FConstraintParent& Parent : Parents)
+			{
+				const float ClampedWeight = FMath::Max(Parent.Weight, 0.f);
+
+				if (ClampedWeight < KINDA_SMALL_NUMBER)
+				{
+					continue;
+				}
+
+				if (!Parent.Item.IsValid())
+				{
+					continue;
+				}
+
+				const float NormalizedWeight = ClampedWeight * WeightNormalizer;
+				AccumulatedWeight += NormalizedWeight;
+
+				FTransform OffsetTransform = FTransform::Identity;
+				FTransform ParentCurrentGlobalTransform = Hierarchy->GetGlobalTransform(Parent.Item, false);
+
+				if (bMaintainOffset)
+				{
+					FTransform ParentInitialGlobalTransform = Hierarchy->GetInitialGlobalTransform(Parent.Item);
+					// offset transform is a transform that transforms parent to child
+					OffsetTransform = ChildInitialGlobalTransform.GetRelativeTransform(ParentInitialGlobalTransform);
+					OffsetTransform.NormalizeRotation();
+				}
+
+				FTransform OffsetParentTransform = OffsetTransform * ParentCurrentGlobalTransform;
+				OffsetParentTransform.NormalizeRotation();
+
+				// deal with different interpolation types
+				if (AdvancedSettings.InterpolationType == EConstraintInterpType::Average)
+				{
+					// component-wise average
+					MixedGlobalTransform.AccumulateWithShortestRotation(OffsetParentTransform, ScalarRegister(NormalizedWeight));
+				}
+				else if (AdvancedSettings.InterpolationType == EConstraintInterpType::Shortest)
+				{
+					FQuat MixedGlobalQuat = MixedGlobalTransform.GetRotation();
+					FQuat OffsetParentQuat = OffsetParentTransform.GetRotation();
+
+					if (MixedGlobalQuat == FQuat(0.0f, 0.0f, 0.0f, 0.0f))
+					{
+						MixedGlobalTransform = OffsetParentTransform;
+					}
+					else
+					{
+						const float Alpha = NormalizedWeight / AccumulatedWeight;
+
+						MixedGlobalTransform.LerpTranslationScale3D(MixedGlobalTransform, OffsetParentTransform, ScalarRegister(Alpha));
+						MixedGlobalQuat = FQuat::Slerp(MixedGlobalQuat, OffsetParentQuat, NormalizedWeight);
+						MixedGlobalTransform.SetRotation(MixedGlobalQuat);
+					}
+				}
+				else
+				{
+					// invalid interpolation type
+					ensure(false);
+					MixedGlobalTransform = ChildCurrentGlobalTransform;
+					break;
+				}
+			}
+
+			MixedGlobalTransform.NormalizeRotation();
+
+			// handle filtering, performed in local(parent) space
+			FTransform ChildParentGlobalTransform = Hierarchy->GetParentTransform(Child, false);
+			FTransform MixedLocalTransform = MixedGlobalTransform.GetRelativeTransform(ChildParentGlobalTransform);
+			MixedLocalTransform.NormalizeRotation();
+			FVector MixedTranslation = MixedLocalTransform.GetTranslation();
+			FQuat MixedRotation = MixedLocalTransform.GetRotation();
+			FVector MixedEulerRotation = FControlRigMathLibrary::EulerFromQuat(MixedRotation, AdvancedSettings.RotationOrderForFilter);
+			FVector MixedScale = MixedLocalTransform.GetScale3D();
+
+			FTransform ChildCurrentLocalTransform = Hierarchy->GetLocalTransform(Child, false);
+			
+			// Controls have an offset transform built-in and thus need to be handled a bit differently
+			FTransform AdditionalOffsetTransform = FTransform::Identity;
+			
+			if (Child.Type == ERigElementType::Control)
+			{
+				if (FRigControlElement* ChildAsControlElement = Hierarchy->Find<FRigControlElement>(Child))
+				{
+					AdditionalOffsetTransform = Hierarchy->GetControlOffsetTransform(ChildAsControlElement, ERigTransformType::CurrentLocal);
+					// Control's local(parent) space transform = control local value * offset
+					ChildCurrentLocalTransform *= AdditionalOffsetTransform;
+				}
+			}
+			
+			FVector ChildTranslation = ChildCurrentLocalTransform.GetTranslation();
+			FQuat ChildRotation = ChildCurrentLocalTransform.GetRotation();
+			FVector ChildEulerRotation = FControlRigMathLibrary::EulerFromQuat(ChildRotation, AdvancedSettings.RotationOrderForFilter);
+			FVector ChildScale = ChildCurrentLocalTransform.GetScale3D();
+
+			FVector FilteredTranslation;
+			FilteredTranslation.X = Filter.TranslationFilter.bX ? MixedTranslation.X : ChildTranslation.X;
+			FilteredTranslation.Y = Filter.TranslationFilter.bY ? MixedTranslation.Y : ChildTranslation.Y;
+			FilteredTranslation.Z = Filter.TranslationFilter.bZ ? MixedTranslation.Z : ChildTranslation.Z;
+
+			FVector FilteredEulerRotation;
+			FilteredEulerRotation.X = Filter.RotationFilter.bX ? MixedEulerRotation.X: ChildEulerRotation.X;
+			FilteredEulerRotation.Y = Filter.RotationFilter.bY ? MixedEulerRotation.Y : ChildEulerRotation.Y;
+			FilteredEulerRotation.Z = Filter.RotationFilter.bZ ? MixedEulerRotation.Z : ChildEulerRotation.Z;
+
+			FVector FilteredScale;
+			FilteredScale.X = Filter.ScaleFilter.bX ? MixedScale.X : ChildScale.X;
+			FilteredScale.Y = Filter.ScaleFilter.bY ? MixedScale.Y : ChildScale.Y;
+			FilteredScale.Z = Filter.ScaleFilter.bZ ? MixedScale.Z : ChildScale.Z;
+
+			FTransform FilteredMixedLocalTransform(FControlRigMathLibrary::QuatFromEuler(FilteredEulerRotation, AdvancedSettings.RotationOrderForFilter), FilteredTranslation, FilteredScale);
+
+			FTransform FinalLocalTransform = FilteredMixedLocalTransform;
+			
+			if (Child.Type == ERigElementType::Control)
+			{
+				// need to convert back to offset space for the actual control value
+				FinalLocalTransform = FilteredMixedLocalTransform.GetRelativeTransform(AdditionalOffsetTransform);
+				FinalLocalTransform.NormalizeRotation();
+			}
+			
+			Hierarchy->SetLocalTransform(Child, FinalLocalTransform);
+		}
+	} 
+ }
+
+#if WITH_DEV_AUTOMATION_TESTS
+#include "Units/RigUnitTest.h"
+
+ IMPLEMENT_RIGUNIT_AUTOMATION_TEST(FRigUnit_ParentConstraint)
+ {
+	// use euler rotation here to match other software's rotation representation more easily
+	EControlRigRotationOrder Order = EControlRigRotationOrder::XZY;
+	const FRigElementKey Child = Controller->AddBone(TEXT("Child"), FRigElementKey(), FTransform( FControlRigMathLibrary::QuatFromEuler( FVector(-10, -10, -10), Order), FVector(0.f, 0.f, 0.f)), true, ERigBoneType::User);
+	const FRigElementKey Parent1 = Controller->AddBone(TEXT("Parent1"), FRigElementKey(), FTransform( FControlRigMathLibrary::QuatFromEuler( FVector(30, -30, -30), Order), FVector(20.f, 20.f, 20.f)), true, ERigBoneType::User);
+	const FRigElementKey Parent2 = Controller->AddBone(TEXT("Parent2"), FRigElementKey(), FTransform( FControlRigMathLibrary::QuatFromEuler( FVector(-40, -40, 40), Order), FVector(40.f, 40.f, 40.f)), true, ERigBoneType::User);
+	const FRigElementKey Parent3 = Controller->AddBone(TEXT("Parent3"), FRigElementKey(), FTransform( FControlRigMathLibrary::QuatFromEuler( FVector(-50, 50, -50), Order), FVector(60.f, 60.f, 60.f)), true, ERigBoneType::User);
+	const FRigElementKey Parent4 = Controller->AddBone(TEXT("Parent4"), FRigElementKey(), FTransform( FControlRigMathLibrary::QuatFromEuler( FVector(60, 60, 60), Order), FVector(80.f, 80.f, 80.f)), true, ERigBoneType::User);
+	
+	Unit.ExecuteContext.Hierarchy = Hierarchy;
+	Unit.Child = Child;
+
+	Unit.Parents.Add(FConstraintParent(Parent1, 1.0));
+	Unit.Parents.Add(FConstraintParent(Parent2, 1.0));
+	Unit.Parents.Add(FConstraintParent(Parent3, 1.0));
+	Unit.Parents.Add(FConstraintParent(Parent4, 1.0));
+
+	Hierarchy->ResetPoseToInitial(ERigElementType::Bone);
+	Unit.bMaintainOffset = false;
+	Unit.AdvancedSettings.InterpolationType = EConstraintInterpType::Average;
+	
+	Execute();
+	AddErrorIfFalse(Hierarchy->GetGlobalTransform(0).GetTranslation().Equals(FVector(50.f, 50.f, 50.f)), TEXT("unexpected translation for average interpolation type"));
+
+	FQuat Result = Hierarchy->GetGlobalTransform(0).GetRotation();
+	FQuat Expected = FControlRigMathLibrary::QuatFromEuler( FVector(-0.852f, 15.189f, -0.572f), Order);
+	AddErrorIfFalse(Result.Equals(Expected, 0.001f), TEXT("unexpected rotation for average interpolation type"));
+	
+	Hierarchy->ResetPoseToInitial(ERigElementType::Bone);
+	Unit.bMaintainOffset = false;
+	Unit.AdvancedSettings.InterpolationType = EConstraintInterpType::Shortest;
+	Execute();
+	
+	Result = Hierarchy->GetGlobalTransform(0).GetRotation();
+	Expected = FControlRigMathLibrary::QuatFromEuler( FVector(16.74f, 8.865f, -5.562f), Order);
+	AddErrorIfFalse(Result.Equals(Expected, 0.001f), TEXT("unexpected rotation for shortest interpolation type"));
+
+	
+	Hierarchy->ResetPoseToInitial(ERigElementType::Bone);
+	Hierarchy->SetGlobalTransform(2, FTransform(FControlRigMathLibrary::QuatFromEuler( FVector(100, 100, -100), Order), FVector(-40.f, -40.f, -40.f)));
+	Unit.bMaintainOffset = true;
+	Unit.AdvancedSettings.InterpolationType = EConstraintInterpType::Average;
+	
+	Execute();
+	AddErrorIfFalse(Hierarchy->GetGlobalTransform(0).GetTranslation().Equals(FVector(-8.66f, 7.01f, -13.0f), 0.02f),
+                    TEXT("unexpected translation for maintain offset and average interpolation type"));
+	
+	Result = Hierarchy->GetGlobalTransform(0).GetRotation();
+	Expected = FControlRigMathLibrary::QuatFromEuler( FVector(5.408f, -5.679f, -34.44f), Order);
+	AddErrorIfFalse(Result.Equals(Expected, 0.001f), TEXT("unexpected rotation for maintain offset and average interpolation type"));
+
+	
+	Hierarchy->ResetPoseToInitial(ERigElementType::Bone);
+	Hierarchy->SetGlobalTransform(2, FTransform(FControlRigMathLibrary::QuatFromEuler( FVector(100.0f, 100.0f, -100.0f), Order), FVector(-40.f, -40.f, -40.f)));
+	Unit.bMaintainOffset = true;
+	Unit.AdvancedSettings.InterpolationType = EConstraintInterpType::Shortest;
+	
+	Execute();
+	Result = Hierarchy->GetGlobalTransform(0).GetRotation();
+	Expected = FControlRigMathLibrary::QuatFromEuler( FVector(-1.209f, -8.332f, -25.022f), Order);
+	AddErrorIfFalse(Result.Equals(Expected, 0.001f), TEXT("unexpected rotation for maintain offset and shortest interpolation type"));
+	
+	return true;
+ }
+
+#endif
+
+FRigUnit_PositionConstraint_Execute()
+{
+	URigHierarchy* Hierarchy = ExecuteContext.Hierarchy;	
+	if (Hierarchy)
+	{
+		if (!Child.IsValid())
+		{
+			return;
+		}
+		
+		FTransform ChildInitialGlobalTransform = Hierarchy->GetInitialGlobalTransform(Child);
+
+		float OverallWeight = 0;
+		for (const FConstraintParent& Parent : Parents)
+		{
+			const float ClampedWeight = FMath::Max(Parent.Weight, 0.f);
+
+			if (ClampedWeight < KINDA_SMALL_NUMBER)
+			{
+				continue;
+			}
+
+			if (!Parent.Item.IsValid())
+			{
+				continue;
+			}
+
+			OverallWeight += ClampedWeight;
+		}
+
+		if (OverallWeight > KINDA_SMALL_NUMBER)
+		{
+			const float WeightNormalizer = 1.0f / OverallWeight;
+
+			FVector MixedGlobalPosition = FVector::ZeroVector;
+
+			for (const FConstraintParent& Parent : Parents)
+			{
+				const float ClampedWeight = FMath::Max(Parent.Weight, 0.f);
+
+				if (ClampedWeight < KINDA_SMALL_NUMBER)
+				{
+					continue;
+				}
+
+				if (!Parent.Item.IsValid())
+				{
+					continue;
+				}
+
+				const float NormalizedWeight = ClampedWeight * WeightNormalizer;
+
+				FVector OffsetPosition = FVector::ZeroVector;
+				FTransform ParentCurrentGlobalTransform = Hierarchy->GetGlobalTransform(Parent.Item, false);
+
+				if (bMaintainOffset)
+				{
+					FTransform ParentInitialGlobalTransform = Hierarchy->GetInitialGlobalTransform(Parent.Item);
+					OffsetPosition = ChildInitialGlobalTransform.GetLocation() - ParentInitialGlobalTransform.GetLocation();
+				}
+
+				FVector OffsetParentPosition = OffsetPosition + ParentCurrentGlobalTransform.GetLocation();
+
+				MixedGlobalPosition += OffsetParentPosition * NormalizedWeight;
+			}
+
+			// handle filtering, performed in local space
+			FTransform ChildParentGlobalTransform = Hierarchy->GetParentTransform(Child, false);
+			FVector MixedPosition = MixedGlobalPosition - ChildParentGlobalTransform.GetLocation();
+			
+			FTransform ChildCurrentLocalTransform = Hierarchy->GetLocalTransform(Child);
+			FVector ChildPosition = ChildCurrentLocalTransform.GetTranslation();
+			
+			// Controls need to be handled a bit differently
+			FTransform AdditionalOffsetTransform = FTransform::Identity;
+			
+			if (Child.Type == ERigElementType::Control)
+			{
+				if (FRigControlElement* ChildAsControlElement = Hierarchy->Find<FRigControlElement>(Child))
+				{
+					AdditionalOffsetTransform = Hierarchy->GetControlOffsetTransform(ChildAsControlElement, ERigTransformType::CurrentLocal);
+					// Control's local(parent) space position = local + offset
+					ChildPosition += AdditionalOffsetTransform.GetLocation();
+				}
+			}
+			
+			FVector FilteredPosition;
+			FilteredPosition.X = Filter.bX ? MixedPosition.X : ChildPosition.X;
+			FilteredPosition.Y = Filter.bY ? MixedPosition.Y : ChildPosition.Y;
+			FilteredPosition.Z = Filter.bZ ? MixedPosition.Z : ChildPosition.Z;
+
+			FVector FinalPosition = FilteredPosition;
+			if (Child.Type == ERigElementType::Control)
+			{
+				// need to convert back to offset space for the actual control value
+				FinalPosition = FilteredPosition - AdditionalOffsetTransform.GetLocation();
+			}
+
+			Hierarchy->SetLocalTransform(Child, FTransform(ChildCurrentLocalTransform.GetRotation(), FinalPosition, ChildCurrentLocalTransform.GetScale3D()));
+		}
+	}
+}
+
+#if WITH_DEV_AUTOMATION_TESTS
+#include "Units/RigUnitTest.h"
+
+ IMPLEMENT_RIGUNIT_AUTOMATION_TEST(FRigUnit_PositionConstraint)
+ {
+	// use euler rotation here to match other software's rotation representation more easily
+	EControlRigRotationOrder Order = EControlRigRotationOrder::XZY;
+	const FRigElementKey Child = Controller->AddBone(TEXT("Child"), FRigElementKey(), FTransform(FVector(0.f, 0.f, 0.f)), true, ERigBoneType::User);
+	const FRigElementKey Parent1 = Controller->AddBone(TEXT("Parent1"), FRigElementKey(), FTransform(FVector(20.f, 20.f, 20.f)), true, ERigBoneType::User);
+	const FRigElementKey Parent2 = Controller->AddBone(TEXT("Parent2"), FRigElementKey(), FTransform(FVector(40.f, 40.f, 40.f)), true, ERigBoneType::User);
+	const FRigElementKey Parent3 = Controller->AddBone(TEXT("Parent3"), FRigElementKey(), FTransform(FVector(60.f, 60.f, 60.f)), true, ERigBoneType::User);
+	
+	Unit.ExecuteContext.Hierarchy = Hierarchy;
+	Unit.Child = Child;
+
+	Unit.Parents.Add(FConstraintParent(Parent1, 1.0));
+	Unit.Parents.Add(FConstraintParent(Parent2, 1.0));
+	Unit.Parents.Add(FConstraintParent(Parent3, 1.0));
+
+	Hierarchy->ResetPoseToInitial(ERigElementType::Bone);
+	Unit.bMaintainOffset = false;
+	
+	Execute();
+	AddErrorIfFalse(Hierarchy->GetGlobalTransform(0).GetTranslation().Equals(FVector(40.f, 40.f, 40.f)), TEXT("unexpected translation for maintain offset off"));
+
+	
+	Hierarchy->ResetPoseToInitial(ERigElementType::Bone);
+	Hierarchy->SetGlobalTransform(2, FTransform(FVector(-40.f, -40.f, -40.f)));
+	Unit.bMaintainOffset = true;
+	
+	Execute();
+	AddErrorIfFalse(Hierarchy->GetGlobalTransform(0).GetTranslation().Equals(FVector(-26.67f, -26.67f, -26.67f), 0.01f),
+                    TEXT("unexpected translation for maintain offset on"));
+
+	return true;
+ }
+
+#endif
+
+FRigUnit_RotationConstraint_Execute()
+{
+	URigHierarchy* Hierarchy = ExecuteContext.Hierarchy;	
+	if (Hierarchy)
+	{
+		if (!Child.IsValid())
+		{
+			return;
+		}
+
+		FTransform ChildInitialGlobalTransform = Hierarchy->GetInitialGlobalTransform(Child);
+		FTransform ChildCurrentGlobalTransform = Hierarchy->GetGlobalTransform(Child, false);
+
+		float OverallWeight = 0;
+		for (const FConstraintParent& Parent : Parents)
+		{
+			const float ClampedWeight = FMath::Max(Parent.Weight, 0.f);
+
+			if (ClampedWeight < KINDA_SMALL_NUMBER)
+			{
+				continue;
+			}
+
+			if (!Parent.Item.IsValid())
+			{
+				continue;
+			}
+
+			OverallWeight += ClampedWeight;
+		}
+
+		if (OverallWeight > KINDA_SMALL_NUMBER)
+		{
+			const float WeightNormalizer = 1.0f / OverallWeight;
+
+			FQuat MixedGlobalRotation = FQuat(0,0,0,0);
+
+			for (const FConstraintParent& Parent : Parents)
+			{
+				const float ClampedWeight = FMath::Max(Parent.Weight, 0.f);
+
+				if (ClampedWeight < KINDA_SMALL_NUMBER)
+				{
+					continue;
+				}
+
+				if (!Parent.Item.IsValid())
+				{
+					continue;
+				}
+
+				const float NormalizedWeight = ClampedWeight * WeightNormalizer;
+
+				FQuat OffsetRotation = FQuat::Identity;
+				FTransform ParentCurrentGlobalTransform = Hierarchy->GetGlobalTransform(Parent.Item, false);
+
+				if (bMaintainOffset)
+				{
+					FTransform ParentInitialGlobalTransform = Hierarchy->GetInitialGlobalTransform(Parent.Item);
+					OffsetRotation = ParentInitialGlobalTransform.GetRotation().Inverse() * ChildInitialGlobalTransform.GetRotation();
+					OffsetRotation.Normalize();
+				}
+
+				FQuat OffsetParentRotation = ParentCurrentGlobalTransform.GetRotation() * OffsetRotation;
+
+				// deal with different interpolation types
+				if (AdvancedSettings.InterpolationType == EConstraintInterpType::Average)
+				{
+					// component-wise average
+					FQuat WeightedOffsetParentRotation = OffsetParentRotation * NormalizedWeight; 
+                    
+                    // To ensure the 'shortest route', we make sure the dot product between the accumulator and the incoming child Quat is positive.
+					if ((WeightedOffsetParentRotation | MixedGlobalRotation) < 0.f)
+					{
+						MixedGlobalRotation -= WeightedOffsetParentRotation;	
+					}
+					else
+					{
+						MixedGlobalRotation += WeightedOffsetParentRotation;	
+					}
+				}
+				else if (AdvancedSettings.InterpolationType == EConstraintInterpType::Shortest)
+				{
+					if (MixedGlobalRotation == FQuat(0.0f, 0.0f, 0.0f, 0.0f))
+					{
+						MixedGlobalRotation = OffsetParentRotation;
+					}
+					else
+					{
+						MixedGlobalRotation = FQuat::Slerp(MixedGlobalRotation, OffsetParentRotation, NormalizedWeight);
+					}
+				}
+				else
+				{
+					// invalid interpolation type
+					ensure(false);
+					MixedGlobalRotation = ChildCurrentGlobalTransform.GetRotation();
+					break;
+				}	
+			}
+
+			MixedGlobalRotation.Normalize();
+			
+			// handle filtering, performed in local space
+			FTransform ChildParentGlobalTransform = Hierarchy->GetParentTransform(Child, false);
+			FQuat MixedLocalRotation = ChildParentGlobalTransform.GetRotation().Inverse() * MixedGlobalRotation;
+			FVector MixedEulerRotation = FControlRigMathLibrary::EulerFromQuat(MixedLocalRotation, AdvancedSettings.RotationOrderForFilter);
+			
+			FTransform ChildCurrentLocalTransform = Hierarchy->GetLocalTransform(Child, false);
+			FQuat ChildRotation = ChildCurrentLocalTransform.GetRotation();
+			
+			// Controls need to be handled a bit differently
+			FTransform AdditionalOffsetTransform = FTransform::Identity;
+			
+			if (Child.Type == ERigElementType::Control)
+			{
+				if (FRigControlElement* ChildAsControlElement = Hierarchy->Find<FRigControlElement>(Child))
+				{
+					AdditionalOffsetTransform = Hierarchy->GetControlOffsetTransform(ChildAsControlElement, ERigTransformType::CurrentLocal);
+					// Control's local(parent) space transform = local * offset
+					ChildRotation = AdditionalOffsetTransform.GetRotation() * ChildRotation;
+				}
+			}
+			
+			FVector ChildEulerRotation = FControlRigMathLibrary::EulerFromQuat(ChildRotation, AdvancedSettings.RotationOrderForFilter);	
+			
+			FVector FilteredEulerRotation;
+			FilteredEulerRotation.X = Filter.bX ? MixedEulerRotation.X : ChildEulerRotation.X;
+			FilteredEulerRotation.Y = Filter.bY ? MixedEulerRotation.Y : ChildEulerRotation.Y;
+			FilteredEulerRotation.Z = Filter.bZ ? MixedEulerRotation.Z : ChildEulerRotation.Z;
+			
+			FQuat FilterRotation = FControlRigMathLibrary::QuatFromEuler(FilteredEulerRotation, AdvancedSettings.RotationOrderForFilter); 
+
+			FQuat FinalRotation = FilterRotation;
+			if (Child.Type == ERigElementType::Control)
+			{
+				// need to convert back to offset space for the actual control value
+				FinalRotation = AdditionalOffsetTransform.GetRotation().Inverse() * FilterRotation;
+			}	
+			
+			Hierarchy->SetLocalTransform(Child, FTransform(FinalRotation, ChildCurrentLocalTransform.GetTranslation(), ChildCurrentLocalTransform.GetScale3D()));
+		}
+	} 	
+}
+
+#if WITH_DEV_AUTOMATION_TESTS
+#include "Units/RigUnitTest.h"
+
+IMPLEMENT_RIGUNIT_AUTOMATION_TEST(FRigUnit_RotationConstraint)
+{
+	// the rotation constraint is expected to behave similarly to parent constraint with translation filter turned off.
+
+	// use euler rotation here to match other software's rotation representation more easily
+	EControlRigRotationOrder Order = EControlRigRotationOrder::XZY;
+	const FRigElementKey Child = Controller->AddBone(TEXT("Child"), FRigElementKey(), FTransform( FControlRigMathLibrary::QuatFromEuler( FVector(-10, -10, -10), Order), FVector(0.f, 0.f, 0.f)), true, ERigBoneType::User);
+	const FRigElementKey Parent1 = Controller->AddBone(TEXT("Parent1"), FRigElementKey(), FTransform( FControlRigMathLibrary::QuatFromEuler( FVector(30, -30, -30), Order), FVector(20.f, 20.f, 20.f)), true, ERigBoneType::User);
+	const FRigElementKey Parent2 = Controller->AddBone(TEXT("Parent2"), FRigElementKey(), FTransform( FControlRigMathLibrary::QuatFromEuler( FVector(-40, -40, 40), Order), FVector(40.f, 40.f, 40.f)), true, ERigBoneType::User);
+	const FRigElementKey Parent3 = Controller->AddBone(TEXT("Parent3"), FRigElementKey(), FTransform( FControlRigMathLibrary::QuatFromEuler( FVector(-50, 50, -50), Order), FVector(60.f, 60.f, 60.f)), true, ERigBoneType::User);
+	const FRigElementKey Parent4 = Controller->AddBone(TEXT("Parent4"), FRigElementKey(), FTransform( FControlRigMathLibrary::QuatFromEuler( FVector(60, 60, 60), Order), FVector(80.f, 80.f, 80.f)), true, ERigBoneType::User);
+	
+	Unit.ExecuteContext.Hierarchy = Hierarchy;
+	Unit.Child = Child;
+
+	Unit.Parents.Add(FConstraintParent(Parent1, 1.0));
+	Unit.Parents.Add(FConstraintParent(Parent2, 1.0));
+	Unit.Parents.Add(FConstraintParent(Parent3, 1.0));
+	Unit.Parents.Add(FConstraintParent(Parent4, 1.0));
+
+	Hierarchy->ResetPoseToInitial(ERigElementType::Bone);
+	Unit.bMaintainOffset = false;
+	Unit.AdvancedSettings.InterpolationType = EConstraintInterpType::Average;
+	
+	Execute();
+
+	FQuat Result = Hierarchy->GetGlobalTransform(0).GetRotation();
+	FQuat Expected = FControlRigMathLibrary::QuatFromEuler( FVector(-0.853f, 15.189f, -0.572f), Order);
+	AddErrorIfFalse(Result.Equals(Expected, 0.001f), TEXT("unexpected rotation for average interpolation type"));
+	
+	Hierarchy->ResetPoseToInitial(ERigElementType::Bone);
+	Unit.bMaintainOffset = false;
+	Unit.AdvancedSettings.InterpolationType = EConstraintInterpType::Shortest;
+	
+	Execute();
+	
+	Result = Hierarchy->GetGlobalTransform(0).GetRotation();
+	Expected = FControlRigMathLibrary::QuatFromEuler( FVector(16.74f, 8.865f, -5.562f), Order);
+	AddErrorIfFalse(Result.Equals(Expected, 0.001f), TEXT("unexpected rotation for shortest interpolation type"));
+
+	
+	Hierarchy->ResetPoseToInitial(ERigElementType::Bone);
+	Hierarchy->SetGlobalTransform(2, FTransform(FControlRigMathLibrary::QuatFromEuler( FVector(100, 100, -100), Order), FVector(-40.f, -40.f, -40.f)));
+	Unit.bMaintainOffset = true;
+	Unit.AdvancedSettings.InterpolationType = EConstraintInterpType::Average;
+	
+	Execute();
+	
+	Result = Hierarchy->GetGlobalTransform(0).GetRotation();
+	Expected = FControlRigMathLibrary::QuatFromEuler( FVector(5.408f, -5.679f, -34.44f), Order);
+	AddErrorIfFalse(Result.Equals(Expected, 0.001f), TEXT("unexpected rotation for maintain offset and average interpolation type"));
+
+	
+	Hierarchy->ResetPoseToInitial(ERigElementType::Bone);
+	Hierarchy->SetGlobalTransform(2, FTransform(FControlRigMathLibrary::QuatFromEuler( FVector(100.0f, 100.0f, -100.0f), Order), FVector(-40.f, -40.f, -40.f)));
+	Unit.bMaintainOffset = true;
+	Unit.AdvancedSettings.InterpolationType = EConstraintInterpType::Shortest;
+	
+	Execute();
+	
+	Result = Hierarchy->GetGlobalTransform(0).GetRotation();
+	Expected = FControlRigMathLibrary::QuatFromEuler( FVector(-1.209f, -8.332f, -25.022f), Order);
+	AddErrorIfFalse(Result.Equals(Expected, 0.001f), TEXT("unexpected rotation for maintain offset and shortest interpolation type"));
+	
+	return true;
+}
+#endif
+
+FRigUnit_ScaleConstraint_Execute()
+{
+	TFunction<FVector (const FVector&)> GetNonZeroScale([RigVMExecuteContext, Context](const FVector& InScale)
+	{
+		FVector NonZeroScale = InScale;
+       
+       	bool bDetectedZeroScale = false;
+       	if (FMath::Abs(NonZeroScale.X) < KINDA_SMALL_NUMBER)
+       	{
+       		NonZeroScale.X = KINDA_SMALL_NUMBER * FMath::Sign(NonZeroScale.X);
+       		bDetectedZeroScale = true;
+       	}
+       	if (FMath::Abs(NonZeroScale.Y) < KINDA_SMALL_NUMBER)
+       	{
+       		NonZeroScale.Y = KINDA_SMALL_NUMBER * FMath::Sign(NonZeroScale.Y);
+       		bDetectedZeroScale = true;
+       	}
+       	if (FMath::Abs(NonZeroScale.Z) < KINDA_SMALL_NUMBER)
+       	{
+       		NonZeroScale.Z = KINDA_SMALL_NUMBER * FMath::Sign(NonZeroScale.Z);
+       		bDetectedZeroScale = true;
+       	}
+       
+       	if (bDetectedZeroScale)
+       	{
+       		UE_CONTROLRIG_RIGUNIT_REPORT_WARNING(TEXT("Scale value: (%f, %f, %f) contains value too close to 0 to use with scale constraint."), InScale.X, InScale.Y, InScale.Z);		
+       	}
+       	
+       	return NonZeroScale;	
+	});
+	
+	URigHierarchy* Hierarchy = ExecuteContext.Hierarchy;	
+	if (Hierarchy)
+	{
+		if (!Child.IsValid())
+		{
+			return;
+		}
+		
+		FTransform ChildInitialGlobalTransform = Hierarchy->GetInitialGlobalTransform(Child);
+
+		float OverallWeight = 0;
+		for (const FConstraintParent& Parent : Parents)
+		{
+			const float ClampedWeight = FMath::Max(Parent.Weight, 0.f);
+
+			if (ClampedWeight < KINDA_SMALL_NUMBER)
+			{
+				continue;
+			}
+
+			if (!Parent.Item.IsValid())
+			{
+				continue;
+			}
+
+			OverallWeight += ClampedWeight;
+		}
+
+		if (OverallWeight > KINDA_SMALL_NUMBER)
+		{
+			const float WeightNormalizer = 1.0f / OverallWeight;
+
+			FVector MixedGlobalScale = FVector::OneVector;
+
+			for (const FConstraintParent& Parent : Parents)
+			{
+				const float ClampedWeight = FMath::Max(Parent.Weight, 0.f);
+
+				if (ClampedWeight < KINDA_SMALL_NUMBER)
+				{
+					continue;
+				}
+
+				if (!Parent.Item.IsValid())
+				{
+					continue;
+				}
+
+				const float NormalizedWeight = ClampedWeight * WeightNormalizer;
+
+				FVector OffsetScale = FVector::OneVector;
+				FTransform ParentCurrentGlobalTransform = Hierarchy->GetGlobalTransform(Parent.Item, false);
+
+				if (bMaintainOffset)
+				{
+					FTransform ParentInitialGlobalTransform = Hierarchy->GetInitialGlobalTransform(Parent.Item);
+
+					FVector ParentInitialGlobalScale = ParentInitialGlobalTransform.GetScale3D();
+					
+					OffsetScale = ChildInitialGlobalTransform.GetScale3D() / GetNonZeroScale(ParentInitialGlobalScale);
+				}
+
+				FVector OffsetParentScale = ParentCurrentGlobalTransform.GetScale3D() * OffsetScale;
+
+				FVector WeightedOffsetParentScale;
+				WeightedOffsetParentScale.X = FMath::Pow(OffsetParentScale.X, NormalizedWeight);
+				WeightedOffsetParentScale.Y = FMath::Pow(OffsetParentScale.Y, NormalizedWeight);
+				WeightedOffsetParentScale.Z = FMath::Pow(OffsetParentScale.Z, NormalizedWeight);
+
+				MixedGlobalScale *= WeightedOffsetParentScale;
+			}
+
+			// handle filtering, performed in local space
+			FTransform ChildParentGlobalTransform = Hierarchy->GetParentTransform(Child, false);
+			FVector ChildParentGlobalScale = ChildParentGlobalTransform.GetScale3D();
+			FVector MixedLocalScale = MixedGlobalScale / GetNonZeroScale(ChildParentGlobalScale);
+			
+			FTransform ChildCurrentLocalTransform = Hierarchy->GetLocalTransform(Child, false);
+			FVector ChildLocalScale = ChildCurrentLocalTransform.GetScale3D();
+			
+			// Controls need to be handled a bit differently
+			FTransform AdditionalOffsetTransform = FTransform::Identity;
+			
+			if (Child.Type == ERigElementType::Control)
+			{
+				if (FRigControlElement* ChildAsControlElement = Hierarchy->Find<FRigControlElement>(Child))
+				{
+					AdditionalOffsetTransform = Hierarchy->GetControlOffsetTransform(ChildAsControlElement, ERigTransformType::CurrentLocal);
+					// Control's local(parent) space transform = local * offset
+					ChildLocalScale *= AdditionalOffsetTransform.GetScale3D();
+				}
+			}
+			
+			FVector FilteredLocalScale;
+			FilteredLocalScale.X = Filter.bX ? MixedLocalScale.X : ChildLocalScale.X;
+			FilteredLocalScale.Y = Filter.bY ? MixedLocalScale.Y : ChildLocalScale.Y;
+			FilteredLocalScale.Z = Filter.bZ ? MixedLocalScale.Z : ChildLocalScale.Z;
+
+			FVector FinalLocalScale = FilteredLocalScale;
+			
+			if (Child.Type == ERigElementType::Control)
+			{
+				FinalLocalScale = FinalLocalScale / GetNonZeroScale(AdditionalOffsetTransform.GetScale3D());
+			}
+			
+			Hierarchy->SetLocalTransform(Child, FTransform(ChildCurrentLocalTransform.GetRotation(), ChildCurrentLocalTransform.GetTranslation(), FinalLocalScale));
+		}
+	}
+}
+
+#if WITH_DEV_AUTOMATION_TESTS
+#include "Units/RigUnitTest.h"
+
+IMPLEMENT_RIGUNIT_AUTOMATION_TEST(FRigUnit_ScaleConstraint)
+{
+	// use euler rotation here to match other software's rotation representation more easily
+	EControlRigRotationOrder Order = EControlRigRotationOrder::XZY;
+	const FRigElementKey Child = Controller->AddBone(TEXT("Child"), FRigElementKey(), FTransform(FQuat::Identity, FVector::ZeroVector, FVector::OneVector), true, ERigBoneType::User);
+	const FRigElementKey Parent1 = Controller->AddBone(TEXT("Parent1"), FRigElementKey(), FTransform(FQuat::Identity, FVector::ZeroVector, FVector(4,4,4)), true, ERigBoneType::User);
+	const FRigElementKey Parent2 = Controller->AddBone(TEXT("Parent2"), FRigElementKey(), FTransform(FQuat::Identity, FVector::ZeroVector, FVector(1,1,1)), true, ERigBoneType::User);
+	
+	Unit.ExecuteContext.Hierarchy = Hierarchy;
+	Unit.Child = Child;
+
+	Unit.Parents.Add(FConstraintParent(Parent1, 1.0));
+	Unit.Parents.Add(FConstraintParent(Parent2, 1.0));
+
+	Hierarchy->ResetPoseToInitial(ERigElementType::Bone);
+	Unit.bMaintainOffset = false;
+	
+	Execute();
+	AddErrorIfFalse(Hierarchy->GetGlobalTransform(0).GetScale3D().Equals(FVector(2,2,2)), TEXT("unexpected scale for maintain offset off"));
+	
+	Hierarchy->ResetPoseToInitial(ERigElementType::Bone);
+	Hierarchy->SetGlobalTransform(2, FTransform(FQuat::Identity, FVector::ZeroVector, FVector(0.5, 0.5, 0.5)));
+	Unit.bMaintainOffset = true;
+	
+	Execute();
+	AddErrorIfFalse(Hierarchy->GetGlobalTransform(0).GetScale3D().Equals(FVector(0.707, 0.707f, 0.707f), 0.001f),
+                    TEXT("unexpected scale for maintain offset on"));
+
+	return true;
+}
+
+#endif
\ No newline at end of file
diff --git a/Engine/Plugins/Experimental/ControlRig/Source/ControlRig/Private/Units/Highlevel/Hierarchy/RigUnit_TransformConstraint.h b/Engine/Plugins/Experimental/ControlRig/Source/ControlRig/Private/Units/Highlevel/Hierarchy/RigUnit_TransformConstraint.h
--- a/Engine/Plugins/Experimental/ControlRig/Source/ControlRig/Private/Units/Highlevel/Hierarchy/RigUnit_TransformConstraint.h
+++ b/Engine/Plugins/Experimental/ControlRig/Source/ControlRig/Private/Units/Highlevel/Hierarchy/RigUnit_TransformConstraint.h
@@ -7,6 +7,8 @@
 #include "Rigs/RigHierarchyContainer.h"
 #include "Constraint.h"
 #include "ControlRigDefines.h"
+#include "Math/ControlRigMathLibrary.h"
+
 #include "RigUnit_TransformConstraint.generated.h"
 
 /** 
@@ -174,3 +176,210 @@ private:
 	UPROPERTY(transient)
 	FRigUnit_TransformConstraint_WorkData WorkData;
 };
+
+USTRUCT()
+struct FConstraintParent
+{
+	GENERATED_BODY()
+
+	UPROPERTY(EditAnywhere, Category = "FConstraintParent", meta = (Input))
+	FRigElementKey Item;
+
+	UPROPERTY(EditAnywhere, Category = "FConstraintParent", meta = (Input))
+	float Weight;
+
+	FConstraintParent()
+        : Item(FRigElementKey(NAME_None, ERigElementType::Bone))
+        , Weight (1.f)
+	{}
+	
+	FConstraintParent(const FRigElementKey InItem, const float InWeight)
+		: Item(InItem)
+		, Weight(InWeight)
+	{}
+};
+
+/**
+*	Options for interpolating rotations
+*/ 
+UENUM(BlueprintType)
+enum class EConstraintInterpType : uint8
+{ 
+	/** Weighted Average of Quaternions by their X,Y,Z,W values, The Shortest Route is Respected. The Order of Parents Doesn't Matter */
+	Average UMETA(DisplayName="Average(Order Independent)"),
+   /** Perform Quaternion Slerp in Sequence, Different Orders of Parents can Produce Different Results */ 
+    Shortest UMETA(DisplayName="Shortest(Order Dependent)"),
+    Max UMETA(Hidden),
+};
+
+USTRUCT()
+struct FRigUnit_ParentConstraint_AdvancedSettings
+{
+	GENERATED_BODY()
+
+	FRigUnit_ParentConstraint_AdvancedSettings()
+		: InterpolationType(EConstraintInterpType::Average)
+		, RotationOrderForFilter(EControlRigRotationOrder::XZY)
+	{}
+	
+	/**
+	*	Options for interpolating rotations
+	*/
+	UPROPERTY(EditAnywhere, Category = "Constraint", meta = (Input))
+	EConstraintInterpType InterpolationType;
+	
+	/**
+	*	Rotation is converted to euler angles using the specified order such that individual axes can be filtered.
+	*/
+	UPROPERTY(EditAnywhere, Category = "Constraint", meta = (Input))
+	EControlRigRotationOrder RotationOrderForFilter;
+};
+
+/**
+* Constrains an item's transform to multiple items' transforms
+*/
+USTRUCT(meta=(DisplayName="Parent Constraint", Category="Transforms", Keywords = "Parent,Orient,Scale"))
+struct FRigUnit_ParentConstraint : public FRigUnit_HighlevelBaseMutable
+{
+	GENERATED_BODY()
+
+	FRigUnit_ParentConstraint()
+        :Child(FRigElementKey(NAME_None, ERigElementType::Bone))
+        , bMaintainOffset(true)
+	{
+	}
+	
+	RIGVM_METHOD()
+    virtual void Execute(const FRigUnitContext& Context) override;
+
+	UPROPERTY(EditAnywhere, Category = "Constraint", meta = (Input, ExpandByDefault))
+	FRigElementKey Child;
+	
+	UPROPERTY(EditAnywhere, Category = "Constraint", meta = (Input))
+	bool bMaintainOffset;
+
+	UPROPERTY(EditAnywhere, Category = "Constraint", meta = (Input)) 
+	FTransformFilter Filter;
+	
+	UPROPERTY(EditAnywhere, Category = "Constraint", meta = (Input)) 
+	FRigUnit_ParentConstraint_AdvancedSettings AdvancedSettings;
+	
+	UPROPERTY(EditAnywhere, Category = "Constraint", meta = (Input, ExpandByDefault, DefaultArraySize = 1))
+	TArray<FConstraintParent> Parents;
+};
+
+/**
+* Constrains an item's position to multiple items' positions 
+*/
+USTRUCT(meta=(DisplayName="Position Constraint", Category="Transforms", Keywords = "Parent,Translation"))
+struct FRigUnit_PositionConstraint : public FRigUnit_HighlevelBaseMutable
+{
+	GENERATED_BODY()
+
+	FRigUnit_PositionConstraint()
+        :Child(FRigElementKey(NAME_None, ERigElementType::Bone))
+        , bMaintainOffset(true)
+	{
+	}
+	
+	RIGVM_METHOD()
+    virtual void Execute(const FRigUnitContext& Context) override;
+
+	UPROPERTY(EditAnywhere, Category = "Constraint", meta = (Input, ExpandByDefault))
+	FRigElementKey Child;
+	
+	UPROPERTY(EditAnywhere, Category = "Constraint", meta = (Input))
+	bool bMaintainOffset;
+
+	UPROPERTY(EditAnywhere, Category = "Constraint", meta = (Input)) 
+	FFilterOptionPerAxis Filter;
+
+	UPROPERTY(EditAnywhere, Category = "Constraint", meta = (Input, ExpandByDefault, DefaultArraySize = 1))
+	TArray<FConstraintParent> Parents;
+};
+
+USTRUCT()
+struct FRigUnit_RotationConstraint_AdvancedSettings
+{
+	GENERATED_BODY()
+
+	FRigUnit_RotationConstraint_AdvancedSettings()
+        : InterpolationType(EConstraintInterpType::Average)
+        , RotationOrderForFilter(EControlRigRotationOrder::XZY)
+	{}
+	
+	/**
+	*	Options for interpolating rotations
+	*/
+	UPROPERTY(EditAnywhere, Category = "Constraint", meta = (Input))
+	EConstraintInterpType InterpolationType;
+	
+	/**
+	*	Rotation is converted to euler angles using the specified order such that individual axes can be filtered.
+	*/
+	UPROPERTY(EditAnywhere, Category = "Constraint", meta = (Input))
+	EControlRigRotationOrder RotationOrderForFilter;
+};
+
+/**
+* Constrains an item's rotation to multiple items' rotations 
+*/
+USTRUCT(meta=(DisplayName="Rotation Constraint", Category="Transforms", Keywords = "Parent,Orientation,Orient,Rotate"))
+struct FRigUnit_RotationConstraint : public FRigUnit_HighlevelBaseMutable
+{
+	GENERATED_BODY()
+
+	FRigUnit_RotationConstraint()
+        :Child(FRigElementKey(NAME_None, ERigElementType::Bone))
+        , bMaintainOffset(true)
+	{
+	}
+	
+	RIGVM_METHOD()
+    virtual void Execute(const FRigUnitContext& Context) override;
+
+	UPROPERTY(EditAnywhere, Category = "Constraint", meta = (Input, ExpandByDefault))
+	FRigElementKey Child;
+	
+	UPROPERTY(EditAnywhere, Category = "Constraint", meta = (Input))
+	bool bMaintainOffset;
+
+	UPROPERTY(EditAnywhere, Category = "Constraint", meta = (Input)) 
+	FFilterOptionPerAxis Filter;
+
+	UPROPERTY(EditAnywhere, Category = "Constraint", meta = (Input)) 
+	FRigUnit_RotationConstraint_AdvancedSettings AdvancedSettings;
+
+	UPROPERTY(EditAnywhere, Category = "Constraint", meta = (Input, ExpandByDefault, DefaultArraySize = 1))
+	TArray<FConstraintParent> Parents;
+};
+
+/**
+* Constrains an item's scale to multiple items' scales
+*/
+USTRUCT(meta=(DisplayName="Scale Constraint", Category="Transforms", Keywords = "Parent, Size"))
+struct FRigUnit_ScaleConstraint : public FRigUnit_HighlevelBaseMutable
+{
+	GENERATED_BODY()
+
+	FRigUnit_ScaleConstraint()
+        : Child(FRigElementKey(NAME_None, ERigElementType::Bone))
+        , bMaintainOffset(true)
+	{
+	}
+	
+	RIGVM_METHOD()
+    virtual void Execute(const FRigUnitContext& Context) override;
+
+	UPROPERTY(EditAnywhere, Category = "Constraint", meta = (Input, ExpandByDefault))
+	FRigElementKey Child;
+	
+	UPROPERTY(EditAnywhere, Category = "Constraint", meta = (Input))
+	bool bMaintainOffset;
+
+	UPROPERTY(EditAnywhere, Category = "Constraint", meta = (Input)) 
+	FFilterOptionPerAxis Filter;
+
+	UPROPERTY(EditAnywhere, Category = "Constraint", meta = (Input, ExpandByDefault, DefaultArraySize = 1))
+	TArray<FConstraintParent> Parents;
+};
\ No newline at end of file
