Control Rig: Follow up changes on Bulk Edit. Make sure to only show dialog for functions that have referenced assets

diff --git a/Engine/Plugins/Experimental/ControlRig/Source/ControlRigEditor/Private/Editor/ControlRigEditor.cpp b/Engine/Plugins/Experimental/ControlRig/Source/ControlRigEditor/Private/Editor/ControlRigEditor.cpp
--- a/Engine/Plugins/Experimental/ControlRig/Source/ControlRigEditor/Private/Editor/ControlRigEditor.cpp
+++ b/Engine/Plugins/Experimental/ControlRig/Source/ControlRigEditor/Private/Editor/ControlRigEditor.cpp
@@ -3426,6 +3426,12 @@ void FControlRigEditor::OnRequestLocalizeFunctionDialog(URigVMLibraryNode* InFun
 bool FControlRigEditor::OnRequestBulkEditDialog(UControlRigBlueprint* InBlueprint, URigVMController* InController,
 	URigVMLibraryNode* InFunction, ERigVMControllerBulkEditType InEditType)
 {
+	const TArray<FAssetData> FirstLevelReferenceAssets = InController->GetAffectedAssets(InEditType, false, true);
+	if(FirstLevelReferenceAssets.Num() == 0)
+	{
+		return true;
+	}
+	
 	TSharedRef<SControlRigFunctionBulkEditDialog> BulkEditDialog = SNew(SControlRigFunctionBulkEditDialog)
 	.Blueprint(InBlueprint)
 	.Controller(InController)
diff --git a/Engine/Plugins/Experimental/ControlRig/Source/ControlRigEditor/Private/Editor/SControlRigFunctionBulkEditWidget.cpp b/Engine/Plugins/Experimental/ControlRig/Source/ControlRigEditor/Private/Editor/SControlRigFunctionBulkEditWidget.cpp
--- a/Engine/Plugins/Experimental/ControlRig/Source/ControlRigEditor/Private/Editor/SControlRigFunctionBulkEditWidget.cpp
+++ b/Engine/Plugins/Experimental/ControlRig/Source/ControlRigEditor/Private/Editor/SControlRigFunctionBulkEditWidget.cpp
@@ -200,7 +200,11 @@ void SControlRigFunctionBulkEditDialog::Construct(const FArguments& InArgs)
 FReply SControlRigFunctionBulkEditDialog::OnButtonClick(EAppReturnType::Type ButtonID)
 {
 	UserResponse = ButtonID;
-	BulkEditWidget->LoadAffectedAssets();
+
+	if(UserResponse == EAppReturnType::Ok)
+	{
+		BulkEditWidget->LoadAffectedAssets();
+	}
 	
 	RequestDestroyWindow();
 
