ControlRig: Release/Debug mode

diff --git a/Engine/Plugins/Experimental/ControlRig/Source/ControlRig/Private/ControlRig.cpp b/Engine/Plugins/Experimental/ControlRig/Source/ControlRig/Private/ControlRig.cpp
--- a/Engine/Plugins/Experimental/ControlRig/Source/ControlRig/Private/ControlRig.cpp
+++ b/Engine/Plugins/Experimental/ControlRig/Source/ControlRig/Private/ControlRig.cpp
@@ -479,18 +479,25 @@ void UControlRig::Execute(const EControlRigState InState, const FName& InEventNa
 	}
 
 #if WITH_EDITOR
-	if (UControlRig* CDO = GetClass()->GetDefaultObject<UControlRig>())
+	if (bIsInDebugMode)
 	{
-		// Copy the breakpoints. This will not override the state of the breakpoints
-		DebugInfo.SetBreakpoints(CDO->DebugInfo.GetBreakpoints());
-
-		// If there are any breakpoints, create the Snapshot VM if it hasn't been created yet
-		if (DebugInfo.GetBreakpoints().Num() > 0)
+		if (UControlRig* CDO = GetClass()->GetDefaultObject<UControlRig>())
 		{
-			GetSnapshotVM();
+			// Copy the breakpoints. This will not override the state of the breakpoints
+			DebugInfo.SetBreakpoints(CDO->DebugInfo.GetBreakpoints());
+
+			// If there are any breakpoints, create the Snapshot VM if it hasn't been created yet
+			if (DebugInfo.GetBreakpoints().Num() > 0)
+			{
+				GetSnapshotVM();
+			}
 		}
+		VM->SetDebugInfo(&DebugInfo);
+	}
+	else
+	{
+		VM->SetDebugInfo(nullptr);
 	}
-	VM->SetDebugInfo(&DebugInfo);
 #endif
 
 	bool bJustRanInit = false;
diff --git a/Engine/Plugins/Experimental/ControlRig/Source/ControlRig/Public/ControlRig.h b/Engine/Plugins/Experimental/ControlRig/Source/ControlRig/Public/ControlRig.h
--- a/Engine/Plugins/Experimental/ControlRig/Source/ControlRig/Public/ControlRig.h
+++ b/Engine/Plugins/Experimental/ControlRig/Source/ControlRig/Public/ControlRig.h
@@ -459,6 +459,10 @@ private:
 	UPROPERTY(transient)
 	TArray<FName> EventQueue;
 
+	/** The current execution mode */
+	UPROPERTY(transient)
+	bool bIsInDebugMode;
+
 	/** Copy the VM from the default object */
 	void InstantiateVMFromCDO();
 	
diff --git a/Engine/Plugins/Experimental/ControlRig/Source/ControlRigDeveloper/Private/ControlRigBlueprint.cpp b/Engine/Plugins/Experimental/ControlRig/Source/ControlRigDeveloper/Private/ControlRigBlueprint.cpp
--- a/Engine/Plugins/Experimental/ControlRig/Source/ControlRigDeveloper/Private/ControlRigBlueprint.cpp
+++ b/Engine/Plugins/Experimental/ControlRig/Source/ControlRigDeveloper/Private/ControlRigBlueprint.cpp
@@ -688,7 +688,7 @@ void UControlRigBlueprint::RecompileVM()
 		CompileLog.NumErrors = CompileLog.NumWarnings = 0;
 
 		URigVMCompiler* Compiler = URigVMCompiler::StaticClass()->GetDefaultObject<URigVMCompiler>();
-		Compiler->Settings = VMCompileSettings;
+		Compiler->Settings = (bCompileInDebugMode) ? FRigVMCompileSettings::Fast() : VMCompileSettings;
 		Compiler->Compile(Model, GetOrCreateController(), CDO->VM, CDO->GetExternalVariablesImpl(false), UserData, &PinToOperandMap);
 
 		if (bErrorsDuringCompilation)
@@ -918,6 +918,7 @@ bool UControlRigBlueprint::AddBreakpoint(URigVMNode* InBreakpointNode, URigVMLib
 		{
 			// Add the breakpoint to the VM
 			bSuccess = AddBreakpointToControlRig(InBreakpointNode);
+			BreakpointAddedEvent.Broadcast();
 		}
 	}
 
diff --git a/Engine/Plugins/Experimental/ControlRig/Source/ControlRigDeveloper/Public/ControlRigBlueprint.h b/Engine/Plugins/Experimental/ControlRig/Source/ControlRigDeveloper/Public/ControlRigBlueprint.h
--- a/Engine/Plugins/Experimental/ControlRig/Source/ControlRigDeveloper/Public/ControlRigBlueprint.h
+++ b/Engine/Plugins/Experimental/ControlRig/Source/ControlRigDeveloper/Public/ControlRigBlueprint.h
@@ -41,6 +41,7 @@ DECLARE_EVENT_OneParam(UControlRigBlueprint, FOnPostEditChangeChainProperty, FPr
 DECLARE_EVENT_ThreeParams(UControlRigBlueprint, FOnLocalizeFunctionDialogRequested, URigVMLibraryNode*, UControlRigBlueprint*, bool);
 DECLARE_EVENT_ThreeParams(UControlRigBlueprint, FOnReportCompilerMessage, EMessageSeverity::Type, UObject*, const FString&);
 DECLARE_DELEGATE_RetVal_FourParams(bool, FControlRigOnBulkEditDialogRequestedDelegate, UControlRigBlueprint*, URigVMController*, URigVMLibraryNode*, ERigVMControllerBulkEditType);
+DECLARE_EVENT(UControlRigBlueprint, FOnBreakpointAdded);
 
 USTRUCT()
 struct CONTROLRIGDEVELOPER_API FControlRigPublicFunctionArg
@@ -432,12 +433,14 @@ private:
 
 	FOnRefreshEditorEvent RefreshEditorEvent;
 	FOnVariableDroppedEvent VariableDroppedEvent;
+	FOnBreakpointAdded BreakpointAddedEvent;
 
 public:
 
 	void BroadcastRefreshEditor() { return RefreshEditorEvent.Broadcast(this); }
 	FOnRefreshEditorEvent& OnRefreshEditor() { return RefreshEditorEvent; }
 	FOnVariableDroppedEvent& OnVariableDropped() { return VariableDroppedEvent; }
+	FOnBreakpointAdded& OnBreakpointAdded() { return BreakpointAddedEvent; }
 
 private:
 
@@ -539,10 +542,15 @@ private:
 
 #if WITH_EDITOR
 private:
+	bool bCompileInDebugMode;
+	
 	TArray<URigVMNode*> RigVMBreakpointNodes;
 
 public:
 
+	/** Sets the execution mode. In Release mode the rig will ignore all breakpoints. */
+	FORCEINLINE void SetDebugMode(const bool bValue) { bCompileInDebugMode = bValue; }
+
 	/** Removes all the breakpoints from the blueprint and the VM */
 	void ClearBreakpoints();
 
diff --git a/Engine/Plugins/Experimental/ControlRig/Source/ControlRigEditor/Private/ControlRigBlueprintCommands.cpp b/Engine/Plugins/Experimental/ControlRig/Source/ControlRigEditor/Private/ControlRigBlueprintCommands.cpp
--- a/Engine/Plugins/Experimental/ControlRig/Source/ControlRigEditor/Private/ControlRigBlueprintCommands.cpp
+++ b/Engine/Plugins/Experimental/ControlRig/Source/ControlRigEditor/Private/ControlRigBlueprintCommands.cpp
@@ -14,6 +14,9 @@ void FControlRigBlueprintCommands::RegisterCommands()
 	UI_COMMAND(UpdateEvent, "Forwards Solve", "Run the normal update graph", EUserInterfaceActionType::Button, FInputChord());
 	UI_COMMAND(InverseEvent, "Backwards Solve", "Run the inverse graph", EUserInterfaceActionType::Button, FInputChord());
 	UI_COMMAND(InverseAndUpdateEvent, "Backwards and Forwards", "Run the inverse graph followed by the update graph", EUserInterfaceActionType::Button, FInputChord());
+	UI_COMMAND(ToggleExecutionMode, "Toggle Execution Mode", "Toggle between Release and Debug execution mode", EUserInterfaceActionType::Button, FInputChord());
+	UI_COMMAND(ReleaseMode, "Release Mode", "Compiles and Executes the rig, ignoring debug data.", EUserInterfaceActionType::Button, FInputChord());
+	UI_COMMAND(DebugMode, "Debug Mode", "Compiles and Executes the unoptimized rig, stopping at breakpoints.", EUserInterfaceActionType::Button, FInputChord());
 	UI_COMMAND(ResumeExecution, "Resume", "Resumes execution after being halted at a breakpoint.", EUserInterfaceActionType::Button, FInputChord(EKeys::F5));
 }
 
diff --git a/Engine/Plugins/Experimental/ControlRig/Source/ControlRigEditor/Private/ControlRigBlueprintCommands.h b/Engine/Plugins/Experimental/ControlRig/Source/ControlRigEditor/Private/ControlRigBlueprintCommands.h
--- a/Engine/Plugins/Experimental/ControlRig/Source/ControlRigEditor/Private/ControlRigBlueprintCommands.h
+++ b/Engine/Plugins/Experimental/ControlRig/Source/ControlRigEditor/Private/ControlRigBlueprintCommands.h
@@ -42,6 +42,15 @@ public:
 	/** Run the inverse graph followed by the update graph */
 	TSharedPtr< FUICommandInfo > InverseAndUpdateEvent;
 
+	/** Toggle between Release and Debug execution mode */
+	TSharedPtr< FUICommandInfo > ToggleExecutionMode;
+
+	/** Compile and run the optimized rig, ignoring any debug data */
+	TSharedPtr< FUICommandInfo > ReleaseMode;
+
+	/** Compile and run the unoptimized rig, ignoring any debug data */
+	TSharedPtr< FUICommandInfo > DebugMode;
+	
 	/** Resume the execution of the graph when halted at a breakpoint */
 	TSharedPtr< FUICommandInfo > ResumeExecution;
 
diff --git a/Engine/Plugins/Experimental/ControlRig/Source/ControlRigEditor/Private/ControlRigEditorStyle.h b/Engine/Plugins/Experimental/ControlRig/Source/ControlRigEditor/Private/ControlRigEditorStyle.h
--- a/Engine/Plugins/Experimental/ControlRig/Source/ControlRigEditor/Private/ControlRigEditorStyle.h
+++ b/Engine/Plugins/Experimental/ControlRig/Source/ControlRigEditor/Private/ControlRigEditorStyle.h
@@ -91,6 +91,8 @@ public:
 			Set("ControlRig.Bug.Solid", new IMAGE_BRUSH("Slate/ControlRig_BugSolid_28x", Icon14x14));
 
 			Set("ControlRig.ResumeExecution", new IMAGE_BRUSH_SVG("Slate/simulate", Icon40x40));
+			Set("ControlRig.ReleaseMode", new IMAGE_BRUSH_SVG("Slate/animation", Icon40x40));
+			Set("ControlRig.DebugMode", new IMAGE_BRUSH_SVG("Slate/Bug", Icon40x40));
 		}
 
 		// Graph styles
diff --git a/Engine/Plugins/Experimental/ControlRig/Source/ControlRigEditor/Private/Editor/ControlRigEditor.cpp b/Engine/Plugins/Experimental/ControlRig/Source/ControlRigEditor/Private/Editor/ControlRigEditor.cpp
--- a/Engine/Plugins/Experimental/ControlRig/Source/ControlRigEditor/Private/Editor/ControlRigEditor.cpp
+++ b/Engine/Plugins/Experimental/ControlRig/Source/ControlRigEditor/Private/Editor/ControlRigEditor.cpp
@@ -107,6 +107,7 @@ FControlRigEditor::FControlRigEditor()
 	, bFirstTimeSelecting(true)
 	, bAnyErrorsLeft(false)
 	, LastEventQueue(EControlRigEditorEventQueue::Setup)
+	, bIsInDebugMode(false)
 	, LastDebuggedRig()
 	, RigHierarchyTabCount(0)
 	, HaltedAtNode(nullptr)
@@ -127,6 +128,7 @@ FControlRigEditor::~FControlRigEditor()
 		}
 		RigBlueprint->OnRefreshEditor().RemoveAll(this);
 		RigBlueprint->OnVariableDropped().RemoveAll(this);
+		RigBlueprint->OnBreakpointAdded().RemoveAll(this);
 		RigBlueprint->OnNodeDoubleClicked().RemoveAll(this);
 		RigBlueprint->OnGraphImported().RemoveAll(this);
 		RigBlueprint->OnPostEditChangeChainProperty().RemoveAll(this);
@@ -350,6 +352,10 @@ void FControlRigEditor::InitControlRigEditor(const EToolkitMode::Type Mode, cons
 		InControlRigBlueprint->Hierarchy->OnModified().AddSP(this, &FControlRigEditor::OnHierarchyModified);
 		InControlRigBlueprint->OnRefreshEditor().AddSP(this, &FControlRigEditor::HandleRefreshEditorFromBlueprint);
 		InControlRigBlueprint->OnVariableDropped().AddSP(this, &FControlRigEditor::HandleVariableDroppedFromBlueprint);
+		InControlRigBlueprint->OnBreakpointAdded().AddLambda([this]()
+		{
+			SetExecutionMode(true);
+		});
 
 		if (FControlRigEditMode* EditMode = GetEditMode())
 		{
@@ -434,6 +440,21 @@ void FControlRigEditor::BindCommands()
 		FExecuteAction::CreateSP(this, &FControlRigEditor::SetEventQueue, EControlRigEditorEventQueue::InverseAndUpdate),
 		FCanExecuteAction());
 
+	GetToolkitCommands()->MapAction(
+		FControlRigBlueprintCommands::Get().ToggleExecutionMode,
+		FExecuteAction::CreateSP(this, &FControlRigEditor::ToggleExecutionMode),
+		FCanExecuteAction());
+
+	GetToolkitCommands()->MapAction(
+		FControlRigBlueprintCommands::Get().ReleaseMode,
+		FExecuteAction::CreateSP(this, &FControlRigEditor::SetExecutionMode, false),
+		FCanExecuteAction());
+
+	GetToolkitCommands()->MapAction(
+		FControlRigBlueprintCommands::Get().DebugMode,
+		FExecuteAction::CreateSP(this, &FControlRigEditor::SetExecutionMode, true),
+		FCanExecuteAction());
+
 	GetToolkitCommands()->MapAction(
 		FControlRigBlueprintCommands::Get().ResumeExecution,
 		FExecuteAction::CreateLambda([this]()
@@ -498,6 +519,11 @@ void FControlRigEditor::ToggleEventQueue()
 	SetEventQueue(LastEventQueue);
 }
 
+void FControlRigEditor::ToggleExecutionMode()
+{
+	SetExecutionMode(!bIsInDebugMode);
+}
+
 TSharedRef<SWidget> FControlRigEditor::GenerateEventQueueMenuContent()
 {
 	FMenuBuilder MenuBuilder(true, GetToolkitCommands());
@@ -515,6 +541,18 @@ TSharedRef<SWidget> FControlRigEditor::GenerateEventQueueMenuContent()
 	return MenuBuilder.MakeWidget();
 }
 
+TSharedRef<SWidget> FControlRigEditor::GenerateExecutionModeMenuContent()
+{
+	FMenuBuilder MenuBuilder(true, GetToolkitCommands());
+
+	MenuBuilder.BeginSection(TEXT("Events"));
+	MenuBuilder.AddMenuEntry(FControlRigBlueprintCommands::Get().ReleaseMode, TEXT("Release"), TAttribute<FText>(), TAttribute<FText>(), GetExecutionModeIcon(false));
+	MenuBuilder.AddMenuEntry(FControlRigBlueprintCommands::Get().DebugMode, TEXT("Debug"), TAttribute<FText>(), TAttribute<FText>(), GetExecutionModeIcon(true));
+	MenuBuilder.EndSection();
+
+	return MenuBuilder.MakeWidget();
+}
+
 void FControlRigEditor::ExtendToolbar()
 {
 	DECLARE_SCOPE_HIERARCHICAL_COUNTER_FUNC()
@@ -583,6 +621,22 @@ void FControlRigEditor::FillToolbar(FToolBarBuilder& ToolbarBuilder)
 		ToolbarBuilder.AddWidget(SNew(SBlueprintEditorSelectedDebugObjectWidget, SharedThis(this)));
 
 		ToolbarBuilder.AddSeparator();
+		ToolbarBuilder.AddToolBarButton(
+			FControlRigBlueprintCommands::Get().ToggleExecutionMode,
+			NAME_None, 
+			TAttribute<FText>::Create(TAttribute<FText>::FGetter::CreateSP(this, &FControlRigEditor::GetExecutionModeLabel)),
+			TAttribute<FText>(), 
+			TAttribute<FSlateIcon>::Create(TAttribute<FSlateIcon>::FGetter::CreateSP(this, &FControlRigEditor::GetExecutionModeIcon))
+		);
+		
+		FUIAction DefaultExecutionMode;
+		ToolbarBuilder.AddComboButton(
+			DefaultExecutionMode,
+			FOnGetContent::CreateSP(this, &FControlRigEditor::GenerateExecutionModeMenuContent),
+			LOCTEXT("ExecutionMode_Label", "Execution Modes"),
+			LOCTEXT("ExecutionMode_ToolTip", "Pick between different execution modes for testing the Control Rig"),
+			FSlateIcon(FEditorStyle::GetStyleSetName(), "LevelEditor.Recompile"),
+			true);
 		ToolbarBuilder.AddToolBarButton(FControlRigBlueprintCommands::Get().ResumeExecution,
 			NAME_None, TAttribute<FText>(), TAttribute<FText>(), FSlateIcon(FControlRigEditorStyle::Get().GetStyleSetName(), "ControlRig.ResumeExecution"));
 
@@ -754,9 +808,55 @@ FSlateIcon FControlRigEditor::GetEventQueueIcon() const
 	return GetEventQueueIcon(GetEventQueue());
 }
 
-void FControlRigEditor::OnEventQueueComboChanged(int32 InValue, ESelectInfo::Type InSelectInfo)
+void FControlRigEditor::SetExecutionMode(const bool bSetDebugMode)
 {
-	SetEventQueue((EControlRigEditorEventQueue)InValue);
+	if (bIsInDebugMode == bSetDebugMode)
+	{
+		return;
+	}
+
+	bIsInDebugMode = bSetDebugMode;
+	GetControlRigBlueprint()->SetDebugMode(bSetDebugMode);
+	Compile();
+	
+	if (ControlRig)
+	{
+		ControlRig->bIsInDebugMode = bSetDebugMode;
+	}
+
+	if (HaltedAtNode)
+	{
+		HaltedAtNode->SetExecutionIsHaltedAtThisNode(false);
+		HaltedAtNode = nullptr;
+	}
+}
+
+int32 FControlRigEditor::GetExecutionModeComboValue() const
+{
+	return (int32) bIsInDebugMode;
+}
+
+FText FControlRigEditor::GetExecutionModeLabel() const
+{
+	if (bIsInDebugMode)
+	{
+		return FText::FromString(TEXT("DebugMode"));
+	}
+	return FText::FromString(TEXT("ReleaseMode"));
+}
+
+FSlateIcon FControlRigEditor::GetExecutionModeIcon(bool bDebugMode)
+{
+	if (bDebugMode)
+	{
+		return FSlateIcon(FControlRigEditorStyle::Get().GetStyleSetName(), "ControlRig.DebugMode");
+	}
+	return FSlateIcon(FControlRigEditorStyle::Get().GetStyleSetName(), "ControlRig.ReleaseMode");
+}
+
+FSlateIcon FControlRigEditor::GetExecutionModeIcon() const
+{
+	return GetExecutionModeIcon(bIsInDebugMode);
 }
 
 void FControlRigEditor::ToggleSetupMode()
@@ -785,7 +885,6 @@ void FControlRigEditor::ToggleSetupMode()
 		{
 			RigBlueprint->ClearTransientControls();
 		}
-
 	}
 
 	if (ControlRig)
diff --git a/Engine/Plugins/Experimental/ControlRig/Source/ControlRigEditor/Private/Editor/ControlRigEditor.h b/Engine/Plugins/Experimental/ControlRig/Source/ControlRigEditor/Private/Editor/ControlRigEditor.h
--- a/Engine/Plugins/Experimental/ControlRig/Source/ControlRigEditor/Private/Editor/ControlRigEditor.h
+++ b/Engine/Plugins/Experimental/ControlRig/Source/ControlRigEditor/Private/Editor/ControlRigEditor.h
@@ -246,7 +246,13 @@ private:
 	FText GetEventQueueLabel() const;
 	static FSlateIcon GetEventQueueIcon(EControlRigEditorEventQueue InEventQueue);
 	FSlateIcon GetEventQueueIcon() const;
-	void OnEventQueueComboChanged(int32 InValue, ESelectInfo::Type InSelectInfo);
+
+	void SetExecutionMode(const bool bSetDebugMode);
+	int32 GetExecutionModeComboValue() const;
+	FText GetExecutionModeLabel() const;
+	static FSlateIcon GetExecutionModeIcon(const bool bInIsDebugMode);
+	FSlateIcon GetExecutionModeIcon() const;
+	
 	virtual void GetCustomDebugObjects(TArray<FCustomDebugObject>& DebugList) const override;
 	virtual bool OnlyShowCustomDebugObjects() const override { return true; }
 	virtual void HandleSetObjectBeingDebugged(UObject* InObject) override;
@@ -292,7 +298,9 @@ private:
 	bool IsAutoCompileGraphOn() const;
 	bool CanAutoCompileGraph() const { return true; }
 	void ToggleEventQueue();
+	void ToggleExecutionMode();
 	TSharedRef<SWidget> GenerateEventQueueMenuContent();
+	TSharedRef<SWidget> GenerateExecutionModeMenuContent();
 
 	enum ERigElementGetterSetterType
 	{
@@ -412,6 +420,7 @@ protected:
 	bool bAnyErrorsLeft;
 
 	EControlRigEditorEventQueue LastEventQueue;
+	bool bIsInDebugMode;
 	FString LastDebuggedRig;
 	int32 RigHierarchyTabCount;
 	TMap<FString, FString> KnownInstructionLimitWarnings;
diff --git a/Engine/Plugins/Experimental/ControlRig/Source/ControlRigEditor/Private/Editor/SControlRigStackView.cpp b/Engine/Plugins/Experimental/ControlRig/Source/ControlRigEditor/Private/Editor/SControlRigStackView.cpp
--- a/Engine/Plugins/Experimental/ControlRig/Source/ControlRigEditor/Private/Editor/SControlRigStackView.cpp
+++ b/Engine/Plugins/Experimental/ControlRig/Source/ControlRigEditor/Private/Editor/SControlRigStackView.cpp
@@ -549,7 +549,10 @@ void SControlRigStackView::RefreshTreeView(URigVM* InVM)
 					}
 				}
 
-				InVM->ExecutionHalted().AddSP(this, &SControlRigStackView::HandleExecutionHalted);
+				if (ControlRig->GetVM())
+				{
+					ControlRig->GetVM()->ExecutionHalted().AddSP(this, &SControlRigStackView::HandleExecutionHalted);
+				}
 			}
 		}
 	}
diff --git a/Engine/Source/Runtime/RigVM/Private/RigVMCore/RigVM.cpp b/Engine/Source/Runtime/RigVM/Private/RigVMCore/RigVM.cpp
--- a/Engine/Source/Runtime/RigVM/Private/RigVMCore/RigVM.cpp
+++ b/Engine/Source/Runtime/RigVM/Private/RigVMCore/RigVM.cpp
@@ -350,10 +350,13 @@ TArray<FName> URigVM::GetEntryNames() const
 bool URigVM::ResumeExecution()
 {
 	HaltedAtInstruction = INDEX_NONE;
-	if (FRigVMBreakpoint* Breakpoint = DebugInfo->FindBreakpoint(Context.InstructionIndex))
+	if (DebugInfo)
 	{
-		DebugInfo->IncrementBreakpointActivationOnHit(Context.InstructionIndex);
-		return true;
+		if (FRigVMBreakpoint* Breakpoint = DebugInfo->FindBreakpoint(Context.InstructionIndex))
+		{
+			DebugInfo->IncrementBreakpointActivationOnHit(Context.InstructionIndex);
+			return true;
+		}
 	}
 
 	return false;
@@ -771,38 +774,8 @@ bool URigVM::Initialize(FRigVMMemoryContainerPtrArray Memory, FRigVMFixedArray<v
 	Context.OpaqueArguments = AdditionalArguments;
 	Context.ExternalVariables = ExternalVariables;
 	
-#if WITH_EDITOR
-	if (DebugInfo)
-	{
-		DebugInfo->StartExecution();
-		HaltedAtInstruction = INDEX_NONE;
-	}
-#endif
-
 	while (Instructions.IsValidIndex(Context.InstructionIndex))
 	{
-#if WITH_EDITOR
-		if (DebugInfo)
-		{
-			if (FRigVMBreakpoint* Breakpoint = DebugInfo->FindBreakpoint(Context.InstructionIndex))
-			{
-				if (DebugInfo->IsActive(Context.InstructionIndex))
-				{
-					if (HaltedAtInstruction != Context.InstructionIndex)
-					{
-						HaltedAtInstruction = Context.InstructionIndex;
-						ExecutionHalted().Broadcast(Context.InstructionIndex, Breakpoint->Subject);
-					}
-					return true;
-				}
-				else
-				{
-					DebugInfo->HitBreakpoint(Context.InstructionIndex);
-				}
-			}
-		}
-#endif
-		
 		const FRigVMInstruction& Instruction = Instructions[Context.InstructionIndex];
 
 		switch (Instruction.OpCode)
@@ -1023,14 +996,6 @@ bool URigVM::Initialize(FRigVMMemoryContainerPtrArray Memory, FRigVMFixedArray<v
 		}
 		Context.InstructionIndex++;
 	}
-
-#if WITH_EDITOR
-	if (HaltedAtInstruction != INDEX_NONE)
-	{
-		HaltedAtInstruction = INDEX_NONE;
-		ExecutionHalted().Broadcast(INDEX_NONE, nullptr);
-	}
-#endif
 	
 	return true;
 }
