ControlRig: Fix Debug unit tests not passing

diff --git a/Engine/Plugins/Experimental/ControlRig/Source/ControlRig/Public/ControlRig.h b/Engine/Plugins/Experimental/ControlRig/Source/ControlRig/Public/ControlRig.h
--- a/Engine/Plugins/Experimental/ControlRig/Source/ControlRig/Public/ControlRig.h
+++ b/Engine/Plugins/Experimental/ControlRig/Source/ControlRig/Public/ControlRig.h
@@ -459,10 +459,6 @@ private:
 	UPROPERTY(transient)
 	TArray<FName> EventQueue;
 
-	/** The current execution mode */
-	UPROPERTY(transient)
-	bool bIsInDebugMode;
-
 	/** Copy the VM from the default object */
 	void InstantiateVMFromCDO();
 	
@@ -649,13 +645,19 @@ private:
 	UPROPERTY(transient)
 	URigVM* VMSnapshotBeforeExecution;
 
+	/** The current execution mode */
+	UPROPERTY(transient)
+	bool bIsInDebugMode;
+
 #endif
 	
 #if WITH_EDITOR	
 
 	FRigVMDebugInfo DebugInfo;
 	
-public:	
+public:
+
+	void SetIsInDebugMode(const bool bValue) { bIsInDebugMode = bValue; }
 	
 	/** Adds a breakpoint in the VM at the InstructionIndex */
 	void AddBreakpoint(int32 InstructionIndex, URigVMNode* InNode);
