Control Rig: Invalidate Curve Container widget on editor close

diff --git a/Engine/Plugins/Experimental/ControlRig/Source/ControlRigEditor/Private/Editor/SRigCurveContainer.cpp b/Engine/Plugins/Experimental/ControlRig/Source/ControlRigEditor/Private/Editor/SRigCurveContainer.cpp
--- a/Engine/Plugins/Experimental/ControlRig/Source/ControlRigEditor/Private/Editor/SRigCurveContainer.cpp
+++ b/Engine/Plugins/Experimental/ControlRig/Source/ControlRigEditor/Private/Editor/SRigCurveContainer.cpp
@@ -145,6 +145,10 @@ void SRigCurveContainer::Construct(const FArguments& InArgs, TSharedRef<FControl
 	ControlRigBlueprint = InControlRigEditor.Get().GetControlRigBlueprint();
 	bIsChangingRigHierarchy = false;
 
+#if WITH_EDITOR
+	GEditor->OnEditorClose().AddRaw(this, &SRigCurveContainer::OnEditorClose);
+#endif
+
 	ControlRigBlueprint->Hierarchy->OnModified().AddRaw(this, &SRigCurveContainer::OnHierarchyModified);
 	ControlRigBlueprint->OnRefreshEditor().AddRaw(this, &SRigCurveContainer::HandleRefreshEditorFromBlueprint);
 
@@ -307,6 +311,12 @@ TSharedPtr<SWidget> SRigCurveContainer::OnGetContextMenuContent() const
 	return MenuBuilder.MakeWidget();
 }
 
+void SRigCurveContainer::OnEditorClose()
+{
+	ControlRigBlueprint.Reset();
+	ControlRigEditor.Reset();
+}
+
 void SRigCurveContainer::OnRenameClicked()
 {
 	TArray< FDisplayedRigCurveInfoPtr > SelectedItems = RigCurveListView->GetSelectedItems();
@@ -364,6 +374,11 @@ void SRigCurveContainer::CreateNewNameEntry(const FText& CommittedText, ETextCom
 
 void SRigCurveContainer::CreateRigCurveList( const FString& SearchText )
 {
+	if(!ControlRigBlueprint.IsValid())
+	{
+		return;
+	}
+	
 	URigHierarchy* Hierarchy = GetHierarchy();
 	if (Hierarchy)
 	{
diff --git a/Engine/Plugins/Experimental/ControlRig/Source/ControlRigEditor/Private/Editor/SRigCurveContainer.h b/Engine/Plugins/Experimental/ControlRig/Source/ControlRigEditor/Private/Editor/SRigCurveContainer.h
--- a/Engine/Plugins/Experimental/ControlRig/Source/ControlRigEditor/Private/Editor/SRigCurveContainer.h
+++ b/Engine/Plugins/Experimental/ControlRig/Source/ControlRigEditor/Private/Editor/SRigCurveContainer.h
@@ -206,6 +206,8 @@ private:
 	/** Handler for context menus */
 	TSharedPtr<SWidget> OnGetContextMenuContent() const;
 
+	void OnEditorClose();
+
 	/**
 	* Clears and rebuilds the table, according to an optional search string
 	*
