Control Rig: Re-enable paste / import for text content from pre-5.0 versions

diff --git a/Engine/Plugins/Experimental/ControlRig/Source/ControlRig/Private/Rigs/RigBoneHierarchy.cpp b/Engine/Plugins/Experimental/ControlRig/Source/ControlRig/Private/Rigs/RigBoneHierarchy.cpp
--- a/Engine/Plugins/Experimental/ControlRig/Source/ControlRig/Private/Rigs/RigBoneHierarchy.cpp
+++ b/Engine/Plugins/Experimental/ControlRig/Source/ControlRig/Private/Rigs/RigBoneHierarchy.cpp
@@ -10,3 +10,17 @@ FRigBoneHierarchy::FRigBoneHierarchy()
 {
 }
 
+FRigBone& FRigBoneHierarchy::Add(const FName& InNewName, const FName& InParentName, ERigBoneType InType, const FTransform& InInitTransform, const FTransform& InLocalTransform, const FTransform& InGlobalTransform)
+{
+	FRigBone NewBone;
+	NewBone.Name = InNewName;
+	NewBone.ParentIndex = INDEX_NONE; // we no longer support parent index lookup
+	NewBone.ParentName = InParentName;
+	NewBone.InitialTransform = InInitTransform;
+	NewBone.LocalTransform = InLocalTransform;
+	NewBone.GlobalTransform = InGlobalTransform;
+	NewBone.Type = InType;
+
+	const int32 Index = Bones.Add(NewBone);
+	return Bones[Index];
+}
\ No newline at end of file
diff --git a/Engine/Plugins/Experimental/ControlRig/Source/ControlRig/Private/Rigs/RigControlHierarchy.cpp b/Engine/Plugins/Experimental/ControlRig/Source/ControlRig/Private/Rigs/RigControlHierarchy.cpp
--- a/Engine/Plugins/Experimental/ControlRig/Source/ControlRig/Private/Rigs/RigControlHierarchy.cpp
+++ b/Engine/Plugins/Experimental/ControlRig/Source/ControlRig/Private/Rigs/RigControlHierarchy.cpp
@@ -34,6 +34,41 @@ FRigControlHierarchy::FRigControlHierarchy()
 {
 }
 
+FRigControl& FRigControlHierarchy::Add(
+    const FName& InNewName,
+    ERigControlType InControlType,
+    const FName& InParentName,
+    const FName& InSpaceName,
+    const FTransform& InOffsetTransform,
+    const FRigControlValue& InValue,
+    const FName& InGizmoName,
+    const FTransform& InGizmoTransform,
+    const FLinearColor& InGizmoColor
+)
+{
+	FRigControl NewControl;
+	NewControl.Name = InNewName;
+	NewControl.ControlType = InControlType;
+	NewControl.ParentIndex = INDEX_NONE; // we don't support indices
+	NewControl.ParentName = InParentName;
+	NewControl.SpaceIndex = INDEX_NONE;
+	NewControl.SpaceName = InSpaceName;
+	NewControl.OffsetTransform = InOffsetTransform;
+	NewControl.InitialValue = InValue;
+	NewControl.Value = FRigControlValue();
+	NewControl.GizmoName = InGizmoName;
+	NewControl.GizmoTransform = InGizmoTransform;
+	NewControl.GizmoColor = InGizmoColor;
+
+	if (!NewControl.InitialValue.IsValid())
+	{
+		NewControl.SetValueFromTransform(FTransform::Identity, ERigControlValueType::Initial);
+	}
+
+	const int32 Index = Controls.Add(NewControl);
+	return Controls[Index];
+}
+
 void FRigControlHierarchy::PostLoad()
 {
 	for (FRigControl& Control : Controls)
diff --git a/Engine/Plugins/Experimental/ControlRig/Source/ControlRig/Private/Rigs/RigCurveContainer.cpp b/Engine/Plugins/Experimental/ControlRig/Source/ControlRig/Private/Rigs/RigCurveContainer.cpp
--- a/Engine/Plugins/Experimental/ControlRig/Source/ControlRig/Private/Rigs/RigCurveContainer.cpp
+++ b/Engine/Plugins/Experimental/ControlRig/Source/ControlRig/Private/Rigs/RigCurveContainer.cpp
@@ -13,3 +13,15 @@
 FRigCurveContainer::FRigCurveContainer()
 {
 }
+
+FRigCurve& FRigCurveContainer::Add(const FName& InNewName, float InValue)
+{
+	DECLARE_SCOPE_HIERARCHICAL_COUNTER_FUNC()
+
+    FRigCurve NewCurve;
+	NewCurve.Name = InNewName;
+	NewCurve.Value = InValue;
+	FName NewCurveName = NewCurve.Name;
+	const int32 Index = Curves.Add(NewCurve);
+	return Curves[Index];
+}
diff --git a/Engine/Plugins/Experimental/ControlRig/Source/ControlRig/Private/Rigs/RigHierarchyContainer.cpp b/Engine/Plugins/Experimental/ControlRig/Source/ControlRig/Private/Rigs/RigHierarchyContainer.cpp
--- a/Engine/Plugins/Experimental/ControlRig/Source/ControlRig/Private/Rigs/RigHierarchyContainer.cpp
+++ b/Engine/Plugins/Experimental/ControlRig/Source/ControlRig/Private/Rigs/RigHierarchyContainer.cpp
@@ -30,3 +30,190 @@ FRigHierarchyContainer& FRigHierarchyContainer::operator= (const FRigHierarchyCo
 
 	return *this;
 }
+
+class FRigHierarchyContainerImportErrorContext : public FOutputDevice
+{
+public:
+
+	int32 NumErrors;
+
+	FRigHierarchyContainerImportErrorContext()
+        : FOutputDevice()
+        , NumErrors(0)
+	{
+	}
+
+	virtual void Serialize(const TCHAR* V, ELogVerbosity::Type Verbosity, const class FName& Category) override
+	{
+		UE_LOG(LogControlRig, Error, TEXT("Error Importing To Hierarchy: %s"), V);
+		NumErrors++;
+	}
+};
+
+TArray<FRigElementKey> FRigHierarchyContainer::ImportFromText(const FRigHierarchyCopyPasteContent& InData)
+{
+	TArray<FRigElementKey> PastedKeys;
+
+	if (InData.Contents.Num() == 0 ||
+		(InData.Types.Num() != InData.Contents.Num()) ||
+		(InData.LocalTransforms.Num() != InData.Contents.Num()) ||
+		(InData.GlobalTransforms.Num() != InData.Contents.Num()))
+	{
+		return PastedKeys;
+	}
+
+	TMap<FRigElementKey, FRigElementKey> ElementMap;
+	for (const FRigBone& Element : BoneHierarchy)
+	{
+		ElementMap.Add(Element.GetElementKey(), Element.GetElementKey());
+	}
+	for (const FRigControl& Element : ControlHierarchy)
+	{
+		ElementMap.Add(Element.GetElementKey(), Element.GetElementKey());
+	}
+	for (const FRigSpace& Element : SpaceHierarchy)
+	{
+		ElementMap.Add(Element.GetElementKey(), Element.GetElementKey());
+	}
+	for (const FRigCurve& Element : CurveContainer)
+	{
+		ElementMap.Add(Element.GetElementKey(), Element.GetElementKey());
+	}
+
+	FRigHierarchyContainerImportErrorContext ErrorPipe;
+	TArray<TSharedPtr<FRigElement>> Elements;
+	for (int32 Index = 0; Index < InData.Types.Num(); Index++)
+	{
+		ErrorPipe.NumErrors = 0;
+		
+		TSharedPtr<FRigElement> NewElement;
+		switch (InData.Types[Index])
+		{
+			case ERigElementType::Bone:
+			{
+				NewElement = MakeShared<FRigBone>();
+				FRigBone::StaticStruct()->ImportText(*InData.Contents[Index], NewElement.Get(), nullptr, EPropertyPortFlags::PPF_None, &ErrorPipe, FRigBone::StaticStruct()->GetName(), true);
+				break;
+			}
+			case ERigElementType::Control:
+			{
+				NewElement = MakeShared<FRigControl>();
+				FRigControl::StaticStruct()->ImportText(*InData.Contents[Index], NewElement.Get(), nullptr, EPropertyPortFlags::PPF_None, &ErrorPipe, FRigControl::StaticStruct()->GetName(), true);
+				break;
+			}
+			case ERigElementType::Space:
+			{
+				NewElement = MakeShared<FRigSpace>();
+				FRigSpace::StaticStruct()->ImportText(*InData.Contents[Index], NewElement.Get(), nullptr, EPropertyPortFlags::PPF_None, &ErrorPipe, FRigSpace::StaticStruct()->GetName(), true);
+				break;
+			}
+			case ERigElementType::Curve:
+			{
+				NewElement = MakeShared<FRigCurve>();
+				FRigCurve::StaticStruct()->ImportText(*InData.Contents[Index], NewElement.Get(), nullptr, EPropertyPortFlags::PPF_None, &ErrorPipe, FRigCurve::StaticStruct()->GetName(), true);
+				break;
+			}
+			default:
+			{
+				ensure(false);
+				break;
+			}
+		}
+
+		if (ErrorPipe.NumErrors > 0)
+		{
+			return PastedKeys;
+		}
+
+		Elements.Add(NewElement);
+	}
+
+	for (int32 Index = 0; Index < InData.Types.Num(); Index++)
+	{
+		switch (InData.Types[Index])
+		{
+			case ERigElementType::Bone:
+			{
+				const FRigBone& Element = *static_cast<FRigBone*>(Elements[Index].Get());
+
+				FName ParentName = NAME_None;
+				if (const FRigElementKey* ParentKey = ElementMap.Find(Element.GetParentElementKey()))
+				{
+					ParentName = ParentKey->Name;
+				}
+
+				FRigBone& NewElement = BoneHierarchy.Add(Element.Name, ParentName, ERigBoneType::User, Element.InitialTransform, Element.LocalTransform, InData.GlobalTransforms[Index]);
+				ElementMap.FindOrAdd(Element.GetElementKey()) = NewElement.GetElementKey();
+				PastedKeys.Add(NewElement.GetElementKey());
+				break;
+			}
+			case ERigElementType::Control:
+			{
+				const FRigControl& Element = *static_cast<FRigControl*>(Elements[Index].Get());
+
+				FName ParentName = NAME_None;
+				if (const FRigElementKey* ParentKey = ElementMap.Find(Element.GetParentElementKey()))
+				{
+					ParentName = ParentKey->Name;
+				}
+
+				FName SpaceName = NAME_None;
+				if (const FRigElementKey* SpaceKey = ElementMap.Find(Element.GetSpaceElementKey()))
+				{
+					SpaceName = SpaceKey->Name;
+				}
+
+				FRigControl& NewElement = ControlHierarchy.Add(Element.Name, Element.ControlType, ParentName, SpaceName, Element.OffsetTransform, Element.InitialValue, Element.GizmoName, Element.GizmoTransform, Element.GizmoColor);
+
+				// copy additional members
+				NewElement.DisplayName = Element.DisplayName;
+				NewElement.bAnimatable = Element.bAnimatable;
+				NewElement.PrimaryAxis = Element.PrimaryAxis;
+				NewElement.bLimitTranslation = Element.bLimitTranslation;
+				NewElement.bLimitRotation = Element.bLimitRotation;
+				NewElement.bLimitScale= Element.bLimitScale;
+				NewElement.MinimumValue = Element.MinimumValue;
+				NewElement.MaximumValue = Element.MaximumValue;
+				NewElement.bDrawLimits = Element.bDrawLimits;
+				NewElement.bGizmoEnabled = Element.bGizmoEnabled;
+				NewElement.bGizmoVisible = Element.bGizmoVisible;
+				NewElement.ControlEnum = Element.ControlEnum;
+
+				ElementMap.FindOrAdd(Element.GetElementKey()) = NewElement.GetElementKey();
+				PastedKeys.Add(NewElement.GetElementKey());
+
+				break;
+			}
+			case ERigElementType::Space:
+			{
+				const FRigSpace& Element = *static_cast<FRigSpace*>(Elements[Index].Get());
+
+				FName ParentName = NAME_None;
+				if (const FRigElementKey* ParentKey = ElementMap.Find(Element.GetParentElementKey()))
+				{
+					ParentName = ParentKey->Name;
+				}
+
+				FRigSpace& NewElement = SpaceHierarchy.Add(Element.Name, Element.SpaceType, ParentName, Element.InitialTransform);
+				ElementMap.FindOrAdd(Element.GetElementKey()) = NewElement.GetElementKey();
+				PastedKeys.Add(NewElement.GetElementKey());
+				break;
+			}
+			case ERigElementType::Curve:
+			{
+				const FRigCurve& Element = *static_cast<FRigCurve*>(Elements[Index].Get());
+				FRigCurve& NewElement = CurveContainer.Add(Element.Name, Element.Value);
+				ElementMap.FindOrAdd(Element.GetElementKey()) = NewElement.GetElementKey();
+				PastedKeys.Add(NewElement.GetElementKey());
+				break;
+			}
+			default:
+			{
+				ensure(false);
+				break;
+			}
+		}
+	}
+
+	return PastedKeys;
+}
diff --git a/Engine/Plugins/Experimental/ControlRig/Source/ControlRig/Private/Rigs/RigHierarchyController.cpp b/Engine/Plugins/Experimental/ControlRig/Source/ControlRig/Private/Rigs/RigHierarchyController.cpp
--- a/Engine/Plugins/Experimental/ControlRig/Source/ControlRig/Private/Rigs/RigHierarchyController.cpp
+++ b/Engine/Plugins/Experimental/ControlRig/Source/ControlRig/Private/Rigs/RigHierarchyController.cpp
@@ -819,6 +819,16 @@ TArray<FRigElementKey> URigHierarchyController::ImportFromText(FString InContent
 
 	if (Data.Elements.Num() == 0)
 	{
+		// check if this is a copy & paste buffer from pre-5.0
+		if(Data.Contents.Num() > 0)
+		{
+			FRigHierarchyContainer OldHierarchy;
+			if(OldHierarchy.ImportFromText(Data).Num() > 0)
+			{
+				return ImportFromHierarchyContainer(OldHierarchy);
+			}
+		}
+		
 		return PastedKeys;
 	}
 
@@ -983,6 +993,125 @@ TArray<FRigElementKey> URigHierarchyController::ImportFromText(FString InContent
 	return PastedKeys;
 }
 
+TArray<FRigElementKey> URigHierarchyController::ImportFromHierarchyContainer(const FRigHierarchyContainer& InContainer)
+{
+	TMap<FRigElementKey, FRigElementKey> KeyMap;;
+
+	for(const FRigBone& Bone : InContainer.BoneHierarchy)
+	{
+		const FRigElementKey OriginalParentKey = Bone.GetParentElementKey(true);
+		const FRigElementKey* ParentKey = nullptr;
+		if(OriginalParentKey.IsValid())
+		{
+			ParentKey = KeyMap.Find(OriginalParentKey);
+		}
+		if(ParentKey == nullptr)
+		{
+			ParentKey = &OriginalParentKey;
+		}
+
+		const FRigElementKey Key = AddBone(Bone.Name, *ParentKey, Bone.InitialTransform, true, Bone.Type, false);
+		KeyMap.Add(Bone.GetElementKey(), Key);
+	}
+	for(const FRigSpace& Space : InContainer.SpaceHierarchy)
+	{
+		const FRigElementKey Key = AddNull(Space.Name, FRigElementKey(), Space.InitialTransform, false, false);
+		KeyMap.Add(Space.GetElementKey(), Key);
+	}
+	for(const FRigControl& Control : InContainer.ControlHierarchy)
+	{
+		FRigControlSettings Settings;
+		Settings.ControlType = Control.ControlType;
+		Settings.DisplayName = Control.DisplayName;
+		Settings.PrimaryAxis = Control.PrimaryAxis;
+		Settings.bIsCurve = Control.bIsCurve;
+		Settings.bAnimatable = Control.bAnimatable;
+		Settings.bLimitTranslation = Control.bLimitTranslation;
+		Settings.bLimitRotation = Control.bLimitRotation;
+		Settings.bLimitScale = Control.bLimitScale;
+		Settings.bDrawLimits = Control.bDrawLimits;
+		Settings.MinimumValue = Control.MinimumValue;
+		Settings.MaximumValue = Control.MaximumValue;
+		Settings.bGizmoEnabled = Control.bGizmoEnabled;
+		Settings.bGizmoVisible = Control.bGizmoVisible;
+		Settings.GizmoName = Control.GizmoName;
+		Settings.GizmoColor = Control.GizmoColor;
+		Settings.ControlEnum = Control.ControlEnum;
+
+		FRigControlValue InitialValue = Control.InitialValue;
+		if(!InitialValue.IsValid())
+		{
+			InitialValue.SetFromTransform(InitialValue.Storage_DEPRECATED, Settings.ControlType, Settings.PrimaryAxis);
+		}
+		
+		const FRigElementKey Key = AddControl(
+			Control.Name,
+			FRigElementKey(),
+			Settings,
+			InitialValue,
+			Control.OffsetTransform,
+			Control.GizmoTransform,
+			false);
+
+		KeyMap.Add(Control.GetElementKey(), Key);
+	}
+	
+	for(const FRigCurve& Curve : InContainer.CurveContainer)
+	{
+		const FRigElementKey Key = AddCurve(Curve.Name, Curve.Value, false);
+		KeyMap.Add(Curve.GetElementKey(), Key);
+	}
+
+	for(const FRigSpace& Space : InContainer.SpaceHierarchy)
+	{
+		const FRigElementKey OriginalParentKey = Space.GetParentElementKey();
+		if(OriginalParentKey.IsValid())
+		{
+			FRigElementKey ParentKey;
+			if(const FRigElementKey* ParentKeyPtr = KeyMap.Find(OriginalParentKey))
+			{
+				ParentKey = *ParentKeyPtr;
+			}
+			SetParent(Space.GetElementKey(), ParentKey, false, false);
+		}
+	}
+
+	for(const FRigControl& Control : InContainer.ControlHierarchy)
+	{
+		FRigElementKey OriginalParentKey = Control.GetParentElementKey();
+		const FRigElementKey SpaceKey = Control.GetSpaceElementKey();
+		OriginalParentKey = SpaceKey.IsValid() ? SpaceKey : OriginalParentKey;
+		if(OriginalParentKey.IsValid())
+		{
+			FRigElementKey ParentKey;
+			if(const FRigElementKey* ParentKeyPtr = KeyMap.Find(OriginalParentKey))
+			{
+				ParentKey = *ParentKeyPtr;
+			}
+			SetParent(Control.GetElementKey(), ParentKey, false, false);
+		}
+	}
+
+#if WITH_EDITOR
+	if(!IsRunningCommandlet()) // don't show warnings like this if we are cooking
+	{
+		for(const TPair<FRigElementKey, FRigElementKey>& Pair: KeyMap)
+		{
+			if(Pair.Key != Pair.Value)
+			{
+				check(Pair.Key.Type == Pair.Value.Type);
+				const FText TypeLabel = StaticEnum<ERigElementType>()->GetDisplayNameTextByValue((int64)Pair.Key.Type);
+				ReportWarningf(TEXT("%s '%s' was renamed to '%s' during load (fixing invalid name)."), *TypeLabel.ToString(), *Pair.Key.Name.ToString(), *Pair.Value.Name.ToString());
+			}
+		}
+	}
+#endif
+
+	TArray<FRigElementKey> AddedKeys;
+	KeyMap.GenerateValueArray(AddedKeys);
+	return AddedKeys;
+}
+
 void URigHierarchyController::Notify(ERigHierarchyNotification InNotifType, const FRigBaseElement* InElement)
 {
 	if(!IsValid())
diff --git a/Engine/Plugins/Experimental/ControlRig/Source/ControlRig/Private/Rigs/RigSpaceHierarchy.cpp b/Engine/Plugins/Experimental/ControlRig/Source/ControlRig/Private/Rigs/RigSpaceHierarchy.cpp
--- a/Engine/Plugins/Experimental/ControlRig/Source/ControlRig/Private/Rigs/RigSpaceHierarchy.cpp
+++ b/Engine/Plugins/Experimental/ControlRig/Source/ControlRig/Private/Rigs/RigSpaceHierarchy.cpp
@@ -12,3 +12,16 @@
 FRigSpaceHierarchy::FRigSpaceHierarchy()
 {
 }
+
+FRigSpace& FRigSpaceHierarchy::Add(const FName& InNewName, ERigSpaceType InSpaceType, const FName& InParentName, const FTransform& InTransform)
+{
+	FRigSpace NewSpace;
+	NewSpace.Name = InNewName;
+	NewSpace.ParentIndex = INDEX_NONE; // we no longer support indices 
+	NewSpace.SpaceType = InSpaceType;
+	NewSpace.ParentName = InParentName;
+	NewSpace.InitialTransform = InTransform;
+	NewSpace.LocalTransform = InTransform;
+	const int32 Index = Spaces.Add(NewSpace);
+	return Spaces[Index];
+}
diff --git a/Engine/Plugins/Experimental/ControlRig/Source/ControlRig/Public/Rigs/RigBoneHierarchy.h b/Engine/Plugins/Experimental/ControlRig/Source/ControlRig/Public/Rigs/RigBoneHierarchy.h
--- a/Engine/Plugins/Experimental/ControlRig/Source/ControlRig/Public/Rigs/RigBoneHierarchy.h
+++ b/Engine/Plugins/Experimental/ControlRig/Source/ControlRig/Public/Rigs/RigBoneHierarchy.h
@@ -60,11 +60,7 @@ struct CONTROLRIG_API FRigBone: public FRigElement
 
 	FORCEINLINE virtual FRigElementKey GetParentElementKey(bool bForce = false) const
 	{
-		if (ParentIndex != INDEX_NONE || bForce)
-		{
-			return FRigElementKey(ParentName, GetElementType());
-		}
-		return FRigElementKey();
+		return FRigElementKey(ParentName, GetElementType());
 	}
 };
 
@@ -80,6 +76,8 @@ struct CONTROLRIG_API FRigBoneHierarchy
 	FORCEINLINE TArray<FRigBone>::RangedForIteratorType      end()         { return Bones.end();   }
 	FORCEINLINE TArray<FRigBone>::RangedForConstIteratorType end() const   { return Bones.end();   }
 
+	FRigBone& Add(const FName& InNewName, const FName& InParentName, ERigBoneType InType, const FTransform& InInitTransform, const FTransform& InLocalTransform, const FTransform& InGlobalTransform);
+	
 private:
 
 	// disable copy constructor
diff --git a/Engine/Plugins/Experimental/ControlRig/Source/ControlRig/Public/Rigs/RigControlHierarchy.h b/Engine/Plugins/Experimental/ControlRig/Source/ControlRig/Public/Rigs/RigControlHierarchy.h
--- a/Engine/Plugins/Experimental/ControlRig/Source/ControlRig/Public/Rigs/RigControlHierarchy.h
+++ b/Engine/Plugins/Experimental/ControlRig/Source/ControlRig/Public/Rigs/RigControlHierarchy.h
@@ -164,22 +164,14 @@ struct CONTROLRIG_API FRigControl : public FRigElement
 		return DisplayName.IsNone() ? Name : DisplayName;
 	}
 
-	FORCEINLINE_DEBUGGABLE virtual FRigElementKey GetParentElementKey(bool bForce = false) const
+	FORCEINLINE_DEBUGGABLE virtual FRigElementKey GetParentElementKey() const
 	{
-		if (ParentIndex != INDEX_NONE || bForce)
-		{
-			return FRigElementKey(ParentName, GetElementType());
-		}
-		return FRigElementKey();
+		return FRigElementKey(ParentName, GetElementType());
 	}
 
-	FORCEINLINE_DEBUGGABLE virtual FRigElementKey GetSpaceElementKey(bool bForce = false) const
+	FORCEINLINE_DEBUGGABLE virtual FRigElementKey GetSpaceElementKey() const
 	{
-		if (SpaceIndex != INDEX_NONE || bForce)
-		{
-			return FRigElementKey(SpaceName, ERigElementType::Null);
-		}
-		return FRigElementKey();
+		return FRigElementKey(SpaceName, ERigElementType::Null);
 	}
 
 	FORCEINLINE_DEBUGGABLE const FRigControlValue& GetValue(ERigControlValueType InValueType = ERigControlValueType::Current) const
@@ -274,6 +266,18 @@ struct CONTROLRIG_API FRigControlHierarchy
 	FORCEINLINE_DEBUGGABLE TArray<FRigControl>::RangedForIteratorType      end()         { return Controls.end();   }
 	FORCEINLINE_DEBUGGABLE TArray<FRigControl>::RangedForConstIteratorType end() const   { return Controls.end();   }
 
+	FRigControl& Add(
+	    const FName& InNewName,
+	    ERigControlType InControlType,
+	    const FName& InParentName,
+	    const FName& InSpaceName,
+	    const FTransform& InOffsetTransform,
+	    const FRigControlValue& InValue,
+	    const FName& InGizmoName,
+	    const FTransform& InGizmoTransform,
+	    const FLinearColor& InGizmoColor
+	);
+
 	void PostLoad();
 	
 private:
diff --git a/Engine/Plugins/Experimental/ControlRig/Source/ControlRig/Public/Rigs/RigCurveContainer.h b/Engine/Plugins/Experimental/ControlRig/Source/ControlRig/Public/Rigs/RigCurveContainer.h
--- a/Engine/Plugins/Experimental/ControlRig/Source/ControlRig/Public/Rigs/RigCurveContainer.h
+++ b/Engine/Plugins/Experimental/ControlRig/Source/ControlRig/Public/Rigs/RigCurveContainer.h
@@ -46,6 +46,8 @@ public:
 	FORCEINLINE TArray<FRigCurve>::RangedForConstIteratorType begin() const { return Curves.begin(); }
 	FORCEINLINE TArray<FRigCurve>::RangedForIteratorType      end()         { return Curves.end();   }
 	FORCEINLINE TArray<FRigCurve>::RangedForConstIteratorType end() const   { return Curves.end();   }
+
+	FRigCurve& Add(const FName& InNewName, float InValue);
 	
 private:
 
diff --git a/Engine/Plugins/Experimental/ControlRig/Source/ControlRig/Public/Rigs/RigHierarchyContainer.h b/Engine/Plugins/Experimental/ControlRig/Source/ControlRig/Public/Rigs/RigHierarchyContainer.h
--- a/Engine/Plugins/Experimental/ControlRig/Source/ControlRig/Public/Rigs/RigHierarchyContainer.h
+++ b/Engine/Plugins/Experimental/ControlRig/Source/ControlRig/Public/Rigs/RigHierarchyContainer.h
@@ -41,7 +41,10 @@ public:
 
 private:
 
+	TArray<FRigElementKey> ImportFromText(const FRigHierarchyCopyPasteContent& InData);
+
 	friend class SRigHierarchy;
+	friend class URigHierarchyController;
 };
 
 // this struct is still here for backwards compatibility - but not used anywhere
diff --git a/Engine/Plugins/Experimental/ControlRig/Source/ControlRig/Public/Rigs/RigHierarchyController.h b/Engine/Plugins/Experimental/ControlRig/Source/ControlRig/Public/Rigs/RigHierarchyController.h
--- a/Engine/Plugins/Experimental/ControlRig/Source/ControlRig/Public/Rigs/RigHierarchyController.h
+++ b/Engine/Plugins/Experimental/ControlRig/Source/ControlRig/Public/Rigs/RigHierarchyController.h
@@ -4,7 +4,9 @@
 
 #include "Rigs/RigHierarchy.h"
 #include "ReferenceSkeleton.h"
+#include "RigHierarchyContainer.h"
 #include "Animation/Skeleton.h"
+#include "RigHierarchyContainer.h"
 #include "RigHierarchyController.generated.h"
 
 UCLASS(BlueprintType)
@@ -323,6 +325,13 @@ public:
 		bool bSelectNewElements = true,
 		bool bSetupUndo = false);
 
+	/**
+	* Imports the content of a RigHierachyContainer (the hierarchy v1 pre 5.0)
+	* This is used for backwards compatbility only during load and does not support undo.
+	* @param InContainer The input hierarchy container
+	*/
+    TArray<FRigElementKey> ImportFromHierarchyContainer(const FRigHierarchyContainer& InContainer);
+
 	/**
 	 * Removes an existing element from the hierarchy
 	 * @param InElement The key of the element to remove
diff --git a/Engine/Plugins/Experimental/ControlRig/Source/ControlRig/Public/Rigs/RigHierarchyDefines.h b/Engine/Plugins/Experimental/ControlRig/Source/ControlRig/Public/Rigs/RigHierarchyDefines.h
--- a/Engine/Plugins/Experimental/ControlRig/Source/ControlRig/Public/Rigs/RigHierarchyDefines.h
+++ b/Engine/Plugins/Experimental/ControlRig/Source/ControlRig/Public/Rigs/RigHierarchyDefines.h
@@ -19,6 +19,7 @@ enum class ERigElementType : uint8
 	None,
 	Bone = 0x001,
 	Null = 0x002,
+	Space = Null UMETA(Hidden),
 	Control = 0x004,
 	Curve = 0x008,
 	RigidBody = 0x010,
@@ -721,6 +722,7 @@ private:
 
 	friend struct FRigControlHierarchy;
 	friend class UControlRigBlueprint;
+	friend class URigHierarchyController;
 };
 
 template<>
diff --git a/Engine/Plugins/Experimental/ControlRig/Source/ControlRig/Public/Rigs/RigHierarchyElements.h b/Engine/Plugins/Experimental/ControlRig/Source/ControlRig/Public/Rigs/RigHierarchyElements.h
--- a/Engine/Plugins/Experimental/ControlRig/Source/ControlRig/Public/Rigs/RigHierarchyElements.h
+++ b/Engine/Plugins/Experimental/ControlRig/Source/ControlRig/Public/Rigs/RigHierarchyElements.h
@@ -410,6 +410,19 @@ struct CONTROLRIG_API FRigHierarchyCopyPasteContent
 
     UPROPERTY()
 	TArray<FRigHierarchyCopyPasteContentPerElement> Elements;
+
+	// Maintain properties below for backwards compatibility pre-5.0
+	UPROPERTY()
+	TArray<ERigElementType> Types;
+
+	UPROPERTY()
+	TArray<FString> Contents;
+
+	UPROPERTY()
+	TArray<FTransform> LocalTransforms;
+
+	UPROPERTY()
+	TArray<FTransform> GlobalTransforms;
 };
 
 USTRUCT(BlueprintType)
diff --git a/Engine/Plugins/Experimental/ControlRig/Source/ControlRig/Public/Rigs/RigSpaceHierarchy.h b/Engine/Plugins/Experimental/ControlRig/Source/ControlRig/Public/Rigs/RigSpaceHierarchy.h
--- a/Engine/Plugins/Experimental/ControlRig/Source/ControlRig/Public/Rigs/RigSpaceHierarchy.h
+++ b/Engine/Plugins/Experimental/ControlRig/Source/ControlRig/Public/Rigs/RigSpaceHierarchy.h
@@ -62,28 +62,25 @@ struct CONTROLRIG_API FRigSpace : public FRigElement
 		return ERigElementType::Null;
 	}
 
-	FORCEINLINE virtual FRigElementKey GetParentElementKey(bool bForce = false) const
+	FORCEINLINE virtual FRigElementKey GetParentElementKey() const
 	{
-		if (ParentIndex != INDEX_NONE || bForce)
+		switch (SpaceType)
 		{
-			switch (SpaceType)
+			case ERigSpaceType::Bone:
 			{
-				case ERigSpaceType::Bone:
-				{
-					return FRigElementKey(ParentName, ERigElementType::Bone);
-				}
-				case ERigSpaceType::Control:
-				{
-					return FRigElementKey(ParentName, ERigElementType::Control);
-				}
-				case ERigSpaceType::Space:
-				{
-					return FRigElementKey(ParentName, ERigElementType::Null);
-				}
-				default:
-				{
-					break;
-				}
+				return FRigElementKey(ParentName, ERigElementType::Bone);
+			}
+			case ERigSpaceType::Control:
+			{
+				return FRigElementKey(ParentName, ERigElementType::Control);
+			}
+			case ERigSpaceType::Space:
+			{
+				return FRigElementKey(ParentName, ERigElementType::Null);
+			}
+			default:
+			{
+				break;
 			}
 		}
 		return FRigElementKey();
@@ -102,6 +99,8 @@ struct CONTROLRIG_API FRigSpaceHierarchy
 	FORCEINLINE TArray<FRigSpace>::RangedForIteratorType      end()         { return Spaces.end();   }
 	FORCEINLINE TArray<FRigSpace>::RangedForConstIteratorType end() const   { return Spaces.end();   }
 
+	FRigSpace& Add(const FName& InNewName, ERigSpaceType InSpaceType, const FName& InParentName, const FTransform& InTransform);
+	
 private:
 
 	// disable copy constructor
diff --git a/Engine/Plugins/Experimental/ControlRig/Source/ControlRigDeveloper/Private/ControlRigBlueprint.cpp b/Engine/Plugins/Experimental/ControlRig/Source/ControlRigDeveloper/Private/ControlRigBlueprint.cpp
--- a/Engine/Plugins/Experimental/ControlRig/Source/ControlRigDeveloper/Private/ControlRigBlueprint.cpp
+++ b/Engine/Plugins/Experimental/ControlRig/Source/ControlRigDeveloper/Private/ControlRigBlueprint.cpp
@@ -404,119 +404,8 @@ void UControlRigBlueprint::PostLoad()
 		TGuardValue<bool> SuspendNotifGuard(Hierarchy->GetSuspendNotificationsFlag(), true);
 		
 		Hierarchy->Reset();
-
-		TMap<FRigElementKey, FRigElementKey> KeyMap;;
-
-		for(const FRigBone& Bone : HierarchyContainer_DEPRECATED.BoneHierarchy)
-		{
-			const FRigElementKey OriginalParentKey = Bone.GetParentElementKey(true);
-			const FRigElementKey* ParentKey = nullptr;
-			if(OriginalParentKey.IsValid())
-			{
-				ParentKey = KeyMap.Find(OriginalParentKey);
-			}
-			if(ParentKey == nullptr)
-			{
-				ParentKey = &OriginalParentKey;
-			}
-
-			const FRigElementKey Key = HierarchyController->AddBone(Bone.Name, *ParentKey, Bone.InitialTransform, true, Bone.Type, false);
-			KeyMap.Add(Bone.GetElementKey(), Key);
-		}
-		for(const FRigSpace& Space : HierarchyContainer_DEPRECATED.SpaceHierarchy)
-		{
-			const FRigElementKey Key = HierarchyController->AddNull(Space.Name, FRigElementKey(), Space.InitialTransform, false, false);
-			KeyMap.Add(Space.GetElementKey(), Key);
-		}
-		for(const FRigControl& Control : HierarchyContainer_DEPRECATED.ControlHierarchy)
-		{
-			FRigControlSettings Settings;
-			Settings.ControlType = Control.ControlType;
-			Settings.DisplayName = Control.DisplayName;
-			Settings.PrimaryAxis = Control.PrimaryAxis;
-			Settings.bIsCurve = Control.bIsCurve;
-			Settings.bAnimatable = Control.bAnimatable;
-			Settings.bLimitTranslation = Control.bLimitTranslation;
-			Settings.bLimitRotation = Control.bLimitRotation;
-			Settings.bLimitScale = Control.bLimitScale;
-			Settings.bDrawLimits = Control.bDrawLimits;
-			Settings.MinimumValue = Control.MinimumValue;
-			Settings.MaximumValue = Control.MaximumValue;
-			Settings.bGizmoEnabled = Control.bGizmoEnabled;
-			Settings.bGizmoVisible = Control.bGizmoVisible;
-			Settings.GizmoName = Control.GizmoName;
-			Settings.GizmoColor = Control.GizmoColor;
-			Settings.ControlEnum = Control.ControlEnum;
-
-			FRigControlValue InitialValue = Control.InitialValue;
-			if(!InitialValue.IsValid())
-			{
-				InitialValue.SetFromTransform(InitialValue.Storage_DEPRECATED, Settings.ControlType, Settings.PrimaryAxis);
-			}
-			
-			const FRigElementKey Key = HierarchyController->AddControl(
-				Control.Name,
-				FRigElementKey(),
-				Settings,
-				InitialValue,
-				Control.OffsetTransform,
-				Control.GizmoTransform,
-				false);
-
-			KeyMap.Add(Control.GetElementKey(), Key);
-		}
-		
-		for(const FRigCurve& Curve : HierarchyContainer_DEPRECATED.CurveContainer)
-		{
-			const FRigElementKey Key = HierarchyController->AddCurve(Curve.Name, Curve.Value, false);
-			KeyMap.Add(Curve.GetElementKey(), Key);
-		}
-
-		for(const FRigSpace& Space : HierarchyContainer_DEPRECATED.SpaceHierarchy)
-		{
-			const FRigElementKey OriginalParentKey = Space.GetParentElementKey(true);
-			if(OriginalParentKey.IsValid())
-			{
-				FRigElementKey ParentKey;
-				if(const FRigElementKey* ParentKeyPtr = KeyMap.Find(OriginalParentKey))
-				{
-					ParentKey = *ParentKeyPtr;
-				}
-				HierarchyController->SetParent(Space.GetElementKey(), ParentKey, false, false);
-			}
-		}
-
-		for(const FRigControl& Control : HierarchyContainer_DEPRECATED.ControlHierarchy)
-		{
-			FRigElementKey OriginalParentKey = Control.GetParentElementKey(true);
-			const FRigElementKey SpaceKey = Control.GetSpaceElementKey(true);
-			OriginalParentKey = SpaceKey.IsValid() ? SpaceKey : OriginalParentKey;
-			if(OriginalParentKey.IsValid())
-			{
-				FRigElementKey ParentKey;
-				if(const FRigElementKey* ParentKeyPtr = KeyMap.Find(OriginalParentKey))
-				{
-					ParentKey = *ParentKeyPtr;
-				}
-				HierarchyController->SetParent(Control.GetElementKey(), ParentKey, false, false);
-			}
-		}
-
-#if WITH_EDITOR
-		if(!IsRunningCommandlet()) // don't show warnings like this if we are cooking
-		{
-			for(const TPair<FRigElementKey, FRigElementKey>& Pair: KeyMap)
-			{
-				if(Pair.Key != Pair.Value)
-				{
-					check(Pair.Key.Type == Pair.Value.Type);
-					const FText TypeLabel = StaticEnum<ERigElementType>()->GetDisplayNameTextByValue((int64)Pair.Key.Type);
-					GetController()->ReportWarningf(TEXT("%s '%s' was renamed to '%s' during load (fixing invalid name)."), *TypeLabel.ToString(), *Pair.Key.Name.ToString(), *Pair.Value.Name.ToString());
-				}
-			}
-		}
+		GetHierarchyController()->ImportFromHierarchyContainer(HierarchyContainer_DEPRECATED);
 	}
-#endif
 
 	PropagateHierarchyFromBPToInstances();
 	
