Playback broken when baking AnimSequence to FKRig
#fix ensure that correct transform value is set for non-additive pass

diff --git a/Engine/Plugins/Experimental/ControlRig/Source/ControlRig/Private/Rigs/FKControlRig.cpp b/Engine/Plugins/Experimental/ControlRig/Source/ControlRig/Private/Rigs/FKControlRig.cpp
--- a/Engine/Plugins/Experimental/ControlRig/Source/ControlRig/Private/Rigs/FKControlRig.cpp
+++ b/Engine/Plugins/Experimental/ControlRig/Source/ControlRig/Private/Rigs/FKControlRig.cpp
@@ -48,14 +48,12 @@ void UFKControlRig::ExecuteUnits(FRigUnitContext& InOutContext, const FName& InE
 			if (IsControlActive[ControlIndex])
 			{
 				FRigControlElement* Control = GetHierarchy()->Get<FRigControlElement>(ControlIndex);
+				const FTransform LocalTransform = GetHierarchy()->GetLocalTransform(ControlIndex);
 				switch (ApplyMode)
 				{
 					case EControlRigFKRigExecuteMode::Replace:
 					{
-						const FTransform LocalTransform = GetHierarchy()->GetLocalTransform(ControlIndex);
-						const FTransform OffsetTransform = Control->Offset.Initial.Local.Transform;
-						const FTransform Transform = LocalTransform * OffsetTransform;
-						GetHierarchy()->SetTransform(BoneElement, Transform, ERigTransformType::CurrentLocal, true, false);
+						GetHierarchy()->SetTransform(BoneElement, LocalTransform, ERigTransformType::CurrentLocal, true, false);
 						break;
 					}
 					case EControlRigFKRigExecuteMode::Additive:
