Fix-up for commit shelved changes in 16192634

diff --git a/Engine/Plugins/Experimental/ControlRig/Source/ControlRig/Private/Rigs/FKControlRig.cpp b/Engine/Plugins/Experimental/ControlRig/Source/ControlRig/Private/Rigs/FKControlRig.cpp
--- a/Engine/Plugins/Experimental/ControlRig/Source/ControlRig/Private/Rigs/FKControlRig.cpp
+++ b/Engine/Plugins/Experimental/ControlRig/Source/ControlRig/Private/Rigs/FKControlRig.cpp
@@ -58,7 +58,6 @@ void UFKControlRig::ExecuteUnits(FRigUnitContext& InOutContext, const FName& InE
 					}
 					case EControlRigFKRigExecuteMode::Additive:
 					{
-						const FTransform LocalTransform = GetHierarchy()->GetLocalTransform(ControlIndex);
 						const FTransform PreviousTransform = GetHierarchy()->GetTransform(BoneElement, ERigTransformType::CurrentLocal);
 						const FTransform Transform = LocalTransform * PreviousTransform;
 						GetHierarchy()->SetTransform(BoneElement, Transform, ERigTransformType::CurrentLocal, true, false);
