Control Rig: Implementing LOD functionality for graph

diff --git a/Engine/Plugins/Experimental/ControlRig/Source/ControlRigDeveloper/Private/Graph/ControlRigGraph.cpp b/Engine/Plugins/Experimental/ControlRig/Source/ControlRigDeveloper/Private/Graph/ControlRigGraph.cpp
--- a/Engine/Plugins/Experimental/ControlRig/Source/ControlRigDeveloper/Private/Graph/ControlRigGraph.cpp
+++ b/Engine/Plugins/Experimental/ControlRig/Source/ControlRigDeveloper/Private/Graph/ControlRigGraph.cpp
@@ -189,6 +189,25 @@ void UControlRigGraph::HandleModifiedEvent(ERigVMGraphNotifType InNotifType, URi
 		return;
 	}
 
+	// increment the node topology version for any interaction
+	// with a node.
+	{
+		UControlRigGraphNode* EdNode = nullptr;
+		if (URigVMNode* ModelNode = Cast<URigVMNode>(InSubject))
+		{
+			EdNode = Cast<UControlRigGraphNode>(FindNodeForModelNodeName(ModelNode->GetFName()));
+		}
+		else if (URigVMPin* ModelPin = Cast<URigVMPin>(InSubject))
+		{
+			EdNode = Cast<UControlRigGraphNode>(FindNodeForModelNodeName(ModelPin->GetNode()->GetFName()));
+		}
+
+		if(EdNode)
+		{
+			EdNode->NodeTopologyVersion++;
+		}
+	}
+
 	switch (InNotifType)
 	{
 		case ERigVMGraphNotifType::GraphChanged:
diff --git a/Engine/Plugins/Experimental/ControlRig/Source/ControlRigDeveloper/Private/Graph/ControlRigGraphNode.cpp b/Engine/Plugins/Experimental/ControlRig/Source/ControlRigDeveloper/Private/Graph/ControlRigGraphNode.cpp
--- a/Engine/Plugins/Experimental/ControlRig/Source/ControlRigDeveloper/Private/Graph/ControlRigGraphNode.cpp
+++ b/Engine/Plugins/Experimental/ControlRig/Source/ControlRigDeveloper/Private/Graph/ControlRigGraphNode.cpp
@@ -40,6 +40,7 @@ UControlRigGraphNode::UControlRigGraphNode()
 , FullNodeTitle(FText::GetEmpty())
 , CachedTitleColor(FLinearColor(0.f, 0.f, 0.f, 0.f))
 , CachedNodeColor(FLinearColor(0.f, 0.f, 0.f, 0.f))
+, NodeTopologyVersion(INDEX_NONE)
 {
 	bHasCompilerMessage = false;
 	ErrorType = (int32)EMessageSeverity::Info + 1;
diff --git a/Engine/Plugins/Experimental/ControlRig/Source/ControlRigDeveloper/Public/Graph/ControlRigGraphNode.h b/Engine/Plugins/Experimental/ControlRig/Source/ControlRigDeveloper/Public/Graph/ControlRigGraphNode.h
--- a/Engine/Plugins/Experimental/ControlRig/Source/ControlRigDeveloper/Public/Graph/ControlRigGraphNode.h
+++ b/Engine/Plugins/Experimental/ControlRig/Source/ControlRigDeveloper/Public/Graph/ControlRigGraphNode.h
@@ -184,6 +184,9 @@ protected:
 
 private:
 
+	FORCEINLINE int32 GetNodeTopologyVersion() const { return NodeTopologyVersion; }
+	int32 NodeTopologyVersion;
+
 	FLinearColor CachedTitleColor;
 	FLinearColor CachedNodeColor;
 
diff --git a/Engine/Plugins/Experimental/ControlRig/Source/ControlRigEditor/Private/Graph/ControlRigConnectionDrawingPolicy.cpp b/Engine/Plugins/Experimental/ControlRig/Source/ControlRigEditor/Private/Graph/ControlRigConnectionDrawingPolicy.cpp
--- a/Engine/Plugins/Experimental/ControlRig/Source/ControlRigEditor/Private/Graph/ControlRigConnectionDrawingPolicy.cpp
+++ b/Engine/Plugins/Experimental/ControlRig/Source/ControlRigEditor/Private/Graph/ControlRigConnectionDrawingPolicy.cpp
@@ -96,7 +96,27 @@ void FControlRigConnectionDrawingPolicy::DrawPinGeometries(TMap<TSharedRef<SWidg
 				{
 					FConnectionParams Params;
 					DetermineWiringStyle(ThePin, TargetPin, /*inout*/ Params);
-					DrawSplineWithArrow(LinkStartWidgetGeometry->Geometry, LinkEndWidgetGeometry->Geometry, Params);
+
+					if(UseLowDetailConnections())
+					{
+						const float StartFudgeX = 4.0f;
+						const float EndFudgeX = 4.0f;
+
+						const FVector2D StartPoint = FGeometryHelper::VerticalMiddleRightOf(LinkStartWidgetGeometry->Geometry) - FVector2D(StartFudgeX, 0.0f);
+						const FVector2D EndPoint = FGeometryHelper::VerticalMiddleLeftOf(LinkEndWidgetGeometry->Geometry) - FVector2D(EndFudgeX, 0);
+						const FVector2D FakeTangent = (EndPoint - StartPoint).GetSafeNormal();
+
+						static TArray<FVector2D> LinePoints;
+						LinePoints.SetNum(2);
+						LinePoints[0] = StartPoint;
+						LinePoints[1] = EndPoint;
+
+						FSlateDrawElement::MakeLines(DrawElementsList, WireLayerID, FPaintGeometry(), LinePoints, ESlateDrawEffect::None, Params.WireColor);
+					}
+					else
+					{
+						DrawSplineWithArrow(LinkStartWidgetGeometry->Geometry, LinkEndWidgetGeometry->Geometry, Params);
+					}
 				}
 			}
 		}
@@ -223,7 +243,7 @@ bool FControlRigConnectionDrawingPolicy::GetAverageConnectedPositionForPin(UEdGr
 void FControlRigConnectionDrawingPolicy::DetermineWiringStyle(UEdGraphPin* OutputPin, UEdGraphPin* InputPin, /*inout*/ FConnectionParams& Params)
 {
 	FKismetConnectionDrawingPolicy::DetermineWiringStyle(OutputPin, InputPin, Params);
-	if (OutputPin == nullptr || InputPin == nullptr)
+	if (OutputPin == nullptr || InputPin == nullptr || UseLowDetailConnections())
 	{
 		return;
 	}
diff --git a/Engine/Plugins/Experimental/ControlRig/Source/ControlRigEditor/Private/Graph/ControlRigConnectionDrawingPolicy.h b/Engine/Plugins/Experimental/ControlRig/Source/ControlRigEditor/Private/Graph/ControlRigConnectionDrawingPolicy.h
--- a/Engine/Plugins/Experimental/ControlRig/Source/ControlRigEditor/Private/Graph/ControlRigConnectionDrawingPolicy.h
+++ b/Engine/Plugins/Experimental/ControlRig/Source/ControlRigEditor/Private/Graph/ControlRigConnectionDrawingPolicy.h
@@ -31,6 +31,8 @@ private:
 	// Each time a reroute node is encountered, input geometry is compared to output geometry to see if the pins on the reroute node need to be reversed
 	TMap<UControlRigGraphNode*, bool> RerouteNodeToReversedDirectionMap;
 
+	FORCEINLINE bool UseLowDetailConnections() const { return ZoomFactor <= 0.175f; /* zoom level -9 */ }
+
 	bool ShouldChangeTangentForReouteControlPoint(UControlRigGraphNode* Node);
 	// Average of the positions of all pins connected to InPin
 	bool GetAverageConnectedPositionForPin(UEdGraphPin* InPin, FVector2D& OutPos) const;
diff --git a/Engine/Plugins/Experimental/ControlRig/Source/ControlRigEditor/Private/Graph/SControlRigGraphNode.cpp b/Engine/Plugins/Experimental/ControlRig/Source/ControlRigEditor/Private/Graph/SControlRigGraphNode.cpp
--- a/Engine/Plugins/Experimental/ControlRig/Source/ControlRigEditor/Private/Graph/SControlRigGraphNode.cpp
+++ b/Engine/Plugins/Experimental/ControlRig/Source/ControlRigEditor/Private/Graph/SControlRigGraphNode.cpp
@@ -10,6 +10,10 @@
 #include "Widgets/Layout/SSpacer.h"
 #include "Widgets/Input/SButton.h"
 #include "Widgets/Layout/SScrollBar.h"
+#include "SLevelOfDetailBranchNode.h"
+#include "SGraphPanel.h"
+#include "Framework/Application/SlateApplication.h"
+#include "Fonts/FontMeasure.h"
 #include "GraphEditorSettings.h"
 #include "ControlRigEditorStyle.h"
 #include "Widgets/Layout/SWrapBox.h"
@@ -29,6 +33,7 @@
 #include "DetailLayoutBuilder.h"
 #include "EditorStyleSet.h"
 #include "SControlRigGraphPinVariableBinding.h"
+#include "Slate/SlateTextures.h"
 
 #if WITH_EDITOR
 #include "Editor.h"
@@ -203,6 +208,8 @@ void SControlRigGraphNode::Construct( const FArguments& InArgs )
 	.ToolTipText(LOCTEXT("NodeHitCountToolTip", "This number represents the hit count for a node.\nFor functions / collapse nodes it represents the sum of all hit counts of contained nodes.\n\nYou can enable / disable the display of the number in the Class Settings\n(Rig Graph Display Settings -> Show Node Run Counts)"));
 
 	ControlRigGraphNode->GetNodeTitleDirtied().BindSP(this, &SControlRigGraphNode::HandleNodeTitleDirtied);
+
+	LastHighDetailSize = FVector2D::ZeroVector;
 }
 
 TSharedRef<SWidget> SControlRigGraphNode::CreateNodeContentArea()
@@ -224,6 +231,29 @@ TSharedRef<SWidget> SControlRigGraphNode::CreateNodeContentArea()
 		];
 }
 
+bool SControlRigGraphNode::UseLowDetailPinContent() const
+{
+	if (const SGraphPanel* MyOwnerPanel = GetOwnerPanel().Get())
+	{
+		return (MyOwnerPanel->GetCurrentLOD() <= EGraphRenderingLOD::MediumDetail);
+	}
+	return false;
+}
+
+bool SControlRigGraphNode::UseLowDetailNodeContent() const
+{
+	if (const SGraphPanel* MyOwnerPanel = GetOwnerPanel().Get())
+	{
+		return (MyOwnerPanel->GetCurrentLOD() <= EGraphRenderingLOD::LowestDetail);
+	}
+	return false;
+}
+
+TOptional<FVector2D> SControlRigGraphNode::GetLowDetailDesiredSize() const
+{
+	return LastHighDetailSize;
+}
+
 TSharedPtr<SGraphPin> SControlRigGraphNode::GetHoveredPin(const FGeometry& MyGeometry, const FPointerEvent& MouseEvent) const
 {
 	TSharedPtr<SGraphPin> HoveredPin = SGraphNode::GetHoveredPin(MyGeometry, MouseEvent);
@@ -330,7 +360,7 @@ void SControlRigGraphNode::AddPin(const TSharedRef<SGraphPin>& PinToAdd)
 
 const FSlateBrush * SControlRigGraphNode::GetNodeBodyBrush() const
 {
-	return FEditorStyle::GetBrush("Graph.Node.TintedBody");
+	return FEditorStyle::GetBrush("Graph.Node.Body");
 }
 
 FReply SControlRigGraphNode::OnMouseButtonDown(const FGeometry& MyGeometry, const FPointerEvent& MouseEvent)
@@ -367,14 +397,9 @@ FReply SControlRigGraphNode::OnMouseButtonDoubleClick(const FGeometry& InMyGeome
 	return SGraphNode::OnMouseButtonDoubleClick(InMyGeometry, InMouseEvent);
 }
 
-bool SControlRigGraphNode::UseLowDetailNodeTitles() const
-{
-	return ParentUseLowDetailNodeTitles();
-}
-
 EVisibility SControlRigGraphNode::GetTitleVisibility() const
 {
-	return ParentUseLowDetailNodeTitles() ? EVisibility::Hidden : EVisibility::Visible;
+	return UseLowDetailNodeTitles() ? EVisibility::Hidden : EVisibility::Visible;
 }
 
 EVisibility SControlRigGraphNode::GetExecutionTreeVisibility() const
@@ -519,7 +544,7 @@ public:
 	SLATE_BEGIN_ARGS(SControlRigPinTreeRow) {}
 
 	SLATE_ARGUMENT(bool, LeftAligned)
-
+	
 	SLATE_END_ARGS()
 
 	void Construct(const FArguments& InArgs, const TSharedRef<STableViewBase>& InOwnerTableView)
@@ -618,7 +643,7 @@ public:
 
 		this->ChildSlot
 		[
-			ContentBox
+            ContentBox
 		];
 
 		InnerContentSlot = InnerContentSlotNativePtr;
@@ -785,11 +810,16 @@ TSharedRef<ITableRow> SControlRigGraphNode::MakePinTableRowWidget(URigVMPin* InI
 	{
 		TWeakPtr<SGraphPin> WeakPin = InputPinWidget.IsValid() ? InputPinWidget : OutputPinWidget;
 
-		TSharedRef<SWidget> LabelWidget = SNew(STextBlock)
+		TSharedPtr<STextBlock> LabelTextBlock;
+		TSharedRef<SWidget> LabelWidget = SAssignNew(LabelTextBlock, STextBlock)
 		.Text(this, &SControlRigGraphNode::GetPinLabel, WeakPin)
 		.TextStyle(FEditorStyle::Get(), NAME_DefaultPinLabelStyle)
+		.SimpleTextMode(true)
 		.ColorAndOpacity(this, &SControlRigGraphNode::GetPinTextColor, WeakPin);
 
+		const FText PinLabel = GetPinLabel(WeakPin);
+		const FVector2D LabelTextSize = FSlateApplication::Get().GetRenderer()->GetFontMeasureService()->Measure(PinLabel, IDetailLayoutBuilder::GetDetailFont()) * FVector2D(1.f, 0.2f);
+
 		// add to mapping that allows labels to act as hover widgets
 		if(InputPinWidget.IsValid())
 		{
@@ -807,14 +837,36 @@ TSharedRef<ITableRow> SControlRigGraphNode::MakePinTableRowWidget(URigVMPin* InI
 		if(OutputPinWidget.IsValid() && !InputPinWidget.IsValid())
 		{
 			TSharedRef<SWidget> InputWidget = 
-				SNew(SHorizontalBox)
-				+SHorizontalBox::Slot()
-				.HAlign(HAlign_Right)
-				.VAlign(VAlign_Center)
-				.FillWidth(1.0f)
-				.Padding(50.0f, 2.0f, 2.0f, 2.0f)
-				[
-					LabelWidget
+                SNew(SHorizontalBox)
+                +SHorizontalBox::Slot()
+                .HAlign(HAlign_Right)
+                .VAlign(VAlign_Center)
+                .FillWidth(1.0f)
+                .Padding(50.0f, 2.0f, 2.0f, 2.0f)
+                [
+                    SNew(SLevelOfDetailBranchNode)
+                    .UseLowDetailSlot(this, &SControlRigGraphNode::UseLowDetailNodeContent)
+                    .LowDetail()
+                    [
+                    	SNew(SSpacer)
+                    ]
+                    .HighDetail()
+                    [
+                        SNew(SLevelOfDetailBranchNode)
+                        .UseLowDetailSlot(this, &SControlRigGraphNode::UseLowDetailPinContent)
+                        .LowDetail()
+                        [
+                            // replacement rectangle in the color of the text
+                            SNew(SImage)
+                            .DesiredSizeOverride(LabelTextSize)
+                            .Image(FEditorStyle::GetBrush("WhiteBrush"))
+                            .ColorAndOpacity(FEditorStyle::GetSlateColor("DefaultForeground"))
+                        ]
+                        .HighDetail()
+                        [
+                            LabelWidget
+                        ]
+					]
 				];
 
 			TSharedRef<SWidget> OutputWidget = 
@@ -849,7 +901,26 @@ TSharedRef<ITableRow> SControlRigGraphNode::MakePinTableRowWidget(URigVMPin* InI
 						InputPinWidget.IsValid() ? StaticCastSharedRef<SWidget>(InputPinWidget.ToSharedRef()) : StaticCastSharedRef<SWidget>(SNew(SSpacer).Size(FVector2D(16.0f, 14.0f)))
 					]
 				];
-				
+
+			TSharedPtr<SWidget> InputValueWrapperWidget;
+			if(InputPinValueWidget.IsValid())
+			{
+				InputValueWrapperWidget = SNew(SLevelOfDetailBranchNode)
+				.UseLowDetailSlot(this, &SControlRigGraphNode::UseLowDetailPinContent)
+				.LowDetail()
+				[
+					SNew(SSpacer)
+				]
+				.HighDetail()
+				[
+					InputPinValueWidget.ToSharedRef()
+				];
+			}
+			else
+			{
+				InputValueWrapperWidget = SNew(SSpacer);
+			}
+
 			TSharedRef<SWidget> OutputWidget = 
 				SNew(SHorizontalBox)
 				+SHorizontalBox::Slot()
@@ -858,15 +929,37 @@ TSharedRef<ITableRow> SControlRigGraphNode::MakePinTableRowWidget(URigVMPin* InI
 				.AutoWidth()
 				.Padding(2.0f)
 				[
-					LabelWidget
-				]
+					SNew(SLevelOfDetailBranchNode)
+	                .UseLowDetailSlot(this, &SControlRigGraphNode::UseLowDetailNodeContent)
+	                .LowDetail()
+	                [
+	                    SNew(SSpacer)
+	                ]
+	                .HighDetail()
+	                [
+		                SNew(SLevelOfDetailBranchNode)
+		                .UseLowDetailSlot(this, &SControlRigGraphNode::UseLowDetailPinContent)
+		                .LowDetail()
+		                [
+		                    // replacement rectangle in the color of the text
+		                    SNew(SImage)
+		                    .DesiredSizeOverride(LabelTextSize)
+		                    .Image(FEditorStyle::GetBrush("WhiteBrush"))
+		                    .ColorAndOpacity(FEditorStyle::GetSlateColor("DefaultForeground"))
+		                ]
+		                .HighDetail()
+		                [
+		                    LabelWidget
+		                ]
+		            ]
+	            ]
 				+SHorizontalBox::Slot()
 				.AutoWidth()
 				.HAlign(HAlign_Left)
 				.VAlign(VAlign_Center)
 				.Padding(2.0f, 2.0f, 18.0f, 2.0f)
 				[
-					InputPinValueWidget.IsValid() ? InputPinValueWidget.ToSharedRef() : SNew(SSpacer)
+					InputValueWrapperWidget.ToSharedRef()
 				]
 				+SHorizontalBox::Slot()
 				.FillWidth(1.0f)
@@ -1273,6 +1366,11 @@ void SControlRigGraphNode::Tick(const FGeometry& AllottedGeometry, const double
 		GraphNode->NodeHeight = (int32)AllottedGeometry.Size.Y;
 		RefreshErrorInfo();
 	}
+
+	if((!UseLowDetailNodeContent()) && (LeftNodeBox != nullptr))
+	{
+		LastHighDetailSize = LeftNodeBox->GetTickSpaceGeometry().Size;
+	}
 }
 
 void SControlRigGraphNode::HandleNodeTitleDirtied()
@@ -1309,4 +1407,13 @@ FText SControlRigGraphNode::GetInstructionCountText() const
 	return FText();
 }
 
-#undef LOCTEXT_NAMESPACE
\ No newline at end of file
+int32 SControlRigGraphNode::GetNodeTopologyVersion() const
+{
+	if(UControlRigGraphNode* ControlRigGraphNode = Cast<UControlRigGraphNode>(GraphNode))
+	{
+		return ControlRigGraphNode->GetNodeTopologyVersion();
+	}
+	return INDEX_NONE;
+}
+
+#undef LOCTEXT_NAMESPACE
diff --git a/Engine/Plugins/Experimental/ControlRig/Source/ControlRigEditor/Private/Graph/SControlRigGraphNode.h b/Engine/Plugins/Experimental/ControlRig/Source/ControlRigEditor/Private/Graph/SControlRigGraphNode.h
--- a/Engine/Plugins/Experimental/ControlRig/Source/ControlRigEditor/Private/Graph/SControlRigGraphNode.h
+++ b/Engine/Plugins/Experimental/ControlRig/Source/ControlRigEditor/Private/Graph/SControlRigGraphNode.h
@@ -34,7 +34,6 @@ public:
 
 	// SGraphNode interface
 	virtual TSharedRef<SWidget> CreateTitleWidget(TSharedPtr<SNodeTitle> InNodeTitle) override;
-	virtual bool UseLowDetailNodeTitles() const override;
 	virtual void EndUserInteraction() const override;
 	virtual void MoveTo(const FVector2D& NewPosition, FNodeSet& NodeFilter) override;
 	virtual void AddPin( const TSharedRef<SGraphPin>& PinToAdd ) override;
@@ -56,10 +55,12 @@ public:
 	virtual void Tick(const FGeometry& AllottedGeometry, const double InCurrentTime, const float InDeltaTime) override;
 
 private:
-	bool ParentUseLowDetailNodeTitles() const
-	{
-		return SGraphNode::UseLowDetailNodeTitles();
-	}
+
+	bool UseLowDetailPinContent() const;
+
+	bool UseLowDetailNodeContent() const;
+
+	TOptional<FVector2D> GetLowDetailDesiredSize() const;
 
 	EVisibility GetTitleVisibility() const;
 
@@ -96,6 +97,9 @@ private:
 	FText GetInstructionCountText() const;
 
 private:
+
+	int32 GetNodeTopologyVersion() const;
+	
 	/** Cached widget title area */
 	TSharedPtr<SOverlay> TitleAreaWidget;
 
@@ -135,4 +139,6 @@ private:
 	TSharedPtr<SNodeTitle> NodeTitle;
 
 	TWeakObjectPtr<UControlRigBlueprint> Blueprint;
+
+	FVector2D LastHighDetailSize;
 };
