Control Rig: Avoid recursive recompilation during rename of elements

diff --git a/Engine/Plugins/Experimental/ControlRig/Source/ControlRigDeveloper/Private/Graph/ControlRigGraph.cpp b/Engine/Plugins/Experimental/ControlRig/Source/ControlRigDeveloper/Private/Graph/ControlRigGraph.cpp
--- a/Engine/Plugins/Experimental/ControlRig/Source/ControlRigDeveloper/Private/Graph/ControlRigGraph.cpp
+++ b/Engine/Plugins/Experimental/ControlRig/Source/ControlRigDeveloper/Private/Graph/ControlRigGraph.cpp
@@ -30,6 +30,7 @@ UControlRigGraph::UControlRigGraph()
 	bSuspendModelNotifications = false;
 	bIsTemporaryGraphForCopyPaste = false;
 	bIsSelecting = false;
+	LastHierarchyTopologyVersion = INDEX_NONE;
 	bIsFunctionDefinition = false;
 }
 
@@ -88,29 +89,34 @@ void UControlRigGraph::CacheNameLists(URigHierarchy* InHierarchy, const FControl
 	check(InHierarchy);
 	check(DrawContainer);
 
-	ElementNameLists.FindOrAdd(ERigElementType::All);
-	ElementNameLists.FindOrAdd(ERigElementType::Bone);
-	ElementNameLists.FindOrAdd(ERigElementType::Null);
-	ElementNameLists.FindOrAdd(ERigElementType::Control);
-	ElementNameLists.FindOrAdd(ERigElementType::Curve);
-	ElementNameLists.FindOrAdd(ERigElementType::RigidBody);
-	ElementNameLists.FindOrAdd(ERigElementType::Socket);
-
-	TArray<TSharedPtr<FString>>& AllNameList = ElementNameLists.FindChecked(ERigElementType::All);
-	TArray<TSharedPtr<FString>>& BoneNameList = ElementNameLists.FindChecked(ERigElementType::Bone);
-	TArray<TSharedPtr<FString>>& NullNameList = ElementNameLists.FindChecked(ERigElementType::Null);
-	TArray<TSharedPtr<FString>>& ControlNameList = ElementNameLists.FindChecked(ERigElementType::Control);
-	TArray<TSharedPtr<FString>>& CurveNameList = ElementNameLists.FindChecked(ERigElementType::Curve);
-	TArray<TSharedPtr<FString>>& RigidBodyNameList = ElementNameLists.FindChecked(ERigElementType::RigidBody);
-	TArray<TSharedPtr<FString>>& SocketNameList = ElementNameLists.FindChecked(ERigElementType::Socket);
-	
-	CacheNameListForHierarchy<FRigBaseElement>(InHierarchy, AllNameList);
-	CacheNameListForHierarchy<FRigBoneElement>(InHierarchy, BoneNameList);
-	CacheNameListForHierarchy<FRigNullElement>(InHierarchy, NullNameList);
-	CacheNameListForHierarchy<FRigControlElement>(InHierarchy, ControlNameList);
-	CacheNameListForHierarchy<FRigCurveElement>(InHierarchy, CurveNameList);
-	CacheNameListForHierarchy<FRigRigidBodyElement>(InHierarchy, RigidBodyNameList);
-	CacheNameListForHierarchy<FRigSocketElement>(InHierarchy, SocketNameList);
+	if(LastHierarchyTopologyVersion != InHierarchy->GetTopologyVersion())
+	{
+		ElementNameLists.FindOrAdd(ERigElementType::All);
+		ElementNameLists.FindOrAdd(ERigElementType::Bone);
+		ElementNameLists.FindOrAdd(ERigElementType::Null);
+		ElementNameLists.FindOrAdd(ERigElementType::Control);
+		ElementNameLists.FindOrAdd(ERigElementType::Curve);
+		ElementNameLists.FindOrAdd(ERigElementType::RigidBody);
+		ElementNameLists.FindOrAdd(ERigElementType::Socket);
+
+		TArray<TSharedPtr<FString>>& AllNameList = ElementNameLists.FindChecked(ERigElementType::All);
+		TArray<TSharedPtr<FString>>& BoneNameList = ElementNameLists.FindChecked(ERigElementType::Bone);
+		TArray<TSharedPtr<FString>>& NullNameList = ElementNameLists.FindChecked(ERigElementType::Null);
+		TArray<TSharedPtr<FString>>& ControlNameList = ElementNameLists.FindChecked(ERigElementType::Control);
+		TArray<TSharedPtr<FString>>& CurveNameList = ElementNameLists.FindChecked(ERigElementType::Curve);
+		TArray<TSharedPtr<FString>>& RigidBodyNameList = ElementNameLists.FindChecked(ERigElementType::RigidBody);
+		TArray<TSharedPtr<FString>>& SocketNameList = ElementNameLists.FindChecked(ERigElementType::Socket);
+		
+		CacheNameListForHierarchy<FRigBaseElement>(InHierarchy, AllNameList);
+		CacheNameListForHierarchy<FRigBoneElement>(InHierarchy, BoneNameList);
+		CacheNameListForHierarchy<FRigNullElement>(InHierarchy, NullNameList);
+		CacheNameListForHierarchy<FRigControlElement>(InHierarchy, ControlNameList);
+		CacheNameListForHierarchy<FRigCurveElement>(InHierarchy, CurveNameList);
+		CacheNameListForHierarchy<FRigRigidBodyElement>(InHierarchy, RigidBodyNameList);
+		CacheNameListForHierarchy<FRigSocketElement>(InHierarchy, SocketNameList);
+
+		LastHierarchyTopologyVersion = InHierarchy->GetTopologyVersion();
+	}
 	CacheNameList<FControlRigDrawContainer>(*DrawContainer, DrawingNameList);
 }
 
diff --git a/Engine/Plugins/Experimental/ControlRig/Source/ControlRigDeveloper/Public/Graph/ControlRigGraph.h b/Engine/Plugins/Experimental/ControlRig/Source/ControlRigDeveloper/Public/Graph/ControlRigGraph.h
--- a/Engine/Plugins/Experimental/ControlRig/Source/ControlRigDeveloper/Public/Graph/ControlRigGraph.h
+++ b/Engine/Plugins/Experimental/ControlRig/Source/ControlRigDeveloper/Public/Graph/ControlRigGraph.h
@@ -126,6 +126,7 @@ private:
 
 	TMap<ERigElementType, TArray<TSharedPtr<FString>>> ElementNameLists;
 	TArray<TSharedPtr<FString>> DrawingNameList;
+	int32 LastHierarchyTopologyVersion;
 
 	bool bIsSelecting;
 
diff --git a/Engine/Plugins/Experimental/ControlRig/Source/ControlRigEditor/Private/Editor/ControlRigEditor.cpp b/Engine/Plugins/Experimental/ControlRig/Source/ControlRigEditor/Private/Editor/ControlRigEditor.cpp
--- a/Engine/Plugins/Experimental/ControlRig/Source/ControlRigEditor/Private/Editor/ControlRigEditor.cpp
+++ b/Engine/Plugins/Experimental/ControlRig/Source/ControlRigEditor/Private/Editor/ControlRigEditor.cpp
@@ -3560,7 +3560,10 @@ void FControlRigEditor::OnHierarchyChanged()
 	CacheNameLists();
 	if (UControlRigBlueprint* ControlRigBP = GetControlRigBlueprint())
 	{
-		ControlRigBP->PropagateHierarchyFromBPToInstances();
+		{
+			TGuardValue<bool> GuardNotifs(ControlRigBP->bSuspendAllNotifications, true);
+			ControlRigBP->PropagateHierarchyFromBPToInstances();
+		}
 
 		FBlueprintEditorUtils::MarkBlueprintAsModified(GetControlRigBlueprint());
 		
@@ -3726,6 +3729,13 @@ void FControlRigEditor::OnHierarchyModified(ERigHierarchyNotification InNotif, U
 					continue;
 				}
 
+				URigVMController* Controller = RigGraph->GetController();
+				if(Controller == nullptr)
+				{
+					continue;
+				}
+
+				RigBlueprint->IncrementVMRecompileBracket();
 				for (UEdGraphNode* Node : RigGraph->Nodes)
 				{
 					if (UControlRigGraphNode* RigNode = Cast<UControlRigGraphNode>(Node))
@@ -3742,7 +3752,7 @@ void FControlRigEditor::OnHierarchyModified(ERigHierarchyNotification InNotif, U
 								{
 									if (ModelPin->GetDefaultValue() == OldNameStr)
 									{
-										GetFocusedController()->SetPinDefaultValue(ModelPin->GetPinPath(), NewNameStr, false);
+										Controller->SetPinDefaultValue(ModelPin->GetPinPath(), NewNameStr, false);
 									}
 								}
 								else if (ModelPin->GetCPPTypeObject() == FRigElementKey::StaticStruct())
@@ -3758,7 +3768,7 @@ void FControlRigEditor::OnHierarchyModified(ERigHierarchyNotification InNotif, U
 												FString NameStr = NamePin->GetDefaultValue();
 												if (NameStr == OldNameStr)
 												{
-													GetFocusedController()->SetPinDefaultValue(NamePin->GetPinPath(), NewNameStr);
+													Controller->SetPinDefaultValue(NamePin->GetPinPath(), NewNameStr);
 												}
 											}
 										}
@@ -3768,6 +3778,7 @@ void FControlRigEditor::OnHierarchyModified(ERigHierarchyNotification InNotif, U
 						}
 					}
 				}
+				RigBlueprint->DecrementVMRecompileBracket();
 			}
 				
 			CacheNameLists();
@@ -3823,10 +3834,21 @@ void FControlRigEditor::OnHierarchyModified_AnyThread(ERigHierarchyNotification
 		Key = InElement->GetKey();
 	}
 
+	if(IsInGameThread())
+	{
+		UControlRigBlueprint* RigBlueprint = GetControlRigBlueprint();
+		check(RigBlueprint);
+
+		if(RigBlueprint->bSuspendAllNotifications)
+		{
+			return;
+		}
+	}
+
 	TWeakObjectPtr<URigHierarchy> WeakHierarchy = InHierarchy;
 	FFunctionGraphTask::CreateAndDispatchWhenReady([this, InNotif, WeakHierarchy, Key]()
     {
-    	if(!WeakHierarchy.IsValid())
+		if(!WeakHierarchy.IsValid())
     	{
     		return;
     	}
@@ -3839,9 +3861,13 @@ void FControlRigEditor::OnHierarchyModified_AnyThread(ERigHierarchyNotification
 			case ERigHierarchyNotification::ElementRemoved:
 			case ERigHierarchyNotification::ElementRenamed:
 			case ERigHierarchyNotification::ParentChanged:
-			case ERigHierarchyNotification::HierarchyReset:
+            case ERigHierarchyNotification::HierarchyReset:
 			{
-				CacheNameLists();
+				// only recache the name lists of this is the main hierarchy
+				if(WeakHierarchy.Get()->GetOuter()->IsA<UControlRigBlueprint>())
+				{
+					CacheNameLists();
+				}
 				break;
 			}
 			case ERigHierarchyNotification::ControlSettingChanged:
