Control Rig: Enable Set Initial Value for Controls in details panel

diff --git a/Engine/Plugins/Experimental/ControlRig/Source/ControlRigEditor/Private/ControlRigElementDetails.cpp b/Engine/Plugins/Experimental/ControlRig/Source/ControlRigEditor/Private/ControlRigElementDetails.cpp
--- a/Engine/Plugins/Experimental/ControlRig/Source/ControlRigEditor/Private/ControlRigElementDetails.cpp
+++ b/Engine/Plugins/Experimental/ControlRig/Source/ControlRigEditor/Private/ControlRigElementDetails.cpp
@@ -72,6 +72,25 @@ void RigElementDetails_GetCustomizedInfo(TSharedRef<IPropertyHandle> InStructPro
 	}
 }
 
+UControlRigBlueprint* RigElementDetails_GetBlueprintFromHierarchy(URigHierarchy* InHierarchy)
+{
+	if(InHierarchy == nullptr)
+	{
+		return nullptr;
+	}
+
+	UControlRigBlueprint* Blueprint = InHierarchy->GetTypedOuter<UControlRigBlueprint>();
+	if(Blueprint == nullptr)
+	{
+		UControlRig* Rig = InHierarchy->GetTypedOuter<UControlRig>();
+		if(Rig)
+		{
+			Blueprint = Cast<UControlRigBlueprint>(Rig->GetClass()->ClassGeneratedBy);
+        }
+	}
+	return Blueprint;
+}
+
 void FRigElementKeyDetails::CustomizeHeader(TSharedRef<IPropertyHandle> InStructPropertyHandle, FDetailWidgetRow& HeaderRow, IPropertyTypeCustomizationUtils& StructCustomizationUtils)
 {
 	BlueprintBeingCustomized = nullptr;
@@ -771,7 +790,14 @@ void FRigControlElementDetails_SetupBoolValueWidget(IDetailCategoryBuilder& InCa
 	        		{
                         const FRigControlValue Value = FRigControlValue::Make<bool>(NewState == ECheckBoxState::Checked);
 						HierarchyPtr->SetControlValue(ControlElement->GetKey(), Value, InValueType, true);
-                    }
+	        			if(InValueType == ERigControlValueType::Initial)
+	        			{
+                            if(UControlRigBlueprint* Blueprint = RigElementDetails_GetBlueprintFromHierarchy(HierarchyPtr.Get()))
+                            {
+                                Blueprint->Hierarchy->SetControlValue(ControlElement->GetKey(), Value, InValueType, true);
+                            }
+                        }
+		            }
 	        	}
 	        })
 	    ]
@@ -853,7 +879,14 @@ void FRigControlElementDetails_SetupIntegerValueWidget(IDetailCategoryBuilder& I
                 		{
                             const FRigControlValue Value = FRigControlValue::Make<int32>(NewSelection);
                             HierarchyPtr->SetControlValue(ControlElement->GetKey(), Value, InValueType, true);
-                        }
+                			if(InValueType == ERigControlValueType::Initial)
+                			{
+                                if(UControlRigBlueprint* Blueprint = RigElementDetails_GetBlueprintFromHierarchy(HierarchyPtr.Get()))
+                                {
+                                    Blueprint->Hierarchy->SetControlValue(ControlElement->GetKey(), Value, InValueType, true);
+                                }
+                            }
+	                    }
                 	}
                 })
                 .Font(FEditorStyle::GetFontStyle(TEXT("MenuItem.Font")))
@@ -930,6 +963,13 @@ void FRigControlElementDetails_SetupIntegerValueWidget(IDetailCategoryBuilder& I
                 			{
                                 const FRigControlValue Value = FRigControlValue::Make<int32>(InNewSelection.GetValue());
 								HierarchyPtr->SetControlValue(ControlElement->GetKey(), Value, InValueType, true);
+                				if(InValueType == ERigControlValueType::Initial)
+                				{
+                                    if(UControlRigBlueprint* Blueprint = RigElementDetails_GetBlueprintFromHierarchy(HierarchyPtr.Get()))
+                                    {
+                                        Blueprint->Hierarchy->SetControlValue(ControlElement->GetKey(), Value, InValueType, true);
+                                    }
+                                }
                             }
                 		}
                 	}
@@ -1034,6 +1074,13 @@ void FRigControlElementDetails_SetupFloatValueWidget(IDetailCategoryBuilder& InC
             			{
                             const FRigControlValue Value = FRigControlValue::Make<float>(InNewSelection.GetValue());
                             HierarchyPtr->SetControlValue(ControlElement->GetKey(), Value, InValueType, true);
+            				if(InValueType == ERigControlValueType::Initial)
+            				{
+            					if(UControlRigBlueprint* Blueprint = RigElementDetails_GetBlueprintFromHierarchy(HierarchyPtr.Get()))
+            					{
+            						Blueprint->Hierarchy->SetControlValue(ControlElement->GetKey(), Value, InValueType, true);
+            					}
+            				}
                         }
             		}
             	}
@@ -1105,6 +1152,13 @@ void FRigControlElementDetails_SetupStructValueWidget(IDetailCategoryBuilder& In
 		{
             const FRigControlValue Value = FRigControlValue::Make(*(T*)StructToDisplay->GetStructMemory());
 			HierarchyPtr->SetControlValue(Key, Value, InValueType, true);
+			if(InValueType == ERigControlValueType::Initial)
+			{
+                if(UControlRigBlueprint* Blueprint = RigElementDetails_GetBlueprintFromHierarchy(HierarchyPtr.Get()))
+                {
+                    Blueprint->Hierarchy->SetControlValue(Key, Value, InValueType, true);
+                }
+            }
 		}
 	});
 
