Control Rig: Improve name list performance and fix invalid name lists for graph UI

diff --git a/Engine/Plugins/Experimental/ControlRig/Source/ControlRig/Private/Rigs/RigHierarchyController.cpp b/Engine/Plugins/Experimental/ControlRig/Source/ControlRig/Private/Rigs/RigHierarchyController.cpp
--- a/Engine/Plugins/Experimental/ControlRig/Source/ControlRig/Private/Rigs/RigHierarchyController.cpp
+++ b/Engine/Plugins/Experimental/ControlRig/Source/ControlRig/Private/Rigs/RigHierarchyController.cpp
@@ -1357,6 +1357,7 @@ bool URigHierarchyController::RenameElement(FRigBaseElement* InElement, const FN
 	}
 	
 	Hierarchy->PreviousNameMap.FindOrAdd(NewKey) = OldKey;
+	Hierarchy->TopologyVersion++;
 	Notify(ERigHierarchyNotification::ElementRenamed, InElement);
 
 	return true;
diff --git a/Engine/Plugins/Experimental/ControlRig/Source/ControlRigDeveloper/Private/Graph/ControlRigGraph.cpp b/Engine/Plugins/Experimental/ControlRig/Source/ControlRigDeveloper/Private/Graph/ControlRigGraph.cpp
--- a/Engine/Plugins/Experimental/ControlRig/Source/ControlRigDeveloper/Private/Graph/ControlRigGraph.cpp
+++ b/Engine/Plugins/Experimental/ControlRig/Source/ControlRigDeveloper/Private/Graph/ControlRigGraph.cpp
@@ -120,7 +120,7 @@ void UControlRigGraph::CacheNameLists(URigHierarchy* InHierarchy, const FControl
 	CacheNameList<FControlRigDrawContainer>(*DrawContainer, DrawingNameList);
 }
 
-const TArray<TSharedPtr<FString>>& UControlRigGraph::GetElementNameList(ERigElementType InElementType) const
+const TArray<TSharedPtr<FString>>* UControlRigGraph::GetElementNameList(ERigElementType InElementType) const
 {
 	if (UControlRigGraph* OuterGraph = Cast<UControlRigGraph>(GetOuter()))
 	{
@@ -129,7 +129,7 @@ const TArray<TSharedPtr<FString>>& UControlRigGraph::GetElementNameList(ERigElem
 	
 	if(InElementType == ERigElementType::None)
 	{
-		return EmptyElementNameList;
+		return &EmptyElementNameList;
 	}
 	
 	if(!ElementNameLists.Contains(InElementType))
@@ -137,16 +137,16 @@ const TArray<TSharedPtr<FString>>& UControlRigGraph::GetElementNameList(ERigElem
 		UControlRigBlueprint* Blueprint = GetBlueprint();
 		if(Blueprint == nullptr)
 		{
-			return EmptyElementNameList;
+			return &EmptyElementNameList;
 		}
 
 		UControlRigGraph* MutableThis = (UControlRigGraph*)this;
 		MutableThis->CacheNameLists(Blueprint->Hierarchy, &Blueprint->DrawContainer);
 	}
-	return ElementNameLists.FindChecked(InElementType);
+	return &ElementNameLists.FindChecked(InElementType);
 }
 
-const TArray<TSharedPtr<FString>>& UControlRigGraph::GetElementNameList(URigVMPin* InPin) const
+const TArray<TSharedPtr<FString>>* UControlRigGraph::GetElementNameList(URigVMPin* InPin) const
 {
 	if (InPin)
 	{
@@ -170,13 +170,13 @@ const TArray<TSharedPtr<FString>>& UControlRigGraph::GetElementNameList(URigVMPi
 	return GetBoneNameList(nullptr);
 }
 
-const TArray<TSharedPtr<FString>>& UControlRigGraph::GetDrawingNameList(URigVMPin* InPin) const
+const TArray<TSharedPtr<FString>>* UControlRigGraph::GetDrawingNameList(URigVMPin* InPin) const
 {
 	if (UControlRigGraph* OuterGraph = Cast<UControlRigGraph>(GetOuter()))
 	{
 		return OuterGraph->GetDrawingNameList(InPin);
 	}
-	return DrawingNameList;
+	return &DrawingNameList;
 }
 
 void UControlRigGraph::HandleModifiedEvent(ERigVMGraphNotifType InNotifType, URigVMGraph* InGraph, UObject* InSubject)
diff --git a/Engine/Plugins/Experimental/ControlRig/Source/ControlRigDeveloper/Public/Graph/ControlRigGraph.h b/Engine/Plugins/Experimental/ControlRig/Source/ControlRigDeveloper/Public/Graph/ControlRigGraph.h
--- a/Engine/Plugins/Experimental/ControlRig/Source/ControlRigDeveloper/Public/Graph/ControlRigGraph.h
+++ b/Engine/Plugins/Experimental/ControlRig/Source/ControlRigDeveloper/Public/Graph/ControlRigGraph.h
@@ -40,27 +40,27 @@ public:
 #endif
 #if WITH_EDITOR
 
-	FORCEINLINE const TArray<TSharedPtr<FString>>& GetBoneNameList(URigVMPin* InPin = nullptr) const
+	FORCEINLINE const TArray<TSharedPtr<FString>>* GetBoneNameList(URigVMPin* InPin = nullptr) const
 	{
 		return GetElementNameList(ERigElementType::Bone);
 	}
-	FORCEINLINE const TArray<TSharedPtr<FString>>& GetControlNameList(URigVMPin* InPin = nullptr) const
+	FORCEINLINE const TArray<TSharedPtr<FString>>* GetControlNameList(URigVMPin* InPin = nullptr) const
 	{
 		return GetElementNameList(ERigElementType::Control);
 	}
-	FORCEINLINE const TArray<TSharedPtr<FString>>& GetNullNameList(URigVMPin* InPin = nullptr) const
+	FORCEINLINE const TArray<TSharedPtr<FString>>* GetNullNameList(URigVMPin* InPin = nullptr) const
 	{
 		return GetElementNameList(ERigElementType::Null);
 	}
-	FORCEINLINE const TArray<TSharedPtr<FString>>& GetCurveNameList(URigVMPin* InPin = nullptr) const
+	FORCEINLINE const TArray<TSharedPtr<FString>>* GetCurveNameList(URigVMPin* InPin = nullptr) const
 	{
 		return GetElementNameList(ERigElementType::Curve);
 	}
 
 	void CacheNameLists(URigHierarchy* InHierarchy, const FControlRigDrawContainer* DrawContainer);
-	const TArray<TSharedPtr<FString>>& GetElementNameList(ERigElementType InElementType = ERigElementType::Bone) const;
-	const TArray<TSharedPtr<FString>>& GetElementNameList(URigVMPin* InPin = nullptr) const;
-	const TArray<TSharedPtr<FString>>& GetDrawingNameList(URigVMPin* InPin = nullptr) const;
+	const TArray<TSharedPtr<FString>>* GetElementNameList(ERigElementType InElementType = ERigElementType::Bone) const;
+	const TArray<TSharedPtr<FString>>* GetElementNameList(URigVMPin* InPin = nullptr) const;
+	const TArray<TSharedPtr<FString>>* GetDrawingNameList(URigVMPin* InPin = nullptr) const;
 
 	bool bSuspendModelNotifications;
 	bool bIsTemporaryGraphForCopyPaste;
diff --git a/Engine/Plugins/Experimental/ControlRig/Source/ControlRigEditor/Private/ControlRigElementDetails.cpp b/Engine/Plugins/Experimental/ControlRig/Source/ControlRigEditor/Private/ControlRigElementDetails.cpp
--- a/Engine/Plugins/Experimental/ControlRig/Source/ControlRigEditor/Private/ControlRigElementDetails.cpp
+++ b/Engine/Plugins/Experimental/ControlRig/Source/ControlRigEditor/Private/ControlRigElementDetails.cpp
@@ -242,7 +242,7 @@ void FRigElementKeyDetails::UpdateElementNameList()
 		{
 			if (UControlRigGraph* RigGraph = Cast<UControlRigGraph>(Graph))
 			{
-				ElementNameList = RigGraph->GetElementNameList(GetElementType());
+				ElementNameList = *RigGraph->GetElementNameList(GetElementType());
 				return;
 			}
 		}
@@ -346,19 +346,19 @@ void FRigUnitDetails::CustomizeDetails(IDetailLayoutBuilder& DetailBuilder)
 				const TArray<TSharedPtr<FString>>* NameList = nullptr;
 				if (CustomWidgetName == TEXT("BoneName"))
 				{
-					NameList = &GraphBeingCustomized->GetBoneNameList();
+					NameList = GraphBeingCustomized->GetBoneNameList();
 				}
 				else if (CustomWidgetName == TEXT("ControlName"))
 				{
-					NameList = &GraphBeingCustomized->GetControlNameList();
+					NameList = GraphBeingCustomized->GetControlNameList();
 				}
 				else if (CustomWidgetName == TEXT("SpaceName"))
 				{
-					NameList = &GraphBeingCustomized->GetNullNameList();
+					NameList = GraphBeingCustomized->GetNullNameList();
 				}
 				else if (CustomWidgetName == TEXT("CurveName"))
 				{
-					NameList = &GraphBeingCustomized->GetCurveNameList();
+					NameList = GraphBeingCustomized->GetCurveNameList();
 				}
 
 				if (NameList)
diff --git a/Engine/Plugins/Experimental/ControlRig/Source/ControlRigEditor/Private/Editor/ControlRigEditor.cpp b/Engine/Plugins/Experimental/ControlRig/Source/ControlRigEditor/Private/Editor/ControlRigEditor.cpp
--- a/Engine/Plugins/Experimental/ControlRig/Source/ControlRigEditor/Private/Editor/ControlRigEditor.cpp
+++ b/Engine/Plugins/Experimental/ControlRig/Source/ControlRigEditor/Private/Editor/ControlRigEditor.cpp
@@ -2666,7 +2666,7 @@ void FControlRigEditor::HandleViewportCreated(const TSharedRef<class IPersonaVie
 					{
 						if (UControlRigGraph* RigGraph = Cast<UControlRigGraph>(Graph))
 						{
-							const TArray<TSharedPtr<FString>>* BoneNameList = &RigGraph->GetBoneNameList();
+							const TArray<TSharedPtr<FString>>* BoneNameList = RigGraph->GetBoneNameList();
 
 							InMenuBuilder.AddWidget(
 								SNew(SBox)
@@ -3557,7 +3557,6 @@ void FControlRigEditor::OnHierarchyChanged()
 {
 	DECLARE_SCOPE_HIERARCHICAL_COUNTER_FUNC()
 
-	CacheNameLists();
 	if (UControlRigBlueprint* ControlRigBP = GetControlRigBlueprint())
 	{
 		{
@@ -3603,6 +3602,8 @@ void FControlRigEditor::OnHierarchyChanged()
 	{
 		ClearDetailObject();
 	}
+	
+	CacheNameLists();
 }
 
 
@@ -3781,7 +3782,6 @@ void FControlRigEditor::OnHierarchyModified(ERigHierarchyNotification InNotif, U
 				RigBlueprint->DecrementVMRecompileBracket();
 			}
 				
-			CacheNameLists();
 			OnHierarchyChanged();
 
 			break;
diff --git a/Engine/Plugins/Experimental/ControlRig/Source/ControlRigEditor/Private/Graph/SControlRigGraphPinNameList.cpp b/Engine/Plugins/Experimental/ControlRig/Source/ControlRigEditor/Private/Graph/SControlRigGraphPinNameList.cpp
--- a/Engine/Plugins/Experimental/ControlRig/Source/ControlRigEditor/Private/Graph/SControlRigGraphPinNameList.cpp
+++ b/Engine/Plugins/Experimental/ControlRig/Source/ControlRigEditor/Private/Graph/SControlRigGraphPinNameList.cpp
@@ -22,7 +22,8 @@ void SControlRigGraphPinNameList::Construct(const FArguments& InArgs, UEdGraphPi
 TSharedRef<SWidget>	SControlRigGraphPinNameList::GetDefaultValueWidget()
 {
 	TSharedPtr<FString> InitialSelected;
-	for (TSharedPtr<FString> Item : GetNameList())
+	const TArray<TSharedPtr<FString>>* List = GetNameList();
+	for (TSharedPtr<FString> Item : (*List))
 	{
 		if (Item->Equals(GetNameListText().ToString()))
 		{
@@ -35,7 +36,7 @@ TSharedRef<SWidget>	SControlRigGraphPinNameList::GetDefaultValueWidget()
 		[
 			SAssignNew(NameListComboBox, SControlRigGraphPinNameListValueWidget)
 				.Visibility(this, &SGraphPin::GetDefaultValueVisibility)
-				.OptionsSource(&CurrentList)
+				.OptionsSource(CurrentList)
 				.OnGenerateWidget(this, &SControlRigGraphPinNameList::MakeNameListItemWidget)
 				.OnSelectionChanged(this, &SControlRigGraphPinNameList::OnNameListChanged)
 				.OnComboBoxOpening(this, &SControlRigGraphPinNameList::OnNameListComboBox)
@@ -49,13 +50,20 @@ TSharedRef<SWidget>	SControlRigGraphPinNameList::GetDefaultValueWidget()
 		];
 }
 
-const TArray<TSharedPtr<FString>>& SControlRigGraphPinNameList::GetNameList() const
+const TArray<TSharedPtr<FString>>* SControlRigGraphPinNameList::GetNameList() const
 {
+	const TArray<TSharedPtr<FString>>* Result = nullptr;
 	if (OnGetNameListContent.IsBound())
 	{
-		return OnGetNameListContent.Execute(ModelPin);
+		Result = OnGetNameListContent.Execute(ModelPin);
 	}
-	return EmptyList;
+
+	if(Result == nullptr)
+	{
+		Result = &EmptyList;
+	}
+
+	return Result;
 }
 
 FText SControlRigGraphPinNameList::GetNameListText() const
@@ -80,7 +88,7 @@ FSlateColor SControlRigGraphPinNameList::GetNameColor() const
 		FString CurrentItem = GetNameListText().ToString();
 		
 		bool bFound = false;
-		for (TSharedPtr<FString> Item : CurrentList)
+		for (TSharedPtr<FString> Item : (*CurrentList))
 		{
 			if (Item->Equals(CurrentItem))
 			{
@@ -119,7 +127,7 @@ void SControlRigGraphPinNameList::OnNameListComboBox()
 {
 	CurrentList = GetNameList();
 	TSharedPtr<FString> CurrentlySelected;
-	for (TSharedPtr<FString> Item : CurrentList)
+	for (TSharedPtr<FString> Item : (*CurrentList))
 	{
 		if (Item->Equals(GetNameListText().ToString()))
 		{
diff --git a/Engine/Plugins/Experimental/ControlRig/Source/ControlRigEditor/Private/Graph/SControlRigGraphPinNameList.h b/Engine/Plugins/Experimental/ControlRig/Source/ControlRigEditor/Private/Graph/SControlRigGraphPinNameList.h
--- a/Engine/Plugins/Experimental/ControlRig/Source/ControlRigEditor/Private/Graph/SControlRigGraphPinNameList.h
+++ b/Engine/Plugins/Experimental/ControlRig/Source/ControlRigEditor/Private/Graph/SControlRigGraphPinNameList.h
@@ -13,7 +13,7 @@ class SControlRigGraphPinNameList : public SGraphPin
 {
 public:
 
-	DECLARE_DELEGATE_RetVal_OneParam( const TArray<TSharedPtr<FString>>&, FOnGetNameListContent, URigVMPin*);
+	DECLARE_DELEGATE_RetVal_OneParam( const TArray<TSharedPtr<FString>>*, FOnGetNameListContent, URigVMPin*);
 
 	SLATE_BEGIN_ARGS(SControlRigGraphPinNameList)
 		: _MarkupInvalidItems(true)
@@ -32,7 +32,7 @@ protected:
 	virtual TSharedRef<SWidget>	GetDefaultValueWidget() override;
 	//~ End SGraphPin Interface
 
-	const TArray<TSharedPtr<FString>>& GetNameList() const;
+	const TArray<TSharedPtr<FString>>* GetNameList() const;
 	FText GetNameListText() const;
 	FSlateColor GetNameColor() const;
 	virtual void SetNameListText(const FText& NewTypeInValue, ETextCommit::Type CommitInfo);
@@ -45,6 +45,6 @@ protected:
 	URigVMPin* ModelPin;
 	TSharedPtr<SControlRigGraphPinNameListValueWidget> NameListComboBox;
 	TArray<TSharedPtr<FString>> EmptyList;
-	TArray<TSharedPtr<FString>> CurrentList;
+	const TArray<TSharedPtr<FString>>* CurrentList;
 	bool bMarkupInvalidItems;
 };
