Sequencer: Control Rig: Added a control rig class property to the control rig section. Changed the control rig property to private and added setter and getter functions for it so that the control rig class property can be set when the control rig is set

diff --git a/Engine/Plugins/Experimental/ControlRig/Source/ControlRig/Private/Sequencer/MovieSceneControlRigParameterSection.cpp b/Engine/Plugins/Experimental/ControlRig/Source/ControlRig/Private/Sequencer/MovieSceneControlRigParameterSection.cpp
--- a/Engine/Plugins/Experimental/ControlRig/Source/ControlRig/Private/Sequencer/MovieSceneControlRigParameterSection.cpp
+++ b/Engine/Plugins/Experimental/ControlRig/Source/ControlRig/Private/Sequencer/MovieSceneControlRigParameterSection.cpp
@@ -765,7 +765,7 @@ void UMovieSceneControlRigParameterSection::PostEditImport()
 	Super::PostEditImport();
 	if (UMovieSceneControlRigParameterTrack* Track = Cast< UMovieSceneControlRigParameterTrack>(GetOuter()))
 	{
-		ControlRig = Track->GetControlRig();
+		SetControlRig(Track->GetControlRig());
 	}
 	ReconstructChannelProxy(true);
 }
@@ -1702,7 +1702,7 @@ float UMovieSceneControlRigParameterSection::GetTotalWeightValue(FFrameTime InTi
 
 void UMovieSceneControlRigParameterSection::RecreateWithThisControlRig(UControlRig* InControlRig, bool bSetDefault)
 {
-	ControlRig = InControlRig;
+	SetControlRig(InControlRig);
 	/* Don't delete old tracks but eventually show that they aren't associated.. but
 	then how to delete?
 	BoolParameterNamesAndCurves.Empty();
@@ -1836,6 +1836,12 @@ void UMovieSceneControlRigParameterSection::RecreateWithThisControlRig(UControlR
 	ReconstructChannelProxy(true);
 }
 
+void UMovieSceneControlRigParameterSection::SetControlRig(UControlRig* InControlRig)
+{
+	ControlRig = InControlRig;
+	ControlRigClass = ControlRig->GetClass();
+}
+
 #if WITH_EDITOR
 
 void UMovieSceneControlRigParameterSection::RecordControlRigKey(FFrameNumber FrameNumber, bool bSetDefault, bool bAutoKey)
diff --git a/Engine/Plugins/Experimental/ControlRig/Source/ControlRig/Private/Sequencer/MovieSceneControlRigParameterTemplate.cpp b/Engine/Plugins/Experimental/ControlRig/Source/ControlRig/Private/Sequencer/MovieSceneControlRigParameterTemplate.cpp
--- a/Engine/Plugins/Experimental/ControlRig/Source/ControlRig/Private/Sequencer/MovieSceneControlRigParameterTemplate.cpp
+++ b/Engine/Plugins/Experimental/ControlRig/Source/ControlRig/Private/Sequencer/MovieSceneControlRigParameterTemplate.cpp
@@ -863,7 +863,7 @@ struct FControlRigParameterExecutionToken : IMovieSceneExecutionToken
 		MOVIESCENE_DETAILED_SCOPE_CYCLE_COUNTER(MovieSceneEval_ControlRigParameterTrack_TokenExecute)
 		
 		FMovieSceneSequenceID SequenceID = Operand.SequenceID;
-		UControlRig* ControlRig = Section->ControlRig;
+		UControlRig* ControlRig = Section->GetControlRig();
 
 		// Update the animation's state
 		
@@ -956,20 +956,20 @@ struct FControlRigParameterExecutionToken : IMovieSceneExecutionToken
 		bWasDoNotKey = Section->GetDoNotKey();
 		Section->SetDoNotKey(true);
 
-		if (Section->ControlRig)
+		if (Section->GetControlRig())
 		{
 			const bool bSetupUndo = false;
-			Section->ControlRig->SetAbsoluteTime((float)Context.GetFrameRate().AsSeconds(Context.GetTime()));
+			Section->GetControlRig()->SetAbsoluteTime((float)Context.GetFrameRate().AsSeconds(Context.GetTime()));
 			if (Section->bAdditive == false)
 			{
 				for (const FBoolParameterStringAndValue& BoolNameAndValue : BoolValues)
 				{
 					if (Section->ControlsToSet.Num() == 0 || Section->ControlsToSet.Contains(BoolNameAndValue.ParameterName))
 					{
-						FRigControlElement* RigControl = Section->ControlRig->FindControl(BoolNameAndValue.ParameterName);
+						FRigControlElement* RigControl = Section->GetControlRig()->FindControl(BoolNameAndValue.ParameterName);
 						if (RigControl && RigControl->Settings.ControlType == ERigControlType::Bool)
 						{
-							Section->ControlRig->SetControlValue<bool>(BoolNameAndValue.ParameterName, BoolNameAndValue.Value, true, EControlRigSetKey::Never,bSetupUndo);
+							Section->GetControlRig()->SetControlValue<bool>(BoolNameAndValue.ParameterName, BoolNameAndValue.Value, true, EControlRigSetKey::Never,bSetupUndo);
 						}
 					}
 				}
@@ -978,10 +978,10 @@ struct FControlRigParameterExecutionToken : IMovieSceneExecutionToken
 				{
 					if (Section->ControlsToSet.Num() == 0 || Section->ControlsToSet.Contains(IntegerNameAndValue.ParameterName))
 					{
-						FRigControlElement* RigControl = Section->ControlRig->FindControl(IntegerNameAndValue.ParameterName);
+						FRigControlElement* RigControl = Section->GetControlRig()->FindControl(IntegerNameAndValue.ParameterName);
 						if (RigControl && RigControl->Settings.ControlType == ERigControlType::Integer)
 						{
-							Section->ControlRig->SetControlValue<int32>(IntegerNameAndValue.ParameterName, IntegerNameAndValue.Value, true, EControlRigSetKey::Never,bSetupUndo);
+							Section->GetControlRig()->SetControlValue<int32>(IntegerNameAndValue.ParameterName, IntegerNameAndValue.Value, true, EControlRigSetKey::Never,bSetupUndo);
 						}
 					}
 				}
@@ -1022,12 +1022,12 @@ struct TControlRigParameterActuatorFloat : TMovieSceneBlendingActuator<FControlR
 	{
 		const UMovieSceneControlRigParameterSection* Section = SectionData.Get();
 
-		if (Section && Section->ControlRig)
+		if (Section && Section->GetControlRig())
 		{
-			FRigControlElement* ControlElement = Section->ControlRig->FindControl(ParameterName);
+			FRigControlElement* ControlElement = Section->GetControlRig()->FindControl(ParameterName);
 			if (ControlElement && ControlElement->Settings.ControlType == ERigControlType::Float)
 			{
-				const float Val = Section->ControlRig->GetHierarchy()->GetControlValue(ControlElement, ERigControlValueType::Current).Get<float>();
+				const float Val = Section->GetControlRig()->GetHierarchy()->GetControlValue(ControlElement, ERigControlValueType::Current).Get<float>();
 				return FControlRigTrackTokenFloat(Val);
 			}
 
@@ -1046,12 +1046,12 @@ struct TControlRigParameterActuatorFloat : TMovieSceneBlendingActuator<FControlR
 			bWasDoNotKey = Section->GetDoNotKey();
 			Section->SetDoNotKey(true);
 
-			if (Section->ControlRig && (Section->ControlsToSet.Num() == 0 || Section->ControlsToSet.Contains(ParameterName)))
+			if (Section->GetControlRig() && (Section->ControlsToSet.Num() == 0 || Section->ControlsToSet.Contains(ParameterName)))
 			{
-				FRigControlElement* ControlElement = Section->ControlRig->FindControl(ParameterName);
+				FRigControlElement* ControlElement = Section->GetControlRig()->FindControl(ParameterName);
 				if (ControlElement && ControlElement->Settings.ControlType == ERigControlType::Float)
 				{
-					Section->ControlRig->SetControlValue<float>(ParameterName, InFinalValue.Value, true, EControlRigSetKey::Never,bSetupUndo);
+					Section->GetControlRig()->SetControlValue<float>(ParameterName, InFinalValue.Value, true, EControlRigSetKey::Never,bSetupUndo);
 				}
 			}
 
@@ -1088,12 +1088,12 @@ struct TControlRigParameterActuatorVector2D : TMovieSceneBlendingActuator<FContr
 	{
 		const UMovieSceneControlRigParameterSection* Section = SectionData.Get();
 
-		if (Section && Section->ControlRig)
+		if (Section && Section->GetControlRig())
 		{
-			FRigControlElement* ControlElement = Section->ControlRig->FindControl(ParameterName);
+			FRigControlElement* ControlElement = Section->GetControlRig()->FindControl(ParameterName);
 			if (ControlElement && (ControlElement->Settings.ControlType == ERigControlType::Vector2D))
 			{
-				FVector2D Val = Section->ControlRig->GetHierarchy()->GetControlValue(ControlElement, ERigControlValueType::Current).Get<FVector2D>();
+				FVector2D Val = Section->GetControlRig()->GetHierarchy()->GetControlValue(ControlElement, ERigControlValueType::Current).Get<FVector2D>();
 				return FControlRigTrackTokenVector2D(Val);
 			}
 		}
@@ -1111,12 +1111,12 @@ struct TControlRigParameterActuatorVector2D : TMovieSceneBlendingActuator<FContr
 			bWasDoNotKey = Section->GetDoNotKey();
 			Section->SetDoNotKey(true);
 
-			if (Section->ControlRig && (Section->ControlsToSet.Num() == 0 || Section->ControlsToSet.Contains(ParameterName)))
+			if (Section->GetControlRig() && (Section->ControlsToSet.Num() == 0 || Section->ControlsToSet.Contains(ParameterName)))
 			{
-				FRigControlElement* ControlElement = Section->ControlRig->FindControl(ParameterName);
+				FRigControlElement* ControlElement = Section->GetControlRig()->FindControl(ParameterName);
 				if (ControlElement && (ControlElement->Settings.ControlType == ERigControlType::Vector2D))
 				{
-					Section->ControlRig->SetControlValue<FVector2D>(ParameterName, InFinalValue.Value, true, EControlRigSetKey::Never,bSetupUndo);
+					Section->GetControlRig()->SetControlValue<FVector2D>(ParameterName, InFinalValue.Value, true, EControlRigSetKey::Never,bSetupUndo);
 				}
 			}
 
@@ -1151,12 +1151,12 @@ struct TControlRigParameterActuatorVector : TMovieSceneBlendingActuator<FControl
 	{
 		const UMovieSceneControlRigParameterSection* Section = SectionData.Get();
 
-		if (Section->ControlRig)
+		if (Section->GetControlRig())
 		{
-			FRigControlElement* ControlElement = Section->ControlRig->FindControl(ParameterName);
+			FRigControlElement* ControlElement = Section->GetControlRig()->FindControl(ParameterName);
 			if (ControlElement && (ControlElement->Settings.ControlType == ERigControlType::Position || ControlElement->Settings.ControlType == ERigControlType::Scale || ControlElement->Settings.ControlType == ERigControlType::Rotator))
 			{
-				FVector Val = Section->ControlRig->GetHierarchy()->GetControlValue(ControlElement, ERigControlValueType::Current).Get<FVector>();
+				FVector Val = Section->GetControlRig()->GetHierarchy()->GetControlValue(ControlElement, ERigControlValueType::Current).Get<FVector>();
 				return FControlRigTrackTokenVector(Val);
 			}
 		}
@@ -1173,12 +1173,12 @@ struct TControlRigParameterActuatorVector : TMovieSceneBlendingActuator<FControl
 			const bool bSetupUndo = false;
 			bWasDoNotKey = Section->GetDoNotKey();
 			Section->SetDoNotKey(true);
-			if (Section->ControlRig && (Section->ControlsToSet.Num() == 0 || Section->ControlsToSet.Contains(ParameterName)))
+			if (Section->GetControlRig() && (Section->ControlsToSet.Num() == 0 || Section->ControlsToSet.Contains(ParameterName)))
 			{
-				FRigControlElement* ControlElement = Section->ControlRig->FindControl(ParameterName);
+				FRigControlElement* ControlElement = Section->GetControlRig()->FindControl(ParameterName);
 				if (ControlElement && (ControlElement->Settings.ControlType == ERigControlType::Position || ControlElement->Settings.ControlType == ERigControlType::Scale || ControlElement->Settings.ControlType == ERigControlType::Rotator))
 				{
-						Section->ControlRig->SetControlValue<FVector>(ParameterName, InFinalValue.Value, true, EControlRigSetKey::Never,bSetupUndo);
+						Section->GetControlRig()->SetControlValue<FVector>(ParameterName, InFinalValue.Value, true, EControlRigSetKey::Never,bSetupUndo);
 				}
 			}
 
@@ -1213,23 +1213,23 @@ struct TControlRigParameterActuatorTransform : TMovieSceneBlendingActuator<FCont
 	{
 		const UMovieSceneControlRigParameterSection* Section = SectionData.Get();
 
-		if (Section->ControlRig && (Section->ControlsToSet.Num() == 0 || Section->ControlsToSet.Contains(ParameterName)))
+		if (Section->GetControlRig() && (Section->ControlsToSet.Num() == 0 || Section->ControlsToSet.Contains(ParameterName)))
 		{
-			FRigControlElement* ControlElement = Section->ControlRig->FindControl(ParameterName);
+			FRigControlElement* ControlElement = Section->GetControlRig()->FindControl(ParameterName);
 			if (ControlElement && ControlElement->Settings.ControlType == ERigControlType::Transform)
 			{
-				FTransform Val = Section->ControlRig->GetHierarchy()->GetControlValue(ControlElement, ERigControlValueType::Current).Get<FTransform>();
+				FTransform Val = Section->GetControlRig()->GetHierarchy()->GetControlValue(ControlElement, ERigControlValueType::Current).Get<FTransform>();
 				return FControlRigTrackTokenTransform(Val);
 			}
 			else if(ControlElement && ControlElement->Settings.ControlType == ERigControlType::TransformNoScale)
 			{
-				FTransformNoScale ValNoScale = Section->ControlRig->GetHierarchy()->GetControlValue(ControlElement, ERigControlValueType::Current).Get<FTransformNoScale>();
+				FTransformNoScale ValNoScale = Section->GetControlRig()->GetHierarchy()->GetControlValue(ControlElement, ERigControlValueType::Current).Get<FTransformNoScale>();
 				FTransform Val = ValNoScale;
 				return FControlRigTrackTokenTransform(Val);
 			}
 			else if (ControlElement && ControlElement->Settings.ControlType == ERigControlType::EulerTransform)
 			{
-				FEulerTransform Euler = Section->ControlRig->GetHierarchy()->GetControlValue(ControlElement, ERigControlValueType::Current).Get<FEulerTransform>();
+				FEulerTransform Euler = Section->GetControlRig()->GetHierarchy()->GetControlValue(ControlElement, ERigControlValueType::Current).Get<FEulerTransform>();
 				FTransform Val = Euler.ToFTransform();
 				return FControlRigTrackTokenTransform(Val);
 		}
@@ -1255,22 +1255,22 @@ struct TControlRigParameterActuatorTransform : TMovieSceneBlendingActuator<FCont
 			bWasDoNotKey = Section->GetDoNotKey();
 			Section->SetDoNotKey(true);
 			const bool bSetupUndo = false;
-			if (Section->ControlRig && (Section->ControlsToSet.Num() == 0 || Section->ControlsToSet.Contains(ParameterName)))
+			if (Section->GetControlRig() && (Section->ControlsToSet.Num() == 0 || Section->ControlsToSet.Contains(ParameterName)))
 			{
-				FRigControlElement* ControlElement = Section->ControlRig->FindControl(ParameterName);
+				FRigControlElement* ControlElement = Section->GetControlRig()->FindControl(ParameterName);
 				if (ControlElement && ControlElement->Settings.ControlType == ERigControlType::Transform)
 				{
-						Section->ControlRig->SetControlValue<FTransform>(ParameterName, InFinalValue.Value,true, EControlRigSetKey::Never,bSetupUndo);
+						Section->GetControlRig()->SetControlValue<FTransform>(ParameterName, InFinalValue.Value,true, EControlRigSetKey::Never,bSetupUndo);
 				}
 				else if (ControlElement && ControlElement->Settings.ControlType == ERigControlType::TransformNoScale)
 				{
 					FTransformNoScale NoScale = InFinalValue.Value;
-						Section->ControlRig->SetControlValue<FTransformNoScale>(ParameterName, NoScale, true, EControlRigSetKey::Never,bSetupUndo);
+						Section->GetControlRig()->SetControlValue<FTransformNoScale>(ParameterName, NoScale, true, EControlRigSetKey::Never,bSetupUndo);
 				}
 				else if (ControlElement && ControlElement->Settings.ControlType == ERigControlType::EulerTransform)
 				{
 					FEulerTransform Euler = InFinalValue.Value;
-					Section->ControlRig->SetControlValue<FEulerTransform>(ParameterName, Euler, true, EControlRigSetKey::Never,bSetupUndo);
+					Section->GetControlRig()->SetControlValue<FEulerTransform>(ParameterName, Euler, true, EControlRigSetKey::Never,bSetupUndo);
 				}
 			}
 
@@ -1295,7 +1295,7 @@ void FMovieSceneControlRigParameterTemplate::Evaluate(const FMovieSceneEvaluatio
 	const FFrameTime Time = Context.GetTime();
 
 	const UMovieSceneControlRigParameterSection* Section = Cast<UMovieSceneControlRigParameterSection>(GetSourceSection());
-	if (Section && Section->ControlRig)
+	if (Section && Section->GetControlRig())
 	{
 		FEvaluatedControlRigParameterSectionChannelMasks* ChannelMasks = PersistentData.FindSectionData<FEvaluatedControlRigParameterSectionChannelMasks>();
 		if (!ChannelMasks)
@@ -1332,7 +1332,7 @@ void FMovieSceneControlRigParameterTemplate::Evaluate(const FMovieSceneEvaluatio
 		FControlRigParameterExecutionToken ExecutionToken(Section,Values);
 		ExecutionTokens.Add(MoveTemp(ExecutionToken));
 
-		FControlRigAnimTypeIDsPtr TypeIDs = FControlRigAnimTypeIDs::Get(Section->ControlRig);
+		FControlRigAnimTypeIDsPtr TypeIDs = FControlRigAnimTypeIDs::Get(Section->GetControlRig());
 
 		for (const FScalarParameterStringAndValue& ScalarNameAndValue : Values.ScalarValues)
 		{
@@ -1589,7 +1589,7 @@ void FMovieSceneControlRigParameterTemplate::Interrogate(const FMovieSceneContex
 	const FFrameTime Time = Context.GetTime();
 
 	const UMovieSceneControlRigParameterSection* Section = Cast<UMovieSceneControlRigParameterSection>(GetSourceSection());
-	if (Section && Section->ControlRig)
+	if (Section && Section->GetControlRig())
 	{
 		FEvaluatedControlRigParameterSectionChannelMasks ChannelMasks;
 		ChannelMasks.Initialize(Section, Scalars, Bools, Integers, Enums, Vector2Ds, Vectors, Colors, Transforms);
@@ -1601,7 +1601,7 @@ void FMovieSceneControlRigParameterTemplate::Interrogate(const FMovieSceneContex
 		EvaluateCurvesWithMasks(Context, Values);
 		HACK_ChannelMasks = nullptr;
 
-		FControlRigAnimTypeIDsPtr TypeIDs = FControlRigAnimTypeIDs::Get(Section->ControlRig);
+		FControlRigAnimTypeIDsPtr TypeIDs = FControlRigAnimTypeIDs::Get(Section->GetControlRig());
 
 		float Weight = 1.f;
 
diff --git a/Engine/Plugins/Experimental/ControlRig/Source/ControlRig/Private/Sequencer/MovieSceneControlRigParameterTrack.cpp b/Engine/Plugins/Experimental/ControlRig/Source/ControlRig/Private/Sequencer/MovieSceneControlRigParameterTrack.cpp
--- a/Engine/Plugins/Experimental/ControlRig/Source/ControlRig/Private/Sequencer/MovieSceneControlRigParameterTrack.cpp
+++ b/Engine/Plugins/Experimental/ControlRig/Source/ControlRig/Private/Sequencer/MovieSceneControlRigParameterTrack.cpp
@@ -37,7 +37,7 @@ bool UMovieSceneControlRigParameterTrack::SupportsType(TSubclassOf<UMovieSceneSe
 UMovieSceneSection* UMovieSceneControlRigParameterTrack::CreateNewSection()
 {
 	UMovieSceneControlRigParameterSection* NewSection = NewObject<UMovieSceneControlRigParameterSection>(this, NAME_None, RF_Transactional);
-	NewSection->ControlRig = ControlRig;
+	NewSection->SetControlRig(ControlRig);
 	bool bSetDefault = false;
 	if (Sections.Num() == 0)
 	{
@@ -71,9 +71,9 @@ void UMovieSceneControlRigParameterTrack::AddSection(UMovieSceneSection& Section
 	Sections.Add(&Section);
 	if (UMovieSceneControlRigParameterSection* CRSection = Cast<UMovieSceneControlRigParameterSection>(&Section))
 	{
-		if (CRSection->ControlRig != ControlRig)
+		if (CRSection->GetControlRig() != ControlRig)
 		{
-			CRSection->ControlRig = ControlRig;
+			CRSection->SetControlRig(ControlRig);
 		}
 		CRSection->ReconstructChannelProxy(true);
 	}
@@ -372,7 +372,7 @@ CONTROLRIG_API void UMovieSceneControlRigParameterTrack::ReplaceControlRig(UCont
 			}
 			else
 			{
-				CRSection->ControlRig = NewControlRig;
+				CRSection->SetControlRig(NewControlRig);
 			}
 		}	
 	}
diff --git a/Engine/Plugins/Experimental/ControlRig/Source/ControlRig/Public/Sequencer/MovieSceneControlRigParameterSection.h b/Engine/Plugins/Experimental/ControlRig/Source/ControlRig/Public/Sequencer/MovieSceneControlRigParameterSection.h
--- a/Engine/Plugins/Experimental/ControlRig/Source/ControlRig/Public/Sequencer/MovieSceneControlRigParameterSection.h
+++ b/Engine/Plugins/Experimental/ControlRig/Source/ControlRig/Public/Sequencer/MovieSceneControlRigParameterSection.h
@@ -158,12 +158,19 @@ public:
 
 	TArray<FIntegerParameterNameAndCurve>& GetIntegerParameterNamesAndCurves();
 	const TArray<FIntegerParameterNameAndCurve>& GetIntegerParameterNamesAndCurves() const;
-public:
+
+private:
 
 	/** Control Rig that controls us*/
 	UPROPERTY()
 	UControlRig* ControlRig;
 
+public:
+
+	/** The class of control rig to instantiate */
+	UPROPERTY(EditAnywhere, Category = "Animation")
+	TSubclassOf<UControlRig> ControlRigClass;
+
 	/** Mask for controls themselves*/
 	UPROPERTY()
 	TArray<bool> ControlsMask;
@@ -274,9 +281,15 @@ public:
 	}
 
 public:
+
 	/** Recreate with this Control Rig*/
 	void RecreateWithThisControlRig(UControlRig* InControlRig, bool bSetDefault);
 
+	/* Set the control rig for this section */
+	void SetControlRig(UControlRig* InControlRig);
+	/* Get the control rig for this section */
+	UControlRig* GetControlRig() const { return ControlRig; }
+
 	/** Whether or not to key currently, maybe evaluating so don't*/
 	void  SetDoNotKey(bool bIn) const { bDoNotKey = bIn; }
 	/** Get Whether to key or not*/
diff --git a/Engine/Plugins/Experimental/ControlRig/Source/ControlRigEditor/Private/Sequencer/ControlRigParameterTrackEditor.cpp b/Engine/Plugins/Experimental/ControlRig/Source/ControlRigEditor/Private/Sequencer/ControlRigParameterTrackEditor.cpp
--- a/Engine/Plugins/Experimental/ControlRig/Source/ControlRigEditor/Private/Sequencer/ControlRigParameterTrackEditor.cpp
+++ b/Engine/Plugins/Experimental/ControlRig/Source/ControlRigEditor/Private/Sequencer/ControlRigParameterTrackEditor.cpp
@@ -1406,7 +1406,7 @@ void FControlRigParameterTrackEditor::OnAddTransformKeysForSelectedObjects(EMovi
 void FControlRigParameterTrackEditor::OnChannelChanged(const FMovieSceneChannelMetaData* MetaData, UMovieSceneSection* InSection)
 {
 	UMovieSceneControlRigParameterSection* Section = Cast<UMovieSceneControlRigParameterSection>(InSection);
-	if (Section && Section->ControlRig && MetaData)
+	if (Section && Section->GetControlRig() && MetaData)
 	{
 		Section->ControlsToSet.Empty();
 		TArray<FString> StringArray;
@@ -1416,9 +1416,9 @@ void FControlRigParameterTrackEditor::OnChannelChanged(const FMovieSceneChannelM
 		{
 			FName ControlName(*StringArray[0]);
 			Section->ControlsToSet.Add(ControlName);
-			FControlRigInteractionScope InteractionScope(Section->ControlRig);
+			FControlRigInteractionScope InteractionScope(Section->GetControlRig());
 			GetSequencer()->ForceEvaluate(); //now run sequencer...
-			Section->ControlRig->Evaluate_AnyThread();
+			Section->GetControlRig()->Evaluate_AnyThread();
 			Section->ControlsToSet.Empty();
 		}
 	}
@@ -1548,7 +1548,7 @@ void FControlRigParameterTrackEditor::OnCurveDisplayChanged(FCurveModel* CurveMo
 		UMovieSceneControlRigParameterSection* MovieSection = Cast<UMovieSceneControlRigParameterSection>(CurveModel->GetOwningObject());
 		if (MovieSection)
 		{
-			ControlRig = MovieSection->ControlRig;
+			ControlRig = MovieSection->GetControlRig();
 			//Only create the edit mode if we have a  curve selected and it's not set and we have some boundobjects.
 			if (!ControlRigEditMode)
 			{
@@ -1754,7 +1754,7 @@ void FControlRigParameterTrackEditor::OnSelectionChanged(TArray<UMovieSceneTrack
 		UMovieSceneControlRigParameterSection* MovieSection = Cast<UMovieSceneControlRigParameterSection>(KeyArea->GetOwningSection());
 		if (MovieSection)
 		{
-			ControlRig = MovieSection->ControlRig;
+			ControlRig = MovieSection->GetControlRig();
 			//Only create the edit mode if we have a KeyAra selected and it's not set and we have some boundobjects.
 			if (!ControlRigEditMode)
 			{
@@ -2643,7 +2643,7 @@ void FControlRigParameterTrackEditor::ToggleFKControlRig(UMovieSceneControlRigPa
 			{
 				Section->Modify();
 				CRSection->ClearAllParameters();
-				CRSection->RecreateWithThisControlRig(CRSection->ControlRig, true);
+				CRSection->RecreateWithThisControlRig(CRSection->GetControlRig(), true);
 			}
 		}
 	}
@@ -2964,7 +2964,7 @@ void FControlRigParameterSection::BuildSectionContextMenu(FMenuBuilder& MenuBuil
 	UMovieSceneControlRigParameterSection* ParameterSection = CastChecked<UMovieSceneControlRigParameterSection>(WeakSection.Get());
 	TSharedPtr<ISequencer> SequencerPtr = WeakSequencer.Pin();
 
-	UControlRig* ControlRig = ParameterSection->ControlRig;
+	UControlRig* ControlRig = ParameterSection->GetControlRig();
 
 	if (ControlRig)
 	{
diff --git a/Engine/Plugins/Experimental/MotionTrailEditorMode/Source/MotionTrailEditorMode/Private/Sequencer/MovieSceneTransformTrail.cpp b/Engine/Plugins/Experimental/MotionTrailEditorMode/Source/MotionTrailEditorMode/Private/Sequencer/MovieSceneTransformTrail.cpp
--- a/Engine/Plugins/Experimental/MotionTrailEditorMode/Source/MotionTrailEditorMode/Private/Sequencer/MovieSceneTransformTrail.cpp
+++ b/Engine/Plugins/Experimental/MotionTrailEditorMode/Source/MotionTrailEditorMode/Private/Sequencer/MovieSceneTransformTrail.cpp
@@ -229,7 +229,7 @@ void FMovieSceneControlTransformTrail::UpdateCacheTimes(const FTrailEvaluateTime
 	check(Section);
 
 	const FRigElementKey ControlKey(ControlName, ERigElementType::Control);
-	const FTransform InitialTransform = Section->ControlRig->GetHierarchy()->GetControlValue(ControlKey, ERigControlValueType::Initial).Get<FTransform>();
+	const FTransform InitialTransform = Section->GetControlRig()->GetHierarchy()->GetControlValue(ControlKey, ERigControlValueType::Initial).Get<FTransform>();
 
 	TArrayView<FMovieSceneFloatChannel*> FloatChannels = Section->GetChannelProxy().GetChannels<FMovieSceneFloatChannel>();
 	FloatChannels = FloatChannels.Slice(GetChannelOffset(), uint8(EMSTrailTransformChannel::MaxChannel) + 1);
@@ -249,7 +249,7 @@ UE::MovieScene::FIntermediate3DTransform FMovieSceneControlTransformTrail::Calcu
 	UMovieSceneControlRigParameterSection* Section = Cast<UMovieSceneControlRigParameterSection>(GetSection());
 
 	const FRigElementKey ControlKey(ControlName, ERigElementType::Control);
-	const FTransform InitialTransform = Section->ControlRig->GetHierarchy()->GetControlValue(ControlKey, ERigControlValueType::Initial).Get<FTransform>();
+	const FTransform InitialTransform = Section->GetControlRig()->GetHierarchy()->GetControlValue(ControlKey, ERigControlValueType::Initial).Get<FTransform>();
 
 	const UE::MovieScene::FIntermediate3DTransform Delta = UE::MovieScene::FIntermediate3DTransform(
 		Current.GetTranslation() - Start.GetTranslation(),
diff --git a/Engine/Plugins/Experimental/MotionTrailEditorMode/Source/MotionTrailEditorMode/Private/Sequencer/SequencerTrailHierarchy.cpp b/Engine/Plugins/Experimental/MotionTrailEditorMode/Source/MotionTrailEditorMode/Private/Sequencer/SequencerTrailHierarchy.cpp
--- a/Engine/Plugins/Experimental/MotionTrailEditorMode/Source/MotionTrailEditorMode/Private/Sequencer/SequencerTrailHierarchy.cpp
+++ b/Engine/Plugins/Experimental/MotionTrailEditorMode/Source/MotionTrailEditorMode/Private/Sequencer/SequencerTrailHierarchy.cpp
@@ -131,7 +131,7 @@ void FSequencerTrailHierarchy::Destroy()
 	for (const TPair<UMovieSceneSection*, FControlRigDelegateHandles>& SectionHandlesPair : ControlRigDelegateHandles)
 	{
 		UMovieSceneControlRigParameterSection* Section = Cast<UMovieSceneControlRigParameterSection>(SectionHandlesPair.Key);
-		URigHierarchy* RigHierarchy = Section->ControlRig->GetHierarchy();
+		URigHierarchy* RigHierarchy = Section->GetControlRig()->GetHierarchy();
 		RigHierarchy->OnModified().Remove(SectionHandlesPair.Value.OnHierarchyModified);
 	}
 
@@ -684,7 +684,7 @@ void FSequencerTrailHierarchy::AddControlsToHierarchy(class USkeletalMeshCompone
 	UMovieSceneControlRigParameterSection* CRParamSection = Cast<UMovieSceneControlRigParameterSection>(Sections[0]);
 	check(CRParamSection);
 
-	URigHierarchy* RigHierarchy = CRParamSection->ControlRig->GetHierarchy();
+	URigHierarchy* RigHierarchy = CRParamSection->GetControlRig()->GetHierarchy();
 	if (!ControlRigDelegateHandles.Contains(CRParamSection))
 	{
 		RegisterControlRigDelegates(CompToAdd, CRParamSection);
@@ -693,7 +693,7 @@ void FSequencerTrailHierarchy::AddControlsToHierarchy(class USkeletalMeshCompone
 	CRParamSection->ReconstructChannelProxy(true);
 
 	TArray<FRigControlElement*> SortedControls;
-	CRParamSection->ControlRig->GetControlsInOrder(SortedControls);
+	CRParamSection->GetControlRig()->GetControlsInOrder(SortedControls);
 	
 	for (const TPair<FName, FChannelMapInfo>& NameInfoPair : CRParamSection->ControlChannelMap)
 	{
@@ -722,7 +722,7 @@ void FSequencerTrailHierarchy::AddControlsToHierarchy(class USkeletalMeshCompone
 
 void FSequencerTrailHierarchy::RegisterControlRigDelegates(USkeletalMeshComponent* Component, class UMovieSceneControlRigParameterSection* CRParamSection)
 {
-	URigHierarchy* RigHierarchy = CRParamSection->ControlRig->GetHierarchy();
+	URigHierarchy* RigHierarchy = CRParamSection->GetControlRig()->GetHierarchy();
 	FControlRigDelegateHandles& DelegateHandles = ControlRigDelegateHandles.Add(CRParamSection);
 	DelegateHandles.OnHierarchyModified = RigHierarchy->OnModified().AddLambda(
 		[this, RigHierarchy, Component, CRParamSection](ERigHierarchyNotification InNotif, URigHierarchy* InHierarchy, const FRigBaseElement* InElement) {
