Sequencer: Control Rig: Remove unused/obsolete

diff --git a/Engine/Plugins/Experimental/ControlRig/Source/ControlRig/Private/Sequencer/MovieSceneControlRigParameterTemplate.cpp b/Engine/Plugins/Experimental/ControlRig/Source/ControlRig/Private/Sequencer/MovieSceneControlRigParameterTemplate.cpp
--- a/Engine/Plugins/Experimental/ControlRig/Source/ControlRig/Private/Sequencer/MovieSceneControlRigParameterTemplate.cpp
+++ b/Engine/Plugins/Experimental/ControlRig/Source/ControlRig/Private/Sequencer/MovieSceneControlRigParameterTemplate.cpp
@@ -960,29 +960,26 @@ struct FControlRigParameterExecutionToken : IMovieSceneExecutionToken
 		{
 			const bool bSetupUndo = false;
 			Section->GetControlRig()->SetAbsoluteTime((float)Context.GetFrameRate().AsSeconds(Context.GetTime()));
-			if (Section->bAdditive == false)
+			for (const FBoolParameterStringAndValue& BoolNameAndValue : BoolValues)
 			{
-				for (const FBoolParameterStringAndValue& BoolNameAndValue : BoolValues)
+				if (Section->ControlsToSet.Num() == 0 || Section->ControlsToSet.Contains(BoolNameAndValue.ParameterName))
 				{
-					if (Section->ControlsToSet.Num() == 0 || Section->ControlsToSet.Contains(BoolNameAndValue.ParameterName))
+					FRigControlElement* RigControl = Section->GetControlRig()->FindControl(BoolNameAndValue.ParameterName);
+					if (RigControl && RigControl->Settings.ControlType == ERigControlType::Bool)
 					{
-						FRigControlElement* RigControl = Section->GetControlRig()->FindControl(BoolNameAndValue.ParameterName);
-						if (RigControl && RigControl->Settings.ControlType == ERigControlType::Bool)
-						{
-							Section->GetControlRig()->SetControlValue<bool>(BoolNameAndValue.ParameterName, BoolNameAndValue.Value, true, EControlRigSetKey::Never,bSetupUndo);
-						}
+						Section->GetControlRig()->SetControlValue<bool>(BoolNameAndValue.ParameterName, BoolNameAndValue.Value, true, EControlRigSetKey::Never,bSetupUndo);
 					}
 				}
+			}
 
-				for (const FIntegerParameterStringAndValue& IntegerNameAndValue : IntegerValues)
+			for (const FIntegerParameterStringAndValue& IntegerNameAndValue : IntegerValues)
+			{
+				if (Section->ControlsToSet.Num() == 0 || Section->ControlsToSet.Contains(IntegerNameAndValue.ParameterName))
 				{
-					if (Section->ControlsToSet.Num() == 0 || Section->ControlsToSet.Contains(IntegerNameAndValue.ParameterName))
+					FRigControlElement* RigControl = Section->GetControlRig()->FindControl(IntegerNameAndValue.ParameterName);
+					if (RigControl && RigControl->Settings.ControlType == ERigControlType::Integer)
 					{
-						FRigControlElement* RigControl = Section->GetControlRig()->FindControl(IntegerNameAndValue.ParameterName);
-						if (RigControl && RigControl->Settings.ControlType == ERigControlType::Integer)
-						{
-							Section->GetControlRig()->SetControlValue<int32>(IntegerNameAndValue.ParameterName, IntegerNameAndValue.Value, true, EControlRigSetKey::Never,bSetupUndo);
-						}
+						Section->GetControlRig()->SetControlValue<int32>(IntegerNameAndValue.ParameterName, IntegerNameAndValue.Value, true, EControlRigSetKey::Never,bSetupUndo);
 					}
 				}
 			}
diff --git a/Engine/Plugins/Experimental/ControlRig/Source/ControlRig/Public/Sequencer/MovieSceneControlRigParameterSection.h b/Engine/Plugins/Experimental/ControlRig/Source/ControlRig/Public/Sequencer/MovieSceneControlRigParameterSection.h
--- a/Engine/Plugins/Experimental/ControlRig/Source/ControlRig/Public/Sequencer/MovieSceneControlRigParameterSection.h
+++ b/Engine/Plugins/Experimental/ControlRig/Source/ControlRig/Public/Sequencer/MovieSceneControlRigParameterSection.h
@@ -179,18 +179,6 @@ public:
 	UPROPERTY()
 	FMovieSceneTransformMask TransformMask;
 
-	/** Blend this track in additively (using the reference pose as a base) */
-	UPROPERTY(EditAnywhere, Category = "Animation")
-	bool bAdditive;
-
-	/** Only apply bones that are in the filter */
-	UPROPERTY(EditAnywhere, Category = "Animation")
-	bool bApplyBoneFilter;
-
-	/** Per-bone filter to apply to our animation */
-	UPROPERTY(EditAnywhere, Category = "Animation", meta=(EditCondition=bApplyBoneFilter))
-	FInputBlendPose BoneFilter;
-
 	/** The weight curve for this animation controller section */
 	UPROPERTY()
 	FMovieSceneFloatChannel Weight;
