Control Rig: Avoid bulk notifications during CopyHierarchy

diff --git a/Engine/Plugins/Experimental/ControlRig/Source/ControlRig/Private/ControlRig.cpp b/Engine/Plugins/Experimental/ControlRig/Source/ControlRig/Private/ControlRig.cpp
--- a/Engine/Plugins/Experimental/ControlRig/Source/ControlRig/Private/ControlRig.cpp
+++ b/Engine/Plugins/Experimental/ControlRig/Source/ControlRig/Private/ControlRig.cpp
@@ -173,8 +173,14 @@ void UControlRig::InitializeFromCDO()
 		UControlRig* CDO = GetClass()->GetDefaultObject<UControlRig>();
 
 		// copy hierarchy
-		GetHierarchy()->CopyHierarchy(CDO->GetHierarchy());
-		GetHierarchy()->ResetPoseToInitial(ERigElementType::All);
+		{
+			TGuardValue<bool> Guard(GetHierarchy()->GetSuspendNotificationsFlag(), true);
+			GetHierarchy()->CopyHierarchy(CDO->GetHierarchy());
+			GetHierarchy()->ResetPoseToInitial(ERigElementType::All);
+		}
+
+		// notify clients that the hierarchy has changed
+		GetHierarchy()->Notify(ERigHierarchyNotification::HierarchyReset, nullptr);
 
 		// copy draw container
 		DrawContainer = CDO->DrawContainer;
diff --git a/Engine/Plugins/Experimental/ControlRig/Source/ControlRigDeveloper/Private/ControlRigBlueprintCompiler.cpp b/Engine/Plugins/Experimental/ControlRig/Source/ControlRigDeveloper/Private/ControlRigBlueprintCompiler.cpp
--- a/Engine/Plugins/Experimental/ControlRig/Source/ControlRigDeveloper/Private/ControlRigBlueprintCompiler.cpp
+++ b/Engine/Plugins/Experimental/ControlRig/Source/ControlRigDeveloper/Private/ControlRigBlueprintCompiler.cpp
@@ -91,7 +91,17 @@ void FControlRigBlueprintCompilerContext::CopyTermDefaultsToDefaultObject(UObjec
 	if (ControlRigBlueprint)
 	{
 		UControlRig* ControlRig = CastChecked<UControlRig>(DefaultObject);
-		ControlRig->GetHierarchy()->CopyHierarchy(ControlRigBlueprint->Hierarchy);
+		
+		// copy hierarchy
+		{
+			TGuardValue<bool> Guard(ControlRig->GetHierarchy()->GetSuspendNotificationsFlag(), false);
+			ControlRig->GetHierarchy()->CopyHierarchy(ControlRigBlueprint->Hierarchy);
+			ControlRig->GetHierarchy()->ResetPoseToInitial(ERigElementType::All);
+		}
+
+		// notify clients that the hierarchy has changed
+		ControlRig->GetHierarchy()->Notify(ERigHierarchyNotification::HierarchyReset, nullptr);
+
 		ControlRig->DrawContainer = ControlRigBlueprint->DrawContainer;
 		ControlRig->Influences = ControlRigBlueprint->Influences;
 	}
