Control Rig: Fix Constraint Nodes's Axis Filter, for position constraints it now uses inverse parent matrix instead of vector substraction for change of space.

diff --git a/Engine/Plugins/Experimental/ControlRig/Source/ControlRig/Private/Units/Highlevel/Hierarchy/RigUnit_TransformConstraint.cpp b/Engine/Plugins/Experimental/ControlRig/Source/ControlRig/Private/Units/Highlevel/Hierarchy/RigUnit_TransformConstraint.cpp
--- a/Engine/Plugins/Experimental/ControlRig/Source/ControlRig/Private/Units/Highlevel/Hierarchy/RigUnit_TransformConstraint.cpp
+++ b/Engine/Plugins/Experimental/ControlRig/Source/ControlRig/Private/Units/Highlevel/Hierarchy/RigUnit_TransformConstraint.cpp
@@ -474,10 +474,9 @@ FRigUnit_PositionConstraint_Execute()
 
 			// handle filtering, performed in local space
 			FTransform ChildParentGlobalTransform = Hierarchy->GetParentTransform(Child, false);
-			FVector MixedPosition = MixedGlobalPosition - ChildParentGlobalTransform.GetLocation();
+			FVector MixedPosition = ChildParentGlobalTransform.Inverse().TransformVector(MixedGlobalPosition);
 			
 			FTransform ChildCurrentLocalTransform = Hierarchy->GetLocalTransform(Child);
-			FVector ChildPosition = ChildCurrentLocalTransform.GetTranslation();
 			
 			// Controls need to be handled a bit differently
 			FTransform AdditionalOffsetTransform = FTransform::Identity;
@@ -487,24 +486,32 @@ FRigUnit_PositionConstraint_Execute()
 				if (FRigControlElement* ChildAsControlElement = Hierarchy->Find<FRigControlElement>(Child))
 				{
 					AdditionalOffsetTransform = Hierarchy->GetControlOffsetTransform(ChildAsControlElement, ERigTransformType::CurrentLocal);
-					// Control's local(parent) space position = local + offset
-					ChildPosition += AdditionalOffsetTransform.GetLocation();
+					// Control's local(parent) space position = local * offset
+					ChildCurrentLocalTransform *= AdditionalOffsetTransform;
 				}
 			}
 			
+			FVector ChildPosition = ChildCurrentLocalTransform.GetTranslation();
+
 			FVector FilteredPosition;
 			FilteredPosition.X = Filter.bX ? MixedPosition.X : ChildPosition.X;
 			FilteredPosition.Y = Filter.bY ? MixedPosition.Y : ChildPosition.Y;
 			FilteredPosition.Z = Filter.bZ ? MixedPosition.Z : ChildPosition.Z;
 
-			FVector FinalPosition = FilteredPosition;
+			FTransform FilteredMixedLocalTransform = ChildCurrentLocalTransform;
+
+			FilteredMixedLocalTransform.SetTranslation(FilteredPosition); 
+
+			FTransform FinalLocalTransform = FilteredMixedLocalTransform;
+
 			if (Child.Type == ERigElementType::Control)
 			{
 				// need to convert back to offset space for the actual control value
-				FinalPosition = FilteredPosition - AdditionalOffsetTransform.GetLocation();
+				FinalLocalTransform = FilteredMixedLocalTransform.GetRelativeTransform(AdditionalOffsetTransform);
+				FinalLocalTransform.NormalizeRotation();
 			}
 
-			Hierarchy->SetLocalTransform(Child, FTransform(ChildCurrentLocalTransform.GetRotation(), FinalPosition, ChildCurrentLocalTransform.GetScale3D()));
+			Hierarchy->SetLocalTransform(Child, FinalLocalTransform);
 		}
 	}
 }
@@ -657,38 +664,43 @@ FRigUnit_RotationConstraint_Execute()
 			FVector MixedEulerRotation = FControlRigMathLibrary::EulerFromQuat(MixedLocalRotation, AdvancedSettings.RotationOrderForFilter);
 			
 			FTransform ChildCurrentLocalTransform = Hierarchy->GetLocalTransform(Child, false);
-			FQuat ChildRotation = ChildCurrentLocalTransform.GetRotation();
 			
 			// Controls need to be handled a bit differently
 			FTransform AdditionalOffsetTransform = FTransform::Identity;
-			
+
 			if (Child.Type == ERigElementType::Control)
 			{
 				if (FRigControlElement* ChildAsControlElement = Hierarchy->Find<FRigControlElement>(Child))
 				{
 					AdditionalOffsetTransform = Hierarchy->GetControlOffsetTransform(ChildAsControlElement, ERigTransformType::CurrentLocal);
-					// Control's local(parent) space transform = local * offset
-					ChildRotation = AdditionalOffsetTransform.GetRotation() * ChildRotation;
+					// Control's local(parent) space = local * offset
+					ChildCurrentLocalTransform *= AdditionalOffsetTransform;
 				}
 			}
+
+			FQuat ChildRotation = ChildCurrentLocalTransform.GetRotation();
 			
 			FVector ChildEulerRotation = FControlRigMathLibrary::EulerFromQuat(ChildRotation, AdvancedSettings.RotationOrderForFilter);	
 			
 			FVector FilteredEulerRotation;
 			FilteredEulerRotation.X = Filter.bX ? MixedEulerRotation.X : ChildEulerRotation.X;
 			FilteredEulerRotation.Y = Filter.bY ? MixedEulerRotation.Y : ChildEulerRotation.Y;
-			FilteredEulerRotation.Z = Filter.bZ ? MixedEulerRotation.Z : ChildEulerRotation.Z;
-			
-			FQuat FilterRotation = FControlRigMathLibrary::QuatFromEuler(FilteredEulerRotation, AdvancedSettings.RotationOrderForFilter); 
+			FilteredEulerRotation.Z = Filter.bZ ? MixedEulerRotation.Z : ChildEulerRotation.Z; 
+
+			FTransform FilteredMixedLocalTransform = ChildCurrentLocalTransform;
+
+			FilteredMixedLocalTransform.SetRotation(FControlRigMathLibrary::QuatFromEuler(FilteredEulerRotation, AdvancedSettings.RotationOrderForFilter));
+
+			FTransform FinalLocalTransform = FilteredMixedLocalTransform;
 
-			FQuat FinalRotation = FilterRotation;
 			if (Child.Type == ERigElementType::Control)
 			{
 				// need to convert back to offset space for the actual control value
-				FinalRotation = AdditionalOffsetTransform.GetRotation().Inverse() * FilterRotation;
-			}	
-			
-			Hierarchy->SetLocalTransform(Child, FTransform(FinalRotation, ChildCurrentLocalTransform.GetTranslation(), ChildCurrentLocalTransform.GetScale3D()));
+				FinalLocalTransform = FilteredMixedLocalTransform.GetRelativeTransform(AdditionalOffsetTransform);
+				FinalLocalTransform.NormalizeRotation();
+			}
+
+			Hierarchy->SetLocalTransform(Child, FinalLocalTransform); 
 		}
 	} 	
 }
@@ -873,34 +885,41 @@ FRigUnit_ScaleConstraint_Execute()
 			FVector MixedLocalScale = MixedGlobalScale / GetNonZeroScale(ChildParentGlobalScale);
 			
 			FTransform ChildCurrentLocalTransform = Hierarchy->GetLocalTransform(Child, false);
-			FVector ChildLocalScale = ChildCurrentLocalTransform.GetScale3D();
-			
+
 			// Controls need to be handled a bit differently
 			FTransform AdditionalOffsetTransform = FTransform::Identity;
-			
+
 			if (Child.Type == ERigElementType::Control)
 			{
 				if (FRigControlElement* ChildAsControlElement = Hierarchy->Find<FRigControlElement>(Child))
 				{
 					AdditionalOffsetTransform = Hierarchy->GetControlOffsetTransform(ChildAsControlElement, ERigTransformType::CurrentLocal);
-					// Control's local(parent) space transform = local * offset
-					ChildLocalScale *= AdditionalOffsetTransform.GetScale3D();
+					// Control's local(parent) space = local * offset
+					ChildCurrentLocalTransform *= AdditionalOffsetTransform;
 				}
 			}
 			
+			FVector ChildLocalScale = ChildCurrentLocalTransform.GetScale3D();
+
 			FVector FilteredLocalScale;
 			FilteredLocalScale.X = Filter.bX ? MixedLocalScale.X : ChildLocalScale.X;
 			FilteredLocalScale.Y = Filter.bY ? MixedLocalScale.Y : ChildLocalScale.Y;
 			FilteredLocalScale.Z = Filter.bZ ? MixedLocalScale.Z : ChildLocalScale.Z;
 
-			FVector FinalLocalScale = FilteredLocalScale;
-			
+			FTransform FilteredMixedLocalTransform = ChildCurrentLocalTransform;
+
+			FilteredMixedLocalTransform.SetScale3D(FilteredLocalScale);
+
+			FTransform FinalLocalTransform = FilteredMixedLocalTransform;
+
 			if (Child.Type == ERigElementType::Control)
 			{
-				FinalLocalScale = FinalLocalScale / GetNonZeroScale(AdditionalOffsetTransform.GetScale3D());
+				// need to convert back to offset space for the actual control value
+				FinalLocalTransform = FilteredMixedLocalTransform.GetRelativeTransform(AdditionalOffsetTransform);
+				FinalLocalTransform.NormalizeRotation();
 			}
-			
-			Hierarchy->SetLocalTransform(Child, FTransform(ChildCurrentLocalTransform.GetRotation(), ChildCurrentLocalTransform.GetTranslation(), FinalLocalScale));
+
+			Hierarchy->SetLocalTransform(Child, FinalLocalTransform); 
 		}
 	}
 }
