Control Rig: Refactoring treeview node layout to fast path

diff --git a/Engine/Plugins/Experimental/ControlRig/Source/ControlRigEditor/Private/Graph/SControlRigGraphNode.cpp b/Engine/Plugins/Experimental/ControlRig/Source/ControlRigEditor/Private/Graph/SControlRigGraphNode.cpp
--- a/Engine/Plugins/Experimental/ControlRig/Source/ControlRigEditor/Private/Graph/SControlRigGraphNode.cpp
+++ b/Engine/Plugins/Experimental/ControlRig/Source/ControlRigEditor/Private/Graph/SControlRigGraphNode.cpp
@@ -9,7 +9,6 @@
 #include "Widgets/SBoxPanel.h"
 #include "Widgets/Layout/SSpacer.h"
 #include "Widgets/Input/SButton.h"
-#include "Widgets/Layout/SScrollBar.h"
 #include "SLevelOfDetailBranchNode.h"
 #include "SGraphPanel.h"
 #include "Framework/Application/SlateApplication.h"
@@ -47,6 +46,10 @@ const FSlateBrush* SControlRigGraphNode::CachedImg_CR_Pin_Disconnected = nullptr
 
 void SControlRigGraphNode::Construct( const FArguments& InArgs )
 {
+	static const float LeftRightPadding = 6.f;
+	static const float TopPadding = 2.f; 
+	static const float MaxHeight = 30.f;
+
 	if (CachedImg_CR_Pin_Connected == nullptr)
 	{
 		static const FName NAME_CR_Pin_Connected("ControlRig.Bug.Solid");
@@ -59,8 +62,9 @@ void SControlRigGraphNode::Construct( const FArguments& InArgs )
 	this->GraphNode = InArgs._GraphNodeObj;
 	this->SetCursor( EMouseCursor::CardinalCross );
 
- 	UControlRigGraphNode* ControlRigGraphNode = InArgs._GraphNodeObj;
-	if (ControlRigGraphNode->GetModelNode() == nullptr)
+ 	UControlRigGraphNode* EdGraphNode = InArgs._GraphNodeObj;
+	ModelNode = EdGraphNode->GetModelNode();
+	if (!ModelNode.IsValid())
 	{
 		return;
 	}
@@ -69,995 +73,641 @@ void SControlRigGraphNode::Construct( const FArguments& InArgs )
 
 	// Re-cache variable info here (unit structure could have changed since last reconstruction, e.g. array add/remove)
 	// and also create missing pins if it hasn't created yet
-	ControlRigGraphNode->AllocateDefaultPins();
+	EdGraphNode->AllocateDefaultPins();
 	
 	NodeErrorType = int32(EMessageSeverity::Info) + 1;
-	InputTree = nullptr;
-	OutputTree = nullptr;
-	InputOutputTree = nullptr;
 	this->UpdateGraphNode();
 
 	SetIsEditable(false);
 
-	ScrollBar = SNew(SScrollBar);
-
-	// create pin-collapse areas
-	LeftNodeBox->AddSlot()
-		.AutoHeight()
-		[
-			SAssignNew(ExecutionTree, STreeView<URigVMPin*>)
-			.Visibility(this, &SControlRigGraphNode::GetExecutionTreeVisibility)
-			.TreeItemsSource(&ControlRigGraphNode->ExecutePins)
-			.SelectionMode(ESelectionMode::None)
-			.OnGenerateRow(this, &SControlRigGraphNode::MakePinTableRowWidget)
-			.OnGetChildren(this, &SControlRigGraphNode::HandleGetChildrenForTree)
-			.OnExpansionChanged(this, &SControlRigGraphNode::HandleExpansionChanged)
-			.OnSetExpansionRecursive(this, &SControlRigGraphNode::HandleExpandRecursively, &ExecutionTree)
-			.ExternalScrollbar(ScrollBar)
-			.ItemHeight(20.0f)
-		];
-
-	LeftNodeBox->AddSlot()
-		.AutoHeight()
-		[
-			SAssignNew(InputTree, STreeView<URigVMPin*>)
-			.Visibility(this, &SControlRigGraphNode::GetInputTreeVisibility)
-			.TreeItemsSource(&ControlRigGraphNode->InputPins)
-			.SelectionMode(ESelectionMode::None)
-			.OnGenerateRow(this, &SControlRigGraphNode::MakePinTableRowWidget)
-			.OnGetChildren(this, &SControlRigGraphNode::HandleGetChildrenForTree)
-			.OnExpansionChanged(this, &SControlRigGraphNode::HandleExpansionChanged)
-			.OnSetExpansionRecursive(this, &SControlRigGraphNode::HandleExpandRecursively, &InputTree)
-			.ExternalScrollbar(ScrollBar)
-			.ItemHeight(20.0f)
-		];
-
-	LeftNodeBox->AddSlot()
-		.AutoHeight()
-		[
-			SAssignNew(InputOutputTree, STreeView<URigVMPin*>)
-			.Visibility(this, &SControlRigGraphNode::GetInputOutputTreeVisibility)
-			.TreeItemsSource(&ControlRigGraphNode->InputOutputPins)
-			.SelectionMode(ESelectionMode::None)
-			.OnGenerateRow(this, &SControlRigGraphNode::MakePinTableRowWidget)
-			.OnGetChildren(this, &SControlRigGraphNode::HandleGetChildrenForTree)
-			.OnExpansionChanged(this, &SControlRigGraphNode::HandleExpansionChanged)
-			.OnSetExpansionRecursive(this, &SControlRigGraphNode::HandleExpandRecursively, &InputOutputTree)
-			.ExternalScrollbar(ScrollBar)
-			.ItemHeight(20.0f)
-		];
+	URigVMController* Controller = EdGraphNode->GetController();
+	Controller->OnModified().AddSP(this, &SControlRigGraphNode::HandleModifiedEvent);
 
-	LeftNodeBox->AddSlot()
-		.AutoHeight()
-		[
-			SAssignNew(OutputTree, STreeView<URigVMPin*>)
-			.Visibility(this, &SControlRigGraphNode::GetOutputTreeVisibility)
-			.TreeItemsSource(&ControlRigGraphNode->OutputPins)
-			.SelectionMode(ESelectionMode::None)
-			.OnGenerateRow(this, &SControlRigGraphNode::MakePinTableRowWidget)
-			.OnGetChildren(this, &SControlRigGraphNode::HandleGetChildrenForTree)
-			.OnExpansionChanged(this, &SControlRigGraphNode::HandleExpansionChanged)
-			.OnSetExpansionRecursive(this, &SControlRigGraphNode::HandleExpandRecursively, &OutputTree)
-			.ExternalScrollbar(ScrollBar)
-			.ItemHeight(20.0f)
-		];
+	TMap<UEdGraphPin*, int32> EdGraphPinToInputPin;
+	for(int32 InputPinIndex = 0; InputPinIndex < InputPins.Num(); InputPinIndex++)
+	{
+		EdGraphPinToInputPin.Add(InputPins[InputPinIndex]->GetPinObj(), InputPinIndex);
+	}
+	TMap<UEdGraphPin*, int32> EdGraphPinToOutputPin;
+	for(int32 OutputPinIndex = 0; OutputPinIndex < OutputPins.Num(); OutputPinIndex++)
+	{
+		EdGraphPinToOutputPin.Add(OutputPins[OutputPinIndex]->GetPinObj(), OutputPinIndex);
+	}
 
-	LeftNodeBox->AddSlot()
-        .AutoHeight()
-        [
-            SAssignNew(VariableRemappingList, SListView<TSharedPtr<FRigVMExternalVariable>>)
-            .Visibility(this, &SControlRigGraphNode::GetVariableListVisibility)
-            .ListItemsSource(&ControlRigGraphNode->ExternalVariables)
-            .SelectionMode(ESelectionMode::None)
-            .OnGenerateRow(this, &SControlRigGraphNode::MakeVariableTableRowWidget)
-            .ExternalScrollbar(ScrollBar)
-            .ItemHeight(20.0f)
-        ];
+	TArray<URigVMPin*> RootModelPins = ModelNode->GetPins();
+	TArray<URigVMPin*> ModelPins;
 
+	// sort model pins
+	// a) execute IOs, b) IO pins, c) input / visible pins, d) output pins
 	struct Local
 	{
-		static void SetItemExpansion_Recursive(UControlRigGraphNode* InControlRigGraphNode, TSharedPtr<STreeView<URigVMPin*>>& TreeWidget, const TArray<URigVMPin*>& InItems)
+		static void VisitPinRecursively(URigVMPin* InPin, TArray<URigVMPin*>& OutPins)
 		{
-			for(URigVMPin* Pin : InItems)
+			OutPins.Add(InPin);
+			for (URigVMPin* SubPin : InPin->GetSubPins())
 			{
-				if(InControlRigGraphNode->IsPinExpanded(Pin->GetPinPath()))
-				{
-					TreeWidget->SetItemExpansion(Pin, true);
-					SetItemExpansion_Recursive(InControlRigGraphNode, TreeWidget, Pin->GetSubPins());
-				}
+				VisitPinRecursively(SubPin, OutPins);
 			}
 		}
 	};
-
-	Local::SetItemExpansion_Recursive(ControlRigGraphNode, ExecutionTree, ControlRigGraphNode->ExecutePins);
-	Local::SetItemExpansion_Recursive(ControlRigGraphNode, InputTree, ControlRigGraphNode->InputPins);
-	Local::SetItemExpansion_Recursive(ControlRigGraphNode, InputOutputTree, ControlRigGraphNode->InputOutputPins);
-	Local::SetItemExpansion_Recursive(ControlRigGraphNode, OutputTree, ControlRigGraphNode->OutputPins);
-
-
-	// force the regeneration of all pins.
-	// the treeview is lazy - to ensure we draw the connections properly we need
-	// to ensure that it updates it's items at least once.
-	FGeometry DummyGeometry(FVector2D(), FVector2D(), FVector2D(FLT_MAX, FLT_MAX), 1.f);
-	ExecutionTree->RequestTreeRefresh();
-	InputTree->RequestTreeRefresh();
-	InputOutputTree->RequestTreeRefresh();
-	OutputTree->RequestTreeRefresh();
-	ExecutionTree->Tick(DummyGeometry, 0.f, 0.f);
-	InputTree->Tick(DummyGeometry, 0.f, 0.f);
-	InputOutputTree->Tick(DummyGeometry, 0.f, 0.f);
-	OutputTree->Tick(DummyGeometry, 0.f, 0.f);
-	VariableRemappingList->Tick(DummyGeometry, 0.f, 0.f);
-
-	const FSlateBrush* ImageBrush = FControlRigEditorStyle::Get().GetBrush(TEXT("ControlRig.Bug.Dot"));
-
-	VisualDebugIndicatorWidget =
-		SNew(SImage)
-		.Image(ImageBrush)
-		.Visibility(EVisibility::Visible);
-
-	static const FSlateColorBrush WhiteBrush(FLinearColor::White);
-	
-	SAssignNew(InstructionCountTextBlockWidget, STextBlock)
-	.Margin(FMargin(2.0f, 2.0f, 2.0f, 1.0f))
-	.Text(this, &SControlRigGraphNode::GetInstructionCountText)
-	.Font(IDetailLayoutBuilder::GetDetailFont())
-	.ColorAndOpacity(FLinearColor::White)
-	.ShadowColorAndOpacity(FLinearColor(0.1f, 0.1f, 0.1f, 1.f))
-	.Visibility(EVisibility::Visible)
-	.ToolTipText(LOCTEXT("NodeHitCountToolTip", "This number represents the hit count for a node.\nFor functions / collapse nodes it represents the sum of all hit counts of contained nodes.\n\nYou can enable / disable the display of the number in the Class Settings\n(Rig Graph Display Settings -> Show Node Run Counts)"));
-
-	ControlRigGraphNode->GetNodeTitleDirtied().BindSP(this, &SControlRigGraphNode::HandleNodeTitleDirtied);
-
-	LastHighDetailSize = FVector2D::ZeroVector;
-}
-
-TSharedRef<SWidget> SControlRigGraphNode::CreateNodeContentArea()
-{
-	// We only use the LeftNodeBox
-	return SNew(SBorder)
-		.BorderImage( FEditorStyle::GetBrush("NoBorder") )
-		.HAlign(HAlign_Fill)
-		.VAlign(VAlign_Fill)
-		.Padding(FMargin(0,3))
-		[
-			SNew(SHorizontalBox)
-			+SHorizontalBox::Slot()
-			.HAlign(HAlign_Fill)
-			.FillWidth(1.0f)
-			[
-				SAssignNew(LeftNodeBox, SVerticalBox)
-			]
-		];
-}
-
-bool SControlRigGraphNode::UseLowDetailPinContent() const
-{
-	if (const SGraphPanel* MyOwnerPanel = GetOwnerPanel().Get())
-	{
-		return (MyOwnerPanel->GetCurrentLOD() <= EGraphRenderingLOD::MediumDetail);
-	}
-	return false;
-}
-
-bool SControlRigGraphNode::UseLowDetailNodeContent() const
-{
-	if (const SGraphPanel* MyOwnerPanel = GetOwnerPanel().Get())
-	{
-		return (MyOwnerPanel->GetCurrentLOD() <= EGraphRenderingLOD::LowestDetail);
-	}
-	return false;
-}
-
-TOptional<FVector2D> SControlRigGraphNode::GetLowDetailDesiredSize() const
-{
-	return LastHighDetailSize;
-}
-
-TSharedPtr<SGraphPin> SControlRigGraphNode::GetHoveredPin(const FGeometry& MyGeometry, const FPointerEvent& MouseEvent) const
-{
-	TSharedPtr<SGraphPin> HoveredPin = SGraphNode::GetHoveredPin(MyGeometry, MouseEvent);
-	if(!HoveredPin.IsValid())
+	for(int32 SortPhase = 0; SortPhase < 4; SortPhase++)
 	{
- 		TArray<TSharedRef<SWidget>> ExtraWidgetArray;
- 		ExtraWidgetToPinMap.GenerateKeyArray(ExtraWidgetArray);
- 		TSet<TSharedRef<SWidget>> ExtraWidgets(ExtraWidgetArray);
- 		
- 		TMap<TSharedRef<SWidget>, FArrangedWidget> Result;
- 		FindChildGeometries(MyGeometry, ExtraWidgets, Result);
-
-		if (Result.Num() > 0)
+		for(URigVMPin* RootPin : RootModelPins)
 		{
-			FArrangedChildren ArrangedWidgets(EVisibility::Visible);
-			Result.GenerateValueArray(ArrangedWidgets.GetInternalArray());
-			int32 HoveredWidgetIndex = SWidget::FindChildUnderMouse(ArrangedWidgets, MouseEvent);
-			if (HoveredWidgetIndex != INDEX_NONE)
+			switch (SortPhase)
 			{
-				return *ExtraWidgetToPinMap.Find(ArrangedWidgets[HoveredWidgetIndex].Widget);
+				case 0: // execute IO pins
+				{
+					if(RootPin->IsExecuteContext() && RootPin->GetDirection() == ERigVMPinDirection::IO)
+					{
+						Local::VisitPinRecursively(RootPin, ModelPins);
+					}
+					break;
+				}
+				case 1: // IO pins
+				{
+					if(!RootPin->IsExecuteContext() && RootPin->GetDirection() == ERigVMPinDirection::IO)
+					{
+						Local::VisitPinRecursively(RootPin, ModelPins);
+					}
+					break;
+				}
+				case 2: // input / visible pins
+				{
+					if(RootPin->GetDirection() == ERigVMPinDirection::Input || RootPin->GetDirection() == ERigVMPinDirection::Visible)
+					{
+						Local::VisitPinRecursively(RootPin, ModelPins);
+					}
+					break;
+				}
+				case 3: // output pins
+				default:
+				{
+					if(RootPin->GetDirection() == ERigVMPinDirection::Output)
+					{
+						Local::VisitPinRecursively(RootPin, ModelPins);
+					}
+					break;
+				}
 			}
 		}
 	}
-	return HoveredPin;
-}
 
-void SControlRigGraphNode::EndUserInteraction() const
-{
-#if WITH_EDITOR
-	if (GEditor)
-	{
-		GEditor->CancelTransaction(0);
-	}
-#endif
+	// add spacer widget at the start
+	LeftNodeBox->AddSlot()
+	.HAlign(HAlign_Fill)
+	.VAlign(VAlign_Center)
+	.AutoHeight()
+	[
+		SNew(SSpacer)
+		.Size(FVector2D(1.f, 2.f))
+	];
 
-	if (GraphNode)
+	TMap<URigVMPin*, int32> ModelPinToInfoIndex;
+	for(URigVMPin* ModelPin : ModelPins)
 	{
-		if (const UControlRigGraphSchema* RigSchema = Cast<UControlRigGraphSchema>(GraphNode->GetSchema()))
+		FPinInfo PinInfo;
+		PinInfo.Index = PinInfos.Num();
+		PinInfo.ParentIndex = INDEX_NONE;
+		PinInfo.bHasChildren = (ModelPin->GetSubPins().Num() > 0);
+		PinInfo.bIsContainer = ModelPin->IsArray();
+		PinInfo.Depth = 0;
+		PinInfo.bExpanded = ModelPin->IsExpanded();
+		PinInfo.ModelPinPath = ModelPin->GetPinPath();
+
+		PinInfo.bHideInputWidget = ModelPin->IsStruct() || PinInfo.bIsContainer;
+		if (!PinInfo.bHideInputWidget)
 		{
-			RigSchema->EndGraphNodeInteraction(GraphNode);
+			if (ModelPin->GetSubPins().Num() == 0)
+			{
+				if (const UControlRigGraphSchema* RigSchema = Cast<UControlRigGraphSchema>(EdGraphNode->GetSchema()))
+				{
+					PinInfo.bHideInputWidget = RigSchema->IsStructEditable(ModelPin->GetScriptStruct());
+				}
+			}
 		}
-	}
-
-	SGraphNode::EndUserInteraction();
-}
-
-void SControlRigGraphNode::MoveTo(const FVector2D& NewPosition, FNodeSet& NodeFilter)
-{
-	if (!NodeFilter.Find(SharedThis(this)))
-	{
-		if (GraphNode && !RequiresSecondPassLayout())
+		
+		if(URigVMPin* ParentPin = ModelPin->GetParentPin())
 		{
-			if (const UControlRigGraphSchema* RigSchema = Cast<UControlRigGraphSchema>(GraphNode->GetSchema()))
+			const int32* ParentIndexPtr = ModelPinToInfoIndex.Find(ParentPin);
+			if(ParentIndexPtr == nullptr)
 			{
-				RigSchema->SetNodePosition(GraphNode, NewPosition);
+				continue;
 			}
+			PinInfo.ParentIndex = *ParentIndexPtr;
+			PinInfo.Depth = PinInfos[PinInfo.ParentIndex].Depth + 1;
 		}
-	}
-}
-
-void SControlRigGraphNode::AddPin(const TSharedRef<SGraphPin>& PinToAdd) 
-{
-	// We show our own label
-	PinToAdd->SetShowLabel(false);
 
-	UControlRigGraphNode* ControlRigGraphNode = CastChecked<UControlRigGraphNode>(GraphNode);
-	if(URigVMNode* ModelNode = ControlRigGraphNode->GetModelNode())
-	{
-		const UEdGraphPin* EdPinObj = PinToAdd->GetPinObj();
+		TAttribute<EVisibility> PinVisibilityAttribute = TAttribute<EVisibility>::CreateSP(this, &SControlRigGraphNode::GetPinVisibility, PinInfo.Index);
 
-		// Remove value widget from combined pin content
-		TSharedPtr<SWrapBox> LabelAndValueWidget = PinToAdd->GetLabelAndValue();
-		TSharedPtr<SHorizontalBox> FullPinHorizontalRowWidget = PinToAdd->GetFullPinHorizontalRowWidget().Pin();
-		if(LabelAndValueWidget.IsValid() && FullPinHorizontalRowWidget.IsValid())
-		{
-			FullPinHorizontalRowWidget->RemoveSlot(LabelAndValueWidget.ToSharedRef());
-		}
+		bool bPinWidgetForExpanderLeft = false;
+		TSharedPtr<SGraphPin> PinWidgetForExpander;
 
-		// Customize the look for pins with injected nodes
-		FString NodeName, PinPath;
-		if (URigVMPin::SplitPinPathAtStart(EdPinObj->GetName(), NodeName, PinPath))
+		bool bPinInfoIsValid = false;
+		if(UEdGraphPin* OutputEdGraphPin = EdGraphNode->FindPin(ModelPin->GetPinPath(), EEdGraphPinDirection::EGPD_Output))
 		{
-			if (URigVMPin* ModelPin = ModelNode->FindPin(PinPath))
+			if(const int32* PinIndexPtr = EdGraphPinToOutputPin.Find(OutputEdGraphPin))
 			{
-				if (ModelPin->HasInjectedNodes())
-				{
-					PinToAdd->SetCustomPinIcon(CachedImg_CR_Pin_Connected, CachedImg_CR_Pin_Disconnected);
-				}
+				PinInfo.OutputPinWidget = OutputPins[*PinIndexPtr];
+				PinInfo.OutputPinWidget->SetVisibility(PinVisibilityAttribute);
+				PinWidgetForExpander = PinInfo.OutputPinWidget;
+				bPinWidgetForExpanderLeft = false;
+				bPinInfoIsValid = true;
 			}
 		}
-
-		PinToAdd->SetOwner(SharedThis(this));
-		PinWidgetMap.Add(EdPinObj, PinToAdd);
-		if(EdPinObj->Direction == EGPD_Input)
-		{
-			InputPins.Add(PinToAdd);
-		}
-		else
+		if(UEdGraphPin* InputEdGraphPin = EdGraphNode->FindPin(ModelPin->GetPinPath(), EEdGraphPinDirection::EGPD_Input))
 		{
-			OutputPins.Add(PinToAdd);
+			if(const int32* PinIndexPtr = EdGraphPinToInputPin.Find(InputEdGraphPin))
+			{
+				PinInfo.InputPinWidget = InputPins[*PinIndexPtr];
+				PinInfo.InputPinWidget->SetVisibility(PinVisibilityAttribute);
+				PinWidgetForExpander = PinInfo.InputPinWidget;
+				bPinWidgetForExpanderLeft = true;
+				bPinInfoIsValid = true;
+			}
 		}
-	}
-}
-
-const FSlateBrush * SControlRigGraphNode::GetNodeBodyBrush() const
-{
-	return FEditorStyle::GetBrush("Graph.Node.Body");
-}
-
-FReply SControlRigGraphNode::OnMouseButtonDown(const FGeometry& MyGeometry, const FPointerEvent& MouseEvent)
-{
-	FReply Reply = SGraphNode::OnMouseButtonDown(MyGeometry, MouseEvent);
 
-	if (UControlRigGraphNode* RigNode = Cast<UControlRigGraphNode>(GraphNode))
-	{
-		if (UControlRigGraph* RigGraph = Cast<UControlRigGraph>(RigNode->GetGraph()))
+		if(!bPinInfoIsValid)
 		{
-			RigGraph->OnGraphNodeClicked.Broadcast(RigNode);
+			continue;
 		}
-	}
 
-	return Reply;
-}
-
-FReply SControlRigGraphNode::OnMouseButtonDoubleClick(const FGeometry& InMyGeometry, const FPointerEvent& InMouseEvent)
-{
-	if (!InMouseEvent.GetModifierKeys().AnyModifiersDown())
-	{
-		if (UControlRigGraphNode* RigNode = Cast<UControlRigGraphNode>(GraphNode))
+		ModelPinToInfoIndex.Add(ModelPin, PinInfos.Add(PinInfo));
+		
+		// check if this pin has sub pins
+		TSharedPtr<SHorizontalBox> FullPinHorizontalRowWidget = PinWidgetForExpander->GetFullPinHorizontalRowWidget().Pin();
+		if(FullPinHorizontalRowWidget.IsValid())
 		{
-			if (URigVMNode* ModelNode = RigNode->GetModelNode())
+			// indent the pin by padding
+			const float DepthIndentation = 12.f * float(PinInfo.Depth + (PinInfo.bHasChildren ? 0 : 1));
+			const float LeftIndentation = bPinWidgetForExpanderLeft ? DepthIndentation : 0.f;
+			const float RightIndentation = bPinWidgetForExpanderLeft ? 0.f : DepthIndentation;
+			
+			if(PinInfo.bHasChildren)
 			{
-				if(Blueprint.IsValid())
+				// only inject the expander arrow for inputs on input / IO
+				// or for output pins
+				if(
+					(
+						(
+							(ModelPin->GetDirection() == ERigVMPinDirection::Input) ||
+							(ModelPin->GetDirection() == ERigVMPinDirection::IO)
+						) &&
+						bPinWidgetForExpanderLeft
+					) ||
+					(
+						(ModelPin->GetDirection() == ERigVMPinDirection::Output) &&
+						(!bPinWidgetForExpanderLeft)
+					)
+				)
 				{
-					Blueprint->BroadcastNodeDoubleClicked(ModelNode);
-					return FReply::Handled();
+					// Add the expander arrow
+					FullPinHorizontalRowWidget->InsertSlot(bPinWidgetForExpanderLeft ? 1 : FullPinHorizontalRowWidget->GetChildren()->Num() - 1)
+					.Padding(LeftIndentation, 0.f, RightIndentation, 0.f)
+					[
+						SNew(SButton)
+						.ButtonStyle(FAppStyle::Get(), TEXT("SimpleButton"))
+						.ContentPadding(0)
+						.VAlign(VAlign_Center)
+						.HAlign(HAlign_Center)
+						.ClickMethod( EButtonClickMethod::MouseDown )
+						.OnClicked(this, &SControlRigGraphNode::OnExpanderArrowClicked, PinInfo.Index)
+						.ToolTipText(LOCTEXT("ExpandSubPin", "Expand Pin"))
+						[
+							SNew(SImage)
+							.Image(this, &SControlRigGraphNode::GetExpanderImage, PinInfo.Index, bPinWidgetForExpanderLeft, false)
+							.ColorAndOpacity(FSlateColor::UseForeground())
+						]
+					];
 				}
 			}
-		}
-	}
-	return SGraphNode::OnMouseButtonDoubleClick(InMyGeometry, InMouseEvent);
-}
-
-EVisibility SControlRigGraphNode::GetTitleVisibility() const
-{
-	return UseLowDetailNodeTitles() ? EVisibility::Hidden : EVisibility::Visible;
-}
-
-EVisibility SControlRigGraphNode::GetExecutionTreeVisibility() const
-{
-	UControlRigGraphNode* ControlRigGraphNode = CastChecked<UControlRigGraphNode>(GraphNode);
-	return ControlRigGraphNode->ExecutePins.Num() > 0 ? EVisibility::Visible : EVisibility::Collapsed;
-}
-
-EVisibility SControlRigGraphNode::GetInputTreeVisibility() const
-{
-	UControlRigGraphNode* ControlRigGraphNode = CastChecked<UControlRigGraphNode>(GraphNode);
-	return ControlRigGraphNode->InputPins.Num() > 0 ? EVisibility::Visible : EVisibility::Collapsed;
-}
-
-EVisibility SControlRigGraphNode::GetInputOutputTreeVisibility() const
-{
-	UControlRigGraphNode* ControlRigGraphNode = CastChecked<UControlRigGraphNode>(GraphNode);
-	return ControlRigGraphNode->InputOutputPins.Num() > 0 ? EVisibility::Visible : EVisibility::Collapsed;
-}
-
-EVisibility SControlRigGraphNode::GetOutputTreeVisibility() const
-{
-	UControlRigGraphNode* ControlRigGraphNode = CastChecked<UControlRigGraphNode>(GraphNode);
-	return ControlRigGraphNode->OutputPins.Num() > 0 ? EVisibility::Visible : EVisibility::Collapsed;
-}
-
-EVisibility SControlRigGraphNode::GetVariableListVisibility() const
-{
-	UControlRigGraphNode* ControlRigGraphNode = CastChecked<UControlRigGraphNode>(GraphNode);
-	return ControlRigGraphNode->ExternalVariables.Num() > 0 ? EVisibility::Visible : EVisibility::Collapsed;
-}
-
-TSharedRef<SWidget> SControlRigGraphNode::CreateTitleWidget(TSharedPtr<SNodeTitle> InNodeTitle)
-{
-	NodeTitle = InNodeTitle;
-
-	TSharedRef<SWidget> WidgetRef = SGraphNode::CreateTitleWidget(NodeTitle);
-	WidgetRef->SetVisibility(MakeAttributeSP(this, &SControlRigGraphNode::GetTitleVisibility));
-	if (NodeTitle.IsValid())
-	{
-		NodeTitle->SetVisibility(MakeAttributeSP(this, &SControlRigGraphNode::GetTitleVisibility));
-	}
-
-	return SNew(SHorizontalBox)
-		+SHorizontalBox::Slot()
-		.Padding(0.0f)
-		[
-			WidgetRef
-		];
-}
-
-class SControlRigExpanderArrow : public SExpanderArrow
-{
-	SLATE_BEGIN_ARGS(SControlRigExpanderArrow) {}
-
-	SLATE_ARGUMENT(bool, LeftAligned)
-
-	SLATE_END_ARGS()
-
-	void Construct( const FArguments& InArgs, const TSharedPtr<class ITableRow>& TableRow)
-	{
-		bLeftAligned = InArgs._LeftAligned;
-
-		SExpanderArrow::Construct(
-			SExpanderArrow::FArguments()
-			.IndentAmount(8.0f),
-			TableRow);
-
-		// override padding
-		ChildSlot.Padding(TAttribute<FMargin>(this, &SControlRigExpanderArrow::GetExpanderPadding_Extended));
-
-		// override image
-		ExpanderArrow->SetContent(
-			SNew(SImage)
-			.Image(this, &SControlRigExpanderArrow::GetExpanderImage_Extended)
-			.ColorAndOpacity(FSlateColor::UseForeground()));
-	}
+			else
+			{
+				TPanelChildren<SBoxPanel::FSlot>* Slots = (TPanelChildren<SBoxPanel::FSlot>*)FullPinHorizontalRowWidget->GetChildren();
+				const int32 SlotToAdjustIndex = bPinWidgetForExpanderLeft ? 0 : Slots->Num() - 1;
+				SBoxPanel::FSlot& Slot = (*Slots)[SlotToAdjustIndex];
 
-	FMargin GetExpanderPadding_Extended() const
-	{
-		const int32 NestingDepth = FMath::Max(0, OwnerRowPtr.Pin()->GetIndentLevel() - BaseIndentLevel.Get());
-		const float Indent = IndentAmount.Get(8.0f);
-		return bLeftAligned ? FMargin( NestingDepth * Indent, 0,0,0 ) : FMargin( 0,0, NestingDepth * Indent,0 );
+				FMargin Padding;
+				if(Slot.SlotPadding.IsSet() || Slot.SlotPadding.IsBound())
+				{
+					Padding = Slot.SlotPadding.Get();
+				}
+				Padding = FMargin(RightIndentation + Padding.Left, Padding.Top, LeftIndentation + Padding.Right, Padding.Bottom);
+				Slot.SlotPadding = Padding; 
+			}
+		}
 	}
 
-	const FSlateBrush* GetExpanderImage_Extended() const
+	for(const FPinInfo& PinInfo : PinInfos)
 	{
-		const bool bIsItemExpanded = OwnerRowPtr.Pin()->IsItemExpanded();
-
-		FName ResourceName;
-		if (bIsItemExpanded)
+		if(PinInfo.InputPinWidget.IsValid())
 		{
-			if ( ExpanderArrow->IsHovered() )
-			{
-				static FName ExpandedHoveredLeftName = "ControlRig.Node.PinTree.Arrow_Expanded_Hovered_Left";
-				static FName ExpandedHoveredRightName = "ControlRig.Node.PinTree.Arrow_Expanded_Hovered_Right";
-				ResourceName = bLeftAligned ? ExpandedHoveredLeftName : ExpandedHoveredRightName;
-			}
-			else
+			if(PinInfo.bHideInputWidget)
 			{
-				static FName ExpandedLeftName = "ControlRig.Node.PinTree.Arrow_Expanded_Left";
-				static FName ExpandedRightName = "ControlRig.Node.PinTree.Arrow_Expanded_Right";
-				ResourceName = bLeftAligned ? ExpandedLeftName : ExpandedRightName;
+				if(PinInfo.InputPinWidget->GetValueWidget() != SNullWidget::NullWidget)
+				{
+					PinInfo.InputPinWidget->GetValueWidget()->SetVisibility(EVisibility::Collapsed);
+				}
 			}
-		}
-		else
-		{
-			if ( ExpanderArrow->IsHovered() )
+				
+			// input pins
+			if(!PinInfo.OutputPinWidget.IsValid())
 			{
-				static FName CollapsedHoveredLeftName = "ControlRig.Node.PinTree.Arrow_Collapsed_Hovered_Left";
-				static FName CollapsedHoveredRightName = "ControlRig.Node.PinTree.Arrow_Collapsed_Hovered_Right";
-				ResourceName = bLeftAligned ? CollapsedHoveredLeftName : CollapsedHoveredRightName;
+				TSharedPtr<SHorizontalBox> SlotLayout;
+				SHorizontalBox::FSlot* FirstSlot = nullptr;
+				
+				LeftNodeBox->AddSlot()
+                .HAlign(HAlign_Fill)
+				.VAlign(VAlign_Center)
+				.AutoHeight()
+				.MaxHeight(MaxHeight)
+                [
+                    SAssignNew(SlotLayout, SHorizontalBox)
+                    .Visibility(this, &SControlRigGraphNode::GetPinVisibility, PinInfo.Index)
+		            
+                    +SHorizontalBox::Slot()
+                    .Expose(FirstSlot)
+                    .FillWidth(1.f)
+                    .HAlign(HAlign_Left)
+                    .Padding(LeftRightPadding, TopPadding, 0.f, 0.f)
+                    [
+                        PinInfo.InputPinWidget.ToSharedRef()
+                    ]
+                ];
+
+				if(PinInfo.bIsContainer)
+				{
+					// make sure to minimize the width of the label
+					FirstSlot->AutoWidth();
+					
+					// add array plus button
+					SlotLayout->AddSlot()
+					.AutoWidth()
+					.HAlign(HAlign_Left)
+                    .Padding(LeftRightPadding, TopPadding, 0.f, 0.f)
+					[
+						SNew(SButton)
+						.ContentPadding(0.0f)
+						.ButtonStyle(FEditorStyle::Get(), "NoBorder")
+						.OnClicked(this, &SControlRigGraphNode::HandleAddArrayElement, PinInfo.ModelPinPath)
+						.IsEnabled(this, &SGraphNode::IsNodeEditable)
+						.Cursor(EMouseCursor::Default)
+						.ToolTipText(LOCTEXT("AddArrayElement", "Add Array Element"))
+						[
+							SNew(SHorizontalBox)
+							+SHorizontalBox::Slot()
+							.AutoWidth()
+							.VAlign(VAlign_Center)
+							[
+								SNew(SImage)
+								.Image(FEditorStyle::GetBrush(TEXT("Icons.PlusCircle")))
+							]
+						]
+					];
+				}
 			}
+			// io pins
 			else
 			{
-				static FName CollapsedLeftName = "ControlRig.Node.PinTree.Arrow_Collapsed_Left";
-				static FName CollapsedRightName = "ControlRig.Node.PinTree.Arrow_Collapsed_Right";
-				ResourceName = bLeftAligned ? CollapsedLeftName : CollapsedRightName;
+				PinInfo.OutputPinWidget->SetShowLabel(false);
+			
+				LeftNodeBox->AddSlot()
+                .HAlign(HAlign_Fill)
+				.VAlign(VAlign_Center)
+                .AutoHeight()
+				.MaxHeight(MaxHeight)
+                [
+                    SNew(SHorizontalBox)
+                    .Visibility(this, &SControlRigGraphNode::GetPinVisibility, PinInfo.Index)
+
+                    +SHorizontalBox::Slot()
+                    .FillWidth(1.f)
+                    .HAlign(HAlign_Left)
+                    .VAlign(VAlign_Center)
+                    .Padding(LeftRightPadding, TopPadding, 0.f, 0.f)
+                    [
+                        PinInfo.InputPinWidget.ToSharedRef()
+                    ]
+
+                    +SHorizontalBox::Slot()
+                    .FillWidth(1.f)
+                    .HAlign(HAlign_Right)
+					.VAlign(VAlign_Center)
+                    .Padding(0.f, TopPadding, LeftRightPadding, 0.f)
+                    [
+                        PinInfo.OutputPinWidget.ToSharedRef()
+                    ]
+                ];
 			}
 		}
-
-		return FControlRigEditorStyle::Get().GetBrush(ResourceName);
-	}
-
-	virtual FCursorReply OnCursorQuery( const FGeometry& MyGeometry, const FPointerEvent& CursorEvent ) const
-	{
-		return FCursorReply::Cursor(EMouseCursor::Default);
-	}
-
-	bool bLeftAligned;
-};
-
-static TAutoConsoleVariable<int32> CVarEnableShowBackground(
-    TEXT("controlrig.ShowBackground"),
-    0,
-    TEXT("Enable unsupported Experimental prototype Modeling Tools"));
-
-
-class SControlRigPinTreeRow : public STableRow<URigVMPin*>
-{
-public:
-	SLATE_BEGIN_ARGS(SControlRigPinTreeRow) {}
-
-	SLATE_ARGUMENT(bool, LeftAligned)
-	
-	SLATE_END_ARGS()
-
-	void Construct(const FArguments& InArgs, const TSharedRef<STableViewBase>& InOwnerTableView)
-	{
-		bLeftAligned = InArgs._LeftAligned;
-
-		STableRow<URigVMPin*>::Construct(STableRow<URigVMPin*>::FArguments(), InOwnerTableView);
-	}
-
-	const FSlateBrush* GetBorder() const
-	{
-		bool bShowBG = (CVarEnableShowBackground.GetValueOnGameThread() > 0);
-
-		if (bShowBG)
-			return STableRow<URigVMPin*>::GetBorder();
-		else
-		// We want a transparent background.
-		return FCoreStyle::Get().GetBrush("NoBrush");
+		// output pins
+		else if(PinInfo.OutputPinWidget.IsValid())
+		{
+			LeftNodeBox->AddSlot()
+            .HAlign(HAlign_Fill)
+			.VAlign(VAlign_Center)
+            .AutoHeight()
+			.MaxHeight(MaxHeight)
+            [
+            	SNew(SHorizontalBox)
+	            .Visibility(this, &SControlRigGraphNode::GetPinVisibility, PinInfo.Index)
+	            
+				+SHorizontalBox::Slot()
+				.FillWidth(1.f)
+				.HAlign(HAlign_Right)
+				.VAlign(VAlign_Center)
+	            .Padding(0.f, TopPadding, LeftRightPadding, 0.f)
+				[
+	                PinInfo.OutputPinWidget.ToSharedRef()
+				]
+            ];
+		}
 	}
 
-	virtual void ConstructChildren( ETableViewMode::Type InOwnerTableMode, const TAttribute<FMargin>& InPadding, const TSharedRef<SWidget>& InContent ) override
+	if(URigVMFunctionReferenceNode* FunctionReferenceNode = Cast<URigVMFunctionReferenceNode>(ModelNode))
 	{
-		const UGraphEditorSettings* Settings = GetDefault<UGraphEditorSettings>();
-		FMargin InputPadding = Settings->GetInputPinPadding();
-		InputPadding.Top = InputPadding.Bottom = 3.0f;
-		InputPadding.Right = 0.0f;
-
-		FMargin OutputPadding = Settings->GetOutputPinPadding();
-		OutputPadding.Top = OutputPadding.Bottom = 3.0f;
-		OutputPadding.Left = 2.0f;
+		TWeakObjectPtr<URigVMFunctionReferenceNode> WeakFunctionReferenceNode = FunctionReferenceNode;
+		TWeakObjectPtr<UControlRigBlueprint> WeakControlRigBlueprint = Blueprint;
 
-		this->Content = InContent;
-
-		SHorizontalBox::FSlot* InnerContentSlotNativePtr = nullptr;
-
-		TSharedRef<SHorizontalBox> ContentBox = SNew(SHorizontalBox);
-
-		if(bLeftAligned)
+		// add the entries for the variable remapping
+		for(const TSharedPtr<FRigVMExternalVariable>& ExternalVariable : EdGraphNode->ExternalVariables)
 		{
-			ContentBox->AddSlot()
-			.AutoWidth()
-			.HAlign(HAlign_Left)
-			.Padding(InputPadding)
-			[
-				SAssignNew(LeftContentBox, SBox)
-			];
-
-			ContentBox->AddSlot()
-			.AutoWidth()
-			.HAlign(HAlign_Right)
-			[
-				SNew(SControlRigExpanderArrow, SharedThis(this))
-				.ToolTipText(LOCTEXT("ExpandSubPin", "Expand Pin"))
-				.LeftAligned(bLeftAligned)
-			];
-
-			ContentBox->AddSlot()
-			.FillWidth(1.0f)
-			.Expose(InnerContentSlotNativePtr)
+			LeftNodeBox->AddSlot()
+			.HAlign(HAlign_Fill)
+			.VAlign(VAlign_Center)
+			.AutoHeight()
+			.MaxHeight(MaxHeight)
 			[
-				SAssignNew(RightContentBox, SBox)
+				SNew(SHorizontalBox)
+	            
+				+SHorizontalBox::Slot()
+				.AutoWidth()
+				.HAlign(HAlign_Left)
+				.VAlign(VAlign_Center)
+				.Padding(LeftRightPadding, TopPadding, LeftRightPadding, 0.f)
 				[
-					InContent
+					SNew(STextBlock)
+					.Text(FText::FromName(ExternalVariable->Name))
+					.TextStyle(FEditorStyle::Get(), NAME_DefaultPinLabelStyle)
+					.ColorAndOpacity(this, &SControlRigGraphNode::GetVariableLabelTextColor, WeakFunctionReferenceNode, ExternalVariable->Name)
+					.ToolTipText(this, &SControlRigGraphNode::GetVariableLabelTooltipText, WeakControlRigBlueprint, ExternalVariable->Name)
 				]
-			];
-		}
-		else
-		{
-			ContentBox->AddSlot()
-			.FillWidth(1.0f)
-			.HAlign(HAlign_Right)
-			[
-				SAssignNew(LeftContentBox, SBox)
-			];
-
-			ContentBox->AddSlot()
-			.AutoWidth()
-			.HAlign(HAlign_Left)
-			[
-				SNew(SControlRigExpanderArrow, SharedThis(this))
-				.LeftAligned(bLeftAligned)
-			];
 
-			ContentBox->AddSlot()
-			.AutoWidth()
-			.HAlign(HAlign_Right)
-			.Expose(InnerContentSlotNativePtr)
-			.Padding(OutputPadding)
-			[
-				SAssignNew(RightContentBox, SBox)
+				+SHorizontalBox::Slot()
+				.AutoWidth()
+				.HAlign(HAlign_Left)
+				.VAlign(VAlign_Center)
+				.Padding(LeftRightPadding, TopPadding, LeftRightPadding, 0.f)
 				[
-					InContent
+					SNew(SControlRigVariableBinding)
+					.Blueprint(Blueprint.Get())
+					.FunctionReferenceNode(FunctionReferenceNode)
+					.InnerVariableName(ExternalVariable->Name)
 				]
 			];
 		}
-
-		this->ChildSlot
-		[
-            ContentBox
-		];
-
-		InnerContentSlot = InnerContentSlotNativePtr;
-	}
-
-	/** Exposed boxes to slot pin widgets into */
-	TSharedPtr<SBox> LeftContentBox;
-	TSharedPtr<SBox> RightContentBox;
-
-	/** Whether we align our content left or right */
-	bool bLeftAligned;
-};
-
-class SControlRigVariableListRow : public STableRow<TSharedPtr<FRigVMExternalVariable>>
-{
-public:
-	SLATE_BEGIN_ARGS(SControlRigVariableListRow) {}
-	SLATE_END_ARGS()
-
-	void Construct(const FArguments& InArgs, const TSharedRef<STableViewBase>& InOwnerTableView)
-	{
-		STableRow<TSharedPtr<FRigVMExternalVariable>>::Construct(STableRow<TSharedPtr<FRigVMExternalVariable>>::FArguments(), InOwnerTableView);
-	}
-
-	const FSlateBrush* GetBorder() const
-	{
-		bool bShowBG = (CVarEnableShowBackground.GetValueOnGameThread() > 0);
-
-		if (bShowBG)
-			return STableRow<TSharedPtr<FRigVMExternalVariable>>::GetBorder();
-		else
-		// We want a transparent background.
-		return FCoreStyle::Get().GetBrush("NoBrush");
 	}
+	
+	// add spacer widget at the end
+	LeftNodeBox->AddSlot()
+	.HAlign(HAlign_Fill)
+	.VAlign(VAlign_Center)
+	.AutoHeight()
+	[
+		SNew(SSpacer)
+		.Size(FVector2D(1.f, 4.f))
+	];
 
-	virtual void ConstructChildren( ETableViewMode::Type InOwnerTableMode, const TAttribute<FMargin>& InPadding, const TSharedRef<SWidget>& InContent ) override
-	{
-		const UGraphEditorSettings* Settings = GetDefault<UGraphEditorSettings>();
-		FMargin InputPadding = Settings->GetInputPinPadding();
-		InputPadding.Top = InputPadding.Bottom = 3.0f;
-		InputPadding.Right = 8.0f;
+	const FSlateBrush* ImageBrush = FControlRigEditorStyle::Get().GetBrush(TEXT("ControlRig.Bug.Dot"));
 
-		this->Content = InContent;
+	VisualDebugIndicatorWidget =
+		SNew(SImage)
+		.Image(ImageBrush)
+		.Visibility(EVisibility::Visible);
 
-		SHorizontalBox::FSlot* InnerContentSlotNativePtr = nullptr;
+	static const FSlateColorBrush WhiteBrush(FLinearColor::White);
+	
+	SAssignNew(InstructionCountTextBlockWidget, STextBlock)
+	.Margin(FMargin(2.0f, 2.0f, 2.0f, 1.0f))
+	.Text(this, &SControlRigGraphNode::GetInstructionCountText)
+	.Font(IDetailLayoutBuilder::GetDetailFont())
+	.ColorAndOpacity(FLinearColor::White)
+	.ShadowColorAndOpacity(FLinearColor(0.1f, 0.1f, 0.1f, 1.f))
+	.Visibility(EVisibility::Visible)
+	.ToolTipText(LOCTEXT("NodeHitCountToolTip", "This number represents the hit count for a node.\nFor functions / collapse nodes it represents the sum of all hit counts of contained nodes.\n\nYou can enable / disable the display of the number in the Class Settings\n(Rig Graph Display Settings -> Show Node Run Counts)"));
 
-		TSharedRef<SHorizontalBox> ContentBox = SNew(SHorizontalBox);
+	EdGraphNode->GetNodeTitleDirtied().BindSP(this, &SControlRigGraphNode::HandleNodeTitleDirtied);
 
-		ContentBox->AddSlot()
-		.AutoWidth()
-		.HAlign(HAlign_Left)
-		.VAlign(VAlign_Center)
-		.Padding(InputPadding)
-		[
-			SAssignNew(LeftContentBox, SBox)
-		];
+	LastHighDetailSize = FVector2D::ZeroVector;
+}
 
-		ContentBox->AddSlot()
-		.AutoWidth()
-		.HAlign(HAlign_Left)
-		.VAlign(VAlign_Center)
-		.MaxWidth(200.f)
-		.Expose(InnerContentSlotNativePtr)
+TSharedRef<SWidget> SControlRigGraphNode::CreateNodeContentArea()
+{
+	return SNew(SLevelOfDetailBranchNode)
+		.UseLowDetailSlot(this, &SControlRigGraphNode::UseLowDetailNodeContent)
+		.LowDetail()
 		[
-			SAssignNew(RightContentBox, SBox)
-			[
-				InContent
-			]
-		];
-
-		this->ChildSlot
+			SNew(SSpacer)
+			.Size(this, &SControlRigGraphNode::GetLowDetailDesiredSize)
+		]
+		.HighDetail()
 		[
-			ContentBox
+			SAssignNew(LeftNodeBox, SVerticalBox)
 		];
+}
 
-		InnerContentSlot = InnerContentSlotNativePtr;
+bool SControlRigGraphNode::UseLowDetailPinNames() const
+{
+	if (const SGraphPanel* MyOwnerPanel = GetOwnerPanel().Get())
+	{
+		return (MyOwnerPanel->GetCurrentLOD() <= EGraphRenderingLOD::LowDetail);
 	}
+	return false;
+}
 
-	/** Exposed boxes to slot pin widgets into */
-	TSharedPtr<SBox> LeftContentBox;
-	TSharedPtr<SBox> RightContentBox;
-};
-
-TSharedRef<SWidget> SControlRigGraphNode::AddContainerPinContent(URigVMPin* InItem, FText InTooltipText)
+bool SControlRigGraphNode::UseLowDetailNodeContent() const
 {
-	return SNew(SButton)
-	.ContentPadding(0.0f)
-	.ButtonStyle(FEditorStyle::Get(), "NoBorder")
-	.OnClicked(this, &SControlRigGraphNode::HandleAddArrayElement, InItem)
-	.IsEnabled(this, &SGraphNode::IsNodeEditable)
-	.ToolTipText(InTooltipText)
-	.Cursor(EMouseCursor::Default)
-	[
-		SNew(SHorizontalBox)
-		+SHorizontalBox::Slot()
-		.AutoWidth()
-		.VAlign(VAlign_Center)
-		.Padding(7,0,0,0)
-		[
-			SNew(SImage)
-			.Image(FEditorStyle::GetBrush(TEXT("Icons.PlusCircle")))
-		]
-	];
+	if(LastHighDetailSize.IsNearlyZero())
+	{
+		return false;
+	}
+	
+	if (const SGraphPanel* MyOwnerPanel = GetOwnerPanel().Get())
+	{
+		return (MyOwnerPanel->GetCurrentLOD() <= EGraphRenderingLOD::LowestDetail);
+	}
+	return false;
 }
 
-TSharedRef<ITableRow> SControlRigGraphNode::MakePinTableRowWidget(URigVMPin* InItem, const TSharedRef<STableViewBase>& OwnerTable)
+FVector2D SControlRigGraphNode::GetLowDetailDesiredSize() const
 {
-	const bool bLeaf = InItem->GetSubPins().Num() == 0;
-	const bool bIsContainer = InItem->IsArray();
-
-	TSharedPtr<SGraphPin> InputPinWidget;
-	TSharedPtr<SGraphPin> OutputPinWidget;
-	TSharedPtr<SWidget> InputPinValueWidget;
-
-	UControlRigGraphNode* RigNode = Cast<UControlRigGraphNode>(GraphNode);
-	const UControlRigGraphNode::PinPair& Pair = RigNode->CachedPins.FindChecked(InItem);
+	return LastHighDetailSize;
+}
 
-	if (Pair.InputPin)
+void SControlRigGraphNode::EndUserInteraction() const
+{
+#if WITH_EDITOR
+	if (GEditor)
 	{
-		TSharedPtr<SGraphPin>* InputGraphPinPtr = PinWidgetMap.Find(Pair.InputPin);
-		if(InputGraphPinPtr != nullptr)
-		{
-			InputPinWidget = *InputGraphPinPtr;
-
-			bool bIsPlainOrEditableStruct = !InItem->IsStruct();
-			if (!bIsPlainOrEditableStruct)
-			{
-				if (InItem->GetSubPins().Num() == 0)
-				{
-					if (const UControlRigGraphSchema* RigSchema = Cast<UControlRigGraphSchema>(RigNode->GetSchema()))
-					{
-						bIsPlainOrEditableStruct = RigSchema->IsStructEditable(InItem->GetScriptStruct());
-					}
-				}
-			}
-
-			// Only leaf pins have value widgets, but not containers
-			if(((bLeaf && bIsPlainOrEditableStruct) || InItem->IsBoundToVariable()) && !bIsContainer)
-			{
-				InputPinValueWidget = (*InputGraphPinPtr)->GetValueWidget();
-			}
-			else if(bIsContainer)
-			{
-				// Add a 'new item' widget for containers
-				InputPinValueWidget = AddContainerPinContent(InItem, LOCTEXT("AddToContainer", "Add a new value to this container"));
-			}
-		}
+		GEditor->CancelTransaction(0);
 	}
+#endif
 
-	if (Pair.OutputPin)
+	if (GraphNode)
 	{
-		TSharedPtr<SGraphPin>* OutputGraphPinPtr = PinWidgetMap.Find(Pair.OutputPin);
-		if(OutputGraphPinPtr != nullptr)
+		if (const UControlRigGraphSchema* RigSchema = Cast<UControlRigGraphSchema>(GraphNode->GetSchema()))
 		{
-			OutputPinWidget = *OutputGraphPinPtr;
+			RigSchema->EndGraphNodeInteraction(GraphNode);
 		}
 	}
 
-	TSharedRef<SControlRigPinTreeRow> ControlRigPinTreeRow = SNew(SControlRigPinTreeRow, OwnerTable)
-		.LeftAligned(!(OutputPinWidget.IsValid() && !InputPinWidget.IsValid()))
-		.ToolTipText(InItem->GetToolTipText());
+	SGraphNode::EndUserInteraction();
+}
 
-	if(InputPinWidget.IsValid() || OutputPinWidget.IsValid())
+void SControlRigGraphNode::MoveTo(const FVector2D& NewPosition, FNodeSet& NodeFilter)
+{
+	if (!NodeFilter.Find(SharedThis(this)))
 	{
-		TWeakPtr<SGraphPin> WeakPin = InputPinWidget.IsValid() ? InputPinWidget : OutputPinWidget;
-
-		TSharedPtr<STextBlock> LabelTextBlock;
-		TSharedRef<SWidget> LabelWidget = SAssignNew(LabelTextBlock, STextBlock)
-		.Text(this, &SControlRigGraphNode::GetPinLabel, WeakPin)
-		.TextStyle(FEditorStyle::Get(), NAME_DefaultPinLabelStyle)
-		.SimpleTextMode(true)
-		.ColorAndOpacity(this, &SControlRigGraphNode::GetPinTextColor, WeakPin);
-
-		const FText PinLabel = GetPinLabel(WeakPin);
-		const FVector2D LabelTextSize = FSlateApplication::Get().GetRenderer()->GetFontMeasureService()->Measure(PinLabel, IDetailLayoutBuilder::GetDetailFont()) * FVector2D(1.f, 0.2f);
-
-		// add to mapping that allows labels to act as hover widgets
-		if(InputPinWidget.IsValid())
-		{
-			ExtraWidgetToPinMap.Emplace(LabelWidget, InputPinWidget.ToSharedRef());
-		}
-		else if(OutputPinWidget.IsValid())
+		if (GraphNode && !RequiresSecondPassLayout())
 		{
-			ExtraWidgetToPinMap.Emplace(LabelWidget, OutputPinWidget.ToSharedRef());
+			if (const UControlRigGraphSchema* RigSchema = Cast<UControlRigGraphSchema>(GraphNode->GetSchema()))
+			{
+				RigSchema->SetNodePosition(GraphNode, NewPosition);
+			}
 		}
+	}
+}
 
-		FMargin OutputPadding = Settings->GetOutputPinPadding();
-		OutputPadding.Top = OutputPadding.Bottom = 3.0f;
-		OutputPadding.Left = 2.0f;
+void SControlRigGraphNode::AddPin(const TSharedRef<SGraphPin>& PinToAdd) 
+{
+	if(ModelNode.IsValid())
+	{
+		const UEdGraphPin* EdPinObj = PinToAdd->GetPinObj();
 
-		if(OutputPinWidget.IsValid() && !InputPinWidget.IsValid())
+		// Customize the look for pins with injected nodes
+		FString NodeName, PinPath;
+		if (URigVMPin::SplitPinPathAtStart(EdPinObj->GetName(), NodeName, PinPath))
 		{
-			TSharedRef<SWidget> InputWidget = 
-                SNew(SHorizontalBox)
-                +SHorizontalBox::Slot()
-                .HAlign(HAlign_Right)
-                .VAlign(VAlign_Center)
-                .FillWidth(1.0f)
-                .Padding(50.0f, 2.0f, 2.0f, 2.0f)
-                [
-                    SNew(SLevelOfDetailBranchNode)
-                    .UseLowDetailSlot(this, &SControlRigGraphNode::UseLowDetailNodeContent)
-                    .LowDetail()
-                    [
-                    	SNew(SSpacer)
-                    ]
-                    .HighDetail()
-                    [
-                        SNew(SLevelOfDetailBranchNode)
-                        .UseLowDetailSlot(this, &SControlRigGraphNode::UseLowDetailPinContent)
-                        .LowDetail()
-                        [
-                            // replacement rectangle in the color of the text
-                            SNew(SImage)
-                            .DesiredSizeOverride(LabelTextSize)
-                            .Image(FEditorStyle::GetBrush("WhiteBrush"))
-                            .ColorAndOpacity(FEditorStyle::GetSlateColor("DefaultForeground"))
-                        ]
-                        .HighDetail()
-                        [
-                            LabelWidget
-                        ]
-					]
-				];
-
-			TSharedRef<SWidget> OutputWidget = 
-				SNew(SHorizontalBox)
-				+SHorizontalBox::Slot()
-				.AutoWidth()
-				.HAlign(HAlign_Right)
-				.VAlign(VAlign_Center)
-				[
-					SNew(SBox)
-					.MinDesiredWidth(24.0f)
-					[
-						OutputPinWidget.IsValid() ? StaticCastSharedRef<SWidget>(OutputPinWidget.ToSharedRef()) : StaticCastSharedRef<SWidget>(SNew(SSpacer).Size(FVector2D(16.0f, 14.0f)))
-					]
-				];
-
-			ControlRigPinTreeRow->LeftContentBox->SetContent(InputWidget);
-			ControlRigPinTreeRow->RightContentBox->SetContent(OutputWidget);
+			if (URigVMPin* ModelPin = ModelNode->FindPin(PinPath))
+			{
+				if (ModelPin->HasInjectedNodes())
+				{
+					PinToAdd->SetCustomPinIcon(CachedImg_CR_Pin_Connected, CachedImg_CR_Pin_Disconnected);
+				}
+				PinToAdd->SetToolTipText(ModelPin->GetToolTipText());
+			}
 		}
-		else
-		{
-			TSharedRef<SWidget> InputWidget = 
-				SNew(SHorizontalBox)
-				+SHorizontalBox::Slot()
-				.AutoWidth()
-				.HAlign(HAlign_Left)
-				.VAlign(VAlign_Center)
-				[
-					SNew(SBox)
-					.MinDesiredWidth(24.0f)
-					[
-						InputPinWidget.IsValid() ? StaticCastSharedRef<SWidget>(InputPinWidget.ToSharedRef()) : StaticCastSharedRef<SWidget>(SNew(SSpacer).Size(FVector2D(16.0f, 14.0f)))
-					]
-				];
 
-			TSharedPtr<SWidget> InputValueWrapperWidget;
-			if(InputPinValueWidget.IsValid())
+		// reformat the pin by
+		// 1. taking out the swrapbox widget
+		// 2. re-inserting all widgets from the label and value wrap box back in the horizontal box
+		TSharedPtr<SHorizontalBox> FullPinHorizontalRowWidget = PinToAdd->GetFullPinHorizontalRowWidget().Pin();
+		TSharedPtr<SWrapBox> LabelAndValueWidget = PinToAdd->GetLabelAndValue();
+		if(FullPinHorizontalRowWidget.IsValid() && LabelAndValueWidget.IsValid())
+		{
+			int32 LabelAndValueWidgetIndex = INDEX_NONE;
+			for(int32 ChildIndex = 0; ChildIndex < FullPinHorizontalRowWidget->GetChildren()->Num(); ChildIndex++)
 			{
-				InputValueWrapperWidget = SNew(SLevelOfDetailBranchNode)
-				.UseLowDetailSlot(this, &SControlRigGraphNode::UseLowDetailPinContent)
-				.LowDetail()
-				[
-					SNew(SSpacer)
-				]
-				.HighDetail()
-				[
-					InputPinValueWidget.ToSharedRef()
-				];
+				TSharedRef<SWidget> ChildWidget = FullPinHorizontalRowWidget->GetChildren()->GetChildAt(ChildIndex);
+				if(ChildWidget == LabelAndValueWidget)
+				{
+					LabelAndValueWidgetIndex = ChildIndex;
+					break;
+				}
 			}
-			else
+			check(LabelAndValueWidgetIndex != INDEX_NONE);
+			
+			FullPinHorizontalRowWidget->RemoveSlot(LabelAndValueWidget.ToSharedRef());
+			
+			for(int32 ChildIndex = 0; ChildIndex < LabelAndValueWidget->GetChildren()->Num(); ChildIndex++)
 			{
-				InputValueWrapperWidget = SNew(SSpacer);
-			}
-
-			TSharedRef<SWidget> OutputWidget = 
-				SNew(SHorizontalBox)
-				+SHorizontalBox::Slot()
-				.HAlign(HAlign_Left)
-				.VAlign(VAlign_Center)
-				.AutoWidth()
-				.Padding(2.0f)
-				[
-					SNew(SLevelOfDetailBranchNode)
-	                .UseLowDetailSlot(this, &SControlRigGraphNode::UseLowDetailNodeContent)
-	                .LowDetail()
-	                [
-	                    SNew(SSpacer)
-	                ]
-	                .HighDetail()
-	                [
-		                SNew(SLevelOfDetailBranchNode)
-		                .UseLowDetailSlot(this, &SControlRigGraphNode::UseLowDetailPinContent)
-		                .LowDetail()
-		                [
-		                    // replacement rectangle in the color of the text
-		                    SNew(SImage)
-		                    .DesiredSizeOverride(LabelTextSize)
-		                    .Image(FEditorStyle::GetBrush("WhiteBrush"))
-		                    .ColorAndOpacity(FEditorStyle::GetSlateColor("DefaultForeground"))
-		                ]
-		                .HighDetail()
-		                [
-		                    LabelWidget
-		                ]
-		            ]
-	            ]
-				+SHorizontalBox::Slot()
-				.AutoWidth()
-				.HAlign(HAlign_Left)
-				.VAlign(VAlign_Center)
-				.Padding(2.0f, 2.0f, 18.0f, 2.0f)
-				[
-					InputValueWrapperWidget.ToSharedRef()
-				]
-				+SHorizontalBox::Slot()
-				.FillWidth(1.0f)
-				.Padding(OutputPadding)
-				.HAlign(HAlign_Right)
-				.VAlign(VAlign_Center)
-				[
-					SNew(SBox)
-					.MinDesiredWidth(24.0f)
+				TSharedRef<SWidget> ChildWidget = LabelAndValueWidget->GetChildren()->GetChildAt(ChildIndex);
+				if(ChildWidget != SNullWidget::NullWidget)
+				{
+					ChildWidget->AssignParentWidget(FullPinHorizontalRowWidget.ToSharedRef());
+					
+					FullPinHorizontalRowWidget->InsertSlot(LabelAndValueWidgetIndex + ChildIndex)
+					.HAlign(HAlign_Fill)
+					.VAlign(VAlign_Center)
+					.Padding(EdPinObj->Direction == EGPD_Input ? 0.f : 2.f, 0.f, EdPinObj->Direction == EGPD_Input ? 2.f : 0.f, 0.f)
+					.AutoWidth()
 					[
-						OutputPinWidget.IsValid() ? StaticCastSharedRef<SWidget>(OutputPinWidget.ToSharedRef()) : StaticCastSharedRef<SWidget>(SNew(SSpacer).Size(FVector2D(16.0f, 14.0f)))
-					]
-				];
+						ChildWidget
+					];
+				}
+			}
+		}
 
-			ControlRigPinTreeRow->LeftContentBox->SetContent(InputWidget);
-			ControlRigPinTreeRow->RightContentBox->SetContent(OutputWidget);
+		PinToAdd->SetOwner(SharedThis(this));
+		if(EdPinObj->Direction == EGPD_Input)
+		{
+			InputPins.Add(PinToAdd);
+		}
+		else
+		{
+			OutputPins.Add(PinToAdd);
 		}
 	}
-
-	return ControlRigPinTreeRow;
 }
 
-TSharedRef<ITableRow> SControlRigGraphNode::MakeVariableTableRowWidget(TSharedPtr<FRigVMExternalVariable> InVariable,
-	const TSharedRef<STableViewBase>& OwnerTable)
+const FSlateBrush * SControlRigGraphNode::GetNodeBodyBrush() const
 {
-	TSharedRef<SControlRigVariableListRow> ControlRigVariableListRow = SNew(SControlRigVariableListRow, OwnerTable);
-
-	UControlRigGraphNode* RigNode = Cast<UControlRigGraphNode>(GraphNode);
-	URigVMFunctionReferenceNode* FunctionReferenceNode = Cast<URigVMFunctionReferenceNode>(RigNode->GetModelNode());
-	check(FunctionReferenceNode);
-	TWeakObjectPtr<URigVMFunctionReferenceNode> WeakFunctionReferenceNode = FunctionReferenceNode;
-
-	UControlRigBlueprint* ControlRigBlueprint = FunctionReferenceNode->GetReferencedNode()->GetTypedOuter<UControlRigBlueprint>(); 
-	check(ControlRigBlueprint);
-	TWeakObjectPtr<UControlRigBlueprint> WeakControlRigBlueprint = ControlRigBlueprint;
-
-	ControlRigVariableListRow->LeftContentBox->SetContent(SNew(STextBlock)
-		.Text(FText::FromName(InVariable->Name))
-        .TextStyle(FEditorStyle::Get(), NAME_DefaultPinLabelStyle)
-		.ColorAndOpacity(this, &SControlRigGraphNode::GetVariableLabelTextColor, WeakFunctionReferenceNode, InVariable->Name)
-		.ToolTipText(this, &SControlRigGraphNode::GetVariableLabelTooltipText, WeakControlRigBlueprint, InVariable->Name)
-    );
-
-	ControlRigVariableListRow->RightContentBox->SetContent(SNew(SControlRigVariableBinding)
-		.Blueprint(Blueprint.Get())
-		.FunctionReferenceNode(FunctionReferenceNode)
-		.InnerVariableName(InVariable->Name)
-	);
-
-	return ControlRigVariableListRow;
+	return FEditorStyle::GetBrush("Graph.Node.Body");
 }
 
-void SControlRigGraphNode::HandleGetChildrenForTree(URigVMPin* InItem, TArray<URigVMPin*>& OutChildren)
+FReply SControlRigGraphNode::OnMouseButtonDown(const FGeometry& MyGeometry, const FPointerEvent& MouseEvent)
 {
-	OutChildren.Append(InItem->GetSubPins());
-}
+	FReply Reply = SGraphNode::OnMouseButtonDown(MyGeometry, MouseEvent);
 
-void SControlRigGraphNode::HandleExpansionChanged(URigVMPin* InItem, bool bExpanded)
-{
-	if (UControlRigGraphNode* RigGraphNode = Cast< UControlRigGraphNode>(GraphNode))
+	if (UControlRigGraphNode* RigNode = Cast<UControlRigGraphNode>(GraphNode))
 	{
-		RigGraphNode->GetController()->SetPinExpansion(InItem->GetPinPath(), bExpanded, true);
+		if (UControlRigGraph* RigGraph = Cast<UControlRigGraph>(RigNode->GetGraph()))
+		{
+			RigGraph->OnGraphNodeClicked.Broadcast(RigNode);
+		}
 	}
+
+	return Reply;
 }
 
-void SControlRigGraphNode::HandleExpandRecursively(URigVMPin* InItem, bool bExpanded, TSharedPtr<STreeView<URigVMPin*>>* TreeWidgetPtr)
+FReply SControlRigGraphNode::OnMouseButtonDoubleClick(const FGeometry& InMyGeometry, const FPointerEvent& InMouseEvent)
 {
-	TSharedPtr<STreeView<URigVMPin*>>& TreeWidget = *TreeWidgetPtr;
-
-	if (UControlRigGraphNode* RigGraphNode = Cast< UControlRigGraphNode>(GraphNode))
+	if (!InMouseEvent.GetModifierKeys().AnyModifiersDown())
 	{
-		RigGraphNode->GetController()->OpenUndoBracket(TEXT("Expand pin recursively"));
-
-		TArray<URigVMPin*> ModelPins;
-		ModelPins.Add(InItem);
-
-		for (int32 PinIndex = 0; PinIndex < ModelPins.Num(); PinIndex++)
-		{
-			URigVMPin* ModelPin = ModelPins[PinIndex];
-			ModelPins.Append(ModelPin->GetSubPins());
-		}
-
-		if (!bExpanded)
-		{
-			Algo::Reverse(ModelPins);
-		}
-
-		for (int32 PinIndex = 0; PinIndex < ModelPins.Num(); PinIndex++)
+		if (ModelNode.IsValid())
 		{
-			URigVMPin* ModelPin = ModelPins[PinIndex];
-			if (ModelPin->IsExpanded() != bExpanded && ModelPin->GetSubPins().Num() > 0)
+			if(Blueprint.IsValid())
 			{
-				TreeWidget->SetItemExpansion(ModelPin, bExpanded);
+				Blueprint->BroadcastNodeDoubleClicked(ModelNode.Get());
+				return FReply::Handled();
 			}
 		}
+	}
+	return SGraphNode::OnMouseButtonDoubleClick(InMyGeometry, InMouseEvent);
+}
+
+EVisibility SControlRigGraphNode::GetTitleVisibility() const
+{
+	return UseLowDetailNodeTitles() ? EVisibility::Hidden : EVisibility::Visible;
+}
+
+TSharedRef<SWidget> SControlRigGraphNode::CreateTitleWidget(TSharedPtr<SNodeTitle> InNodeTitle)
+{
+	NodeTitle = InNodeTitle;
 
-		RigGraphNode->GetController()->CloseUndoBracket();
+	TSharedRef<SWidget> WidgetRef = SGraphNode::CreateTitleWidget(NodeTitle);
+	WidgetRef->SetVisibility(MakeAttributeSP(this, &SControlRigGraphNode::GetTitleVisibility));
+	if (NodeTitle.IsValid())
+	{
+		NodeTitle->SetVisibility(MakeAttributeSP(this, &SControlRigGraphNode::GetTitleVisibility));
 	}
+
+	return SNew(SHorizontalBox)
+		+SHorizontalBox::Slot()
+		.Padding(0.0f)
+		[
+			WidgetRef
+		];
 }
 
 FText SControlRigGraphNode::GetPinLabel(TWeakPtr<SGraphPin> GraphPin) const
@@ -1130,16 +780,15 @@ FText SControlRigGraphNode::GetVariableLabelTooltipText(TWeakObjectPtr<UControlR
 	return FText();
 }
 
-FReply SControlRigGraphNode::HandleAddArrayElement(URigVMPin* InItem)
+FReply SControlRigGraphNode::HandleAddArrayElement(FString InModelPinPath)
 {
-	if(InItem)
+	if(!InModelPinPath.IsEmpty())
 	{
 		if (UControlRigGraphNode* ControlRigGraphNode = Cast<UControlRigGraphNode>(GraphNode))
 		{
-			ControlRigGraphNode->HandleAddArrayElement(InItem->GetPinPath());
+			ControlRigGraphNode->HandleAddArrayElement(InModelPinPath);
 		}
 	}
-
 	return FReply::Handled();
 }
 
@@ -1263,78 +912,75 @@ TArray<FOverlayWidgetInfo> SControlRigGraphNode::GetOverlayWidgets(bool bSelecte
 {
 	TArray<FOverlayWidgetInfo> Widgets = SGraphNode::GetOverlayWidgets(bSelected, WidgetSize);
 
-	if (UControlRigGraphNode* RigNode = CastChecked<UControlRigGraphNode>(GraphNode, ECastCheckedType::NullAllowed))
+	if (ModelNode.IsValid())
 	{
-		if (URigVMNode* ModelNode = RigNode->GetModelNode())
-		{
-			bool bSetColor = false;
-			FLinearColor Color = FLinearColor::Black;
-			int32 PreviousNumWidgets = Widgets.Num();
-			VisualDebugIndicatorWidget->SetColorAndOpacity(Color);
+		bool bSetColor = false;
+		FLinearColor Color = FLinearColor::Black;
+		int32 PreviousNumWidgets = Widgets.Num();
+		VisualDebugIndicatorWidget->SetColorAndOpacity(Color);
 
-			for (URigVMPin* ModelPin : ModelNode->GetPins())
+		for (URigVMPin* ModelPin : ModelNode->GetPins())
+		{
+			if (ModelPin->HasInjectedNodes())
 			{
-				if (ModelPin->HasInjectedNodes())
+				for (URigVMInjectionInfo* Injection : ModelPin->GetInjectedNodes())
 				{
-					for (URigVMInjectionInfo* Injection : ModelPin->GetInjectedNodes())
-					{
-						URigVMUnitNode* VisualDebugNode = Injection->UnitNode;
+					URigVMUnitNode* VisualDebugNode = Injection->UnitNode;
 
-						FString PrototypeName;
-						if (VisualDebugNode->GetScriptStruct()->GetStringMetaDataHierarchical(TEXT("PrototypeName"), &PrototypeName))
+					FString PrototypeName;
+					if (VisualDebugNode->GetScriptStruct()->GetStringMetaDataHierarchical(TEXT("PrototypeName"), &PrototypeName))
+					{
+						if (PrototypeName == TEXT("VisualDebug"))
 						{
-							if (PrototypeName == TEXT("VisualDebug"))
+							if (!bSetColor)
 							{
-								if (!bSetColor)
+								if (VisualDebugNode->FindPin(TEXT("bEnabled"))->GetDefaultValue() == TEXT("True"))
 								{
-									if (VisualDebugNode->FindPin(TEXT("bEnabled"))->GetDefaultValue() == TEXT("True"))
+									if (URigVMPin* ColorPin = VisualDebugNode->FindPin(TEXT("Color")))
+									{
+										TBaseStructure<FLinearColor>::Get()->ImportText(*ColorPin->GetDefaultValue(), &Color, nullptr, PPF_None, nullptr, TBaseStructure<FLinearColor>::Get()->GetName());
+									}
+									else
 									{
-										if (URigVMPin* ColorPin = VisualDebugNode->FindPin(TEXT("Color")))
-										{
-											TBaseStructure<FLinearColor>::Get()->ImportText(*ColorPin->GetDefaultValue(), &Color, nullptr, PPF_None, nullptr, TBaseStructure<FLinearColor>::Get()->GetName());
-										}
-										else
-										{
-											Color = FLinearColor::White;
-										}
-
-										VisualDebugIndicatorWidget->SetColorAndOpacity(Color);
-										bSetColor = true;
+										Color = FLinearColor::White;
 									}
+
+									VisualDebugIndicatorWidget->SetColorAndOpacity(Color);
+									bSetColor = true;
 								}
+							}
 
-								if (Widgets.Num() == PreviousNumWidgets)
-								{
-									const FVector2D ImageSize = VisualDebugIndicatorWidget->GetDesiredSize();
+							if (Widgets.Num() == PreviousNumWidgets)
+							{
+								const FVector2D ImageSize = VisualDebugIndicatorWidget->GetDesiredSize();
 
-									FOverlayWidgetInfo Info;
-									Info.OverlayOffset = FVector2D(WidgetSize.X - ImageSize.X - 6.f, 6.f);
-									Info.Widget = VisualDebugIndicatorWidget;
+								FOverlayWidgetInfo Info;
+								Info.OverlayOffset = FVector2D(WidgetSize.X - ImageSize.X - 6.f, 6.f);
+								Info.Widget = VisualDebugIndicatorWidget;
 
-									Widgets.Add(Info);
-								}
+								Widgets.Add(Info);
 							}
 						}
 					}
 				}
 			}
+		}
 
-			if(Blueprint.IsValid())
+		if(Blueprint.IsValid())
+		{
+			if(Blueprint->RigGraphDisplaySettings.bShowNodeRunCounts)
 			{
-				if(Blueprint->RigGraphDisplaySettings.bShowNodeRunCounts)
+				if(UControlRig* DebuggedControlRig = Cast<UControlRig>(Blueprint->GetObjectBeingDebugged()))
 				{
-					if(UControlRig* DebuggedControlRig = Cast<UControlRig>(Blueprint->GetObjectBeingDebugged()))
+					const int32 Count = ModelNode->GetInstructionVisitedCount(DebuggedControlRig->GetVM(), FRigVMASTProxy(), false);
+					if(Count > Blueprint->RigGraphDisplaySettings.NodeRunLowerBound)
 					{
-						const int32 Count = ModelNode->GetInstructionVisitedCount(DebuggedControlRig->GetVM(), FRigVMASTProxy(), false);
-						if(Count > Blueprint->RigGraphDisplaySettings.NodeRunLowerBound)
-						{
-							const int32 VOffset = bSelected ? -2 : 2;
-							const FVector2D TextSize = InstructionCountTextBlockWidget->GetDesiredSize();
-							FOverlayWidgetInfo Info;
-							Info.OverlayOffset = FVector2D(WidgetSize.X - TextSize.X - 8.f, VOffset - TextSize.Y);
-							Info.Widget = InstructionCountTextBlockWidget;
-							Widgets.Add(Info);
-						}
+						const int32 VOffset = bSelected ? -2 : 2;
+						const FVector2D TextSize = InstructionCountTextBlockWidget->GetDesiredSize();
+						FOverlayWidgetInfo Info;
+						Info.OverlayOffset = FVector2D(WidgetSize.X - TextSize.X - 8.f, VOffset - TextSize.Y);
+						Info.Widget = InstructionCountTextBlockWidget;
+						Widgets.Add(Info);
 					}
 				}
 			}
@@ -1387,17 +1033,14 @@ FText SControlRigGraphNode::GetInstructionCountText() const
 	{
 		if(Blueprint->RigGraphDisplaySettings.bShowNodeRunCounts)
 		{
-			if (UControlRigGraphNode* RigNode = CastChecked<UControlRigGraphNode>(GraphNode, ECastCheckedType::NullAllowed))
+			if (ModelNode.IsValid())
 			{
-				if (URigVMNode* ModelNode = RigNode->GetModelNode())
+				if(UControlRig* DebuggedControlRig = Cast<UControlRig>(Blueprint->GetObjectBeingDebugged()))
 				{
-					if(UControlRig* DebuggedControlRig = Cast<UControlRig>(Blueprint->GetObjectBeingDebugged()))
+					const int32 Count = ModelNode->GetInstructionVisitedCount(DebuggedControlRig->GetVM(), FRigVMASTProxy(), true);
+					if(Count > Blueprint->RigGraphDisplaySettings.NodeRunLowerBound)
 					{
-						const int32 Count = ModelNode->GetInstructionVisitedCount(DebuggedControlRig->GetVM(), FRigVMASTProxy(), true);
-						if(Count > Blueprint->RigGraphDisplaySettings.NodeRunLowerBound)
-						{
-							return FText::FromString(FString::FromInt(Count));
-						}
+						return FText::FromString(FString::FromInt(Count));
 					}
 				}
 			}
@@ -1416,4 +1059,154 @@ int32 SControlRigGraphNode::GetNodeTopologyVersion() const
 	return INDEX_NONE;
 }
 
+EVisibility SControlRigGraphNode::GetPinVisibility(int32 InPinInfoIndex) const
+{
+	if(PinInfos.IsValidIndex(InPinInfoIndex))
+	{
+		const int32 ParentPinIndex = PinInfos[InPinInfoIndex].ParentIndex;
+		if(ParentPinIndex != INDEX_NONE)
+		{
+			const EVisibility ParentPinVisibility = GetPinVisibility(ParentPinIndex); 
+			if(GetPinVisibility(ParentPinIndex) != EVisibility::Visible)
+			{
+				return ParentPinVisibility;
+			}
+
+			if(!PinInfos[ParentPinIndex].bExpanded)
+			{
+				return EVisibility::Collapsed;
+			}
+		}
+	}
+
+	return EVisibility::Visible;
+}
+
+const FSlateBrush* SControlRigGraphNode::GetExpanderImage(int32 InPinInfoIndex, bool bLeft, bool bHovered) const
+{
+	static const FSlateBrush* ExpandedHoveredLeftBrush = nullptr;
+	static const FSlateBrush* ExpandedHoveredRightBrush = nullptr;
+	static const FSlateBrush* ExpandedLeftBrush = nullptr;
+	static const FSlateBrush* ExpandedRightBrush = nullptr;
+	static const FSlateBrush* CollapsedHoveredLeftBrush = nullptr;
+	static const FSlateBrush* CollapsedHoveredRightBrush = nullptr;
+	static const FSlateBrush* CollapsedLeftBrush = nullptr;
+	static const FSlateBrush* CollapsedRightBrush = nullptr;
+	
+	if(ExpandedHoveredLeftBrush == nullptr)
+	{
+		static const TCHAR* ExpandedHoveredLeftName = TEXT("ControlRig.Node.PinTree.Arrow_Expanded_Hovered_Left");
+		static const TCHAR* ExpandedHoveredRightName = TEXT("ControlRig.Node.PinTree.Arrow_Expanded_Hovered_Right");
+		static const TCHAR* ExpandedLeftName = TEXT("ControlRig.Node.PinTree.Arrow_Expanded_Left");
+		static const TCHAR* ExpandedRightName = TEXT("ControlRig.Node.PinTree.Arrow_Expanded_Right");
+		static const TCHAR* CollapsedHoveredLeftName = TEXT("ControlRig.Node.PinTree.Arrow_Collapsed_Hovered_Left");
+		static const TCHAR* CollapsedHoveredRightName = TEXT("ControlRig.Node.PinTree.Arrow_Collapsed_Hovered_Right");
+		static const TCHAR* CollapsedLeftName = TEXT("ControlRig.Node.PinTree.Arrow_Collapsed_Left");
+		static const TCHAR* CollapsedRightName = TEXT("ControlRig.Node.PinTree.Arrow_Collapsed_Right");
+		
+		ExpandedHoveredLeftBrush = FControlRigEditorStyle::Get().GetBrush(ExpandedHoveredLeftName);
+		ExpandedHoveredRightBrush = FControlRigEditorStyle::Get().GetBrush(ExpandedHoveredRightName);
+		ExpandedLeftBrush = FControlRigEditorStyle::Get().GetBrush(ExpandedLeftName);
+		ExpandedRightBrush = FControlRigEditorStyle::Get().GetBrush(ExpandedRightName);
+		CollapsedHoveredLeftBrush = FControlRigEditorStyle::Get().GetBrush(CollapsedHoveredLeftName);
+		CollapsedHoveredRightBrush = FControlRigEditorStyle::Get().GetBrush(CollapsedHoveredRightName);
+		CollapsedLeftBrush = FControlRigEditorStyle::Get().GetBrush(CollapsedLeftName);
+		CollapsedRightBrush = FControlRigEditorStyle::Get().GetBrush(CollapsedRightName);
+	}
+
+	if(PinInfos[InPinInfoIndex].bExpanded)
+	{
+		if(bHovered)
+		{
+			return bLeft ? ExpandedHoveredLeftBrush : ExpandedHoveredRightBrush;
+		}
+		return bLeft ? ExpandedLeftBrush : ExpandedRightBrush;
+	}
+
+	if(bHovered)
+	{
+		return bLeft ? CollapsedHoveredLeftBrush : CollapsedHoveredRightBrush;
+	}
+	return bLeft ? CollapsedLeftBrush : CollapsedRightBrush;
+}
+
+FReply SControlRigGraphNode::OnExpanderArrowClicked(int32 InPinInfoIndex)
+{
+	if(UControlRigGraphNode* EdGraphNode = Cast<UControlRigGraphNode>(GraphNode))
+	{
+		if(URigVMController* Controller = EdGraphNode->GetController())
+		{
+			const FPinInfo& PinInfo = PinInfos[InPinInfoIndex];
+			TArray<FString> PinPathsToModify;
+			PinPathsToModify.Add(PinInfo.ModelPinPath);
+
+			// with shift clicked we expand recursively
+			if (FSlateApplication::Get().GetModifierKeys().IsShiftDown())
+			{
+				if(URigVMGraph* ModelGraph = EdGraphNode->GetModel())
+				{
+					if(URigVMPin* ModelPin = ModelGraph->FindPin(PinInfo.ModelPinPath))
+					{
+						TArray<URigVMPin*> SubPins = ModelPin->GetNode()->GetAllPinsRecursively();
+						for(URigVMPin* SubPin : SubPins)
+						{
+							if(SubPin->IsInOuter(ModelPin))
+							{
+								PinPathsToModify.Add(SubPin->GetPinPath());
+							}
+						}
+
+						Algo::Reverse(PinPathsToModify);
+					}
+				}
+			}
+
+			Controller->OpenUndoBracket(PinInfo.bExpanded ? TEXT("Collapsing Pin") : TEXT("Expanding Pin"));
+			for(const FString& PinPathToModify : PinPathsToModify)
+			{
+				Controller->SetPinExpansion(PinPathToModify, !PinInfo.bExpanded, true);
+			}
+			Controller->CloseUndoBracket();
+			return FReply::Handled();
+		}
+	}
+	return FReply::Unhandled();
+}
+
+void SControlRigGraphNode::HandleModifiedEvent(ERigVMGraphNotifType InNotifType, URigVMGraph* InGraph,
+	UObject* InSubject)
+{
+	switch(InNotifType)
+	{
+		case ERigVMGraphNotifType::PinExpansionChanged:
+		{
+			if(!ModelNode.IsValid())
+			{
+				return;
+			}
+				
+			if(URigVMPin* Pin = Cast<URigVMPin>(InSubject))
+			{
+				if(Pin->GetNode() == ModelNode.Get())
+				{
+					const FString PinPath = Pin->GetPinPath();
+					for(FPinInfo& PinInfo : PinInfos)
+					{
+						if(PinInfo.ModelPinPath == PinPath)
+						{
+							PinInfo.bExpanded = Pin->IsExpanded();
+							break;
+						}
+					}
+				}
+			}
+			break;
+		}
+		default:
+		{
+			break;
+		}			
+	}
+}
+
 #undef LOCTEXT_NAMESPACE
diff --git a/Engine/Plugins/Experimental/ControlRig/Source/ControlRigEditor/Private/Graph/SControlRigGraphNode.h b/Engine/Plugins/Experimental/ControlRig/Source/ControlRigEditor/Private/Graph/SControlRigGraphNode.h
--- a/Engine/Plugins/Experimental/ControlRig/Source/ControlRigEditor/Private/Graph/SControlRigGraphNode.h
+++ b/Engine/Plugins/Experimental/ControlRig/Source/ControlRigEditor/Private/Graph/SControlRigGraphNode.h
@@ -5,8 +5,6 @@
 #include "CoreMinimal.h"
 #include "Widgets/DeclarativeSyntaxSupport.h"
 #include "SGraphNode.h"
-#include "Widgets/Views/STreeView.h"
-#include "Widgets/Views/SListView.h"
 #include "Widgets/Images/SImage.h"
 #include "Widgets/Text/STextBlock.h"
 #include "RigVMModel/RigVMPin.h"
@@ -17,7 +15,6 @@ class STableViewBase;
 class SOverlay;
 class SGraphPin;
 class UEdGraphPin;
-class SScrollBar;
 
 class SControlRigGraphNode : public SGraphNode
 {
@@ -46,51 +43,30 @@ public:
 	virtual FReply OnMouseButtonDoubleClick(const FGeometry& InMyGeometry, const FPointerEvent& InMouseEvent) override;
 
 	virtual TSharedRef<SWidget> CreateNodeContentArea() override;
-	virtual TSharedPtr<SGraphPin> GetHoveredPin( const FGeometry& MyGeometry, const FPointerEvent& MouseEvent ) const override;
 	virtual void GetOverlayBrushes(bool bSelected, const FVector2D WidgetSize, TArray<FOverlayBrushInfo>& Brushes) const override;
 	virtual void GetNodeInfoPopups(FNodeInfoContext* Context, TArray<FGraphInformationPopupInfo>& Popups) const override;
 	virtual TArray<FOverlayWidgetInfo> GetOverlayWidgets(bool bSelected, const FVector2D& WidgetSize) const override;
 
 	virtual void RefreshErrorInfo() override;
 	virtual void Tick(const FGeometry& AllottedGeometry, const double InCurrentTime, const float InDeltaTime) override;
+	virtual bool IsHidingPinWidgets() const override { return UseLowDetailNodeContent(); }
+	virtual bool UseLowDetailPinNames() const override;
 
 private:
 
-	bool UseLowDetailPinContent() const;
-
 	bool UseLowDetailNodeContent() const;
 
-	TOptional<FVector2D> GetLowDetailDesiredSize() const;
+	FVector2D GetLowDetailDesiredSize() const;
 
 	EVisibility GetTitleVisibility() const;
 
-	EVisibility GetExecutionTreeVisibility() const;
-
-	EVisibility GetInputTreeVisibility() const;
-
-	EVisibility GetInputOutputTreeVisibility() const;
-
-	EVisibility GetOutputTreeVisibility() const;
-
-	EVisibility GetVariableListVisibility() const;
-
-	TSharedRef<ITableRow> MakePinTableRowWidget(URigVMPin* InItem, const TSharedRef<STableViewBase>& OwnerTable);
-	TSharedRef<ITableRow> MakeVariableTableRowWidget(TSharedPtr<FRigVMExternalVariable> InVariable, const TSharedRef<STableViewBase>& OwnerTable);
-
-	void HandleGetChildrenForTree(URigVMPin* InItem, TArray<URigVMPin*>& OutChildren);
-
-	void HandleExpansionChanged(URigVMPin* InItem, bool bExpanded);
-	void HandleExpandRecursively(URigVMPin* InItem, bool bExpanded, TSharedPtr<STreeView<URigVMPin*>>* TreeWidgetPtr);
-
 	FText GetPinLabel(TWeakPtr<SGraphPin> GraphPin) const;
 
 	FSlateColor GetPinTextColor(TWeakPtr<SGraphPin> GraphPin) const;
 	FSlateColor GetVariableLabelTextColor(TWeakObjectPtr<URigVMFunctionReferenceNode> FunctionReferenceNode, FName InVariableName) const;
 	FText GetVariableLabelTooltipText(TWeakObjectPtr<UControlRigBlueprint> InBlueprint, FName InVariableName) const;
 
-	TSharedRef<SWidget> AddContainerPinContent(URigVMPin* InItem, FText InTooltipText);
-
-	FReply HandleAddArrayElement(URigVMPin* InItem);
+	FReply HandleAddArrayElement(FString InModelPinPath);
 
 	void HandleNodeTitleDirtied();
 
@@ -99,34 +75,14 @@ private:
 private:
 
 	int32 GetNodeTopologyVersion() const;
-	
+	EVisibility GetPinVisibility(int32 InPinInfoIndex) const;
+	const FSlateBrush * GetExpanderImage(int32 InPinInfoIndex, bool bLeft, bool bHovered) const;
+	FReply OnExpanderArrowClicked(int32 InPinInfoIndex);
+	void HandleModifiedEvent(ERigVMGraphNotifType InNotifType, URigVMGraph* InGraph, UObject* InSubject);
+
 	/** Cached widget title area */
 	TSharedPtr<SOverlay> TitleAreaWidget;
 
-	/** Widget representing collapsible execution pins */
-	TSharedPtr<STreeView<URigVMPin*>> ExecutionTree;
-
-	/** Widget representing collapsible input pins */
-	TSharedPtr<STreeView<URigVMPin*>> InputTree;
-
-	/** Widget representing collapsible input-output pins */
-	TSharedPtr<STreeView<URigVMPin*>> InputOutputTree;
-
-	/** Widget representing collapsible output pins */
-	TSharedPtr<STreeView<URigVMPin*>> OutputTree;
-
-	/** Widget representing the variable remapping information */
-	TSharedPtr<SListView<TSharedPtr<FRigVMExternalVariable>>> VariableRemappingList;
-
-	/** Dummy scrollbar, as we cant create a tree view without one! */
-	TSharedPtr<SScrollBar> ScrollBar;
-
-	/** Map of pin->widget */
-	TMap<const UEdGraphPin*, TSharedPtr<SGraphPin>> PinWidgetMap;
-
-	/** Map of pin widgets to extra pin widgets */
-	TMap<TSharedRef<SWidget>, TSharedRef<SGraphPin>> ExtraWidgetToPinMap;
-
 	int32 NodeErrorType;
 
 	TSharedPtr<SImage> VisualDebugIndicatorWidget;
@@ -141,4 +97,21 @@ private:
 	TWeakObjectPtr<UControlRigBlueprint> Blueprint;
 
 	FVector2D LastHighDetailSize;
+
+	struct FPinInfo
+	{
+		int32 Index;
+		int32 ParentIndex;
+		bool bHasChildren;
+		bool bHideInputWidget;
+		bool bIsContainer;
+		int32 Depth;
+		FString ModelPinPath;
+		TSharedPtr<SGraphPin> InputPinWidget;
+		TSharedPtr<SGraphPin> OutputPinWidget;
+		bool bExpanded;
+	};
+	
+	TArray<FPinInfo> PinInfos;
+	TWeakObjectPtr<URigVMNode> ModelNode;
 };
diff --git a/Engine/Source/Editor/GraphEditor/Private/SGraphPanel.cpp b/Engine/Source/Editor/GraphEditor/Private/SGraphPanel.cpp
--- a/Engine/Source/Editor/GraphEditor/Private/SGraphPanel.cpp
+++ b/Engine/Source/Editor/GraphEditor/Private/SGraphPanel.cpp
@@ -348,7 +348,7 @@ int32 SGraphPanel::OnPaint( const FPaintArgs& Args, const FGeometry& AllottedGeo
 			TSharedRef<SGraphNode> ChildNode = StaticCastSharedRef<SGraphNode>(Children[ChildIndex]);
 
 			// If this is a culled node, approximate the pin geometry to the corner of the node it is within
-			if (IsNodeCulled(ChildNode, AllottedGeometry))
+			if (IsNodeCulled(ChildNode, AllottedGeometry) || ChildNode->IsHidingPinWidgets())
 			{
 				TArray< TSharedRef<SWidget> > NodePins;
 				ChildNode->GetPins(NodePins);
diff --git a/Engine/Source/Editor/GraphEditor/Private/SGraphPin.cpp b/Engine/Source/Editor/GraphEditor/Private/SGraphPin.cpp
--- a/Engine/Source/Editor/GraphEditor/Private/SGraphPin.cpp
+++ b/Engine/Source/Editor/GraphEditor/Private/SGraphPin.cpp
@@ -1228,14 +1228,19 @@ bool SGraphPin::IsEditingEnabled() const
 bool SGraphPin::UseLowDetailPinNames() const
 {
 	SGraphNode* MyOwnerNode = OwnerNodePtr.Pin().Get();
-	if (MyOwnerNode && MyOwnerNode->GetOwnerPanel().IsValid())
+	if (MyOwnerNode)
 	{
-		return MyOwnerNode->GetOwnerPanel()->GetCurrentLOD() <= EGraphRenderingLOD::LowDetail;
-	}
-	else
-	{
-		return false;
+		if(MyOwnerNode->UseLowDetailPinNames())
+		{
+			return true;
+		}
+
+		if(MyOwnerNode->GetOwnerPanel().IsValid())
+		{
+			return MyOwnerNode->GetOwnerPanel()->GetCurrentLOD() <= EGraphRenderingLOD::LowDetail;
+		}
 	}
+	return false;
 }
 
 EVisibility SGraphPin::GetPinVisiblity() const
diff --git a/Engine/Source/Editor/GraphEditor/Public/SGraphNode.h b/Engine/Source/Editor/GraphEditor/Public/SGraphNode.h
--- a/Engine/Source/Editor/GraphEditor/Public/SGraphNode.h
+++ b/Engine/Source/Editor/GraphEditor/Public/SGraphNode.h
@@ -243,6 +243,10 @@ public:
 	
 	/** Called when ed graph data is cleared, indicating this widget can no longer safely access GraphNode - forwards call to owned pins: */
 	void InvalidateGraphData();
+
+	/** Returns true if the node is hiding its pins */
+	virtual bool IsHidingPinWidgets() const { return false; }
+	
 protected:
 	SGraphNode();
 
@@ -307,6 +311,13 @@ protected:
 	// Should we use low-detail node titles?
 	virtual bool UseLowDetailNodeTitles() const;
 
+public:
+	
+	// Should we use low-detail pin names?
+	virtual bool UseLowDetailPinNames() const { return false; }
+
+protected:
+	
 	/** Return the desired comment bubble color */
 	virtual FSlateColor GetCommentColor() const { return FLinearColor::White; }
 
