Control Rig: Improve performance for URigHierarchy

* improved ::Num()
* improved ::Load()
* improved ::GetChildren / ::UpdateCachedChildren performance
* improved ::CopyHierarchy

diff --git a/Engine/Plugins/Experimental/ControlRig/Source/ControlRig/Private/Rigs/RigHierarchy.cpp b/Engine/Plugins/Experimental/ControlRig/Source/ControlRig/Private/Rigs/RigHierarchy.cpp
--- a/Engine/Plugins/Experimental/ControlRig/Source/ControlRig/Private/Rigs/RigHierarchy.cpp
+++ b/Engine/Plugins/Experimental/ControlRig/Source/ControlRig/Private/Rigs/RigHierarchy.cpp
@@ -105,49 +105,12 @@ void URigHierarchy::Load(FArchive& Ar)
 		FRigElementKey Key;
 		Ar << Key;
 
-		FRigBaseElement* Element = nullptr;
-		switch(Key.Type)
-		{
-			case ERigElementType::Bone:
-			{
-				Element = new FRigBoneElement();
-				break;
-			}
-			case ERigElementType::Null:
-			{
-				Element = new FRigNullElement();
-				break;
-			}
-			case ERigElementType::Control:
-			{
-				Element = new FRigControlElement();
-				break;
-			}
-			case ERigElementType::Curve:
-			{
-				Element = new FRigCurveElement();
-				break;
-			}
-			case ERigElementType::RigidBody:
-			{
-				Element = new FRigRigidBodyElement();
-				break;
-			}
-			case ERigElementType::Socket:
-			{
-				Element = new FRigSocketElement();
-				break;
-			}
-			default:
-			{
-				ensure(false);
-			}
-		}
-
+		FRigBaseElement* Element = MakeElement(Key.Type);
 		check(Element);
 
 		Element->SubIndex = Num(Key.Type);
 		Element->Index = Elements.Add(Element);
+		ElementsPerType[(int32)Key.Type].Add(Element);
 		IndexLookup.Add(Key, Element->Index);
 		
 		Element->Load(Ar, this, FRigBaseElement::StaticData);
@@ -182,45 +145,8 @@ void URigHierarchy::Load(FArchive& Ar)
 		}
 	}
 
-	TArray<int32> CountedAsChild;
-	CountedAsChild.AddZeroed(Elements.Num());
-
-	for (int32 ElementIndex = 0; ElementIndex < Elements.Num(); ElementIndex++)
-	{
-		TArray<FRigBaseElement*> Children = GetChildren(Elements[ElementIndex]);
-		for (FRigBaseElement* Child : Children)
-		{
-			CountedAsChild[Child->Index]++;
-		}
-	}
-
-	for (int32 ElementIndex = 0; ElementIndex < Elements.Num(); ElementIndex++)
-	{
-		FRigBaseElement* Element = Elements[ElementIndex];
-		if(GetNumberOfParents(Element) > 0)
-		{
-			continue;
-		}
-		
-		Traverse(Element, true, [this, &CountedAsChild](FRigBaseElement* InElement, bool& bContinue)
-		{
-
-			int32& Count = CountedAsChild[InElement->Index];
-			Count = FMath::Max(Count - 1, 0);
-
-			// if the count is 0 it means we have hit all parents
-			// and we can finally notify of our existence
-			if(Count == 0)
-			{
-				Notify(ERigHierarchyNotification::ElementAdded, InElement);
-				bContinue = true;
-			}
-			else if(Count > 0)
-			{
-				bContinue = false;
-			}
-		});
-	}
+	UpdateAllCachedChildren();
+	Notify(ERigHierarchyNotification::HierarchyReset, nullptr);
 }
 
 void URigHierarchy::Reset()
@@ -233,6 +159,11 @@ void URigHierarchy::Reset()
 		delete Elements[ElementIndex];
 	}
 	Elements.Reset();
+	ElementsPerType.Reset();
+	for(int32 TypeIndex=0;TypeIndex<(int32)ERigElementType::All;TypeIndex++)
+	{
+		ElementsPerType.Add(TArray<FRigBaseElement*>());
+	}
 	IndexLookup.Reset();
 
 	Notify(ERigHierarchyNotification::HierarchyReset, nullptr);
@@ -244,14 +175,33 @@ void URigHierarchy::CopyHierarchy(URigHierarchy* InHierarchy)
 	
 	Reset();
 
-	FRigVMByteArray ArchiveBytes;
-	FMemoryWriter ArchiveWriter(ArchiveBytes);
-	InHierarchy->Save(ArchiveWriter);
+	for(int32 Index = 0; Index < InHierarchy->Num(); Index++)
+	{
+		FRigBaseElement* Source = InHierarchy->Get(Index);
+		const FRigElementKey& Key = Source->Key;
+
+		FRigBaseElement* Target = MakeElement(Key.Type);
+		
+		Target->Key = Key;
+		Target->SubIndex = Num(Target->GetType());
+		Target->Index = Elements.Add(Target);
+
+		ElementsPerType[(int32)Key.Type].Add(Target);
+		IndexLookup.Add(Key, Target->Index);
 
-	FMemoryReader ArchiveReader(ArchiveBytes);
-	Load(ArchiveReader);
+		check(Source->Index == Index);
+		check(Target->Index == Index);
+	}
+
+	for(int32 Index = 0; Index < InHierarchy->Num(); Index++)
+	{
+		FRigBaseElement* Source = InHierarchy->Get(Index);
+		FRigBaseElement* Target = Elements[Index];
+		Target->CopyFrom(this, Source, InHierarchy);
+	}
 
 	TopologyVersion = InHierarchy->GetTopologyVersion();
+	UpdateAllCachedChildren();
 }
 
 void URigHierarchy::CopyPose(URigHierarchy* InHierarchy, bool bCurrent, bool bInitial)
@@ -363,16 +313,7 @@ void URigHierarchy::ResetCurveValues()
 
 int32 URigHierarchy::Num(ERigElementType InElementType) const
 {
-	int32 Count = 0;
-	for (int32 ElementIndex = 0; ElementIndex < Elements.Num(); ElementIndex++)
-	{
-		FRigBaseElement* Element = Elements[ElementIndex];
-		if(Element->IsTypeOf(InElementType))
-		{
-			Count++;
-		}
-	}
-	return Count;
+	return ElementsPerType[(int32)InElementType].Num();
 }
 
 TArray<FRigBaseElement*> URigHierarchy::GetSelectedElements(ERigElementType InTypeFilter) const
@@ -1799,6 +1740,89 @@ void URigHierarchy::UpdateCachedChildren(const FRigBaseElement* InElement, bool
 	InElement->TopologyVersion = TopologyVersion;
 }
 
+void URigHierarchy::UpdateAllCachedChildren() const
+{
+	TArray<bool> ParentVisited;
+	ParentVisited.AddZeroed(Elements.Num());
+	
+	for (int32 ElementIndex = 0; ElementIndex < Elements.Num(); ElementIndex++)
+	{
+		FRigBaseElement* Element = Elements[ElementIndex];
+		Element->TopologyVersion = TopologyVersion;
+		
+		if(FRigSingleParentElement* SingleParentElement = Cast<FRigSingleParentElement>(Element))
+		{
+			if(FRigTransformElement* ParentElement = SingleParentElement->ParentElement)
+			{
+				if(!ParentVisited[ParentElement->Index])
+				{
+					ParentElement->CachedChildren.Reset();
+					ParentVisited[ParentElement->Index] = true;
+				}
+				ParentElement->CachedChildren.Add(Element);
+			}
+		}
+		else if(FRigMultiParentElement* MultiParentElement = Cast<FRigMultiParentElement>(Element))
+		{
+			for(FRigTransformElement* ParentElement : MultiParentElement->ParentElements)
+			{
+				if(ParentElement)
+				{
+					if(!ParentVisited[ParentElement->Index])
+					{
+						ParentElement->CachedChildren.Reset();
+						ParentVisited[ParentElement->Index] = true;
+					}
+					ParentElement->CachedChildren.Add(Element);
+				}
+			}
+		}
+	}
+}
+
+FRigBaseElement* URigHierarchy::MakeElement(ERigElementType InElementType)
+{
+	FRigBaseElement* Element = nullptr;
+	switch(InElementType)
+	{
+		case ERigElementType::Bone:
+		{
+			Element = new FRigBoneElement();
+			break;
+		}
+		case ERigElementType::Null:
+		{
+			Element = new FRigNullElement();
+			break;
+		}
+		case ERigElementType::Control:
+		{
+			Element = new FRigControlElement();
+			break;
+		}
+		case ERigElementType::Curve:
+		{
+			Element = new FRigCurveElement();
+			break;
+		}
+		case ERigElementType::RigidBody:
+		{
+			Element = new FRigRigidBodyElement();
+			break;
+		}
+		case ERigElementType::Socket:
+		{
+			Element = new FRigSocketElement();
+			break;
+		}
+		default:
+		{
+			ensure(false);
+		}
+	}
+	return Element;
+}
+
 #if URIGHIERARCHY_RECURSIVE_DIRTY_PROPAGATION
 void URigHierarchy::PropagateDirtyFlags(FRigTransformElement* InTransformElement, bool bInitial, bool bAffectChildren, bool bComputeOpposed, bool bMarkDirty) const
 #else
diff --git a/Engine/Plugins/Experimental/ControlRig/Source/ControlRig/Private/Rigs/RigHierarchyController.cpp b/Engine/Plugins/Experimental/ControlRig/Source/ControlRig/Private/Rigs/RigHierarchyController.cpp
--- a/Engine/Plugins/Experimental/ControlRig/Source/ControlRig/Private/Rigs/RigHierarchyController.cpp
+++ b/Engine/Plugins/Experimental/ControlRig/Source/ControlRig/Private/Rigs/RigHierarchyController.cpp
@@ -1142,6 +1142,7 @@ int32 URigHierarchyController::AddElement(FRigBaseElement* InElementToAdd, FRigB
 
 	InElementToAdd->SubIndex = Hierarchy->Num(InElementToAdd->Key.Type);
 	InElementToAdd->Index = Hierarchy->Elements.Add(InElementToAdd);
+	Hierarchy->ElementsPerType[(int32)InElementToAdd->GetKey().Type].Add(InElementToAdd);
 
 	Hierarchy->IndexLookup.Add(InElementToAdd->Key, InElementToAdd->Index);
 	Hierarchy->TopologyVersion++;
@@ -1231,6 +1232,9 @@ bool URigHierarchyController::RemoveElement(FRigBaseElement* InElement)
 	const int32 NumElementsRemoved = Hierarchy->Elements.Remove(InElement);
 	ensure(NumElementsRemoved == 1);
 
+	const int32 NumTypeElementsRemoved = Hierarchy->ElementsPerType[(int32)InElement->GetKey().Type].Remove(InElement);
+	ensure(NumTypeElementsRemoved == 1);
+
 	const int32 NumLookupsRemoved = Hierarchy->IndexLookup.Remove(InElement->Key);
 	ensure(NumLookupsRemoved == 1);
 	for(TPair<FRigElementKey, int32>& Pair : Hierarchy->IndexLookup)
diff --git a/Engine/Plugins/Experimental/ControlRig/Source/ControlRig/Private/Rigs/RigHierarchyElements.cpp b/Engine/Plugins/Experimental/ControlRig/Source/ControlRig/Private/Rigs/RigHierarchyElements.cpp
--- a/Engine/Plugins/Experimental/ControlRig/Source/ControlRig/Private/Rigs/RigHierarchyElements.cpp
+++ b/Engine/Plugins/Experimental/ControlRig/Source/ControlRig/Private/Rigs/RigHierarchyElements.cpp
@@ -137,6 +137,25 @@ void FRigTransformElement::CopyPose(FRigBaseElement* InOther, bool bCurrent, boo
 	}
 }
 
+void FRigTransformElement::CopyFrom(URigHierarchy* InHierarchy, FRigBaseElement* InOther,
+	URigHierarchy* InOtherHierarchy)
+{
+	Super::CopyFrom(InHierarchy, InOther, InOtherHierarchy);
+	
+	const FRigTransformElement* SourceTransform = CastChecked<FRigTransformElement>(InOther);
+	Pose = SourceTransform->Pose;
+
+	ElementsToDirty.Reset();
+	for(int32 ElementToDirtyIndex = 0; ElementToDirtyIndex < SourceTransform->ElementsToDirty.Num(); ElementToDirtyIndex++)
+	{
+		const FElementToDirty& Source = SourceTransform->ElementsToDirty[ElementToDirtyIndex];
+		FRigTransformElement* TargetTransform = CastChecked<FRigTransformElement>(InHierarchy->Get(Source.Element->Index));
+		const FElementToDirty Target(TargetTransform, Source.HierarchyDistance);
+		ElementsToDirty.Add(Target);
+		check(ElementsToDirty[ElementToDirtyIndex].Element->GetKey() == Source.Element->GetKey());
+	}
+}
+
 ////////////////////////////////////////////////////////////////////////////////
 // FRigSingleParentElement
 ////////////////////////////////////////////////////////////////////////////////
@@ -172,6 +191,23 @@ void FRigSingleParentElement::Load(FArchive& Ar, URigHierarchy* Hierarchy, ESeri
 	}
 }
 
+void FRigSingleParentElement::CopyFrom(URigHierarchy* InHierarchy, FRigBaseElement* InOther,
+	URigHierarchy* InOtherHierarchy)
+{
+	Super::CopyFrom(InHierarchy, InOther, InOtherHierarchy);
+
+	const FRigSingleParentElement* Source = CastChecked<FRigSingleParentElement>(InOther); 
+	if(Source->ParentElement)
+	{
+		ParentElement = CastChecked<FRigTransformElement>(InHierarchy->Get(Source->ParentElement->Index));
+		check(ParentElement->GetKey() == Source->ParentElement->GetKey());
+	}
+	else
+	{
+		ParentElement = nullptr;
+	}
+}
+
 ////////////////////////////////////////////////////////////////////////////////
 // FRigMultiParentElement
 ////////////////////////////////////////////////////////////////////////////////
@@ -237,6 +273,25 @@ void FRigMultiParentElement::Load(FArchive& Ar, URigHierarchy* Hierarchy, ESeria
 	}
 }
 
+void FRigMultiParentElement::CopyFrom(URigHierarchy* InHierarchy, FRigBaseElement* InOther,
+	URigHierarchy* InOtherHierarchy)
+{
+	Super::CopyFrom(InHierarchy, InOther, InOtherHierarchy);
+	
+	const FRigMultiParentElement* Source = CastChecked<FRigMultiParentElement>(InOther);
+	Parent = Source->Parent;
+	ParentElements.Reset();
+	ParentWeights = Source->ParentWeights;
+	ParentWeightsInitial = Source->ParentWeightsInitial;
+
+	for(int32 ParentIndex = 0; ParentIndex < Source->ParentElements.Num(); ParentIndex++)
+	{
+		const FRigTransformElement* SourceParentElement = Source->ParentElements[ParentIndex];
+		ParentElements.Add(CastChecked<FRigTransformElement>(InHierarchy->Get(SourceParentElement->Index)));
+		check(ParentElements[ParentIndex]->GetKey() == SourceParentElement->GetKey());
+	}
+}
+
 void FRigMultiParentElement::CopyPose(FRigBaseElement* InOther, bool bCurrent, bool bInitial)
 {
 	Super::CopyPose(InOther, bCurrent, bInitial);
@@ -283,6 +338,14 @@ void FRigBoneElement::Load(FArchive& Ar, URigHierarchy* Hierarchy, ESerializatio
 	}
 }
 
+void FRigBoneElement::CopyFrom(URigHierarchy* InHierarchy, FRigBaseElement* InOther, URigHierarchy* InOtherHierarchy)
+{
+	Super::CopyFrom(InHierarchy, InOther, InOtherHierarchy);
+	
+	const FRigBoneElement* Source = CastChecked<FRigBoneElement>(InOther);
+	BoneType = Source->BoneType;
+}
+
 ////////////////////////////////////////////////////////////////////////////////
 // FRigControlSettings
 ////////////////////////////////////////////////////////////////////////////////
@@ -409,6 +472,16 @@ void FRigControlElement::Load(FArchive& Ar, URigHierarchy* Hierarchy, ESerializa
 	}
 }
 
+void FRigControlElement::CopyFrom(URigHierarchy* InHierarchy, FRigBaseElement* InOther, URigHierarchy* InOtherHierarchy)
+{
+	Super::CopyFrom(InHierarchy, InOther, InOtherHierarchy);
+	
+	const FRigControlElement* Source = CastChecked<FRigControlElement>(InOther);
+	Settings = Source->Settings;
+	Offset = Source->Offset;
+	Gizmo = Source->Gizmo;
+}
+
 void FRigControlElement::CopyPose(FRigBaseElement* InOther, bool bCurrent, bool bInitial)
 {
 	Super::CopyPose(InOther, bCurrent, bInitial);
@@ -462,6 +535,12 @@ void FRigCurveElement::CopyPose(FRigBaseElement* InOther, bool bCurrent, bool bI
 	}
 }
 
+void FRigCurveElement::CopyFrom(URigHierarchy* InHierarchy, FRigBaseElement* InOther, URigHierarchy* InOtherHierarchy)
+{
+	Super::CopyFrom(InHierarchy, InOther, InOtherHierarchy);
+	Value = CastChecked<FRigCurveElement>(InOther)->Value;
+}
+
 ////////////////////////////////////////////////////////////////////////////////
 // FRigRigidBodySettings
 ////////////////////////////////////////////////////////////////////////////////
@@ -505,6 +584,15 @@ void FRigRigidBodyElement::Load(FArchive& Ar, URigHierarchy* Hierarchy, ESeriali
 	}
 }
 
+void FRigRigidBodyElement::CopyFrom(URigHierarchy* InHierarchy, FRigBaseElement* InOther,
+	URigHierarchy* InOtherHierarchy)
+{
+	Super::CopyFrom(InHierarchy, InOther, InOtherHierarchy);
+	
+	const FRigRigidBodyElement* Source = CastChecked<FRigRigidBodyElement>(InOther);
+	Settings = Source->Settings;
+}
+
 ////////////////////////////////////////////////////////////////////////////////
 // FRigSocketElement
 ////////////////////////////////////////////////////////////////////////////////
@@ -519,6 +607,14 @@ void FRigSocketElement::Load(FArchive& Ar, URigHierarchy* Hierarchy, ESerializat
 	Super::Load(Ar, Hierarchy, SerializationPhase);
 }
 
+void FRigSocketElement::CopyFrom(URigHierarchy* InHierarchy, FRigBaseElement* InOther, URigHierarchy* InOtherHierarchy)
+{
+	Super::CopyFrom(InHierarchy, InOther, InOtherHierarchy);
+	
+	const FRigSocketElement* Source = CastChecked<FRigSocketElement>(InOther);
+	GetWorldTransformDelegate = Source->GetWorldTransformDelegate;
+}
+
 FTransform FRigSocketElement::GetSocketWorldTransform(const FRigUnitContext* InContext, bool bInitial) const
 {
 	if(GetWorldTransformDelegate.IsBound())
diff --git a/Engine/Plugins/Experimental/ControlRig/Source/ControlRig/Public/Rigs/RigHierarchy.h b/Engine/Plugins/Experimental/ControlRig/Source/ControlRig/Public/Rigs/RigHierarchy.h
--- a/Engine/Plugins/Experimental/ControlRig/Source/ControlRig/Public/Rigs/RigHierarchy.h
+++ b/Engine/Plugins/Experimental/ControlRig/Source/ControlRig/Public/Rigs/RigHierarchy.h
@@ -32,7 +32,9 @@ public:
 	, bTransactingForTransformChange(false)
 	, bIsInteracting(false)
 	, bSuspendNotifications(false)
-	{}
+	{
+		Reset();
+	}
 
 	virtual ~URigHierarchy();
 
@@ -2023,6 +2025,17 @@ private:
 	 */
 	void UpdateCachedChildren(const FRigBaseElement* InElement, bool bForce = false) const;
 
+	/**
+	* Updates the transient cached children table for all elements if needed (or if bForce == true).
+	* @param bForce If set to true the table will always be updated
+	*/
+	void UpdateAllCachedChildren() const;
+
+	/*
+	 * Helper function to create an element for a given type
+	 */
+	static FRigBaseElement* MakeElement(ERigElementType InElementType); 
+
 	/**
 	 * Marks all affected elements of a given element as dirty
 	 * @param InTransformElement The element that has changed
@@ -2051,6 +2064,9 @@ private:
 	// Storage for the elements
 	mutable TArray<FRigBaseElement*> Elements;
 
+	// Storage for the elements
+	mutable TArray<TArray<FRigBaseElement*>> ElementsPerType;
+
 	// Managed lookup from Key to Index
 	TMap<FRigElementKey, int32> IndexLookup;
 
diff --git a/Engine/Plugins/Experimental/ControlRig/Source/ControlRig/Public/Rigs/RigHierarchyElements.h b/Engine/Plugins/Experimental/ControlRig/Source/ControlRig/Public/Rigs/RigHierarchyElements.h
--- a/Engine/Plugins/Experimental/ControlRig/Source/ControlRig/Public/Rigs/RigHierarchyElements.h
+++ b/Engine/Plugins/Experimental/ControlRig/Source/ControlRig/Public/Rigs/RigHierarchyElements.h
@@ -535,6 +535,9 @@ public:
 
 protected:
 
+	// helper function to be called as part of URigHierarchy::CopyHierarchy
+	virtual void CopyFrom(URigHierarchy* InHierarchy, FRigBaseElement* InOther, URigHierarchy* InOtherHierarchy) {}
+
 	mutable uint16 TopologyVersion;
 	mutable TArray<FRigBaseElement*> CachedChildren;
 
@@ -600,9 +603,11 @@ protected:
 public:
 	
 	virtual void CopyPose(FRigBaseElement* InOther, bool bCurrent, bool bInitial) override;
-
-protected:
 	
+protected:
+
+	virtual void CopyFrom(URigHierarchy* InHierarchy, FRigBaseElement* InOther, URigHierarchy* InOtherHierarchy) override;
+
 	friend struct FRigBaseElement;
 	friend struct FRigSingleParentElement;
 	friend struct FRigMultiParentElement;
@@ -631,7 +636,9 @@ public:
 	FRigTransformElement* ParentElement;
 
 protected:
-	
+
+	virtual void CopyFrom(URigHierarchy* InHierarchy, FRigBaseElement* InOther, URigHierarchy* InOtherHierarchy) override;
+
 	FORCEINLINE static bool IsClassOf(const FRigBaseElement* InElement)
 	{
 		return InElement->GetType() == ERigElementType::Bone ||
@@ -658,7 +665,7 @@ public:
 
 	virtual void Save(FArchive& A, URigHierarchy* Hierarchy, ESerializationPhase SerializationPhase) override;
 	virtual void Load(FArchive& Ar, URigHierarchy* Hierarchy, ESerializationPhase SerializationPhase) override;
-
+	
 	UPROPERTY(BlueprintReadOnly, Category = RigElement)
 	FRigCurrentAndInitialTransform Parent;
 
@@ -669,6 +676,8 @@ public:
 
 protected:
 
+	virtual void CopyFrom(URigHierarchy* InHierarchy, FRigBaseElement* InOther, URigHierarchy* InOtherHierarchy) override;
+
 	FORCEINLINE static bool IsClassOf(const FRigBaseElement* InElement)
 	{
 		return InElement->GetType() == ERigElementType::Null ||
@@ -709,7 +718,9 @@ public:
 	ERigBoneType BoneType;
 
 protected:
-	
+
+	virtual void CopyFrom(URigHierarchy* InHierarchy, FRigBaseElement* InOther, URigHierarchy* InOtherHierarchy) override;
+
 	FORCEINLINE static bool IsClassOf(const FRigBaseElement* InElement)
 	{
 		return InElement->GetType() == ERigElementType::Bone;
@@ -879,6 +890,12 @@ struct CONTROLRIG_API FRigControlElement : public FRigMultiParentElement
 	virtual void Save(FArchive& A, URigHierarchy* Hierarchy, ESerializationPhase SerializationPhase) override;
 	virtual void Load(FArchive& Ar, URigHierarchy* Hierarchy, ESerializationPhase SerializationPhase) override;
 
+private:
+
+	virtual void CopyFrom(URigHierarchy* InHierarchy, FRigBaseElement* InOther, URigHierarchy* InOtherHierarchy) override;
+
+public:
+
 	UPROPERTY(BlueprintReadOnly, EditAnywhere, Category = Control, meta=(ShowOnlyInnerProperties))
 	FRigControlSettings Settings;
 
@@ -924,6 +941,12 @@ public:
 	virtual void Save(FArchive& A, URigHierarchy* Hierarchy, ESerializationPhase SerializationPhase) override;
 	virtual void Load(FArchive& Ar, URigHierarchy* Hierarchy, ESerializationPhase SerializationPhase) override;
 
+private:
+
+	virtual void CopyFrom(URigHierarchy* InHierarchy, FRigBaseElement* InOther, URigHierarchy* InOtherHierarchy) override;
+
+public:
+	
 	float Value;
 	
 	FORCEINLINE static bool IsClassOf(const FRigBaseElement* InElement)
@@ -974,6 +997,12 @@ public:
 	virtual void Save(FArchive& A, URigHierarchy* Hierarchy, ESerializationPhase SerializationPhase) override;
 	virtual void Load(FArchive& Ar, URigHierarchy* Hierarchy, ESerializationPhase SerializationPhase) override;
 
+private:
+
+	virtual void CopyFrom(URigHierarchy* InHierarchy, FRigBaseElement* InOther, URigHierarchy* InOtherHierarchy) override;
+
+public:
+
 	UPROPERTY(BlueprintReadOnly, EditAnywhere, Category = Control, meta=(ShowOnlyInnerProperties))
 	FRigRigidBodySettings Settings;
 
@@ -1014,6 +1043,8 @@ protected:
 
 	FRigSocketGetWorldTransformDelegate GetWorldTransformDelegate;
 
+	virtual void CopyFrom(URigHierarchy* InHierarchy, FRigBaseElement* InOther, URigHierarchy* InOtherHierarchy) override;
+
 	FORCEINLINE static bool IsClassOf(const FRigBaseElement* InElement)
 	{
 		return InElement->GetType() == ERigElementType::Socket;
