Remove calls to deprecated SetAssetEditorModeManager and replace with direct initialization.

#rnx

diff --git a/Engine/Plugins/Experimental/ControlRig/Source/ControlRigEditor/Private/Editor/ControlRigEditor.cpp b/Engine/Plugins/Experimental/ControlRig/Source/ControlRigEditor/Private/Editor/ControlRigEditor.cpp
--- a/Engine/Plugins/Experimental/ControlRig/Source/ControlRigEditor/Private/Editor/ControlRigEditor.cpp
+++ b/Engine/Plugins/Experimental/ControlRig/Source/ControlRigEditor/Private/Editor/ControlRigEditor.cpp
@@ -2277,11 +2277,7 @@ void FControlRigEditor::HandleControlRigExecutionHalted(const int32 InstructionI
 
 void FControlRigEditor::CreateEditorModeManager()
 {
-	TSharedPtr<FAssetEditorModeManager> ModeManager = MakeShareable(FModuleManager::LoadModuleChecked<FPersonaModule>("Persona").CreatePersonaEditorModeManager());
-
-	PRAGMA_DISABLE_DEPRECATION_WARNINGS
-	SetAssetEditorModeManager(ModeManager.Get());
-	PRAGMA_ENABLE_DEPRECATION_WARNINGS
+	EditorModeManager = MakeShareable(FModuleManager::LoadModuleChecked<FPersonaModule>("Persona").CreatePersonaEditorModeManager());
 }
 
 void FControlRigEditor::Tick(float DeltaTime)
diff --git a/Engine/Source/Editor/AnimationBlueprintEditor/Private/AnimationBlueprintEditor.cpp b/Engine/Source/Editor/AnimationBlueprintEditor/Private/AnimationBlueprintEditor.cpp
--- a/Engine/Source/Editor/AnimationBlueprintEditor/Private/AnimationBlueprintEditor.cpp
+++ b/Engine/Source/Editor/AnimationBlueprintEditor/Private/AnimationBlueprintEditor.cpp
@@ -1655,11 +1655,7 @@ void FAnimationBlueprintEditor::OnBlueprintChangedImpl(UBlueprint* InBlueprint,
 
 void FAnimationBlueprintEditor::CreateEditorModeManager()
 {
-	TSharedPtr<FAssetEditorModeManager> ModeManager = MakeShareable(FModuleManager::LoadModuleChecked<FPersonaModule>("Persona").CreatePersonaEditorModeManager());
-
-	PRAGMA_DISABLE_DEPRECATION_WARNINGS
-	SetAssetEditorModeManager(ModeManager.Get());
-	PRAGMA_ENABLE_DEPRECATION_WARNINGS
+	EditorModeManager = MakeShareable(FModuleManager::LoadModuleChecked<FPersonaModule>("Persona").CreatePersonaEditorModeManager());
 }
 
 void FAnimationBlueprintEditor::JumpToHyperlink(const UObject* ObjectReference, bool bRequestRename)
