Sequencer: Control Rig: Change order of how we restore from control rig so we restore control values first and then reset the anim instance and unbind.

diff --git a/Engine/Plugins/Experimental/ControlRig/Source/ControlRig/Private/Sequencer/MovieSceneControlRigParameterTemplate.cpp b/Engine/Plugins/Experimental/ControlRig/Source/ControlRig/Private/Sequencer/MovieSceneControlRigParameterTemplate.cpp
--- a/Engine/Plugins/Experimental/ControlRig/Source/ControlRig/Private/Sequencer/MovieSceneControlRigParameterTemplate.cpp
+++ b/Engine/Plugins/Experimental/ControlRig/Source/ControlRig/Private/Sequencer/MovieSceneControlRigParameterTemplate.cpp
@@ -643,14 +643,8 @@ struct FControlRigParameterPreAnimatedTokenProducer : IMovieScenePreAnimatedToke
 				{
 					if (ControlRig->GetObjectBinding())
 					{
+						//Restore control rig first
 						const bool bSetupUndo = false;
-						if (USkeletalMeshComponent* SkeletalMeshComponent = Cast<USkeletalMeshComponent>(ControlRig->GetObjectBinding()->GetBoundObject()))
-						{
-							SkeletalMeshRestoreState.RestoreState(SkeletalMeshComponent);
-						}
-
-						FControlRigBindingHelper::UnBindFromSequencerInstance(ControlRig);
-						
 						for (TNameAndValue<float>& Value : ScalarValues)
 						{
 							if (ControlRig->FindControl(Value.Name))
@@ -718,11 +712,20 @@ struct FControlRigParameterPreAnimatedTokenProducer : IMovieScenePreAnimatedToke
 								}
 							}
 						}
+						//unbind instances and reset animbp
+						FControlRigBindingHelper::UnBindFromSequencerInstance(ControlRig);
+
+						//restore skel mesh and due a tick
+						if (USkeletalMeshComponent* SkeletalMeshComponent = Cast<USkeletalMeshComponent>(ControlRig->GetObjectBinding()->GetBoundObject()))
+						{
+							SkeletalMeshRestoreState.RestoreState(SkeletalMeshComponent);
+						}
 						//only unbind if not a component
 						if (Cast<UControlRigComponent>(ControlRig->GetObjectBinding()->GetBoundObject()) == nullptr)
 						{
 							ControlRig->GetObjectBinding()->UnbindFromObject();
 						}
+
 					}
 				}
 			}
@@ -1685,4 +1688,3 @@ void FMovieSceneControlRigParameterTemplate::Interrogate(const FMovieSceneContex
 
 	}
 }
-
