Control Rig: Enable wires for hierarchy tree

diff --git a/Engine/Plugins/Experimental/ControlRig/Source/ControlRigEditor/Private/Editor/SRigHierarchy.cpp b/Engine/Plugins/Experimental/ControlRig/Source/ControlRigEditor/Private/Editor/SRigHierarchy.cpp
--- a/Engine/Plugins/Experimental/ControlRig/Source/ControlRigEditor/Private/Editor/SRigHierarchy.cpp
+++ b/Engine/Plugins/Experimental/ControlRig/Source/ControlRigEditor/Private/Editor/SRigHierarchy.cpp
@@ -204,6 +204,7 @@ void SRigHierarchyItem::Construct(const FArguments& InArgs, TSharedPtr<FControlR
 		.OnDragDetected(InHierarchy.Get(), &SRigHierarchy::OnDragDetected)
 		.OnCanAcceptDrop(InHierarchy.Get(), &SRigHierarchy::OnCanAcceptDrop)
 		.OnAcceptDrop(InHierarchy.Get(), &SRigHierarchy::OnAcceptDrop)
+		.ShowWires(true)
 		.Content()
 		[
 			SNew(SHorizontalBox)
