Control Rig: Improve padding on UI nodes to avoid squished nodes

diff --git a/Engine/Plugins/Experimental/ControlRig/Source/ControlRigEditor/Private/Graph/SControlRigGraphNode.cpp b/Engine/Plugins/Experimental/ControlRig/Source/ControlRigEditor/Private/Graph/SControlRigGraphNode.cpp
--- a/Engine/Plugins/Experimental/ControlRig/Source/ControlRigEditor/Private/Graph/SControlRigGraphNode.cpp
+++ b/Engine/Plugins/Experimental/ControlRig/Source/ControlRigEditor/Private/Graph/SControlRigGraphNode.cpp
@@ -46,7 +46,8 @@ const FSlateBrush* SControlRigGraphNode::CachedImg_CR_Pin_Disconnected = nullptr
 
 void SControlRigGraphNode::Construct( const FArguments& InArgs )
 {
-	static const float LeftRightPadding = 6.f;
+	static const float PinWidgetSidePadding = 6.f;
+	static const float EmptySidePadding = 60.f;
 	static const float TopPadding = 2.f; 
 	static const float MaxHeight = 30.f;
 
@@ -315,6 +316,8 @@ void SControlRigGraphNode::Construct( const FArguments& InArgs )
 			{
 				TSharedPtr<SHorizontalBox> SlotLayout;
 				SHorizontalBox::FSlot* FirstSlot = nullptr;
+
+				const float MyEmptySidePadding = PinInfo.bHideInputWidget ? EmptySidePadding : 0.f; 
 				
 				LeftNodeBox->AddSlot()
                 .HAlign(HAlign_Fill)
@@ -329,7 +332,7 @@ void SControlRigGraphNode::Construct( const FArguments& InArgs )
                     .Expose(FirstSlot)
                     .FillWidth(1.f)
                     .HAlign(HAlign_Left)
-                    .Padding(LeftRightPadding, TopPadding, 0.f, 0.f)
+                    .Padding(PinWidgetSidePadding, TopPadding, PinInfo.bIsContainer ? 0.f : MyEmptySidePadding, 0.f)
                     [
                         PinInfo.InputPinWidget.ToSharedRef()
                     ]
@@ -344,7 +347,7 @@ void SControlRigGraphNode::Construct( const FArguments& InArgs )
 					SlotLayout->AddSlot()
 					.AutoWidth()
 					.HAlign(HAlign_Left)
-                    .Padding(LeftRightPadding, TopPadding, 0.f, 0.f)
+                    .Padding(PinWidgetSidePadding, TopPadding, MyEmptySidePadding, 0.f)
 					[
 						SNew(SButton)
 						.ContentPadding(0.0f)
@@ -384,7 +387,7 @@ void SControlRigGraphNode::Construct( const FArguments& InArgs )
                     .FillWidth(1.f)
                     .HAlign(HAlign_Left)
                     .VAlign(VAlign_Center)
-                    .Padding(LeftRightPadding, TopPadding, 0.f, 0.f)
+                    .Padding(PinWidgetSidePadding, TopPadding, 0.f, 0.f)
                     [
                         PinInfo.InputPinWidget.ToSharedRef()
                     ]
@@ -393,7 +396,7 @@ void SControlRigGraphNode::Construct( const FArguments& InArgs )
                     .FillWidth(1.f)
                     .HAlign(HAlign_Right)
 					.VAlign(VAlign_Center)
-                    .Padding(0.f, TopPadding, LeftRightPadding, 0.f)
+                    .Padding(0.f, TopPadding, PinWidgetSidePadding, 0.f)
                     [
                         PinInfo.OutputPinWidget.ToSharedRef()
                     ]
@@ -416,7 +419,7 @@ void SControlRigGraphNode::Construct( const FArguments& InArgs )
 				.FillWidth(1.f)
 				.HAlign(HAlign_Right)
 				.VAlign(VAlign_Center)
-	            .Padding(0.f, TopPadding, LeftRightPadding, 0.f)
+	            .Padding(EmptySidePadding, TopPadding, PinWidgetSidePadding, 0.f)
 				[
 	                PinInfo.OutputPinWidget.ToSharedRef()
 				]
@@ -444,7 +447,7 @@ void SControlRigGraphNode::Construct( const FArguments& InArgs )
 				.AutoWidth()
 				.HAlign(HAlign_Left)
 				.VAlign(VAlign_Center)
-				.Padding(LeftRightPadding, TopPadding, LeftRightPadding, 0.f)
+				.Padding(PinWidgetSidePadding, TopPadding, PinWidgetSidePadding, 0.f)
 				[
 					SNew(STextBlock)
 					.Text(FText::FromName(ExternalVariable->Name))
@@ -457,7 +460,7 @@ void SControlRigGraphNode::Construct( const FArguments& InArgs )
 				.AutoWidth()
 				.HAlign(HAlign_Left)
 				.VAlign(VAlign_Center)
-				.Padding(LeftRightPadding, TopPadding, LeftRightPadding, 0.f)
+				.Padding(PinWidgetSidePadding, TopPadding, PinWidgetSidePadding, 0.f)
 				[
 					SNew(SControlRigVariableBinding)
 					.Blueprint(Blueprint.Get())
