Control Rig: Fix ControlRigComponent mapping between multiple skeletal meshes

diff --git a/Engine/Plugins/Experimental/ControlRig/Source/ControlRig/Private/ControlRigComponent.cpp b/Engine/Plugins/Experimental/ControlRig/Source/ControlRig/Private/ControlRigComponent.cpp
--- a/Engine/Plugins/Experimental/ControlRig/Source/ControlRig/Private/ControlRigComponent.cpp
+++ b/Engine/Plugins/Experimental/ControlRig/Source/ControlRig/Private/ControlRigComponent.cpp
@@ -1087,6 +1087,9 @@ void UControlRigComponent::ValidateMappingData()
 				continue;
 			}
 
+			// cache the scene component also in the override component to avoid on further relying on names
+			MappedElement.ComponentReference.OverrideComponent = MappedElement.SceneComponent;
+
 			if (MappedElement.Direction == EControlRigComponentMapDirection::Output && MappedElement.Weight <= SMALL_NUMBER)
 			{
 				continue;
@@ -1111,11 +1114,7 @@ void UControlRigComponent::ValidateMappingData()
 					{
 						if (USkeleton* Skeleton = SkeletalMesh->GetSkeleton())
 						{
-							if (MappedElement.ElementType == ERigElementType::Bone)
-							{
-								MappedElement.SubIndex = Skeleton->GetReferenceSkeleton().FindBoneIndex(MappedElement.TransformName);
-							}
-							else if (MappedElement.ElementType == ERigElementType::Curve)
+							if (MappedElement.ElementType == ERigElementType::Curve)
 							{
 								const FSmartNameMapping* CurveNameMapping = Skeleton->GetSmartNameContainer(USkeleton::AnimCurveMappingName);
 								if (CurveNameMapping)
@@ -1127,19 +1126,23 @@ void UControlRigComponent::ValidateMappingData()
 									}
 								}
 							}
+							else
+							{
+								MappedElement.SubIndex = Skeleton->GetReferenceSkeleton().FindBoneIndex(MappedElement.TransformName);
+							}
 						}
 						else
 						{
 							ReportError(FString::Printf(TEXT("%s does not have a Skeleton set."), *SkeletalMesh->GetPathName()));
 						}
 					}
-				}
 
-				// if we didn't find the bone, disable this mapped element
-				if (MappedElement.SubIndex == INDEX_NONE)
-				{
-					MappedElement.ElementIndex = INDEX_NONE;
-					continue;
+					// if we didn't find the bone, disable this mapped element
+					if (MappedElement.SubIndex == INDEX_NONE)
+					{
+						MappedElement.ElementIndex = INDEX_NONE;
+						continue;
+					}
 				}
 
 				if (MappedElement.Direction == EControlRigComponentMapDirection::Output)
@@ -1214,7 +1217,7 @@ void UControlRigComponent::TransferInputs()
 				Transform = MappedElement.SceneComponent->GetComponentToWorld();
 			}
 
-			Transform = MappedElement.Offset * Transform;
+			Transform = Transform * MappedElement.Offset;
 
 			ConvertTransformToRigSpace(Transform, MappedElement.Space);
 			ControlRig->GetHierarchy()->SetGlobalTransform(MappedElement.ElementIndex, Transform);
@@ -1259,7 +1262,7 @@ void UControlRigComponent::TransferOutputs()
 				FTransform Transform = ControlRig->GetHierarchy()->GetGlobalTransform(MappedElement.ElementIndex);
 				ConvertTransformFromRigSpace(Transform, MappedElement.Space);
 
-				Transform = MappedElement.Offset * Transform;
+				Transform = Transform * MappedElement.Offset;
 
 				if (MappedElement.SubIndex >= 0)
 				{
