Sequencer: Control Rig: Expose SequencerBindingProxy

diff --git a/Engine/Plugins/Experimental/ControlRig/Source/ControlRigEditor/Public/ControlRigSequencerEditorLibrary.h b/Engine/Plugins/Experimental/ControlRig/Source/ControlRigEditor/Public/ControlRigSequencerEditorLibrary.h
--- a/Engine/Plugins/Experimental/ControlRig/Source/ControlRigEditor/Public/ControlRigSequencerEditorLibrary.h
+++ b/Engine/Plugins/Experimental/ControlRig/Source/ControlRigEditor/Public/ControlRigSequencerEditorLibrary.h
@@ -34,7 +34,7 @@ struct FControlRigSequencerBindingProxy
 		, ControlRig(InControlRig)
 	{}
 
-	UPROPERTY()
+	UPROPERTY(BlueprintReadOnly, Category = ControlRig)
 	FSequencerBindingProxy Proxy;
 
 	UPROPERTY(BlueprintReadOnly, Category = ControlRig)
