Fixed font for SControlRigGraphPinNameList to be default matching rest of pin entries

diff --git a/Engine/Plugins/Experimental/ControlRig/Source/ControlRigEditor/Private/Graph/SControlRigGraphPinNameList.cpp b/Engine/Plugins/Experimental/ControlRig/Source/ControlRigEditor/Private/Graph/SControlRigGraphPinNameList.cpp
--- a/Engine/Plugins/Experimental/ControlRig/Source/ControlRigEditor/Private/Graph/SControlRigGraphPinNameList.cpp
+++ b/Engine/Plugins/Experimental/ControlRig/Source/ControlRigEditor/Private/Graph/SControlRigGraphPinNameList.cpp
@@ -46,6 +46,7 @@ TSharedRef<SWidget>	SControlRigGraphPinNameList::GetDefaultValueWidget()
 					SNew(STextBlock)
 					.Text(this, &SControlRigGraphPinNameList::GetNameListText)
 					.ColorAndOpacity(this, &SControlRigGraphPinNameList::GetNameColor)
+					.Font( FEditorStyle::GetFontStyle( TEXT("PropertyWindow.NormalFont") ) )
 				]
 		];
 }
@@ -107,7 +108,7 @@ FSlateColor SControlRigGraphPinNameList::GetNameColor() const
 
 TSharedRef<SWidget> SControlRigGraphPinNameList::MakeNameListItemWidget(TSharedPtr<FString> InItem)
 {
-	return 	SNew(STextBlock).Text(FText::FromString(*InItem));// .Font(FEditorStyle::GetFontStyle(TEXT("PropertyWindow.NormalFont")));
+	return 	SNew(STextBlock).Text(FText::FromString(*InItem)).Font(FEditorStyle::GetFontStyle(TEXT("PropertyWindow.NormalFont")));
 }
 
 void SControlRigGraphPinNameList::OnNameListChanged(TSharedPtr<FString> NewSelection, ESelectInfo::Type SelectInfo)
