Control Rig: BoneType to be based on user interaction type

diff --git a/Engine/Plugins/Experimental/ControlRig/Source/ControlRig/Private/Rigs/RigHierarchyController.cpp b/Engine/Plugins/Experimental/ControlRig/Source/ControlRig/Private/Rigs/RigHierarchyController.cpp
--- a/Engine/Plugins/Experimental/ControlRig/Source/ControlRig/Private/Rigs/RigHierarchyController.cpp
+++ b/Engine/Plugins/Experimental/ControlRig/Source/ControlRig/Private/Rigs/RigHierarchyController.cpp
@@ -825,7 +825,7 @@ TArray<FRigElementKey> URigHierarchyController::ImportFromText(FString InContent
 			FRigHierarchyContainer OldHierarchy;
 			if(OldHierarchy.ImportFromText(Data).Num() > 0)
 			{
-				return ImportFromHierarchyContainer(OldHierarchy);
+				return ImportFromHierarchyContainer(OldHierarchy, true);
 			}
 		}
 		
@@ -859,6 +859,7 @@ TArray<FRigElementKey> URigHierarchyController::ImportFromText(FString InContent
 			{
 				NewElement = new FRigBoneElement();
 				FRigBoneElement::StaticStruct()->ImportText(*PerElementData.Content, NewElement, nullptr, EPropertyPortFlags::PPF_None, &ErrorPipe, FRigBoneElement::StaticStruct()->GetName(), true);
+				CastChecked<FRigBoneElement>(NewElement)->BoneType = ERigBoneType::User;
 				break;
 			}
 			case ERigElementType::Null:
@@ -993,7 +994,7 @@ TArray<FRigElementKey> URigHierarchyController::ImportFromText(FString InContent
 	return PastedKeys;
 }
 
-TArray<FRigElementKey> URigHierarchyController::ImportFromHierarchyContainer(const FRigHierarchyContainer& InContainer)
+TArray<FRigElementKey> URigHierarchyController::ImportFromHierarchyContainer(const FRigHierarchyContainer& InContainer, bool bIsCopyAndPaste)
 {
 	TMap<FRigElementKey, FRigElementKey> KeyMap;;
 
@@ -1010,7 +1011,7 @@ TArray<FRigElementKey> URigHierarchyController::ImportFromHierarchyContainer(con
 			ParentKey = &OriginalParentKey;
 		}
 
-		const FRigElementKey Key = AddBone(Bone.Name, *ParentKey, Bone.InitialTransform, true, Bone.Type, false);
+		const FRigElementKey Key = AddBone(Bone.Name, *ParentKey, Bone.InitialTransform, true, bIsCopyAndPaste ? ERigBoneType::User : Bone.Type, false);
 		KeyMap.Add(Bone.GetElementKey(), Key);
 	}
 	for(const FRigSpace& Space : InContainer.SpaceHierarchy)
diff --git a/Engine/Plugins/Experimental/ControlRig/Source/ControlRig/Public/Rigs/RigHierarchyController.h b/Engine/Plugins/Experimental/ControlRig/Source/ControlRig/Public/Rigs/RigHierarchyController.h
--- a/Engine/Plugins/Experimental/ControlRig/Source/ControlRig/Public/Rigs/RigHierarchyController.h
+++ b/Engine/Plugins/Experimental/ControlRig/Source/ControlRig/Public/Rigs/RigHierarchyController.h
@@ -330,7 +330,7 @@ public:
 	* This is used for backwards compatbility only during load and does not support undo.
 	* @param InContainer The input hierarchy container
 	*/
-    TArray<FRigElementKey> ImportFromHierarchyContainer(const FRigHierarchyContainer& InContainer);
+    TArray<FRigElementKey> ImportFromHierarchyContainer(const FRigHierarchyContainer& InContainer, bool bIsCopyAndPaste);
 
 	/**
 	 * Removes an existing element from the hierarchy
diff --git a/Engine/Plugins/Experimental/ControlRig/Source/ControlRigDeveloper/Private/ControlRigBlueprint.cpp b/Engine/Plugins/Experimental/ControlRig/Source/ControlRigDeveloper/Private/ControlRigBlueprint.cpp
--- a/Engine/Plugins/Experimental/ControlRig/Source/ControlRigDeveloper/Private/ControlRigBlueprint.cpp
+++ b/Engine/Plugins/Experimental/ControlRig/Source/ControlRigDeveloper/Private/ControlRigBlueprint.cpp
@@ -404,7 +404,7 @@ void UControlRigBlueprint::PostLoad()
 		TGuardValue<bool> SuspendNotifGuard(Hierarchy->GetSuspendNotificationsFlag(), true);
 		
 		Hierarchy->Reset();
-		GetHierarchyController()->ImportFromHierarchyContainer(HierarchyContainer_DEPRECATED);
+		GetHierarchyController()->ImportFromHierarchyContainer(HierarchyContainer_DEPRECATED, false);
 	}
 
 	PropagateHierarchyFromBPToInstances();
