Control Rig: Child controls with illegal characters are not reparented correctly

#UE-114929

diff --git a/Engine/Plugins/Experimental/ControlRig/Source/ControlRig/Private/Rigs/RigHierarchyController.cpp b/Engine/Plugins/Experimental/ControlRig/Source/ControlRig/Private/Rigs/RigHierarchyController.cpp
--- a/Engine/Plugins/Experimental/ControlRig/Source/ControlRig/Private/Rigs/RigHierarchyController.cpp
+++ b/Engine/Plugins/Experimental/ControlRig/Source/ControlRig/Private/Rigs/RigHierarchyController.cpp
@@ -1065,6 +1065,7 @@ TArray<FRigElementKey> URigHierarchyController::ImportFromHierarchyContainer(con
 
 	for(const FRigSpace& Space : InContainer.SpaceHierarchy)
 	{
+		const FRigElementKey SpaceKey = KeyMap.FindChecked(Space.GetElementKey());
 		const FRigElementKey OriginalParentKey = Space.GetParentElementKey();
 		if(OriginalParentKey.IsValid())
 		{
@@ -1073,12 +1074,13 @@ TArray<FRigElementKey> URigHierarchyController::ImportFromHierarchyContainer(con
 			{
 				ParentKey = *ParentKeyPtr;
 			}
-			SetParent(Space.GetElementKey(), ParentKey, false, false);
+			SetParent(SpaceKey, ParentKey, false, false);
 		}
 	}
 
 	for(const FRigControl& Control : InContainer.ControlHierarchy)
 	{
+		const FRigElementKey ControlKey = KeyMap.FindChecked(Control.GetElementKey());
 		FRigElementKey OriginalParentKey = Control.GetParentElementKey();
 		const FRigElementKey SpaceKey = Control.GetSpaceElementKey();
 		OriginalParentKey = SpaceKey.IsValid() ? SpaceKey : OriginalParentKey;
@@ -1089,7 +1091,7 @@ TArray<FRigElementKey> URigHierarchyController::ImportFromHierarchyContainer(con
 			{
 				ParentKey = *ParentKeyPtr;
 			}
-			SetParent(Control.GetElementKey(), ParentKey, false, false);
+			SetParent(ControlKey, ParentKey, false, false);
 		}
 	}
 
