Sequencer: Control Rig: Fix crash duplicaitng a control rig section

diff --git a/Engine/Plugins/Experimental/ControlRig/Source/ControlRig/Private/Sequencer/MovieSceneControlRigParameterSection.cpp b/Engine/Plugins/Experimental/ControlRig/Source/ControlRig/Private/Sequencer/MovieSceneControlRigParameterSection.cpp
--- a/Engine/Plugins/Experimental/ControlRig/Source/ControlRig/Private/Sequencer/MovieSceneControlRigParameterSection.cpp
+++ b/Engine/Plugins/Experimental/ControlRig/Source/ControlRig/Private/Sequencer/MovieSceneControlRigParameterSection.cpp
@@ -1839,7 +1839,7 @@ void UMovieSceneControlRigParameterSection::RecreateWithThisControlRig(UControlR
 void UMovieSceneControlRigParameterSection::SetControlRig(UControlRig* InControlRig)
 {
 	ControlRig = InControlRig;
-	ControlRigClass = ControlRig->GetClass();
+	ControlRigClass = ControlRig ? ControlRig->GetClass() : nullptr;
 }
 
 #if WITH_EDITOR
