Control Rig: Improve zoom level and add ground for collision testing

diff --git a/Engine/Plugins/Experimental/ControlRig/Source/ControlRigEditor/Private/Editor/ControlRigEditor.cpp b/Engine/Plugins/Experimental/ControlRig/Source/ControlRigEditor/Private/Editor/ControlRigEditor.cpp
--- a/Engine/Plugins/Experimental/ControlRig/Source/ControlRigEditor/Private/Editor/ControlRigEditor.cpp
+++ b/Engine/Plugins/Experimental/ControlRig/Source/ControlRigEditor/Private/Editor/ControlRigEditor.cpp
@@ -81,6 +81,10 @@
 #include "IMessageLogListing.h"
 #include "SControlRigFunctionLocalizationWidget.h"
 #include "SControlRigFunctionBulkEditWidget.h"
+#include "SGraphPanel.h"
+#include "Engine/StaticMesh.h"
+#include "Engine/StaticMeshActor.h"
+#include "Components/StaticMeshComponent.h"
 
 #define LOCTEXT_NAMESPACE "ControlRigEditor"
 
@@ -95,6 +99,96 @@ namespace ControlRigEditorTabs
 // 	const FName AdvancedPreviewTab(TEXT("AdvancedPreviewTab"));
 };
 
+struct FControlRigZoomLevelsContainer : public FZoomLevelsContainer
+{
+	struct FControlRigZoomLevelEntry
+	{
+	public:
+		FControlRigZoomLevelEntry(float InZoomAmount, const FText& InDisplayText, EGraphRenderingLOD::Type InLOD)
+			: DisplayText(FText::Format(NSLOCTEXT("GraphEditor", "Zoom", "Zoom {0}"), InDisplayText))
+		, ZoomAmount(InZoomAmount)
+		, LOD(InLOD)
+		{
+		}
+
+	public:
+		FText DisplayText;
+		float ZoomAmount;
+		EGraphRenderingLOD::Type LOD;
+	};
+	
+	FControlRigZoomLevelsContainer()
+	{
+		ZoomLevels.Reserve(22);
+		ZoomLevels.Add(FControlRigZoomLevelEntry(0.025f, FText::FromString(TEXT("-14")), EGraphRenderingLOD::LowestDetail));
+		ZoomLevels.Add(FControlRigZoomLevelEntry(0.070f, FText::FromString(TEXT("-13")), EGraphRenderingLOD::LowestDetail));
+		ZoomLevels.Add(FControlRigZoomLevelEntry(0.100f, FText::FromString(TEXT("-12")), EGraphRenderingLOD::LowestDetail));
+		ZoomLevels.Add(FControlRigZoomLevelEntry(0.125f, FText::FromString(TEXT("-11")), EGraphRenderingLOD::LowestDetail));
+		ZoomLevels.Add(FControlRigZoomLevelEntry(0.150f, FText::FromString(TEXT("-10")), EGraphRenderingLOD::LowestDetail));
+		ZoomLevels.Add(FControlRigZoomLevelEntry(0.175f, FText::FromString(TEXT("-9")), EGraphRenderingLOD::LowestDetail));
+		ZoomLevels.Add(FControlRigZoomLevelEntry(0.200f, FText::FromString(TEXT("-8")), EGraphRenderingLOD::LowestDetail));
+		ZoomLevels.Add(FControlRigZoomLevelEntry(0.225f, FText::FromString(TEXT("-7")), EGraphRenderingLOD::LowDetail));
+		ZoomLevels.Add(FControlRigZoomLevelEntry(0.250f, FText::FromString(TEXT("-6")), EGraphRenderingLOD::LowDetail));
+		ZoomLevels.Add(FControlRigZoomLevelEntry(0.375f, FText::FromString(TEXT("-5")), EGraphRenderingLOD::MediumDetail));
+		ZoomLevels.Add(FControlRigZoomLevelEntry(0.500f, FText::FromString(TEXT("-4")), EGraphRenderingLOD::MediumDetail));
+		ZoomLevels.Add(FControlRigZoomLevelEntry(0.675f, FText::FromString(TEXT("-3")), EGraphRenderingLOD::MediumDetail));
+		ZoomLevels.Add(FControlRigZoomLevelEntry(0.750f, FText::FromString(TEXT("-2")), EGraphRenderingLOD::DefaultDetail));
+		ZoomLevels.Add(FControlRigZoomLevelEntry(0.875f, FText::FromString(TEXT("-1")), EGraphRenderingLOD::DefaultDetail));
+		ZoomLevels.Add(FControlRigZoomLevelEntry(1.000f, FText::FromString(TEXT("1:1")), EGraphRenderingLOD::DefaultDetail));
+		ZoomLevels.Add(FControlRigZoomLevelEntry(1.250f, FText::FromString(TEXT("+1")), EGraphRenderingLOD::DefaultDetail));
+		ZoomLevels.Add(FControlRigZoomLevelEntry(1.375f, FText::FromString(TEXT("+2")), EGraphRenderingLOD::DefaultDetail));
+		ZoomLevels.Add(FControlRigZoomLevelEntry(1.500f, FText::FromString(TEXT("+3")), EGraphRenderingLOD::FullyZoomedIn));
+		ZoomLevels.Add(FControlRigZoomLevelEntry(1.675f, FText::FromString(TEXT("+4")), EGraphRenderingLOD::FullyZoomedIn));
+		ZoomLevels.Add(FControlRigZoomLevelEntry(1.750f, FText::FromString(TEXT("+5")), EGraphRenderingLOD::FullyZoomedIn));
+		ZoomLevels.Add(FControlRigZoomLevelEntry(1.875f, FText::FromString(TEXT("+6")), EGraphRenderingLOD::FullyZoomedIn));
+		ZoomLevels.Add(FControlRigZoomLevelEntry(2.000f, FText::FromString(TEXT("+7")), EGraphRenderingLOD::FullyZoomedIn));
+	}
+
+	float GetZoomAmount(int32 InZoomLevel) const override
+	{
+		checkSlow(ZoomLevels.IsValidIndex(InZoomLevel));
+		return ZoomLevels[InZoomLevel].ZoomAmount;
+	}
+
+	int32 GetNearestZoomLevel(float InZoomAmount) const override
+	{
+		for (int32 ZoomLevelIndex=0; ZoomLevelIndex < GetNumZoomLevels(); ++ZoomLevelIndex)
+		{
+			if (InZoomAmount <= GetZoomAmount(ZoomLevelIndex))
+			{
+				return ZoomLevelIndex;
+			}
+		}
+
+		return GetDefaultZoomLevel();
+	}
+	
+	FText GetZoomText(int32 InZoomLevel) const override
+	{
+		checkSlow(ZoomLevels.IsValidIndex(InZoomLevel));
+		return ZoomLevels[InZoomLevel].DisplayText;
+	}
+	
+	int32 GetNumZoomLevels() const override
+	{
+		return ZoomLevels.Num();
+	}
+	
+	int32 GetDefaultZoomLevel() const override
+	{
+		return 12;
+	}
+
+	EGraphRenderingLOD::Type GetLOD(int32 InZoomLevel) const override
+	{
+		checkSlow(ZoomLevels.IsValidIndex(InZoomLevel));
+		return ZoomLevels[InZoomLevel].LOD;
+	}
+
+	TArray<FControlRigZoomLevelEntry> ZoomLevels;
+};
+
+
 FControlRigEditor::FControlRigEditor()
 	: ControlRig(nullptr)
 	, ActiveController(nullptr)
@@ -1415,6 +1509,13 @@ void FControlRigEditor::OnCreateGraphEditorCommands(TSharedPtr<FUICommandList> G
 
 }
 
+TSharedRef<SGraphEditor> FControlRigEditor::CreateGraphEditorWidget(TSharedRef<FTabInfo> InTabInfo, UEdGraph* InGraph)
+{
+	TSharedRef<SGraphEditor> GraphEditor = FBlueprintEditor::CreateGraphEditorWidget(InTabInfo, InGraph);
+	GraphEditor->GetGraphPanel()->SetZoomLevelsContainer<FControlRigZoomLevelsContainer>();
+	return GraphEditor;
+}
+
 void FControlRigEditor::Compile()
 {
 	{
@@ -2899,6 +3000,24 @@ void FControlRigEditor::HandlePreviewSceneCreated(const TSharedRef<IPersonaPrevi
 {
 	DECLARE_SCOPE_HIERARCHICAL_COUNTER_FUNC()
 
+	// load a ground mesh
+	static const TCHAR* GroundAssetPath = TEXT("/Engine/MapTemplates/SM_Template_Map_Floor.SM_Template_Map_Floor");
+	static const TCHAR* DefaultMaterialPath = TEXT("/Engine/EngineMaterials/DefaultMaterial.DefaultMaterial");
+	UStaticMesh* FloorMesh = Cast<UStaticMesh>(StaticLoadObject(UStaticMesh::StaticClass(), NULL, GroundAssetPath, NULL, LOAD_None, NULL));
+	UMaterial* DefaultMaterial = Cast<UMaterial>(StaticLoadObject(UMaterial::StaticClass(), NULL, DefaultMaterialPath, NULL, LOAD_None, NULL));
+	check(FloorMesh);
+	check(DefaultMaterial);
+
+	// create ground mesh actor
+	AStaticMeshActor* GroundActor = InPersonaPreviewScene->GetWorld()->SpawnActor<AStaticMeshActor>(AStaticMeshActor::StaticClass(), FTransform::Identity);
+	GroundActor->SetFlags(RF_Transient);
+	GroundActor->GetStaticMeshComponent()->SetStaticMesh(FloorMesh);
+	GroundActor->GetStaticMeshComponent()->SetMaterial(0, DefaultMaterial);
+	GroundActor->SetMobility(EComponentMobility::Static);
+	//GroundActor->GetStaticMeshComponent()->SetVisibility(false);
+	GroundActor->GetStaticMeshComponent()->SetCollisionEnabled(ECollisionEnabled::QueryAndPhysics);
+	GroundActor->GetStaticMeshComponent()->SetCollisionProfileName(UCollisionProfile::BlockAll_ProfileName);
+
 	AAnimationEditorPreviewActor* Actor = InPersonaPreviewScene->GetWorld()->SpawnActor<AAnimationEditorPreviewActor>(AAnimationEditorPreviewActor::StaticClass(), FTransform::Identity);
 	Actor->SetFlags(RF_Transient);
 	InPersonaPreviewScene->SetActor(Actor);
diff --git a/Engine/Source/Editor/GraphEditor/Private/SGraphEditorImpl.cpp b/Engine/Source/Editor/GraphEditor/Private/SGraphEditorImpl.cpp
--- a/Engine/Source/Editor/GraphEditor/Private/SGraphEditorImpl.cpp
+++ b/Engine/Source/Editor/GraphEditor/Private/SGraphEditorImpl.cpp
@@ -1644,6 +1644,11 @@ UEdGraphNode* SGraphEditorImpl::GetSingleSelectedNode() const
 	return (SelectedNodes.Num() == 1) ? Cast<UEdGraphNode>(*SelectedNodes.CreateConstIterator()) : nullptr;
 }
 
+SGraphPanel* SGraphEditorImpl::GetGraphPanel() const
+{
+	return GraphPanel.Get();
+}
+
 /////////////////////////////////////////////////////
 
 
diff --git a/Engine/Source/Editor/GraphEditor/Private/SGraphEditorImpl.h b/Engine/Source/Editor/GraphEditor/Private/SGraphEditorImpl.h
--- a/Engine/Source/Editor/GraphEditor/Private/SGraphEditorImpl.h
+++ b/Engine/Source/Editor/GraphEditor/Private/SGraphEditorImpl.h
@@ -330,6 +330,9 @@ public:
 	virtual int32 GetNumberOfSelectedNodes() const override;
 
 	virtual UEdGraphNode* GetSingleSelectedNode() const override;
+
+	virtual SGraphPanel* GetGraphPanel() const override;
+	
 	// End of SGraphEditor interface
 protected:
 	//
diff --git a/Engine/Source/Editor/GraphEditor/Private/SNodePanel.cpp b/Engine/Source/Editor/GraphEditor/Private/SNodePanel.cpp
--- a/Engine/Source/Editor/GraphEditor/Private/SNodePanel.cpp
+++ b/Engine/Source/Editor/GraphEditor/Private/SNodePanel.cpp
@@ -289,7 +289,7 @@ void SNodePanel::Construct()
 {
 	if (!ZoomLevels)
 	{
-		ZoomLevels = MakeUnique<FFixedZoomLevelsContainer>();
+		SetZoomLevelsContainer<FFixedZoomLevelsContainer>();
 	}
 	ZoomLevel = ZoomLevels->GetDefaultZoomLevel();
 	PreviousZoomLevel = ZoomLevels->GetDefaultZoomLevel();
diff --git a/Engine/Source/Editor/GraphEditor/Public/SNodePanel.h b/Engine/Source/Editor/GraphEditor/Public/SNodePanel.h
--- a/Engine/Source/Editor/GraphEditor/Public/SNodePanel.h
+++ b/Engine/Source/Editor/GraphEditor/Public/SNodePanel.h
@@ -833,7 +833,21 @@ protected:
 
 	// Cancels any active zoom-to-fit action
 	void CancelZoomToFit();
+
+public:
+	
+	// Sets the zoom levels container
+	template<typename T>
+	void SetZoomLevelsContainer()
+	{
+		ZoomLevels = MakeUnique<T>();
+		OldZoomAmount = ZoomLevels->GetZoomAmount(ZoomLevel);
+		ZoomLevel = PreviousZoomLevel = ZoomLevels->GetNearestZoomLevel(OldZoomAmount);
+		PostChangedZoom();
+	}
+	
 protected:
+	
 	// The interface for mapping ZoomLevel values to actual node scaling values
 	TUniquePtr<FZoomLevelsContainer> ZoomLevels;
 
diff --git a/Engine/Source/Editor/Kismet/Public/BlueprintEditor.h b/Engine/Source/Editor/Kismet/Public/BlueprintEditor.h
--- a/Engine/Source/Editor/Kismet/Public/BlueprintEditor.h
+++ b/Engine/Source/Editor/Kismet/Public/BlueprintEditor.h
@@ -1165,11 +1165,15 @@ private:
 	/* User wants to edit tunnel via function editor */
 	void OnEditTunnel();
 
+protected:
+	
 	/* Create comment node on graph */
 	virtual void OnCreateComment();
 
 	// Create new graph editor widget for the supplied document container
-	TSharedRef<SGraphEditor> CreateGraphEditorWidget(TSharedRef<class FTabInfo> InTabInfo, class UEdGraph* InGraph);
+	virtual TSharedRef<SGraphEditor> CreateGraphEditorWidget(TSharedRef<class FTabInfo> InTabInfo, class UEdGraph* InGraph);
+
+private:
 
 	/** Helper to move focused graph when clicking on graph breadcrumb */
 	void OnChangeBreadCrumbGraph( class UEdGraph* InGraph);
diff --git a/Engine/Source/Editor/UnrealEd/Public/GraphEditor.h b/Engine/Source/Editor/UnrealEd/Public/GraphEditor.h
--- a/Engine/Source/Editor/UnrealEd/Public/GraphEditor.h
+++ b/Engine/Source/Editor/UnrealEd/Public/GraphEditor.h
@@ -16,6 +16,7 @@
 #include "Engine/LevelStreaming.h"
 
 class UEdGraph;
+class SGraphPanel;
 struct FNotificationInfo;
 struct Rect;
 class FMenuBuilder;
@@ -598,6 +599,17 @@ public:
 	// Returns the first graph editor that is viewing the specified graph
 	UNREALED_API static TSharedPtr<SGraphEditor> FindGraphEditorForGraph(const UEdGraph* Graph);
 
+
+	/** Returns the graph panel used for this graph editor */
+	UNREALED_API virtual SGraphPanel* GetGraphPanel() const
+	{
+		if (Implementation.IsValid())
+		{
+			return Implementation->GetGraphPanel();
+		}
+		return nullptr;
+	}
+
 protected:
 	/** Invoked when the underlying Graph is being changed. */
 	virtual void OnGraphChanged(const struct FEdGraphEditAction& InAction)
