Control Rig: Improve performance of SControlRigStackview label generation

diff --git a/Engine/Plugins/Experimental/ControlRig/Source/ControlRigEditor/Private/Editor/SControlRigStackView.cpp b/Engine/Plugins/Experimental/ControlRig/Source/ControlRigEditor/Private/Editor/SControlRigStackView.cpp
--- a/Engine/Plugins/Experimental/ControlRig/Source/ControlRigEditor/Private/Editor/SControlRigStackView.cpp
+++ b/Engine/Plugins/Experimental/ControlRig/Source/ControlRigEditor/Private/Editor/SControlRigStackView.cpp
@@ -18,6 +18,8 @@
 
 #define LOCTEXT_NAMESPACE "SControlRigStackView"
 
+TAutoConsoleVariable<bool> CVarControlRigExecutionStackDetailedLabels(TEXT("ControlRig.StackDetailedLabels"), false, TEXT("Set to true to turn on detailed labels for the execution stack widget"));
+
 //////////////////////////////////////////////////////////////
 /// FRigStackEntry
 ///////////////////////////////////////////////////////////
@@ -393,6 +395,177 @@ void SControlRigStackView::PopulateStackView(URigVM* InVM)
 		const FRigVMInstructionArray Instructions = InVM->GetInstructions();
 		const FRigVMByteCode& ByteCode = InVM->GetByteCode();
 
+		const bool bUseSimpleLabels = !CVarControlRigExecutionStackDetailedLabels.GetValueOnAnyThread();
+		if(bUseSimpleLabels)
+		{
+			for (int32 InstructionIndex = 0; InstructionIndex < Instructions.Num(); InstructionIndex++)
+			{
+				URigVMNode* Node = Cast<URigVMNode>(ByteCode.GetSubjectForInstruction(InstructionIndex));
+
+				FString DisplayName;
+
+				if(Node)
+				{
+					DisplayName = Node->GetName();
+					
+					// only unit nodes among all nodes has StaticExecute() that generates actual instructions
+					if (URigVMUnitNode* UnitNode = Cast<URigVMUnitNode>(Node))
+					{
+						DisplayName = UnitNode->GetNodeTitle();
+	#if WITH_EDITOR
+						UScriptStruct* Struct = UnitNode->GetScriptStruct();
+						FString MenuDescSuffixMetadata;
+						if (Struct)
+						{
+							Struct->GetStringMetaDataHierarchical(FRigVMStruct::MenuDescSuffixMetaName, &MenuDescSuffixMetadata);
+						}
+						if (!MenuDescSuffixMetadata.IsEmpty())
+						{
+							DisplayName = FString::Printf(TEXT("%s %s"), *UnitNode->GetNodeTitle(), *MenuDescSuffixMetadata);
+						}
+	#endif
+					}
+				}
+
+				FString Label;
+
+				switch (Instructions[InstructionIndex].OpCode)
+				{
+					case ERigVMOpCode::Execute_0_Operands:
+					case ERigVMOpCode::Execute_1_Operands:
+					case ERigVMOpCode::Execute_2_Operands:
+					case ERigVMOpCode::Execute_3_Operands:
+					case ERigVMOpCode::Execute_4_Operands:
+					case ERigVMOpCode::Execute_5_Operands:
+					case ERigVMOpCode::Execute_6_Operands:
+					case ERigVMOpCode::Execute_7_Operands:
+					case ERigVMOpCode::Execute_8_Operands:
+					case ERigVMOpCode::Execute_9_Operands:
+					case ERigVMOpCode::Execute_10_Operands:
+					case ERigVMOpCode::Execute_11_Operands:
+					case ERigVMOpCode::Execute_12_Operands:
+					case ERigVMOpCode::Execute_13_Operands:
+					case ERigVMOpCode::Execute_14_Operands:
+					case ERigVMOpCode::Execute_15_Operands:
+					case ERigVMOpCode::Execute_16_Operands:
+					case ERigVMOpCode::Execute_17_Operands:
+					case ERigVMOpCode::Execute_18_Operands:
+					case ERigVMOpCode::Execute_19_Operands:
+					case ERigVMOpCode::Execute_20_Operands:
+					case ERigVMOpCode::Execute_21_Operands:
+					case ERigVMOpCode::Execute_22_Operands:
+					case ERigVMOpCode::Execute_23_Operands:
+					case ERigVMOpCode::Execute_24_Operands:
+					case ERigVMOpCode::Execute_25_Operands:
+					case ERigVMOpCode::Execute_26_Operands:
+					case ERigVMOpCode::Execute_27_Operands:
+					case ERigVMOpCode::Execute_28_Operands:
+					case ERigVMOpCode::Execute_29_Operands:
+					case ERigVMOpCode::Execute_30_Operands:
+					case ERigVMOpCode::Execute_31_Operands:
+					case ERigVMOpCode::Execute_32_Operands:
+					case ERigVMOpCode::Execute_33_Operands:
+					case ERigVMOpCode::Execute_34_Operands:
+					case ERigVMOpCode::Execute_35_Operands:
+					case ERigVMOpCode::Execute_36_Operands:
+					case ERigVMOpCode::Execute_37_Operands:
+					case ERigVMOpCode::Execute_38_Operands:
+					case ERigVMOpCode::Execute_39_Operands:
+					case ERigVMOpCode::Execute_40_Operands:
+					case ERigVMOpCode::Execute_41_Operands:
+					case ERigVMOpCode::Execute_42_Operands:
+					case ERigVMOpCode::Execute_43_Operands:
+					case ERigVMOpCode::Execute_44_Operands:
+					case ERigVMOpCode::Execute_45_Operands:
+					case ERigVMOpCode::Execute_46_Operands:
+					case ERigVMOpCode::Execute_47_Operands:
+					case ERigVMOpCode::Execute_48_Operands:
+					case ERigVMOpCode::Execute_49_Operands:
+					case ERigVMOpCode::Execute_50_Operands:
+					case ERigVMOpCode::Execute_51_Operands:
+					case ERigVMOpCode::Execute_52_Operands:
+					case ERigVMOpCode::Execute_53_Operands:
+					case ERigVMOpCode::Execute_54_Operands:
+					case ERigVMOpCode::Execute_55_Operands:
+					case ERigVMOpCode::Execute_56_Operands:
+					case ERigVMOpCode::Execute_57_Operands:
+					case ERigVMOpCode::Execute_58_Operands:
+					case ERigVMOpCode::Execute_59_Operands:
+					case ERigVMOpCode::Execute_60_Operands:
+					case ERigVMOpCode::Execute_61_Operands:
+					case ERigVMOpCode::Execute_62_Operands:
+					case ERigVMOpCode::Execute_63_Operands:
+					case ERigVMOpCode::Execute_64_Operands:
+					{
+						const FRigVMExecuteOp& Op = ByteCode.GetOpAt<FRigVMExecuteOp>(Instructions[InstructionIndex]);
+
+						if(Node)
+						{
+							Label = DisplayName;
+						}
+						else
+						{
+							Label = InVM->GetRigVMFunctionName(Op.FunctionIndex);
+						}
+						break;
+					}
+					case ERigVMOpCode::Copy:
+					case ERigVMOpCode::Zero:
+					case ERigVMOpCode::BoolFalse:
+					case ERigVMOpCode::BoolTrue:
+					case ERigVMOpCode::Increment:
+					case ERigVMOpCode::Decrement:
+					case ERigVMOpCode::Equals:
+					case ERigVMOpCode::NotEquals:
+					case ERigVMOpCode::JumpAbsolute:
+					case ERigVMOpCode::JumpForward:
+					case ERigVMOpCode::JumpBackward:
+					case ERigVMOpCode::JumpAbsoluteIf:
+					case ERigVMOpCode::JumpForwardIf:
+					case ERigVMOpCode::JumpBackwardIf:
+					case ERigVMOpCode::BeginBlock:
+					case ERigVMOpCode::EndBlock:
+					case ERigVMOpCode::ChangeType:
+					{
+						const FText OpCodeText = StaticEnum<ERigVMOpCode>()->GetDisplayNameTextByValue((int32)Instructions[InstructionIndex].OpCode);
+						if(Node)
+						{
+							Label = DisplayName + TEXT(" - ") + OpCodeText.ToString();
+						}
+						else
+						{
+							Label = OpCodeText.ToString();
+						}
+						break;
+					}
+					case ERigVMOpCode::Exit:
+					{
+						Label = TEXT("Exit");
+						break;
+					}
+					default:
+					{
+						ensure(false);
+						break;
+					}
+				}
+				
+				// add the entry with the new label to the stack view
+				const FRigVMInstruction& Instruction = Instructions[InstructionIndex];
+				if (FilterText.IsEmpty() || Label.Contains(FilterText.ToString()))
+				{
+					FRigVMASTProxy Proxy;
+					if(const TArray<UObject*>* Callstack = ByteCode.GetCallstackForInstruction(InstructionIndex))
+					{
+						Proxy = FRigVMASTProxy::MakeFromCallstack(Callstack);
+					}
+					TSharedPtr<FRigStackEntry> NewEntry = MakeShared<FRigStackEntry>(InstructionIndex, ERigStackEntry::Operator, InstructionIndex, Instruction.OpCode, Label, Proxy);
+					Operators.Add(NewEntry);
+				}
+			}
+			return;
+		}
+
 		// 1. cache information about instructions/nodes, which will be used later 
 		TMap<FString, FString> OperandFormatMap;
 		for (int32 InstructionIndex = 0; InstructionIndex < Instructions.Num(); InstructionIndex++)
diff --git a/Engine/Source/Developer/RigVMDeveloper/Private/RigVMCompiler/RigVMASTProxy.cpp b/Engine/Source/Developer/RigVMDeveloper/Private/RigVMCompiler/RigVMASTProxy.cpp
--- a/Engine/Source/Developer/RigVMDeveloper/Private/RigVMCompiler/RigVMASTProxy.cpp
+++ b/Engine/Source/Developer/RigVMDeveloper/Private/RigVMCompiler/RigVMASTProxy.cpp
@@ -161,3 +161,17 @@ FRigVMASTProxy FRigVMASTProxy::MakeFromCallPath(const FString& InCallPath, UObje
 	return Proxy;
 }
 
+FRigVMASTProxy FRigVMASTProxy::MakeFromCallstack(const FRigVMCallstack& InCallstack)
+{
+	FRigVMASTProxy Proxy;
+	Proxy.Callstack = InCallstack;
+	return Proxy;
+}
+
+FRigVMASTProxy FRigVMASTProxy::MakeFromCallstack(const TArray<UObject*>* InCallstack)
+{
+	check(InCallstack);
+	FRigVMCallstack Callstack;
+	Callstack.Stack = *InCallstack;
+	return MakeFromCallstack(Callstack);
+}
diff --git a/Engine/Source/Developer/RigVMDeveloper/Public/RigVMCompiler/RigVMASTProxy.h b/Engine/Source/Developer/RigVMDeveloper/Public/RigVMCompiler/RigVMASTProxy.h
--- a/Engine/Source/Developer/RigVMDeveloper/Public/RigVMCompiler/RigVMASTProxy.h
+++ b/Engine/Source/Developer/RigVMDeveloper/Public/RigVMCompiler/RigVMASTProxy.h
@@ -146,6 +146,8 @@ public:
 
 	static FRigVMASTProxy MakeFromUObject(UObject* InSubject);
 	static FRigVMASTProxy MakeFromCallPath(const FString& InCallPath, UObject* InRootObject);
+	static FRigVMASTProxy MakeFromCallstack(const FRigVMCallstack& InCallstack);
+	static FRigVMASTProxy MakeFromCallstack(const TArray<UObject*>* InCallstack);
 
 	FRigVMASTProxy GetSibling(UObject* InSubject) const
 	{
