Control Rig: Improve compilation performance for simple hierarchy changes

diff --git a/Engine/Plugins/Experimental/ControlRig/Source/ControlRig/Private/Rigs/RigHierarchyController.cpp b/Engine/Plugins/Experimental/ControlRig/Source/ControlRig/Private/Rigs/RigHierarchyController.cpp
--- a/Engine/Plugins/Experimental/ControlRig/Source/ControlRig/Private/Rigs/RigHierarchyController.cpp
+++ b/Engine/Plugins/Experimental/ControlRig/Source/ControlRig/Private/Rigs/RigHierarchyController.cpp
@@ -1150,9 +1150,13 @@ int32 URigHierarchyController::AddElement(FRigBaseElement* InElementToAdd, FRigB
 	Hierarchy->IndexLookup.Add(InElementToAdd->Key, InElementToAdd->Index);
 	Hierarchy->TopologyVersion++;
 
-	Notify(ERigHierarchyNotification::ElementAdded, InElementToAdd);
+	{
+		const TGuardValue<bool> Guard(bSuspendNotifications, true);
+		SetParent(InElementToAdd, InFirstParent, bMaintainGlobalTransform);
+	}
 
-	SetParent(InElementToAdd, InFirstParent, bMaintainGlobalTransform);
+	// only notify once at the end
+	Notify(ERigHierarchyNotification::ElementAdded, InElementToAdd);
 
 	return InElementToAdd->Index;
 }
diff --git a/Engine/Plugins/Experimental/ControlRig/Source/ControlRigEditor/Private/Editor/ControlRigEditor.cpp b/Engine/Plugins/Experimental/ControlRig/Source/ControlRigEditor/Private/Editor/ControlRigEditor.cpp
--- a/Engine/Plugins/Experimental/ControlRig/Source/ControlRigEditor/Private/Editor/ControlRigEditor.cpp
+++ b/Engine/Plugins/Experimental/ControlRig/Source/ControlRigEditor/Private/Editor/ControlRigEditor.cpp
@@ -3703,7 +3703,7 @@ void FControlRigEditor::OnHierarchyChanged()
 		{
 			ControlRigBP->Hierarchy->OnModified().Broadcast(ERigHierarchyNotification::ElementSelected, ControlRigBP->Hierarchy, SelectedElement);
 		}
-		GetControlRigBlueprint()->RecompileVM();
+		GetControlRigBlueprint()->RequestAutoVMRecompilation();
 
 		SynchronizeViewportBoneSelection();
 
