ControlRig: Disable colliding against gizmos

diff --git a/Engine/Plugins/Experimental/ControlRig/Source/ControlRig/Private/ControlRigGizmoActor.cpp b/Engine/Plugins/Experimental/ControlRig/Source/ControlRig/Private/ControlRigGizmoActor.cpp
--- a/Engine/Plugins/Experimental/ControlRig/Source/ControlRig/Private/ControlRigGizmoActor.cpp
+++ b/Engine/Plugins/Experimental/ControlRig/Source/ControlRig/Private/ControlRigGizmoActor.cpp
@@ -22,7 +22,7 @@ AControlRigGizmoActor::AControlRigGizmoActor(const FObjectInitializer& ObjectIni
 	StaticMeshComponent->SetCollisionProfileName(UCollisionProfile::NoCollision_ProfileName);
 	StaticMeshComponent->Mobility = EComponentMobility::Movable;
 	StaticMeshComponent->SetGenerateOverlapEvents(false);
-	StaticMeshComponent->bUseDefaultCollision = true;
+	StaticMeshComponent->bUseDefaultCollision = false;
 #if WITH_EDITORONLY_DATA
 	StaticMeshComponent->HitProxyPriority = HPP_Wireframe;
 #endif
