Control Rig: Sync both current and initial offset transforms during setup mode interactions

diff --git a/Engine/Plugins/Experimental/ControlRig/Source/ControlRig/Private/ControlRig.cpp b/Engine/Plugins/Experimental/ControlRig/Source/ControlRig/Private/ControlRig.cpp
--- a/Engine/Plugins/Experimental/ControlRig/Source/ControlRig/Private/ControlRig.cpp
+++ b/Engine/Plugins/Experimental/ControlRig/Source/ControlRig/Private/ControlRig.cpp
@@ -1172,7 +1172,7 @@ FTransform UControlRig::SetupControlFromGlobalTransform(const FName& InControlNa
 			const FTransform ParentTransform = GetHierarchy()->GetParentTransform(ControlElement, ERigTransformType::CurrentGlobal);
 			const FTransform OffsetTransform = InGlobalTransform.GetRelativeTransform(ParentTransform);
 			GetHierarchy()->SetControlOffsetTransform(ControlElement, OffsetTransform, ERigTransformType::InitialLocal, true, true);
-			ControlElement->Offset.Current = ControlElement->Offset.Initial; 
+			GetHierarchy()->SetControlOffsetTransform(ControlElement, OffsetTransform, ERigTransformType::CurrentLocal, true, true);
 		}
 	}
 	return InGlobalTransform;
