Control Rig: Improve selection performance by bracketing notifs for selection

diff --git a/Engine/Plugins/Experimental/ControlRig/Source/ControlRigEditor/Private/Editor/ControlRigEditor.cpp b/Engine/Plugins/Experimental/ControlRig/Source/ControlRigEditor/Private/Editor/ControlRigEditor.cpp
--- a/Engine/Plugins/Experimental/ControlRig/Source/ControlRigEditor/Private/Editor/ControlRigEditor.cpp
+++ b/Engine/Plugins/Experimental/ControlRig/Source/ControlRigEditor/Private/Editor/ControlRigEditor.cpp
@@ -205,6 +205,7 @@ FControlRigEditor::FControlRigEditor()
 	, LastDebuggedRig()
 	, RigHierarchyTabCount(0)
 	, HaltedAtNode(nullptr)
+	, bSuspendDetailsPanelRefresh(false)
 {
 }
 
@@ -1254,6 +1255,11 @@ UBlueprint* FControlRigEditor::GetBlueprintObj() const
 
 void FControlRigEditor::SetDetailObjects(const TArray<UObject*>& InObjects)
 {
+	if(bSuspendDetailsPanelRefresh)
+	{
+		return;
+	}
+
 	ClearDetailObject();
 
 	if (InObjects.Num() == 1)
@@ -1279,6 +1285,11 @@ void FControlRigEditor::SetDetailObjects(const TArray<UObject*>& InObjects)
 
 void FControlRigEditor::SetDetailObject(UObject* Obj)
 {
+	if(bSuspendDetailsPanelRefresh)
+	{
+		return;
+	}
+
 	if (URigVMUnitNode* UnitNode = Cast<URigVMUnitNode>(Obj))
 	{
 		ClearDetailObject();
@@ -1400,6 +1411,11 @@ void FControlRigEditor::SetDetailObject(UObject* Obj)
 
 void FControlRigEditor::SetDetailStruct(const FRigElementKey& InElement)
 {
+	if(bSuspendDetailsPanelRefresh)
+	{
+		return;
+	}
+
 	if (RigElementInDetailPanel == InElement)
 	{
 		return;
@@ -1472,6 +1488,11 @@ void FControlRigEditor::SetDetailStruct(const FRigElementKey& InElement)
 
 void FControlRigEditor::ClearDetailObject()
 {
+	if(bSuspendDetailsPanelRefresh)
+	{
+		return;
+	}
+
 	RigElementInDetailPanel = FRigElementKey();
 
 	if (NodeDetailBuffer.Num() > 0 && NodeDetailStruct != nullptr)
diff --git a/Engine/Plugins/Experimental/ControlRig/Source/ControlRigEditor/Private/Editor/ControlRigEditor.h b/Engine/Plugins/Experimental/ControlRig/Source/ControlRigEditor/Private/Editor/ControlRigEditor.h
--- a/Engine/Plugins/Experimental/ControlRig/Source/ControlRigEditor/Private/Editor/ControlRigEditor.h
+++ b/Engine/Plugins/Experimental/ControlRig/Source/ControlRigEditor/Private/Editor/ControlRigEditor.h
@@ -432,6 +432,8 @@ protected:
 
 	URigVMNode* HaltedAtNode;
 
+	bool bSuspendDetailsPanelRefresh;
+
 	friend class FControlRigEditorMode;
 	friend class SControlRigStackView;
 	friend class SRigHierarchy;
diff --git a/Engine/Plugins/Experimental/ControlRig/Source/ControlRigEditor/Private/Editor/SRigHierarchy.cpp b/Engine/Plugins/Experimental/ControlRig/Source/ControlRigEditor/Private/Editor/SRigHierarchy.cpp
--- a/Engine/Plugins/Experimental/ControlRig/Source/ControlRigEditor/Private/Editor/SRigHierarchy.cpp
+++ b/Engine/Plugins/Experimental/ControlRig/Source/ControlRigEditor/Private/Editor/SRigHierarchy.cpp
@@ -750,7 +750,10 @@ TArray<FRigElementKey> SRigHierarchy::GetSelectedKeys() const
 	TArray<FRigElementKey> SelectedKeys;
 	for (const TSharedPtr<FRigTreeElement>& SelectedItem : SelectedItems)
 	{
-		SelectedKeys.AddUnique(SelectedItem->Key);
+		if(SelectedItem->Key.IsValid())
+		{
+			SelectedKeys.AddUnique(SelectedItem->Key);
+		}
 	}
 
 	return SelectedKeys;
@@ -795,6 +798,9 @@ void SRigHierarchy::OnSelectionChanged(TSharedPtr<FRigTreeElement> Selection, ES
 		return;
 	}
 
+	// an element to use for the control rig editor's detail panel
+	FRigElementKey LastSelectedElement;
+
 	URigHierarchy* Hierarchy = GetHierarchy();
 	if (Hierarchy)
 	{
@@ -803,6 +809,17 @@ void SRigHierarchy::OnSelectionChanged(TSharedPtr<FRigTreeElement> Selection, ES
 		
 		TGuardValue<bool> GuardRigHierarchyChanges(bIsChangingRigHierarchy, true);
 
+		// flag to guard during selection changes.
+		// in case there's no editor we'll use the local variable.
+		bool bDummySuspensionFlag = false;
+		bool* SuspensionFlagPtr = &bDummySuspensionFlag;
+		if (ControlRigEditor.IsValid())
+		{
+			SuspensionFlagPtr = &ControlRigEditor.Pin()->bSuspendDetailsPanelRefresh;
+		}
+
+		TGuardValue<bool> SuspendDetailsPanelRefreshGuard(*SuspensionFlagPtr, true);
+		
 		const TArray<FRigElementKey> NewSelection = GetSelectedKeys();
 		if(!Controller->SetSelection(NewSelection))
 		{
@@ -818,6 +835,20 @@ void SRigHierarchy::OnSelectionChanged(TSharedPtr<FRigTreeElement> Selection, ES
 					HandleControlBoneOrSpaceTransform();
 				}
 			}
+
+			LastSelectedElement = NewSelection.Last();
+		}
+	}
+
+	if (ControlRigEditor.IsValid())
+	{
+		if(LastSelectedElement.IsValid())
+		{
+			ControlRigEditor.Pin()->SetDetailStruct(LastSelectedElement);
+		}
+		else
+		{
+			ControlRigEditor.Pin()->ClearDetailObject();
 		}
 	}
 }
