Control Rig: Fix default zoom level limit

diff --git a/Engine/Plugins/Experimental/ControlRig/Source/ControlRigEditor/Private/Editor/ControlRigEditor.cpp b/Engine/Plugins/Experimental/ControlRig/Source/ControlRigEditor/Private/Editor/ControlRigEditor.cpp
--- a/Engine/Plugins/Experimental/ControlRig/Source/ControlRigEditor/Private/Editor/ControlRigEditor.cpp
+++ b/Engine/Plugins/Experimental/ControlRig/Source/ControlRigEditor/Private/Editor/ControlRigEditor.cpp
@@ -176,7 +176,7 @@ struct FControlRigZoomLevelsContainer : public FZoomLevelsContainer
 	
 	int32 GetDefaultZoomLevel() const override
 	{
-		return 12;
+		return 14;
 	}
 
 	EGraphRenderingLOD::Type GetLOD(int32 InZoomLevel) const override
