ControlRig: Closing Mirror Menu still applies mirror options

diff --git a/Engine/Plugins/Experimental/ControlRig/Source/ControlRigEditor/Private/Editor/SRigHierarchy.cpp b/Engine/Plugins/Experimental/ControlRig/Source/ControlRigEditor/Private/Editor/SRigHierarchy.cpp
--- a/Engine/Plugins/Experimental/ControlRig/Source/ControlRigEditor/Private/Editor/SRigHierarchy.cpp
+++ b/Engine/Plugins/Experimental/ControlRig/Source/ControlRigEditor/Private/Editor/SRigHierarchy.cpp
@@ -36,6 +36,7 @@
 #include "IPersonaToolkit.h"
 #include "SKismetInspector.h"
 #include "Types/WidgetActiveTimerDelegate.h"
+#include "Dialogs/CustomDialog.h"
 
 #define LOCTEXT_NAMESPACE "SRigHierarchy"
 
@@ -1860,20 +1861,29 @@ void SRigHierarchy::HandleMirrorItem()
 		TSharedRef<SKismetInspector> KismetInspector = SNew(SKismetInspector);
 		KismetInspector->ShowSingleStruct(StructToDisplay);
 
-		SGenericDialogWidget::OpenDialog(LOCTEXT("ControlRigHierarchyMirror", "Mirror Hierarchy"), KismetInspector, SGenericDialogWidget::FArguments(), true);
+		TSharedRef<SCustomDialog> MirrorDialog = SNew(SCustomDialog)
+			.Title(FText(LOCTEXT("ControlRigHierarchyMirror", "Mirror Hierarchy")))
+			.DialogContent( KismetInspector)
+			.Buttons({
+				SCustomDialog::FButton(LOCTEXT("OK", "OK")),
+				SCustomDialog::FButton(LOCTEXT("Cancel", "Cancel"))
+		});
 
-		ClearDetailPanel();
+		if (MirrorDialog->ShowModal() == 0)
 		{
-			TGuardValue<bool> GuardRigHierarchyChanges(bIsChangingRigHierarchy, true);
-			TGuardValue<bool> SuspendBlueprintNotifs(ControlRigBlueprint->bSuspendAllNotifications, true);
+			ClearDetailPanel();
+			{
+				TGuardValue<bool> GuardRigHierarchyChanges(bIsChangingRigHierarchy, true);
+				TGuardValue<bool> SuspendBlueprintNotifs(ControlRigBlueprint->bSuspendAllNotifications, true);
 
-			FScopedTransaction Transaction(LOCTEXT("HierarchyTreeMirrorSelected", "Mirror selected items from hierarchy"));
+				FScopedTransaction Transaction(LOCTEXT("HierarchyTreeMirrorSelected", "Mirror selected items from hierarchy"));
 
-			const TArray<FRigElementKey> KeysToMirror = GetSelectedKeys();
-			const TArray<FRigElementKey> KeysToDuplicate = GetSelectedKeys();
-			Controller->MirrorElements(KeysToDuplicate, Settings, true, true);
+				const TArray<FRigElementKey> KeysToMirror = GetSelectedKeys();
+				const TArray<FRigElementKey> KeysToDuplicate = GetSelectedKeys();
+				Controller->MirrorElements(KeysToDuplicate, Settings, true, true);
+			}
+			ControlRigBlueprint->PropagateHierarchyFromBPToInstances();
 		}
-		ControlRigBlueprint->PropagateHierarchyFromBPToInstances();
 	}
 
 	FSlateApplication::Get().DismissAllMenus();
