Control Rig: Avoid clearing the kismet inspector when navigating the graph

diff --git a/Engine/Plugins/Experimental/ControlRig/Source/ControlRigEditor/Private/Editor/ControlRigEditor.cpp b/Engine/Plugins/Experimental/ControlRig/Source/ControlRigEditor/Private/Editor/ControlRigEditor.cpp
--- a/Engine/Plugins/Experimental/ControlRig/Source/ControlRigEditor/Private/Editor/ControlRigEditor.cpp
+++ b/Engine/Plugins/Experimental/ControlRig/Source/ControlRigEditor/Private/Editor/ControlRigEditor.cpp
@@ -3299,6 +3299,11 @@ void FControlRigEditor::SetupGraphEditorEvents(UEdGraph* InGraph, SGraphEditor::
 	InEvents.OnTextCommitted = FOnNodeTextCommitted::CreateSP(this, &FControlRigEditor::OnNodeTitleCommitted);
 }
 
+void FControlRigEditor::FocusInspectorOnGraphSelection(const TSet<UObject*>& NewSelection, bool bForceRefresh)
+{
+	// nothing to do here for control rig
+}
+
 FActionMenuContent FControlRigEditor::HandleCreateGraphActionMenu(UEdGraph* InGraph, const FVector2D& InNodePosition, const TArray<UEdGraphPin*>& InDraggedPins, bool bAutoExpand, SGraphEditor::FActionMenuClosed InOnMenuClosed)
 {
 	return FBlueprintEditor::OnCreateGraphActionMenu(InGraph, InNodePosition, InDraggedPins, bAutoExpand, InOnMenuClosed);
diff --git a/Engine/Plugins/Experimental/ControlRig/Source/ControlRigEditor/Private/Editor/ControlRigEditor.h b/Engine/Plugins/Experimental/ControlRig/Source/ControlRigEditor/Private/Editor/ControlRigEditor.h
--- a/Engine/Plugins/Experimental/ControlRig/Source/ControlRigEditor/Private/Editor/ControlRigEditor.h
+++ b/Engine/Plugins/Experimental/ControlRig/Source/ControlRigEditor/Private/Editor/ControlRigEditor.h
@@ -210,6 +210,7 @@ protected:
 	virtual void OnSelectedNodesChangedImpl(const TSet<class UObject*>& NewSelection) override;
 	virtual void OnBlueprintChangedImpl(UBlueprint* InBlueprint, bool bIsJustBeingCompiled) override;
 	virtual void SetupGraphEditorEvents(UEdGraph* InGraph, SGraphEditor::FGraphEditorEvents& InEvents) override;
+	virtual void FocusInspectorOnGraphSelection(const TSet<class UObject*>& NewSelection, bool bForceRefresh = false) override;
 
 	void HandleModifiedEvent(ERigVMGraphNotifType InNotifType, URigVMGraph* InGraph, UObject* InSubject);
 	void HandleVMCompiledEvent(UBlueprint* InBlueprint, URigVM* InVM);
diff --git a/Engine/Source/Editor/Kismet/Public/BlueprintEditor.h b/Engine/Source/Editor/Kismet/Public/BlueprintEditor.h
--- a/Engine/Source/Editor/Kismet/Public/BlueprintEditor.h
+++ b/Engine/Source/Editor/Kismet/Public/BlueprintEditor.h
@@ -458,7 +458,7 @@ public:
 	/** Called when a token in a log message is clicked */
 	void LogToken_OnClicked(const class IMessageToken& Token);
 
-	void FocusInspectorOnGraphSelection(const TSet<class UObject*>& NewSelection, bool bForceRefresh = false);
+	virtual void FocusInspectorOnGraphSelection(const TSet<class UObject*>& NewSelection, bool bForceRefresh = false);
 
 	/** Variable list window calls this after it is updated */
 	void VariableListWasUpdated();
