Control Rig: Add option for bone radius drawing

diff --git a/Engine/Plugins/Experimental/ControlRig/Source/ControlRig/Private/ControlRig.cpp b/Engine/Plugins/Experimental/ControlRig/Source/ControlRig/Private/ControlRig.cpp
--- a/Engine/Plugins/Experimental/ControlRig/Source/ControlRig/Private/ControlRig.cpp
+++ b/Engine/Plugins/Experimental/ControlRig/Source/ControlRig/Private/ControlRig.cpp
@@ -73,6 +73,7 @@ UControlRig::UControlRig(const FObjectInitializer& ObjectInitializer)
 #if WITH_EDITORONLY_DATA
 	, VMSnapshotBeforeExecution(nullptr)
 #endif
+	, DebugBoneRadiusMultiplier(1.f)
 {
 	SetVM(ObjectInitializer.CreateDefaultSubobject<URigVM>(this, TEXT("VM")));
 	DynamicHierarchy = ObjectInitializer.CreateDefaultSubobject<URigHierarchy>(this, TEXT("DynamicHierarchy"));
diff --git a/Engine/Plugins/Experimental/ControlRig/Source/ControlRig/Private/ControlRigComponent.cpp b/Engine/Plugins/Experimental/ControlRig/Source/ControlRig/Private/ControlRigComponent.cpp
--- a/Engine/Plugins/Experimental/ControlRig/Source/ControlRig/Private/ControlRigComponent.cpp
+++ b/Engine/Plugins/Experimental/ControlRig/Source/ControlRig/Private/ControlRigComponent.cpp
@@ -1645,11 +1645,12 @@ void FControlRigSceneProxy::GetDynamicMeshElements(const TArray<const FSceneView
 
 			if (bShouldDrawBones)
 			{
+				const float RadiusMultiplier = ControlRigComponent->ControlRig->GetDebugBoneRadiusMultiplier();
 				const FTransform Transform = ControlRigComponent->GetComponentToWorld();
 				const float MaxDrawRadius = ControlRigComponent->Bounds.SphereRadius * 0.02f;
 
 				URigHierarchy* Hierarchy = ControlRigComponent->ControlRig->GetHierarchy();
-				Hierarchy->ForEach<FRigBoneElement>([PDI, Hierarchy, Transform, MaxDrawRadius](FRigBoneElement* BoneElement) -> bool
+				Hierarchy->ForEach<FRigBoneElement>([PDI, Hierarchy, Transform, MaxDrawRadius, RadiusMultiplier](FRigBoneElement* BoneElement) -> bool
                 {
                     const int32 ParentIndex = Hierarchy->GetFirstParent(BoneElement->GetIndex());
 					const FLinearColor LineColor = FLinearColor::White;
@@ -1671,7 +1672,7 @@ void FControlRigSceneProxy::GetDynamicMeshElements(const TArray<const FSceneView
 
 					const float BoneLength = (End - Start).Size();
 					// clamp by bound, we don't want too long or big
-					const float Radius = FMath::Clamp(BoneLength * 0.05f, 0.1f, MaxDrawRadius);
+					const float Radius = FMath::Clamp(BoneLength * 0.05f, 0.1f, MaxDrawRadius) * RadiusMultiplier;
 
 					//Render Sphere for bone end point and a cone between it and its parent.
 					SkeletalDebugRendering::DrawWireBone(PDI, Start, End, LineColor, SDPG_Foreground, Radius);
diff --git a/Engine/Plugins/Experimental/ControlRig/Source/ControlRig/Public/ControlRig.h b/Engine/Plugins/Experimental/ControlRig/Source/ControlRig/Public/ControlRig.h
--- a/Engine/Plugins/Experimental/ControlRig/Source/ControlRig/Public/ControlRig.h
+++ b/Engine/Plugins/Experimental/ControlRig/Source/ControlRig/Public/ControlRig.h
@@ -552,6 +552,8 @@ public:
 	virtual void PostEditChangeProperty(FPropertyChangedEvent& PropertyChangedEvent) override;
 #endif
 
+	float GetDebugBoneRadiusMultiplier() const { return DebugBoneRadiusMultiplier; }
+
 public:
 	//~ Begin IInterface_AssetUserData Interface
 	virtual void AddAssetUserData(UAssetUserData* InUserData) override;
@@ -650,6 +652,8 @@ private:
 	bool bIsInDebugMode;
 
 #endif
+
+	float DebugBoneRadiusMultiplier;
 	
 #if WITH_EDITOR	
 
diff --git a/Engine/Plugins/Experimental/ControlRig/Source/ControlRigDeveloper/Private/ControlRigBlueprint.cpp b/Engine/Plugins/Experimental/ControlRig/Source/ControlRigDeveloper/Private/ControlRigBlueprint.cpp
--- a/Engine/Plugins/Experimental/ControlRig/Source/ControlRigDeveloper/Private/ControlRigBlueprint.cpp
+++ b/Engine/Plugins/Experimental/ControlRig/Source/ControlRigDeveloper/Private/ControlRigBlueprint.cpp
@@ -87,6 +87,8 @@ UControlRigBlueprint::UControlRigBlueprint(const FObjectInitializer& ObjectIniti
 
 	Validator = ObjectInitializer.CreateDefaultSubobject<UControlRigValidator>(this, TEXT("ControlRigValidator"));
 
+	DebugBoneRadius = 1.f;
+	
 	bDirtyDuringLoad = false;
 	bErrorsDuringCompilation = false;
 
diff --git a/Engine/Plugins/Experimental/ControlRig/Source/ControlRigDeveloper/Public/ControlRigBlueprint.h b/Engine/Plugins/Experimental/ControlRig/Source/ControlRigDeveloper/Public/ControlRigBlueprint.h
--- a/Engine/Plugins/Experimental/ControlRig/Source/ControlRigDeveloper/Public/ControlRigBlueprint.h
+++ b/Engine/Plugins/Experimental/ControlRig/Source/ControlRigDeveloper/Public/ControlRigBlueprint.h
@@ -491,6 +491,9 @@ private:
 		TMap<FName, FRigControlValue> ControlValues;
 	};
 
+	UPROPERTY()
+	float DebugBoneRadius;
+
 #if WITH_EDITOR
 
 public:
@@ -594,4 +597,5 @@ private:
 	friend class UControlRigThumbnailRenderer;
 	friend class FControlRigGraphDetails;
 	friend class FControlRigEditorModule;
+	friend class UControlRigComponent;
 };
diff --git a/Engine/Plugins/Experimental/ControlRig/Source/ControlRigEditor/Private/Editor/ControlRigEditor.cpp b/Engine/Plugins/Experimental/ControlRig/Source/ControlRigEditor/Private/Editor/ControlRigEditor.cpp
--- a/Engine/Plugins/Experimental/ControlRig/Source/ControlRigEditor/Private/Editor/ControlRigEditor.cpp
+++ b/Engine/Plugins/Experimental/ControlRig/Source/ControlRigEditor/Private/Editor/ControlRigEditor.cpp
@@ -490,6 +490,8 @@ void FControlRigEditor::InitControlRigEditor(const EToolkitMode::Type Mode, cons
 			Compile();
 		}
 	}
+	
+	OnToolbarBoneRadiusChanged(InControlRigBlueprint->DebugBoneRadius);
 }
 
 void FControlRigEditor::BindCommands()
@@ -1182,6 +1184,11 @@ void FControlRigEditor::HandleSetObjectBeingDebugged(UObject* InObject)
 		DebuggedControlRig->ControlRigLog = &ControlRigLog;
 		DebuggedControlRig->DynamicHierarchy->HierarchyForSelectionPtr = GetControlRigBlueprint()->Hierarchy;
 
+		if (UControlRigBlueprint* RigBlueprint = Cast<UControlRigBlueprint>(GetBlueprintObj()))
+		{
+			DebuggedControlRig->DebugBoneRadiusMultiplier = RigBlueprint->DebugBoneRadius;
+		}
+
 		UControlRigSkeletalMeshComponent* EditorSkelComp = Cast<UControlRigSkeletalMeshComponent>(GetPersonaToolkit()->GetPreviewScene()->GetPreviewMeshComponent());
 		if (EditorSkelComp)
 		{
@@ -2795,6 +2802,27 @@ void FControlRigEditor::HandleViewportCreated(const TSharedRef<class IPersonaVie
 					LOCTEXT("ControlRigAxesScale", "Axes Scale")
 				);
 
+				InMenuBuilder.AddWidget(
+					SNew(SBox)
+					.HAlign(HAlign_Right)
+					[
+						SNew(SBox)
+						.Padding(FMargin(4.0f, 0.0f, 0.0f, 0.0f))
+						.WidthOverride(100.0f)
+						[
+							SNew(SNumericEntryBox<float>)
+							.Font(FEditorStyle::GetFontStyle(TEXT("MenuItem.Font")))
+							.AllowSpin(true)
+							.MinSliderValue(0.0f)
+							.MaxSliderValue(10.0f)
+							.Value(this, &FControlRigEditor::GetToolbarBoneRadius)
+							.OnValueChanged(this, &FControlRigEditor::OnToolbarBoneRadiusChanged)
+							.ToolTipText(LOCTEXT("ControlRigBoneRadiusToolTip", "The radius of the bones drawn for selected rig elements"))
+						]
+					], 
+					LOCTEXT("ControlRigBoneRadius", "Bone Radius")
+				);
+
 				if (UControlRigBlueprint* ControlRigBlueprint = CastChecked<UControlRigBlueprint>(GetBlueprintObj()))
 				{
 					for (UEdGraph* Graph : ControlRigBlueprint->UbergraphPages)
@@ -2862,6 +2890,31 @@ void FControlRigEditor::OnToolbarAxesScaleChanged(float InValue)
 	}
 }
 
+TOptional<float> FControlRigEditor::GetToolbarBoneRadius() const
+{
+	if (UControlRigBlueprint* ControlRigBlueprint = CastChecked<UControlRigBlueprint>(GetBlueprintObj()))
+	{
+		return ControlRigBlueprint->DebugBoneRadius;
+	}
+	return 0.f;
+}
+
+void FControlRigEditor::OnToolbarBoneRadiusChanged(float InValue)
+{
+	if (UControlRigBlueprint* ControlRigBlueprint = CastChecked<UControlRigBlueprint>(GetBlueprintObj()))
+	{
+		if(!FMath::IsNearlyEqual(ControlRigBlueprint->DebugBoneRadius, InValue))
+		{
+			ControlRigBlueprint->Modify();
+			ControlRigBlueprint->DebugBoneRadius = InValue;
+		}
+	}
+	if (UDebugSkelMeshComponent* MeshComponent = GetPersonaToolkit()->GetPreviewScene()->GetPreviewMeshComponent())
+	{
+		MeshComponent->BoneRadiusMultiplier = InValue;
+	}
+}
+
 ECheckBoxState FControlRigEditor::GetToolbarDrawAxesOnSelection() const
 {
 	if (FControlRigEditMode* EditMode = GetEditMode())
@@ -3073,6 +3126,14 @@ void FControlRigEditor::HandlePreviewSceneCreated(const TSharedRef<IPersonaPrevi
 			GEditor->UnregisterForUndo(AnimationEditorPreviewScene);
 		}
 	}
+
+	if(HasEditingObject())
+	{
+		if(UControlRigBlueprint* Blueprint = Cast<UControlRigBlueprint>(GetBlueprintObj()))
+		{
+			OnToolbarBoneRadiusChanged(Blueprint->DebugBoneRadius);
+		}
+	}
 }
 
 void FControlRigEditor::UpdateControlRig()
diff --git a/Engine/Plugins/Experimental/ControlRig/Source/ControlRigEditor/Private/Editor/ControlRigEditor.h b/Engine/Plugins/Experimental/ControlRig/Source/ControlRigEditor/Private/Editor/ControlRigEditor.h
--- a/Engine/Plugins/Experimental/ControlRig/Source/ControlRigEditor/Private/Editor/ControlRigEditor.h
+++ b/Engine/Plugins/Experimental/ControlRig/Source/ControlRigEditor/Private/Editor/ControlRigEditor.h
@@ -274,6 +274,8 @@ private:
 	void OnToolbarDrawAxesOnSelectionChanged(ECheckBoxState InNewValue);
 	TOptional<float> GetToolbarAxesScale() const;
 	void OnToolbarAxesScaleChanged(float InValue);
+	TOptional<float> GetToolbarBoneRadius() const;
+	void OnToolbarBoneRadiusChanged(float InValue);
 
 		/** Handle switching skeletal meshes */
 	void HandlePreviewMeshChanged(USkeletalMesh* InOldSkeletalMesh, USkeletalMesh* InNewSkeletalMesh);
diff --git a/Engine/Source/Editor/Persona/Private/AnimationEditorViewportClient.cpp b/Engine/Source/Editor/Persona/Private/AnimationEditorViewportClient.cpp
--- a/Engine/Source/Editor/Persona/Private/AnimationEditorViewportClient.cpp
+++ b/Engine/Source/Editor/Persona/Private/AnimationEditorViewportClient.cpp
@@ -1346,11 +1346,11 @@ void FAnimationViewportClient::DrawMeshBones(UDebugSkelMeshComponent * MeshCompo
 			}
 		}
 
-		DrawBones(DrawBoneIndices, MeshComponent->GetReferenceSkeleton(), WorldTransforms, MeshComponent->BonesOfInterest, PDI, BoneColours, MeshComponent->Bounds.SphereRadius);
+		DrawBones(DrawBoneIndices, MeshComponent->GetReferenceSkeleton(), WorldTransforms, MeshComponent->BonesOfInterest, PDI, BoneColours, MeshComponent->Bounds.SphereRadius, 0.f, false, MeshComponent->BoneRadiusMultiplier);
 	}
 }
 
-void FAnimationViewportClient::DrawBones(const TArray<FBoneIndexType>& RequiredBones, const FReferenceSkeleton& RefSkeleton, const TArray<FTransform> & WorldTransforms, const TArray<int32>& InSelectedBones, FPrimitiveDrawInterface* PDI, const TArray<FLinearColor>& BoneColours, float BoundRadius, float LineThickness/*=0.f*/, bool bForceDraw/*=false*/) const
+void FAnimationViewportClient::DrawBones(const TArray<FBoneIndexType>& RequiredBones, const FReferenceSkeleton& RefSkeleton, const TArray<FTransform> & WorldTransforms, const TArray<int32>& InSelectedBones, FPrimitiveDrawInterface* PDI, const TArray<FLinearColor>& BoneColours, float BoundRadius, float LineThickness/*=0.f*/, bool bForceDraw/*=false*/, float InBoneRadius/*=1.f*/) const
 {
 	TArray<int32> SelectedBones = InSelectedBones;
 	if(InSelectedBones.Num() > 0 && GetBoneDrawMode() == EBoneDrawMode::SelectedAndParents)
@@ -1395,8 +1395,9 @@ void FAnimationViewportClient::DrawBones(const TArray<FBoneIndexType>& RequiredB
 			}
 
 			const float BoneLength = (End - Start).Size();
+
 			// clamp by bound, we don't want too long or big
-			const float Radius = FMath::Clamp(BoneLength * 0.05f, 0.1f, MaxDrawRadius);
+			const float Radius = FMath::Clamp(BoneLength * 0.05f, 0.1f, MaxDrawRadius) * InBoneRadius;
 			//Render Sphere for bone end point and a cone between it and its parent.
 			PDI->SetHitProxy(new HPersonaBoneProxy(RefSkeleton.GetBoneName(BoneIndex)));
 			SkeletalDebugRendering::DrawWireBone(PDI, Start, End, LineColor, SDPG_Foreground, Radius);
diff --git a/Engine/Source/Editor/Persona/Private/AnimationEditorViewportClient.h b/Engine/Source/Editor/Persona/Private/AnimationEditorViewportClient.h
--- a/Engine/Source/Editor/Persona/Private/AnimationEditorViewportClient.h
+++ b/Engine/Source/Editor/Persona/Private/AnimationEditorViewportClient.h
@@ -408,7 +408,7 @@ private:
 	/** Draw Bones for non retargeted animation. */
 	void DrawMeshBonesBakedAnimation(UDebugSkelMeshComponent * MeshComponent, FPrimitiveDrawInterface* PDI) const;
 	/** Draws Bones for RequiredBones with WorldTransform **/
-	void DrawBones(const TArray<FBoneIndexType> & RequiredBones, const FReferenceSkeleton& RefSkeleton, const TArray<FTransform> & WorldTransforms, const TArray<int32>& InSelectedBones, FPrimitiveDrawInterface* PDI, const TArray<FLinearColor>& BoneColours, float BoundRadius, float LineThickness = 0.f, bool bForceDraw = false) const;
+	void DrawBones(const TArray<FBoneIndexType> & RequiredBones, const FReferenceSkeleton& RefSkeleton, const TArray<FTransform> & WorldTransforms, const TArray<int32>& InSelectedBones, FPrimitiveDrawInterface* PDI, const TArray<FLinearColor>& BoneColours, float BoundRadius, float LineThickness = 0.f, bool bForceDraw = false, float InBoneRadius = 1.f) const;
 	/** Draw Sub set of Bones **/
 	void DrawMeshSubsetBones(const UDebugSkelMeshComponent* MeshComponent, const TArray<int32>& BonesOfInterest, FPrimitiveDrawInterface* PDI) const;
 
diff --git a/Engine/Source/Editor/UnrealEd/Classes/Animation/DebugSkelMeshComponent.h b/Engine/Source/Editor/UnrealEd/Classes/Animation/DebugSkelMeshComponent.h
--- a/Engine/Source/Editor/UnrealEd/Classes/Animation/DebugSkelMeshComponent.h
+++ b/Engine/Source/Editor/UnrealEd/Classes/Animation/DebugSkelMeshComponent.h
@@ -232,6 +232,10 @@ class UNREALED_API UDebugSkelMeshComponent : public USkeletalMeshComponent
 	UPROPERTY(transient)
 	uint32 bRequiredBonesUpToDateDuringTick : 1;
 
+	/** Multiplier for the bone radius rendering */
+	UPROPERTY(transient)
+	float BoneRadiusMultiplier;
+
 	/* Bounds computed from cloth. */
 	FBoxSphereBounds CachedClothBounds;
 
diff --git a/Engine/Source/Editor/UnrealEd/Private/Animation/DebugSkelMeshComponent.cpp b/Engine/Source/Editor/UnrealEd/Private/Animation/DebugSkelMeshComponent.cpp
--- a/Engine/Source/Editor/UnrealEd/Private/Animation/DebugSkelMeshComponent.cpp
+++ b/Engine/Source/Editor/UnrealEd/Private/Animation/DebugSkelMeshComponent.cpp
@@ -43,6 +43,8 @@ UDebugSkelMeshComponent::UDebugSkelMeshComponent(const FObjectInitializer& Objec
 	bPauseClothingSimulationWithAnim = false;
 	bPerformSingleClothingTick = false;
 
+	BoneRadiusMultiplier = 1.f;
+	
 	CachedClothBounds = FBoxSphereBounds(ForceInit);
 }
 
diff --git a/Engine/Source/Editor/UnrealEd/Public/Toolkits/AssetEditorToolkit.h b/Engine/Source/Editor/UnrealEd/Public/Toolkits/AssetEditorToolkit.h
--- a/Engine/Source/Editor/UnrealEd/Public/Toolkits/AssetEditorToolkit.h
+++ b/Engine/Source/Editor/UnrealEd/Public/Toolkits/AssetEditorToolkit.h
@@ -278,6 +278,9 @@ protected:
 	 */
 	virtual void PostInitAssetEditor() {}
 
+	/**	Returns true if this toolkit has any objects being edited */
+	bool HasEditingObject() const { return !EditingObjects.IsEmpty(); }
+
 	/**	Returns the single object currently being edited. Asserts if currently editing no object or multiple objects */
 	UObject* GetEditingObject() const;
 
