- Expose GetToolkitHost in ISequencer
- Update ControlRig editor mode and track editor to use currently active ToolKitHost rather than GLevelEditor singleton
- Acted on to-do and removed some code
- Made FControlRigEditMode derive from IControlRigEditMode rather than IPersonaEditMode

diff --git a/Engine/Plugins/Experimental/ControlRig/Source/ControlRigEditor/Private/EditMode/ControlRigEditMode.cpp b/Engine/Plugins/Experimental/ControlRig/Source/ControlRigEditor/Private/EditMode/ControlRigEditMode.cpp
--- a/Engine/Plugins/Experimental/ControlRig/Source/ControlRigEditor/Private/EditMode/ControlRigEditMode.cpp
+++ b/Engine/Plugins/Experimental/ControlRig/Source/ControlRigEditor/Private/EditMode/ControlRigEditMode.cpp
@@ -209,6 +209,11 @@ void FControlRigEditMode::SetObjects(const TWeakObjectPtr<>& InSelectedObject,
 	SetObjects_Internal();
 }
 
+bool FControlRigEditMode::IsInLevelEditor() const
+{
+	return GetModeManager() == &GLevelEditorModeTools();
+}
+
 void FControlRigEditMode::SetUpDetailPanel()
 {
 	if (IsInLevelEditor())
@@ -303,20 +308,21 @@ void FControlRigEditMode::Enter()
 
 		Toolkit->Init(Owner->GetToolkitHost());
 
+		FEditorModeTools* ModeManager = GetModeManager();
+
 		bIsChangingCoordSystem = false;
 		if (CoordSystemPerWidgetMode.Num() < (UE::Widget::WM_Max))
 		{
 			CoordSystemPerWidgetMode.SetNum(UE::Widget::WM_Max);
-			ECoordSystem CoordSystem = GLevelEditorModeTools().GetCoordSystem();
+			ECoordSystem CoordSystem = ModeManager->GetCoordSystem();
 			for (int32 i = 0; i < UE::Widget::WM_Max; ++i)
 			{
 				CoordSystemPerWidgetMode[i] = CoordSystem;
 			}
 		}
 	
-		GLevelEditorModeTools().OnWidgetModeChanged().AddSP(this, &FControlRigEditMode::OnWidgetModeChanged);
-		GLevelEditorModeTools().OnCoordSystemChanged().AddSP(this, &FControlRigEditMode::OnCoordSystemChanged);
-
+		ModeManager->OnWidgetModeChanged().AddSP(this, &FControlRigEditMode::OnWidgetModeChanged);
+		ModeManager->OnCoordSystemChanged().AddSP(this, &FControlRigEditMode::OnCoordSystemChanged);
 	}
 
 	if (DelegateHelper.IsValid())
@@ -375,8 +381,9 @@ void FControlRigEditMode::Exit()
 	}
 
 	//clear delegates
-	GLevelEditorModeTools().OnWidgetModeChanged().RemoveAll(this);
-	GLevelEditorModeTools().OnCoordSystemChanged().RemoveAll(this);
+	FEditorModeTools* ModeManager = GetModeManager();
+	ModeManager->OnWidgetModeChanged().RemoveAll(this);
+	ModeManager->OnCoordSystemChanged().RemoveAll(this);
 
 	//clear proxies
 	ControlProxy->RemoveAllProxies();
@@ -1702,22 +1709,19 @@ void FControlRigEditMode::FrameSelection()
 void FControlRigEditMode::IncreaseGizmoSize()
 {
 	Settings->GizmoScale += 0.1f;
-	FEditorModeTools& ModeTools = GLevelEditorModeTools();
-	ModeTools.SetWidgetScale(Settings->GizmoScale);
+	GetModeManager()->SetWidgetScale(Settings->GizmoScale);
 }
 
 void FControlRigEditMode::DecreaseGizmoSize()
 {
 	Settings->GizmoScale -= 0.1f;
-	FEditorModeTools& ModeTools = GLevelEditorModeTools();
-	ModeTools.SetWidgetScale(Settings->GizmoScale);
+	GetModeManager()->SetWidgetScale(Settings->GizmoScale);
 }
 
 void FControlRigEditMode::ResetGizmoSize()
 {
 	Settings->GizmoScale = 1.0f;
-	FEditorModeTools& ModeTools = GLevelEditorModeTools();
-	ModeTools.SetWidgetScale(Settings->GizmoScale);
+	GetModeManager()->SetWidgetScale(Settings->GizmoScale);
 }
 
 void FControlRigEditMode::ToggleManipulators()
@@ -2039,10 +2043,11 @@ void FControlRigEditMode::OnWidgetModeChanged(UE::Widget::EWidgetMode InWidgetMo
 	{
 		TGuardValue<bool> ReentrantGuardSelf(bIsChangingCoordSystem, true);
 
-		int32 WidgetMode = (int32)GLevelEditorModeTools().GetWidgetMode();
+		FEditorModeTools* ModeManager = GetModeManager();
+		int32 WidgetMode = (int32)ModeManager->GetWidgetMode();
 		if (WidgetMode >= 0 && WidgetMode < CoordSystemPerWidgetMode.Num())
 		{
-			GLevelEditorModeTools().SetCoordSystem(CoordSystemPerWidgetMode[WidgetMode]);
+			ModeManager->SetCoordSystem(CoordSystemPerWidgetMode[WidgetMode]);
 		}
 	}
 }
@@ -2051,8 +2056,9 @@ void FControlRigEditMode::OnCoordSystemChanged(ECoordSystem InCoordSystem)
 {
 	TGuardValue<bool> ReentrantGuardSelf(bIsChangingCoordSystem, true);
 
-	int32 WidgetMode = (int32)GLevelEditorModeTools().GetWidgetMode();
-	ECoordSystem CoordSystem = GLevelEditorModeTools().GetCoordSystem();
+	FEditorModeTools* ModeManager = GetModeManager();
+	int32 WidgetMode = (int32)ModeManager->GetWidgetMode();
+	ECoordSystem CoordSystem = ModeManager->GetCoordSystem();
 	if (WidgetMode >= 0 && WidgetMode < CoordSystemPerWidgetMode.Num())
 	{
 		CoordSystemPerWidgetMode[WidgetMode] = CoordSystem;
diff --git a/Engine/Plugins/Experimental/ControlRig/Source/ControlRigEditor/Private/EditMode/ControlRigEditMode.h b/Engine/Plugins/Experimental/ControlRig/Source/ControlRigEditor/Private/EditMode/ControlRigEditMode.h
--- a/Engine/Plugins/Experimental/ControlRig/Source/ControlRigEditor/Private/EditMode/ControlRigEditMode.h
+++ b/Engine/Plugins/Experimental/ControlRig/Source/ControlRigEditor/Private/EditMode/ControlRigEditMode.h
@@ -4,7 +4,6 @@
 
 #include "CoreMinimal.h"
 #include "InputCoreTypes.h"
-#include "IPersonaEditMode.h"
 #include "IControlRigObjectBinding.h"
 #include "RigVMModel/RigVMGraph.h"
 #include "Rigs/RigHierarchyContainer.h"
@@ -12,6 +11,7 @@
 #include "Framework/MultiBox/MultiBoxBuilder.h"
 #include "UObject/StrongObjectPtr.h"
 #include "UnrealWidgetFwd.h"
+#include "IControlRigEditMode.h"
 #include "ControlRigEditMode.generated.h"
 
 
@@ -69,7 +69,7 @@ private:
 	FDelegateHandle OnBoneTransformsFinalizedHandle;
 };
 
-class FControlRigEditMode : public IPersonaEditMode
+class FControlRigEditMode : public IControlRigEditMode
 {
 public:
 	static FName ModeName;
@@ -78,10 +78,10 @@ public:
 	~FControlRigEditMode();
 
 	/** Set the objects to be displayed in the details panel */
-	void SetObjects(const TWeakObjectPtr<>& InSelectedObject,  UObject* BindingObject, TWeakPtr<ISequencer> InSequencer);
+	virtual void SetObjects(const TWeakObjectPtr<>& InSelectedObject,  UObject* BindingObject, TWeakPtr<ISequencer> InSequencer) override;
 
 	/** This edit mode is re-used between the level editor and the control rig editor. Calling this indicates which context we are in */
-	virtual bool IsInLevelEditor() const { return true; }
+	bool IsInLevelEditor() const;
 
 	// FEdMode interface
 	virtual bool UsesToolkits() const override;
diff --git a/Engine/Plugins/Experimental/ControlRig/Source/ControlRigEditor/Private/EditMode/SControlRigEditModeTools.cpp b/Engine/Plugins/Experimental/ControlRig/Source/ControlRigEditor/Private/EditMode/SControlRigEditModeTools.cpp
--- a/Engine/Plugins/Experimental/ControlRig/Source/ControlRigEditor/Private/EditMode/SControlRigEditModeTools.cpp
+++ b/Engine/Plugins/Experimental/ControlRig/Source/ControlRigEditor/Private/EditMode/SControlRigEditModeTools.cpp
@@ -35,10 +35,13 @@
 class FControlRigEditModeGenericDetails : public IDetailCustomization
 {
 public:
+	FControlRigEditModeGenericDetails() = delete;
+	FControlRigEditModeGenericDetails(FEditorModeTools* InModeTools) : ModeTools(InModeTools) {}
+	
 	/** Makes a new instance of this detail layout class for a specific detail view requesting it */
-	static TSharedRef<IDetailCustomization> MakeInstance()
+	static TSharedRef<IDetailCustomization> MakeInstance(FEditorModeTools* InModeTools)
 	{
-		return MakeShareable(new FControlRigEditModeGenericDetails);
+		return MakeShareable(new FControlRigEditModeGenericDetails(InModeTools));
 	}
 
 	/** IDetailCustomization interface */
@@ -108,7 +111,7 @@ public:
 			}
 
 			URigHierarchy* Hierarchy = Proxy->ControlRig->GetHierarchy();
-			Hierarchy->ForEach<FRigControlElement>([Hierarchy, Proxy, &Category](FRigControlElement* ControlElement) -> bool
+			Hierarchy->ForEach<FRigControlElement>([Hierarchy, Proxy, &Category, this](FRigControlElement* ControlElement) -> bool
             {
 				FName ParentControlName = NAME_None;
 				FRigControlElement* ParentControlElement = Cast<FRigControlElement>(Hierarchy->GetFirstParent(ControlElement));
@@ -119,7 +122,7 @@ public:
 				
 				if (ParentControlName == ControlElement->GetName())
 				{
-					if (FControlRigEditMode* EditMode = static_cast<FControlRigEditMode*>(GLevelEditorModeTools().GetActiveMode(FControlRigEditMode::ModeName)))
+					if (FControlRigEditMode* EditMode = static_cast<FControlRigEditMode*>(ModeTools->GetActiveMode(FControlRigEditMode::ModeName)))
 					{
 						if (UObject* NestedProxy = EditMode->ControlProxy->FindProxy(ControlElement->GetName()))
 						{
@@ -180,6 +183,8 @@ public:
 			});
 		}
 	}
+protected:
+	FEditorModeTools* ModeTools = nullptr;
 };
 
 void SControlRigEditModeTools::SetControlRig(UControlRig* ControlRig)
@@ -219,12 +224,14 @@ void SControlRigEditModeTools::Construct(const FArguments& InArgs, FControlRigEd
 		DetailsViewArgs.bAllowMultipleTopLevelObjects = true;
 		DetailsViewArgs.bShowScrollBar = false; // Don't need to show this, as we are putting it in a scroll box
 	}
+	
+	ModeTools = InEditMode.GetModeManager();
 
 	ControlDetailsView = FModuleManager::GetModuleChecked<FPropertyEditorModule>("PropertyEditor").CreateDetailView(DetailsViewArgs);
 	ControlDetailsView->SetKeyframeHandler(SharedThis(this));
 	ControlDetailsView->SetIsPropertyVisibleDelegate(FIsPropertyVisible::CreateSP(this, &SControlRigEditModeTools::ShouldShowPropertyOnDetailCustomization));
 	ControlDetailsView->SetIsPropertyReadOnlyDelegate(FIsPropertyReadOnly::CreateSP(this, &SControlRigEditModeTools::IsReadOnlyPropertyOnDetailCustomization));
-	ControlDetailsView->SetGenericLayoutDetailsDelegate(FOnGetDetailCustomizationInstance::CreateStatic(&FControlRigEditModeGenericDetails::MakeInstance));
+	ControlDetailsView->SetGenericLayoutDetailsDelegate(FOnGetDetailCustomizationInstance::CreateStatic(&FControlRigEditModeGenericDetails::MakeInstance, ModeTools));
 
 	RigOptionsDetailsView = FModuleManager::GetModuleChecked<FPropertyEditorModule>("PropertyEditor").CreateDetailView(DetailsViewArgs);
 	RigOptionsDetailsView->SetKeyframeHandler(SharedThis(this));
@@ -440,7 +447,7 @@ static bool bPickerChangingSelection = false;
 
 void SControlRigEditModeTools::OnManipulatorsPicked(const TArray<FName>& Manipulators)
 {
-	FControlRigEditMode* ControlRigEditMode = static_cast<FControlRigEditMode*>(GLevelEditorModeTools().GetActiveMode(FControlRigEditMode::ModeName));
+	FControlRigEditMode* ControlRigEditMode = static_cast<FControlRigEditMode*>(ModeTools->GetActiveMode(FControlRigEditMode::ModeName));
 	if (ControlRigEditMode)
 	{
 		if (!bPickerChangingSelection)
@@ -496,7 +503,7 @@ void SControlRigEditModeTools::OnRigOptionFinishedChange(const FPropertyChangedE
 {
 	SetControlRig(SequencerRig.Get());
 
-	if (FControlRigEditMode* EditMode = static_cast<FControlRigEditMode*>(GLevelEditorModeTools().GetActiveMode(FControlRigEditMode::ModeName)))
+	if (FControlRigEditMode* EditMode = static_cast<FControlRigEditMode*>(ModeTools->GetActiveMode(FControlRigEditMode::ModeName)))
 	{
 		EditMode->SetObjects_Internal();
 	}
@@ -508,15 +515,15 @@ void SControlRigEditModeTools::CustomizeToolBarPalette(FToolBarBuilder& ToolBarB
 	ToolBarBuilder.AddToolBarButton(
 		FUIAction(
 			FExecuteAction::CreateLambda([this] {
-		FControlRigEditMode* ControlRigEditMode = static_cast<FControlRigEditMode*>(GLevelEditorModeTools().GetActiveMode(FControlRigEditMode::ModeName));
+		FControlRigEditMode* ControlRigEditMode = static_cast<FControlRigEditMode*>(ModeTools->GetActiveMode(FControlRigEditMode::ModeName));
 		if (ControlRigEditMode)
 		{
 			ControlRigEditMode->SetOnlySelectRigControls(!ControlRigEditMode->GetOnlySelectRigControls());
 		}
 	}),
 			FCanExecuteAction(),
-		FIsActionChecked::CreateLambda([] {
-		FControlRigEditMode* ControlRigEditMode = static_cast<FControlRigEditMode*>(GLevelEditorModeTools().GetActiveMode(FControlRigEditMode::ModeName));
+		FIsActionChecked::CreateLambda([this] {
+		FControlRigEditMode* ControlRigEditMode = static_cast<FControlRigEditMode*>(ModeTools->GetActiveMode(FControlRigEditMode::ModeName));
 		if (ControlRigEditMode)
 		{
 			return ControlRigEditMode->GetOnlySelectRigControls();
@@ -578,7 +585,7 @@ void SControlRigEditModeTools::CustomizeToolBarPalette(FToolBarBuilder& ToolBarB
 
 void SControlRigEditModeTools::MakePoseDialog()
 {
-	FControlRigEditMode* ControlRigEditMode = static_cast<FControlRigEditMode*>(GLevelEditorModeTools().GetActiveMode(FControlRigEditMode::ModeName));
+	FControlRigEditMode* ControlRigEditMode = static_cast<FControlRigEditMode*>(ModeTools->GetActiveMode(FControlRigEditMode::ModeName));
 	if (ControlRigEditMode)
 	{
 		FGlobalTabmanager::Get()->TryInvokeTab(IControlRigEditorModule::ControlRigPoseTab);
@@ -587,7 +594,7 @@ void SControlRigEditModeTools::MakePoseDialog()
 
 void SControlRigEditModeTools::MakeTweenDialog()
 {
-	FControlRigEditMode* ControlRigEditMode = static_cast<FControlRigEditMode*>(GLevelEditorModeTools().GetActiveMode(FControlRigEditMode::ModeName));
+	FControlRigEditMode* ControlRigEditMode = static_cast<FControlRigEditMode*>(ModeTools->GetActiveMode(FControlRigEditMode::ModeName));
 	if (ControlRigEditMode)
 	{
 		FGlobalTabmanager::Get()->TryInvokeTab(IControlRigEditorModule::ControlRigTweenTab);
@@ -597,7 +604,7 @@ void SControlRigEditModeTools::MakeTweenDialog()
 
 void SControlRigEditModeTools::MakeSnapperDialog()
 {
-	FControlRigEditMode* ControlRigEditMode = static_cast<FControlRigEditMode*>(GLevelEditorModeTools().GetActiveMode(FControlRigEditMode::ModeName));
+	FControlRigEditMode* ControlRigEditMode = static_cast<FControlRigEditMode*>(ModeTools->GetActiveMode(FControlRigEditMode::ModeName));
 	if (ControlRigEditMode)
 	{
 		FGlobalTabmanager::Get()->TryInvokeTab(IControlRigEditorModule::ControlRigSnapperTab);
@@ -608,7 +615,7 @@ void SControlRigEditModeTools::MakeSnapperDialog()
 
 void SControlRigEditModeTools::MakeTempPivotDialog()
 {
-	FControlRigEditMode* ControlRigEditMode = static_cast<FControlRigEditMode*>(GLevelEditorModeTools().GetActiveMode(FControlRigEditMode::ModeName));
+	FControlRigEditMode* ControlRigEditMode = static_cast<FControlRigEditMode*>(ModeTools->GetActiveMode(FControlRigEditMode::ModeName));
 	if (ControlRigEditMode)
 	{
 		FGlobalTabmanager::Get()->TryInvokeTab(IControlRigEditorModule::ControlRigTempPivotTab);
diff --git a/Engine/Plugins/Experimental/ControlRig/Source/ControlRigEditor/Private/EditMode/SControlRigEditModeTools.h b/Engine/Plugins/Experimental/ControlRig/Source/ControlRigEditor/Private/EditMode/SControlRigEditModeTools.h
--- a/Engine/Plugins/Experimental/ControlRig/Source/ControlRigEditor/Private/EditMode/SControlRigEditModeTools.h
+++ b/Engine/Plugins/Experimental/ControlRig/Source/ControlRigEditor/Private/EditMode/SControlRigEditModeTools.h
@@ -84,6 +84,8 @@ private:
 
 	//TODO may put back void MakeSelectionSetDialog();
 	//TWeakPtr<SWindow> SelectionSetWindow;
+
+	FEditorModeTools* ModeTools = nullptr;
 public:
 	/** Modes Panel Header Information **/
 	void CustomizeToolBarPalette(FToolBarBuilder& ToolBarBuilder);
diff --git a/Engine/Plugins/Experimental/ControlRig/Source/ControlRigEditor/Private/Editor/ControlRigEditorEditMode.h b/Engine/Plugins/Experimental/ControlRig/Source/ControlRigEditor/Private/Editor/ControlRigEditorEditMode.h
--- a/Engine/Plugins/Experimental/ControlRig/Source/ControlRigEditor/Private/Editor/ControlRigEditorEditMode.h
+++ b/Engine/Plugins/Experimental/ControlRig/Source/ControlRigEditor/Private/Editor/ControlRigEditorEditMode.h
@@ -15,9 +15,6 @@ class FControlRigEditorEditMode : public FControlRigEditMode
 public:
 	static FName ModeName;
 
-	/** FControlRigEditMode interface */
-	virtual bool IsInLevelEditor() const { return false; }
-
 	// FEdMode interface
 	virtual void Render(const FSceneView* View, FViewport* Viewport, FPrimitiveDrawInterface* PDI) override;
 
diff --git a/Engine/Plugins/Experimental/ControlRig/Source/ControlRigEditor/Private/Sequencer/ControlRigParameterTrackEditor.cpp b/Engine/Plugins/Experimental/ControlRig/Source/ControlRigEditor/Private/Sequencer/ControlRigParameterTrackEditor.cpp
--- a/Engine/Plugins/Experimental/ControlRig/Source/ControlRigEditor/Private/Sequencer/ControlRigParameterTrackEditor.cpp
+++ b/Engine/Plugins/Experimental/ControlRig/Source/ControlRigEditor/Private/Sequencer/ControlRigParameterTrackEditor.cpp
@@ -192,8 +192,6 @@ FControlRigParameterTrackEditor::FControlRigParameterTrackEditor(TSharedRef<ISeq
 	OnActorAddedToSequencerHandle = InSequencer->OnActorAddedToSequencer().AddRaw(this, &FControlRigParameterTrackEditor::HandleActorAdded);
 	OnTreeViewChangedHandle = InSequencer->OnTreeViewChanged().AddRaw(this, &FControlRigParameterTrackEditor::OnTreeViewChanged);
 
-	//REMOVE ME IN UE5
-	//InSequencer->GetObjectChangeListener().GetOnPropagateObjectChanges().AddRaw(this, &FControlRigParameterTrackEditor::OnPropagateObjectChanges);
 	{
 		//we check for two things, one if the control rig has been replaced if so we need to switch.
 		//the other is if bound object on the edit mode is null we request a re-evaluate which will reset it up.
@@ -202,7 +200,7 @@ FControlRigParameterTrackEditor::FControlRigParameterTrackEditor(TSharedRef<ISeq
 			if (GetSequencer().IsValid())
 			{
 				TMap<UControlRig*, UControlRig*> OldToNewControlRigs;
-				FControlRigEditMode* ControlRigEditMode = static_cast<FControlRigEditMode*>(GLevelEditorModeTools().GetActiveMode(FControlRigEditMode::ModeName));
+				FControlRigEditMode* ControlRigEditMode = GetEditMode();
 				if (ControlRigEditMode && ControlRigEditMode->GetControlRig(true) && ControlRigEditMode->GetControlRig(true)->GetObjectBinding())
 				{
 					if (ControlRigEditMode->GetControlRig(true)->GetObjectBinding()->GetBoundObject() == nullptr)
@@ -413,12 +411,12 @@ void FControlRigParameterTrackEditor::OnRelease()
 			}
 		}
 	}
-	FControlRigEditMode* ControlRigEditMode = static_cast<FControlRigEditMode*>(GLevelEditorModeTools().GetActiveMode(FControlRigEditMode::ModeName));
+	FControlRigEditMode* ControlRigEditMode = GetEditMode();
 	if (ControlRigEditMode)
 	{
-		if (GLevelEditorModeTools().HasToolkitHost())
+		if (FEditorModeTools* Tools = GetEditorModeTools())
 		{
-			GLevelEditorModeTools().DeactivateMode(FControlRigEditMode::ModeName);
+			Tools->DeactivateMode(FControlRigEditMode::ModeName);
 		}
 
 		ControlRigEditMode->SetObjects(nullptr, nullptr, GetSequencer());
@@ -925,12 +923,10 @@ void FControlRigParameterTrackEditor::BakeToControlRig(UClass* InClass, FGuid Ob
 						return;
 					}
 
-					FControlRigEditMode* ControlRigEditMode = static_cast<FControlRigEditMode*>(GLevelEditorModeTools().GetActiveMode(FControlRigEditMode::ModeName));
+					FControlRigEditMode* ControlRigEditMode = GetEditMode();
 					if (!ControlRigEditMode)
 					{
-						GLevelEditorModeTools().ActivateMode(FControlRigEditMode::ModeName);
-						ControlRigEditMode = static_cast<FControlRigEditMode*>(GLevelEditorModeTools().GetActiveMode(FControlRigEditMode::ModeName));
-
+						ControlRigEditMode = GetEditMode(true);
 					}
 					else
 					{
@@ -1282,13 +1278,8 @@ void FControlRigParameterTrackEditor::AddControlRig(UClass* InClass, UObject* Bo
 			GetSequencer()->ThrobSectionSelection();
 			GetSequencer()->NotifyMovieSceneDataChanged(EMovieSceneDataChangeType::MovieSceneStructureItemAdded);
 
-			FControlRigEditMode* ControlRigEditMode = static_cast<FControlRigEditMode*>(GLevelEditorModeTools().GetActiveMode(FControlRigEditMode::ModeName));
-			if (!ControlRigEditMode)
-			{
-				GLevelEditorModeTools().ActivateMode(FControlRigEditMode::ModeName);
-				ControlRigEditMode = static_cast<FControlRigEditMode*>(GLevelEditorModeTools().GetActiveMode(FControlRigEditMode::ModeName));
-
-			}
+			FControlRigEditMode* ControlRigEditMode = GetEditMode(true);
+			
 			if (ControlRigEditMode)
 			{
 				ControlRigEditMode->SetObjects(ControlRig, nullptr, GetSequencer());
@@ -1360,7 +1351,7 @@ bool FControlRigParameterTrackEditor::CanAddTransformKeysForSelectedObjects() co
 		return false;
 	}
 
-	FControlRigEditMode* ControlRigEditMode = static_cast<FControlRigEditMode*>(GLevelEditorModeTools().GetActiveMode(FControlRigEditMode::ModeName));
+	FControlRigEditMode* ControlRigEditMode = GetEditMode();
 	if (ControlRigEditMode && ControlRigEditMode->GetControlRig(false))
 	{
 		UControlRig* ControlRig = ControlRigEditMode->GetControlRig(false);
@@ -1382,7 +1373,7 @@ void FControlRigParameterTrackEditor::OnAddTransformKeysForSelectedObjects(EMovi
 		return;
 	}
 
-	FControlRigEditMode* ControlRigEditMode = static_cast<FControlRigEditMode*>(GLevelEditorModeTools().GetActiveMode(FControlRigEditMode::ModeName));
+	FControlRigEditMode* ControlRigEditMode = GetEditMode();
 	if (ControlRigEditMode && ControlRigEditMode->GetControlRig(false))
 	{
 		UControlRig* ControlRig = ControlRigEditMode->GetControlRig(false);
@@ -1502,7 +1493,7 @@ void FControlRigParameterTrackEditor::OnActivateSequenceChanged(FMovieSceneSeque
 void FControlRigParameterTrackEditor::OnSequencerDataChanged(EMovieSceneDataChangeType DataChangeType)
 {
 	UMovieScene* MovieScene = GetSequencer()->GetFocusedMovieSceneSequence()->GetMovieScene();
-	FControlRigEditMode* ControlRigEditMode = static_cast<FControlRigEditMode*>(GLevelEditorModeTools().GetActiveMode(FControlRigEditMode::ModeName));
+	FControlRigEditMode* ControlRigEditMode = GetEditMode();
 
 	//if we have a valid control rig edit mode need to check and see the control rig in that mode is still in a track
 	//if not we get rid of it.
@@ -1521,10 +1512,10 @@ void FControlRigParameterTrackEditor::OnSequencerDataChanged(EMovieSceneDataChan
 				return; //just exit out we still have a good track
 			}
 		}
-		//okay no good track so deactive it and delete it's Control Rig and bingings.
-		if (GLevelEditorModeTools().HasToolkitHost())
+		//okay no good track so deactive it and delete it's Control Rig and bindings.
+		if (FEditorModeTools* Tools = GetEditorModeTools())
 		{
-			GLevelEditorModeTools().DeactivateMode(FControlRigEditMode::ModeName);
+			Tools->DeactivateMode(FControlRigEditMode::ModeName);
 		}
 		ControlRigEditMode->SetObjects(nullptr, nullptr, GetSequencer());
 	}
@@ -1540,7 +1531,7 @@ void FControlRigParameterTrackEditor::OnCurveDisplayChanged(FCurveModel* CurveMo
 	FScopedTransaction ScopedTransaction(LOCTEXT("SelectControlTransaction", "Select Control"), !GIsTransacting);
 
 	TArray<FString> StringArray;
-	FControlRigEditMode* ControlRigEditMode = static_cast<FControlRigEditMode*>(GLevelEditorModeTools().GetActiveMode(FControlRigEditMode::ModeName));
+	FControlRigEditMode* ControlRigEditMode = GetEditMode();
 	UControlRig* ControlRig = nullptr;
 
 	if (CurveModel)
@@ -1552,8 +1543,7 @@ void FControlRigParameterTrackEditor::OnCurveDisplayChanged(FCurveModel* CurveMo
 			//Only create the edit mode if we have a  curve selected and it's not set and we have some boundobjects.
 			if (!ControlRigEditMode)
 			{
-				GLevelEditorModeTools().ActivateMode(FControlRigEditMode::ModeName);
-				ControlRigEditMode = static_cast<FControlRigEditMode*>(GLevelEditorModeTools().GetActiveMode(FControlRigEditMode::ModeName));
+				ControlRigEditMode = GetEditMode(true);
 				if (TSharedPtr<IControlRigObjectBinding> ObjectBinding = ControlRig->GetObjectBinding())
 				{
 					if (ControlRigEditMode)
@@ -1642,7 +1632,7 @@ void FControlRigParameterTrackEditor::OnTreeViewChanged()
 {
 	if (!bIsDoingSelection)
 	{
-		FControlRigEditMode* ControlRigEditMode = static_cast<FControlRigEditMode*>(GLevelEditorModeTools().GetActiveMode(FControlRigEditMode::ModeName));
+		FControlRigEditMode* ControlRigEditMode = GetEditMode();
 		if (ControlRigEditMode)
 		{
 			if (UControlRig* ControlRig = ControlRigEditMode->GetControlRig(true))
@@ -1696,7 +1686,7 @@ void FControlRigParameterTrackEditor::OnSelectionChanged(TArray<UMovieSceneTrack
 	TGuardValue<bool> Guard(bIsDoingSelection, true);
 
 	TArray<FString> StringArray;
-	FControlRigEditMode* ControlRigEditMode = static_cast<FControlRigEditMode*>(GLevelEditorModeTools().GetActiveMode(FControlRigEditMode::ModeName));
+	FControlRigEditMode* ControlRigEditMode = GetEditMode();
 	UControlRig* ControlRig = nullptr;
 
 	TArray<const IKeyArea*> KeyAreas;
@@ -1732,8 +1722,7 @@ void FControlRigParameterTrackEditor::OnSelectionChanged(TArray<UMovieSceneTrack
 					}
 					else
 					{
-						GLevelEditorModeTools().ActivateMode(FControlRigEditMode::ModeName);
-						ControlRigEditMode = static_cast<FControlRigEditMode*>(GLevelEditorModeTools().GetActiveMode(FControlRigEditMode::ModeName));
+						ControlRigEditMode = GetEditMode(true);
 						if (TSharedPtr<IControlRigObjectBinding> ObjectBinding = TrackControlRig->GetObjectBinding())
 						{
 							if (ControlRigEditMode)
@@ -1758,8 +1747,7 @@ void FControlRigParameterTrackEditor::OnSelectionChanged(TArray<UMovieSceneTrack
 			//Only create the edit mode if we have a KeyAra selected and it's not set and we have some boundobjects.
 			if (!ControlRigEditMode)
 			{
-				GLevelEditorModeTools().ActivateMode(FControlRigEditMode::ModeName);
-				ControlRigEditMode = static_cast<FControlRigEditMode*>(GLevelEditorModeTools().GetActiveMode(FControlRigEditMode::ModeName));
+				ControlRigEditMode = GetEditMode(true);
 				if (TSharedPtr<IControlRigObjectBinding> ObjectBinding = ControlRig->GetObjectBinding())
 				{
 					if (ControlRigEditMode)
@@ -2001,42 +1989,6 @@ void FControlRigParameterTrackEditor::HandleControlSelected(UControlRig* Subject
 	}
 }
 
-//REMOVE ME IN UE5
-void FControlRigParameterTrackEditor::OnPropagateObjectChanges(UObject* InChangedObject)
-{
-	//not needed
-	/*
-	if (AActor* Actor = Cast<AActor>(InChangedObject))
-	{
-		if (UMovieScene* MovieScene = GetFocusedMovieScene())
-		{
-			const TArray<FMovieSceneBinding>& Bindings = MovieScene->GetBindings();
-			for (const FMovieSceneBinding& Binding : Bindings)
-			{
-				if (UMovieSceneControlRigParameterTrack* Track = Cast<UMovieSceneControlRigParameterTrack>(MovieScene->FindTrack(UMovieSceneControlRigParameterTrack::StaticClass(), Binding.GetObjectGuid(), NAME_None)))
-				{
-					if (UControlRig* ControlRig = Track->GetControlRig())
-					{
-						if (ControlRig->GetObjectBinding())
-						{
-							if (USceneComponent* SceneComponent = Cast<USceneComponent>(ControlRig->GetObjectBinding()->GetBoundObject()))
-							{
-								if (SceneComponent->GetOwner() == Actor)
-								{
-									if (GetSequencer().IsValid())
-									{
-										GetSequencer()->NotifyMovieSceneDataChanged(EMovieSceneDataChangeType::Unknown);
-									}
-								}
-							}
-						}
-					}
-				}
-			}
-		}
-	}
-	*/
-}
 
 void FControlRigParameterTrackEditor::HandleOnInitialized(UControlRig* ControlRig, const EControlRigState InState, const FName& InEventName)
 {
@@ -2053,7 +2005,7 @@ void FControlRigParameterTrackEditor::HandleOnInitialized(UControlRig* ControlRi
 
 void FControlRigParameterTrackEditor::HandleControlModified(UControlRig* ControlRig, FRigControlElement* ControlElement, const FRigControlModifiedContext& Context)
 {
-	if (GetSequencer().IsValid() && !GetSequencer()->IsAllowedToChange())
+	if (!GetSequencer().IsValid() || !GetSequencer()->IsAllowedToChange())
 	{
 		return;
 	}
@@ -3321,6 +3273,36 @@ void FControlRigParameterSection::OnAnimationAssetEnterPressedForFK(const TArray
 	}
 }
 
+FEditorModeTools* FControlRigParameterTrackEditor::GetEditorModeTools() const
+{
+	TSharedPtr<ISequencer> SharedSequencer = GetSequencer();
+	if (SharedSequencer.IsValid())
+	{
+		TSharedPtr<IToolkitHost> ToolkitHost = SharedSequencer->GetToolkitHost();
+		if(ToolkitHost.IsValid())
+		{
+			return &ToolkitHost->GetEditorModeManager();
+		}
+	}
+
+	return nullptr;
+}
+
+FControlRigEditMode* FControlRigParameterTrackEditor::GetEditMode(bool bForceActivate /*= false*/) const
+{
+	if (FEditorModeTools* EditorModetools = GetEditorModeTools())
+	{
+		if(bForceActivate && !EditorModetools->IsModeActive(FControlRigEditMode::ModeName))
+		{
+			EditorModetools->ActivateMode(FControlRigEditMode::ModeName);
+		}
+		
+		return static_cast<FControlRigEditMode*>(EditorModetools->GetActiveMode(FControlRigEditMode::ModeName));
+	}
+
+	return nullptr;
+}
+
 #undef LOCTEXT_NAMESPACE
 
 
diff --git a/Engine/Plugins/Experimental/ControlRig/Source/ControlRigEditor/Private/Sequencer/ControlRigParameterTrackEditor.h b/Engine/Plugins/Experimental/ControlRig/Source/ControlRigEditor/Private/Sequencer/ControlRigParameterTrackEditor.h
--- a/Engine/Plugins/Experimental/ControlRig/Source/ControlRigEditor/Private/Sequencer/ControlRigParameterTrackEditor.h
+++ b/Engine/Plugins/Experimental/ControlRig/Source/ControlRigEditor/Private/Sequencer/ControlRigParameterTrackEditor.h
@@ -89,7 +89,9 @@ private:
 	/** Delegate for Selection Changed Event */
 	void OnSelectionChanged(TArray<UMovieSceneTrack*> InTracks);
 
-
+	/** Returns the Edit Mode tools manager hosting the edit mode */
+	FEditorModeTools* GetEditorModeTools() const;
+	FControlRigEditMode* GetEditMode(bool bForceActivate = false) const;
 
 	/** Delegate for  Tree View Changed Event */
 	void OnTreeViewChanged();
@@ -115,8 +117,6 @@ private:
 	void HandleControlModified(UControlRig* Subject, FRigControlElement* ControlElement, const FRigControlModifiedContext& Context);
 	void HandleControlSelected(UControlRig* Subject, FRigControlElement* ControlElement, bool bSelected);
 	void HandleOnInitialized(UControlRig* Subject, const EControlRigState InState, const FName& InEventName);
-	/** Post Edit Delegates */
-	void OnPropagateObjectChanges(UObject* InChangedObject);
 
 	/** Handle Creation for SkelMeshComp or Actor Owner, either may have a binding*/
 	FMovieSceneTrackEditor::FFindOrCreateHandleResult FindOrCreateHandleToSceneCompOrOwner(USceneComponent* InComp);
diff --git a/Engine/Plugins/Experimental/ControlRig/Source/ControlRigEditor/Public/IControlRigEditMode.h b/Engine/Plugins/Experimental/ControlRig/Source/ControlRigEditor/Public/IControlRigEditMode.h
new file mode 100644
--- /dev/null
+++ b/Engine/Plugins/Experimental/ControlRig/Source/ControlRigEditor/Public/IControlRigEditMode.h
@@ -0,0 +1,14 @@
+// Copyright Epic Games, Inc. All Rights Reserved.
+
+#pragma once
+
+#include "CoreMinimal.h"
+#include "IPersonaEditMode.h"
+
+class ISequencer;
+
+class IControlRigEditMode : public IPersonaEditMode
+{
+public:	
+	virtual void SetObjects(const TWeakObjectPtr<>& InSelectedObject,  UObject* BindingObject, TWeakPtr<ISequencer> InSequencer) = 0;
+};
diff --git a/Engine/Source/Editor/Sequencer/Private/Sequencer.h b/Engine/Source/Editor/Sequencer/Private/Sequencer.h
--- a/Engine/Source/Editor/Sequencer/Private/Sequencer.h
+++ b/Engine/Source/Editor/Sequencer/Private/Sequencer.h
@@ -499,7 +499,7 @@ public:
 	EPlaybackMode::Type GetPlaybackMode() const;
 
 	/** @return The toolkit that this sequencer is hosted in (if any) */
-	TSharedPtr<IToolkitHost> GetToolkitHost() const { return ToolkitHost.Pin(); }
+	virtual TSharedPtr<IToolkitHost> GetToolkitHost() const override { return ToolkitHost.Pin(); }
 
 	const FSequencerHostCapabilities& GetHostCapabilities() const { return HostCapabilities; }
 
diff --git a/Engine/Source/Editor/Sequencer/Public/ISequencer.h b/Engine/Source/Editor/Sequencer/Public/ISequencer.h
--- a/Engine/Source/Editor/Sequencer/Public/ISequencer.h
+++ b/Engine/Source/Editor/Sequencer/Public/ISequencer.h
@@ -331,6 +331,9 @@ public:
 		return GetAllowEditsMode() != EAllowEditsMode::AllowLevelEditsOnly || GetAutoChangeMode() != EAutoChangeMode::None;
 	}
 
+	/** Returns the Toolkit hosting the sequencer instance, if any */	
+	virtual TSharedPtr<IToolkitHost> GetToolkitHost() const = 0;
+
 	/**
 	 * Gets the current time of the time slider relative to the currently focused movie scene
 	 */
