Control Rig: Fix world space compensation in Control Rig Component

diff --git a/Engine/Plugins/Experimental/ControlRig/Source/ControlRig/Private/ControlRigComponent.cpp b/Engine/Plugins/Experimental/ControlRig/Source/ControlRig/Private/ControlRigComponent.cpp
--- a/Engine/Plugins/Experimental/ControlRig/Source/ControlRig/Private/ControlRigComponent.cpp
+++ b/Engine/Plugins/Experimental/ControlRig/Source/ControlRig/Private/ControlRigComponent.cpp
@@ -1101,7 +1101,7 @@ void UControlRigComponent::ValidateMappingData()
 
 			if (USkeletalMeshComponent* SkeletalMeshComponent = Cast<USkeletalMeshComponent>(MappedElement.SceneComponent))
 			{
-				MappedElement.Space = EControlRigComponentSpace::ComponentSpace;
+				//MappedElement.Space = EControlRigComponentSpace::ComponentSpace;
 
 				MappedElement.SubIndex = INDEX_NONE;
 				if (MappedElement.TransformIndex >= 0 && MappedElement.TransformIndex < SkeletalMeshComponent->GetNumBones())
@@ -1197,8 +1197,11 @@ void UControlRigComponent::TransferInputs()
 			{
 				if (USkeletalMeshComponent* SkeletalMeshComponent = Cast<USkeletalMeshComponent>(MappedElement.SceneComponent))
 				{
-					// todo: optimize
 					Transform = SkeletalMeshComponent->GetBoneTransform(MappedElement.SubIndex, FTransform::Identity);
+					if(MappedElement.Space == EControlRigComponentSpace::WorldSpace)
+					{
+						Transform = Transform * MappedElement.SceneComponent->GetComponentToWorld();
+					}
 				}
 				else if (UInstancedStaticMeshComponent* InstancingComponent = Cast<UInstancedStaticMeshComponent>(MappedElement.SceneComponent))
 				{
@@ -1278,6 +1281,10 @@ void UControlRigComponent::TransferOutputs()
 					if (Proxy)
 					{
 						ComponentsToTick.AddUnique(Cast<USkeletalMeshComponent>(MappedElement.SceneComponent));
+						if(MappedElement.Space == EControlRigComponentSpace::WorldSpace)
+						{
+							Transform = Transform.GetRelativeTransform(MappedElement.SceneComponent->GetComponentToWorld());
+						}
 						Proxy->StoredTransforms.FindOrAdd(MappedElement.SubIndex) = Transform;
 					}
 					else if (UInstancedStaticMeshComponent* InstancingComponent = Cast<UInstancedStaticMeshComponent>(MappedElement.SceneComponent))
