Control Rig: Introduce linear force to verlet integration

diff --git a/Engine/Plugins/Experimental/ControlRig/Source/ControlRig/Private/Units/Simulation/RigUnit_Verlet.cpp b/Engine/Plugins/Experimental/ControlRig/Source/ControlRig/Private/Units/Simulation/RigUnit_Verlet.cpp
--- a/Engine/Plugins/Experimental/ControlRig/Source/ControlRig/Private/Units/Simulation/RigUnit_Verlet.cpp
+++ b/Engine/Plugins/Experimental/ControlRig/Source/ControlRig/Private/Units/Simulation/RigUnit_Verlet.cpp
@@ -25,9 +25,9 @@ FRigUnit_VerletIntegrateVector_Execute()
 	if (Context.DeltaTime > SMALL_NUMBER)
 	{
 		float U = FMath::Clamp<float>(Blend * Context.DeltaTime, 0.f, 1.f);
-		FVector Force = (Target - Point.Position) * FMath::Max(Strength, 0.0001f);
-		FVector PreviousVelocity = Point.LinearVelocity;
-		Point = Point.IntegrateVerlet(Force, Blend, Context.DeltaTime);
+		const FVector CombinedForce = (Target - Point.Position) * FMath::Max(Strength, 0.0001f) + Force * Context.DeltaTime * 60.f;
+		const FVector PreviousVelocity = Point.LinearVelocity;
+		Point = Point.IntegrateVerlet(CombinedForce, Blend, Context.DeltaTime);
 		Acceleration = (Point.LinearVelocity - PreviousVelocity) / Context.DeltaTime;
 		Position = Point.Position;
 		Velocity = Point.LinearVelocity;
diff --git a/Engine/Plugins/Experimental/ControlRig/Source/ControlRig/Private/Units/Simulation/RigUnit_Verlet.h b/Engine/Plugins/Experimental/ControlRig/Source/ControlRig/Private/Units/Simulation/RigUnit_Verlet.h
--- a/Engine/Plugins/Experimental/ControlRig/Source/ControlRig/Private/Units/Simulation/RigUnit_Verlet.h
+++ b/Engine/Plugins/Experimental/ControlRig/Source/ControlRig/Private/Units/Simulation/RigUnit_Verlet.h
@@ -16,7 +16,7 @@ struct CONTROLRIG_API FRigUnit_VerletIntegrateVector : public FRigUnit_SimBase
 	
 	FRigUnit_VerletIntegrateVector()
 	{
-		Target = Position = Velocity = Acceleration = FVector::ZeroVector;
+		Target = Force = Position = Velocity = Acceleration = FVector::ZeroVector;
 		Strength = 64.f;
 		Damp = 0.01;
 		Blend = 5.f;
@@ -42,6 +42,10 @@ struct CONTROLRIG_API FRigUnit_VerletIntegrateVector : public FRigUnit_SimBase
 	UPROPERTY(meta = (Input))
 	float Blend;
 
+	/** The force feeding into the solver. Can be used for gravity. */
+	UPROPERTY(meta = (Input))
+	FVector Force;
+
 	UPROPERTY(meta = (Output))
 	FVector Position;
 
