Sequencer: Control Rig: GetBoundObjects by creating a temporary player to support spawnables

diff --git a/Engine/Plugins/Experimental/ControlRig/Source/ControlRigEditor/Private/ControlRigSequencerEditorLibrary.cpp b/Engine/Plugins/Experimental/ControlRig/Source/ControlRigEditor/Private/ControlRigSequencerEditorLibrary.cpp
--- a/Engine/Plugins/Experimental/ControlRig/Source/ControlRigEditor/Private/ControlRigSequencerEditorLibrary.cpp
+++ b/Engine/Plugins/Experimental/ControlRig/Source/ControlRigEditor/Private/ControlRigSequencerEditorLibrary.cpp
@@ -79,6 +79,28 @@ TArray<FControlRigSequencerBindingProxy> UControlRigSequencerEditorLibrary::GetC
 	return ControlRigBindingProxies;
 }
 
+static TArray<UObject*> GetBoundObjects(UWorld* World, ULevelSequence* LevelSequence, const FSequencerBindingProxy& Binding, ULevelSequencePlayer** OutPlayer, ALevelSequenceActor** OutActor)
+{
+	FMovieSceneSequencePlaybackSettings Settings;
+	FLevelSequenceCameraSettings CameraSettings;
+
+	ULevelSequencePlayer* Player = ULevelSequencePlayer::CreateLevelSequencePlayer(World, LevelSequence, Settings, *OutActor);
+	*OutPlayer = Player;
+
+	Player->Initialize(LevelSequence, World->PersistentLevel, Settings, CameraSettings);
+	Player->State.AssignSequence(MovieSceneSequenceID::Root, *LevelSequence, *Player);
+
+	// Evaluation needs to occur in order to obtain spawnables
+	Player->SetPlaybackPosition(FMovieSceneSequencePlaybackParams(LevelSequence->GetMovieScene()->GetPlaybackRange().GetLowerBoundValue().Value, EUpdatePositionMethod::Play));
+
+	FMovieSceneSequenceID SequenceId = Player->State.FindSequenceId(LevelSequence);
+
+	FMovieSceneObjectBindingID ObjectBinding = UE::MovieScene::FFixedObjectBindingID(Binding.BindingID, SequenceId);
+	TArray<UObject*> BoundObjects = Player->GetBoundObjects(ObjectBinding);
+	
+	return BoundObjects;
+}
+
 static void AcquireSkeletonAndSkelMeshCompFromObject(UObject* BoundObject, USkeleton** OutSkeleton, USkeletalMeshComponent** OutSkeletalMeshComponent)
 {
 	*OutSkeletalMeshComponent = nullptr;
@@ -233,8 +255,7 @@ static UMovieSceneControlRigParameterTrack* AddControlRig(ULevelSequence* LevelS
 	return nullptr;
 }
 
-
-UMovieSceneTrack* UControlRigSequencerEditorLibrary::FindOrCreateControlRigTrack(ULevelSequence* LevelSequence, const UClass* ControlRigClass, const FSequencerBindingProxy& InBinding)
+UMovieSceneTrack* UControlRigSequencerEditorLibrary::FindOrCreateControlRigTrack(UWorld* World, ULevelSequence* LevelSequence, const UClass* ControlRigClass, const FSequencerBindingProxy& InBinding)
 {
 	UMovieScene* MovieScene = InBinding.Sequence ? InBinding.Sequence->GetMovieScene() : nullptr;
 	UMovieSceneTrack* BaseTrack = nullptr;
@@ -257,7 +278,9 @@ UMovieSceneTrack* UControlRigSequencerEditorLibrary::FindOrCreateControlRigTrack
 
 				TArray<UObject*, TInlineAllocator<1>> Result;
 				UObject* Context = nullptr;
-				LevelSequence->LocateBoundObjects(InBinding.BindingID, Context, Result);
+				ALevelSequenceActor* OutActor = nullptr;
+				ULevelSequencePlayer* OutPlayer = nullptr;
+				Result = GetBoundObjects(World, LevelSequence, InBinding, &OutPlayer, &OutActor);
 				if (Result.Num() > 0 && Result[0])
 				{
 					UObject* BoundObject = Result[0];
@@ -276,13 +299,24 @@ UMovieSceneTrack* UControlRigSequencerEditorLibrary::FindOrCreateControlRigTrack
 						}
 					}
 				}
+
+				if (OutPlayer)
+				{
+					OutPlayer->Stop();
+				}
+
+				if (OutActor)
+				{
+					World->DestroyActor(OutActor);
+				}
 			}
 		}
 	}
 	return BaseTrack;
 }
 
-TArray<UMovieSceneTrack*> UControlRigSequencerEditorLibrary::FindOrCreateControlRigComponentTrack(ULevelSequence* LevelSequence, const FSequencerBindingProxy& InBinding)
+
+TArray<UMovieSceneTrack*> UControlRigSequencerEditorLibrary::FindOrCreateControlRigComponentTrack(UWorld* World, ULevelSequence* LevelSequence, const FSequencerBindingProxy& InBinding)
 {
 	TArray< UMovieSceneTrack*> Tracks;
 	TArray<UObject*, TInlineAllocator<1>> Result;
@@ -292,7 +326,9 @@ TArray<UMovieSceneTrack*> UControlRigSequencerEditorLibrary::FindOrCreateControl
 	}
 	UMovieScene* MovieScene = LevelSequence->GetMovieScene();
 	UObject* Context = nullptr;
-	LevelSequence->LocateBoundObjects(InBinding.BindingID, Context, Result);
+	ALevelSequenceActor* OutActor = nullptr;
+	ULevelSequencePlayer* OutPlayer = nullptr;
+	Result = GetBoundObjects(World, LevelSequence, InBinding, &OutPlayer, &OutActor);
 	if (Result.Num() > 0 && Result[0])
 	{
 		UObject* BoundObject = Result[0];
@@ -329,10 +365,20 @@ TArray<UMovieSceneTrack*> UControlRigSequencerEditorLibrary::FindOrCreateControl
 			}
 		}
 	}
+
+	if (OutPlayer)
+	{
+		OutPlayer->Stop();
+	}
+
+	if (OutActor)
+	{
+		World->DestroyActor(OutActor);
+	}
+
 	return Tracks;
 }
 
-
 bool UControlRigSequencerEditorLibrary::TweenControlRig(ULevelSequence* LevelSequence, UControlRig* ControlRig, float TweenValue)
 {
 	TWeakPtr<ISequencer> WeakSequencer = GetSequencerFromAsset(LevelSequence);
diff --git a/Engine/Plugins/Experimental/ControlRig/Source/ControlRigEditor/Public/ControlRigSequencerEditorLibrary.h b/Engine/Plugins/Experimental/ControlRig/Source/ControlRigEditor/Public/ControlRigSequencerEditorLibrary.h
--- a/Engine/Plugins/Experimental/ControlRig/Source/ControlRigEditor/Public/ControlRigSequencerEditorLibrary.h
+++ b/Engine/Plugins/Experimental/ControlRig/Source/ControlRigEditor/Public/ControlRigSequencerEditorLibrary.h
@@ -70,22 +70,24 @@ public:
 	
 	/**
 	* Find or create a Control Rig track of a specific class based upon the binding
+	* @param World The world used to spawn into temporarily if binding is a spawnable
 	* @param LevelSequence The LevelSequence to find or create
 	* @param ControlRigClass The class of the Control Rig
 	* @param InBinding The binding (actor or component binding) to find or create the Control Rig track
 	* @return returns Return the found or created track
 	*/
 	UFUNCTION(BlueprintCallable, Category = "Editor Scripting | Sequencer Tools | Control Rig")
-	static UMovieSceneTrack* FindOrCreateControlRigTrack(ULevelSequence* LevelSequence, const UClass* ControlRigClass, const FSequencerBindingProxy& InBinding);
+	static UMovieSceneTrack* FindOrCreateControlRigTrack(UWorld* World, ULevelSequence* LevelSequence, const UClass* ControlRigClass, const FSequencerBindingProxy& InBinding);
 
 	/**
 	* Find or create a Control Rig Component
+	* @param World The world used to spawn into temporarily if binding is a spawnable
 	* @param LevelSequence The LevelSequence to find or create
 	* @param InBinding The binding (actor or component binding) to find or create the Control Rig tracks
 	* @return returns Find array of component Control Rigs that were found or created
 	*/
 	UFUNCTION(BlueprintCallable, Category = "Editor Scripting | Sequencer Tools | Control Rig")
-	static TArray<UMovieSceneTrack*> FindOrCreateControlRigComponentTrack(ULevelSequence* LevelSequence, const FSequencerBindingProxy& InBinding);
+	static TArray<UMovieSceneTrack*> FindOrCreateControlRigComponentTrack(UWorld* World, ULevelSequence* LevelSequence, const FSequencerBindingProxy& InBinding);
 	
 	/**
 	* Load anim sequence into this control rig section
