Control Rig: Fix static analysis problem with ControlRigComponent

diff --git a/Engine/Plugins/Experimental/ControlRig/Source/ControlRig/Private/ControlRigComponent.cpp b/Engine/Plugins/Experimental/ControlRig/Source/ControlRig/Private/ControlRigComponent.cpp
--- a/Engine/Plugins/Experimental/ControlRig/Source/ControlRig/Private/ControlRigComponent.cpp
+++ b/Engine/Plugins/Experimental/ControlRig/Source/ControlRig/Private/ControlRigComponent.cpp
@@ -1278,7 +1278,7 @@ void UControlRigComponent::TransferOutputs()
 						}
 					}
 
-					if (Proxy)
+					if (Proxy && (MappedElement.SceneComponent != nullptr))
 					{
 						ComponentsToTick.AddUnique(Cast<USkeletalMeshComponent>(MappedElement.SceneComponent));
 						if(MappedElement.Space == EControlRigComponentSpace::WorldSpace)
