Control Rig: Pin Name List Combo Box to reflect chosen type

diff --git a/Engine/Plugins/Experimental/ControlRig/Source/ControlRigEditor/Private/Graph/SControlRigGraphPinNameList.cpp b/Engine/Plugins/Experimental/ControlRig/Source/ControlRigEditor/Private/Graph/SControlRigGraphPinNameList.cpp
--- a/Engine/Plugins/Experimental/ControlRig/Source/ControlRigEditor/Private/Graph/SControlRigGraphPinNameList.cpp
+++ b/Engine/Plugins/Experimental/ControlRig/Source/ControlRigEditor/Private/Graph/SControlRigGraphPinNameList.cpp
@@ -135,5 +135,6 @@ void SControlRigGraphPinNameList::OnNameListComboBox()
 			CurrentlySelected = Item;
 		}
 	}
+	NameListComboBox->SetOptionsSource(CurrentList);
 	NameListComboBox->SetSelectedItem(CurrentlySelected);
 }
diff --git a/Engine/Plugins/Experimental/ControlRig/Source/ControlRigEditor/Private/Graph/SControlRigGraphPinNameListValueWidget.cpp b/Engine/Plugins/Experimental/ControlRig/Source/ControlRigEditor/Private/Graph/SControlRigGraphPinNameListValueWidget.cpp
--- a/Engine/Plugins/Experimental/ControlRig/Source/ControlRigEditor/Private/Graph/SControlRigGraphPinNameListValueWidget.cpp
+++ b/Engine/Plugins/Experimental/ControlRig/Source/ControlRigEditor/Private/Graph/SControlRigGraphPinNameListValueWidget.cpp
@@ -99,6 +99,13 @@ void SControlRigGraphPinNameListValueWidget::SetSelectedItem(TSharedPtr<FString>
 	}
 }
 
+void SControlRigGraphPinNameListValueWidget::SetOptionsSource(const TArray<TSharedPtr<FString>>* InOptionsSource)
+{
+	check(InOptionsSource);
+	OptionsSource = InOptionsSource;
+	ComboListView->SetListItemsSource(*OptionsSource);
+}
+
 TSharedPtr<FString> SControlRigGraphPinNameListValueWidget::GetSelectedItem()
 {
 	return SelectedItem;
diff --git a/Engine/Plugins/Experimental/ControlRig/Source/ControlRigEditor/Private/Graph/SControlRigGraphPinNameListValueWidget.h b/Engine/Plugins/Experimental/ControlRig/Source/ControlRigEditor/Private/Graph/SControlRigGraphPinNameListValueWidget.h
--- a/Engine/Plugins/Experimental/ControlRig/Source/ControlRigEditor/Private/Graph/SControlRigGraphPinNameListValueWidget.h
+++ b/Engine/Plugins/Experimental/ControlRig/Source/ControlRigEditor/Private/Graph/SControlRigGraphPinNameListValueWidget.h
@@ -75,6 +75,7 @@ public:
 	void ClearSelection();
 
 	void SetSelectedItem(TSharedPtr<FString> InSelectedItem);
+	void SetOptionsSource(const TArray< TSharedPtr<FString> >* InOptionsSource);
 
 	/** @return the item currently selected by the combo box. */
 	TSharedPtr<FString> GetSelectedItem();
