Control Rig: Don't manage / update the UI graph for the function library

diff --git a/Engine/Plugins/Experimental/ControlRig/Source/ControlRigDeveloper/Private/ControlRigBlueprint.cpp b/Engine/Plugins/Experimental/ControlRig/Source/ControlRigDeveloper/Private/ControlRigBlueprint.cpp
--- a/Engine/Plugins/Experimental/ControlRig/Source/ControlRigDeveloper/Private/ControlRigBlueprint.cpp
+++ b/Engine/Plugins/Experimental/ControlRig/Source/ControlRigDeveloper/Private/ControlRigBlueprint.cpp
@@ -3367,6 +3367,8 @@ bool UControlRigBlueprint::RemoveEdGraphForCollapseNode(URigVMCollapseNode* InNo
 						}
 
 						FunctionGraphs.Remove(RigFunctionGraph);
+						RigFunctionGraph->Rename(nullptr, GetTransientPackage());
+						RigFunctionGraph->MarkPendingKill();
 						return bNotify;
 					}
 				}
@@ -3394,6 +3396,8 @@ bool UControlRigBlueprint::RemoveEdGraphForCollapseNode(URigVMCollapseNode* InNo
 						}
 
 						RigGraph->SubGraphs.Remove(SubRigGraph);
+						SubRigGraph->Rename(nullptr, GetTransientPackage());
+						SubRigGraph->MarkPendingKill();
 						return bNotify;
 					}
 				}
diff --git a/Engine/Plugins/Experimental/ControlRig/Source/ControlRigDeveloper/Private/Graph/ControlRigGraph.cpp b/Engine/Plugins/Experimental/ControlRig/Source/ControlRigDeveloper/Private/Graph/ControlRigGraph.cpp
--- a/Engine/Plugins/Experimental/ControlRig/Source/ControlRigDeveloper/Private/Graph/ControlRigGraph.cpp
+++ b/Engine/Plugins/Experimental/ControlRig/Source/ControlRigDeveloper/Private/Graph/ControlRigGraph.cpp
@@ -195,6 +195,12 @@ void UControlRigGraph::HandleModifiedEvent(ERigVMGraphNotifType InNotifType, URi
 		return;
 	}
 
+	// We don't need a visual representation for the low level function library graph
+	if(InGraph->IsA<URigVMFunctionLibrary>())
+	{
+		return;
+	}
+
 	// increment the node topology version for any interaction
 	// with a node.
 	{
