Control Rig: Fix detail panel element key name list refresh

diff --git a/Engine/Plugins/Experimental/ControlRig/Source/ControlRigEditor/Private/ControlRigElementDetails.cpp b/Engine/Plugins/Experimental/ControlRig/Source/ControlRigEditor/Private/ControlRigElementDetails.cpp
--- a/Engine/Plugins/Experimental/ControlRig/Source/ControlRigEditor/Private/ControlRigElementDetails.cpp
+++ b/Engine/Plugins/Experimental/ControlRig/Source/ControlRigEditor/Private/ControlRigElementDetails.cpp
@@ -13,7 +13,6 @@
 #include "Widgets/Input/SCheckBox.h"
 #include "Widgets/Input/SNumericEntryBox.h"
 #include "Widgets/Input/SComboBox.h"
-#include "SSearchableComboBox.h"
 #include "Editor/SControlRigGizmoNameList.h"
 #include "ControlRigBlueprint.h"
 #include "Graph/ControlRigGraph.h"
@@ -162,7 +161,7 @@ void FRigElementKeyDetails::CustomizeHeader(TSharedRef<IPropertyHandle> InStruct
 			.AutoWidth()
 			.Padding(4.f, 0.f, 0.f, 0.f)
 			[
-				SNew(SSearchableComboBox)
+				SAssignNew(SearchableComboBox, SSearchableComboBox)
 				.OptionsSource(&ElementNameList)
 				.OnSelectionChanged(this, &FRigElementKeyDetails::OnElementNameChanged)
 				.OnGenerateWidget(this, &FRigElementKeyDetails::OnGetElementNameWidget)
@@ -243,6 +242,10 @@ void FRigElementKeyDetails::UpdateElementNameList()
 			if (UControlRigGraph* RigGraph = Cast<UControlRigGraph>(Graph))
 			{
 				ElementNameList = *RigGraph->GetElementNameList(GetElementType());
+				if(SearchableComboBox.IsValid())
+				{
+					SearchableComboBox->RefreshOptions();
+				}
 				return;
 			}
 		}
diff --git a/Engine/Plugins/Experimental/ControlRig/Source/ControlRigEditor/Private/ControlRigElementDetails.h b/Engine/Plugins/Experimental/ControlRig/Source/ControlRigEditor/Private/ControlRigElementDetails.h
--- a/Engine/Plugins/Experimental/ControlRig/Source/ControlRigEditor/Private/ControlRigElementDetails.h
+++ b/Engine/Plugins/Experimental/ControlRig/Source/ControlRigEditor/Private/ControlRigElementDetails.h
@@ -13,6 +13,7 @@
 #include "Graph/SControlRigGraphPinNameListValueWidget.h"
 #include "Styling/SlateTypes.h"
 #include "IPropertyUtilities.h"
+#include "SSearchableComboBox.h"
 
 class IPropertyHandle;
 
@@ -43,6 +44,7 @@ protected:
 	TSharedPtr<IPropertyHandle> NameHandle;
 	TArray<TSharedPtr<FString>> ElementNameList;
 	UControlRigBlueprint* BlueprintBeingCustomized;
+	TSharedPtr<SSearchableComboBox> SearchableComboBox;
 };
 
 class FRigUnitDetails : public IDetailCustomization
