Control Rig: Control Rig Editor to open MyBlueprint widget as non-floating and only if required

diff --git a/Engine/Plugins/Experimental/ControlRig/Source/ControlRigEditor/Private/Editor/ControlRigEditor.cpp b/Engine/Plugins/Experimental/ControlRig/Source/ControlRigEditor/Private/Editor/ControlRigEditor.cpp
--- a/Engine/Plugins/Experimental/ControlRig/Source/ControlRigEditor/Private/Editor/ControlRigEditor.cpp
+++ b/Engine/Plugins/Experimental/ControlRig/Source/ControlRigEditor/Private/Editor/ControlRigEditor.cpp
@@ -350,9 +350,6 @@ void FControlRigEditor::InitControlRigEditor(const EToolkitMode::Type Mode, cons
 	// Activate the initial mode (which will populate with a real layout)
 	SetCurrentMode(FControlRigEditorModes::ControlRigEditorMode);
 
-	// Always show the myblueprint tab
-	GetTabManager()->TryInvokeTab(FTabId(FBlueprintEditorTabs::MyBlueprintID));
-
 	// Activate our edit mode
 	GetEditorModeManager().SetDefaultMode(FControlRigEditorEditMode::ModeName);
 	GetEditorModeManager().ActivateMode(FControlRigEditorEditMode::ModeName);
@@ -492,6 +489,17 @@ void FControlRigEditor::InitControlRigEditor(const EToolkitMode::Type Mode, cons
 	}
 	
 	OnToolbarBoneRadiusChanged(InControlRigBlueprint->DebugBoneRadius);
+
+	FFunctionGraphTask::CreateAndDispatchWhenReady([this]()
+	{
+		// Always show the myblueprint tab
+		FTabId MyBlueprintTabId(FBlueprintEditorTabs::MyBlueprintID);
+		if(!GetTabManager()->FindExistingLiveTab(MyBlueprintTabId).IsValid())
+		{
+			GetTabManager()->TryInvokeTab(MyBlueprintTabId);
+		}
+		
+	}, TStatId(), NULL, ENamedThreads::GameThread);
 }
 
 void FControlRigEditor::BindCommands()
