Control Rig Editor: Allow to edit bone transforms in setup mode

diff --git a/Engine/Plugins/Experimental/ControlRig/Source/ControlRigEditor/Private/Editor/ControlRigEditor.cpp b/Engine/Plugins/Experimental/ControlRig/Source/ControlRigEditor/Private/Editor/ControlRigEditor.cpp
--- a/Engine/Plugins/Experimental/ControlRig/Source/ControlRigEditor/Private/Editor/ControlRigEditor.cpp
+++ b/Engine/Plugins/Experimental/ControlRig/Source/ControlRigEditor/Private/Editor/ControlRigEditor.cpp
@@ -3629,6 +3629,7 @@ void FControlRigEditor::OnBlueprintPropertyChainEvent(FPropertyChangedChainEvent
 		if(bSetupModeEnabled || bIsInitial)
 		{
 			Hierarchy = ControlRigBP->Hierarchy;
+			DebuggedHierarchy = Hierarchy;
 		}
 
 		FRigBaseElement* SourceElement = DebuggedHierarchy->Find(RigElementInDetailPanel);
@@ -3705,7 +3706,7 @@ void FControlRigEditor::OnBlueprintPropertyChainEvent(FPropertyChangedChainEvent
 			}
 		}
 
-		if(bIsInitial)
+		if(bSetupModeEnabled || bIsInitial)
 		{
 			ControlRigBP->PropagatePoseFromBPToInstances();
 			ControlRigBP->Modify();
