Control rig sequencer related non-unity fixes

diff --git a/Engine/Plugins/Experimental/ControlRig/Source/ControlRigEditor/Private/EditMode/SControlRigEditModeTools.h b/Engine/Plugins/Experimental/ControlRig/Source/ControlRigEditor/Private/EditMode/SControlRigEditModeTools.h
--- a/Engine/Plugins/Experimental/ControlRig/Source/ControlRigEditor/Private/EditMode/SControlRigEditModeTools.h
+++ b/Engine/Plugins/Experimental/ControlRig/Source/ControlRigEditor/Private/EditMode/SControlRigEditModeTools.h
@@ -16,6 +16,7 @@ class SExpandableArea;
 class SControlHierarchy;
 class UControlRig;
 class FToolBarBuilder;
+class FEditorModeTools;
 
 class SControlRigEditModeTools : public SCompoundWidget, public IDetailKeyframeHandler
 {
diff --git a/Engine/Plugins/Experimental/ControlRig/Source/ControlRigEditor/Private/Sequencer/ControlRigParameterTrackEditor.cpp b/Engine/Plugins/Experimental/ControlRig/Source/ControlRigEditor/Private/Sequencer/ControlRigParameterTrackEditor.cpp
--- a/Engine/Plugins/Experimental/ControlRig/Source/ControlRigEditor/Private/Sequencer/ControlRigParameterTrackEditor.cpp
+++ b/Engine/Plugins/Experimental/ControlRig/Source/ControlRigEditor/Private/Sequencer/ControlRigParameterTrackEditor.cpp
@@ -75,6 +75,7 @@
 #include "BakeToControlRigSettings.h"
 #include "UnrealEdGlobals.h"
 #include "Editor/UnrealEdEngine.h"
+#include "Toolkits/IToolkitHost.h"
 
 #define LOCTEXT_NAMESPACE "FControlRigParameterTrackEditor"
 
diff --git a/Engine/Plugins/Experimental/ControlRig/Source/ControlRigEditor/Private/Sequencer/ControlRigParameterTrackEditor.h b/Engine/Plugins/Experimental/ControlRig/Source/ControlRigEditor/Private/Sequencer/ControlRigParameterTrackEditor.h
--- a/Engine/Plugins/Experimental/ControlRig/Source/ControlRigEditor/Private/Sequencer/ControlRigParameterTrackEditor.h
+++ b/Engine/Plugins/Experimental/ControlRig/Source/ControlRigEditor/Private/Sequencer/ControlRigParameterTrackEditor.h
@@ -23,6 +23,7 @@ class FMenuBuilder;
 class USkeleton;
 class UMovieSceneControlRigParameterSection;
 class UFKControlRig;
+class FEditorModeTools;
 
 /**
  * Tools for animation tracks
diff --git a/Engine/Source/Editor/Sequencer/Public/ISequencer.h b/Engine/Source/Editor/Sequencer/Public/ISequencer.h
--- a/Engine/Source/Editor/Sequencer/Public/ISequencer.h
+++ b/Engine/Source/Editor/Sequencer/Public/ISequencer.h
@@ -37,6 +37,7 @@ enum class EMapChangeType : uint8;
 class FCurveModel;
 struct FMovieSceneSequencePlaybackParams;
 struct FMovieSceneChannelMetaData;
+class IToolkitHost;
 
 /**
  * Defines auto change modes.
