Fix control rig create pose dialog not displaying any assets

diff --git a/Engine/Plugins/Experimental/ControlRig/Source/ControlRigEditor/ControlRigEditor.Build.cs b/Engine/Plugins/Experimental/ControlRig/Source/ControlRigEditor/ControlRigEditor.Build.cs
--- a/Engine/Plugins/Experimental/ControlRig/Source/ControlRigEditor/ControlRigEditor.Build.cs
+++ b/Engine/Plugins/Experimental/ControlRig/Source/ControlRigEditor/ControlRigEditor.Build.cs
@@ -59,6 +59,7 @@ namespace UnrealBuildTool.Rules
                     "ClassViewer",
                     "AssetTools",
                     "ContentBrowser",
+					"ContentBrowserData",
                     "LevelEditor",
                     "SceneOutliner",
                     "LevelSequence",
diff --git a/Engine/Plugins/Experimental/ControlRig/Source/ControlRigEditor/Private/EditMode/SControlRigBaseListWidget.cpp b/Engine/Plugins/Experimental/ControlRig/Source/ControlRigEditor/Private/EditMode/SControlRigBaseListWidget.cpp
--- a/Engine/Plugins/Experimental/ControlRig/Source/ControlRigEditor/Private/EditMode/SControlRigBaseListWidget.cpp
+++ b/Engine/Plugins/Experimental/ControlRig/Source/ControlRigEditor/Private/EditMode/SControlRigBaseListWidget.cpp
@@ -25,6 +25,8 @@
 
 #include "IContentBrowserSingleton.h"
 #include "ContentBrowserModule.h"
+#include "IContentBrowserDataModule.h"
+#include "ContentBrowserDataSubsystem.h"
 #include "Tools/ControlRigPoseProjectSettings.h"
 #include "Widgets/Layout/SScrollBox.h"
 #include "Framework/Commands/GenericCommands.h"
@@ -40,7 +42,7 @@
 
 #define LOCTEXT_NAMESPACE "ControlRigBaseListWidget"
 
-FString SControlRigBaseListWidget::CurrentlySelectedPath("");
+FString SControlRigBaseListWidget::CurrentlySelectedInternalPath("");
 
 enum class FControlRigAssetType {
 	ControlRigPose,
@@ -664,11 +666,15 @@ void SControlRigBaseListWidget::Construct(const FArguments& InArgs)
 	TArray<FString> PosesDirectories = PoseSettings->GetAssetPaths();
 	if (PosesDirectories.Num() > 0)
 	{
-		if (CurrentlySelectedPath.IsEmpty())
+		if (CurrentlySelectedInternalPath.IsEmpty())
 		{
-			CurrentlySelectedPath = PosesDirectories[0];
+			CurrentlySelectedInternalPath = PosesDirectories[0];
 		}
 	}
+
+	FName CurrentlySelectedVirtualPath;
+	IContentBrowserDataModule::Get().GetSubsystem()->ConvertInternalPathToVirtual(FStringView(CurrentlySelectedInternalPath), CurrentlySelectedVirtualPath);
+
 	FContentBrowserModule& ContentBrowserModule = FModuleManager::Get().LoadModuleChecked<FContentBrowserModule>(TEXT("ContentBrowser"));
 
 	// Configure filter for asset picker
@@ -683,7 +689,7 @@ void SControlRigBaseListWidget::Construct(const FArguments& InArgs)
 	AssetPickerConfig.bCanShowRealTimeThumbnails = true;
 	AssetPickerConfig.ThumbnailLabel = EThumbnailLabel::AssetName;
 	AssetPickerConfig.bFocusSearchBoxWhenOpened = false;
-	AssetPickerConfig.Filter.PackagePaths.Add(FName(*CurrentlySelectedPath));
+	AssetPickerConfig.Filter.PackagePaths.Add(CurrentlySelectedVirtualPath);
 	AssetPickerConfig.OnAssetSelected = FOnAssetSelected::CreateSP(this, &SControlRigBaseListWidget::OnAssetSelected);
 	AssetPickerConfig.OnAssetsActivated = FOnAssetsActivated::CreateSP(this, &SControlRigBaseListWidget::OnAssetsActivated);
 	AssetPickerConfig.SaveSettingsName = TEXT("ControlPoseDialog");
@@ -699,7 +705,8 @@ void SControlRigBaseListWidget::Construct(const FArguments& InArgs)
 
 	FPathPickerConfig PathPickerConfig;
 	PathPickerConfig.bAddDefaultPath = true;
-	PathPickerConfig.DefaultPath = CurrentlySelectedPath;
+	PathPickerConfig.bOnPathSelectedPassesVirtualPaths = false;
+	PathPickerConfig.DefaultPath = CurrentlySelectedInternalPath;
 	CustomFolderBlacklist = PathPickerConfig.CustomFolderBlacklist = MakeShared<FBlacklistPaths>();
 	for (const FString& Path : PosesDirectories)
 	{
@@ -816,18 +823,18 @@ FText SControlRigBaseListWidget::GetAssetNameText() const
 
 FText SControlRigBaseListWidget::GetPathNameText() const
 {
-	return FText::FromString(CurrentlySelectedPath);
+	return FText::FromString(CurrentlySelectedInternalPath);
 }
 
 void SControlRigBaseListWidget::SetCurrentlySelectedPath(const FString& NewPath)
 {
-	CurrentlySelectedPath = NewPath;
+	CurrentlySelectedInternalPath = NewPath;
 	UpdateInputValidity();
 }
 
 FString SControlRigBaseListWidget::GetCurrentlySelectedPath() const
 {
-	return CurrentlySelectedPath;
+	return CurrentlySelectedInternalPath;
 }
 void SControlRigBaseListWidget::SetCurrentlyEnteredAssetName(const FString& NewName)
 {
@@ -847,7 +854,7 @@ void SControlRigBaseListWidget::UpdateInputValidity()
 
 	if (bLastInputValidityCheckSuccessful)
 	{
-		if (CurrentlySelectedPath.IsEmpty())
+		if (CurrentlySelectedInternalPath.IsEmpty())
 		{
 			bLastInputValidityCheckSuccessful = false;
 		}
@@ -872,7 +879,7 @@ void SControlRigBaseListWidget::UpdateInputValidity()
 
 FString SControlRigBaseListWidget::GetObjectPathForSave() const
 {
-	return CurrentlySelectedPath / CurrentlyEnteredAssetName + TEXT(".") + CurrentlyEnteredAssetName;
+	return CurrentlySelectedInternalPath / CurrentlyEnteredAssetName + TEXT(".") + CurrentlyEnteredAssetName;
 }
 
 void SControlRigBaseListWidget::SelectThisAsset(UObject* Asset)
@@ -929,8 +936,11 @@ void SControlRigBaseListWidget::OnAssetsActivated(const TArray<FAssetData>& Sele
 
 void SControlRigBaseListWidget::FilterChanged()
 {
+	FName CurrentlySelectedVirtualPath;
+	IContentBrowserDataModule::Get().GetSubsystem()->ConvertInternalPathToVirtual(FStringView(CurrentlySelectedInternalPath), CurrentlySelectedVirtualPath);
+
 	FARFilter NewFilter;
-	NewFilter.PackagePaths.Add(FName(*CurrentlySelectedPath));
+	NewFilter.PackagePaths.Add(CurrentlySelectedVirtualPath);
 	NewFilter.ClassNames.Add(UControlRigPoseAsset::StaticClass()->GetFName());
 	SetFilterDelegate.ExecuteIfBound(NewFilter);
 }
diff --git a/Engine/Plugins/Experimental/ControlRig/Source/ControlRigEditor/Private/EditMode/SControlRigBaseListWidget.h b/Engine/Plugins/Experimental/ControlRig/Source/ControlRigEditor/Private/EditMode/SControlRigBaseListWidget.h
--- a/Engine/Plugins/Experimental/ControlRig/Source/ControlRigEditor/Private/EditMode/SControlRigBaseListWidget.h
+++ b/Engine/Plugins/Experimental/ControlRig/Source/ControlRigEditor/Private/EditMode/SControlRigBaseListWidget.h
@@ -124,7 +124,7 @@ private:
 
 	TSharedPtr<FAssetFileContextMenu> AssetFileContextMenu;
 
-	static FString CurrentlySelectedPath;
+	static FString CurrentlySelectedInternalPath;
 	FString CurrentlyEnteredAssetName;
 	bool bLastInputValidityCheckSuccessful;
 
