Build issue 73700 - non unity fix

diff --git a/Engine/Plugins/Experimental/ControlRig/Source/ControlRigEditor/Private/Sequencer/ControlRigParameterTrackEditor.h b/Engine/Plugins/Experimental/ControlRig/Source/ControlRigEditor/Private/Sequencer/ControlRigParameterTrackEditor.h
--- a/Engine/Plugins/Experimental/ControlRig/Source/ControlRigEditor/Private/Sequencer/ControlRigParameterTrackEditor.h
+++ b/Engine/Plugins/Experimental/ControlRig/Source/ControlRigEditor/Private/Sequencer/ControlRigParameterTrackEditor.h
@@ -24,6 +24,7 @@ class USkeleton;
 class UMovieSceneControlRigParameterSection;
 class UFKControlRig;
 class FEditorModeTools;
+class FControlRigEditMode;
 
 /**
  * Tools for animation tracks
