ControlRig: Use Selected Item helper button on pin

diff --git a/Engine/Plugins/Experimental/ControlRig/Source/ControlRigDeveloper/Private/Graph/ControlRigGraph.cpp b/Engine/Plugins/Experimental/ControlRig/Source/ControlRigDeveloper/Private/Graph/ControlRigGraph.cpp
--- a/Engine/Plugins/Experimental/ControlRig/Source/ControlRigDeveloper/Private/Graph/ControlRigGraph.cpp
+++ b/Engine/Plugins/Experimental/ControlRig/Source/ControlRigDeveloper/Private/Graph/ControlRigGraph.cpp
@@ -170,6 +170,31 @@ const TArray<TSharedPtr<FString>>* UControlRigGraph::GetElementNameList(URigVMPi
 	return GetBoneNameList(nullptr);
 }
 
+const TArray<TSharedPtr<FString>> UControlRigGraph::GetSelectedElementsNameList() const
+{
+	TArray<TSharedPtr<FString>> Result;
+	if (UControlRigGraph* OuterGraph = Cast<UControlRigGraph>(GetOuter()))
+	{
+		return OuterGraph->GetSelectedElementsNameList();
+	}
+	
+	UControlRigBlueprint* Blueprint = GetBlueprint();
+	if(Blueprint == nullptr)
+	{
+		return Result;
+	}
+	
+	TArray<FRigElementKey> Keys = Blueprint->Hierarchy->GetSelectedKeys();
+	for (const FRigElementKey& Key : Keys)
+	{
+		FString ValueStr;
+		FRigElementKey::StaticStruct()->ExportText(ValueStr, &Key, nullptr, nullptr, PPF_None, nullptr);
+		Result.Add(MakeShared<FString>(ValueStr));
+	}
+	
+	return Result;
+}
+
 const TArray<TSharedPtr<FString>>* UControlRigGraph::GetDrawingNameList(URigVMPin* InPin) const
 {
 	if (UControlRigGraph* OuterGraph = Cast<UControlRigGraph>(GetOuter()))
diff --git a/Engine/Plugins/Experimental/ControlRig/Source/ControlRigDeveloper/Public/Graph/ControlRigGraph.h b/Engine/Plugins/Experimental/ControlRig/Source/ControlRigDeveloper/Public/Graph/ControlRigGraph.h
--- a/Engine/Plugins/Experimental/ControlRig/Source/ControlRigDeveloper/Public/Graph/ControlRigGraph.h
+++ b/Engine/Plugins/Experimental/ControlRig/Source/ControlRigDeveloper/Public/Graph/ControlRigGraph.h
@@ -60,6 +60,7 @@ public:
 	void CacheNameLists(URigHierarchy* InHierarchy, const FControlRigDrawContainer* DrawContainer);
 	const TArray<TSharedPtr<FString>>* GetElementNameList(ERigElementType InElementType = ERigElementType::Bone) const;
 	const TArray<TSharedPtr<FString>>* GetElementNameList(URigVMPin* InPin = nullptr) const;
+	const TArray<TSharedPtr<FString>> GetSelectedElementsNameList() const;
 	const TArray<TSharedPtr<FString>>* GetDrawingNameList(URigVMPin* InPin = nullptr) const;
 
 	bool bSuspendModelNotifications;
diff --git a/Engine/Plugins/Experimental/ControlRig/Source/ControlRigEditor/Private/ControlRigElementDetails.cpp b/Engine/Plugins/Experimental/ControlRig/Source/ControlRigEditor/Private/ControlRigElementDetails.cpp
--- a/Engine/Plugins/Experimental/ControlRig/Source/ControlRigEditor/Private/ControlRigElementDetails.cpp
+++ b/Engine/Plugins/Experimental/ControlRig/Source/ControlRigEditor/Private/ControlRigElementDetails.cpp
@@ -22,6 +22,18 @@
 
 #define LOCTEXT_NAMESPACE "ControlRigElementDetails"
 
+namespace FRigElementKeyDetailsDefs
+{
+	// Active foreground pin alpha
+	static const float ActivePinForegroundAlpha = 1.f;
+	// InActive foreground pin alpha
+	static const float InactivePinForegroundAlpha = 0.15f;
+	// Active background pin alpha
+	static const float ActivePinBackgroundAlpha = 0.8f;
+	// InActive background pin alpha
+	static const float InactivePinBackgroundAlpha = 0.4f;
+};
+
 void RigElementDetails_GetCustomizedInfo(TSharedRef<IPropertyHandle> InStructPropertyHandle, UControlRigBlueprint*& OutBlueprint)
 {
 	TArray<UObject*> Objects;
@@ -172,6 +184,42 @@ void FRigElementKeyDetails::CustomizeHeader(TSharedRef<IPropertyHandle> InStruct
 					.Font(IDetailLayoutBuilder::GetDetailFont())
 				]
 			]
+			// Use button
+			+SHorizontalBox::Slot()
+			.AutoWidth()
+			.Padding(1,0)
+			.VAlign(VAlign_Center)
+			[
+				SAssignNew(UseSelectedButton, SButton)
+				.ButtonStyle( FEditorStyle::Get(), "NoBorder" )
+				.ButtonColorAndOpacity_Lambda([this]() { return OnGetWidgetBackground(UseSelectedButton); })
+				.OnClicked(this, &FRigElementKeyDetails::OnGetSelectedClicked)
+				.ContentPadding(1.f)
+				.ToolTipText(NSLOCTEXT("GraphEditor", "ObjectGraphPin_Use_Tooltip", "Use item selected"))
+				[
+					SNew(SImage)
+					.ColorAndOpacity_Lambda( [this]() { return OnGetWidgetForeground(UseSelectedButton); })
+					.Image(FEditorStyle::GetBrush("Icons.CircleArrowLeft"))
+				]
+			]
+			// Select in hierarchy button
+			+SHorizontalBox::Slot()
+			.AutoWidth()
+			.Padding(1,0)
+			.VAlign(VAlign_Center)
+			[
+				SAssignNew(SelectElementButton, SButton)
+				.ButtonStyle( FEditorStyle::Get(), "NoBorder" )
+				.ButtonColorAndOpacity_Lambda([this]() { return OnGetWidgetBackground(SelectElementButton); })
+				.OnClicked(this, &FRigElementKeyDetails::OnSelectInHierarchyClicked)
+				.ContentPadding(0)
+				.ToolTipText(NSLOCTEXT("GraphEditor", "ObjectGraphPin_Browse_Tooltip", "Select in hierarchy"))
+				[
+					SNew(SImage)
+					.ColorAndOpacity_Lambda( [this]() { return OnGetWidgetForeground(SelectElementButton); })
+					.Image(FEditorStyle::GetBrush("Icons.Search"))
+				]
+			]			
 		];
 	}
 }
@@ -276,6 +324,61 @@ FText FRigElementKeyDetails::GetElementNameAsText() const
 	return FText::FromString(GetElementName());
 }
 
+FSlateColor FRigElementKeyDetails::OnGetWidgetForeground(const TSharedPtr<SButton> Button) const
+{
+	float Alpha = (Button.IsValid() && Button->IsHovered()) ? FRigElementKeyDetailsDefs::ActivePinForegroundAlpha : FRigElementKeyDetailsDefs::InactivePinForegroundAlpha;
+	return FSlateColor(FLinearColor(1.f, 1.f, 1.f, Alpha));
+}
+
+FSlateColor FRigElementKeyDetails::OnGetWidgetBackground(const TSharedPtr<SButton> Button) const
+{
+	float Alpha = (Button.IsValid() && Button->IsHovered()) ? FRigElementKeyDetailsDefs::ActivePinBackgroundAlpha : FRigElementKeyDetailsDefs::InactivePinBackgroundAlpha;
+	return FSlateColor(FLinearColor(1.f, 1.f, 1.f, Alpha));
+}
+
+FReply FRigElementKeyDetails::OnGetSelectedClicked()
+{
+	if (BlueprintBeingCustomized)
+	{
+		const TArray<FRigElementKey>& Selected = BlueprintBeingCustomized->Hierarchy->GetSelectedKeys();
+		if (Selected.Num() > 0)
+		{
+			if (TypeHandle.IsValid())
+			{
+				uint8 Index = (uint8) Selected[0].Type;
+				TypeHandle->SetValue(Index);
+			}
+			SetElementName(Selected[0].Name.ToString());
+		}
+	}
+	return FReply::Handled();
+}
+
+FReply FRigElementKeyDetails::OnSelectInHierarchyClicked()
+{
+	if (BlueprintBeingCustomized)
+	{
+		FRigElementKey Key;
+		if (TypeHandle.IsValid())
+		{
+			uint8 Type;
+			TypeHandle->GetValue(Type);
+			Key.Type = (ERigElementType) Type;
+		}
+
+		if (NameHandle.IsValid())
+		{
+			NameHandle->GetValue(Key.Name);
+		}
+				
+		if (Key.IsValid())
+		{
+			BlueprintBeingCustomized->GetHierarchyController()->SetSelection({Key});
+		}	
+	}
+	return FReply::Handled();
+}
+
 void FRigUnitDetails::CustomizeDetails(IDetailLayoutBuilder& DetailBuilder)
 {
 	TArray<TSharedPtr<FStructOnScope>> StructsBeingCustomized;
diff --git a/Engine/Plugins/Experimental/ControlRig/Source/ControlRigEditor/Private/ControlRigElementDetails.h b/Engine/Plugins/Experimental/ControlRig/Source/ControlRigEditor/Private/ControlRigElementDetails.h
--- a/Engine/Plugins/Experimental/ControlRig/Source/ControlRigEditor/Private/ControlRigElementDetails.h
+++ b/Engine/Plugins/Experimental/ControlRig/Source/ControlRigEditor/Private/ControlRigElementDetails.h
@@ -40,6 +40,14 @@ protected:
 	TSharedRef<SWidget> OnGetElementNameWidget(TSharedPtr<FString> InItem);
 	FText GetElementNameAsText() const;
 
+	/** Helper buttons. */
+	TSharedPtr<SButton> UseSelectedButton;
+	TSharedPtr<SButton> SelectElementButton;
+	FSlateColor OnGetWidgetForeground(const TSharedPtr<SButton> Button) const;
+	FSlateColor OnGetWidgetBackground(const TSharedPtr<SButton> Button) const;
+	FReply OnGetSelectedClicked();
+	FReply OnSelectInHierarchyClicked();
+	
 	TSharedPtr<IPropertyHandle> TypeHandle;
 	TSharedPtr<IPropertyHandle> NameHandle;
 	TArray<TSharedPtr<FString>> ElementNameList;
diff --git a/Engine/Plugins/Experimental/ControlRig/Source/ControlRigEditor/Private/Graph/ControlRigGraphPanelPinFactory.cpp b/Engine/Plugins/Experimental/ControlRig/Source/ControlRigEditor/Private/Graph/ControlRigGraphPanelPinFactory.cpp
--- a/Engine/Plugins/Experimental/ControlRig/Source/ControlRigEditor/Private/Graph/ControlRigGraphPanelPinFactory.cpp
+++ b/Engine/Plugins/Experimental/ControlRig/Source/ControlRigEditor/Private/Graph/ControlRigGraphPanelPinFactory.cpp
@@ -43,36 +43,42 @@ TSharedPtr<SGraphPin> FControlRigGraphPanelPinFactory::CreatePin(UEdGraphPin* In
 				{
 					return SNew(SControlRigGraphPinNameList, InPin)
 						.ModelPin(ModelPin)
+						.OnGetNameFromSelection_UObject(RigGraph, &UControlRigGraph::GetSelectedElementsNameList)
 						.OnGetNameListContent_UObject(RigGraph, &UControlRigGraph::GetBoneNameList);
 				}
 				else if (CustomWidgetName == TEXT("ControlName"))
 				{
 					return SNew(SControlRigGraphPinNameList, InPin)
 						.ModelPin(ModelPin)
+						.OnGetNameFromSelection_UObject(RigGraph, &UControlRigGraph::GetSelectedElementsNameList)
 						.OnGetNameListContent_UObject(RigGraph, &UControlRigGraph::GetControlNameList);
 				}
 				else if (CustomWidgetName == TEXT("SpaceName"))
 				{
 					return SNew(SControlRigGraphPinNameList, InPin)
 						.ModelPin(ModelPin)
+						.OnGetNameFromSelection_UObject(RigGraph, &UControlRigGraph::GetSelectedElementsNameList)
 						.OnGetNameListContent_UObject(RigGraph, &UControlRigGraph::GetNullNameList);
 				}
 				else if (CustomWidgetName == TEXT("CurveName"))
 				{
 					return SNew(SControlRigGraphPinNameList, InPin)
 						.ModelPin(ModelPin)
+						.OnGetNameFromSelection_UObject(RigGraph, &UControlRigGraph::GetSelectedElementsNameList)
 						.OnGetNameListContent_UObject(RigGraph, &UControlRigGraph::GetCurveNameList);
 				}
 				else if (CustomWidgetName == TEXT("ElementName"))
 				{
 					return SNew(SControlRigGraphPinNameList, InPin)
 						.ModelPin(ModelPin)
+						.OnGetNameFromSelection_UObject(RigGraph, &UControlRigGraph::GetSelectedElementsNameList)
 						.OnGetNameListContent_UObject(RigGraph, &UControlRigGraph::GetElementNameList);
 				}
 				else if (CustomWidgetName == TEXT("DrawingName"))
 				{
 					return SNew(SControlRigGraphPinNameList, InPin)
 						.ModelPin(ModelPin)
+						.OnGetNameFromSelection_UObject(RigGraph, &UControlRigGraph::GetSelectedElementsNameList)
 						.OnGetNameListContent_UObject(RigGraph, &UControlRigGraph::GetDrawingNameList);
 				}
 				else if (CustomWidgetName == TEXT("VariableName"))
diff --git a/Engine/Plugins/Experimental/ControlRig/Source/ControlRigEditor/Private/Graph/SControlRigGraphPinNameList.cpp b/Engine/Plugins/Experimental/ControlRig/Source/ControlRigEditor/Private/Graph/SControlRigGraphPinNameList.cpp
--- a/Engine/Plugins/Experimental/ControlRig/Source/ControlRigEditor/Private/Graph/SControlRigGraphPinNameList.cpp
+++ b/Engine/Plugins/Experimental/ControlRig/Source/ControlRigEditor/Private/Graph/SControlRigGraphPinNameList.cpp
@@ -8,11 +8,25 @@
 #include "EdGraphSchema_K2.h"
 #include "ScopedTransaction.h"
 #include "DetailLayoutBuilder.h"
+#include "Graph/ControlRigGraph.h"
+
+namespace SControlRigGraphPinNameListDefs
+{
+	// Active foreground pin alpha
+	static const float ActivePinForegroundAlpha = 1.f;
+	// InActive foreground pin alpha
+	static const float InactivePinForegroundAlpha = 0.15f;
+	// Active background pin alpha
+	static const float ActivePinBackgroundAlpha = 0.8f;
+	// InActive background pin alpha
+	static const float InactivePinBackgroundAlpha = 0.4f;
+};
 
 void SControlRigGraphPinNameList::Construct(const FArguments& InArgs, UEdGraphPin* InGraphPinObj)
 {
 	this->ModelPin = InArgs._ModelPin;
 	this->OnGetNameListContent = InArgs._OnGetNameListContent;
+	this->OnGetNameFromSelection = InArgs._OnGetNameFromSelection;
 	this->bMarkupInvalidItems = InArgs._MarkupInvalidItems;
 
 	CurrentList = GetNameList();
@@ -34,20 +48,63 @@ TSharedRef<SWidget>	SControlRigGraphPinNameList::GetDefaultValueWidget()
 	return SNew(SBox)
 		.MinDesiredWidth(150)
 		[
-			SAssignNew(NameListComboBox, SControlRigGraphPinNameListValueWidget)
-				.Visibility(this, &SGraphPin::GetDefaultValueVisibility)
-				.OptionsSource(CurrentList)
-				.OnGenerateWidget(this, &SControlRigGraphPinNameList::MakeNameListItemWidget)
-				.OnSelectionChanged(this, &SControlRigGraphPinNameList::OnNameListChanged)
-				.OnComboBoxOpening(this, &SControlRigGraphPinNameList::OnNameListComboBox)
-				.InitiallySelectedItem(InitialSelected)
-				.Content()
+			SNew(SHorizontalBox)
+			+SHorizontalBox::Slot()
+			[
+				SAssignNew(NameListComboBox, SControlRigGraphPinNameListValueWidget)
+					.Visibility(this, &SGraphPin::GetDefaultValueVisibility)
+					.OptionsSource(CurrentList)
+					.OnGenerateWidget(this, &SControlRigGraphPinNameList::MakeNameListItemWidget)
+					.OnSelectionChanged(this, &SControlRigGraphPinNameList::OnNameListChanged)
+					.OnComboBoxOpening(this, &SControlRigGraphPinNameList::OnNameListComboBox)
+					.InitiallySelectedItem(InitialSelected)
+					.Content()
+					[
+						SNew(STextBlock)
+						.Text(this, &SControlRigGraphPinNameList::GetNameListText)
+						.ColorAndOpacity(this, &SControlRigGraphPinNameList::GetNameColor)
+						.Font( FEditorStyle::GetFontStyle( TEXT("PropertyWindow.NormalFont") ) )
+					]
+			]
+		
+			// Use button
+			+SHorizontalBox::Slot()
+			.AutoWidth()
+			.Padding(1,0)
+			.VAlign(VAlign_Center)
+			[
+				SNew(SButton)
+				.ButtonStyle( FEditorStyle::Get(), "NoBorder" )
+				.ButtonColorAndOpacity( this, &SControlRigGraphPinNameList::OnGetWidgetBackground )
+				.OnClicked(this, &SControlRigGraphPinNameList::OnGetSelectedClicked)
+				.ContentPadding(1.f)
+				.ToolTipText(NSLOCTEXT("GraphEditor", "ObjectGraphPin_Use_Tooltip", "Use item selected"))
+				[
+					SNew(SImage)
+					.ColorAndOpacity( this, &SControlRigGraphPinNameList::OnGetWidgetForeground )
+					.Image(FEditorStyle::GetBrush("Icons.CircleArrowLeft"))
+				]
+			]
+
+			// Browse button
+			+SHorizontalBox::Slot()
+			.AutoWidth()
+			.Padding(1,0)
+			.VAlign(VAlign_Center)
+			[
+				SNew(SButton)
+				.ButtonStyle( FEditorStyle::Get(), "NoBorder" )
+				.ButtonColorAndOpacity( this, &SControlRigGraphPinNameList::OnGetWidgetBackground )
+				.OnClicked(this, &SControlRigGraphPinNameList::OnBrowseClicked)
+				.ContentPadding(0)
+				.ToolTipText(NSLOCTEXT("GraphEditor", "ObjectGraphPin_Browse_Tooltip", "Browse"))
 				[
-					SNew(STextBlock)
-					.Text(this, &SControlRigGraphPinNameList::GetNameListText)
-					.ColorAndOpacity(this, &SControlRigGraphPinNameList::GetNameColor)
-					.Font( FEditorStyle::GetFontStyle( TEXT("PropertyWindow.NormalFont") ) )
+					SNew(SImage)
+					.ColorAndOpacity( this, &SControlRigGraphPinNameList::OnGetWidgetForeground )
+					.Image(FEditorStyle::GetBrush("Icons.Search"))
 				]
+			]
+				
 		];
 }
 
@@ -76,7 +133,7 @@ void SControlRigGraphPinNameList::SetNameListText(const FText& NewTypeInValue, E
 {
 	if(!GraphPinObj->GetDefaultAsString().Equals(NewTypeInValue.ToString()))
 	{
-		const FScopedTransaction Transaction( NSLOCTEXT("GraphEditor", "ChangeBoneNameListPinValue", "Change Bone Name Pin Value" ) );
+		const FScopedTransaction Transaction( NSLOCTEXT("GraphEditor", "ChangeElementNameListPinValue", "Change Element Name Pin Value" ) );
 		GraphPinObj->Modify();
 		GraphPinObj->GetSchema()->TrySetDefaultValue(*GraphPinObj, NewTypeInValue.ToString());
 	}
@@ -133,8 +190,71 @@ void SControlRigGraphPinNameList::OnNameListComboBox()
 		if (Item->Equals(GetNameListText().ToString()))
 		{
 			CurrentlySelected = Item;
+			break;
 		}
 	}
 	NameListComboBox->SetOptionsSource(CurrentList);
 	NameListComboBox->SetSelectedItem(CurrentlySelected);
 }
+
+FSlateColor SControlRigGraphPinNameList::OnGetWidgetForeground() const
+{
+	float Alpha = IsHovered() ? SControlRigGraphPinNameListDefs::ActivePinForegroundAlpha : SControlRigGraphPinNameListDefs::InactivePinForegroundAlpha;
+	return FSlateColor(FLinearColor(1.f, 1.f, 1.f, Alpha));
+}
+
+FSlateColor SControlRigGraphPinNameList::OnGetWidgetBackground() const
+{
+	float Alpha = IsHovered() ? SControlRigGraphPinNameListDefs::ActivePinBackgroundAlpha : SControlRigGraphPinNameListDefs::InactivePinBackgroundAlpha;
+	return FSlateColor(FLinearColor(1.f, 1.f, 1.f, Alpha));
+}
+
+FReply SControlRigGraphPinNameList::OnGetSelectedClicked()
+{
+	if (ModelPin->GetCustomWidgetName() == TEXT("ElementName"))
+	{
+		if (OnGetNameFromSelection.IsBound())
+		{
+			const TArray<TSharedPtr<FString>> Result = OnGetNameFromSelection.Execute();
+			if (Result.Num() > 0)
+			{
+				if (Result[0].IsValid() && Result[0] != nullptr)
+				{
+					if (URigVMPin* ParentPin = ModelPin->GetParentPin())
+					{
+						if (UControlRigGraph* Graph = Cast<UControlRigGraph>(GetPinObj()->GetOwningNode()->GetGraph()))
+						{
+							Graph->GetController()->SetPinDefaultValue(ParentPin->GetPinPath(), *Result[0].Get());
+							CurrentList = GetNameList();
+						}
+					}
+				}
+			}
+		}
+	}
+
+	return FReply::Handled();
+}
+
+FReply SControlRigGraphPinNameList::OnBrowseClicked()
+{
+	if (UControlRigGraph* Graph = Cast<UControlRigGraph>(GetPinObj()->GetOwningNode()->GetGraph()))
+	{
+		TSharedPtr<FString> Selected = NameListComboBox->GetSelectedItem();
+		if (Selected.IsValid())
+		{
+			FString DefaultValue = ModelPin->GetParentPin()->GetDefaultValue();
+			if (!DefaultValue.IsEmpty())
+			{
+				FRigElementKey Key;
+				FRigElementKey::StaticStruct()->ImportText(*DefaultValue, &Key, nullptr, EPropertyPortFlags::PPF_None, nullptr, FRigElementKey::StaticStruct()->GetName(), true);
+				if (Key.IsValid())
+				{
+					Graph->GetBlueprint()->GetHierarchyController()->SetSelection({Key});
+				}
+			}
+		}
+	}
+
+	return FReply::Handled();
+}
diff --git a/Engine/Plugins/Experimental/ControlRig/Source/ControlRigEditor/Private/Graph/SControlRigGraphPinNameList.h b/Engine/Plugins/Experimental/ControlRig/Source/ControlRigEditor/Private/Graph/SControlRigGraphPinNameList.h
--- a/Engine/Plugins/Experimental/ControlRig/Source/ControlRigEditor/Private/Graph/SControlRigGraphPinNameList.h
+++ b/Engine/Plugins/Experimental/ControlRig/Source/ControlRigEditor/Private/Graph/SControlRigGraphPinNameList.h
@@ -14,6 +14,7 @@ class SControlRigGraphPinNameList : public SGraphPin
 public:
 
 	DECLARE_DELEGATE_RetVal_OneParam( const TArray<TSharedPtr<FString>>*, FOnGetNameListContent, URigVMPin*);
+	DECLARE_DELEGATE_RetVal( const TArray<TSharedPtr<FString>>, FOnGetNameFromSelection);
 
 	SLATE_BEGIN_ARGS(SControlRigGraphPinNameList)
 		: _MarkupInvalidItems(true)
@@ -21,6 +22,7 @@ public:
 
 		SLATE_ARGUMENT(URigVMPin*, ModelPin)
 		SLATE_EVENT(FOnGetNameListContent, OnGetNameListContent)
+		SLATE_EVENT(FOnGetNameFromSelection, OnGetNameFromSelection)
 		SLATE_ARGUMENT(bool, MarkupInvalidItems)
 
 	SLATE_END_ARGS()
@@ -47,4 +49,11 @@ protected:
 	TArray<TSharedPtr<FString>> EmptyList;
 	const TArray<TSharedPtr<FString>>* CurrentList;
 	bool bMarkupInvalidItems;
+
+	/** Helper buttons. */
+	FSlateColor OnGetWidgetForeground() const;
+	FSlateColor OnGetWidgetBackground() const;
+	FReply OnGetSelectedClicked();
+	FReply OnBrowseClicked();
+	FOnGetNameFromSelection OnGetNameFromSelection;
 };
