Trace initlialize/destroy control rig, so that the object will exist in GameplayProvider, which makes it show in the Rewind Debugger component hierarchy

diff --git a/Engine/Plugins/Experimental/ControlRig/Source/ControlRig/Private/ControlRig.cpp b/Engine/Plugins/Experimental/ControlRig/Source/ControlRig/Private/ControlRig.cpp
--- a/Engine/Plugins/Experimental/ControlRig/Source/ControlRig/Private/ControlRig.cpp
+++ b/Engine/Plugins/Experimental/ControlRig/Source/ControlRig/Private/ControlRig.cpp
@@ -6,6 +6,7 @@
 #include "Misc/RuntimeErrors.h"
 #include "IControlRigObjectBinding.h"
 #include "HelperUtil.h"
+#include "ObjectTrace.h"
 #include "ControlRigBlueprintGeneratedClass.h"
 #include "ControlRigObjectVersion.h"
 #include "Rigs/RigHierarchyController.h"
@@ -101,6 +102,8 @@ void UControlRig::BeginDestroy()
 		VMSnapshotBeforeExecution = nullptr;
 	}
 #endif
+
+	TRACE_OBJECT_EVENT(this, DestroyControlRig);
 }
 
 UWorld* UControlRig::GetWorld() const
@@ -131,6 +134,8 @@ UWorld* UControlRig::GetWorld() const
 
 void UControlRig::Initialize(bool bInitRigUnits)
 {
+	TRACE_OBJECT_EVENT(this, InitializeControlRig);
+
 	DECLARE_SCOPE_HIERARCHICAL_COUNTER_FUNC()
 	QUICK_SCOPE_CYCLE_COUNTER(STAT_ControlRig_Initialize);
 
