ControlRig: Closing Mirror Menu still Applies Mirror Options

diff --git a/Engine/Plugins/Experimental/ControlRig/Source/ControlRigEditor/Private/ControlRigEditorModule.cpp b/Engine/Plugins/Experimental/ControlRig/Source/ControlRigEditor/Private/ControlRigEditorModule.cpp
--- a/Engine/Plugins/Experimental/ControlRig/Source/ControlRigEditor/Private/ControlRigEditorModule.cpp
+++ b/Engine/Plugins/Experimental/ControlRig/Source/ControlRigEditor/Private/ControlRigEditorModule.cpp
@@ -113,6 +113,7 @@
 #include "EditMode/SControlRigBaseListWidget.h"
 #include "EditMode/SControlRigTweenWidget.h"
 #include "EditMode/SControlRigSnapper.h"
+#include "Dialogs/CustomDialog.h"
 
 #define LOCTEXT_NAMESPACE "ControlRigEditorModule"
 
@@ -1773,64 +1774,72 @@ void FControlRigEditorModule::GetContextMenuActions(const UControlRigGraphSchema
 							TSharedRef<SKismetInspector> KismetInspector = SNew(SKismetInspector);
 							KismetInspector->ShowSingleStruct(StructToDisplay);
 
-							SGenericDialogWidget::OpenDialog(LOCTEXT("ControlRigHierarchyMirror", "Mirror Graph"), KismetInspector, SGenericDialogWidget::FArguments(), true);
-
-							Controller->OpenUndoBracket(TEXT("Mirroring Graph"));
-							int32 ReplacedNames = 0;
-							TArray<FString> UnchangedItems;
-
-							for (const TPair<const URigVMPin*, FRigElementKey>& Pair : PinToKey)
+							TSharedRef<SCustomDialog> MirrorDialog = SNew(SCustomDialog)
+								.Title(FText(LOCTEXT("ControlRigHierarchyMirror", "Mirror Graph")))
+								.DialogContent( KismetInspector)
+								.Buttons({
+									SCustomDialog::FButton(LOCTEXT("OK", "OK")),
+									SCustomDialog::FButton(LOCTEXT("Cancel", "Cancel"))
+							});
+							if (MirrorDialog->ShowModal() == 0)
 							{
-								const URigVMPin* Pin = Pair.Key;
-								FRigElementKey Key = Pair.Value;
+								Controller->OpenUndoBracket(TEXT("Mirroring Graph"));
+								int32 ReplacedNames = 0;
+								TArray<FString> UnchangedItems;
 
-								if (Key.Name.IsNone())
+								for (const TPair<const URigVMPin*, FRigElementKey>& Pair : PinToKey)
 								{
-									continue;
-								}
+									const URigVMPin* Pin = Pair.Key;
+									FRigElementKey Key = Pair.Value;
 
-								FString OldNameStr = Key.Name.ToString();
-								FString NewNameStr = OldNameStr.Replace(*Settings.OldName, *Settings.NewName, ESearchCase::CaseSensitive);
-								if(NewNameStr != OldNameStr)
-								{
-									Key.Name = *NewNameStr;
-									if(RigBlueprint->Hierarchy->GetIndex(Key) != INDEX_NONE)
+									if (Key.Name.IsNone())
 									{
-										Controller->SetPinDefaultValue(Pin->GetPinPath(), NewNameStr, false);
-										ReplacedNames++;
+										continue;
 									}
-									else
+
+									FString OldNameStr = Key.Name.ToString();
+									FString NewNameStr = OldNameStr.Replace(*Settings.OldName, *Settings.NewName, ESearchCase::CaseSensitive);
+									if(NewNameStr != OldNameStr)
 									{
-										// save the names of the items that we skipped during this search & replace
-										UnchangedItems.AddUnique(OldNameStr);
-									} 
+										Key.Name = *NewNameStr;
+										if(RigBlueprint->Hierarchy->GetIndex(Key) != INDEX_NONE)
+										{
+											Controller->SetPinDefaultValue(Pin->GetPinPath(), NewNameStr, false);
+											ReplacedNames++;
+										}
+										else
+										{
+											// save the names of the items that we skipped during this search & replace
+											UnchangedItems.AddUnique(OldNameStr);
+										} 
+									}
 								}
-							}
 
-							if (UnchangedItems.Num() > 0)
-							{
-								FString ListOfUnchangedItems;
-								for (int Index = 0; Index < UnchangedItems.Num(); Index++)
+								if (UnchangedItems.Num() > 0)
 								{
-									// construct the string "item1, item2, item3"
-									ListOfUnchangedItems += UnchangedItems[Index];
-									if (Index != UnchangedItems.Num() - 1)
+									FString ListOfUnchangedItems;
+									for (int Index = 0; Index < UnchangedItems.Num(); Index++)
 									{
-										ListOfUnchangedItems += TEXT(", ");
+										// construct the string "item1, item2, item3"
+										ListOfUnchangedItems += UnchangedItems[Index];
+										if (Index != UnchangedItems.Num() - 1)
+										{
+											ListOfUnchangedItems += TEXT(", ");
+										}
 									}
+									
+									// inform the user that some items were skipped due to invalid new names
+									Controller->ReportAndNotifyError(FString::Printf(TEXT("Invalid Names after Search & Replace, action skipped for %s"), *ListOfUnchangedItems));
 								}
-								
-								// inform the user that some items were skipped due to invalid new names
-								Controller->ReportAndNotifyError(FString::Printf(TEXT("Invalid Names after Search & Replace, action skipped for %s"), *ListOfUnchangedItems));
-							}
 
-							if (ReplacedNames > 0)
-							{ 
-								Controller->CloseUndoBracket();
-							}
-							else
-							{
-								Controller->CancelUndoBracket();
+								if (ReplacedNames > 0)
+								{ 
+									Controller->CloseUndoBracket();
+								}
+								else
+								{
+									Controller->CancelUndoBracket();
+								}
 							}
 						})
 					));
