Control Rig: Avoid log spam

diff --git a/Engine/Plugins/Experimental/ControlRig/Source/ControlRig/Private/ControlRig.cpp b/Engine/Plugins/Experimental/ControlRig/Source/ControlRig/Private/ControlRig.cpp
--- a/Engine/Plugins/Experimental/ControlRig/Source/ControlRig/Private/ControlRig.cpp
+++ b/Engine/Plugins/Experimental/ControlRig/Source/ControlRig/Private/ControlRig.cpp
@@ -715,23 +715,36 @@ void UControlRig::Execute(const EControlRigState InState, const FName& InEventNa
 				continue;
 			}
 
+			FString PerInstructionMessage = 
+				FString::Printf(
+					TEXT("Instruction[%d] '%s': '%s'"),
+					Entry.InstructionIndex,
+					*Entry.FunctionName.ToString(),
+					*Entry.Message
+				);
+
+			if(LoggedMessages.Contains(PerInstructionMessage))
+			{
+				continue;
+			}
+
 			switch (Entry.Severity)
 			{
 				case EMessageSeverity::CriticalError:
 				case EMessageSeverity::Error:
 				{
-					UE_LOG(LogControlRig, Error, TEXT("Instruction[%d] '%s': '%s'"), Entry.InstructionIndex, *Entry.FunctionName.ToString(), *Entry.Message);
+					UE_LOG(LogControlRig, Error, TEXT("%s"), *PerInstructionMessage);
 					break;
 				}
 				case EMessageSeverity::PerformanceWarning:
 				case EMessageSeverity::Warning:
 				{
-					UE_LOG(LogControlRig, Warning, TEXT("Instruction[%d] '%s': '%s'"), Entry.InstructionIndex, *Entry.FunctionName.ToString(), *Entry.Message);
+					UE_LOG(LogControlRig, Warning, TEXT("%s"), *PerInstructionMessage);
 					break;
 				}
 				case EMessageSeverity::Info:
 				{
-					UE_LOG(LogControlRig, Display, TEXT("Instruction[%d] '%s': '%s'"), Entry.InstructionIndex, *Entry.FunctionName.ToString(), *Entry.Message);
+					UE_LOG(LogControlRig, Display, TEXT("%s"), *PerInstructionMessage);
 					break;
 				}
 				default:
@@ -739,12 +752,15 @@ void UControlRig::Execute(const EControlRigState InState, const FName& InEventNa
 					break;
 				}
 			}
+
+			LoggedMessages.Add(PerInstructionMessage, true);
 		}
 	}
 
 	if (bJustRanInit && ControlRigLog != nullptr)
 	{
 		ControlRigLog->KnownMessages.Reset();
+		LoggedMessages.Reset();
 	}
 #endif
 
diff --git a/Engine/Plugins/Experimental/ControlRig/Source/ControlRig/Public/ControlRig.h b/Engine/Plugins/Experimental/ControlRig/Source/ControlRig/Public/ControlRig.h
--- a/Engine/Plugins/Experimental/ControlRig/Source/ControlRig/Public/ControlRig.h
+++ b/Engine/Plugins/Experimental/ControlRig/Source/ControlRig/Public/ControlRig.h
@@ -658,6 +658,7 @@ private:
 #if WITH_EDITOR	
 
 	FRigVMDebugInfo DebugInfo;
+	TMap<FString, bool> LoggedMessages;
 	
 public:
 
