Control Rig: Show graph definition in details panel for collapse nodes and function reference nodes

diff --git a/Engine/Plugins/Experimental/ControlRig/Source/ControlRigEditor/Private/Editor/ControlRigEditor.cpp b/Engine/Plugins/Experimental/ControlRig/Source/ControlRigEditor/Private/Editor/ControlRigEditor.cpp
--- a/Engine/Plugins/Experimental/ControlRig/Source/ControlRigEditor/Private/Editor/ControlRigEditor.cpp
+++ b/Engine/Plugins/Experimental/ControlRig/Source/ControlRigEditor/Private/Editor/ControlRigEditor.cpp
@@ -1386,9 +1386,9 @@ void FControlRigEditor::SetDetailObject(UObject* Obj)
 		}
 		return;
 	}
-	else if (URigVMCollapseNode* CollapseNode = Cast<URigVMCollapseNode>(Obj))
+	else if (URigVMLibraryNode* LibraryNode = Cast<URigVMLibraryNode>(Obj))
 	{
-		UEdGraph* EdGraph = GetControlRigBlueprint()->GetEdGraph(CollapseNode->GetContainedGraph());
+		UEdGraph* EdGraph = GetControlRigBlueprint()->GetEdGraph(LibraryNode->GetContainedGraph());
 		if (EdGraph)
 		{
 			TArray<UObject*> Objects;
