ControlRig: Enable preview different skeletal meshes with their respective bone initial transforms in Control Rig Editor. Once the preview mesh is changed, hit "compile" to set bone initial transforms from the new skeletal mesh

diff --git a/Engine/Plugins/Experimental/ControlRig/Source/ControlRigEditor/Private/Editor/ControlRigEditor.cpp b/Engine/Plugins/Experimental/ControlRig/Source/ControlRigEditor/Private/Editor/ControlRigEditor.cpp
--- a/Engine/Plugins/Experimental/ControlRig/Source/ControlRigEditor/Private/Editor/ControlRigEditor.cpp
+++ b/Engine/Plugins/Experimental/ControlRig/Source/ControlRigEditor/Private/Editor/ControlRigEditor.cpp
@@ -1225,6 +1225,14 @@ void FControlRigEditor::HandleSetObjectBeingDebugged(UObject* InObject)
 					EditMode->SetObjects(DebuggedControlRig, EditorSkelComp,nullptr);
 				}
 			}
+			
+			DebuggedControlRig->SetBoneInitialTransformsFromSkeletalMeshComponent(EditorSkelComp);
+			if(UControlRigBlueprint* RigBlueprint = Cast<UControlRigBlueprint>(GetBlueprintObj()))
+			{
+				// copy the initial transforms back to the blueprint
+				RigBlueprint->Hierarchy->Modify();
+				RigBlueprint->Hierarchy->CopyPose(DebuggedControlRig->GetHierarchy(), false, true);
+			}
 		}
 
 		DebuggedControlRig->GetHierarchy()->OnModified().AddSP(this, &FControlRigEditor::OnHierarchyModified_AnyThread);
