Control Rig Pose Cache functionality

diff --git a/Engine/Plugins/Experimental/ControlRig/Source/ControlRig/Private/Drawing/ControlRigDrawInterface.cpp b/Engine/Plugins/Experimental/ControlRig/Source/ControlRig/Private/Drawing/ControlRigDrawInterface.cpp
--- a/Engine/Plugins/Experimental/ControlRig/Source/ControlRig/Private/Drawing/ControlRigDrawInterface.cpp
+++ b/Engine/Plugins/Experimental/ControlRig/Source/ControlRig/Private/Drawing/ControlRigDrawInterface.cpp
@@ -156,6 +156,57 @@ void FControlRigDrawInterface::DrawAxes(const FTransform& WorldOffset, const FRi
 	Instructions.Add(InstructionZ);
 }
 
+void FControlRigDrawInterface::DrawAxes(const FTransform& WorldOffset, const FTransform& Transform,
+	const FLinearColor& InColor, float Size, float Thickness)
+{
+	if (!IsEnabled())
+	{
+		return;
+	}
+
+	FControlRigDrawInstruction Instruction(EControlRigDrawSettings::Lines, InColor, Thickness, WorldOffset);
+
+	Instruction.Positions.Reserve(6);
+	Instruction.Positions.Add(Transform.GetLocation());
+	Instruction.Positions.Add(Transform.TransformPosition(FVector(Size, 0.f, 0.f)));
+	Instruction.Positions.Add(Transform.GetLocation());
+	Instruction.Positions.Add(Transform.TransformPosition(FVector(0.f, Size, 0.f)));
+	Instruction.Positions.Add(Transform.GetLocation());
+	Instruction.Positions.Add(Transform.TransformPosition(FVector(0.f, 0.f, Size)));
+
+	Instructions.Add(Instruction);
+}
+
+void FControlRigDrawInterface::DrawAxes(const FTransform& WorldOffset, const FRigVMFixedArray<FTransform>& Transforms,
+	const FLinearColor& InColor, float Size, float Thickness)
+{
+	if (!IsEnabled())
+	{
+		return;
+	}
+
+	if (Transforms.Num() == 0)
+	{
+		return;
+	}
+
+	FControlRigDrawInstruction Instruction(EControlRigDrawSettings::Lines, InColor, Thickness, WorldOffset);
+
+	Instruction.Positions.Reserve(Transforms.Num() * 6);
+
+	for (const FTransform& Transform : Transforms)
+	{
+		Instruction.Positions.Add(Transform.GetLocation());
+		Instruction.Positions.Add(Transform.TransformPosition(FVector(Size, 0.f, 0.f)));
+		Instruction.Positions.Add(Transform.GetLocation());
+		Instruction.Positions.Add(Transform.TransformPosition(FVector(0.f, Size, 0.f)));
+		Instruction.Positions.Add(Transform.GetLocation());
+		Instruction.Positions.Add(Transform.TransformPosition(FVector(0.f, 0.f, Size)));
+	}
+
+	Instructions.Add(Instruction);
+}
+
 void FControlRigDrawInterface::DrawRectangle(const FTransform& WorldOffset, const FTransform& Transform, float Size, const FLinearColor& Color, float Thickness)
 {
 	if (!IsEnabled())
@@ -235,7 +286,7 @@ void FControlRigDrawInterface::DrawBezier(const FTransform& WorldOffset, const F
 	Instructions.Add(Instruction);
 }
 
-void FControlRigDrawInterface::DrawHierarchy(const FTransform& WorldOffset, URigHierarchy* Hierarchy, EControlRigDrawHierarchyMode::Type Mode, float Scale, const FLinearColor& Color, float Thickness)
+void FControlRigDrawInterface::DrawHierarchy(const FTransform& WorldOffset, URigHierarchy* Hierarchy, EControlRigDrawHierarchyMode::Type Mode, float Scale, const FLinearColor& Color, float Thickness, const FRigPose* InPose)
 {
 	if (!IsEnabled())
 	{
@@ -257,7 +308,33 @@ void FControlRigDrawInterface::DrawHierarchy(const FTransform& WorldOffset, URig
 
 			Hierarchy->ForEach<FRigTransformElement>([&](FRigTransformElement* Child) -> bool
 			{
-				const FTransform Transform = Hierarchy->GetTransform(Child, ERigTransformType::CurrentGlobal);
+				auto GetTransformLambda = [InPose, Hierarchy](FRigTransformElement* InElement, FTransform& Transform) -> bool
+				{
+					bool bValid = true;
+					if(InPose)
+					{
+						const int32 ElementIndex = InPose->GetIndex(InElement->GetKey());
+						if(ElementIndex != INDEX_NONE)
+						{
+							Transform = InPose->operator[](ElementIndex).GlobalTransform;
+							return true;
+						}
+						else
+						{
+							bValid = false;
+						}
+					}
+
+					Transform = Hierarchy->GetTransform(InElement, ERigTransformType::CurrentGlobal);
+					return bValid;
+				};
+				
+				FTransform Transform = FTransform::Identity;
+				if(!GetTransformLambda(Child, Transform))
+				{
+					return true;
+				}
+
 				const FVector P0 = Transform.GetLocation();
 				const FVector PX = Transform.TransformPosition(FVector(Scale, 0.f, 0.f));
 				const FVector PY = Transform.TransformPosition(FVector(0.f, Scale, 0.f));
@@ -274,7 +351,9 @@ void FControlRigDrawInterface::DrawHierarchy(const FTransform& WorldOffset, URig
 				{
 					if(FRigTransformElement* ParentTransformElement = Cast<FRigTransformElement>(Parent))
 					{
-						const FTransform ParentTransform = Hierarchy->GetTransform(ParentTransformElement, ERigTransformType::CurrentGlobal);
+						FTransform ParentTransform = FTransform::Identity;
+						GetTransformLambda(ParentTransformElement, ParentTransform);
+
 						const FVector P1 = ParentTransform.GetLocation();
 						InstructionParent.Positions.Add(P0);
 						InstructionParent.Positions.Add(P1);
diff --git a/Engine/Plugins/Experimental/ControlRig/Source/ControlRig/Private/Rigs/RigHierarchy.cpp b/Engine/Plugins/Experimental/ControlRig/Source/ControlRig/Private/Rigs/RigHierarchy.cpp
--- a/Engine/Plugins/Experimental/ControlRig/Source/ControlRig/Private/Rigs/RigHierarchy.cpp
+++ b/Engine/Plugins/Experimental/ControlRig/Source/ControlRig/Private/Rigs/RigHierarchy.cpp
@@ -999,13 +999,35 @@ URigHierarchyController* URigHierarchy::GetController(bool bCreateIfNeeded)
 	return nullptr;
 }
 
-FRigPose URigHierarchy::GetPose(bool bInitial) const
+FRigPose URigHierarchy::GetPose(
+	bool bInitial,
+	ERigElementType InElementType,
+	const FRigElementKeyCollection& InItems 
+) const
 {
 	FRigPose Pose;
+	Pose.HierarchyTopologyVersion = GetTopologyVersion();
+	Pose.PoseHash = Pose.HierarchyTopologyVersion;
 
 	for (int32 ElementIndex = 0; ElementIndex < Elements.Num(); ElementIndex++)
 	{
 		FRigBaseElement* Element = Elements[ElementIndex];
+
+		// filter by type
+		if (((uint8)InElementType & (uint8)Element->GetType()) == 0)
+		{
+			continue;
+		}
+
+		// filter by optional collection
+		if(InItems.Num() > 0)
+		{
+			if(!InItems.Contains(Element->GetKey()))
+			{
+				continue;
+			}
+		}
+		
 		FRigPoseElement PoseElement;
 		PoseElement.Index.UpdateCache(Element->GetKey(), this);
 		
@@ -1023,28 +1045,61 @@ FRigPose URigHierarchy::GetPose(bool bInitial) const
 			continue;
 		}
 		Pose.Elements.Add(PoseElement);
+		Pose.PoseHash = HashCombine(Pose.PoseHash, GetTypeHash(PoseElement.Index.GetKey()));
 	}
 	return Pose;
 }
 
-void URigHierarchy::SetPose(const FRigPose& InPose, ERigTransformType::Type InTransformType)
+void URigHierarchy::SetPose(
+	const FRigPose& InPose,
+	ERigTransformType::Type InTransformType,
+	ERigElementType InElementType,
+	const FRigElementKeyCollection& InItems,
+	float InWeight
+)
 {
+	const float U = FMath::Clamp(InWeight, 0.f, 1.f);
+	if(U < SMALL_NUMBER)
+	{
+		return;
+	}
+	
 	for(const FRigPoseElement& PoseElement : InPose)
 	{
 		FCachedRigElement Index = PoseElement.Index;
+
+		// filter by type
+		if (((uint8)InElementType & (uint8)Index.GetKey().Type) == 0)
+		{
+			continue;
+		}
+
+		// filter by optional collection
+		if(InItems.Num() > 0)
+		{
+			if(!InItems.Contains(Index.GetKey()))
+			{
+				continue;
+			}
+		}
+
 		if(Index.UpdateCache(this))
 		{
 			FRigBaseElement* Element = Get(Index.GetIndex());
 			if(FRigTransformElement* TransformElement = Cast<FRigTransformElement>(Element))
 			{
-				if(ERigTransformType::IsLocal(InTransformType))
+				FTransform TransformToSet =
+					ERigTransformType::IsLocal(InTransformType) ?
+						PoseElement.LocalTransform :
+						PoseElement.GlobalTransform;
+				
+				if(U < 1.f - SMALL_NUMBER)
 				{
-					SetTransform(TransformElement, PoseElement.LocalTransform, InTransformType, true);
-				}
-				else
-				{
-					SetTransform(TransformElement, PoseElement.GlobalTransform, InTransformType, true);
+					const FTransform PreviousTransform = GetTransform(TransformElement, InTransformType);
+					TransformToSet = FControlRigMathLibrary::LerpTransform(PreviousTransform, TransformToSet, U);
 				}
+
+				SetTransform(TransformElement, TransformToSet, InTransformType, true);
 			}
 			else if(FRigCurveElement* CurveElement = Cast<FRigCurveElement>(Element))
 			{
diff --git a/Engine/Plugins/Experimental/ControlRig/Source/ControlRig/Private/Units/Debug/RigUnit_DebugHierarchy.cpp b/Engine/Plugins/Experimental/ControlRig/Source/ControlRig/Private/Units/Debug/RigUnit_DebugHierarchy.cpp
--- a/Engine/Plugins/Experimental/ControlRig/Source/ControlRig/Private/Units/Debug/RigUnit_DebugHierarchy.cpp
+++ b/Engine/Plugins/Experimental/ControlRig/Source/ControlRig/Private/Units/Debug/RigUnit_DebugHierarchy.cpp
@@ -21,4 +21,24 @@ FRigUnit_DebugHierarchy_Execute()
 	{
 		Context.DrawInterface->DrawHierarchy(WorldOffset, Hierarchy, EControlRigDrawHierarchyMode::Axes, Scale, Color, Thickness);
 	}
-}
\ No newline at end of file
+}
+
+FRigUnit_DebugPose_Execute()
+{
+	DECLARE_SCOPE_HIERARCHICAL_COUNTER_RIGUNIT()
+	if (Context.State == EControlRigState::Init)
+	{
+		return;
+	}
+
+	if (Context.DrawInterface == nullptr || !bEnabled)
+	{
+		return;
+	}
+
+	URigHierarchy* Hierarchy = Context.Hierarchy;
+	if (Hierarchy)
+	{
+		Context.DrawInterface->DrawHierarchy(WorldOffset, Hierarchy, EControlRigDrawHierarchyMode::Axes, Scale, Color, Thickness, &Pose);
+	}
+}
diff --git a/Engine/Plugins/Experimental/ControlRig/Source/ControlRig/Private/Units/Debug/RigUnit_DebugHierarchy.h b/Engine/Plugins/Experimental/ControlRig/Source/ControlRig/Private/Units/Debug/RigUnit_DebugHierarchy.h
--- a/Engine/Plugins/Experimental/ControlRig/Source/ControlRig/Private/Units/Debug/RigUnit_DebugHierarchy.h
+++ b/Engine/Plugins/Experimental/ControlRig/Source/ControlRig/Private/Units/Debug/RigUnit_DebugHierarchy.h
@@ -41,3 +41,42 @@ struct CONTROLRIG_API FRigUnit_DebugHierarchy : public FRigUnit_DebugBaseMutable
 	UPROPERTY(meta = (Input, Constant))
 	bool bEnabled;
 };
+
+/**
+* Draws vectors on each bone in the viewport across the entire pose
+*/
+USTRUCT(meta=(DisplayName="Draw Pose Cache"))
+struct CONTROLRIG_API FRigUnit_DebugPose : public FRigUnit_DebugBaseMutable
+{
+	GENERATED_BODY()
+
+	FRigUnit_DebugPose()
+	{
+		Scale = 10.f;
+		Color = FLinearColor::White;
+		Thickness = 0.f;
+		WorldOffset = FTransform::Identity;
+		bEnabled = true;
+	}
+
+	RIGVM_METHOD()
+	virtual void Execute(const FRigUnitContext& Context) override;
+
+	UPROPERTY(meta = (Input))
+	FRigPose Pose;
+
+	UPROPERTY(meta = (Input))
+	float Scale;
+
+	UPROPERTY(meta = (Input))
+	FLinearColor Color;
+
+	UPROPERTY(meta = (Input))
+	float Thickness;
+
+	UPROPERTY(meta = (Input))
+	FTransform WorldOffset;
+
+	UPROPERTY(meta = (Input, Constant))
+	bool bEnabled;
+};
\ No newline at end of file
diff --git a/Engine/Plugins/Experimental/ControlRig/Source/ControlRig/Private/Units/Execution/RigUnit_Hierarchy.cpp b/Engine/Plugins/Experimental/ControlRig/Source/ControlRig/Private/Units/Execution/RigUnit_Hierarchy.cpp
--- a/Engine/Plugins/Experimental/ControlRig/Source/ControlRig/Private/Units/Execution/RigUnit_Hierarchy.cpp
+++ b/Engine/Plugins/Experimental/ControlRig/Source/ControlRig/Private/Units/Execution/RigUnit_Hierarchy.cpp
@@ -147,4 +147,341 @@ FRigUnit_HierarchyGetSiblings_Execute()
 	}
 
 	Siblings = CachedSiblings;
+}
+
+FRigUnit_HierarchyGetPose_Execute()
+{
+	Pose = Context.Hierarchy->GetPose(Initial, ElementType, ItemsToGet);
+}
+
+FRigUnit_HierarchySetPose_Execute()
+{
+	ExecuteContext.Hierarchy->SetPose(
+		Pose,
+		Space == EBoneGetterSetterMode::GlobalSpace ? ERigTransformType::CurrentGlobal : ERigTransformType::CurrentLocal,
+		ElementType,
+		ItemsToSet,
+		Weight
+	);
+}
+
+FRigUnit_PoseIsEmpty_Execute()
+{
+	IsEmpty = Pose.Num() == 0;
+}
+
+FRigUnit_PoseGetItems_Execute()
+{
+	Items.Reset();
+
+	for(const FRigPoseElement& PoseElement : Pose)
+	{
+		// filter by type
+		if (((uint8)ElementType & (uint8)PoseElement.Index.GetKey().Type) == 0)
+		{
+			continue;
+		}
+		Items.Add(PoseElement.Index.GetKey());
+	}
+}
+
+FRigUnit_PoseGetDelta_Execute()
+{
+	PosesAreEqual = true;
+	ItemsWithDelta.Reset();
+
+	if(PoseA.Num() == 0 && PoseB.Num() == 0)
+	{
+		PosesAreEqual = true;
+		return;
+	}
+
+	if(PoseA.Num() == 0 && PoseB.Num() != 0)
+	{
+		PosesAreEqual = false;
+		FRigUnit_PoseGetItems::StaticExecute(RigVMExecuteContext, PoseB, ElementType, ItemsWithDelta, Context);
+		return;
+	}
+
+	if(PoseA.Num() != 0 && PoseB.Num() == 0)
+	{
+		PosesAreEqual = false;
+		FRigUnit_PoseGetItems::StaticExecute(RigVMExecuteContext, PoseA, ElementType, ItemsWithDelta, Context);
+		return;
+	}
+
+	const float PositionU = FMath::Abs(PositionThreshold);
+	const float RotationU = FMath::Abs(RotationThreshold);
+	const float ScaleU = FMath::Abs(ScaleThreshold);
+	const float CurveU = FMath::Abs(CurveThreshold);
+
+	auto ArePoseElementsEqual = [
+		PositionU,
+		RotationU,
+		ScaleU,
+		CurveU,
+		Space
+		] (
+			const FRigPoseElement& A,
+			const FRigPoseElement& B
+		) -> bool
+	{
+		const FRigElementKey& KeyA = A.Index.GetKey();
+		const FRigElementKey& KeyB = B.Index.GetKey();
+		check(KeyA == KeyB);
+		
+		if(KeyA.Type == ERigElementType::Curve)
+		{
+			if(CurveU > SMALL_NUMBER)
+			{
+				return FMath::Abs(A.CurveValue - B.CurveValue) < CurveU;
+			}
+			return true;
+		}
+		
+		const FTransform& TransformA = (Space == EBoneGetterSetterMode::GlobalSpace) ? A.GlobalTransform : A.LocalTransform;
+		const FTransform& TransformB = (Space == EBoneGetterSetterMode::GlobalSpace) ? B.GlobalTransform : B.LocalTransform;
+		
+		if(PositionU > SMALL_NUMBER)
+		{
+			const FVector PositionA = TransformA.GetLocation();
+			const FVector PositionB = TransformB.GetLocation();
+			const FVector Delta = (PositionA - PositionB).GetAbs();
+			if( (Delta.X >= PositionU) ||
+				(Delta.Y >= PositionU) ||
+				(Delta.Z >= PositionU))
+			{
+				return false;
+			}
+		}
+
+		if(RotationU > SMALL_NUMBER)
+		{
+			const FVector RotationA = TransformA.GetRotation().Rotator().Euler();
+			const FVector RotationB = TransformB.GetRotation().Rotator().Euler();
+			const FVector Delta = (RotationA - RotationB).GetAbs();
+			if( (Delta.X >= RotationU) ||
+				(Delta.Y >= RotationU) ||
+				(Delta.Z >= RotationU))
+			{
+				return false;
+			}
+		}
+
+		if(ScaleU > SMALL_NUMBER)
+		{
+			const FVector ScaleA = TransformA.GetScale3D();
+			const FVector ScaleB = TransformB.GetScale3D();
+			const FVector Delta = (ScaleA - ScaleB).GetAbs();
+			if( (Delta.X >= ScaleU) ||
+				(Delta.Y >= ScaleU) ||
+				(Delta.Z >= ScaleU))
+			{
+				return false;
+			}
+		}
+
+		return true;
+	};
+
+	// if we should compare a sub set of things
+	if(!ItemsToCompare.IsEmpty())
+	{
+		for(int32 Index = 0; Index < ItemsToCompare.Num(); Index++)
+		{
+			const FRigElementKey& Key = ItemsToCompare[Index];
+			if (((uint8)ElementType & (uint8)Key.Type) == 0)
+			{
+				continue;
+			}
+
+			const int32 IndexA = PoseA.GetIndex(Key);
+			if(IndexA == INDEX_NONE)
+			{
+				PosesAreEqual = false;
+				ItemsWithDelta.Add(Key);
+				continue;
+			}
+			
+			const int32 IndexB = PoseB.GetIndex(Key);
+			if(IndexB == INDEX_NONE)
+			{
+				PosesAreEqual = false;
+				ItemsWithDelta.Add(Key);
+				continue;
+			}
+
+			const FRigPoseElement& A = PoseA[IndexA];
+			const FRigPoseElement& B = PoseB[IndexB];
+
+			if(!ArePoseElementsEqual(A, B))
+			{
+				PosesAreEqual = false;
+				ItemsWithDelta.Add(Key);
+			}
+		}		
+	}
+	
+	// if the poses have the same hash we can accelerate this
+	// since they are structurally the same
+	else if(PoseA.PoseHash == PoseB.PoseHash)
+	{
+		for(int32 Index = 0; Index < PoseA.Num(); Index++)
+		{
+			const FRigPoseElement& A = PoseA[Index];
+			const FRigPoseElement& B = PoseB[Index];
+			const FRigElementKey& KeyA = A.Index.GetKey(); 
+			const FRigElementKey& KeyB = B.Index.GetKey(); 
+			check(KeyA == KeyB);
+			
+			if (((uint8)ElementType & (uint8)KeyA.Type) == 0)
+			{
+				continue;
+			}
+
+			if(!ArePoseElementsEqual(A, B))
+			{
+				PosesAreEqual = false;
+				ItemsWithDelta.Add(KeyA);
+			}
+		}
+	}
+	
+	// if the poses have different hashes they might not contain the same elements
+	else
+	{
+		for(int32 IndexA = 0; IndexA < PoseA.Num(); IndexA++)
+		{
+			const FRigPoseElement& A = PoseA[IndexA];
+			const FRigElementKey& KeyA = A.Index.GetKey(); 
+			
+			if (((uint8)ElementType & (uint8)KeyA.Type) == 0)
+			{
+				continue;
+			}
+
+			const int32 IndexB = PoseB.GetIndex(KeyA);
+			if(IndexB == INDEX_NONE)
+			{
+				PosesAreEqual = false;
+				ItemsWithDelta.Add(KeyA);
+				continue;
+			}
+
+			const FRigPoseElement& B = PoseB[IndexB];
+			
+			if(!ArePoseElementsEqual(A, B))
+			{
+				PosesAreEqual = false;
+				ItemsWithDelta.Add(KeyA);
+			}
+		}
+
+		// let's loop the other way as well and find relevant
+		// elements which are not part of this
+		for(int32 IndexB = 0; IndexB < PoseB.Num(); IndexB++)
+		{
+			const FRigPoseElement& B = PoseB[IndexB];
+			const FRigElementKey& KeyB = B.Index.GetKey(); 
+			
+			if (((uint8)ElementType & (uint8)KeyB.Type) == 0)
+			{
+				continue;
+			}
+
+			const int32 IndexA = PoseA.GetIndex(KeyB);
+			if(IndexA == INDEX_NONE)
+			{
+				PosesAreEqual = false;
+				ItemsWithDelta.AddUnique(KeyB);
+			}
+		}
+	}
+}
+
+FRigUnit_PoseGetTransform_Execute()
+{
+	if (Context.State == EControlRigState::Init)
+	{
+		CachedPoseElementIndex = CachedPoseHash = INDEX_NONE;
+	}
+
+	// set up defaults
+	Valid = false;
+	Transform = FTransform::Identity;
+
+	if(CachedPoseHash != Pose.HierarchyTopologyVersion)
+	{
+		CachedPoseHash = Pose.PoseHash;
+		CachedPoseElementIndex = Pose.GetIndex(Item);
+	}
+
+	if(CachedPoseElementIndex == INDEX_NONE)
+	{
+		return;
+	}
+
+	Valid = true;
+
+	const FRigPoseElement& PoseElement = Pose[CachedPoseElementIndex];
+
+	if(Space == EBoneGetterSetterMode::GlobalSpace)
+	{
+		Transform = PoseElement.GlobalTransform;
+	}
+	else
+	{
+		Transform = PoseElement.LocalTransform;
+	}
+}
+
+FRigUnit_PoseGetCurve_Execute()
+{
+	if (Context.State == EControlRigState::Init)
+	{
+		CachedPoseElementIndex = CachedPoseHash = INDEX_NONE;
+	}
+
+	// set up defaults
+	Valid = false;
+	CurveValue = 0.f;
+
+	if(CachedPoseHash != Pose.HierarchyTopologyVersion)
+	{
+		CachedPoseHash = Pose.PoseHash;
+		CachedPoseElementIndex = Pose.GetIndex(FRigElementKey(Curve, ERigElementType::Curve));
+	}
+
+	if(CachedPoseElementIndex == INDEX_NONE)
+	{
+		return;
+	}
+
+	Valid = true;
+
+	const FRigPoseElement& PoseElement = Pose[CachedPoseElementIndex];
+	CurveValue = PoseElement.CurveValue;
+}
+
+FRigUnit_PoseLoop_Execute()
+{
+	DECLARE_SCOPE_HIERARCHICAL_COUNTER_RIGUNIT()
+	Count = Pose.Num();
+	Continue = Pose.IsValidIndex(Index);
+	Ratio = GetRatioFromIndex(Index, Count);
+
+	if(Continue)
+	{
+		const FRigPoseElement& PoseElement = Pose.Elements[Index];
+		Item = PoseElement.Index.GetKey();
+		GlobalTransform = PoseElement.GlobalTransform;
+		LocalTransform = PoseElement.LocalTransform;
+		CurveValue = PoseElement.CurveValue;
+	}
+	else
+	{
+		Item = FRigElementKey();
+		GlobalTransform = LocalTransform = FTransform::Identity;
+		CurveValue = 0.f;
+	}
 }
\ No newline at end of file
diff --git a/Engine/Plugins/Experimental/ControlRig/Source/ControlRig/Private/Units/Execution/RigUnit_Hierarchy.h b/Engine/Plugins/Experimental/ControlRig/Source/ControlRig/Private/Units/Execution/RigUnit_Hierarchy.h
--- a/Engine/Plugins/Experimental/ControlRig/Source/ControlRig/Private/Units/Execution/RigUnit_Hierarchy.h
+++ b/Engine/Plugins/Experimental/ControlRig/Source/ControlRig/Private/Units/Execution/RigUnit_Hierarchy.h
@@ -12,6 +12,12 @@ struct CONTROLRIG_API FRigUnit_HierarchyBase : public FRigUnit
 	GENERATED_BODY()
 };
 
+USTRUCT(meta = (Abstract, NodeColor="0.462745, 1,0, 0.329412", Category = "Hierarchy"))
+struct CONTROLRIG_API FRigUnit_HierarchyBaseMutable : public FRigUnitMutable
+{
+	GENERATED_BODY()
+};
+
 /**
  * Returns the item's parent
  */
@@ -161,4 +167,338 @@ struct CONTROLRIG_API FRigUnit_HierarchyGetSiblings : public FRigUnit_HierarchyB
 	// Used to cache the internally
 	UPROPERTY()
 	FRigElementKeyCollection CachedSiblings;
-};
\ No newline at end of file
+};
+
+/**
+ * Returns the hierarchy's pose
+ */
+USTRUCT(meta=(DisplayName="Get Pose Cache", Keywords="Hierarchy,Pose,State", Varying))
+struct CONTROLRIG_API FRigUnit_HierarchyGetPose : public FRigUnit_HierarchyBase
+{
+	GENERATED_BODY()
+
+	FRigUnit_HierarchyGetPose()
+	{
+		Initial = false;
+		ElementType = ERigElementType::All;
+		ItemsToGet = FRigElementKeyCollection();
+		Pose = FRigPose();
+	}
+
+	RIGVM_METHOD()
+	virtual void Execute(const FRigUnitContext& Context) override;
+
+	UPROPERTY(meta = (Input))
+	bool Initial;
+
+	UPROPERTY(meta = (Input))
+	ERigElementType ElementType;
+
+	// An optional collection to filter against
+	UPROPERTY(meta = (Input))
+	FRigElementKeyCollection ItemsToGet;
+
+	UPROPERTY(meta = (Output))
+	FRigPose Pose;
+};
+
+/**
+ * Sets the hierarchy's pose
+ */
+USTRUCT(meta=(DisplayName="Apply Pose Cache", Keywords="Hierarchy,Pose,State", Varying))
+struct CONTROLRIG_API FRigUnit_HierarchySetPose : public FRigUnit_HierarchyBaseMutable
+{
+	GENERATED_BODY()
+
+	FRigUnit_HierarchySetPose()
+	{
+		Pose = FRigPose();
+		ElementType = ERigElementType::All;
+		Space = EBoneGetterSetterMode::LocalSpace;
+		ItemsToSet = FRigElementKeyCollection();
+		Weight = 1.f;
+	}
+
+	RIGVM_METHOD()
+	virtual void Execute(const FRigUnitContext& Context) override;
+
+	UPROPERTY(meta = (Input))
+	FRigPose Pose;
+
+	UPROPERTY(meta = (Input))
+	ERigElementType ElementType;
+
+	UPROPERTY(meta = (Input))
+	EBoneGetterSetterMode Space;
+
+	// An optional collection to filter against
+	UPROPERTY(meta = (Input))
+	FRigElementKeyCollection ItemsToSet;
+
+	UPROPERTY(meta = (Input))
+	float Weight;
+};
+
+/**
+* Returns true if the hierarchy pose is empty (has no items)
+*/
+USTRUCT(meta=(DisplayName="Is Pose Cache Empty", Keywords="Hierarchy,Pose,State", Varying))
+struct CONTROLRIG_API FRigUnit_PoseIsEmpty : public FRigUnit_HierarchyBase
+{
+	GENERATED_BODY()
+
+	FRigUnit_PoseIsEmpty()
+	{
+		Pose = FRigPose();
+		IsEmpty = true;
+	}
+
+	RIGVM_METHOD()
+	virtual void Execute(const FRigUnitContext& Context) override;
+
+	UPROPERTY(meta = (Input))
+	FRigPose Pose;
+
+	UPROPERTY(meta = (Output))
+	bool IsEmpty;
+};
+
+/**
+* Returns the items in the hierarchy pose
+*/
+USTRUCT(meta=(DisplayName="Get Pose Cache Items", Keywords="Hierarchy,Pose,State", Varying))
+struct CONTROLRIG_API FRigUnit_PoseGetItems : public FRigUnit_HierarchyBase
+{
+	GENERATED_BODY()
+
+	FRigUnit_PoseGetItems()
+	{
+		Pose = FRigPose();
+		ElementType = ERigElementType::All;
+		Items = FRigElementKeyCollection();
+	}
+
+	RIGVM_METHOD()
+	virtual void Execute(const FRigUnitContext& Context) override;
+
+	UPROPERTY(meta = (Input))
+	FRigPose Pose;
+
+	UPROPERTY(meta = (Input))
+	ERigElementType ElementType;
+
+	UPROPERTY(meta = (Output))
+	FRigElementKeyCollection Items;
+};
+
+/**
+* Compares two pose caches and compares their values.
+*/
+USTRUCT(meta=(DisplayName="Get Pose Cache Delta", Keywords="Hierarchy,Pose,State", Varying))
+struct CONTROLRIG_API FRigUnit_PoseGetDelta : public FRigUnit_HierarchyBase
+{
+	GENERATED_BODY()
+
+	FRigUnit_PoseGetDelta()
+	{
+		PoseA = PoseB = FRigPose();
+		ElementType = ERigElementType::All;
+		ItemsToCompare = ItemsWithDelta = FRigElementKeyCollection();
+		PositionThreshold = 0.1f;
+		RotationThreshold = ScaleThreshold = CurveThreshold = 0.f;
+		PosesAreEqual = false;
+	}
+
+	RIGVM_METHOD()
+	virtual void Execute(const FRigUnitContext& Context) override;
+
+	UPROPERTY(meta = (Input))
+	FRigPose PoseA;
+
+	UPROPERTY(meta = (Input))
+	FRigPose PoseB;
+
+	// The delta threshold for a translation / position difference. 0.0 disables position differences.
+	UPROPERTY(meta = (Input))
+	float PositionThreshold;
+
+	// The delta threshold for a rotation difference (in degrees). 0.0 disables rotation differences.
+	UPROPERTY(meta = (Input))
+	float RotationThreshold;
+
+	// The delta threshold for a scale difference. 0.0 disables scale differences.
+	UPROPERTY(meta = (Input))
+	float ScaleThreshold;
+
+	// The delta threshold for curve value difference. 0.0 disables curve differences.
+	UPROPERTY(meta = (Input))
+	float CurveThreshold;
+
+	UPROPERTY(meta = (Input))
+	ERigElementType ElementType;
+	
+	// Defines in which space transform deltas should be computed
+	UPROPERTY(meta = (Input))
+	EBoneGetterSetterMode Space;
+
+	// An optional list of items to compare
+	UPROPERTY(meta = (Input))
+	FRigElementKeyCollection ItemsToCompare;
+
+	UPROPERTY(meta = (Output))
+	bool PosesAreEqual;
+
+	UPROPERTY(meta = (Output))
+	FRigElementKeyCollection ItemsWithDelta;
+};
+
+/**
+* Returns the hierarchy's pose transform
+*/
+USTRUCT(meta=(DisplayName="Get Pose Cache Transform", Keywords="Hierarchy,Pose,State", Varying))
+struct CONTROLRIG_API FRigUnit_PoseGetTransform : public FRigUnit_HierarchyBase
+{
+	GENERATED_BODY()
+
+	FRigUnit_PoseGetTransform()
+	{
+		Pose = FRigPose();
+		Item = FRigElementKey();
+		Space = EBoneGetterSetterMode::GlobalSpace;
+		Valid = false;
+		Transform = FTransform::Identity;
+		CachedPoseElementIndex = CachedPoseHash = INDEX_NONE;
+	}
+
+	RIGVM_METHOD()
+	virtual void Execute(const FRigUnitContext& Context) override;
+
+	UPROPERTY(meta = (Input))
+	FRigPose Pose;
+
+	UPROPERTY(meta = (Input))
+	FRigElementKey Item;
+
+	/**
+	* Defines if the transform should be retrieved in local or global space
+	*/ 
+	UPROPERTY(meta = (Input))
+	EBoneGetterSetterMode Space;
+
+	UPROPERTY(meta = (Output))
+	bool Valid;
+	
+	UPROPERTY(meta = (Output))
+	FTransform Transform;
+
+	UPROPERTY(meta = (Output))
+	float CurveValue;
+
+	UPROPERTY()
+	int32 CachedPoseElementIndex;
+
+	UPROPERTY()
+	int32 CachedPoseHash;
+};
+
+/**
+* Returns the hierarchy's pose curve value
+*/
+USTRUCT(meta=(DisplayName="Get Pose Cache Curve", Keywords="Hierarchy,Pose,State", Varying))
+struct CONTROLRIG_API FRigUnit_PoseGetCurve : public FRigUnit_HierarchyBase
+{
+	GENERATED_BODY()
+
+	FRigUnit_PoseGetCurve()
+	{
+		Pose = FRigPose();
+		Curve = NAME_None;
+		Valid = false;
+		CurveValue = 0.f;
+		CachedPoseElementIndex = CachedPoseHash = INDEX_NONE;
+	}
+
+	RIGVM_METHOD()
+	virtual void Execute(const FRigUnitContext& Context) override;
+
+	UPROPERTY(meta = (Input))
+	FRigPose Pose;
+
+	UPROPERTY(meta = (Input, CustomWidget = "CurveName"))
+	FName Curve;
+
+	UPROPERTY(meta = (Output))
+	bool Valid;
+	
+	UPROPERTY(meta = (Output))
+	float CurveValue;
+
+	UPROPERTY()
+	int32 CachedPoseElementIndex;
+
+	UPROPERTY()
+	int32 CachedPoseHash;
+};
+
+/**
+* Given a pose, execute iteratively across all items in the pose
+*/
+USTRUCT(meta=(DisplayName="For Each Pose Cache Element", Keywords="Collection,Loop,Iterate", Icon="EditorStyle|GraphEditor.Macro.ForEach_16x"))
+struct CONTROLRIG_API FRigUnit_PoseLoop : public FRigUnit_HierarchyBaseMutable
+{
+	GENERATED_BODY()
+
+	FRigUnit_PoseLoop()
+	{
+		Pose = FRigPose();
+		Item = FRigElementKey();
+		GlobalTransform = LocalTransform = FTransform::Identity;
+		CurveValue = 0.f;
+		Count = 0;
+		Index = 0;
+		Ratio = 0.f;
+		Continue = false;
+	}
+
+	// FRigVMStruct overrides
+	FORCEINLINE virtual bool IsForLoop() const override { return true; }
+	FORCEINLINE virtual int32 GetNumSlices() const override { return Count; }
+
+	RIGVM_METHOD()
+	virtual void Execute(const FRigUnitContext& Context) override;
+
+	UPROPERTY(meta = (Input))
+	FRigPose Pose;
+
+	UPROPERTY(meta = (Singleton, Output))
+	FRigElementKey Item;
+
+	UPROPERTY(meta = (Singleton, Output))
+	FTransform GlobalTransform;
+
+	UPROPERTY(meta = (Singleton, Output))
+	FTransform LocalTransform;
+
+	UPROPERTY(meta = (Singleton, Output))
+	float CurveValue;
+
+	UPROPERTY(meta = (Singleton, Output))
+	int32 Index;
+
+	UPROPERTY(meta = (Singleton, Output))
+	int32 Count;
+
+	/**
+	* Ranging from 0.0 (first item) and 1.0 (last item)
+	* This is useful to drive a consecutive node with a 
+	* curve or an ease to distribute a value.
+	*/
+	UPROPERTY(meta = (Singleton, Output))
+	float Ratio;
+
+	UPROPERTY(meta = (Singleton))
+	bool Continue;
+
+	UPROPERTY(meta = (Output))
+	FControlRigExecuteContext Completed;
+};
diff --git a/Engine/Plugins/Experimental/ControlRig/Source/ControlRig/Public/Drawing/ControlRigDrawInterface.h b/Engine/Plugins/Experimental/ControlRig/Source/ControlRig/Public/Drawing/ControlRigDrawInterface.h
--- a/Engine/Plugins/Experimental/ControlRig/Source/ControlRig/Public/Drawing/ControlRigDrawInterface.h
+++ b/Engine/Plugins/Experimental/ControlRig/Source/ControlRig/Public/Drawing/ControlRigDrawInterface.h
@@ -37,10 +37,12 @@ public:
 	void DrawBox(const FTransform& WorldOffset, const FTransform& Transform, const FLinearColor& Color, float Thickness = 0.f);
 	void DrawAxes(const FTransform& WorldOffset, const FTransform& Transform, float Size, float Thickness = 0.f);
 	void DrawAxes(const FTransform& WorldOffset, const FRigVMFixedArray<FTransform>& Transforms, float Size, float Thickness = 0.f);
+	void DrawAxes(const FTransform& WorldOffset, const FTransform& Transform, const FLinearColor& InColor, float Size, float Thickness = 0.f);
+	void DrawAxes(const FTransform& WorldOffset, const FRigVMFixedArray<FTransform>& Transforms, const FLinearColor& InColor, float Size, float Thickness = 0.f);
 	void DrawRectangle(const FTransform& WorldOffset, const FTransform& Transform, float Size, const FLinearColor& Color, float Thickness);
 	void DrawArc(const FTransform& WorldOffset, const FTransform& Transform, float Radius, float MinimumAngle, float MaximumAngle, const FLinearColor& Color, float Thickness, int32 Detail);
 	void DrawBezier(const FTransform& WorldOffset, const FCRFourPointBezier& InBezier, float MinimumU, float MaximumU, const FLinearColor& Color, float Thickness, int32 Detail);
-	void DrawHierarchy(const FTransform& WorldOffset, URigHierarchy* Hierarchy, EControlRigDrawHierarchyMode::Type Mode, float Scale, const FLinearColor& Color, float Thickness);
+	void DrawHierarchy(const FTransform& WorldOffset, URigHierarchy* Hierarchy, EControlRigDrawHierarchyMode::Type Mode, float Scale, const FLinearColor& Color, float Thickness, const FRigPose* InPose = nullptr);
 	void DrawPointSimulation(const FTransform& WorldOffset, const FCRSimPointContainer& Simulation, const FLinearColor& Color, float Thickness, float PrimitiveSize = 0.f, bool bDrawPointsAsSphere = false);
 	void DrawCone(const FTransform& WorldOffset, const FTransform& ConeOffset, float Angle1, float Angle2, uint32 NumSides, bool bDrawSideLines, const FLinearColor& SideLineColor, FMaterialRenderProxy* const MaterialRenderProxy);
 
diff --git a/Engine/Plugins/Experimental/ControlRig/Source/ControlRig/Public/Rigs/RigHierarchy.h b/Engine/Plugins/Experimental/ControlRig/Source/ControlRig/Public/Rigs/RigHierarchy.h
--- a/Engine/Plugins/Experimental/ControlRig/Source/ControlRig/Public/Rigs/RigHierarchy.h
+++ b/Engine/Plugins/Experimental/ControlRig/Source/ControlRig/Public/Rigs/RigHierarchy.h
@@ -1643,14 +1643,54 @@ public:
 	 * @return The pose of the hierarchy
 	 */
 	UFUNCTION(BlueprintCallable, Category = URigHierarchy)
-	FRigPose GetPose(bool bInitial = false) const;
+	FORCEINLINE FRigPose GetPose(
+		bool bInitial = false
+	) const
+	{
+		return GetPose(bInitial, ERigElementType::All, FRigElementKeyCollection());
+	}
+
+	/**
+	 * Returns the current / initial pose of the hierarchy
+	 * @param bInitial If set to true the initial pose will be returned
+	 * @param InElementType The types of elements to get
+	 * @param InItems An optional list of items to get
+	 * @return The pose of the hierarchy
+	 */
+	FRigPose GetPose(
+		bool bInitial,
+		ERigElementType InElementType,
+		const FRigElementKeyCollection& InItems 
+	) const;
 
 	/**
 	 * Sets the current / initial pose of the hierarchy
 	 * @param InPose The pose to set on the hierarchy
 	 * @param InTransformType The transform type to set
 	 */
-	void SetPose(const FRigPose& InPose, ERigTransformType::Type InTransformType = ERigTransformType::CurrentLocal);
+	FORCEINLINE void SetPose(
+		const FRigPose& InPose,
+		ERigTransformType::Type InTransformType = ERigTransformType::CurrentLocal
+	)
+	{
+		SetPose(InPose, InTransformType, ERigElementType::All, FRigElementKeyCollection(), 1.f);
+	}
+
+	/**
+	 * Sets the current / initial pose of the hierarchy
+	 * @param InPose The pose to set on the hierarchy
+	 * @param InTransformType The transform type to set
+	 * @param InElementType The types of elements to set
+	 * @param InItems An optional list of items to set
+	 * @param InWeight A weight to define how much the pose needs to be mixed in
+	 */
+	void SetPose(
+		const FRigPose& InPose,
+		ERigTransformType::Type InTransformType,
+		ERigElementType InElementType,
+		const FRigElementKeyCollection& InItems,
+		float InWeight
+	);
 
 	/**
 	 * Sets the current / initial pose of the hierarchy
diff --git a/Engine/Plugins/Experimental/ControlRig/Source/ControlRig/Public/Rigs/RigHierarchyPose.h b/Engine/Plugins/Experimental/ControlRig/Source/ControlRig/Public/Rigs/RigHierarchyPose.h
--- a/Engine/Plugins/Experimental/ControlRig/Source/ControlRig/Public/Rigs/RigHierarchyPose.h
+++ b/Engine/Plugins/Experimental/ControlRig/Source/ControlRig/Public/Rigs/RigHierarchyPose.h
@@ -42,12 +42,51 @@ public:
 
 	GENERATED_BODY()
 
+	FRigPose()
+	: HierarchyTopologyVersion(INDEX_NONE)
+	, PoseHash(INDEX_NONE)
+	, CachedPoseHash(INDEX_NONE)
+	{
+	}
+
 	FORCEINLINE void Reset() { Elements.Reset(); }
 
 	FORCEINLINE int32 Num() const { return Elements.Num(); }
+	FORCEINLINE bool IsValidIndex(int32 InIndex) const { return Elements.IsValidIndex(InIndex); }
+
+	FORCEINLINE int32 GetIndex(const FRigElementKey& InKey) const
+	{
+		if(CachedPoseHash != PoseHash)
+		{
+			KeyToIndex.Reset();
+			CachedPoseHash = PoseHash;
+		}
+
+		if(const int32* IndexPtr = KeyToIndex.Find(InKey))
+		{
+			return *IndexPtr;
+		}
+		
+		for(int32 Index=0;Index<Elements.Num();Index++)
+		{
+			const FRigElementKey& Key = Elements[Index].Index.GetKey();
+
+			KeyToIndex.Add(Key, Index);
+			if(Key == InKey)
+			{
+				return Index;
+			}
+		}
+		return INDEX_NONE;
+	}
+
+	FORCEINLINE bool Contains(const FRigElementKey& InKey) const
+	{
+		return GetIndex(InKey) != INDEX_NONE;
+	}
+
 	FORCEINLINE const FRigPoseElement& operator[](int32 InIndex) const { return Elements[InIndex]; }
 	FORCEINLINE FRigPoseElement& operator[](int32 InIndex) { return Elements[InIndex]; }
-
 	FORCEINLINE TArray<FRigPoseElement>::RangedForIteratorType      begin()       { return Elements.begin(); }
 	FORCEINLINE TArray<FRigPoseElement>::RangedForConstIteratorType begin() const { return Elements.begin(); }
 	FORCEINLINE TArray<FRigPoseElement>::RangedForIteratorType      end()         { return Elements.end();   }
@@ -55,4 +94,15 @@ public:
 
 	UPROPERTY()
 	TArray<FRigPoseElement> Elements;
+
+	UPROPERTY()
+	int32 HierarchyTopologyVersion;
+
+	UPROPERTY()
+	int32 PoseHash;
+
+private:
+	
+	mutable int32 CachedPoseHash;
+	mutable TMap<FRigElementKey, int32> KeyToIndex;
 };
diff --git a/Engine/Plugins/Experimental/ControlRig/Source/ControlRigDeveloper/Private/ControlRigBlueprint.cpp b/Engine/Plugins/Experimental/ControlRig/Source/ControlRigDeveloper/Private/ControlRigBlueprint.cpp
--- a/Engine/Plugins/Experimental/ControlRig/Source/ControlRigDeveloper/Private/ControlRigBlueprint.cpp
+++ b/Engine/Plugins/Experimental/ControlRig/Source/ControlRigDeveloper/Private/ControlRigBlueprint.cpp
@@ -1082,6 +1082,10 @@ URigVMController* UControlRigBlueprint::GetOrCreateController(URigVMGraph* InGra
 		{
 			return false;
 		}
+		if (InStruct == FRigPose::StaticStruct())
+		{
+			return false;
+		}
 		return true;
 		});
 
diff --git a/Engine/Plugins/Experimental/ControlRig/Source/ControlRigDeveloper/Private/Graph/ControlRigGraphSchema.cpp b/Engine/Plugins/Experimental/ControlRig/Source/ControlRigDeveloper/Private/Graph/ControlRigGraphSchema.cpp
--- a/Engine/Plugins/Experimental/ControlRig/Source/ControlRigDeveloper/Private/Graph/ControlRigGraphSchema.cpp
+++ b/Engine/Plugins/Experimental/ControlRig/Source/ControlRigDeveloper/Private/Graph/ControlRigGraphSchema.cpp
@@ -347,6 +347,11 @@ FLinearColor UControlRigGraphSchema::GetPinTypeColor(const FEdGraphPinType& PinT
 				return FLinearColor(0.0, 0.6588, 0.9490);
 			}
 
+			if (Struct == FRigElementKey::StaticStruct() || Struct == FRigPose::StaticStruct())
+			{
+				return FLinearColor(0.0, 0.3588, 0.5490);
+			}
+
 			// external types can register their own colors, check if there are any
 			if (IControlRigDeveloperModule* Module = FModuleManager::GetModulePtr<IControlRigDeveloperModule>("ControlRigDeveloper"))
 			{
