Control Rig: Capsule based bone selection

diff --git a/Engine/Plugins/Experimental/ControlRig/Source/ControlRigEditor/Private/EditMode/ControlRigEditMode.cpp b/Engine/Plugins/Experimental/ControlRig/Source/ControlRigEditor/Private/EditMode/ControlRigEditMode.cpp
--- a/Engine/Plugins/Experimental/ControlRig/Source/ControlRigEditor/Private/EditMode/ControlRigEditMode.cpp
+++ b/Engine/Plugins/Experimental/ControlRig/Source/ControlRigEditor/Private/EditMode/ControlRigEditMode.cpp
@@ -46,6 +46,8 @@
 //#include "IPersonaPreviewScene.h"
 //#include "Animation/DebugSkelMeshComponent.h"
 //#include "Persona/Private/AnimationEditorViewportClient.h"
+#include "IPersonaPreviewScene.h"
+#include "PersonaSelectionComponent.h"
 #include "Framework/Application/SlateApplication.h"
 #include "UnrealEdGlobals.h"
 #include "Editor/UnrealEdEngine.h"
@@ -547,15 +549,25 @@ void FControlRigEditMode::Render(const FSceneView* View, FViewport* Viewport, FP
 			});
 		}
 
-		if (Settings->bDisplaySpaces || ControlRig->IsSetupModeEnabled())
+		if (Settings->bDisplayNulls || ControlRig->IsSetupModeEnabled())
 		{
 			TArray<FTransform> SpaceTransforms;
-			Hierarchy->ForEach<FRigNullElement>([&SpaceTransforms, Hierarchy](FRigNullElement* NullElement) -> bool
+			TArray<FTransform> SelectedSpaceTransforms;
+			Hierarchy->ForEach<FRigNullElement>([&SpaceTransforms, &SelectedSpaceTransforms, Hierarchy](FRigNullElement* NullElement) -> bool
             {
-				SpaceTransforms.Add(Hierarchy->GetTransform(NullElement, ERigTransformType::CurrentGlobal));
+				if(Hierarchy->IsSelected(NullElement->GetIndex()))
+				{
+					SelectedSpaceTransforms.Add(Hierarchy->GetTransform(NullElement, ERigTransformType::CurrentGlobal));
+				}
+				else
+				{
+					SpaceTransforms.Add(Hierarchy->GetTransform(NullElement, ERigTransformType::CurrentGlobal));
+				}
 				return true;
 			});
+
 			GetControlRig(true)->DrawInterface.DrawAxes(FTransform::Identity, SpaceTransforms, Settings->AxisScale);
+			GetControlRig(true)->DrawInterface.DrawAxes(FTransform::Identity, SelectedSpaceTransforms, FLinearColor(1.0f, 0.34f, 0.0f, 1.0f), Settings->AxisScale);
 		}
 
 		if (Settings->bDisplayAxesOnSelection && Settings->AxisScale > SMALL_NUMBER)
@@ -889,6 +901,11 @@ bool FControlRigEditMode::HandleClick(FEditorViewportClient* InViewportClient, H
 			}
 		}
 	}
+	else if(HPersonaSelectionHitProxy* CapsuleHitProxy = HitProxyCast<HPersonaSelectionHitProxy>(HitProxy))
+	{
+		static_cast<HPersonaSelectionHitProxy*>(HitProxy)->BroadcastClicked();
+		return true;
+	}
 
 	// for now we show this menu all the time if body is selected
 	// if we want some global menu, we'll have to move this
diff --git a/Engine/Plugins/Experimental/ControlRig/Source/ControlRigEditor/Private/EditMode/ControlRigEditModeSettings.h b/Engine/Plugins/Experimental/ControlRig/Source/ControlRigEditor/Private/EditMode/ControlRigEditModeSettings.h
--- a/Engine/Plugins/Experimental/ControlRig/Source/ControlRigEditor/Private/EditMode/ControlRigEditModeSettings.h
+++ b/Engine/Plugins/Experimental/ControlRig/Source/ControlRigEditor/Private/EditMode/ControlRigEditModeSettings.h
@@ -15,7 +15,7 @@ class UControlRigEditModeSettings : public UObject
 
 	UControlRigEditModeSettings()
 		: bDisplayHierarchy(false)
-		, bDisplaySpaces(false)
+		, bDisplayNulls(false)
 		, bHideManipulators(false)
 		, bDisplayAxesOnSelection(false)
 		, AxisScale(10.f)
@@ -37,9 +37,9 @@ public:
 	UPROPERTY(EditAnywhere, AdvancedDisplay, Category = "Animation")
 	bool bDisplayHierarchy;
 
-	/** Whether to show all spaces in the hierarchy */
+	/** Whether to show all nulls in the hierarchy */
 	UPROPERTY(EditAnywhere, AdvancedDisplay, Category = "Animation")
-	bool bDisplaySpaces;
+	bool bDisplayNulls;
 
 	/** Should we always hide manipulators in viewport */
 	UPROPERTY(EditAnywhere, AdvancedDisplay, Category = "Animation")
diff --git a/Engine/Plugins/Experimental/ControlRig/Source/ControlRigEditor/Private/Editor/ControlRigEditor.cpp b/Engine/Plugins/Experimental/ControlRig/Source/ControlRigEditor/Private/Editor/ControlRigEditor.cpp
--- a/Engine/Plugins/Experimental/ControlRig/Source/ControlRigEditor/Private/Editor/ControlRigEditor.cpp
+++ b/Engine/Plugins/Experimental/ControlRig/Source/ControlRigEditor/Private/Editor/ControlRigEditor.cpp
@@ -503,6 +503,8 @@ void FControlRigEditor::InitControlRigEditor(const EToolkitMode::Type Mode, cons
 		}
 		
 	}, TStatId(), NULL, ENamedThreads::GameThread);
+
+	UpdateCapsules();
 }
 
 void FControlRigEditor::BindCommands()
@@ -1026,7 +1028,7 @@ void FControlRigEditor::ToggleSetupMode()
 			EditMode->RecreateGizmoActors(RigBlueprint->Hierarchy->GetSelectedKeys());
 		}
 
-		EditMode->Settings->bDisplaySpaces = bSetupModeEnabled;
+		EditMode->Settings->bDisplayNulls = bSetupModeEnabled;
 	}
 
 	if (PreviousSelection.Num() > 0)
@@ -2778,10 +2780,10 @@ void FControlRigEditor::HandleViewportCreated(const TSharedRef<class IPersonaVie
 							SNew(SCheckBox)
 							.IsChecked(this, &FControlRigEditor::GetToolbarDrawNulls)
 							.OnCheckStateChanged(this, &FControlRigEditor::OnToolbarDrawNullsChanged)
-							.ToolTipText(LOCTEXT("ControlRigDrawNullsToolTip", "If checked all spaces are drawn as axes."))
+							.ToolTipText(LOCTEXT("ControlRigDrawNullsToolTip", "If checked all nulls are drawn as axes."))
 						]
 					],
-					LOCTEXT("ControlRigDisplaySpaces", "Display Spaces")
+					LOCTEXT("ControlRigDisplayNulls", "Display Nulls")
 				);
 
 				InMenuBuilder.AddWidget(
@@ -2933,6 +2935,10 @@ void FControlRigEditor::OnToolbarBoneRadiusChanged(float InValue)
 	{
 		MeshComponent->BoneRadiusMultiplier = InValue;
 	}
+	if(UControlRig* DebuggedControlRig = Cast<UControlRig>(GetBlueprintObj()->GetObjectBeingDebugged()))
+	{
+		DebuggedControlRig->DebugBoneRadiusMultiplier = InValue;
+	}
 }
 
 ECheckBoxState FControlRigEditor::GetToolbarDrawAxesOnSelection() const
@@ -2968,7 +2974,7 @@ ECheckBoxState FControlRigEditor::GetToolbarDrawNulls() const
 {
 	if (FControlRigEditMode* EditMode = GetEditMode())
 	{
-		return EditMode->Settings->bDisplaySpaces ? ECheckBoxState::Checked : ECheckBoxState::Unchecked;
+		return EditMode->Settings->bDisplayNulls ? ECheckBoxState::Checked : ECheckBoxState::Unchecked;
 	}
 	return ECheckBoxState::Unchecked;
 }
@@ -2977,7 +2983,7 @@ void FControlRigEditor::OnToolbarDrawNullsChanged(ECheckBoxState InNewValue)
 {
 	if (FControlRigEditMode* EditMode = GetEditMode())
 	{
-		EditMode->Settings->bDisplaySpaces = InNewValue == ECheckBoxState::Checked;
+		EditMode->Settings->bDisplayNulls = InNewValue == ECheckBoxState::Checked;
 	}
 }
 
@@ -3108,9 +3114,9 @@ void FControlRigEditor::HandlePreviewSceneCreated(const TSharedRef<IPersonaPrevi
 	GroundActor->GetStaticMeshComponent()->SetStaticMesh(FloorMesh);
 	GroundActor->GetStaticMeshComponent()->SetMaterial(0, DefaultMaterial);
 	GroundActor->SetMobility(EComponentMobility::Static);
-	//GroundActor->GetStaticMeshComponent()->SetVisibility(false);
 	GroundActor->GetStaticMeshComponent()->SetCollisionEnabled(ECollisionEnabled::QueryAndPhysics);
 	GroundActor->GetStaticMeshComponent()->SetCollisionProfileName(UCollisionProfile::BlockAll_ProfileName);
+	GroundActor->GetStaticMeshComponent()->bSelectable = false;
 
 	AAnimationEditorPreviewActor* Actor = InPersonaPreviewScene->GetWorld()->SpawnActor<AAnimationEditorPreviewActor>(AAnimationEditorPreviewActor::StaticClass(), FTransform::Identity);
 	Actor->SetFlags(RF_Transient);
@@ -3126,6 +3132,11 @@ void FControlRigEditor::HandlePreviewSceneCreated(const TSharedRef<IPersonaPrevi
 
 	// set root component, so we can attach to it. 
 	Actor->SetRootComponent(EditorSkelComp);
+	EditorSkelComp->bSelectable = false;
+	EditorSkelComp->MarkRenderStateDirty();
+	
+	InPersonaPreviewScene->SetAllowMeshHitProxies(false);
+	InPersonaPreviewScene->SetAdditionalMeshesSelectable(false);
 
 	PreviewInstance = nullptr;
 	if (UControlRigLayerInstance* ControlRigLayerInstance = Cast<UControlRigLayerInstance>(EditorSkelComp->GetAnimInstance()))
@@ -3154,6 +3165,88 @@ void FControlRigEditor::HandlePreviewSceneCreated(const TSharedRef<IPersonaPrevi
 			OnToolbarBoneRadiusChanged(Blueprint->DebugBoneRadius);
 		}
 	}
+
+	UPersonaSelectionComponent* SelectionComponent = InPersonaPreviewScene->GetSelectionComponent();
+	check(SelectionComponent != nullptr);
+
+	// make sure the selection component ticks after the skeletal mesh comp
+	SelectionComponent->AddTickPrerequisiteComponent(EditorSkelComp);
+
+	SelectionComponent->OnUpdateCapsules().BindLambda([this](UPersonaSelectionComponent*, const TArray<int32>& InIndices, TArray<FPersonaSelectionCapsule>& OutCapsules)
+	{
+		if(UControlRig* DebuggedControlRig = Cast<UControlRig>(GetBlueprintObj()->GetObjectBeingDebugged()))
+		{
+			URigHierarchy* Hierarchy = DebuggedControlRig->GetHierarchy();
+			check(Hierarchy != nullptr);
+			
+			for(int32 Index = 0; Index < InIndices.Num(); Index++)
+			{
+				const int32 CapsuleIndex = InIndices[Index];
+				const int32 HierarchyIndex = CapsuleToHierarchyIndex.FindChecked(CapsuleIndex);
+				check(Hierarchy->IsValidIndex(HierarchyIndex));
+				check(OutCapsules.IsValidIndex(CapsuleIndex));
+
+				FPersonaSelectionCapsule& OutCapsule = OutCapsules[CapsuleIndex];
+				OutCapsule.Transform = Hierarchy->GetGlobalTransform(HierarchyIndex);
+				OutCapsule.Radius = 10.f;
+				OutCapsule.HalfHeight = 5.f;
+
+				if(Hierarchy->GetKey(HierarchyIndex).Type == ERigElementType::Bone)
+				{
+					const int32 ParentIndex = Hierarchy->GetFirstParent(HierarchyIndex);
+
+					FVector Start, End;
+					if (ParentIndex >= 0)
+					{
+						Start = Hierarchy->GetGlobalTransform(ParentIndex).GetLocation();
+						End = OutCapsule.Transform.GetLocation();
+					}
+					else
+					{
+						Start = FVector::ZeroVector;
+						End = OutCapsule.Transform.GetLocation();
+					}
+					
+					UPersonaSelectionComponent::ComputeCapsuleFromBonePositions(
+						Start,
+						End,
+						100.f,
+						DebuggedControlRig->GetDebugBoneRadiusMultiplier(),
+						OutCapsule);
+				}
+			}
+		}
+	});
+
+	// setup the capsules in the preview scene
+	SelectionComponent->OnClicked().BindLambda([this] (UPersonaSelectionComponent*, int32 InCapsuleIndex, const FPersonaSelectionCapsule& InCapsule)
+	{
+		const int32 HierarchyIndex = CapsuleToHierarchyIndex.FindChecked(InCapsuleIndex);
+
+		if (UControlRigBlueprint* ControlRigBP = GetControlRigBlueprint())
+		{
+			URigHierarchy* Hierarchy = ControlRigBP->Hierarchy;
+			const FRigElementKey ElementToSelect = Hierarchy->GetKey(HierarchyIndex);
+
+			if (FSlateApplication::Get().GetModifierKeys().IsShiftDown())
+			{
+				Hierarchy->GetController()->SelectElement(ElementToSelect, true);
+			}
+			else if (FSlateApplication::Get().GetModifierKeys().IsControlDown())
+			{
+				const bool bSelect = !Hierarchy->IsSelected(ElementToSelect);
+				Hierarchy->GetController()->SelectElement(ElementToSelect, bSelect);
+			}
+			else
+			{
+				TArray<FRigElementKey> NewSelection;
+				NewSelection.Add(ElementToSelect);
+				Hierarchy->GetController()->SetSelection(NewSelection);
+			}
+		}
+	});
+
+	UpdateCapsules();
 }
 
 void FControlRigEditor::UpdateControlRig()
@@ -3844,6 +3937,7 @@ void FControlRigEditor::OnHierarchyChanged()
 	}
 	
 	CacheNameLists();
+	UpdateCapsules();
 }
 
 
@@ -5054,6 +5148,46 @@ bool FControlRigEditor::IsHaltedAtBreakpoint() const
 	return false;
 }
 
+void FControlRigEditor::UpdateCapsules()
+{
+	if(!PersonaToolkit.IsValid())
+	{
+		return;
+	}
+	
+	TSharedRef<IPersonaPreviewScene> CurrentPreviewScene = GetPersonaToolkit()->GetPreviewScene();
+	UPersonaSelectionComponent* SelectionComponent = CurrentPreviewScene->GetSelectionComponent();
+	check(SelectionComponent != nullptr);
+
+	SelectionComponent->Reset();
+	CapsuleToHierarchyIndex.Reset();
+	
+	if(UControlRig* DebuggedControlRig = Cast<UControlRig>(GetBlueprintObj()->GetObjectBeingDebugged()))
+	{
+		URigHierarchy* Hierarchy = DebuggedControlRig->GetHierarchy();
+		
+		for(int32 Index = 0; Index < Hierarchy->Num(); Index++)
+		{
+			const FRigElementKey Key = Hierarchy->GetKey(Index);
+			switch(Key.Type)
+			{
+				case ERigElementType::Bone:
+				case ERigElementType::Null:
+				//case ERigElementType::Control:
+				{
+					const int32 CapsuleIndex = SelectionComponent->Add();
+					CapsuleToHierarchyIndex.Add(CapsuleIndex, Index);
+					break;
+				}
+				default:
+				{
+					break;
+				}						
+			}
+		}
+	}
+}
+
 void FControlRigEditor::UpdateGraphCompilerErrors()
 {
 	DECLARE_SCOPE_HIERARCHICAL_COUNTER_FUNC()
diff --git a/Engine/Plugins/Experimental/ControlRig/Source/ControlRigEditor/Private/Editor/ControlRigEditor.h b/Engine/Plugins/Experimental/ControlRig/Source/ControlRigEditor/Private/Editor/ControlRigEditor.h
--- a/Engine/Plugins/Experimental/ControlRig/Source/ControlRigEditor/Private/Editor/ControlRigEditor.h
+++ b/Engine/Plugins/Experimental/ControlRig/Source/ControlRigEditor/Private/Editor/ControlRigEditor.h
@@ -337,6 +337,8 @@ private:
 
 	bool IsHaltedAtBreakpoint() const;
 
+	void UpdateCapsules();
+
 protected:
 
 	/** Toolbox hosting widget */
@@ -437,6 +439,8 @@ protected:
 
 	bool bSuspendDetailsPanelRefresh;
 
+	TMap<int32, int32> CapsuleToHierarchyIndex;
+
 	friend class FControlRigEditorMode;
 	friend class SControlRigStackView;
 	friend class SRigHierarchy;
diff --git a/Engine/Source/Editor/Persona/Private/AnimationEditorPreviewScene.cpp b/Engine/Source/Editor/Persona/Private/AnimationEditorPreviewScene.cpp
--- a/Engine/Source/Editor/Persona/Private/AnimationEditorPreviewScene.cpp
+++ b/Engine/Source/Editor/Persona/Private/AnimationEditorPreviewScene.cpp
@@ -1112,6 +1112,27 @@ void FAnimationEditorPreviewScene::SetUsePhysicsBodiesForBoneSelection(bool bUse
 	bUsePhysicsBodiesForBoneSelection = bUsePhysicsBodies;
 }
 
+UPersonaSelectionComponent* FAnimationEditorPreviewScene::GetSelectionComponent()
+{
+	if(!SelectionComponentPtr.IsValid())
+	{
+		FActorSpawnParameters SelectionActorParameters;
+		SelectionActorParameters.Name = TEXT("SelectionComponent");
+		AActor* SelectionActor = GetWorld()->SpawnActor<AActor>(AActor::StaticClass(), FTransform::Identity, SelectionActorParameters);
+		check(SelectionActor);
+
+		SelectionActor->SetFlags(RF_Transient);
+		UPersonaSelectionComponent* SelectionComponent = NewObject<UPersonaSelectionComponent>(SelectionActor);
+		SelectionActor->AddOwnedComponent(SelectionComponent);
+		SelectionActor->SetRootComponent(SelectionComponent);
+		SelectionComponent->RegisterComponent();
+		SelectionComponent->SetCollisionEnabled(ECollisionEnabled::NoCollision);
+
+		SelectionComponentPtr = SelectionComponent;
+	}
+	return SelectionComponentPtr.Get();
+}
+
 void FAnimationEditorPreviewScene::SetAllowMeshHitProxies(bool bState)
 {
 	bEnableMeshHitProxies = bState;
diff --git a/Engine/Source/Editor/Persona/Private/AnimationEditorPreviewScene.h b/Engine/Source/Editor/Persona/Private/AnimationEditorPreviewScene.h
--- a/Engine/Source/Editor/Persona/Private/AnimationEditorPreviewScene.h
+++ b/Engine/Source/Editor/Persona/Private/AnimationEditorPreviewScene.h
@@ -323,6 +323,7 @@ public:
 	/** Set whether bones can be selected by their physics bodies */
 	virtual void SetUsePhysicsBodiesForBoneSelection(bool bUsePhysicsBodies) override;
 
+	virtual UPersonaSelectionComponent* GetSelectionComponent() override;
 
 private:
 	/** Set preview mesh internal use only. The mesh should be verified by now. */
@@ -473,4 +474,7 @@ private:
 
 	/** Delegate Remove attach component */
 	FOnRemoveAttachedComponentFilter OnRemoveAttachedComponentFilter;
+
+	/** The component storing the capsules */
+	TWeakObjectPtr<UPersonaSelectionComponent> SelectionComponentPtr;
 };
diff --git a/Engine/Source/Editor/Persona/Private/PersonaSelectionComponent.cpp b/Engine/Source/Editor/Persona/Private/PersonaSelectionComponent.cpp
new file mode 100644
--- /dev/null
+++ b/Engine/Source/Editor/Persona/Private/PersonaSelectionComponent.cpp
@@ -0,0 +1,254 @@
+// Copyright Epic Games, Inc. All Rights Reserved.
+
+#include "PersonaSelectionComponent.h"
+
+IMPLEMENT_HIT_PROXY(HPersonaSelectionHitProxy, HHitProxy)
+
+const FPersonaSelectionCapsule& HPersonaSelectionHitProxy::GetCapsule() const
+{
+	return Component->operator[](CapsuleIndex);
+}
+
+void HPersonaSelectionHitProxy::BroadcastClicked() const
+{
+	Component->OnClicked().ExecuteIfBound(Component, CapsuleIndex, GetCapsule());
+}
+
+////////////////////////////////////////////////////////////////////////////////////////
+
+UPersonaSelectionComponent::UPersonaSelectionComponent(const FObjectInitializer& ObjectInitializer)
+	: Super(ObjectInitializer)
+	, bVisibleInMainPass(false)
+{
+	PrimaryComponentTick.TickGroup = TG_LastDemotable;
+	PrimaryComponentTick.bStartWithTickEnabled = true;
+	PrimaryComponentTick.bCanEverTick = true;
+	PrimaryComponentTick.bTickEvenWhenPaused = true;
+	bTickInEditor = true;
+	bAutoActivate = true;
+}
+
+void UPersonaSelectionComponent::Reset()
+{
+	Capsules.Reset();
+	CapsuleIndices.Reset();
+}
+
+void UPersonaSelectionComponent::SetNum(int32 InCount)
+{
+	Capsules.SetNum(InCount);
+	SetCapsulesIndicesToFullArray();
+}
+
+int32 UPersonaSelectionComponent::Add(int32 InCount)
+{
+	int32 FirstIndex = INDEX_NONE;
+	for(int32 Index=0;Index<InCount;Index++)
+	{
+		const int32 CapsuleIndex = Capsules.Add(FPersonaSelectionCapsule());
+		if(FirstIndex == INDEX_NONE)
+		{
+			FirstIndex = CapsuleIndex;
+		}
+		CapsuleIndices.Add(CapsuleIndex);
+	}
+	return FirstIndex;
+}
+
+void UPersonaSelectionComponent::SetCapsulesIndicesToFullArray()
+{
+	CapsuleIndices.SetNum(Capsules.Num());
+	for(int32 Index=0;Index<CapsuleIndices.Num();Index++)
+	{
+		CapsuleIndices[Index] = Index;
+	}
+}
+
+void UPersonaSelectionComponent::MarkCapsulesChanged()
+{
+	while(HitProxies.Num() > Capsules.Num())
+	{
+		HitProxies.Pop();
+	}
+	while(HitProxies.Num() < Capsules.Num())
+	{
+		HitProxies.Add(new HPersonaSelectionHitProxy(HitProxies.Num(), this));
+	}
+	
+	UpdateBounds();
+	MarkRenderStateDirty();
+}
+
+FPrimitiveSceneProxy* UPersonaSelectionComponent::CreateSceneProxy()
+{
+	return new FPersonaSelectionComponentProxy(this);
+}
+
+FBoxSphereBounds UPersonaSelectionComponent::CalcBounds(const FTransform& LocalToWorld) const
+{
+	FBox BBox(ForceInit);
+
+	if(CapsuleIndices.Num() > 0)
+	{
+		const FTransform ComponentTransform = GetComponentToWorld();
+		
+		for(int32 CapsuleIndex : CapsuleIndices)
+		{
+			const FPersonaSelectionCapsule& Capsule = Capsules[CapsuleIndex];
+			const FTransform& Transform = Capsule.Transform * ComponentTransform;
+			BBox += Transform.TransformPosition(FVector(0.f, 0.f, Capsule.HalfHeight + Capsule.Radius));
+			BBox += Transform.TransformPosition(-FVector(0.f, 0.f, Capsule.HalfHeight + Capsule.Radius));
+		}
+	}
+
+	if (!BBox.IsValid)
+	{
+		const FVector BoxExtent(1.f);
+		return FBoxSphereBounds(LocalToWorld.GetLocation(), BoxExtent, 1.f);
+	}
+
+	// Points are in world space, so no need to transform.
+	return FBoxSphereBounds(BBox);
+}
+
+void UPersonaSelectionComponent::GetUsedMaterials(TArray<UMaterialInterface*>& OutMaterials,
+	bool bGetDebugMaterials) const
+{
+	OutMaterials.Add(GEngine->DefaultFlattenMaterial);
+}
+
+void UPersonaSelectionComponent::TickComponent(float DeltaTime, ELevelTick TickType,
+	FActorComponentTickFunction* ThisTickFunction)
+{
+	Super::TickComponent(DeltaTime, TickType, ThisTickFunction);
+
+	if(Capsules.Num() > 0 && CapsuleIndices.Num() > 0)
+	{
+		if(OnPersonaSelectionUpdateCapsules.IsBound())
+		{
+			OnPersonaSelectionUpdateCapsules.Execute(this, CapsuleIndices, Capsules);
+			MarkCapsulesChanged();
+		}
+	}
+}
+
+void UPersonaSelectionComponent::ComputeCapsuleFromBonePositions(
+	const FVector& InWorldStart,
+	const FVector& InWorldEnd,
+	float InBounds,
+	float InRadiusMultiplier,
+	FPersonaSelectionCapsule& OutCapsule
+	)
+{
+	static const float MaxDrawRadius = InBounds * 0.02f;
+	const FVector Direction = InWorldEnd - InWorldStart;
+	const float BoneLength = Direction.Size();
+	OutCapsule.Radius = FMath::Clamp(BoneLength * 0.05f, 0.1f, MaxDrawRadius) * InRadiusMultiplier;
+	OutCapsule.HalfHeight = BoneLength * 0.5f;
+
+	FVector Location = (InWorldStart + InWorldEnd) * 0.5f;
+	FQuat AlignmentQuat = FQuat::Identity;
+	if(!Direction.IsNearlyZero())
+	{
+		Location -= Direction.GetSafeNormal() * OutCapsule.Radius;
+		AlignmentQuat = FQuat::FindBetween(FVector(0.f, 0.f, 1.f), Direction);
+	}
+
+	OutCapsule.Transform.SetLocation(Location);
+	OutCapsule.Transform.SetRotation(AlignmentQuat);
+	OutCapsule.Transform.SetScale3D(FVector::OneVector);
+}
+
+////////////////////////////////////////////////////////////////////////////////////////
+
+FPersonaSelectionComponentProxy::FPersonaSelectionComponentProxy(const UPersonaSelectionComponent* InComponent)
+: FPrimitiveSceneProxy(InComponent)
+, SelectionComponentPtr(InComponent)
+, bVisibleInMainPass(false)
+{
+	bWillEverBeLit = false;
+
+	if(InComponent)
+	{
+		bVisibleInMainPass = InComponent->bVisibleInMainPass;
+	}
+}
+
+SIZE_T FPersonaSelectionComponentProxy::GetTypeHash() const
+{
+	static size_t UniquePointer;
+	return reinterpret_cast<size_t>(&UniquePointer);
+}
+
+void FPersonaSelectionComponentProxy::GetDynamicMeshElements(const TArray<const FSceneView*>& Views, const FSceneViewFamily& ViewFamily, uint32 VisibilityMap, FMeshElementCollector& Collector) const
+{
+	if(!SelectionComponentPtr.IsValid())
+	{
+		return;
+	}
+
+	const UPersonaSelectionComponent* SelectionComponent = SelectionComponentPtr.Get();
+	if(SelectionComponent->Capsules.IsEmpty())
+	{
+		return;
+	}
+
+	for (int32 ViewIndex = 0; ViewIndex < Views.Num(); ViewIndex++)
+	{
+		if (VisibilityMap & (1 << ViewIndex))
+		{
+			const FSceneView* View = Views[ViewIndex];
+			FPrimitiveDrawInterface* PDI = Collector.GetPDI(ViewIndex);
+
+			FMaterialRenderProxy* MaterialRenderProxy = GEngine->DefaultFlattenMaterial->GetRenderProxy();
+			const FTransform ComponentToWorld = SelectionComponent->GetComponentToWorld();
+
+			for(int32 CapsuleIndex : SelectionComponent->CapsuleIndices)
+            {
+				const FPersonaSelectionCapsule& Capsule = SelectionComponent->Capsules[CapsuleIndex];
+				const FLinearColor LineColor = FLinearColor::Yellow;
+				const FTransform& Transform = Capsule.Transform * ComponentToWorld;
+				HPersonaSelectionHitProxy* HitProxy = SelectionComponent->HitProxies[CapsuleIndex].GetReference();
+
+				GetBoxMesh(
+					Transform.ToMatrixWithScale(),
+					FVector(Capsule.Radius,
+					Capsule.Radius,
+					Capsule.HalfHeight),
+					MaterialRenderProxy,
+					SDPG_Foreground,
+					ViewIndex,
+					Collector,
+					HitProxy
+				);
+			}
+		}
+	}
+}
+
+/**
+*  Returns a struct that describes to the renderer when to draw this proxy.
+*	@param		Scene view to use to determine our relevence.
+*  @return		View relevance struct
+*/
+FPrimitiveViewRelevance FPersonaSelectionComponentProxy::GetViewRelevance(const FSceneView* View) const
+{
+	FPrimitiveViewRelevance ViewRelevance;
+	ViewRelevance.bDrawRelevance = IsShown(View);
+	ViewRelevance.bOpaque = !bVisibleInMainPass;
+	ViewRelevance.bDynamicRelevance = true;
+	ViewRelevance.bShadowRelevance = false;
+	ViewRelevance.bEditorPrimitiveRelevance = !bVisibleInMainPass;
+	ViewRelevance.bSeparateTranslucency = ViewRelevance.bNormalTranslucency = true;
+	return ViewRelevance;
+}
+
+uint32 FPersonaSelectionComponentProxy::GetMemoryFootprint(void) const
+{
+	return(sizeof(*this) + GetAllocatedSize());
+}
+
+uint32 FPersonaSelectionComponentProxy::GetAllocatedSize(void) const
+{
+	return FPrimitiveSceneProxy::GetAllocatedSize();
+}
diff --git a/Engine/Source/Editor/Persona/Public/IPersonaPreviewScene.h b/Engine/Source/Editor/Persona/Public/IPersonaPreviewScene.h
--- a/Engine/Source/Editor/Persona/Public/IPersonaPreviewScene.h
+++ b/Engine/Source/Editor/Persona/Public/IPersonaPreviewScene.h
@@ -7,6 +7,7 @@
 #include "Containers/ArrayView.h"
 #include "Types/SlateEnums.h"
 #include "Animation/AnimBlueprint.h"
+#include "PersonaSelectionComponent.h"
 
 class UAnimationAsset;
 class UDebugSkelMeshComponent;
@@ -304,4 +305,8 @@ public:
 
 	/** Set whether bones can be selected by their physics bodies */
 	virtual void SetUsePhysicsBodiesForBoneSelection(bool bUsePhysicsBodies) = 0;
+
+	/** Returns the persona selection component (or nullptr) */
+	virtual UPersonaSelectionComponent* GetSelectionComponent() = 0;
 };
+
diff --git a/Engine/Source/Editor/Persona/Public/PersonaSelectionComponent.h b/Engine/Source/Editor/Persona/Public/PersonaSelectionComponent.h
new file mode 100644
--- /dev/null
+++ b/Engine/Source/Editor/Persona/Public/PersonaSelectionComponent.h
@@ -0,0 +1,172 @@
+// Copyright Epic Games, Inc. All Rights Reserved.
+
+#pragma once
+
+#include "Components/PrimitiveComponent.h"
+#include "PersonaSelectionComponent.generated.h"
+
+class UPersonaSelectionComponent;
+
+USTRUCT()
+struct FPersonaSelectionCapsule
+{
+	GENERATED_BODY()
+
+	FPersonaSelectionCapsule()
+	{
+		Name = NAME_None;
+		Transform = FTransform::Identity;
+		Radius = HalfHeight = 0.f;
+	}
+
+	UPROPERTY()
+	FName Name;
+
+	UPROPERTY()
+	FTransform Transform;
+
+	UPROPERTY()
+	float Radius;
+
+	UPROPERTY()
+	float HalfHeight;
+};
+
+struct PERSONA_API HPersonaSelectionHitProxy : public HHitProxy
+{
+	DECLARE_HIT_PROXY()
+
+	int32 CapsuleIndex;
+	UPersonaSelectionComponent* Component;
+	
+	HPersonaSelectionHitProxy()
+		:CapsuleIndex(INDEX_NONE)
+		,Component(nullptr)
+	{}
+
+	HPersonaSelectionHitProxy(int32 InCapsuleIndex, UPersonaSelectionComponent* InComponent)
+		: HHitProxy(HPP_Foreground)
+		, CapsuleIndex(InCapsuleIndex)
+		, Component(InComponent) 
+	{
+	}
+
+	const FPersonaSelectionCapsule& GetCapsule() const;
+
+	void BroadcastClicked() const;
+	
+	// HHitProxy interface
+	virtual EMouseCursor::Type GetMouseCursor() override { return EMouseCursor::Crosshairs; }
+	// End of HHitProxy interface
+};
+
+// called when a persona selection capsule has been clicked on
+DECLARE_DELEGATE_ThreeParams(FOnPersonaSelectionClicked, UPersonaSelectionComponent*, int32 /*Index*/, const FPersonaSelectionCapsule& /*CapsuleInfo*/);
+
+// called when a persona selection component requires an update for a set of capsules
+DECLARE_DELEGATE_ThreeParams(FOnPersonaSelectionUpdateCapsules, UPersonaSelectionComponent*, const TArray<int32>& /*Indices*/, TArray<FPersonaSelectionCapsule>& /*AllCapsules*/);
+
+UCLASS(Blueprintable, ClassGroup = "Persona")
+class PERSONA_API UPersonaSelectionComponent : public UPrimitiveComponent
+{
+	GENERATED_UCLASS_BODY()
+
+protected:
+
+	//~ Begin UPrimitiveComponent Interface.
+	virtual FPrimitiveSceneProxy* CreateSceneProxy() override;
+	virtual FBoxSphereBounds CalcBounds(const FTransform& LocalToWorld) const override;
+	virtual void GetUsedMaterials( TArray<UMaterialInterface*>& OutMaterials, bool bGetDebugMaterials = false ) const override;
+	virtual void TickComponent(float DeltaTime, enum ELevelTick TickType, FActorComponentTickFunction *ThisTickFunction);
+	//~ End UPrimitiveComponent Interface.
+
+public:
+
+	/** Removes all capsules from the component */
+	void Reset();
+	
+	/** Returns the number of capsules in this selection component */
+	FORCEINLINE int32 Num() const { return Capsules.Num(); }
+
+	FORCEINLINE const FPersonaSelectionCapsule& operator[](int32 InIndex) const { return Capsules[InIndex]; }
+	FORCEINLINE FPersonaSelectionCapsule& operator[](int32 InIndex) { return Capsules[InIndex]; }
+	FORCEINLINE TArray<FPersonaSelectionCapsule>::RangedForIteratorType      begin()       { return Capsules.begin(); }
+	FORCEINLINE TArray<FPersonaSelectionCapsule>::RangedForConstIteratorType begin() const { return Capsules.begin(); }
+	FORCEINLINE TArray<FPersonaSelectionCapsule>::RangedForIteratorType      end()         { return Capsules.end();   }
+	FORCEINLINE TArray<FPersonaSelectionCapsule>::RangedForConstIteratorType end() const   { return Capsules.end();   }
+
+	/** Sets the number of capsules */
+	void SetNum(int32 InCount);
+
+	/** Adds one or more capsules and returns the first added index */
+	int32 Add(int32 InCount = 1);
+
+	/** Marks the capsules as changed and dirties rendering */ 
+	void MarkCapsulesChanged();
+
+	/** Returns the delegate firing when a capsule has been clicked on */
+	FOnPersonaSelectionClicked& OnClicked() { return OnPersonaSelectioneClicked; }
+
+	/** Returns the delegate used to retrieve the transforms for all capsules */
+	FOnPersonaSelectionUpdateCapsules& OnUpdateCapsules() { return OnPersonaSelectionUpdateCapsules; }
+
+	/** Computes a capsule given start end end position of a bone */
+	static void ComputeCapsuleFromBonePositions(
+    	const FVector& InWorldStart,
+    	const FVector& InWorldEnd,
+    	float InBounds,
+    	float InRadiusMultiplier,
+    	FPersonaSelectionCapsule& OutCapsule
+    	);
+
+private:
+
+	void SetCapsulesIndicesToFullArray();
+
+	/** The capsules stored in this component */
+	UPROPERTY()
+	TArray<FPersonaSelectionCapsule> Capsules;
+
+	/** The indices of capsules we want to retrieve */
+	UPROPERTY()
+	TArray<int32> CapsuleIndices;
+
+	/** If set to true the capsule shapes will be visible in the main pass */
+	UPROPERTY()
+	bool bVisibleInMainPass;
+
+	/** Hit proxies for this capsule component */
+	TArray<TRefCountPtr<HPersonaSelectionHitProxy>> HitProxies;
+
+	/** Fired when a capsule is clicked on */
+	FOnPersonaSelectionClicked OnPersonaSelectioneClicked;
+
+	/** Fired when the component needs to update a set of capsules */
+	FOnPersonaSelectionUpdateCapsules OnPersonaSelectionUpdateCapsules;
+
+	friend class FPersonaSelectionComponentProxy;
+};
+
+class FPersonaSelectionComponentProxy : public FPrimitiveSceneProxy
+{
+public:
+
+	SIZE_T GetTypeHash() const override;
+	FPersonaSelectionComponentProxy(const UPersonaSelectionComponent* InComponent);
+
+	virtual void GetDynamicMeshElements(const TArray<const FSceneView*>& Views, const FSceneViewFamily& ViewFamily, uint32 VisibilityMap, FMeshElementCollector& Collector) const override;
+
+	/**
+	*  Returns a struct that describes to the renderer when to draw this proxy.
+	*	@param		Scene view to use to determine our relevence.
+	*  @return		View relevance struct
+	*/
+	virtual FPrimitiveViewRelevance GetViewRelevance(const FSceneView* View) const override;
+	virtual uint32 GetMemoryFootprint(void) const override;
+	uint32 GetAllocatedSize(void) const;
+
+private:
+
+	TWeakObjectPtr<const UPersonaSelectionComponent> SelectionComponentPtr;
+	bool bVisibleInMainPass;
+};
\ No newline at end of file
diff --git a/Engine/Source/Editor/UnrealEd/Classes/Animation/DebugSkelMeshComponent.h b/Engine/Source/Editor/UnrealEd/Classes/Animation/DebugSkelMeshComponent.h
--- a/Engine/Source/Editor/UnrealEd/Classes/Animation/DebugSkelMeshComponent.h
+++ b/Engine/Source/Editor/UnrealEd/Classes/Animation/DebugSkelMeshComponent.h
@@ -127,6 +127,10 @@ public:
 	{
 		return sizeof(*this) + GetAllocatedSize();
 	}
+
+private:
+
+	bool bSelectable;
 };
 
 UCLASS(transient)
diff --git a/Engine/Source/Editor/UnrealEd/Private/Animation/DebugSkelMeshComponent.cpp b/Engine/Source/Editor/UnrealEd/Private/Animation/DebugSkelMeshComponent.cpp
--- a/Engine/Source/Editor/UnrealEd/Private/Animation/DebugSkelMeshComponent.cpp
+++ b/Engine/Source/Editor/UnrealEd/Private/Animation/DebugSkelMeshComponent.cpp
@@ -892,6 +892,7 @@ bool UDebugSkelMeshComponent::IsReferencePoseShown() const
  ***************************************************/
 FDebugSkelMeshSceneProxy::FDebugSkelMeshSceneProxy(const UDebugSkelMeshComponent* InComponent, FSkeletalMeshRenderData* InSkelMeshRenderData, const FColor& InWireframeOverlayColor /*= FColor::White*/) :
 	FSkeletalMeshSceneProxy(InComponent, InSkelMeshRenderData)
+	, bSelectable(false)
 {
 	DynamicData = nullptr;
 	SetWireframeColor(FLinearColor(InWireframeOverlayColor));
@@ -900,6 +901,11 @@ FDebugSkelMeshSceneProxy::FDebugSkelMeshSceneProxy(const UDebugSkelMeshComponent
 	{
 		MaterialRelevance |= GEngine->ClothPaintMaterial->GetRelevance_Concurrent(GetScene().GetFeatureLevel());
 	}
+
+	if(InComponent)
+	{
+		bSelectable = InComponent->bSelectable;
+	}
 }
 
 SIZE_T FDebugSkelMeshSceneProxy::GetTypeHash() const
@@ -912,7 +918,7 @@ void FDebugSkelMeshSceneProxy::GetDynamicMeshElements(const TArray<const FSceneV
 {
 	if(!DynamicData || DynamicData->bDrawMesh)
 	{
-		GetMeshElementsConditionallySelectable(Views, ViewFamily, /*bSelectable=*/true, VisibilityMap, Collector);
+		GetMeshElementsConditionallySelectable(Views, ViewFamily, bSelectable, VisibilityMap, Collector);
 	}
 
 	if(MeshObject && DynamicData && (DynamicData->bDrawNormals || DynamicData->bDrawTangents || DynamicData->bDrawBinormals))
diff --git a/Engine/Source/Runtime/Engine/Private/PrimitiveDrawingUtils.cpp b/Engine/Source/Runtime/Engine/Private/PrimitiveDrawingUtils.cpp
--- a/Engine/Source/Runtime/Engine/Private/PrimitiveDrawingUtils.cpp
+++ b/Engine/Source/Runtime/Engine/Private/PrimitiveDrawingUtils.cpp
@@ -131,7 +131,7 @@ void DrawTriangle(class FPrimitiveDrawInterface* PDI, const FVector& A, const FV
 }
 
 
-void GetBoxMesh(const FMatrix& BoxToWorld,const FVector& Radii,const FMaterialRenderProxy* MaterialRenderProxy,uint8 DepthPriorityGroup,int32 ViewIndex,FMeshElementCollector& Collector)
+void GetBoxMesh(const FMatrix& BoxToWorld,const FVector& Radii,const FMaterialRenderProxy* MaterialRenderProxy,uint8 DepthPriorityGroup,int32 ViewIndex,FMeshElementCollector& Collector, HHitProxy* HitProxy)
 {
 	// Calculate verts for a face pointing down Z
 	FVector Positions[4] =
@@ -181,7 +181,7 @@ void GetBoxMesh(const FMatrix& BoxToWorld,const FVector& Radii,const FMaterialRe
 		MeshBuilder.AddTriangle(VertexIndices[0],VertexIndices[2],VertexIndices[3]);
 	}
 
-	MeshBuilder.GetMesh(FScaleMatrix(Radii) * BoxToWorld,MaterialRenderProxy,DepthPriorityGroup,false,false,ViewIndex,Collector);
+	MeshBuilder.GetMesh(FScaleMatrix(Radii) * BoxToWorld,MaterialRenderProxy,DepthPriorityGroup,false,false,true,ViewIndex,Collector,HitProxy);
 }
 
 void DrawBox(FPrimitiveDrawInterface* PDI,const FMatrix& BoxToWorld,const FVector& Radii,const FMaterialRenderProxy* MaterialRenderProxy,uint8 DepthPriorityGroup)
@@ -649,7 +649,7 @@ void BuildCylinderVerts(const FVector& Base, const FVector& XAxis, const FVector
 
 }
 
-void GetCylinderMesh(const FVector& Start, const FVector& End, float Radius, int32 Sides, const FMaterialRenderProxy* MaterialInstance, uint8 DepthPriority, int32 ViewIndex, FMeshElementCollector& Collector)
+void GetCylinderMesh(const FVector& Start, const FVector& End, float Radius, int32 Sides, const FMaterialRenderProxy* MaterialInstance, uint8 DepthPriority, int32 ViewIndex, FMeshElementCollector& Collector, HHitProxy* HitProxy)
 {
 	FVector Dir = End - Start;
 	float Length = Dir.Size();
@@ -660,18 +660,18 @@ void GetCylinderMesh(const FVector& Start, const FVector& End, float Radius, int
 		FVector X, Y;
 		Z.GetUnsafeNormal().FindBestAxisVectors(X, Y);
 
-		GetCylinderMesh(FMatrix::Identity, Z * Length*0.5 + Start, X, Y, Z, Radius, Length * 0.5f, Sides, MaterialInstance, DepthPriority, ViewIndex, Collector);
+		GetCylinderMesh(FMatrix::Identity, Z * Length*0.5 + Start, X, Y, Z, Radius, Length * 0.5f, Sides, MaterialInstance, DepthPriority, ViewIndex, Collector, HitProxy);
 	}
 
 }
 
 void GetCylinderMesh(const FVector& Base, const FVector& XAxis, const FVector& YAxis, const FVector& ZAxis,
-				  float Radius, float HalfHeight, int32 Sides, const FMaterialRenderProxy* MaterialRenderProxy, uint8 DepthPriority, int32 ViewIndex, FMeshElementCollector& Collector)
+				  float Radius, float HalfHeight, int32 Sides, const FMaterialRenderProxy* MaterialRenderProxy, uint8 DepthPriority, int32 ViewIndex, FMeshElementCollector& Collector, HHitProxy* HitProxy)
 {
-	GetCylinderMesh( FMatrix::Identity, Base, XAxis, YAxis, ZAxis, Radius, HalfHeight, Sides, MaterialRenderProxy, DepthPriority, ViewIndex, Collector );
+	GetCylinderMesh( FMatrix::Identity, Base, XAxis, YAxis, ZAxis, Radius, HalfHeight, Sides, MaterialRenderProxy, DepthPriority, ViewIndex, Collector, HitProxy );
 }
 
-void GetCylinderMesh(const FMatrix& CylToWorld, const FVector& Base, const FVector& XAxis, const FVector& YAxis, const FVector& ZAxis, float Radius, float HalfHeight, uint32 Sides, const FMaterialRenderProxy* MaterialRenderProxy, uint8 DepthPriority, int32 ViewIndex, FMeshElementCollector& Collector)
+void GetCylinderMesh(const FMatrix& CylToWorld, const FVector& Base, const FVector& XAxis, const FVector& YAxis, const FVector& ZAxis, float Radius, float HalfHeight, uint32 Sides, const FMaterialRenderProxy* MaterialRenderProxy, uint8 DepthPriority, int32 ViewIndex, FMeshElementCollector& Collector, HHitProxy* HitProxy)
 {
 	TArray<FDynamicMeshVertex> MeshVerts;
 	TArray<uint32> MeshIndices;
@@ -682,10 +682,10 @@ void GetCylinderMesh(const FMatrix& CylToWorld, const FVector& Base, const FVect
 	MeshBuilder.AddVertices(MeshVerts);
 	MeshBuilder.AddTriangles(MeshIndices);
 
-	MeshBuilder.GetMesh(CylToWorld, MaterialRenderProxy, DepthPriority, false, false, ViewIndex, Collector);
+	MeshBuilder.GetMesh(CylToWorld, MaterialRenderProxy, DepthPriority, false, false, true, ViewIndex, Collector, HitProxy);
 }
 
-void GetConeMesh(const FMatrix& LocalToWorld, float AngleWidth, float AngleHeight, uint32 NumSides, const FMaterialRenderProxy* MaterialRenderProxy, uint8 DepthPriority, int32 ViewIndex, FMeshElementCollector& Collector)
+void GetConeMesh(const FMatrix& LocalToWorld, float AngleWidth, float AngleHeight, uint32 NumSides, const FMaterialRenderProxy* MaterialRenderProxy, uint8 DepthPriority, int32 ViewIndex, FMeshElementCollector& Collector, HHitProxy* HitProxy)
 {
 	TArray<FDynamicMeshVertex> MeshVerts;
 	TArray<uint32> MeshIndices;
@@ -693,10 +693,10 @@ void GetConeMesh(const FMatrix& LocalToWorld, float AngleWidth, float AngleHeigh
 	FDynamicMeshBuilder MeshBuilder(Collector.GetFeatureLevel());
 	MeshBuilder.AddVertices(MeshVerts);
 	MeshBuilder.AddTriangles(MeshIndices);
-	MeshBuilder.GetMesh(LocalToWorld, MaterialRenderProxy, DepthPriority, false, false, ViewIndex, Collector);
+	MeshBuilder.GetMesh(LocalToWorld, MaterialRenderProxy, DepthPriority, false, false, true, ViewIndex, Collector, HitProxy);
 }
 
-void GetCapsuleMesh(const FVector& Origin, const FVector& XAxis, const FVector& YAxis, const FVector& ZAxis, const FLinearColor& Color, float Radius, float HalfHeight, int32 NumSides, const FMaterialRenderProxy* MaterialRenderProxy, uint8 DepthPriority, bool bDisableBackfaceCulling, int32 ViewIndex, FMeshElementCollector& Collector)
+void GetCapsuleMesh(const FVector& Origin, const FVector& XAxis, const FVector& YAxis, const FVector& ZAxis, const FLinearColor& Color, float Radius, float HalfHeight, int32 NumSides, const FMaterialRenderProxy* MaterialRenderProxy, uint8 DepthPriority, bool bDisableBackfaceCulling, int32 ViewIndex, FMeshElementCollector& Collector, HHitProxy* HitProxy)
 {
 	const float HalfAxis = FMath::Max<float>(HalfHeight - Radius, 1.f);
 	const FVector BottomEnd = Origin + Radius * ZAxis;
@@ -704,9 +704,9 @@ void GetCapsuleMesh(const FVector& Origin, const FVector& XAxis, const FVector&
 	const float CylinderHalfHeight = (TopEnd - BottomEnd).Size() * 0.5;
 	const FVector CylinderLocation = BottomEnd + CylinderHalfHeight * ZAxis;
 
-	GetOrientedHalfSphereMesh(TopEnd, FRotationMatrix::MakeFromXY(XAxis, YAxis).Rotator(), FVector(Radius), NumSides, NumSides, 0, PI / 2, MaterialRenderProxy, DepthPriority, bDisableBackfaceCulling, ViewIndex, Collector);
-	GetCylinderMesh(CylinderLocation, XAxis, YAxis, ZAxis, Radius, CylinderHalfHeight, NumSides, MaterialRenderProxy, DepthPriority, ViewIndex, Collector);
-	GetOrientedHalfSphereMesh(BottomEnd, FRotationMatrix::MakeFromXY(XAxis, YAxis).Rotator(), FVector(Radius), NumSides, NumSides, PI / 2, PI, MaterialRenderProxy, DepthPriority, bDisableBackfaceCulling, ViewIndex, Collector);
+	GetOrientedHalfSphereMesh(TopEnd, FRotationMatrix::MakeFromXY(XAxis, YAxis).Rotator(), FVector(Radius), NumSides, NumSides, 0, PI / 2, MaterialRenderProxy, DepthPriority, bDisableBackfaceCulling, ViewIndex, Collector, false, HitProxy);
+	GetCylinderMesh(CylinderLocation, XAxis, YAxis, ZAxis, Radius, CylinderHalfHeight, NumSides, MaterialRenderProxy, DepthPriority, ViewIndex, Collector, HitProxy);
+	GetOrientedHalfSphereMesh(BottomEnd, FRotationMatrix::MakeFromXY(XAxis, YAxis).Rotator(), FVector(Radius), NumSides, NumSides, PI / 2, PI, MaterialRenderProxy, DepthPriority, bDisableBackfaceCulling, ViewIndex, Collector, false, HitProxy);
 }
 
 
diff --git a/Engine/Source/Runtime/Engine/Public/SceneManagement.h b/Engine/Source/Runtime/Engine/Public/SceneManagement.h
--- a/Engine/Source/Runtime/Engine/Public/SceneManagement.h
+++ b/Engine/Source/Runtime/Engine/Public/SceneManagement.h
@@ -2590,7 +2590,7 @@ extern ENGINE_API void DrawCylinder(class FPrimitiveDrawInterface* PDI, const FM
 extern ENGINE_API void DrawCylinder(class FPrimitiveDrawInterface* PDI, const FVector& Start, const FVector& End, float Radius, int32 Sides, const FMaterialRenderProxy* MaterialInstance, uint8 DepthPriority);
 
 
-extern ENGINE_API void GetBoxMesh(const FMatrix& BoxToWorld,const FVector& Radii,const FMaterialRenderProxy* MaterialRenderProxy,uint8 DepthPriority,int32 ViewIndex,FMeshElementCollector& Collector);
+extern ENGINE_API void GetBoxMesh(const FMatrix& BoxToWorld,const FVector& Radii,const FMaterialRenderProxy* MaterialRenderProxy,uint8 DepthPriority,int32 ViewIndex,FMeshElementCollector& Collector, HHitProxy* HitProxy = NULL);
 extern ENGINE_API void GetOrientedHalfSphereMesh(const FVector& Center, const FRotator& Orientation, const FVector& Radii, int32 NumSides, int32 NumRings, float StartAngle, float EndAngle, const FMaterialRenderProxy* MaterialRenderProxy, uint8 DepthPriority, bool bDisableBackfaceCulling,
 									int32 ViewIndex, FMeshElementCollector& Collector, bool bUseSelectionOutline = false, HHitProxy* HitProxy = NULL);
 extern ENGINE_API void GetHalfSphereMesh(const FVector& Center, const FVector& Radii, int32 NumSides, int32 NumRings, float StartAngle, float EndAngle, const FMaterialRenderProxy* MaterialRenderProxy, uint8 DepthPriority, bool bDisableBackfaceCulling,
@@ -2600,17 +2600,17 @@ extern ENGINE_API void GetSphereMesh(const FVector& Center, const FVector& Radii
 extern ENGINE_API void GetSphereMesh(const FVector& Center,const FVector& Radii,int32 NumSides,int32 NumRings,const FMaterialRenderProxy* MaterialRenderProxy,uint8 DepthPriority,
 									bool bDisableBackfaceCulling,int32 ViewIndex,FMeshElementCollector& Collector, bool bUseSelectionOutline, HHitProxy* HitProxy);
 extern ENGINE_API void GetCylinderMesh(const FVector& Base, const FVector& XAxis, const FVector& YAxis, const FVector& ZAxis,
-									float Radius, float HalfHeight, int32 Sides, const FMaterialRenderProxy* MaterialInstance, uint8 DepthPriority, int32 ViewIndex, FMeshElementCollector& Collector);
+									float Radius, float HalfHeight, int32 Sides, const FMaterialRenderProxy* MaterialInstance, uint8 DepthPriority, int32 ViewIndex, FMeshElementCollector& Collector, HHitProxy* HitProxy = NULL);
 extern ENGINE_API void GetCylinderMesh(const FMatrix& CylToWorld, const FVector& Base, const FVector& XAxis, const FVector& YAxis, const FVector& ZAxis,
-									float Radius, float HalfHeight, uint32 Sides, const FMaterialRenderProxy* MaterialInstance, uint8 DepthPriority, int32 ViewIndex, FMeshElementCollector& Collector);
+									float Radius, float HalfHeight, uint32 Sides, const FMaterialRenderProxy* MaterialInstance, uint8 DepthPriority, int32 ViewIndex, FMeshElementCollector& Collector, HHitProxy* HitProxy = NULL);
 //Draws a cylinder along the axis from Start to End
-extern ENGINE_API void GetCylinderMesh(const FVector& Start, const FVector& End, float Radius, int32 Sides, const FMaterialRenderProxy* MaterialInstance, uint8 DepthPriority, int32 ViewIndex, FMeshElementCollector& Collector);
+extern ENGINE_API void GetCylinderMesh(const FVector& Start, const FVector& End, float Radius, int32 Sides, const FMaterialRenderProxy* MaterialInstance, uint8 DepthPriority, int32 ViewIndex, FMeshElementCollector& Collector, HHitProxy* HitProxy = NULL);
 
 
 extern ENGINE_API void GetConeMesh(const FMatrix& LocalToWorld, float AngleWidth, float AngleHeight, uint32 NumSides,
-									const FMaterialRenderProxy* MaterialRenderProxy, uint8 DepthPriority, int32 ViewIndex, FMeshElementCollector& Collector);
+									const FMaterialRenderProxy* MaterialRenderProxy, uint8 DepthPriority, int32 ViewIndex, FMeshElementCollector& Collector, HHitProxy* HitProxy = NULL);
 extern ENGINE_API void GetCapsuleMesh(const FVector& Origin, const FVector& XAxis, const FVector& YAxis, const FVector& ZAxis, const FLinearColor& Color, float Radius, float HalfHeight, int32 NumSides,
-									const FMaterialRenderProxy* MaterialRenderProxy, uint8 DepthPriority, bool bDisableBackfaceCulling, int32 ViewIndex, FMeshElementCollector& Collector);
+									const FMaterialRenderProxy* MaterialRenderProxy, uint8 DepthPriority, bool bDisableBackfaceCulling, int32 ViewIndex, FMeshElementCollector& Collector, HHitProxy* HitProxy = NULL);
 
 
 /**
