Control Rig: fixed a editor ux regression caused by calling modify() unnecessarily, which was marking the bp as dirty on load

diff --git a/Engine/Plugins/Experimental/ControlRig/Source/ControlRigEditor/Private/Editor/ControlRigEditor.cpp b/Engine/Plugins/Experimental/ControlRig/Source/ControlRigEditor/Private/Editor/ControlRigEditor.cpp
--- a/Engine/Plugins/Experimental/ControlRig/Source/ControlRigEditor/Private/Editor/ControlRigEditor.cpp
+++ b/Engine/Plugins/Experimental/ControlRig/Source/ControlRigEditor/Private/Editor/ControlRigEditor.cpp
@@ -1228,11 +1228,12 @@ void FControlRigEditor::HandleSetObjectBeingDebugged(UObject* InObject)
 				}
 			}
 			
+			// get the bone intial transforms from the preview skeletal mesh
 			DebuggedControlRig->SetBoneInitialTransformsFromSkeletalMeshComponent(EditorSkelComp);
 			if(UControlRigBlueprint* RigBlueprint = Cast<UControlRigBlueprint>(GetBlueprintObj()))
 			{
 				// copy the initial transforms back to the blueprint
-				RigBlueprint->Hierarchy->Modify();
+				// no need to call modify here since this code only modifies the bp if the preview mesh changed
 				RigBlueprint->Hierarchy->CopyPose(DebuggedControlRig->GetHierarchy(), false, true);
 			}
 		}
