Control Rig: Improve performance for ResetPoseToInitial

diff --git a/Engine/Plugins/Experimental/ControlRig/Source/ControlRig/Private/Rigs/RigHierarchy.cpp b/Engine/Plugins/Experimental/ControlRig/Source/ControlRig/Private/Rigs/RigHierarchy.cpp
--- a/Engine/Plugins/Experimental/ControlRig/Source/ControlRig/Private/Rigs/RigHierarchy.cpp
+++ b/Engine/Plugins/Experimental/ControlRig/Source/ControlRig/Private/Rigs/RigHierarchy.cpp
@@ -166,6 +166,9 @@ void URigHierarchy::Reset()
 	}
 	IndexLookup.Reset();
 
+	ResetPoseHash = INDEX_NONE;
+	ResetPoseHasFilteredChildren.Reset();
+
 	Notify(ERigHierarchyNotification::HierarchyReset, nullptr);
 }
 
@@ -246,10 +249,69 @@ void URigHierarchy::UpdateSockets(const FRigUnitContext* InContext)
 
 void URigHierarchy::ResetPoseToInitial(ERigElementType InTypeFilter)
 {
-	const bool bAppliesToAllTypes = InTypeFilter == ERigElementType::All;
+	bool bPerformFiltering = InTypeFilter != ERigElementType::All;
+
+	// if we are resetting the pose on some elements, we need to check if
+	// any of affected elements has any children that would not be affected
+	// by resetting the pose. if all children are affected we can use the
+	// fast path.
+	if(bPerformFiltering)
+	{
+		const int32 Hash = HashCombine(GetTopologyVersion(), (int32)InTypeFilter);
+		if(Hash != ResetPoseHash)
+		{
+			ResetPoseHasFilteredChildren.Reset();
+			ResetPoseHash = Hash;
+
+			// let's look at all elements and mark all parent of unaffected children
+			bool bHitAnyParentWithFilteredChildren = false;
+			ResetPoseHasFilteredChildren.AddZeroed(Elements.Num());
+
+			Traverse([this, &bHitAnyParentWithFilteredChildren, InTypeFilter](FRigBaseElement* InElement, bool& bContinue)
+			{
+				bContinue = true;
+
+				const bool bFilteredOut = (!InElement->IsTypeOf(InTypeFilter)) || ResetPoseHasFilteredChildren[InElement->Index];
+				if(bFilteredOut)
+				{
+					const TArray<FRigBaseElement*> Parents = GetParents(InElement);
+					for(const FRigBaseElement* Parent : Parents)
+					{
+						// only mark this up if the parent is not filtered out / 
+						// if we want the parent to reset its pose to initial.
+						if(Parent->IsTypeOf(InTypeFilter))
+						{
+							bHitAnyParentWithFilteredChildren = true;
+						}
+						ResetPoseHasFilteredChildren[Parent->GetIndex()] = true;
+					}
+				}
+
+			}, false);
+
+			if(!bHitAnyParentWithFilteredChildren)
+			{
+				ResetPoseHasFilteredChildren.Reset();
+			}
+		}
+
+		// if the per element state is empty
+		// it means that the filter doesn't affect 
+		if(ResetPoseHasFilteredChildren.IsEmpty())
+		{
+			bPerformFiltering = false;
+		}
+	}
+	
 	for(int32 ElementIndex=0; ElementIndex<Elements.Num(); ElementIndex++)
 	{
-		if(!bAppliesToAllTypes)
+		bool bHasFilteredChildren = bPerformFiltering;
+		if(bHasFilteredChildren)
+		{
+			bHasFilteredChildren = ResetPoseHasFilteredChildren[ElementIndex];
+		}
+		
+		if(bHasFilteredChildren)
 		{
 			if(!Elements[ElementIndex]->IsTypeOf(InTypeFilter))
 			{
@@ -259,7 +321,7 @@ void URigHierarchy::ResetPoseToInitial(ERigElementType InTypeFilter)
 
 		if(FRigControlElement* ControlElement = Cast<FRigControlElement>(Elements[ElementIndex]))
 		{
-			if(!bAppliesToAllTypes)
+			if(bHasFilteredChildren)
 			{
 				const FTransform OffsetTransform = GetControlOffsetTransform(ControlElement, ERigTransformType::InitialLocal);
 				SetControlOffsetTransform(ControlElement, OffsetTransform, ERigTransformType::CurrentLocal, true);
@@ -275,7 +337,7 @@ void URigHierarchy::ResetPoseToInitial(ERigElementType InTypeFilter)
 
 		if(FRigTransformElement* TransformElement = Cast<FRigTransformElement>(Elements[ElementIndex]))
 		{
-			if(!bAppliesToAllTypes)
+			if(bHasFilteredChildren)
 			{
 				const FTransform Transform = GetTransform(TransformElement, ERigTransformType::InitialLocal);
 				SetTransform(TransformElement, Transform, ERigTransformType::CurrentLocal, true);
@@ -288,7 +350,7 @@ void URigHierarchy::ResetPoseToInitial(ERigElementType InTypeFilter)
 
 		if(FRigMultiParentElement* MultiParentElement = Cast<FRigMultiParentElement>(Elements[ElementIndex]))
 		{
-			if(!bAppliesToAllTypes)
+			if(bHasFilteredChildren)
 			{
 				MultiParentElement->Parent.MarkDirty(ERigTransformType::CurrentGlobal);
 			}
diff --git a/Engine/Plugins/Experimental/ControlRig/Source/ControlRig/Public/Rigs/RigHierarchy.h b/Engine/Plugins/Experimental/ControlRig/Source/ControlRig/Public/Rigs/RigHierarchy.h
--- a/Engine/Plugins/Experimental/ControlRig/Source/ControlRig/Public/Rigs/RigHierarchy.h
+++ b/Engine/Plugins/Experimental/ControlRig/Source/ControlRig/Public/Rigs/RigHierarchy.h
@@ -32,6 +32,7 @@ public:
 	, bTransactingForTransformChange(false)
 	, bIsInteracting(false)
 	, bSuspendNotifications(false)
+	, ResetPoseHash(INDEX_NONE)
 	{
 		Reset();
 	}
@@ -2235,6 +2236,9 @@ private:
 
 	TMap<FRigElementKey, FRigElementKey> PreviousParentMap;
 	TMap<FRigElementKey, FRigElementKey> PreviousNameMap;
+
+	int32 ResetPoseHash;
+	TArray<bool> ResetPoseHasFilteredChildren;
 	
 	friend class URigHierarchyController;
 	friend class UControlRig;
