Control Rig: Fix regression for control interactoin in 16508040

diff --git a/Engine/Plugins/Experimental/ControlRig/Source/ControlRig/Private/Rigs/RigHierarchy.cpp b/Engine/Plugins/Experimental/ControlRig/Source/ControlRig/Private/Rigs/RigHierarchy.cpp
--- a/Engine/Plugins/Experimental/ControlRig/Source/ControlRig/Private/Rigs/RigHierarchy.cpp
+++ b/Engine/Plugins/Experimental/ControlRig/Source/ControlRig/Private/Rigs/RigHierarchy.cpp
@@ -311,12 +311,9 @@ void URigHierarchy::ResetPoseToInitial(ERigElementType InTypeFilter)
 			bHasFilteredChildren = ResetPoseHasFilteredChildren[ElementIndex];
 		}
 		
-		if(bHasFilteredChildren)
+		if(!Elements[ElementIndex]->IsTypeOf(InTypeFilter))
 		{
-			if(!Elements[ElementIndex]->IsTypeOf(InTypeFilter))
-			{
-				continue;
-			}
+			continue;
 		}
 
 		if(FRigControlElement* ControlElement = Cast<FRigControlElement>(Elements[ElementIndex]))
