Control Rig: Hide ALL enum entry for element type when within element key

diff --git a/Engine/Plugins/Experimental/ControlRig/Source/ControlRig/Public/Rigs/RigHierarchyDefines.h b/Engine/Plugins/Experimental/ControlRig/Source/ControlRig/Public/Rigs/RigHierarchyDefines.h
--- a/Engine/Plugins/Experimental/ControlRig/Source/ControlRig/Public/Rigs/RigHierarchyDefines.h
+++ b/Engine/Plugins/Experimental/ControlRig/Source/ControlRig/Public/Rigs/RigHierarchyDefines.h
@@ -24,7 +24,7 @@ enum class ERigElementType : uint8
 	Curve = 0x008,
 	RigidBody = 0x010,
 	Socket = 0x020,
-	All = Bone | Null | Control | Curve | RigidBody | Socket UMETA(Hidden),
+	All = Bone | Null | Control | Curve | RigidBody | Socket,
 };
 
 UENUM(BlueprintType)
diff --git a/Engine/Plugins/Experimental/ControlRig/Source/ControlRigEditor/Private/Graph/ControlRigGraphPanelPinFactory.cpp b/Engine/Plugins/Experimental/ControlRig/Source/ControlRigEditor/Private/Graph/ControlRigGraphPanelPinFactory.cpp
--- a/Engine/Plugins/Experimental/ControlRig/Source/ControlRigEditor/Private/Graph/ControlRigGraphPanelPinFactory.cpp
+++ b/Engine/Plugins/Experimental/ControlRig/Source/ControlRigEditor/Private/Graph/ControlRigGraphPanelPinFactory.cpp
@@ -10,6 +10,7 @@
 #include "Graph/SControlRigGraphPinParameterName.h"
 #include "Graph/SControlRigGraphPinVariableBinding.h"
 #include "KismetPins/SGraphPinExec.h"
+#include "SGraphPinComboBox.h"
 #include "ControlRig.h"
 #include "NodeFactory.h"
 #include "EdGraphSchema_K2.h"
@@ -108,8 +109,22 @@ TSharedPtr<SGraphPin> FControlRigGraphPanelPinFactory::CreatePin(UEdGraphPin* In
 	TSharedPtr<SGraphPin> K2PinWidget = FNodeFactory::CreateK2PinWidget(InPin);
 	if(K2PinWidget.IsValid())
 	{
+		// if we are an enum pin - and we are inside a RigElementKey,
+		// let's remove the "all" entry.
+		if(InPin->PinType.PinSubCategoryObject == StaticEnum<ERigElementType>())
+		{
+			if(InPin->ParentPin)
+			{
+				if(InPin->ParentPin->PinType.PinSubCategoryObject == FRigElementKey::StaticStruct())
+				{
+					TSharedPtr<SPinComboBox> EnumCombo = StaticCastSharedRef<SPinComboBox>(K2PinWidget->GetValueWidget());
+					EnumCombo->RemoveItemByIndex(StaticEnum<ERigElementType>()->GetIndexByValue((int64)ERigElementType::All));
+				}
+			}
+		}
+
 		return K2PinWidget;
 	}
 
 	return nullptr;
-}
\ No newline at end of file
+}
diff --git a/Engine/Source/Editor/GraphEditor/Private/KismetPins/SGraphPinEnum.cpp b/Engine/Source/Editor/GraphEditor/Private/KismetPins/SGraphPinEnum.cpp
--- a/Engine/Source/Editor/GraphEditor/Private/KismetPins/SGraphPinEnum.cpp
+++ b/Engine/Source/Editor/GraphEditor/Private/KismetPins/SGraphPinEnum.cpp
@@ -52,6 +52,18 @@ void SPinComboBox::Construct( const FArguments& InArgs )
 	];
 }
 
+void SPinComboBox::RemoveItemByIndex(int32 InIndexToRemove)
+{
+	for(int32 Index=0;Index<ComboItemList.Num();Index++)
+	{
+		if(*ComboItemList[Index] == InIndexToRemove)
+		{
+			ComboItemList.RemoveAt(Index);
+			break;
+		}
+	}
+}
+
 //Function to set the newly selected item
 void SPinComboBox::OnSelectionChangedInternal( TSharedPtr<int32> NewSelection, ESelectInfo::Type SelectInfo )
 {
diff --git a/Engine/Source/Editor/GraphEditor/Public/SGraphPinComboBox.h b/Engine/Source/Editor/GraphEditor/Public/SGraphPinComboBox.h
--- a/Engine/Source/Editor/GraphEditor/Public/SGraphPinComboBox.h
+++ b/Engine/Source/Editor/GraphEditor/Public/SGraphPinComboBox.h
@@ -39,6 +39,9 @@ public:
 	{
 		return CurrentSelection.Pin();
 	}
+
+	void RemoveItemByIndex(int32 InIndexToRemove);
+	
 private:
 
 	//Function to set the newly selected item
