Control Rig: Minor fix to avoid a checkslow() in FCompactPose when calling UControlRig::SetBoneInitialTransformsFromSkeletalMeshComponent with a skeletal mesh component that does not contain any bones

diff --git a/Engine/Plugins/Experimental/ControlRig/Source/ControlRig/Private/ControlRig.cpp b/Engine/Plugins/Experimental/ControlRig/Source/ControlRig/Private/ControlRig.cpp
--- a/Engine/Plugins/Experimental/ControlRig/Source/ControlRig/Private/ControlRig.cpp
+++ b/Engine/Plugins/Experimental/ControlRig/Source/ControlRig/Private/ControlRig.cpp
@@ -2464,7 +2464,7 @@ void UControlRig::SetBoneInitialTransformsFromSkeletalMeshComponent(USkeletalMes
 		FMemMark Mark(FMemStack::Get());
 		FCompactPose OutPose;
 		OutPose.ResetToRefPose(InSkelMeshComp->GetAnimInstance()->GetRequiredBones());
-		if(!OutPose.GetBoneContainer().IsValid())
+		if(!OutPose.IsValid())
 		{
 			return;
 		}
