Control Rig: changed a Check to an Ensure to avoid uncessary crash

diff --git a/Engine/Plugins/Experimental/ControlRig/Source/ControlRigDeveloper/Private/ControlRigBlueprint.cpp b/Engine/Plugins/Experimental/ControlRig/Source/ControlRigDeveloper/Private/ControlRigBlueprint.cpp
--- a/Engine/Plugins/Experimental/ControlRig/Source/ControlRigDeveloper/Private/ControlRigBlueprint.cpp
+++ b/Engine/Plugins/Experimental/ControlRig/Source/ControlRigDeveloper/Private/ControlRigBlueprint.cpp
@@ -438,7 +438,7 @@ void UControlRigBlueprint::PostLoad()
 	GraphsToDetach.Add(GetModel());
 	GraphsToDetach.Add(GetLocalFunctionLibrary());
 
-	check(IsInGameThread());
+	if (ensure(IsInGameThread()))
 	{
 		for (URigVMGraph* GraphToDetach : GraphsToDetach)
 		{
