fix for unititialized struct member warnings
#rnx

diff --git a/Engine/Plugins/Experimental/ControlRig/Source/ControlRig/Private/Units/Execution/RigUnit_Hierarchy.h b/Engine/Plugins/Experimental/ControlRig/Source/ControlRig/Private/Units/Execution/RigUnit_Hierarchy.h
--- a/Engine/Plugins/Experimental/ControlRig/Source/ControlRig/Private/Units/Execution/RigUnit_Hierarchy.h
+++ b/Engine/Plugins/Experimental/ControlRig/Source/ControlRig/Private/Units/Execution/RigUnit_Hierarchy.h
@@ -303,6 +303,7 @@ struct CONTROLRIG_API FRigUnit_PoseGetDelta : public FRigUnit_HierarchyBase
 	{
 		PoseA = PoseB = FRigPose();
 		ElementType = ERigElementType::All;
+		Space = EBoneGetterSetterMode::LocalSpace;
 		ItemsToCompare = ItemsWithDelta = FRigElementKeyCollection();
 		PositionThreshold = 0.1f;
 		RotationThreshold = ScaleThreshold = CurveThreshold = 0.f;
@@ -367,6 +368,7 @@ struct CONTROLRIG_API FRigUnit_PoseGetTransform : public FRigUnit_HierarchyBase
 		Space = EBoneGetterSetterMode::GlobalSpace;
 		Valid = false;
 		Transform = FTransform::Identity;
+		CurveValue = 0.f;
 		CachedPoseElementIndex = CachedPoseHash = INDEX_NONE;
 	}
 
