Control Rig: Fix code analysis null pointer validity checking

diff --git a/Engine/Plugins/Experimental/ControlRig/Source/ControlRigEditor/Private/Graph/ControlRigGraphPanelPinFactory.cpp b/Engine/Plugins/Experimental/ControlRig/Source/ControlRigEditor/Private/Graph/ControlRigGraphPanelPinFactory.cpp
--- a/Engine/Plugins/Experimental/ControlRig/Source/ControlRigEditor/Private/Graph/ControlRigGraphPanelPinFactory.cpp
+++ b/Engine/Plugins/Experimental/ControlRig/Source/ControlRigEditor/Private/Graph/ControlRigGraphPanelPinFactory.cpp
@@ -104,26 +104,26 @@ TSharedPtr<SGraphPin> FControlRigGraphPanelPinFactory::CreatePin(UEdGraphPin* In
 				}
 			}
 		}
-	}
 
-	TSharedPtr<SGraphPin> K2PinWidget = FNodeFactory::CreateK2PinWidget(InPin);
-	if(K2PinWidget.IsValid())
-	{
-		// if we are an enum pin - and we are inside a RigElementKey,
-		// let's remove the "all" entry.
-		if(InPin->PinType.PinSubCategoryObject == StaticEnum<ERigElementType>())
+		TSharedPtr<SGraphPin> K2PinWidget = FNodeFactory::CreateK2PinWidget(InPin);
+		if(K2PinWidget.IsValid())
 		{
-			if(InPin->ParentPin)
+			// if we are an enum pin - and we are inside a RigElementKey,
+			// let's remove the "all" entry.
+			if(InPin->PinType.PinSubCategoryObject == StaticEnum<ERigElementType>())
 			{
-				if(InPin->ParentPin->PinType.PinSubCategoryObject == FRigElementKey::StaticStruct())
+				if(InPin->ParentPin)
 				{
-					TSharedPtr<SPinComboBox> EnumCombo = StaticCastSharedRef<SPinComboBox>(K2PinWidget->GetValueWidget());
-					EnumCombo->RemoveItemByIndex(StaticEnum<ERigElementType>()->GetIndexByValue((int64)ERigElementType::All));
+					if(InPin->ParentPin->PinType.PinSubCategoryObject == FRigElementKey::StaticStruct())
+					{
+						TSharedPtr<SPinComboBox> EnumCombo = StaticCastSharedRef<SPinComboBox>(K2PinWidget->GetValueWidget());
+						EnumCombo->RemoveItemByIndex(StaticEnum<ERigElementType>()->GetIndexByValue((int64)ERigElementType::All));
+					}
 				}
 			}
-		}
 
-		return K2PinWidget;
+			return K2PinWidget;
+		}
 	}
 
 	return nullptr;
