Control Rig: Perform name sanitization for functions

diff --git a/Engine/Plugins/Experimental/ControlRig/Source/ControlRigDeveloper/Private/Graph/ControlRigGraphNode.cpp b/Engine/Plugins/Experimental/ControlRig/Source/ControlRigDeveloper/Private/Graph/ControlRigGraphNode.cpp
--- a/Engine/Plugins/Experimental/ControlRig/Source/ControlRigDeveloper/Private/Graph/ControlRigGraphNode.cpp
+++ b/Engine/Plugins/Experimental/ControlRig/Source/ControlRigDeveloper/Private/Graph/ControlRigGraphNode.cpp
@@ -67,15 +67,18 @@ FText UControlRigGraphNode::GetNodeTitle(ENodeTitleType::Type TitleType) const
 
 			else if(URigVMFunctionReferenceNode* FunctionReferenceNode = Cast<URigVMFunctionReferenceNode>(ModelNode))
 			{
-				UPackage* ReferencedPackage = FunctionReferenceNode->GetReferencedNode()->GetOutermost();
-				if(ReferencedPackage != ModelNode->GetOutermost())
+				if(URigVMLibraryNode* ReferencedNode = FunctionReferenceNode->GetReferencedNode())
 				{
-					SubTitle = FString::Printf(TEXT("From %s"), *ReferencedPackage->GetName());
-				}
-				else
-				{
-					static const FString LocalFunctionString = TEXT("Local Function");
-					SubTitle = LocalFunctionString;
+					UPackage* ReferencedPackage = ReferencedNode->GetOutermost();
+					if(ReferencedPackage != ModelNode->GetOutermost())
+					{
+						SubTitle = FString::Printf(TEXT("From %s"), *ReferencedPackage->GetName());
+					}
+					else
+					{
+						static const FString LocalFunctionString = TEXT("Local Function");
+						SubTitle = LocalFunctionString;
+					}
 				}
 			}
 
diff --git a/Engine/Plugins/Experimental/ControlRig/Source/ControlRigEditor/Private/Graph/ControlRigGraphPanelPinFactory.cpp b/Engine/Plugins/Experimental/ControlRig/Source/ControlRigEditor/Private/Graph/ControlRigGraphPanelPinFactory.cpp
--- a/Engine/Plugins/Experimental/ControlRig/Source/ControlRigEditor/Private/Graph/ControlRigGraphPanelPinFactory.cpp
+++ b/Engine/Plugins/Experimental/ControlRig/Source/ControlRigEditor/Private/Graph/ControlRigGraphPanelPinFactory.cpp
@@ -108,16 +108,28 @@ TSharedPtr<SGraphPin> FControlRigGraphPanelPinFactory::CreatePin(UEdGraphPin* In
 		TSharedPtr<SGraphPin> K2PinWidget = FNodeFactory::CreateK2PinWidget(InPin);
 		if(K2PinWidget.IsValid())
 		{
-			// if we are an enum pin - and we are inside a RigElementKey,
-			// let's remove the "all" entry.
-			if(InPin->PinType.PinSubCategoryObject == StaticEnum<ERigElementType>())
+			if(InPin->Direction == EEdGraphPinDirection::EGPD_Input)
 			{
-				if(InPin->ParentPin)
+				// if we are an enum pin - and we are inside a RigElementKey,
+				// let's remove the "all" entry.
+				if(InPin->PinType.PinSubCategoryObject == StaticEnum<ERigElementType>())
 				{
-					if(InPin->ParentPin->PinType.PinSubCategoryObject == FRigElementKey::StaticStruct())
+					if(InPin->ParentPin)
 					{
-						TSharedPtr<SPinComboBox> EnumCombo = StaticCastSharedRef<SPinComboBox>(K2PinWidget->GetValueWidget());
-						EnumCombo->RemoveItemByIndex(StaticEnum<ERigElementType>()->GetIndexByValue((int64)ERigElementType::All));
+						if(InPin->ParentPin->PinType.PinSubCategoryObject == FRigElementKey::StaticStruct())
+						{
+							TSharedPtr<SWidget> ValueWidget = K2PinWidget->GetValueWidget();
+							if(ValueWidget.IsValid())
+							{
+								if(TSharedPtr<SPinComboBox> EnumCombo = StaticCastSharedPtr<SPinComboBox>(ValueWidget))
+								{
+									if(EnumCombo.IsValid())
+									{
+										EnumCombo->RemoveItemByIndex(StaticEnum<ERigElementType>()->GetIndexByValue((int64)ERigElementType::All));
+									}
+								}
+							}
+						}
 					}
 				}
 			}
diff --git a/Engine/Source/Developer/RigVMDeveloper/Private/RigVMModel/RigVMController.cpp b/Engine/Source/Developer/RigVMDeveloper/Private/RigVMModel/RigVMController.cpp
--- a/Engine/Source/Developer/RigVMDeveloper/Private/RigVMModel/RigVMController.cpp
+++ b/Engine/Source/Developer/RigVMDeveloper/Private/RigVMModel/RigVMController.cpp
@@ -7465,7 +7465,10 @@ FString URigVMController::GetValidNodeName(const FString& InPrefix)
 	URigVMGraph* Graph = GetGraph();
 	check(Graph);
 
-	return GetUniqueName(*InPrefix, [&](const FName& InName) {
+	FString SanitizedPrefix = InPrefix;
+	SanitizeName(SanitizedPrefix);
+
+	return GetUniqueName(*SanitizedPrefix, [&](const FName& InName) {
 		return Graph->IsNameAvailable(InName.ToString());
 	}).ToString();
 }
@@ -9504,3 +9507,27 @@ void URigVMController::RefreshFunctionReferences(URigVMLibraryNode* InFunctionDe
 		});
 	}
 }
+
+void URigVMController::SanitizeName(FString& InOutName)
+{
+	// Sanitize the name
+	for (int32 i = 0; i < InOutName.Len(); ++i)
+	{
+		TCHAR& C = InOutName[i];
+
+		const bool bGoodChar =
+			((C >= 'A') && (C <= 'Z')) || ((C >= 'a') && (C <= 'z')) ||		// A-Z (upper and lowercase) anytime
+			(C == '_') || (C == '-') || (C == '.') || (C == ' ') ||			// _  - . and space anytime
+			((i > 0) && (C >= '0') && (C <= '9'));							// 0-9 after the first character
+
+		if (!bGoodChar)
+		{
+			C = '_';
+		}
+	}
+
+	if (InOutName.Len() > GetMaxNameLength())
+	{
+		InOutName.LeftChopInline(InOutName.Len() - GetMaxNameLength());
+	}
+}
diff --git a/Engine/Source/Developer/RigVMDeveloper/Public/RigVMModel/RigVMController.h b/Engine/Source/Developer/RigVMDeveloper/Public/RigVMModel/RigVMController.h
--- a/Engine/Source/Developer/RigVMDeveloper/Public/RigVMModel/RigVMController.h
+++ b/Engine/Source/Developer/RigVMDeveloper/Public/RigVMModel/RigVMController.h
@@ -775,6 +775,9 @@ private:
 
 	void RefreshFunctionReferences(URigVMLibraryNode* InFunctionDefinition, bool bSetupUndoRedo);
 
+	static int32 GetMaxNameLength() { return 100; }
+	static void SanitizeName(FString& InOutName);
+	
 	UPROPERTY(transient)
 	TArray<URigVMGraph*> Graphs;
 
