Control Rig: Minimum set of nodes to interact with matrices

diff --git a/Engine/Plugins/Experimental/ControlRig/Source/ControlRig/Private/Units/Math/RigUnit_MathMatrix.cpp b/Engine/Plugins/Experimental/ControlRig/Source/ControlRig/Private/Units/Math/RigUnit_MathMatrix.cpp
new file mode 100644
--- /dev/null
+++ b/Engine/Plugins/Experimental/ControlRig/Source/ControlRig/Private/Units/Math/RigUnit_MathMatrix.cpp
@@ -0,0 +1,45 @@
+// Copyright Epic Games, Inc. All Rights Reserved.
+
+#include "Units/Math/RigUnit_MathMatrix.h"
+#include "Units/Math/RigUnit_MathTransform.h"
+#include "Units/RigUnitContext.h"
+
+FRigUnit_MathMatrixFromTransform_Execute()
+{
+    DECLARE_SCOPE_HIERARCHICAL_COUNTER_RIGUNIT()
+	Result = Transform.ToMatrixWithScale();
+}
+
+FRigUnit_MathMatrixToTransform_Execute()
+{
+    DECLARE_SCOPE_HIERARCHICAL_COUNTER_RIGUNIT()
+	Result.SetFromMatrix(Value);
+}
+
+FRigUnit_MathMatrixFromVectors_Execute()
+{
+	DECLARE_SCOPE_HIERARCHICAL_COUNTER_RIGUNIT()
+	Result = FMatrix(X, Y, Z, FVector::ZeroVector);
+	Result.SetOrigin(Origin);
+}
+
+FRigUnit_MathMatrixToVectors_Execute()
+{
+	DECLARE_SCOPE_HIERARCHICAL_COUNTER_RIGUNIT()
+	X = Value.GetScaledAxis(EAxis::X);
+	Y = Value.GetScaledAxis(EAxis::Y);
+	Z = Value.GetScaledAxis(EAxis::Z);
+	Origin = Value.GetOrigin();
+}
+
+FRigUnit_MathMatrixMul_Execute()
+{
+    DECLARE_SCOPE_HIERARCHICAL_COUNTER_RIGUNIT()
+	Result = A * B;
+}
+
+FRigUnit_MathMatrixInverse_Execute()
+{
+    DECLARE_SCOPE_HIERARCHICAL_COUNTER_RIGUNIT()
+	Result = Value.Inverse();
+}
diff --git a/Engine/Plugins/Experimental/ControlRig/Source/ControlRig/Private/Units/Math/RigUnit_MathMatrix.h b/Engine/Plugins/Experimental/ControlRig/Source/ControlRig/Private/Units/Math/RigUnit_MathMatrix.h
new file mode 100644
--- /dev/null
+++ b/Engine/Plugins/Experimental/ControlRig/Source/ControlRig/Private/Units/Math/RigUnit_MathMatrix.h
@@ -0,0 +1,193 @@
+// Copyright Epic Games, Inc. All Rights Reserved.
+
+#pragma once
+
+#include "EulerTransform.h"
+#include "Units/Math/RigUnit_MathBase.h"
+#include "Math/ControlRigMathLibrary.h"
+#include "RigUnit_MathMatrix.generated.h"
+
+USTRUCT(meta=(Abstract, Category="Math|Matrix", MenuDescSuffix="(Matrix)"))
+struct CONTROLRIG_API FRigUnit_MathMatrixBase : public FRigUnit_MathBase
+{
+	GENERATED_BODY()
+};
+
+USTRUCT(meta=(Abstract))
+struct CONTROLRIG_API FRigUnit_MathMatrixUnaryOp : public FRigUnit_MathMatrixBase
+{
+	GENERATED_BODY()
+
+	FRigUnit_MathMatrixUnaryOp()
+	{
+		Value = Result = FMatrix::Identity;
+	}
+
+	UPROPERTY(meta=(Input))
+	FMatrix Value;
+
+	UPROPERTY(meta=(Output))
+	FMatrix Result;
+};
+
+USTRUCT(meta=(Abstract))
+struct CONTROLRIG_API FRigUnit_MathMatrixBinaryOp : public FRigUnit_MathMatrixBase
+{
+	GENERATED_BODY()
+
+	FRigUnit_MathMatrixBinaryOp()
+	{
+		A = B = Result = FMatrix::Identity;
+	}
+
+	UPROPERTY(meta=(Input))
+	FMatrix A;
+
+	UPROPERTY(meta=(Input))
+	FMatrix B;
+
+	UPROPERTY(meta=(Output))
+	FMatrix Result;
+};
+
+/**
+* Makes a transform from a matrix
+*/
+USTRUCT(meta=(DisplayName="To Transform", Keywords="Make,Construct"))
+struct CONTROLRIG_API FRigUnit_MathMatrixToTransform : public FRigUnit_MathMatrixBase
+{
+	GENERATED_BODY()
+
+	FRigUnit_MathMatrixToTransform()
+	{
+		Value = FMatrix::Identity;
+		Result = FTransform::Identity;
+	}
+
+	RIGVM_METHOD()
+	virtual void Execute(const FRigUnitContext& Context) override;
+
+	UPROPERTY(meta=(Input))
+	FMatrix Value;
+
+	UPROPERTY(meta=(Output))
+	FTransform Result;
+};
+
+/**
+ * Makes a matrix from a transform
+ */
+USTRUCT(meta=(DisplayName="From Transform", Keywords="Make,Construct"))
+struct CONTROLRIG_API FRigUnit_MathMatrixFromTransform : public FRigUnit_MathMatrixBase
+{
+	GENERATED_BODY()
+
+	FRigUnit_MathMatrixFromTransform()
+	{
+		Transform = FTransform::Identity;
+		Result = FMatrix::Identity;
+	}
+
+	RIGVM_METHOD()
+	virtual void Execute(const FRigUnitContext& Context) override;
+
+	UPROPERTY(meta=(Input))
+	FTransform Transform;
+
+	UPROPERTY(meta=(Output))
+	FMatrix Result;
+};
+
+/**
+ * Converts the matrix to its vectors
+ */
+USTRUCT(meta=(DisplayName="To Vectors", Keywords="Make,Construct"))
+struct CONTROLRIG_API FRigUnit_MathMatrixToVectors : public FRigUnit_MathMatrixBase
+{
+	GENERATED_BODY()
+
+	FRigUnit_MathMatrixToVectors()
+	{
+		Value = FMatrix::Identity;
+		Origin = FVector::ZeroVector;
+		X = FVector::XAxisVector;
+		Y = FVector::YAxisVector;
+		Z = FVector::ZAxisVector;
+	}
+
+	RIGVM_METHOD()
+	virtual void Execute(const FRigUnitContext& Context) override;
+
+	UPROPERTY(meta=(Input))
+	FMatrix Value;
+
+	UPROPERTY(meta=(Output))
+	FVector Origin;
+
+	UPROPERTY(meta=(Output))
+	FVector X;
+
+	UPROPERTY(meta=(Output))
+	FVector Y;
+
+	UPROPERTY(meta=(Output))
+	FVector Z;
+};
+
+/**
+* Makes a matrix from its vectors
+*/
+USTRUCT(meta=(DisplayName="From Vectors", Keywords="Make,Construct"))
+struct CONTROLRIG_API FRigUnit_MathMatrixFromVectors : public FRigUnit_MathMatrixBase
+{
+	GENERATED_BODY()
+
+	FRigUnit_MathMatrixFromVectors()
+	{
+		Origin = FVector::ZeroVector;
+		X = FVector::XAxisVector;
+		Y = FVector::YAxisVector;
+		Z = FVector::ZAxisVector;
+		Result = FMatrix::Identity;
+	}
+
+	RIGVM_METHOD()
+	virtual void Execute(const FRigUnitContext& Context) override;
+
+	UPROPERTY(meta=(Input))
+	FVector Origin;
+
+	UPROPERTY(meta=(Input))
+	FVector X;
+
+	UPROPERTY(meta=(Input))
+	FVector Y;
+
+	UPROPERTY(meta=(Input))
+	FVector Z;
+
+	UPROPERTY(meta=(Output))
+	FMatrix Result;
+};
+
+/**
+ * Returns the product of the two values
+ */
+USTRUCT(meta=(DisplayName="Multiply", PrototypeName="Multiply", Keywords="Product,*,Global"))
+struct CONTROLRIG_API FRigUnit_MathMatrixMul : public FRigUnit_MathMatrixBinaryOp
+{
+	GENERATED_BODY()
+	RIGVM_METHOD()
+	virtual void Execute(const FRigUnitContext& Context) override;
+};
+
+/**
+ * Returns the inverse value
+ */
+USTRUCT(meta=(DisplayName="Inverse", PrototypeName="Inverse"))
+struct CONTROLRIG_API FRigUnit_MathMatrixInverse : public FRigUnit_MathMatrixUnaryOp
+{
+	GENERATED_BODY()
+	RIGVM_METHOD()
+	virtual void Execute(const FRigUnitContext& Context) override;
+};
