Control Rig: Add workflow for "Collapse To Function"

diff --git a/Engine/Plugins/Experimental/ControlRig/Source/ControlRigEditor/Private/ControlRigEditorModule.cpp b/Engine/Plugins/Experimental/ControlRig/Source/ControlRigEditor/Private/ControlRigEditorModule.cpp
--- a/Engine/Plugins/Experimental/ControlRig/Source/ControlRigEditor/Private/ControlRigEditorModule.cpp
+++ b/Engine/Plugins/Experimental/ControlRig/Source/ControlRigEditor/Private/ControlRigEditorModule.cpp
@@ -1946,6 +1946,19 @@ void FControlRigEditorModule::GetContextMenuActions(const UControlRigGraphSchema
 							Controller->CollapseNodes(Nodes);
 						})
 					));
+					OrganizationSection.AddMenuEntry(
+						"Collapse to Function",
+						LOCTEXT("CollapseNodesToFunction", "Collapse to Function"),
+						LOCTEXT("CollapseNodesToFunction_Tooltip", "Turns the selected nodes into a new Function"),
+						FSlateIcon(),
+						FUIAction(FExecuteAction::CreateLambda([Model, Controller]() {
+							TArray<FName> Nodes = Model->GetSelectNodes();
+							Controller->OpenUndoBracket(TEXT("Collapse to Function"));
+							URigVMCollapseNode* CollapseNode = Controller->CollapseNodes(Nodes, TEXT("New Function"));
+							Controller->PromoteCollapseNodeToFunctionReferenceNode(CollapseNode->GetFName());
+							Controller->CloseUndoBracket();
+						})
+					));
 
 					if (URigVMCollapseNode* CollapseNode = Cast<URigVMCollapseNode>(RigNode->GetModelNode()))
 					{
