Control Rig: Contraint Nodes, added an overall weight pin + changed pin order slightly

diff --git a/Engine/Plugins/Experimental/ControlRig/Source/ControlRig/Private/Units/Highlevel/Hierarchy/RigUnit_TransformConstraint.cpp b/Engine/Plugins/Experimental/ControlRig/Source/ControlRig/Private/Units/Highlevel/Hierarchy/RigUnit_TransformConstraint.cpp
--- a/Engine/Plugins/Experimental/ControlRig/Source/ControlRig/Private/Units/Highlevel/Hierarchy/RigUnit_TransformConstraint.cpp
+++ b/Engine/Plugins/Experimental/ControlRig/Source/ControlRig/Private/Units/Highlevel/Hierarchy/RigUnit_TransformConstraint.cpp
@@ -167,6 +167,11 @@ void FRigUnit_TransformConstraintPerItem::AddConstraintData(const FRigVMFixedArr
 
 FRigUnit_ParentConstraint_Execute()
  {
+	if (Weight < KINDA_SMALL_NUMBER)
+	{
+		return;
+	}
+	
  	URigHierarchy* Hierarchy = ExecuteContext.Hierarchy;	
 	if (Hierarchy)
 	{
@@ -322,14 +327,19 @@ FRigUnit_ParentConstraint_Execute()
 			FTransform FilteredMixedLocalTransform(FControlRigMathLibrary::QuatFromEuler(FilteredEulerRotation, AdvancedSettings.RotationOrderForFilter), FilteredTranslation, FilteredScale);
 
 			FTransform FinalLocalTransform = FilteredMixedLocalTransform;
+
+			if (Weight < 1.0f - KINDA_SMALL_NUMBER)
+			{
+				FinalLocalTransform = FControlRigMathLibrary::LerpTransform(ChildCurrentLocalTransform, FinalLocalTransform, Weight);
+			}
 			
 			if (Child.Type == ERigElementType::Control)
 			{
 				// need to convert back to offset space for the actual control value
-				FinalLocalTransform = FilteredMixedLocalTransform.GetRelativeTransform(AdditionalOffsetTransform);
+				FinalLocalTransform = FinalLocalTransform.GetRelativeTransform(AdditionalOffsetTransform);
 				FinalLocalTransform.NormalizeRotation();
 			}
-			
+
 			Hierarchy->SetLocalTransform(Child, FinalLocalTransform);
 		}
 	} 
@@ -408,6 +418,11 @@ FRigUnit_ParentConstraint_Execute()
 
 FRigUnit_PositionConstraint_Execute()
 {
+	if (Weight < KINDA_SMALL_NUMBER)
+	{
+		return;
+	}
+	
 	URigHierarchy* Hierarchy = ExecuteContext.Hierarchy;	
 	if (Hierarchy)
 	{
@@ -503,11 +518,16 @@ FRigUnit_PositionConstraint_Execute()
 			FilteredMixedLocalTransform.SetTranslation(FilteredPosition); 
 
 			FTransform FinalLocalTransform = FilteredMixedLocalTransform;
-
+			
+			if (Weight < 1.0f - KINDA_SMALL_NUMBER)
+			{
+				FinalLocalTransform = FControlRigMathLibrary::LerpTransform(ChildCurrentLocalTransform, FinalLocalTransform, Weight);
+			}
+			
 			if (Child.Type == ERigElementType::Control)
 			{
 				// need to convert back to offset space for the actual control value
-				FinalLocalTransform = FilteredMixedLocalTransform.GetRelativeTransform(AdditionalOffsetTransform);
+				FinalLocalTransform = FinalLocalTransform.GetRelativeTransform(AdditionalOffsetTransform);
 				FinalLocalTransform.NormalizeRotation();
 			}
 
@@ -557,6 +577,11 @@ FRigUnit_PositionConstraint_Execute()
 
 FRigUnit_RotationConstraint_Execute()
 {
+	if (Weight < KINDA_SMALL_NUMBER)
+	{
+		return;
+	}
+	
 	URigHierarchy* Hierarchy = ExecuteContext.Hierarchy;	
 	if (Hierarchy)
 	{
@@ -693,14 +718,19 @@ FRigUnit_RotationConstraint_Execute()
 
 			FTransform FinalLocalTransform = FilteredMixedLocalTransform;
 
+			if (Weight < 1.0f - KINDA_SMALL_NUMBER)
+			{
+				FinalLocalTransform = FControlRigMathLibrary::LerpTransform(ChildCurrentLocalTransform, FinalLocalTransform, Weight);
+			}
+
 			if (Child.Type == ERigElementType::Control)
 			{
 				// need to convert back to offset space for the actual control value
-				FinalLocalTransform = FilteredMixedLocalTransform.GetRelativeTransform(AdditionalOffsetTransform);
+				FinalLocalTransform = FinalLocalTransform.GetRelativeTransform(AdditionalOffsetTransform);
 				FinalLocalTransform.NormalizeRotation();
 			}
-
-			Hierarchy->SetLocalTransform(Child, FinalLocalTransform); 
+			
+			Hierarchy->SetLocalTransform(Child, FinalLocalTransform);
 		}
 	} 	
 }
@@ -806,6 +836,11 @@ FRigUnit_ScaleConstraint_Execute()
        	
        	return NonZeroScale;	
 	});
+
+	if (Weight < KINDA_SMALL_NUMBER)
+	{
+		return;
+	}
 	
 	URigHierarchy* Hierarchy = ExecuteContext.Hierarchy;	
 	if (Hierarchy)
@@ -912,14 +947,19 @@ FRigUnit_ScaleConstraint_Execute()
 
 			FTransform FinalLocalTransform = FilteredMixedLocalTransform;
 
+			if (Weight < 1.0f - KINDA_SMALL_NUMBER)
+			{
+				FinalLocalTransform = FControlRigMathLibrary::LerpTransform(ChildCurrentLocalTransform, FinalLocalTransform, Weight);
+			}
+			
 			if (Child.Type == ERigElementType::Control)
 			{
 				// need to convert back to offset space for the actual control value
-				FinalLocalTransform = FilteredMixedLocalTransform.GetRelativeTransform(AdditionalOffsetTransform);
+				FinalLocalTransform = FinalLocalTransform.GetRelativeTransform(AdditionalOffsetTransform);
 				FinalLocalTransform.NormalizeRotation();
 			}
 
-			Hierarchy->SetLocalTransform(Child, FinalLocalTransform); 
+			Hierarchy->SetLocalTransform(Child, FinalLocalTransform);
 		}
 	}
 }
diff --git a/Engine/Plugins/Experimental/ControlRig/Source/ControlRig/Private/Units/Highlevel/Hierarchy/RigUnit_TransformConstraint.h b/Engine/Plugins/Experimental/ControlRig/Source/ControlRig/Private/Units/Highlevel/Hierarchy/RigUnit_TransformConstraint.h
--- a/Engine/Plugins/Experimental/ControlRig/Source/ControlRig/Private/Units/Highlevel/Hierarchy/RigUnit_TransformConstraint.h
+++ b/Engine/Plugins/Experimental/ControlRig/Source/ControlRig/Private/Units/Highlevel/Hierarchy/RigUnit_TransformConstraint.h
@@ -209,6 +209,7 @@ enum class EConstraintInterpType : uint8
 	Average UMETA(DisplayName="Average(Order Independent)"),
    /** Perform Quaternion Slerp in Sequence, Different Orders of Parents can Produce Different Results */ 
     Shortest UMETA(DisplayName="Shortest(Order Dependent)"),
+	
     Max UMETA(Hidden),
 };
 
@@ -244,9 +245,11 @@ struct FRigUnit_ParentConstraint : public FRigUnit_HighlevelBaseMutable
 	GENERATED_BODY()
 
 	FRigUnit_ParentConstraint()
-        :Child(FRigElementKey(NAME_None, ERigElementType::Bone))
+        : Child(FRigElementKey(NAME_None, ERigElementType::Bone))
         , bMaintainOffset(true)
+		, Weight(1.0f)
 	{
+		Parents.Add(FConstraintParent());
 	}
 	
 	RIGVM_METHOD()
@@ -261,11 +264,14 @@ struct FRigUnit_ParentConstraint : public FRigUnit_HighlevelBaseMutable
 	UPROPERTY(EditAnywhere, Category = "Constraint", meta = (Input)) 
 	FTransformFilter Filter;
 	
-	UPROPERTY(EditAnywhere, Category = "Constraint", meta = (Input)) 
-	FRigUnit_ParentConstraint_AdvancedSettings AdvancedSettings;
-	
 	UPROPERTY(EditAnywhere, Category = "Constraint", meta = (Input, ExpandByDefault, DefaultArraySize = 1))
 	TArray<FConstraintParent> Parents;
+	
+	UPROPERTY(EditAnywhere, Category = "Constraint", meta = (Input)) 
+	FRigUnit_ParentConstraint_AdvancedSettings AdvancedSettings;
+
+	UPROPERTY(EditAnywhere, Category = "Constraint", meta = (Input))
+	float Weight;	
 };
 
 /**
@@ -277,9 +283,11 @@ struct FRigUnit_PositionConstraint : public FRigUnit_HighlevelBaseMutable
 	GENERATED_BODY()
 
 	FRigUnit_PositionConstraint()
-        :Child(FRigElementKey(NAME_None, ERigElementType::Bone))
+        : Child(FRigElementKey(NAME_None, ERigElementType::Bone))
         , bMaintainOffset(true)
+		, Weight(1.0f)
 	{
+		Parents.Add(FConstraintParent());
 	}
 	
 	RIGVM_METHOD()
@@ -296,6 +304,9 @@ struct FRigUnit_PositionConstraint : public FRigUnit_HighlevelBaseMutable
 
 	UPROPERTY(EditAnywhere, Category = "Constraint", meta = (Input, ExpandByDefault, DefaultArraySize = 1))
 	TArray<FConstraintParent> Parents;
+
+	UPROPERTY(EditAnywhere, Category = "Constraint", meta = (Input))
+	float Weight;
 };
 
 USTRUCT()
@@ -330,9 +341,11 @@ struct FRigUnit_RotationConstraint : public FRigUnit_HighlevelBaseMutable
 	GENERATED_BODY()
 
 	FRigUnit_RotationConstraint()
-        :Child(FRigElementKey(NAME_None, ERigElementType::Bone))
+        : Child(FRigElementKey(NAME_None, ERigElementType::Bone))
         , bMaintainOffset(true)
+		, Weight(1.0f)
 	{
+		Parents.Add(FConstraintParent());
 	}
 	
 	RIGVM_METHOD()
@@ -347,11 +360,14 @@ struct FRigUnit_RotationConstraint : public FRigUnit_HighlevelBaseMutable
 	UPROPERTY(EditAnywhere, Category = "Constraint", meta = (Input)) 
 	FFilterOptionPerAxis Filter;
 
+	UPROPERTY(EditAnywhere, Category = "Constraint", meta = (Input, ExpandByDefault, DefaultArraySize = 1))
+	TArray<FConstraintParent> Parents;
+	
 	UPROPERTY(EditAnywhere, Category = "Constraint", meta = (Input)) 
 	FRigUnit_RotationConstraint_AdvancedSettings AdvancedSettings;
 
-	UPROPERTY(EditAnywhere, Category = "Constraint", meta = (Input, ExpandByDefault, DefaultArraySize = 1))
-	TArray<FConstraintParent> Parents;
+	UPROPERTY(EditAnywhere, Category = "Constraint", meta = (Input))
+	float Weight;
 };
 
 /**
@@ -365,7 +381,9 @@ struct FRigUnit_ScaleConstraint : public FRigUnit_HighlevelBaseMutable
 	FRigUnit_ScaleConstraint()
         : Child(FRigElementKey(NAME_None, ERigElementType::Bone))
         , bMaintainOffset(true)
+		, Weight(1.0f)
 	{
+		Parents.Add(FConstraintParent());
 	}
 	
 	RIGVM_METHOD()
@@ -382,4 +400,7 @@ struct FRigUnit_ScaleConstraint : public FRigUnit_HighlevelBaseMutable
 
 	UPROPERTY(EditAnywhere, Category = "Constraint", meta = (Input, ExpandByDefault, DefaultArraySize = 1))
 	TArray<FConstraintParent> Parents;
-};
\ No newline at end of file
+
+	UPROPERTY(EditAnywhere, Category = "Constraint", meta = (Input))
+	float Weight;
+};
