Control Rig: Update tab layout to include curve container + my blueprint by default

diff --git a/Engine/Plugins/Experimental/ControlRig/Source/ControlRigEditor/Private/Editor/ControlRigEditorMode.cpp b/Engine/Plugins/Experimental/ControlRig/Source/ControlRigEditor/Private/Editor/ControlRigEditorMode.cpp
--- a/Engine/Plugins/Experimental/ControlRig/Source/ControlRigEditor/Private/Editor/ControlRigEditorMode.cpp
+++ b/Engine/Plugins/Experimental/ControlRig/Source/ControlRigEditor/Private/Editor/ControlRigEditorMode.cpp
@@ -39,7 +39,7 @@ FControlRigEditorMode::FControlRigEditorMode(const TSharedRef<FControlRigEditor>
  	TabFactories.RegisterFactory(PersonaModule.CreatePersonaViewportTabFactory(InControlRigEditor, ViewportArgs));
 	TabFactories.RegisterFactory(PersonaModule.CreateAdvancedPreviewSceneTabFactory(InControlRigEditor, InControlRigEditor->GetPersonaToolkit()->GetPreviewScene()));
 
-	TabLayout = FTabManager::NewLayout("Standalone_ControlRigEditMode_Layout_v1.4")
+	TabLayout = FTabManager::NewLayout("Standalone_ControlRigEditMode_Layout_v1.5")
 		->AddArea
 		(
 			// Main application area
@@ -70,6 +70,8 @@ FControlRigEditorMode::FControlRigEditorMode(const TSharedRef<FControlRigEditor>
 						->SetSizeCoefficient(0.5f)
 						->AddTab(FRigHierarchyTabSummoner::TabID, ETabState::OpenedTab)
 						->AddTab(FRigStackTabSummoner::TabID, ETabState::OpenedTab)
+						->AddTab(FRigCurveContainerTabSummoner::TabID, ETabState::OpenedTab)
+						->AddTab(FBlueprintEditorTabs::MyBlueprintID, ETabState::OpenedTab)
 					)
 				)
 				->Split
