Control Rig: Performance improvement for ControlRigGraph UI

diff --git a/Engine/Plugins/Experimental/ControlRig/Source/ControlRigDeveloper/Private/Graph/ControlRigGraph.cpp b/Engine/Plugins/Experimental/ControlRig/Source/ControlRigDeveloper/Private/Graph/ControlRigGraph.cpp
--- a/Engine/Plugins/Experimental/ControlRig/Source/ControlRigDeveloper/Private/Graph/ControlRigGraph.cpp
+++ b/Engine/Plugins/Experimental/ControlRig/Source/ControlRigDeveloper/Private/Graph/ControlRigGraph.cpp
@@ -249,6 +249,8 @@ void UControlRigGraph::HandleModifiedEvent(ERigVMGraphNotifType InNotifType, URi
 	{
 		case ERigVMGraphNotifType::GraphChanged:
 		{
+			ModelNodePathToEdNode.Reset();
+
 			for (URigVMNode* Node : InGraph->GetNodes())
 			{
 				UEdGraphNode* EdNode = FindNodeForModelNodeName(Node->GetFName());
@@ -280,6 +282,8 @@ void UControlRigGraph::HandleModifiedEvent(ERigVMGraphNotifType InNotifType, URi
 		}
 		case ERigVMGraphNotifType::NodeAdded:
 		{
+			ModelNodePathToEdNode.Reset();
+				
 			if (URigVMNode* ModelNode = Cast<URigVMNode>(InSubject))
 			{
 				if (!ModelNode->IsVisibleInUI())
@@ -365,6 +369,8 @@ void UControlRigGraph::HandleModifiedEvent(ERigVMGraphNotifType InNotifType, URi
 		}
 		case ERigVMGraphNotifType::NodeRemoved:
 		{
+			ModelNodePathToEdNode.Reset();
+
 			if (URigVMNode* ModelNode = Cast<URigVMNode>(InSubject))
 			{
 				if (URigVMInjectionInfo* Injection = ModelNode->GetInjectionInfo())
@@ -625,6 +631,8 @@ void UControlRigGraph::HandleModifiedEvent(ERigVMGraphNotifType InNotifType, URi
 		}
 		case ERigVMGraphNotifType::NodeRenamed:
 		{
+			ModelNodePathToEdNode.Reset();
+				
 			if (URigVMNode* ModelNode = Cast<URigVMNode>(InSubject))
 			{
 				if (UControlRigGraphNode* RigNode = Cast<UControlRigGraphNode>(FindNodeForModelNodeName(ModelNode->GetPreviousFName())))
@@ -739,12 +747,19 @@ UEdGraphNode* UControlRigGraph::FindNodeForModelNodeName(const FName& InModelNod
 {
 	DECLARE_SCOPE_HIERARCHICAL_COUNTER_FUNC()
 
+	if(UEdGraphNode** MappedNode = ModelNodePathToEdNode.Find(InModelNodeName))
+	{
+		return *MappedNode;
+	}
+	
+	const FString InModelNodePath = InModelNodeName.ToString();
 	for (UEdGraphNode* EdNode : Nodes)
 	{
 		if (UControlRigGraphNode* RigNode = Cast<UControlRigGraphNode>(EdNode))
 		{
-			if (RigNode->ModelNodePath == InModelNodeName.ToString())
+			if (RigNode->ModelNodePath == InModelNodePath)
 			{
+				ModelNodePathToEdNode.Add(InModelNodeName, EdNode);
 				return EdNode;
 			}
 		}
@@ -752,6 +767,7 @@ UEdGraphNode* UControlRigGraph::FindNodeForModelNodeName(const FName& InModelNod
 		{
 			if (EdNode->GetFName() == InModelNodeName)
 			{
+				ModelNodePathToEdNode.Add(InModelNodeName, EdNode);
 				return EdNode;
 			}
 		}
diff --git a/Engine/Plugins/Experimental/ControlRig/Source/ControlRigDeveloper/Public/Graph/ControlRigGraph.h b/Engine/Plugins/Experimental/ControlRig/Source/ControlRigDeveloper/Public/Graph/ControlRigGraph.h
--- a/Engine/Plugins/Experimental/ControlRig/Source/ControlRigDeveloper/Public/Graph/ControlRigGraph.h
+++ b/Engine/Plugins/Experimental/ControlRig/Source/ControlRigDeveloper/Public/Graph/ControlRigGraph.h
@@ -148,6 +148,7 @@ private:
 	void HandleVMCompiledEvent(UBlueprint* InBlueprint, URigVM* InVM);
 
 	static TArray<TSharedPtr<FString>> EmptyElementNameList;
+	TMap<FName, UEdGraphNode*> ModelNodePathToEdNode;
 	
 	friend class UControlRigUnitNodeSpawner;
 	friend class UControlRigVariableNodeSpawner;
