Control Rig: Colorize viewport outline for different modes

diff --git a/Engine/Plugins/Experimental/ControlRig/Source/ControlRig/Private/Settings/ControlRigSettings.cpp b/Engine/Plugins/Experimental/ControlRig/Source/ControlRig/Private/Settings/ControlRigSettings.cpp
--- a/Engine/Plugins/Experimental/ControlRig/Source/ControlRig/Private/Settings/ControlRigSettings.cpp
+++ b/Engine/Plugins/Experimental/ControlRig/Source/ControlRig/Private/Settings/ControlRigSettings.cpp
@@ -12,6 +12,10 @@ UControlRigSettings::UControlRigSettings(const FObjectInitializer& ObjectInitial
 	DefaultGizmoLibrary = LoadObject<UControlRigGizmoLibrary>(nullptr, TEXT("/ControlRig/Controls/DefaultGizmoLibrary.DefaultGizmoLibrary"));
 	bResetControlsOnCompile = true;
 	bResetControlsOnPinValueInteraction = false;
+
+	SetupEventBorderColor = FLinearColor::Red;
+	BackwardsSolveBorderColor = FLinearColor::Yellow;
+	BackwardsAndForwardsBorderColor = FLinearColor::Blue;
 #endif
 }
 
diff --git a/Engine/Plugins/Experimental/ControlRig/Source/ControlRig/Public/Settings/ControlRigSettings.h b/Engine/Plugins/Experimental/ControlRig/Source/ControlRig/Public/Settings/ControlRigSettings.h
--- a/Engine/Plugins/Experimental/ControlRig/Source/ControlRig/Public/Settings/ControlRigSettings.h
+++ b/Engine/Plugins/Experimental/ControlRig/Source/ControlRig/Public/Settings/ControlRigSettings.h
@@ -58,6 +58,24 @@ public:
 	 */
 	UPROPERTY(EditAnywhere, config, Category = NodeGraph)
 	TMap<FString, FControlRigSettingsPerPinBool> RigUnitPinExpansion;
+	
+	/**
+	 * The border color of the viewport when entering "Setup Event" mode
+	 */
+	UPROPERTY(EditAnywhere, config, Category = Viewport)
+	FLinearColor SetupEventBorderColor;
+	
+	/**
+	 * The border color of the viewport when entering "Backwards Solve" mode
+	 */
+	UPROPERTY(EditAnywhere, config, Category = Viewport)
+	FLinearColor BackwardsSolveBorderColor;
+	
+	/**
+	 * The border color of the viewport when entering "Backwards And Forwards" mode
+	 */
+	UPROPERTY(EditAnywhere, config, Category = Viewport)
+	FLinearColor BackwardsAndForwardsBorderColor;
 #endif
 
 	static UControlRigSettings * Get() { return CastChecked<UControlRigSettings>(UControlRigSettings::StaticClass()->GetDefaultObject()); }
diff --git a/Engine/Plugins/Experimental/ControlRig/Source/ControlRigEditor/Private/ControlRigEditorStyle.h b/Engine/Plugins/Experimental/ControlRig/Source/ControlRigEditor/Private/ControlRigEditorStyle.h
--- a/Engine/Plugins/Experimental/ControlRig/Source/ControlRigEditor/Private/ControlRigEditorStyle.h
+++ b/Engine/Plugins/Experimental/ControlRig/Source/ControlRigEditor/Private/ControlRigEditorStyle.h
@@ -24,7 +24,9 @@ public:
 		const FVector2D Icon24x24(24.0f, 24.0f);
 		const FVector2D Icon32x32(32.0f, 32.0f);
 		const FVector2D Icon40x40(40.0f, 40.0f);
-		SetContentRoot(FPaths::EnginePluginsDir() / TEXT("Experimental/ControlRig/Content"));
+		const FString PluginContentDir = FPaths::EnginePluginsDir() / TEXT("Experimental/ControlRig/Content");
+		const FString EngineEditorSlateDir = FPaths::EngineContentDir() / TEXT("Editor/Slate");
+		SetContentRoot(PluginContentDir);
 
 		const FSlateColor DefaultForeground(FLinearColor(0.72f, 0.72f, 0.72f, 1.f));
 
@@ -97,6 +99,16 @@ public:
 			Set( "ControlRig.StepOut", new IMAGE_BRUSH("Slate/icon_DebugStepOut_40x", Icon40x40));
 			Set( "ControlRig.StepInto", new IMAGE_BRUSH( "Slate/icon_DebugStepIn_40x", Icon40x40 ) );
 			Set( "ControlRig.StepOver", new IMAGE_BRUSH( "Slate/icon_DebugStepOver_40x", Icon40x40 ) );
+
+			{
+				const FString PreviousContentRoot = GetContentRootDir();
+				// Common/RoundedSelection_16x lives in the editor slate folder, change the root temporarily
+				SetContentRoot(EngineEditorSlateDir);
+				// similar style to "LevelViewport.StartingPlayInEditorBorder"
+				Set( "ControlRig.Viewport.Border", new BOX_BRUSH( "Old/Window/ViewportDebugBorder", 0.8f, FLinearColor(1.0f,1.0f,1.0f,1.0f) ) );
+				// restore the previous content root
+				SetContentRoot(PreviousContentRoot);
+			}
 		}
 
 		// Graph styles
diff --git a/Engine/Plugins/Experimental/ControlRig/Source/ControlRigEditor/Private/Editor/ControlRigEditor.cpp b/Engine/Plugins/Experimental/ControlRig/Source/ControlRigEditor/Private/Editor/ControlRigEditor.cpp
--- a/Engine/Plugins/Experimental/ControlRig/Source/ControlRigEditor/Private/Editor/ControlRigEditor.cpp
+++ b/Engine/Plugins/Experimental/ControlRig/Source/ControlRigEditor/Private/Editor/ControlRigEditor.cpp
@@ -2934,6 +2934,46 @@ void FControlRigEditor::HandleViewportCreated(const TSharedRef<class IPersonaVie
 		}
 	));
 
+	auto GetBorderColorAndOpacity = [this]()
+	{
+		FLinearColor Color = FLinearColor::Transparent;
+		switch (GetEventQueue())
+		{
+			case EControlRigEditorEventQueue::Setup:
+				Color = UControlRigSettings::Get()->SetupEventBorderColor;
+				break;
+			case EControlRigEditorEventQueue::Inverse:
+				Color = UControlRigSettings::Get()->BackwardsSolveBorderColor;
+				break;
+			case EControlRigEditorEventQueue::InverseAndUpdate:
+				Color = UControlRigSettings::Get()->BackwardsAndForwardsBorderColor;
+				break;
+			default:
+				Color = FLinearColor::Transparent;
+				break;
+		}
+		return Color;
+	};
+
+	auto GetBorderVisibility = [this]()
+	{
+		EVisibility Visibility = EVisibility::Collapsed;
+		if (GetEventQueue() != EControlRigEditorEventQueue::Update)
+		{
+			Visibility = EVisibility::HitTestInvisible;
+		}
+		return Visibility;
+	};
+	
+	InViewport->AddOverlayWidget(
+		SNew(SBorder)
+        .BorderImage(FControlRigEditorStyle::Get().GetBrush( "ControlRig.Viewport.Border"))
+        .BorderBackgroundColor_Lambda(GetBorderColorAndOpacity)
+        .Visibility_Lambda(GetBorderVisibility)
+        .Padding(0.0f)
+        .ShowEffectWhenDisabled(false)
+	);
+	
 	InViewport->GetKeyDownDelegate().BindLambda([&](const FGeometry& MyGeometry, const FKeyEvent& InKeyEvent) -> FReply {
 		if (OnKeyDownDelegate.IsBound())
 		{
