Control Rig: Hierarchy tracing functionality

diff --git a/Engine/Plugins/Experimental/ControlRig/Source/ControlRig/ControlRig.Build.cs b/Engine/Plugins/Experimental/ControlRig/Source/ControlRig/ControlRig.Build.cs
--- a/Engine/Plugins/Experimental/ControlRig/Source/ControlRig/ControlRig.Build.cs
+++ b/Engine/Plugins/Experimental/ControlRig/Source/ControlRig/ControlRig.Build.cs
@@ -46,6 +46,9 @@ namespace UnrealBuildTool.Rules
 						"EditorStyle",
 						"RigVMDeveloper",
                         "AnimGraph",
+                        "Json",
+                        "Serialization",
+                        "JsonUtilities",
                     }
                 );
 
diff --git a/Engine/Plugins/Experimental/ControlRig/Source/ControlRig/Private/AnimNode_ControlRigBase.cpp b/Engine/Plugins/Experimental/ControlRig/Source/ControlRig/Private/AnimNode_ControlRigBase.cpp
--- a/Engine/Plugins/Experimental/ControlRig/Source/ControlRig/Private/AnimNode_ControlRigBase.cpp
+++ b/Engine/Plugins/Experimental/ControlRig/Source/ControlRig/Private/AnimNode_ControlRigBase.cpp
@@ -12,6 +12,9 @@
 TAutoConsoleVariable<int32> CVarAnimNodeControlRigDebug(TEXT("a.AnimNode.ControlRig.Debug"), 0, TEXT("Set to 1 to turn on debug drawing for AnimNode_ControlRigBase"));
 #endif
 
+// CVar to disable control rig execution within an anim node
+static TAutoConsoleVariable<int32> CVarControlRigDisableExecutionAnimNode(TEXT("ControlRig.DisableExecutionInAnimNode"), 0, TEXT("if nonzero we disable the execution of Control Rigs inside an anim node."));
+
 FAnimNode_ControlRigBase::FAnimNode_ControlRigBase()
 	: FAnimNode_CustomProperty()
 	, InputSettings(FControlRigIOSettings())
@@ -84,8 +87,39 @@ void FAnimNode_ControlRigBase::Update_AnyThread(const FAnimationUpdateContext& C
 	}
 }
 
+bool FAnimNode_ControlRigBase::CanExecute()
+{
+	if(CVarControlRigDisableExecutionAnimNode->GetInt() != 0)
+	{
+		return false;
+	}
+
+	if (UControlRig* ControlRig = GetControlRig())
+	{
+		return ControlRig->CanExecute(); 
+	}
+
+	return false;
+}
+
 void FAnimNode_ControlRigBase::UpdateInput(UControlRig* ControlRig, const FPoseContext& InOutput)
 {
+	if(!CanExecute())
+	{
+		return;
+	}
+
+#if WITH_EDITOR
+	// if we are recording any change - let's clear the undo stack
+	if(URigHierarchy* Hierarchy = ControlRig->GetHierarchy())
+	{
+		if(Hierarchy->IsTracingChanges())
+		{
+			Hierarchy->ResetTransformStack();
+		}
+	}
+#endif
+
 	DECLARE_SCOPE_HIERARCHICAL_COUNTER_FUNC()
 
 	if (InputSettings.bUpdatePose)
@@ -137,10 +171,25 @@ void FAnimNode_ControlRigBase::UpdateInput(UControlRig* ControlRig, const FPoseC
 			ControlRig->GetHierarchy()->SetCurveValue(Key, InOutput.Curve.Get(Index));
 		}
 	}
+
+#if WITH_EDITOR
+	if(URigHierarchy* Hierarchy = ControlRig->GetHierarchy())
+	{
+		if(Hierarchy->IsTracingChanges())
+		{
+			Hierarchy->StorePoseForTrace(TEXT("FAnimNode_ControlRigBase::UpdateInput"));
+		}
+	}
+#endif
 }
 
 void FAnimNode_ControlRigBase::UpdateOutput(UControlRig* ControlRig, FPoseContext& InOutput)
 {
+	if(!CanExecute())
+	{
+		return;
+	}
+
 	DECLARE_SCOPE_HIERARCHICAL_COUNTER_FUNC()
 
 	if (OutputSettings.bUpdatePose)
@@ -188,10 +237,21 @@ void FAnimNode_ControlRigBase::UpdateOutput(UControlRig* ControlRig, FPoseContex
 			{
 				// this causes a side effect of marking the curve as "valid"
 				// so only apply it for curves that have really changed
-			InOutput.Curve.Set(Index, Value);
+				InOutput.Curve.Set(Index, Value);
+			}
 		}
 	}
-}
+
+#if WITH_EDITOR
+	if(URigHierarchy* Hierarchy = ControlRig->GetHierarchy())
+	{
+		if(Hierarchy->IsTracingChanges())
+		{
+			Hierarchy->StorePoseForTrace(TEXT("FAnimNode_ControlRigBase::UpdateOutput"));
+			Hierarchy->DumpTransformStackToFile();
+		}
+	}
+#endif
 }
 
 void FAnimNode_ControlRigBase::Evaluate_AnyThread(FPoseContext& Output)
@@ -210,7 +270,7 @@ void FAnimNode_ControlRigBase::Evaluate_AnyThread(FPoseContext& Output)
 		SourcePose.ResetToRefPose();
 	}
 
-	if (FAnimWeight::IsRelevant(InternalBlendAlpha))
+	if (CanExecute() && FAnimWeight::IsRelevant(InternalBlendAlpha) && GetControlRig())
 	{
 		if (FAnimWeight::IsFullWeight(InternalBlendAlpha))
 		{
@@ -251,6 +311,16 @@ void FAnimNode_ControlRigBase::ExecuteControlRig(FPoseContext& InOutput)
 
 		if (bExecute)
 		{
+#if WITH_EDITOR
+			if(URigHierarchy* Hierarchy = ControlRig->GetHierarchy())
+			{
+				if(Hierarchy->IsTracingChanges())
+				{
+					Hierarchy->StorePoseForTrace(TEXT("FAnimNode_ControlRigBase::BeforeEvaluate"));
+				}
+			}
+#endif
+
 			// first evaluate control rig
 			ControlRig->Evaluate_AnyThread();
 
@@ -263,6 +333,16 @@ void FAnimNode_ControlRigBase::ExecuteControlRig(FPoseContext& InOutput)
 				QueueControlRigDrawInstructions(ControlRig, InOutput.AnimInstanceProxy);
 			}
 #endif
+
+#if WITH_EDITOR
+			if(URigHierarchy* Hierarchy = ControlRig->GetHierarchy())
+			{
+				if(Hierarchy->IsTracingChanges())
+				{
+					Hierarchy->StorePoseForTrace(TEXT("FAnimNode_ControlRigBase::AfterEvaluate"));
+				}
+			}
+#endif
 		}
 
 		// now update output
diff --git a/Engine/Plugins/Experimental/ControlRig/Source/ControlRig/Private/ControlRig.cpp b/Engine/Plugins/Experimental/ControlRig/Source/ControlRig/Private/ControlRig.cpp
--- a/Engine/Plugins/Experimental/ControlRig/Source/ControlRig/Private/ControlRig.cpp
+++ b/Engine/Plugins/Experimental/ControlRig/Source/ControlRig/Private/ControlRig.cpp
@@ -43,6 +43,9 @@ static TAutoConsoleVariable<int32> CVarControlRigCreateFloatControlsForCurves(
 	TEXT("If nonzero we create a float control for each curve in the curve container, useful for debugging low level controls."),
 	ECVF_Default);
 
+// CVar to disable all control rig execution 
+static TAutoConsoleVariable<int32> CVarControlRigDisableExecutionAll(TEXT("ControlRig.DisableExecutionAll"), 0, TEXT("if nonzero we disable all execution of Control Rigs."));
+
 UControlRig::UControlRig(const FObjectInitializer& ObjectInitializer)
 	: Super(ObjectInitializer)
 	, DeltaTime(0.0f)
@@ -459,6 +462,11 @@ void UControlRig::InstantiateVMFromCDO()
 
 void UControlRig::Execute(const EControlRigState InState, const FName& InEventName)
 {
+	if(!CanExecute())
+	{
+		return;
+	}
+	
 	DECLARE_SCOPE_HIERARCHICAL_COUNTER_FUNC()
 	QUICK_SCOPE_CYCLE_COUNTER(STAT_ControlRig_Execute);
 	
@@ -1131,8 +1139,14 @@ void UControlRig::PostEditUndo()
 {
 	Super::PostEditUndo();
 }
+
 #endif // WITH_EDITOR
 
+bool UControlRig::CanExecute() const
+{
+	return CVarControlRigDisableExecutionAll->GetInt() == 0;
+}
+
 void UControlRig::Serialize(FArchive& Ar)
 {
 	DECLARE_SCOPE_HIERARCHICAL_COUNTER_FUNC()
diff --git a/Engine/Plugins/Experimental/ControlRig/Source/ControlRig/Private/ControlRigComponent.cpp b/Engine/Plugins/Experimental/ControlRig/Source/ControlRig/Private/ControlRigComponent.cpp
--- a/Engine/Plugins/Experimental/ControlRig/Source/ControlRig/Private/ControlRigComponent.cpp
+++ b/Engine/Plugins/Experimental/ControlRig/Source/ControlRig/Private/ControlRigComponent.cpp
@@ -7,6 +7,9 @@
 #include "Components/InstancedStaticMeshComponent.h"
 #include "AnimCustomInstanceHelper.h"
 
+// CVar to disable control rig execution within a component
+static TAutoConsoleVariable<int32> CVarControlRigDisableExecutionComponent(TEXT("ControlRig.DisableExecutionInComponent"), 0, TEXT("if nonzero we disable the execution of Control Rigs inside a ControlRigComponent."));
+
 FControlRigAnimInstanceProxy* FControlRigComponentMappedElement::GetAnimProxyOnGameThread() const
 {
 	if (USkeletalMeshComponent* SkeletalMeshComponent = Cast<USkeletalMeshComponent>(SceneComponent))
@@ -225,6 +228,21 @@ UControlRig* UControlRigComponent::GetControlRig()
 	return SetupControlRigIfRequired();
 }
 
+bool UControlRigComponent::CanExecute()
+{
+	if(CVarControlRigDisableExecutionComponent->GetInt() != 0)
+	{
+		return false;
+	}
+	
+	if(UControlRig* CR = GetControlRig())
+	{
+		return CR->CanExecute();
+	}
+	
+	return false;
+}
+
 float UControlRigComponent::GetAbsoluteTime() const
 {
 	if(ControlRig)
@@ -283,6 +301,11 @@ void UControlRigComponent::Update(float DeltaTime)
 {
 	if(UControlRig* CR = SetupControlRigIfRequired())
 	{
+		if(!CanExecute())
+		{
+			return;
+		}
+		
 		if (CR->IsExecuting() || CR->IsInitializing())
 		{
 			ReportError(TEXT("Update is being called recursively."));
@@ -301,6 +324,15 @@ void UControlRigComponent::Update(float DeltaTime)
 			}
 
 #if WITH_EDITOR
+			// if we are recording any change - let's clear the undo stack
+			if(URigHierarchy* Hierarchy = CR->GetHierarchy())
+			{
+				if(Hierarchy->IsTracingChanges())
+				{
+					Hierarchy->ResetTransformStack();
+				}
+			}
+
 			if (bUpdateInEditor)
 			{
 				FEditorScriptExecutionGuard AllowScripts;
@@ -312,7 +344,27 @@ void UControlRigComponent::Update(float DeltaTime)
 				OnPreUpdate(this);
 			}
 
+#if WITH_EDITOR
+			if(URigHierarchy* Hierarchy = CR->GetHierarchy())
+			{
+				if(Hierarchy->IsTracingChanges())
+				{
+					Hierarchy->StorePoseForTrace(TEXT("UControlRigComponent::BeforeEvaluate"));
+				}
+			}
+#endif
+
 			CR->Evaluate_AnyThread();
+
+#if WITH_EDITOR
+			if(URigHierarchy* Hierarchy = CR->GetHierarchy())
+			{
+				if(Hierarchy->IsTracingChanges())
+				{
+					Hierarchy->StorePoseForTrace(TEXT("UControlRigComponent::AfterEvaluate"));
+				}
+			}
+#endif
 		} 
 
 		if (bShowDebugDrawing)
@@ -322,6 +374,17 @@ void UControlRigComponent::Update(float DeltaTime)
 				MarkRenderStateDirty();
 			}
 		}
+
+#if WITH_EDITOR
+		// if we are recording any change - dump the trace to file
+		if(URigHierarchy* Hierarchy = CR->GetHierarchy())
+		{
+			if(Hierarchy->IsTracingChanges())
+			{
+				Hierarchy->DumpTransformStackToFile();
+			}
+		}
+#endif
 	} 
 }
 
@@ -1225,6 +1288,15 @@ void UControlRigComponent::TransferInputs()
 			ConvertTransformToRigSpace(Transform, MappedElement.Space);
 			ControlRig->GetHierarchy()->SetGlobalTransform(MappedElement.ElementIndex, Transform);
 		}
+#if WITH_EDITOR
+		if(URigHierarchy* Hierarchy = ControlRig->GetHierarchy())
+		{
+			if(Hierarchy->IsTracingChanges())
+			{
+				Hierarchy->StorePoseForTrace(TEXT("UControlRigComponent::TransferInputs"));
+			}
+		}
+#endif
 	}
 }
 
@@ -1352,6 +1424,16 @@ void UControlRigComponent::TransferOutputs()
 				}
 			}
 		}
+
+#if WITH_EDITOR
+		if(URigHierarchy* Hierarchy = ControlRig->GetHierarchy())
+		{
+			if(Hierarchy->IsTracingChanges())
+			{
+				Hierarchy->StorePoseForTrace(TEXT("UControlRigComponent::TransferOutputs"));
+			}
+		}
+#endif
 	}
 }
 
diff --git a/Engine/Plugins/Experimental/ControlRig/Source/ControlRig/Private/Rigs/RigHierarchy.cpp b/Engine/Plugins/Experimental/ControlRig/Source/ControlRig/Private/Rigs/RigHierarchy.cpp
--- a/Engine/Plugins/Experimental/ControlRig/Source/ControlRig/Private/Rigs/RigHierarchy.cpp
+++ b/Engine/Plugins/Experimental/ControlRig/Source/ControlRig/Private/Rigs/RigHierarchy.cpp
@@ -1,6 +1,8 @@
 // Copyright Epic Games, Inc. All Rights Reserved.
 
 #include "Rigs/RigHierarchy.h"
+#include "ControlRig.h"
+
 #include "Rigs/RigHierarchyElements.h"
 #include "Rigs/RigHierarchyController.h"
 #include "Units/RigUnitContext.h"
@@ -10,7 +12,63 @@
 #include "Serialization/MemoryWriter.h"
 
 #if WITH_EDITOR
-	#include "ScopedTransaction.h"
+
+#include "ScopedTransaction.h"
+#include "HAL/PlatformStackWalk.h"
+#include "HAL/PlatformFileManager.h"
+#include "Misc/FileHelper.h"
+#include "Misc/Paths.h"
+#include "JsonObjectConverter.h"
+#include "Serialization/JsonReader.h"
+#include "Serialization/JsonSerializer.h"
+
+static FCriticalSection GRigHierarchyStackTraceMutex;
+static char GRigHierarchyStackTrace[65536];
+static void RigHierarchyCaptureCallStack(FString& OutCallstack, uint32 NumCallsToIgnore)
+{
+	FScopeLock ScopeLock(&GRigHierarchyStackTraceMutex);
+	GRigHierarchyStackTrace[0] = 0;
+	FPlatformStackWalk::StackWalkAndDump(GRigHierarchyStackTrace, 65535, 1 + NumCallsToIgnore);
+	OutCallstack = ANSI_TO_TCHAR(GRigHierarchyStackTrace);
+}
+
+// CVar to record all transform changes 
+static TAutoConsoleVariable<int32> CVarControlRigHierarchyTraceAlways(TEXT("ControlRig.Hierarchy.TraceAlways"), 0, TEXT("if nonzero we will record all transform changes."));
+static TAutoConsoleVariable<int32> CVarControlRigHierarchyTraceCallstack(TEXT("ControlRig.Hierarchy.TraceCallstack"), 0, TEXT("if nonzero we will record the callstack for any trace entry.\nOnly works if(ControlRig.Hierarchy.TraceEnabled != 0)"));
+
+// A console command to trace a single frame / single execution for a control rig anim node / control rig component
+FAutoConsoleCommandWithWorldAndArgs FCmdControlRigHierarchyTraceFrames
+(
+	TEXT("ControlRig.Hierarchy.Trace"),
+	TEXT("Traces changes in a hierarchy for a provided number of executions (defaults to 1).\nYou can use ControlRig.Hierarchy.TraceCallstack to enable callstack tracing as part of this."),
+	FConsoleCommandWithWorldAndArgsDelegate::CreateLambda([](const TArray<FString>& InParams, UWorld* InWorld)
+	{
+		int32 NumFrames = 1;
+		if(InParams.Num() > 0)
+		{
+			NumFrames = FCString::Atoi(*InParams[0]);
+		}
+		
+		TArray<UObject*> Instances;
+		URigHierarchy::StaticClass()->GetDefaultObject()->GetArchetypeInstances(Instances);
+
+		for(UObject* Instance : Instances)
+		{
+			if (Instance->HasAnyFlags(RF_ClassDefaultObject))
+			{
+				continue;
+			}
+			
+			// we'll just trace all of them for now
+			//if(Instance->GetWorld() == InWorld)
+			if(Instance->GetTypedOuter<UControlRig>() != nullptr)
+			{
+				CastChecked<URigHierarchy>(Instance)->TraceFrames(NumFrames);
+			}
+		}
+	})
+);
+
 #endif
 
 ////////////////////////////////////////////////////////////////////////////////
@@ -920,7 +978,7 @@ bool URigHierarchy::Undo()
 		return false;
 	}
 
-	const FTransformStackEntry Entry = TransformUndoStack.Pop();
+	const FRigTransformStackEntry Entry = TransformUndoStack.Pop();
 	ApplyTransformFromStack(Entry, true);
 	UndoRedoEvent.Broadcast(this, Entry.Key, Entry.TransformType, Entry.OldTransform, true);
 	TransformRedoStack.Push(Entry);
@@ -943,7 +1001,7 @@ bool URigHierarchy::Redo()
 		return false;
 	}
 
-	const FTransformStackEntry Entry = TransformRedoStack.Pop();
+	const FRigTransformStackEntry Entry = TransformRedoStack.Pop();
 	ApplyTransformFromStack(Entry, false);
 	UndoRedoEvent.Broadcast(this, Entry.Key, Entry.TransformType, Entry.NewTransform, false);
 	TransformUndoStack.Push(Entry);
@@ -1259,15 +1317,16 @@ void URigHierarchy::SetTransform(FRigTransformElement* InTransformElement, const
 	}
 
 #if WITH_EDITOR
-	if(bSetupUndo)
+	if(bSetupUndo || IsTracingChanges())
 	{
 		PushTransformToStack(
 			InTransformElement->GetKey(),
-			ETransformStackEntryType::TransformPose,
+			ERigTransformStackEntryType::TransformPose,
 			InTransformType,
 			PreviousTransform,
 			InTransformElement->Pose.Get(InTransformType),
-			bAffectChildren);
+			bAffectChildren,
+			bSetupUndo);
 	}
 #endif
 }
@@ -1328,15 +1387,16 @@ void URigHierarchy::SetControlOffsetTransform(FRigControlElement* InControlEleme
 	InControlElement->Gizmo.MarkDirty(MakeGlobal(InTransformType));
 
 #if WITH_EDITOR
-	if(bSetupUndo)
+	if(bSetupUndo || IsTracingChanges())
 	{
 		PushTransformToStack(
             InControlElement->GetKey(),
-            ETransformStackEntryType::ControlOffset,
+            ERigTransformStackEntryType::ControlOffset,
             InTransformType,
             PreviousTransform,
             InControlElement->Offset.Get(InTransformType),
-            bAffectChildren);
+            bAffectChildren,
+            bSetupUndo);
 	}
 #endif
 }
@@ -1391,15 +1451,16 @@ void URigHierarchy::SetControlGizmoTransform(FRigControlElement* InControlElemen
 	InControlElement->Gizmo.MarkDirty(OpposedType);
 
 #if WITH_EDITOR
-	if(bSetupUndo)
+	if(bSetupUndo || IsTracingChanges())
 	{
 		PushTransformToStack(
             InControlElement->GetKey(),
-            ETransformStackEntryType::ControlGizmo,
+            ERigTransformStackEntryType::ControlGizmo,
             InTransformType,
             PreviousTransform,
             InControlElement->Gizmo.Get(InTransformType),
-            false);
+            false,
+            bSetupUndo);
 	}
 #endif
 
@@ -1736,9 +1797,9 @@ void URigHierarchy::SetCurveValue(FRigCurveElement* InCurveElement, float InValu
 	InCurveElement->Value = InValue;
 
 #if WITH_EDITOR
-	if(bSetupUndo)
+	if(bSetupUndo || IsTracingChanges())
 	{
-		PushCurveToStack(InCurveElement->GetKey(), PreviousValue, InCurveElement->Value);
+		PushCurveToStack(InCurveElement->GetKey(), PreviousValue, InCurveElement->Value, bSetupUndo);
 	}
 #endif
 }
@@ -1795,6 +1856,123 @@ bool URigHierarchy::IsParentedTo(FRigBaseElement* InChild, FRigBaseElement* InPa
 	return false;
 }
 
+bool URigHierarchy::IsTracingChanges() const
+{
+#if WITH_EDITOR
+	return (CVarControlRigHierarchyTraceAlways->GetInt() != 0) || (TraceFramesLeft > 0);
+#else
+	return false;
+#endif
+}
+
+#if WITH_EDITOR
+
+void URigHierarchy::ResetTransformStack()
+{
+	TransformUndoStack.Reset();
+	TransformRedoStack.Reset();
+	TransformStackIndex = TransformUndoStack.Num();
+
+	if(IsTracingChanges())
+	{
+		TracePoses.Reset();
+		StorePoseForTrace(TEXT("BeginOfFrame"));
+	}
+}
+
+void URigHierarchy::StorePoseForTrace(const FString& InPrefix)
+{
+	check(!InPrefix.IsEmpty());
+	
+	FName InitialKey = *FString::Printf(TEXT("%s_Initial"), *InPrefix);
+	FName CurrentKey = *FString::Printf(TEXT("%s_Current"), *InPrefix);
+	TracePoses.FindOrAdd(InitialKey) = GetPose(true);
+	TracePoses.FindOrAdd(CurrentKey) = GetPose(false);
+}
+
+template <class CharType>
+struct TRigHierarchyJsonPrintPolicy
+	: public TPrettyJsonPrintPolicy<CharType>
+{
+	static inline void WriteDouble(  FArchive* Stream, double Value )
+	{
+		TJsonPrintPolicy<CharType>::WriteString(Stream, FString::Printf(TEXT("%.3f"), Value));
+	}
+};
+
+void URigHierarchy::DumpTransformStackToFile(FString* OutFilePath)
+{
+	if(IsTracingChanges())
+	{
+		StorePoseForTrace(TEXT("EndOfFrame"));
+	}
+
+	FString PathName = GetPathName();
+	PathName.Split(TEXT(":"), nullptr, &PathName);
+	PathName.ReplaceCharInline('.', '/');
+
+	FString Suffix;
+	if(TraceFramesLeft > 0)
+	{
+		Suffix = FString::Printf(TEXT("_Trace_%03d"), TraceFramesCaptured);
+	}
+
+	FString FileName = FString::Printf(TEXT("%sControlRig/%s%s.json"), *FPaths::ProjectLogDir(), *PathName, *Suffix);
+	FString FullFilename = FPlatformFileManager::Get().GetPlatformFile().ConvertToAbsolutePathForExternalAppForWrite(*FileName);
+
+	TSharedPtr<FJsonObject> JsonData = MakeShareable(new FJsonObject);
+	JsonData->SetStringField(TEXT("PathName"), GetPathName());
+
+	TSharedRef<FJsonObject> JsonTracedPoses = MakeShareable(new FJsonObject);
+	for(const TPair<FName, FRigPose>& Pair : TracePoses)
+	{
+		TSharedRef<FJsonObject> JsonTracedPose = MakeShareable(new FJsonObject);
+		if (FJsonObjectConverter::UStructToJsonObject(FRigPose::StaticStruct(), &Pair.Value, JsonTracedPose, 0, 0))
+		{
+			JsonTracedPoses->SetObjectField(Pair.Key.ToString(), JsonTracedPose);
+		}
+	}
+	JsonData->SetObjectField(TEXT("TracedPoses"), JsonTracedPoses);
+
+	TArray<TSharedPtr<FJsonValue>> JsonTransformStack;
+	for (const FRigTransformStackEntry& TransformStackEntry : TransformUndoStack)
+	{
+		TSharedRef<FJsonObject> JsonTransformStackEntry = MakeShareable(new FJsonObject);
+		if (FJsonObjectConverter::UStructToJsonObject(FRigTransformStackEntry::StaticStruct(), &TransformStackEntry, JsonTransformStackEntry, 0, 0))
+		{
+			JsonTransformStack.Add(MakeShareable(new FJsonValueObject(JsonTransformStackEntry)));
+		}
+	}
+	JsonData->SetArrayField(TEXT("TransformStack"), JsonTransformStack);
+
+	FString JsonText;
+	const TSharedRef< TJsonWriter< TCHAR, TRigHierarchyJsonPrintPolicy<TCHAR> > > JsonWriter = TJsonWriterFactory< TCHAR, TRigHierarchyJsonPrintPolicy<TCHAR> >::Create(&JsonText);
+	if (FJsonSerializer::Serialize(JsonData.ToSharedRef(), JsonWriter))
+	{
+		if ( FFileHelper::SaveStringToFile(JsonText, *FullFilename) )
+		{
+			UE_LOG(LogControlRig, Display, TEXT("Saved hierarchy trace to %s"), *FullFilename);
+
+			if(OutFilePath)
+			{
+				*OutFilePath = FullFilename;
+			}
+		}
+	}
+
+	TraceFramesLeft = FMath::Max(0, TraceFramesLeft - 1);
+	TraceFramesCaptured++;
+}
+
+void URigHierarchy::TraceFrames(int32 InNumFramesToTrace)
+{
+	TraceFramesLeft = InNumFramesToTrace;
+	TraceFramesCaptured = 0;
+	ResetTransformStack();
+}
+
+#endif
+
 bool URigHierarchy::IsSelected(const FRigBaseElement* InElement) const
 {
 	if(InElement == nullptr)
@@ -2055,9 +2233,9 @@ void URigHierarchy::PropagateDirtyFlags(FRigTransformElement* InTransformElement
 	}
 }
 
-void URigHierarchy::PushTransformToStack(const FRigElementKey& InKey, ETransformStackEntryType InEntryType,
+void URigHierarchy::PushTransformToStack(const FRigElementKey& InKey, ERigTransformStackEntryType InEntryType,
 	ERigTransformType::Type InTransformType, const FTransform& InOldTransform, const FTransform& InNewTransform,
-	bool bAffectChildren)
+	bool bAffectChildren, bool bModify)
 {
 #if WITH_EDITOR
 
@@ -2074,22 +2252,22 @@ void URigHierarchy::PushTransformToStack(const FRigElementKey& InKey, ETransform
 	FText Title;
 	switch(InEntryType)
 	{
-		case ETransformStackEntryType::TransformPose:
+		case ERigTransformStackEntryType::TransformPose:
 		{
 			Title = TransformPoseTitle;
 			break;
 		}
-		case ETransformStackEntryType::ControlOffset:
+		case ERigTransformStackEntryType::ControlOffset:
 		{
 			Title = TransformPoseTitle;
 			break;
 		}
-		case ETransformStackEntryType::ControlGizmo:
+		case ERigTransformStackEntryType::ControlGizmo:
 		{
 			Title = TransformPoseTitle;
 			break;
 		}
-		case ETransformStackEntryType::CurveValue:
+		case ERigTransformStackEntryType::CurveValue:
 		{
 			Title = TransformPoseTitle;
 			break;
@@ -2098,11 +2276,15 @@ void URigHierarchy::PushTransformToStack(const FRigElementKey& InKey, ETransform
 
 	TGuardValue<bool> TransactingGuard(bTransactingForTransformChange, true);
 
-	FScopedTransaction Transaction(Title);
+	TSharedPtr<FScopedTransaction> TransactionPtr;
+	if(bModify)
+	{
+		TransactionPtr = MakeShareable(new FScopedTransaction(Title));
+	}
 
 	if(bIsInteracting)
 	{
-		FTransformStackEntry LastEntry;
+		FRigTransformStackEntry LastEntry;
 		if(!TransformUndoStack.IsEmpty())
 		{
 			LastEntry = TransformUndoStack.Last();
@@ -2112,14 +2294,14 @@ void URigHierarchy::PushTransformToStack(const FRigElementKey& InKey, ETransform
 		{
 			// merge the entries on the stack
 			TransformUndoStack.Last() = 
-                FTransformStackEntry(InKey, InEntryType, InTransformType, LastEntry.OldTransform, InNewTransform, bAffectChildren);
+                FRigTransformStackEntry(InKey, InEntryType, InTransformType, LastEntry.OldTransform, InNewTransform, bAffectChildren);
 		}
 		else
 		{
 			Modify();
 
 			TransformUndoStack.Add(
-                FTransformStackEntry(InKey, InEntryType, InTransformType, InOldTransform, InNewTransform, bAffectChildren));
+                FRigTransformStackEntry(InKey, InEntryType, InTransformType, InOldTransform, InNewTransform, bAffectChildren));
 			TransformStackIndex = TransformUndoStack.Num();
 		}
 
@@ -2127,10 +2309,40 @@ void URigHierarchy::PushTransformToStack(const FRigElementKey& InKey, ETransform
 		return;
 	}
 
-	Modify();
+	if(bModify)
+	{
+		Modify();
+	}
+
+	TArray<FString> Callstack;
+	if(IsTracingChanges() && (CVarControlRigHierarchyTraceCallstack->GetInt() != 0))
+	{
+		FString JoinedCallStack;
+		RigHierarchyCaptureCallStack(JoinedCallStack, 1);
+		JoinedCallStack.ReplaceInline(TEXT("\r"), TEXT(""));
+
+		FString Left, Right;
+		do
+		{
+			if(!JoinedCallStack.Split(TEXT("\n"), &Left, &Right))
+			{
+				Left = JoinedCallStack;
+				Right.Empty();
+			}
+
+			Left.TrimStartAndEndInline();
+			if(Left.StartsWith(TEXT("0x")))
+			{
+				Left.Split(TEXT(" "), nullptr, &Left);
+			}
+			Callstack.Add(Left);
+			JoinedCallStack = Right;
+		}
+		while(!JoinedCallStack.IsEmpty());
+	}
 
 	TransformUndoStack.Add(
-		FTransformStackEntry(InKey, InEntryType, InTransformType, InOldTransform, InNewTransform, bAffectChildren));
+		FRigTransformStackEntry(InKey, InEntryType, InTransformType, InOldTransform, InNewTransform, bAffectChildren, Callstack));
 	TransformStackIndex = TransformUndoStack.Num();
 
 	TransformRedoStack.Reset();
@@ -2138,7 +2350,7 @@ void URigHierarchy::PushTransformToStack(const FRigElementKey& InKey, ETransform
 #endif
 }
 
-void URigHierarchy::PushCurveToStack(const FRigElementKey& InKey, float InOldCurveValue, float InNewCurveValue)
+void URigHierarchy::PushCurveToStack(const FRigElementKey& InKey, float InOldCurveValue, float InNewCurveValue, bool bModify)
 {
 #if WITH_EDITOR
 
@@ -2148,12 +2360,12 @@ void URigHierarchy::PushCurveToStack(const FRigElementKey& InKey, float InOldCur
 	OldTransform.SetTranslation(FVector(InOldCurveValue, 0.f, 0.f));
 	NewTransform.SetTranslation(FVector(InNewCurveValue, 0.f, 0.f));
 
-	PushTransformToStack(InKey, ETransformStackEntryType::CurveValue, ERigTransformType::CurrentLocal, OldTransform, NewTransform, false);
+	PushTransformToStack(InKey, ERigTransformStackEntryType::CurveValue, ERigTransformType::CurrentLocal, OldTransform, NewTransform, false, bModify);
 
 #endif
 }
 
-bool URigHierarchy::ApplyTransformFromStack(const FTransformStackEntry& InEntry, bool bUndo)
+bool URigHierarchy::ApplyTransformFromStack(const FRigTransformStackEntry& InEntry, bool bUndo)
 {
 #if WITH_EDITOR
 	
@@ -2167,22 +2379,22 @@ bool URigHierarchy::ApplyTransformFromStack(const FTransformStackEntry& InEntry,
 	
 	switch(InEntry.EntryType)
 	{
-		case ETransformStackEntryType::TransformPose:
+		case ERigTransformStackEntryType::TransformPose:
 		{
 			SetTransform(Cast<FRigTransformElement>(Element), Transform, InEntry.TransformType, InEntry.bAffectChildren, false); 
 			break;
 		}
-		case ETransformStackEntryType::ControlOffset:
+		case ERigTransformStackEntryType::ControlOffset:
 		{
 			SetControlOffsetTransform(Cast<FRigControlElement>(Element), Transform, InEntry.TransformType, InEntry.bAffectChildren, false); 
 			break;
 		}
-		case ETransformStackEntryType::ControlGizmo:
+		case ERigTransformStackEntryType::ControlGizmo:
 		{
 			SetControlGizmoTransform(Cast<FRigControlElement>(Element), Transform, InEntry.TransformType, false); 
 			break;
 		}
-		case ETransformStackEntryType::CurveValue:
+		case ERigTransformStackEntryType::CurveValue:
 		{
 			const float CurveValue = Transform.GetTranslation().X;
 			SetCurveValue(Cast<FRigCurveElement>(Element), CurveValue, false);
diff --git a/Engine/Plugins/Experimental/ControlRig/Source/ControlRig/Public/AnimNode_ControlRigBase.h b/Engine/Plugins/Experimental/ControlRig/Source/ControlRig/Public/AnimNode_ControlRigBase.h
--- a/Engine/Plugins/Experimental/ControlRig/Source/ControlRig/Public/AnimNode_ControlRigBase.h
+++ b/Engine/Plugins/Experimental/ControlRig/Source/ControlRig/Public/AnimNode_ControlRigBase.h
@@ -92,6 +92,7 @@ protected:
 	// The below is alpha value support for control rig
 	float InternalBlendAlpha;
 
+	virtual bool CanExecute();
 	// update input/output to control rig
 	virtual void UpdateInput(UControlRig* ControlRig, const FPoseContext& InOutput);
 	virtual void UpdateOutput(UControlRig* ControlRig, FPoseContext& InOutput);
diff --git a/Engine/Plugins/Experimental/ControlRig/Source/ControlRig/Public/ControlRig.h b/Engine/Plugins/Experimental/ControlRig/Source/ControlRig/Public/ControlRig.h
--- a/Engine/Plugins/Experimental/ControlRig/Source/ControlRig/Public/ControlRig.h
+++ b/Engine/Plugins/Experimental/ControlRig/Source/ControlRig/Public/ControlRig.h
@@ -84,6 +84,10 @@ public:
 #if WITH_EDITOR
 	virtual void PostEditUndo() override;
 #endif
+
+	/** Is valid for execution */
+	UFUNCTION(BlueprintPure, Category="Control Rig")
+	virtual bool CanExecute() const;
 	
 	/** Set the current delta time */
 	UFUNCTION(BlueprintCallable, Category="Control Rig")
diff --git a/Engine/Plugins/Experimental/ControlRig/Source/ControlRig/Public/ControlRigComponent.h b/Engine/Plugins/Experimental/ControlRig/Source/ControlRig/Public/ControlRigComponent.h
--- a/Engine/Plugins/Experimental/ControlRig/Source/ControlRig/Public/ControlRigComponent.h
+++ b/Engine/Plugins/Experimental/ControlRig/Source/ControlRig/Public/ControlRigComponent.h
@@ -225,6 +225,10 @@ public:
 	UFUNCTION(BlueprintPure, Category = "ControlRig")
 	UControlRig* GetControlRig();
 
+	/** Returns true if the Component can execute its Control Rig */
+	UFUNCTION(BlueprintPure, Category = "ControlRig")
+	bool CanExecute();
+
 	/** Get the ControlRig's local time in seconds since its last initialize */
 	UFUNCTION(BlueprintPure, Category = "ControlRig")
 	float GetAbsoluteTime() const;
diff --git a/Engine/Plugins/Experimental/ControlRig/Source/ControlRig/Public/Rigs/RigHierarchy.h b/Engine/Plugins/Experimental/ControlRig/Source/ControlRig/Public/Rigs/RigHierarchy.h
--- a/Engine/Plugins/Experimental/ControlRig/Source/ControlRig/Public/Rigs/RigHierarchy.h
+++ b/Engine/Plugins/Experimental/ControlRig/Source/ControlRig/Public/Rigs/RigHierarchy.h
@@ -16,6 +16,69 @@ class URigHierarchyController;
 DECLARE_MULTICAST_DELEGATE_ThreeParams(FRigHierarchyModifiedEvent, ERigHierarchyNotification /* type */, URigHierarchy* /* hierarchy */, const FRigBaseElement* /* element */);
 DECLARE_EVENT_FiveParams(URigHierarchy, FRigHierarchyUndoRedoTransformEvent, URigHierarchy*, const FRigElementKey&, ERigTransformType::Type, const FTransform&, bool /* bUndo */);
 
+UENUM()
+enum ERigTransformStackEntryType
+{
+	TransformPose,
+	ControlOffset,
+	ControlGizmo,
+	CurveValue
+};
+
+USTRUCT()
+struct FRigTransformStackEntry
+{
+	GENERATED_BODY()
+
+	FORCEINLINE FRigTransformStackEntry()
+		: Key()
+		, EntryType(ERigTransformStackEntryType::TransformPose)
+		, TransformType(ERigTransformType::CurrentLocal)
+		, OldTransform(FTransform::Identity)
+		, NewTransform(FTransform::Identity)
+		, bAffectChildren(true)
+		, Callstack() 
+	{}
+
+	FORCEINLINE FRigTransformStackEntry(
+		const FRigElementKey& InKey,
+		ERigTransformStackEntryType InEntryType,
+		ERigTransformType::Type InTransformType,
+		const FTransform& InOldTransform,
+		const FTransform& InNewTransform,
+		bool bInAffectChildren,
+		const TArray<FString>& InCallstack =  TArray<FString>())
+		: Key(InKey)
+		, EntryType(InEntryType)
+		, TransformType(InTransformType)
+		, OldTransform(InOldTransform)
+		, NewTransform(InNewTransform)
+		, bAffectChildren(bInAffectChildren)
+		, Callstack(InCallstack)
+	{}
+
+	UPROPERTY()
+	FRigElementKey Key;
+
+	UPROPERTY()
+	TEnumAsByte<ERigTransformStackEntryType> EntryType;
+
+	UPROPERTY()
+	TEnumAsByte<ERigTransformType::Type> TransformType;
+	
+	UPROPERTY()
+	FTransform OldTransform;
+
+	UPROPERTY()
+	FTransform NewTransform;
+
+	UPROPERTY()
+	bool bAffectChildren;
+
+	UPROPERTY()
+	TArray<FString> Callstack;
+};
+
 UCLASS(BlueprintType)
 class CONTROLRIG_API URigHierarchy : public UObject
 {
@@ -33,6 +96,10 @@ public:
 	, bIsInteracting(false)
 	, bSuspendNotifications(false)
 	, ResetPoseHash(INDEX_NONE)
+#if WITH_EDITOR
+	, TraceFramesLeft(0)
+	, TraceFramesCaptured(0)
+#endif
 	{
 		Reset();
 	}
@@ -2042,10 +2109,39 @@ public:
 	/**
 	 * Returns a reference to the suspend notifications flag
 	 */
-	FORCEINLINE bool& GetSuspendNotificationsFlag() { return bSuspendNotifications; } 
+	FORCEINLINE bool& GetSuspendNotificationsFlag() { return bSuspendNotifications; }
 
-private:
+	/*
+	 * Returns true if a hierarchy will record any change.
+	 * This is used for debugging purposes.
+	 */
+	bool IsTracingChanges() const;
+
+#if WITH_EDITOR
+
+	/**
+	 * Clears the undo / redo stack of this hierarchy
+	 */
+	void ResetTransformStack();
 
+	/**
+	 * Stores the current pose for tracing
+	 */
+	void StorePoseForTrace(const FString& InPrefix);
+
+	/**
+	 * Dumps the content of the transform stack to a string
+	 */
+	void DumpTransformStackToFile(FString* OutFilePath = nullptr);
+
+	/**
+	 * Tells this hierarchy to trace a series of frames
+	 */
+	void TraceFrames(int32 InNumFramesToTrace);
+	
+#endif
+
+private:
 
 	/**
 	 * Returns true if a given element is selected
@@ -2118,48 +2214,6 @@ private:
 	/// Undo redo related
 	///////////////////////////////////////////////
 
-	enum ETransformStackEntryType
-	{
-		TransformPose,
-		ControlOffset,
-		ControlGizmo,
-		CurveValue
-	};
-
-	struct FTransformStackEntry
-	{
-		FORCEINLINE FTransformStackEntry()
-			: Key()
-			, EntryType(ETransformStackEntryType::TransformPose)
-			, TransformType(ERigTransformType::CurrentLocal)
-			, OldTransform(FTransform::Identity)
-			, NewTransform(FTransform::Identity)
-			, bAffectChildren(true)
-		{}
-
-		FORCEINLINE FTransformStackEntry(
-			const FRigElementKey& InKey,
-			ETransformStackEntryType InEntryType,
-			ERigTransformType::Type InTransformType,
-			const FTransform& InOldTransform,
-			const FTransform& InNewTransform,
-			bool bInAffectChildren)
-			: Key(InKey)
-			, EntryType(InEntryType)
-			, TransformType(InTransformType)
-			, OldTransform(InOldTransform)
-			, NewTransform(InNewTransform)
-			, bAffectChildren(bInAffectChildren)
-		{}
-
-		FRigElementKey Key;
-		ETransformStackEntryType EntryType;
-		ERigTransformType::Type TransformType;
-		FTransform OldTransform;
-		FTransform NewTransform;
-		bool bAffectChildren;
-	};
-
 	/**
 	 * The index identifying where we stand with the stack
 	 */
@@ -2172,14 +2226,15 @@ private:
 	bool bTransactingForTransformChange;
 	
 	/**
-	 * The stack of actions to undo
+	 * The stack of actions to undo.
+	 * Note: This is also used when performing traces on the hierarchy.
 	 */
-	TArray<FTransformStackEntry> TransformUndoStack;
+	TArray<FRigTransformStackEntry> TransformUndoStack;
 
 	/**
 	 * The stack of actions to undo
 	 */
-	TArray<FTransformStackEntry> TransformRedoStack;
+	TArray<FRigTransformStackEntry> TransformRedoStack;
 
 	/**
 	 * Sets the transform stack index - which in turns performs a series of undos / redos
@@ -2192,11 +2247,12 @@ private:
 	 */
 	void PushTransformToStack(
 			const FRigElementKey& InKey,
-            ETransformStackEntryType InEntryType,
+            ERigTransformStackEntryType InEntryType,
             ERigTransformType::Type InTransformType,
             const FTransform& InOldTransform,
             const FTransform& InNewTransform,
-            bool bAffectChildren);
+            bool bAffectChildren,
+            bool bModify);
 
 	/**
 	 * Stores a curve value on the stack
@@ -2204,12 +2260,13 @@ private:
 	void PushCurveToStack(
             const FRigElementKey& InKey,
             float InOldCurveValue,
-            float InNewCurveValue);
+            float InNewCurveValue,
+            bool bModify);
 
 	/**
 	 * Restores a transform on the stack
 	 */
-	bool ApplyTransformFromStack(const FTransformStackEntry& InEntry, bool bUndo);
+	bool ApplyTransformFromStack(const FRigTransformStackEntry& InEntry, bool bUndo);
 
 	/**
 	 * Computes all parts of the pose
@@ -2239,6 +2296,16 @@ private:
 
 	int32 ResetPoseHash;
 	TArray<bool> ResetPoseHasFilteredChildren;
+
+#if WITH_EDITOR
+
+protected:
+	
+	int32 TraceFramesLeft;
+	int32 TraceFramesCaptured;
+	TMap<FName, FRigPose> TracePoses;
+
+#endif
 	
 	friend class URigHierarchyController;
 	friend class UControlRig;
diff --git a/Engine/Plugins/Experimental/ControlRig/Source/ControlRig/Public/Rigs/RigHierarchyElements.h b/Engine/Plugins/Experimental/ControlRig/Source/ControlRig/Public/Rigs/RigHierarchyElements.h
--- a/Engine/Plugins/Experimental/ControlRig/Source/ControlRig/Public/Rigs/RigHierarchyElements.h
+++ b/Engine/Plugins/Experimental/ControlRig/Source/ControlRig/Public/Rigs/RigHierarchyElements.h
@@ -24,17 +24,21 @@ friend FORCEINLINE T* Cast(ElementType* InElement) \
    return Cast<T>((FRigBaseElement*) InElement); \
 }
 
+UENUM()
 namespace ERigTransformType
 {
 	enum Type
 	{
 		InitialLocal,
-        CurrentLocal,
-        InitialGlobal,
-        CurrentGlobal,
+		CurrentLocal,
+		InitialGlobal,
+		CurrentGlobal,
 		NumTransformTypes
 	};
+}
 
+namespace ERigTransformType
+{
 	FORCEINLINE ERigTransformType::Type SwapCurrentAndInitial(const Type InTransformType)
 	{
 		switch(InTransformType)
