ControlRig: Adding local variables to CR graphs UI

diff --git a/Engine/Plugins/Experimental/ControlRig/Source/ControlRigDeveloper/Private/ControlRigBlueprint.cpp b/Engine/Plugins/Experimental/ControlRig/Source/ControlRigDeveloper/Private/ControlRigBlueprint.cpp
--- a/Engine/Plugins/Experimental/ControlRig/Source/ControlRigDeveloper/Private/ControlRigBlueprint.cpp
+++ b/Engine/Plugins/Experimental/ControlRig/Source/ControlRigDeveloper/Private/ControlRigBlueprint.cpp
@@ -1665,7 +1665,7 @@ FName UControlRigBlueprint::AddMemberVariable(const FName& InName, const FString
 		return NAME_None;
 	}
 
-	Variable.bIsArray = CPPType.StartsWith(TEXT("TArray<"));
+	Variable.bIsArray = RigVMUtilities::IsArrayType(CPPType);
 	if (Variable.bIsArray)
 	{
 		CPPType = CPPType.RightChop(7).LeftChop(1);
diff --git a/Engine/Plugins/Experimental/ControlRig/Source/ControlRigDeveloper/Private/Graph/ControlRigGraphSchema.cpp b/Engine/Plugins/Experimental/ControlRig/Source/ControlRigDeveloper/Private/Graph/ControlRigGraphSchema.cpp
--- a/Engine/Plugins/Experimental/ControlRig/Source/ControlRigDeveloper/Private/Graph/ControlRigGraphSchema.cpp
+++ b/Engine/Plugins/Experimental/ControlRig/Source/ControlRigDeveloper/Private/Graph/ControlRigGraphSchema.cpp
@@ -24,6 +24,8 @@
 #include "RigVMModel/Nodes/RigVMCollapseNode.h"
 #include "RigVMModel/Nodes/RigVMFunctionEntryNode.h"
 #include "RigVMModel/Nodes/RigVMFunctionReturnNode.h"
+#include "RigVMModel/RigVMVariableDescription.h"
+#include "Kismet2/Kismet2NameValidators.h"
 
 #if WITH_EDITOR
 #include "ControlRigEditor/Private/Editor/SControlRigFunctionLocalizationWidget.h"
@@ -34,9 +36,177 @@
 
 const FName UControlRigGraphSchema::GraphName_ControlRig(TEXT("Rig Graph"));
 
+FEdGraphPinType FControlRigGraphSchemaAction_LocalVar::GetPinType() const
+{
+	if (UControlRigGraph* Graph = Cast<UControlRigGraph>(GetVariableScope()))
+	{
+		for (FRigVMGraphVariableDescription Variable : Graph->GetModel()->GetLocalVariables())
+		{
+			if (Variable.Name == GetVariableName())
+			{
+				return Variable.ToPinType();
+			}
+		}
+	}
+
+	return FEdGraphPinType();
+}
+
+void FControlRigGraphSchemaAction_LocalVar::ChangeVariableType(const FEdGraphPinType& NewPinType)
+{
+	if (UControlRigGraph* Graph = Cast<UControlRigGraph>(GetVariableScope()))
+	{
+		FString NewCPPType;
+		UObject* NewCPPTypeObject = nullptr;
+		FRigVMGraphVariableDescription::CPPTypeFromPinType(NewPinType, NewCPPType, NewCPPTypeObject);
+		Graph->GetController()->SetLocalVariableType(GetVariableName(), NewCPPType, NewCPPTypeObject);
+	}
+}
+
+void FControlRigGraphSchemaAction_LocalVar::RenameVariable(const FName& NewName)
+{
+	const FName OldName = GetVariableName();
+	
+	if (UControlRigGraph* Graph = Cast<UControlRigGraph>(GetVariableScope()))
+	{
+		if (Graph->GetController()->RenameLocalVariable(OldName, NewName))
+		{
+			SetVariableInfo(NewName, GetVariableScope(), GetPinType().PinCategory == TEXT("bool"));		
+		}
+	}	
+}
+
+bool FControlRigGraphSchemaAction_LocalVar::IsValidName(const FName& NewName, FText& OutErrorMessage) const
+{
+	if (UControlRigGraph* ControlRigGraph = Cast<UControlRigGraph>(GetVariableScope()))
+	{
+		FKismetNameValidator NameValidator(ControlRigGraph->GetBlueprint());
+		EValidatorResult Result = NameValidator.IsValid(NewName.ToString());
+		if (Result != EValidatorResult::Ok)
+		{
+			OutErrorMessage = FText::FromString(TEXT("Name with invalid format"));
+			return false;
+		}
+		
+		for (const FRigVMGraphVariableDescription& LocalVar : ControlRigGraph->GetModel()->GetLocalVariables())
+		{
+			if (LocalVar.Name != GetVariableName() && LocalVar.Name == NewName)
+			{
+				OutErrorMessage = FText::FromString(TEXT("Name already in use"));
+				return false;
+			}
+		}
+	}
+	return true;
+}
+
+void FControlRigGraphSchemaAction_LocalVar::DeleteVariable() 
+{
+	if (UControlRigGraph* Graph = Cast<UControlRigGraph>(GetVariableScope()))
+	{
+		Graph->GetController()->RemoveLocalVariable(GetVariableName());
+	}
+}
+
+bool FControlRigGraphSchemaAction_LocalVar::IsVariableUsed()
+{
+	if (UControlRigGraph* ControlRigGraph = Cast<UControlRigGraph>(GetVariableScope()))
+	{
+		const FString VarNameStr = GetVariableName().ToString();
+		for (URigVMNode* Node : ControlRigGraph->GetModel()->GetNodes())
+		{
+			if (URigVMVariableNode* VarNode = Cast<URigVMVariableNode>(Node))
+			{
+				if (VarNode->FindPin(TEXT("Variable"))->GetDefaultValue() == VarNameStr)
+				{
+					return true;		
+				}
+			}
+		}
+	}
+	return false;
+}
+
+void FControlRigGraphSchemaAction_LocalVar::GetVariableTypeTree(
+	TArray<TSharedPtr<UEdGraphSchema_K2::FPinTypeTreeInfo>>& TypeTree, ETypeTreeFilter TypeTreeFilter) const
+{
+	const UEdGraphSchema_K2* Schema = GetDefault<UEdGraphSchema_K2>();
+	
+	// Clear the list
+	TypeTree.Empty();
+
+	TypeTree.Add( MakeShareable( new UEdGraphSchema_K2::FPinTypeTreeInfo(UEdGraphSchema_K2::GetCategoryText(UEdGraphSchema_K2::PC_Boolean, true), UEdGraphSchema_K2::PC_Boolean, Schema, LOCTEXT("BooleanType", "True or false value"))));
+	TypeTree.Add( MakeShareable( new UEdGraphSchema_K2::FPinTypeTreeInfo(UEdGraphSchema_K2::GetCategoryText(UEdGraphSchema_K2::PC_Int, true), UEdGraphSchema_K2::PC_Int, Schema, LOCTEXT("IntegerType", "Integer number")) ) );;
+	
+	TypeTree.Add(MakeShareable(new UEdGraphSchema_K2::FPinTypeTreeInfo(UEdGraphSchema_K2::GetCategoryText(UEdGraphSchema_K2::PC_Float, true), UEdGraphSchema_K2::PC_Float, Schema, LOCTEXT("FloatType", "Floating point number"))));
+	TypeTree.Add(MakeShareable(new UEdGraphSchema_K2::FPinTypeTreeInfo(UEdGraphSchema_K2::GetCategoryText(UEdGraphSchema_K2::PC_Name, true), UEdGraphSchema_K2::PC_Name, Schema, LOCTEXT("NameType", "A text name"))));
+	TypeTree.Add(MakeShareable(new UEdGraphSchema_K2::FPinTypeTreeInfo(UEdGraphSchema_K2::GetCategoryText(UEdGraphSchema_K2::PC_String, true), UEdGraphSchema_K2::PC_String, Schema, LOCTEXT("StringType", "A text string"))));
+
+	TypeTree.Add(MakeShareable(new UEdGraphSchema_K2::FPinTypeTreeInfo(UEdGraphSchema_K2::PC_Struct, TBaseStructure<FVector>::Get(), LOCTEXT("VectorType", "A 3D vector"))));
+	TypeTree.Add(MakeShareable(new UEdGraphSchema_K2::FPinTypeTreeInfo(UEdGraphSchema_K2::PC_Struct, TBaseStructure<FRotator>::Get(), LOCTEXT("RotatorType", "A 3D rotation"))));
+	TypeTree.Add(MakeShareable(new UEdGraphSchema_K2::FPinTypeTreeInfo(UEdGraphSchema_K2::PC_Struct, TBaseStructure<FTransform>::Get(), LOCTEXT("TransformType", "A 3D transformation, including translation, rotation and 3D scale."))));
+}
+
 FReply FControlRigFunctionDragDropAction::DroppedOnPanel(const TSharedRef< class SWidget >& Panel, FVector2D ScreenPosition, FVector2D GraphPosition, UEdGraph& Graph)
 {
-	if (UControlRigGraph* TargetRigGraph = Cast<UControlRigGraph>(&Graph))
+	// For local variables
+	if (SourceAction->GetTypeId() == FControlRigGraphSchemaAction_LocalVar::StaticGetTypeId())
+	{
+		if (UControlRigGraph* TargetRigGraph = Cast<UControlRigGraph>(&Graph))
+		{
+			if (TargetRigGraph == SourceRigGraph)
+			{
+				FControlRigGraphSchemaAction_LocalVar* VarAction = (FControlRigGraphSchemaAction_LocalVar*) SourceAction.Get();
+				for (FRigVMGraphVariableDescription LocalVariable : TargetRigGraph->GetModel()->GetLocalVariables())
+				{
+					if (LocalVariable.Name == VarAction->GetVariableName())
+					{
+						URigVMController* Controller = TargetRigGraph->GetController();
+						FMenuBuilder MenuBuilder(true, NULL);
+						const FText VariableNameText = FText::FromName( LocalVariable.Name );
+
+						MenuBuilder.BeginSection("BPVariableDroppedOn", VariableNameText );
+
+						MenuBuilder.AddMenuEntry(
+							FText::Format( LOCTEXT("CreateGetVariable", "Get {0}"), VariableNameText ),
+							FText::Format( LOCTEXT("CreateVariableGetterToolTip", "Create Getter for variable '{0}'\n(Ctrl-drag to automatically create a getter)"), VariableNameText ),
+							FSlateIcon(),
+							FUIAction(
+							FExecuteAction::CreateLambda([Controller, LocalVariable, GraphPosition]()
+							{
+								Controller->AddVariableNode(LocalVariable.Name, LocalVariable.CPPType, LocalVariable.CPPTypeObject, true, LocalVariable.DefaultValue, GraphPosition);
+							}),
+							FCanExecuteAction()));
+
+						MenuBuilder.AddMenuEntry(
+							FText::Format( LOCTEXT("CreateSetVariable", "Set {0}"), VariableNameText ),
+							FText::Format( LOCTEXT("CreateVariableSetterToolTip", "Create Setter for variable '{0}'\n(Alt-drag to automatically create a setter)"), VariableNameText ),
+							FSlateIcon(),
+							FUIAction(
+							FExecuteAction::CreateLambda([Controller, LocalVariable, GraphPosition]()
+							{
+								Controller->AddVariableNode(LocalVariable.Name, LocalVariable.CPPType, LocalVariable.CPPTypeObject, false, LocalVariable.DefaultValue, GraphPosition);
+							}),
+							FCanExecuteAction()));
+
+						TSharedRef< SWidget > PanelWidget = Panel;
+						// Show dialog to choose getter vs setter
+						FSlateApplication::Get().PushMenu(
+							PanelWidget,
+							FWidgetPath(),
+							MenuBuilder.MakeWidget(),
+							ScreenPosition,
+							FPopupTransitionEffect( FPopupTransitionEffect::ContextMenu)
+							);
+
+						MenuBuilder.EndSection();
+					}	
+				}				
+			}
+		}
+	}
+	// For functions
+	else if (UControlRigGraph* TargetRigGraph = Cast<UControlRigGraph>(&Graph))
 	{
 		if (UControlRigBlueprint* TargetRigBlueprint = Cast<UControlRigBlueprint>(FBlueprintEditorUtils::FindBlueprintForGraph(TargetRigGraph)))
 		{
@@ -68,6 +238,9 @@ FReply FControlRigFunctionDragDropAction::DroppedOnPanel(const TSharedRef< class
 			}
 		}
 	}
+
+	
+	
 	return FReply::Unhandled();
 }
 
@@ -408,11 +581,19 @@ bool UControlRigGraphSchema::CanGraphBeDropped(TSharedPtr<FEdGraphSchemaAction>
 			return true;
 		}
 	}
+	else if (InAction->GetTypeId() == FControlRigGraphSchemaAction_LocalVar::StaticGetTypeId())
+	{
+		FControlRigGraphSchemaAction_LocalVar* VarAction = (FControlRigGraphSchemaAction_LocalVar*)InAction.Get();
+		if (UControlRigGraph* RigGraph = Cast<UControlRigGraph>((UEdGraph*)VarAction->GetVariableScope()))
+		{
+			return true;
+		}
+	}
 	
 	return false;
 }
 
-FReply UControlRigGraphSchema::BeginGraphDragAction(TSharedPtr<FEdGraphSchemaAction> InAction) const
+FReply UControlRigGraphSchema::BeginGraphDragAction(TSharedPtr<FEdGraphSchemaAction> InAction, const FPointerEvent& MouseEvent) const
 {
 	if (!InAction.IsValid())
 	{
@@ -422,13 +603,30 @@ FReply UControlRigGraphSchema::BeginGraphDragAction(TSharedPtr<FEdGraphSchemaAct
 	if (InAction->GetTypeId() == FEdGraphSchemaAction_K2Graph::StaticGetTypeId())
 	{
 		FEdGraphSchemaAction_K2Graph* FuncAction = (FEdGraphSchemaAction_K2Graph*)InAction.Get();
-if (UControlRigGraph* RigGraph = Cast<UControlRigGraph>(FuncAction->EdGraph))
-{
-	if (UControlRigBlueprint* RigBlueprint = Cast<UControlRigBlueprint>(FBlueprintEditorUtils::FindBlueprintForGraph(RigGraph)))
-	{
-		return FReply::Handled().BeginDragDrop(FControlRigFunctionDragDropAction::New(InAction, RigBlueprint, RigGraph));
+		if (UControlRigGraph* RigGraph = Cast<UControlRigGraph>(FuncAction->EdGraph))
+		{
+			if (UControlRigBlueprint* RigBlueprint = Cast<UControlRigBlueprint>(FBlueprintEditorUtils::FindBlueprintForGraph(RigGraph)))
+			{
+				TSharedRef<FControlRigFunctionDragDropAction> Action = FControlRigFunctionDragDropAction::New(InAction, RigBlueprint, RigGraph);
+				Action->SetAltDrag(MouseEvent.IsAltDown());
+				Action->SetCtrlDrag(MouseEvent.IsControlDown());
+				return FReply::Handled().BeginDragDrop(Action);
+			}
+		}
 	}
-}
+	else if(InAction->GetTypeId() == FControlRigGraphSchemaAction_LocalVar::StaticGetTypeId())
+	{
+		FControlRigGraphSchemaAction_LocalVar* VarAction = (FControlRigGraphSchemaAction_LocalVar*)InAction.Get();
+		if (UControlRigGraph* RigGraph = Cast<UControlRigGraph>(VarAction->GetVariableScope()))
+		{
+			if (UControlRigBlueprint* RigBlueprint = Cast<UControlRigBlueprint>(FBlueprintEditorUtils::FindBlueprintForGraph(RigGraph)))
+			{
+				TSharedRef<FControlRigFunctionDragDropAction> Action = FControlRigFunctionDragDropAction::New(InAction, RigBlueprint, RigGraph);
+				Action->SetAltDrag(MouseEvent.IsAltDown());
+				Action->SetCtrlDrag(MouseEvent.IsControlDown());
+				return FReply::Handled().BeginDragDrop(Action);
+			}
+		}
 	}
 	return FReply::Unhandled();
 }
@@ -521,6 +719,47 @@ void UControlRigGraphSchema::GetGraphDisplayInformation(const UEdGraph& Graph, /
 	}
 }
 
+bool UControlRigGraphSchema::GetLocalVariables(const UEdGraph* InGraph, TArray<FBPVariableDescription>& OutLocalVariables) const
+{
+	OutLocalVariables.Reset();
+	if (UControlRigGraph* RigGraph = Cast<UControlRigGraph>((UEdGraph*)InGraph))
+	{
+		if (URigVMGraph* Model = RigGraph->GetModel())
+		{
+			TArray<FRigVMGraphVariableDescription> LocalVariables = Model->GetLocalVariables();
+			for (FRigVMGraphVariableDescription LocalVariable : LocalVariables)
+			{
+				FBPVariableDescription VariableDescription;
+				VariableDescription.VarName = LocalVariable.Name;
+				VariableDescription.FriendlyName = LocalVariable.Name.ToString();
+				VariableDescription.DefaultValue = LocalVariable.DefaultValue;
+				VariableDescription.VarType = LocalVariable.ToPinType();
+				VariableDescription.PropertyFlags |= CPF_BlueprintVisible;
+				OutLocalVariables.Add(VariableDescription);
+			}
+		}
+	}
+	return true;
+}
+
+TSharedPtr<FEdGraphSchemaAction> UControlRigGraphSchema::MakeActionFromVariableDescription(const UEdGraph* InEdGraph,
+	const FBPVariableDescription& Variable) const
+{
+	if (UControlRigGraph* RigGraph = Cast<UControlRigGraph>((UEdGraph*)InEdGraph))
+	{
+		FText Category = Variable.Category;
+		if (Variable.Category.EqualTo(UEdGraphSchema_K2::VR_DefaultCategory))
+		{
+			Category = FText::GetEmpty();
+		}
+
+		TSharedPtr<FControlRigGraphSchemaAction_LocalVar> Action = MakeShareable(new FControlRigGraphSchemaAction_LocalVar(Category, FText::FromName(Variable.VarName), FText::GetEmpty(), 0, NodeSectionID::LOCAL_VARIABLE));
+		Action->SetVariableInfo(Variable.VarName, RigGraph, Variable.VarType.PinCategory == UEdGraphSchema_K2::PC_Boolean);
+		return Action;
+	}
+	return nullptr;
+}
+
 FText UControlRigGraphSchema::GetGraphCategory(const UEdGraph* InGraph) const
 {
 	if (UControlRigGraph* RigGraph = Cast<UControlRigGraph>((UEdGraph*)InGraph))
diff --git a/Engine/Plugins/Experimental/ControlRig/Source/ControlRigDeveloper/Public/ControlRigBlueprint.h b/Engine/Plugins/Experimental/ControlRig/Source/ControlRigDeveloper/Public/ControlRigBlueprint.h
--- a/Engine/Plugins/Experimental/ControlRig/Source/ControlRigDeveloper/Public/ControlRigBlueprint.h
+++ b/Engine/Plugins/Experimental/ControlRig/Source/ControlRigDeveloper/Public/ControlRigBlueprint.h
@@ -177,7 +177,7 @@ public:
 	virtual void ReplaceDeprecatedNodes() override;
 
 	virtual bool SupportsGlobalVariables() const override { return true; }
-	virtual bool SupportsLocalVariables() const override { return false; }
+	virtual bool SupportsLocalVariables() const override { return true; }
 	virtual bool SupportsFunctions() const override { return true; }
 	virtual bool SupportsMacros() const override { return false; }
 	virtual bool SupportsDelegates() const override { return false; }
diff --git a/Engine/Plugins/Experimental/ControlRig/Source/ControlRigDeveloper/Public/Graph/ControlRigGraphSchema.h b/Engine/Plugins/Experimental/ControlRig/Source/ControlRigDeveloper/Public/Graph/ControlRigGraphSchema.h
--- a/Engine/Plugins/Experimental/ControlRig/Source/ControlRigDeveloper/Public/Graph/ControlRigGraphSchema.h
+++ b/Engine/Plugins/Experimental/ControlRig/Source/ControlRigDeveloper/Public/Graph/ControlRigGraphSchema.h
@@ -4,6 +4,9 @@
 
 #include "EdGraph/EdGraphSchema.h"
 #include "GraphEditorDragDropAction.h"
+#include "RigVMModel/RigVMGraph.h"
+#include "EdGraphSchema_K2_Actions.h"
+
 #include "ControlRigGraphSchema.generated.h"
 
 class UControlRigBlueprint;
@@ -40,6 +43,46 @@ struct FControlRigPinConnectionResponse
 	ECanCreateConnectionResponse_Extended ExtendedResponse;
 };
 
+USTRUCT()
+struct CONTROLRIGDEVELOPER_API FControlRigGraphSchemaAction_LocalVar : public FEdGraphSchemaAction_BlueprintVariableBase
+{
+	GENERATED_BODY()
+
+	
+public:
+
+	// Simple type info
+	static FName StaticGetTypeId() {static FName Type("FControlRigGraphSchemaAction_LocalVar"); return Type;}
+	virtual FName GetTypeId() const override { return StaticGetTypeId(); } 
+
+	FControlRigGraphSchemaAction_LocalVar()
+		: FEdGraphSchemaAction_BlueprintVariableBase()
+	{}
+
+	FControlRigGraphSchemaAction_LocalVar(FText InNodeCategory, FText InMenuDesc, FText InToolTip, const int32 InGrouping, const int32 InSectionID)
+		: FEdGraphSchemaAction_BlueprintVariableBase(MoveTemp(InNodeCategory), MoveTemp(InMenuDesc), MoveTemp(InToolTip), InGrouping, InSectionID)
+	{}
+
+	virtual bool IsA(const FName& InType) const override
+	{
+		return InType == GetTypeId() || InType == FEdGraphSchemaAction_BlueprintVariableBase::StaticGetTypeId();
+	}
+
+	virtual FEdGraphPinType GetPinType() const override;
+
+	virtual void ChangeVariableType(const FEdGraphPinType& NewPinType) override;
+
+	virtual void RenameVariable(const FName& NewName) override;
+
+	virtual bool IsValidName(const FName& NewName, FText& OutErrorMessage) const override;
+
+	virtual void DeleteVariable() override;
+
+	virtual bool IsVariableUsed() override;
+
+	virtual void GetVariableTypeTree(TArray< TSharedPtr<UEdGraphSchema_K2::FPinTypeTreeInfo> >& TypeTree, ETypeTreeFilter TypeTreeFilter) const override;
+};
+
 /** DragDropAction class for drag and dropping an item from the My Blueprints tree (e.g., variable or function) */
 class CONTROLRIGDEVELOPER_API FControlRigFunctionDragDropAction : public FGraphSchemaActionDragDropAction
 {
@@ -99,7 +142,7 @@ public:
 	virtual void BreakPinLinks(UEdGraphPin& TargetPin, bool bSendsNodeNotifcation) const override;
 	virtual void BreakSinglePinLink(UEdGraphPin* SourcePin, UEdGraphPin* TargetPin) const override;
 	virtual bool CanGraphBeDropped(TSharedPtr<FEdGraphSchemaAction> InAction) const override;
-	virtual FReply BeginGraphDragAction(TSharedPtr<FEdGraphSchemaAction> InAction) const override;
+	virtual FReply BeginGraphDragAction(TSharedPtr<FEdGraphSchemaAction> InAction, const FPointerEvent& MouseEvent = FPointerEvent() ) const override;
 	virtual class FConnectionDrawingPolicy* CreateConnectionDrawingPolicy(int32 InBackLayerID, int32 InFrontLayerID, float InZoomFactor, const FSlateRect& InClippingRect, class FSlateWindowElementList& InDrawElements, class UEdGraph* InGraphObj) const override;
 	virtual bool ShouldHidePinDefaultValue(UEdGraphPin* Pin) const override;
 	virtual void TrySetDefaultValue(UEdGraphPin& InPin, const FString& InNewDefaultValue, bool bMarkAsModified = true) const override;
@@ -119,6 +162,8 @@ public:
 	virtual bool IsStructEditable(UStruct* InStruct) const;
 	virtual void SetNodePosition(UEdGraphNode* Node, const FVector2D& Position) const override;
 	virtual void GetGraphDisplayInformation(const UEdGraph& Graph, /*out*/ FGraphDisplayInfo& DisplayInfo) const override;
+	virtual bool GetLocalVariables(const UEdGraph* InGraph, TArray<FBPVariableDescription>& OutLocalVariables) const override;
+	virtual TSharedPtr<FEdGraphSchemaAction> MakeActionFromVariableDescription(const UEdGraph* InEdGraph, const FBPVariableDescription& Variable) const override;
 	virtual FText GetGraphCategory(const UEdGraph* InGraph) const override;
 	virtual FReply TrySetGraphCategory(const UEdGraph* InGraph, const FText& InCategory) override;
 	virtual bool TryDeleteGraph(UEdGraph* GraphToDelete) const override;
diff --git a/Engine/Plugins/Experimental/ControlRig/Source/ControlRigEditor/Private/Editor/ControlRigEditor.cpp b/Engine/Plugins/Experimental/ControlRig/Source/ControlRigEditor/Private/Editor/ControlRigEditor.cpp
--- a/Engine/Plugins/Experimental/ControlRig/Source/ControlRigEditor/Private/Editor/ControlRigEditor.cpp
+++ b/Engine/Plugins/Experimental/ControlRig/Source/ControlRigEditor/Private/Editor/ControlRigEditor.cpp
@@ -1845,6 +1845,24 @@ bool FControlRigEditor::TransactionObjectAffectsBlueprint(UObject* InTransactedO
 	return FBlueprintEditor::TransactionObjectAffectsBlueprint(InTransactedObject);
 }
 
+void FControlRigEditor::OnAddNewLocalVariable()
+{
+	if (!CanAddNewLocalVariable())
+	{
+		return;
+	}
+
+	FRigVMGraphVariableDescription NewVar = GetFocusedController()->AddLocalVariable(TEXT("NewLocalVar"), TEXT("bool"), nullptr, TEXT("0"));
+	if(NewVar.Name.IsNone())
+	{
+		LogSimpleMessage( LOCTEXT("AddLocalVariable_Error", "Adding new local variable failed.") );
+	}
+	else
+	{
+		RenameNewlyAddedAction(NewVar.Name);
+	}
+}
+
 void FControlRigEditor::DeleteSelectedNodes()
 {
 	DECLARE_SCOPE_HIERARCHICAL_COUNTER_FUNC()
@@ -2186,6 +2204,7 @@ bool FControlRigEditor::IsSectionVisible(NodeSectionID::Type InSectionID) const
 		case NodeSectionID::GRAPH:
 		case NodeSectionID::VARIABLE:
 		case NodeSectionID::FUNCTION:
+		case NodeSectionID::LOCAL_VARIABLE:
 		{
 			return true;
 		}
diff --git a/Engine/Plugins/Experimental/ControlRig/Source/ControlRigEditor/Private/Editor/ControlRigEditor.h b/Engine/Plugins/Experimental/ControlRig/Source/ControlRigEditor/Private/Editor/ControlRigEditor.h
--- a/Engine/Plugins/Experimental/ControlRig/Source/ControlRigEditor/Private/Editor/ControlRigEditor.h
+++ b/Engine/Plugins/Experimental/ControlRig/Source/ControlRigEditor/Private/Editor/ControlRigEditor.h
@@ -104,7 +104,9 @@ public:
 
 	// BlueprintEditor interface
 	virtual bool TransactionObjectAffectsBlueprint(UObject* InTransactedObject) override;
-	virtual bool CanAddNewLocalVariable() const override { return false; }
+	virtual bool CanAddNewLocalVariable() const override { return true; }
+	virtual void OnAddNewLocalVariable() override;
+	
 	virtual void DeleteSelectedNodes() override;
 	virtual bool CanDeleteNodes() const override;
 
diff --git a/Engine/Source/Developer/RigVMDeveloper/Private/RigVMCompiler/RigVMCompiler.cpp b/Engine/Source/Developer/RigVMDeveloper/Private/RigVMCompiler/RigVMCompiler.cpp
--- a/Engine/Source/Developer/RigVMDeveloper/Private/RigVMCompiler/RigVMCompiler.cpp
+++ b/Engine/Source/Developer/RigVMDeveloper/Private/RigVMCompiler/RigVMCompiler.cpp
@@ -1236,7 +1236,7 @@ void URigVMCompiler::InitializeLocalVariables(const FRigVMExprAST* InExpr, FRigV
 						for (FRigVMGraphVariableDescription Variable : Node->GetContainedGraph()->LocalVariables)
 						{
 							FString TargetPath = FString::Printf(TEXT("LocalVariable::%s|%s"), *Node->GetNodePath(), *Variable.Name.ToString());
-							FString SourcePath = FString::Printf(TEXT("LocalVariable::%s|%s::Const"), *Node->GetContainedGraph()->GetGraphName(), *Variable.Name.ToString());
+							FString SourcePath = FString::Printf(TEXT("LocalVariableDefault::%s|%s::Const"), *Node->GetContainedGraph()->GetGraphName(), *Variable.Name.ToString());
 							FRigVMOperand* TargetPtr = WorkData.PinPathToOperand->Find(TargetPath);
 							FRigVMOperand* SourcePtr = WorkData.PinPathToOperand->Find(SourcePath);
 							if (SourcePtr && TargetPtr) 
@@ -1282,7 +1282,7 @@ void URigVMCompiler::InitializeLocalVariables(const FRigVMExprAST* InExpr, FRigV
 					for (FRigVMGraphVariableDescription Variable : Node->GetGraph()->LocalVariables)
 					{
 						FString TargetPath = FString::Printf(TEXT("LocalVariable::%s"), *Variable.Name.ToString());
-						FString SourcePath = FString::Printf(TEXT("LocalVariable::|%s::Const"), *Variable.Name.ToString());
+						FString SourcePath = FString::Printf(TEXT("LocalVariableDefault::|%s::Const"), *Variable.Name.ToString());
 						FRigVMOperand* TargetPtr = WorkData.PinPathToOperand->Find(TargetPath);
 						FRigVMOperand* SourcePtr = WorkData.PinPathToOperand->Find(SourcePath);
 						if (SourcePtr && TargetPtr)
@@ -1362,7 +1362,14 @@ FString URigVMCompiler::GetPinHash(const URigVMPin* InPin, const FRigVMVarExprAS
 				if (bIsLiteral)
 				{
 					// Literal values will be reused for all instance of local variables
-					return FString::Printf(TEXT("%sLocalVariable::%s|%s%s"), *Prefix, *Node->GetGraph()->GetGraphName(), *VariableName.ToString(), *Suffix);
+					if (InVarExpr->NumParents() == 0 && InVarExpr->NumChildren() == 0)
+					{
+						return FString::Printf(TEXT("%sLocalVariableDefault::%s|%s%s"), *Prefix, *Node->GetGraph()->GetGraphName(), *VariableName.ToString(), *Suffix);
+					}
+					else
+					{
+						return FString::Printf(TEXT("%sLocalVariable::%s|%s%s"), *Prefix, *Node->GetGraph()->GetGraphName(), *VariableName.ToString(), *Suffix);					
+					}			
 				}
 				else
 				{
diff --git a/Engine/Source/Developer/RigVMDeveloper/Private/RigVMModel/RigVMController.cpp b/Engine/Source/Developer/RigVMDeveloper/Private/RigVMModel/RigVMController.cpp
--- a/Engine/Source/Developer/RigVMDeveloper/Private/RigVMModel/RigVMController.cpp
+++ b/Engine/Source/Developer/RigVMDeveloper/Private/RigVMModel/RigVMController.cpp
@@ -615,7 +615,7 @@ void URigVMController::RefreshVariableNode(const FName& InNodeName, const FName&
 			{
 				if (URigVMPin* ValuePin = VariableNode->FindPin(URigVMVariableNode::ValueName))
 				{
-					if (ValuePin->CPPType != InCPPType)
+					if (ValuePin->CPPType != InCPPType || ValuePin->GetCPPTypeObject() != InCPPTypeObject)
 					{
 						if (bSetupUndoRedo)
 						{
@@ -6644,39 +6644,35 @@ bool URigVMController::RemoveFunctionFromLibrary(const FName& InFunctionName, bo
 	return RemoveNodeByName(InFunctionName, bSetupUndoRedo);
 }
 
-bool URigVMController::AddLocalVariable(const FName& InVariableName,
+FRigVMGraphVariableDescription URigVMController::AddLocalVariable(const FName& InVariableName,
 	const FString& InCPPType, UObject* InCPPTypeObject, const FString& InDefaultValue, bool bSetupUndoRedo)
 {
+	FRigVMGraphVariableDescription NewVariable;
 	if (!IsValidGraph())
 	{
-		return false;
+		return NewVariable;
 	}
 	
 	URigVMGraph* Graph = GetGraph();
 	check(Graph);
 
-	if (Graph->GetParentGraph() && !Graph->GetParentGraph()->IsA<URigVMFunctionLibrary>())
-	{
-		ReportError(TEXT("Can only add local variables to root graphs."));
-		return false;
-	}
-
-	TArray<FRigVMGraphVariableDescription>& LocalVariables = Graph->LocalVariables;
-	for (FRigVMGraphVariableDescription& Variable : LocalVariables)
-	{
-		if (Variable.Name == InVariableName)
+	FName VariableName = GetUniqueName(InVariableName, [Graph](const FName& InName) {
+		for (FRigVMGraphVariableDescription LocalVariable : Graph->GetLocalVariables())
 		{
-			return false;
+			if (LocalVariable.Name == InName)
+			{
+				return false;
+			}
 		}
-	}
+		return true;
+	});
 
-	FRigVMGraphVariableDescription NewVariable;
-	NewVariable.Name = InVariableName;
+	NewVariable.Name = VariableName;
 	NewVariable.CPPType = InCPPType;
 	NewVariable.CPPTypeObject = InCPPTypeObject;
 	NewVariable.DefaultValue = InDefaultValue;
 
-	LocalVariables.Add(NewVariable);
+	Graph->LocalVariables.Add(NewVariable);
 
 	if (bSetupUndoRedo)
 	{
@@ -6693,7 +6689,7 @@ bool URigVMController::AddLocalVariable(const FName& InVariableName,
 		Graph->MarkPackageDirty();
 	}
 
-	return true;
+	return NewVariable;
 }
 
 bool URigVMController::RemoveLocalVariable(const FName& InVariableName, bool bSetupUndoRedo)
@@ -6741,6 +6737,134 @@ bool URigVMController::RemoveLocalVariable(const FName& InVariableName, bool bSe
 	return false;
 }
 
+bool URigVMController::RenameLocalVariable(const FName& InVariableName, const FName& InNewVariableName,
+	bool bSetupUndoRedo)
+{
+	if (!IsValidGraph())
+	{
+		return false;
+	}
+
+	URigVMGraph* Graph = GetGraph();
+	check(Graph);
+
+	TArray<FRigVMGraphVariableDescription>& LocalVariables = Graph->LocalVariables;
+	int32 FoundIndex = INDEX_NONE;
+	for (int32 Index = 0; Index < LocalVariables.Num(); ++Index)
+	{
+		if (LocalVariables[Index].Name == InVariableName)
+		{
+			FoundIndex = Index;
+			break;
+		}
+	}
+
+	if (FoundIndex == INDEX_NONE)
+	{
+		return false;
+	}
+
+	for (int32 Index = 0; Index < LocalVariables.Num(); ++Index)
+	{
+		if (LocalVariables[Index].Name == InNewVariableName)
+		{
+			return false;
+		}
+	}
+	
+	if (!bSuspendNotifications)
+	{
+		Graph->MarkPackageDirty();
+	}
+	
+	if (bSetupUndoRedo)
+	{
+		FRigVMInverseAction InverseAction;
+		InverseAction.Title = FString::Printf(TEXT("Rename Local Variable %s to %s"), *InVariableName.ToString(), *InNewVariableName.ToString());
+
+		ActionStack->BeginAction(InverseAction);
+		ActionStack->AddAction(FRigVMRenameLocalVariableAction(LocalVariables[FoundIndex].Name, InNewVariableName));
+		ActionStack->EndAction(InverseAction);
+	}	
+	
+	LocalVariables[FoundIndex].Name = InNewVariableName;
+	return true;
+}
+
+bool URigVMController::SetLocalVariableType(const FName& InVariableName, const FString& InCPPType,
+	UObject* InCPPTypeObject, bool bSetupUndoRedo)
+{
+	if (!IsValidGraph())
+	{
+		return false;
+	}
+
+	URigVMGraph* Graph = GetGraph();
+	check(Graph);
+
+	TArray<FRigVMGraphVariableDescription>& LocalVariables = Graph->LocalVariables;
+	int32 FoundIndex = INDEX_NONE;
+	for (int32 Index = 0; Index < LocalVariables.Num(); ++Index)
+	{
+		if (LocalVariables[Index].Name == InVariableName)
+		{
+			FoundIndex = Index;
+			break;
+		}
+	}
+
+	if (FoundIndex == INDEX_NONE)
+	{
+		return false;
+	}
+
+	if (!bSuspendNotifications)
+	{
+		Graph->MarkPackageDirty();
+	}
+	
+	if (bSetupUndoRedo)
+	{
+		FRigVMInverseAction InverseAction;
+		InverseAction.Title = FString::Printf(TEXT("Change Local Variable type %s to %s"), *InVariableName.ToString(), *InCPPType);
+
+		ActionStack->BeginAction(InverseAction);
+		ActionStack->AddAction(FRigVMChangeLocalVariableTypeAction(LocalVariables[FoundIndex], InCPPType, InCPPTypeObject));
+		ActionStack->EndAction(InverseAction);
+	}	
+	
+	LocalVariables[FoundIndex].CPPType = InCPPType;
+	LocalVariables[FoundIndex].CPPTypeObject = InCPPTypeObject;
+
+	// Set default value
+	if (UScriptStruct* ScriptStruct = Cast<UScriptStruct>(InCPPTypeObject))
+	{
+		FString DefaultValue;
+		CreateDefaultValueForStructIfRequired(ScriptStruct, DefaultValue);
+		LocalVariables[FoundIndex].DefaultValue = DefaultValue;
+	}
+
+
+	// Change pin types on variable nodes
+	TArray<URigVMNode*> Nodes = Graph->GetNodes();
+	for (URigVMNode* Node : Nodes)
+	{
+		if (URigVMVariableNode* VariableNode = Cast<URigVMVariableNode>(Node))
+		{
+			if (URigVMPin* VariablePin = VariableNode->FindPin(URigVMVariableNode::VariableName))
+			{
+				if (VariablePin->GetDefaultValue() == InVariableName.ToString())
+				{
+					RefreshVariableNode(Node->GetFName(), InVariableName, InCPPType, InCPPTypeObject, bSetupUndoRedo);
+					continue;
+				}
+			}
+		}
+	}
+	
+	return true;
+}
+
 TArray<TSoftObjectPtr<URigVMFunctionReferenceNode>> URigVMController::GetAffectedReferences(ERigVMControllerBulkEditType InEditType, bool bForceLoad, bool bNotify)
 {
 	TArray<TSoftObjectPtr<URigVMFunctionReferenceNode>> FunctionReferencePtrs;
@@ -9304,17 +9428,17 @@ bool URigVMController::ChangePinType(const FString& InPinPath, const FString& In
 
 bool URigVMController::ChangePinType(URigVMPin* InPin, const FString& InCPPType, const FName& InCPPTypeObjectPath, bool bSetupUndoRedo, bool bSetupOrphanPins)
 {
-	if (InPin->CPPType == InCPPType)
+	if (InCPPType == TEXT("None") || InCPPType.IsEmpty())
 	{
 		return false;
 	}
 
-	if (InCPPType == TEXT("None") || InCPPType.IsEmpty())
+	UObject* CPPTypeObject = URigVMPin::FindObjectFromCPPTypeObjectPath(InCPPTypeObjectPath.ToString());
+	if (InPin->CPPType == InCPPType && InPin->CPPTypeObject == CPPTypeObject)
 	{
 		return false;
 	}
 
-	UObject* CPPTypeObject = URigVMPin::FindObjectFromCPPTypeObjectPath(InCPPTypeObjectPath.ToString());
 	if (CPPTypeObject)
 	{
 		// for now we don't allow UObjects
diff --git a/Engine/Source/Developer/RigVMDeveloper/Private/RigVMModel/RigVMControllerActions.cpp b/Engine/Source/Developer/RigVMDeveloper/Private/RigVMModel/RigVMControllerActions.cpp
--- a/Engine/Source/Developer/RigVMDeveloper/Private/RigVMModel/RigVMControllerActions.cpp
+++ b/Engine/Source/Developer/RigVMDeveloper/Private/RigVMModel/RigVMControllerActions.cpp
@@ -1748,7 +1748,7 @@ bool FRigVMAddLocalVariableAction::Redo(URigVMController* InController)
 {
 	if (!LocalVariable.Name.IsNone())
 	{
-		return InController->AddLocalVariable(LocalVariable.Name, LocalVariable.CPPType, LocalVariable.CPPTypeObject, LocalVariable.DefaultValue);
+		return InController->AddLocalVariable(LocalVariable.Name, LocalVariable.CPPType, LocalVariable.CPPTypeObject, LocalVariable.DefaultValue).Name.IsNone() == false;
 	}
 	return FRigVMBaseAction::Redo(InController);
 }
@@ -1764,7 +1764,7 @@ bool FRigVMRemoveLocalVariableAction::Undo(URigVMController* InController)
 	{
 		return false;
 	}
-	return InController->AddLocalVariable(LocalVariable.Name, LocalVariable.CPPType, LocalVariable.CPPTypeObject, LocalVariable.DefaultValue, false);
+	return InController->AddLocalVariable(LocalVariable.Name, LocalVariable.CPPType, LocalVariable.CPPTypeObject, LocalVariable.DefaultValue, false).Name.IsNone() != false;
 }
 
 bool FRigVMRemoveLocalVariableAction::Redo(URigVMController* InController)
@@ -1774,4 +1774,59 @@ bool FRigVMRemoveLocalVariableAction::Redo(URigVMController* InController)
 		return InController->RemoveLocalVariable(LocalVariable.Name, false);
 	}
 	return FRigVMBaseAction::Redo(InController);
-}
\ No newline at end of file
+}
+
+FRigVMRenameLocalVariableAction::FRigVMRenameLocalVariableAction(const FName& InOldName, const FName& InNewName)
+	: OldVariableName(InOldName), NewVariableName(InNewName)
+{
+	
+}
+
+bool FRigVMRenameLocalVariableAction::Undo(URigVMController* InController)
+{
+	if (!FRigVMBaseAction::Undo(InController))
+	{
+		return false;
+	}
+	return InController->RenameLocalVariable(NewVariableName, OldVariableName, false);
+}
+
+bool FRigVMRenameLocalVariableAction::Redo(URigVMController* InController)
+{
+	if (!InController->RenameLocalVariable(OldVariableName, NewVariableName, false))
+	{
+		return false;
+	}
+	return FRigVMBaseAction::Redo(InController);
+}
+
+FRigVMChangeLocalVariableTypeAction::FRigVMChangeLocalVariableTypeAction()
+	: LocalVariable()
+	, CPPType()
+	, CPPTypeObject(nullptr)
+{
+}
+
+FRigVMChangeLocalVariableTypeAction::FRigVMChangeLocalVariableTypeAction(
+	const FRigVMGraphVariableDescription& InLocalVariable, const FString& InCPPType, UObject* InCPPTypeObject)
+		: LocalVariable(InLocalVariable), CPPType(InCPPType), CPPTypeObject(InCPPTypeObject)
+{
+}
+
+bool FRigVMChangeLocalVariableTypeAction::Undo(URigVMController* InController)
+{
+	if (!FRigVMBaseAction::Undo(InController))
+	{
+		return false;
+	}
+	return InController->SetLocalVariableType(LocalVariable.Name, LocalVariable.CPPType, LocalVariable.CPPTypeObject, false);
+}
+
+bool FRigVMChangeLocalVariableTypeAction::Redo(URigVMController* InController)
+{
+	if (!InController->SetLocalVariableType(LocalVariable.Name, CPPType, CPPTypeObject, false))
+	{
+		return false;
+	}
+	return FRigVMBaseAction::Redo(InController);
+}
diff --git a/Engine/Source/Developer/RigVMDeveloper/Private/RigVMModel/RigVMGraph.cpp b/Engine/Source/Developer/RigVMDeveloper/Private/RigVMModel/RigVMGraph.cpp
--- a/Engine/Source/Developer/RigVMDeveloper/Private/RigVMModel/RigVMGraph.cpp
+++ b/Engine/Source/Developer/RigVMDeveloper/Private/RigVMModel/RigVMGraph.cpp
@@ -242,40 +242,6 @@ void URigVMGraph::SetDefaultFunctionLibrary(URigVMFunctionLibrary* InFunctionLib
 	DefaultFunctionLibraryPtr = InFunctionLibrary;
 }
 
-bool URigVMGraph::AddLocalVariable(const FRigVMGraphVariableDescription& NewVar)
-{
-	for (FRigVMGraphVariableDescription Variable : LocalVariables)
-	{
-		if (Variable.Name == NewVar.Name)
-		{
-			return false;
-		}
-	}
-	
-	LocalVariables.Add(NewVar);
-	return true;
-}
-
-bool URigVMGraph::RemoveLocalVariable(const FName& InVariableName)
-{
-	int32 FoundIndex = INDEX_NONE;
-	for (int32 Index = 0; Index < LocalVariables.Num(); ++Index)
-	{
-		if (LocalVariables[Index].Name == InVariableName)
-		{
-			FoundIndex = Index;
-			break;
-		}
-	}
-
-	if (FoundIndex != INDEX_NONE)
-	{
-		LocalVariables.RemoveAt(FoundIndex);
-		return true;
-	}
-	return false;
-}
-
 TArray<FRigVMExternalVariable> URigVMGraph::GetExternalVariables() const
 {
 	TArray<FRigVMExternalVariable> Variables;
diff --git a/Engine/Source/Developer/RigVMDeveloper/Private/RigVMModel/RigVMPin.cpp b/Engine/Source/Developer/RigVMDeveloper/Private/RigVMModel/RigVMPin.cpp
--- a/Engine/Source/Developer/RigVMDeveloper/Private/RigVMModel/RigVMPin.cpp
+++ b/Engine/Source/Developer/RigVMDeveloper/Private/RigVMModel/RigVMPin.cpp
@@ -837,7 +837,7 @@ FRigVMExternalVariable URigVMPin::ToExternalVariable() const
 	}
 	ExternalVariable.Name = *VariableName;
 
-	if (CPPType.StartsWith(TEXT("TArray<")))
+	if (RigVMUtilities::IsArrayType(CPPType))
 	{
 		ExternalVariable.bIsArray = true;
 		ExternalVariable.TypeName = *CPPType.Mid(7, CPPType.Len() - 8);
diff --git a/Engine/Source/Developer/RigVMDeveloper/Public/RigVMModel/Nodes/RigVMParameterNode.h b/Engine/Source/Developer/RigVMDeveloper/Public/RigVMModel/Nodes/RigVMParameterNode.h
--- a/Engine/Source/Developer/RigVMDeveloper/Public/RigVMModel/Nodes/RigVMParameterNode.h
+++ b/Engine/Source/Developer/RigVMDeveloper/Public/RigVMModel/Nodes/RigVMParameterNode.h
@@ -4,6 +4,7 @@
 
 #include "RigVMCore/RigVMExternalVariable.h"
 #include "RigVMModel/RigVMNode.h"
+#include "RigVMCore/RigVMUtilities.h"
 #include "RigVMParameterNode.generated.h"
 
 /**
@@ -50,7 +51,7 @@ public:
 		FRigVMExternalVariable ExternalVariable;
 		ExternalVariable.Name = Name;
 
-		if (CPPType.StartsWith(TEXT("TArray<")))
+		if (RigVMUtilities::IsArrayType(CPPType))
 		{
 			ExternalVariable.bIsArray = true;
 			ExternalVariable.TypeName = *CPPType.Mid(7, CPPType.Len() - 8);
diff --git a/Engine/Source/Developer/RigVMDeveloper/Public/RigVMModel/Nodes/RigVMVariableNode.h b/Engine/Source/Developer/RigVMDeveloper/Public/RigVMModel/Nodes/RigVMVariableNode.h
--- a/Engine/Source/Developer/RigVMDeveloper/Public/RigVMModel/Nodes/RigVMVariableNode.h
+++ b/Engine/Source/Developer/RigVMDeveloper/Public/RigVMModel/Nodes/RigVMVariableNode.h
@@ -4,68 +4,9 @@
 
 #include "RigVMCore/RigVMExternalVariable.h"
 #include "RigVMModel/RigVMNode.h"
+#include "RigVMModel/RigVMVariableDescription.h"
 #include "RigVMVariableNode.generated.h"
 
-/**
- * The variable description is used to convey information
- * about unique variables within a Graph. Multiple Variable
- * Nodes can share the same variable description.
- */
-USTRUCT(BlueprintType)
-struct FRigVMGraphVariableDescription
-{
-	GENERATED_BODY()
-
-public:
-
-	// comparison operator
-	bool operator ==(const FRigVMGraphVariableDescription& Other) const
-	{
-		return Name == Other.Name;
-	}
-
-	// The name of the variable
-	UPROPERTY(BlueprintReadOnly, Category = RigVMGraphVariableDescription)
-	FName Name;
-
-	// The C++ data type of the variable
-	UPROPERTY(BlueprintReadOnly, Category = RigVMGraphVariableDescription)
-	FString CPPType;
-
-	// The Struct of the C++ data type of the variable (or nullptr)
-	UPROPERTY(BlueprintReadOnly, Category = RigVMGraphVariableDescription)
-	UObject* CPPTypeObject = nullptr;
-
-	// The default value of the variable
-	UPROPERTY(BlueprintReadOnly, Category = RigVMGraphVariableDescription)
-	FString DefaultValue;
-
-	// Returns nullptr external variable matching this description
-	FORCEINLINE FRigVMExternalVariable ToExternalVariable() const
-	{
-		FRigVMExternalVariable ExternalVariable;
-		ExternalVariable.Name = Name;
-
-		if (CPPType.StartsWith(TEXT("TArray<")))
-		{
-			ExternalVariable.bIsArray = true;
-			ExternalVariable.TypeName = *CPPType.Mid(7, CPPType.Len() - 8);
-			ExternalVariable.TypeObject = CPPTypeObject;
-		}
-		else
-		{
-			ExternalVariable.bIsArray = false;
-			ExternalVariable.TypeName = *CPPType;
-			ExternalVariable.TypeObject = CPPTypeObject;
-		}
-
-		ExternalVariable.bIsPublic = false;
-		ExternalVariable.bIsReadOnly = false;
-		ExternalVariable.Memory = nullptr;
-		return ExternalVariable;
-	}
-};
-
 /**
  * The Variable Node represents a mutable value / local state within the
  * the Function / Graph. Variable Node's can be a getter or a setter.
diff --git a/Engine/Source/Developer/RigVMDeveloper/Public/RigVMModel/RigVMController.h b/Engine/Source/Developer/RigVMDeveloper/Public/RigVMModel/RigVMController.h
--- a/Engine/Source/Developer/RigVMDeveloper/Public/RigVMModel/RigVMController.h
+++ b/Engine/Source/Developer/RigVMDeveloper/Public/RigVMModel/RigVMController.h
@@ -584,12 +584,20 @@ public:
 
 	// Add a local variable to the graph
 	UFUNCTION(BlueprintCallable, Category = RigVMController)
-	bool AddLocalVariable(const FName& InVariableName, const FString& InCPPType, UObject* InCPPTypeObject, const FString& InDefaultValue, bool bSetupUndoRedo = true);
+	FRigVMGraphVariableDescription AddLocalVariable(const FName& InVariableName, const FString& InCPPType, UObject* InCPPTypeObject, const FString& InDefaultValue, bool bSetupUndoRedo = true);
 
 	// Remove a local variable from the graph
 	UFUNCTION(BlueprintCallable, Category = RigVMController)
 	bool RemoveLocalVariable(const FName& InVariableName, bool bSetupUndoRedo = true);
 
+	// Rename a local variable from the graph
+	UFUNCTION(BlueprintCallable, Category = RigVMController)
+	bool RenameLocalVariable(const FName& InVariableName, const FName& InNewVariableName, bool bSetupUndoRedo = true);
+
+	// Sets the type of the local variable
+	UFUNCTION(BlueprintCallable, Category = RigVMController)
+	bool SetLocalVariableType(const FName& InVariableName, const FString& InCPPType, UObject* InCPPTypeObject, bool bSetupUndoRedo = true);
+
 	// Determine affected function references for a potential bulk edit on a library node
 	TArray<TSoftObjectPtr<URigVMFunctionReferenceNode>> GetAffectedReferences(ERigVMControllerBulkEditType InEditType, bool bForceLoad = false, bool bNotify = true);
 
diff --git a/Engine/Source/Developer/RigVMDeveloper/Public/RigVMModel/RigVMControllerActions.h b/Engine/Source/Developer/RigVMDeveloper/Public/RigVMModel/RigVMControllerActions.h
--- a/Engine/Source/Developer/RigVMDeveloper/Public/RigVMModel/RigVMControllerActions.h
+++ b/Engine/Source/Developer/RigVMDeveloper/Public/RigVMModel/RigVMControllerActions.h
@@ -1507,3 +1507,51 @@ public:
 	UPROPERTY()
 	FRigVMGraphVariableDescription LocalVariable;
 };
+
+/**
+* An action to rename a local variable.
+*/
+USTRUCT()
+struct FRigVMRenameLocalVariableAction : public FRigVMBaseAction
+{
+	GENERATED_BODY()
+
+public:
+
+	FRigVMRenameLocalVariableAction() {}
+	FRigVMRenameLocalVariableAction(const FName& InOldName, const FName& InNewName);
+	virtual ~FRigVMRenameLocalVariableAction() {};
+	virtual bool Undo(URigVMController* InController) override;
+	virtual bool Redo(URigVMController* InController) override;
+
+	UPROPERTY()
+	FName OldVariableName;
+
+	UPROPERTY()
+	FName NewVariableName;
+};
+
+/**
+* An action to change the type of a local variable.
+*/
+USTRUCT()
+struct FRigVMChangeLocalVariableTypeAction : public FRigVMBaseAction
+{
+	GENERATED_BODY()
+
+public:
+	FRigVMChangeLocalVariableTypeAction();
+	FRigVMChangeLocalVariableTypeAction(const FRigVMGraphVariableDescription& InLocalVariable, const FString& InCPPType, UObject* InCPPTypeObject);
+	virtual ~FRigVMChangeLocalVariableTypeAction() {};
+	virtual bool Undo(URigVMController* InController) override;
+	virtual bool Redo(URigVMController* InController) override;
+
+	UPROPERTY()
+	FRigVMGraphVariableDescription LocalVariable;
+
+	UPROPERTY()
+	FString CPPType;
+
+	UPROPERTY()
+	UObject* CPPTypeObject;
+};
\ No newline at end of file
diff --git a/Engine/Source/Developer/RigVMDeveloper/Public/RigVMModel/RigVMGraph.h b/Engine/Source/Developer/RigVMDeveloper/Public/RigVMModel/RigVMGraph.h
--- a/Engine/Source/Developer/RigVMDeveloper/Public/RigVMModel/RigVMGraph.h
+++ b/Engine/Source/Developer/RigVMDeveloper/Public/RigVMModel/RigVMGraph.h
@@ -119,14 +119,6 @@ public:
 	UFUNCTION(BlueprintCallable, Category = RigVMGraph)
 	void SetDefaultFunctionLibrary(URigVMFunctionLibrary* InFunctionLibrary);
 
-	// Returns true if the variable was added to the graph's local variables
-	UFUNCTION(BlueprintCallable, Category = RigVMGraph)
-	bool AddLocalVariable(const FRigVMGraphVariableDescription& InVariable);
-
-	// Returns true if the variable was removed from the graph's local variables
-	UFUNCTION(BlueprintCallable, Category = RigVMGraph)
-	bool RemoveLocalVariable(const FName& InVariableName);
-
 	TArray<FRigVMExternalVariable> GetExternalVariables() const;
 
 	const TArray<FRigVMGraphVariableDescription>& GetLocalVariables() const { return LocalVariables; }
diff --git a/Engine/Source/Developer/RigVMDeveloper/Public/RigVMModel/RigVMVariableDescription.h b/Engine/Source/Developer/RigVMDeveloper/Public/RigVMModel/RigVMVariableDescription.h
new file mode 100644
--- /dev/null
+++ b/Engine/Source/Developer/RigVMDeveloper/Public/RigVMModel/RigVMVariableDescription.h
@@ -0,0 +1,172 @@
+﻿// Copyright Epic Games, Inc. All Rights Reserved.
+
+#pragma once
+
+#include "RigVMCore/RigVMExternalVariable.h"
+#include "RigVMModel/RigVMNode.h"
+#include "EdGraphSchema_K2.h"
+#include "RigVMVariableDescription.generated.h"
+
+/**
+ * The variable description is used to convey information
+ * about unique variables within a Graph. Multiple Variable
+ * Nodes can share the same variable description.
+ */
+USTRUCT(BlueprintType)
+struct FRigVMGraphVariableDescription
+{
+	GENERATED_BODY()
+
+public:
+
+	// comparison operator
+	bool operator ==(const FRigVMGraphVariableDescription& Other) const
+	{
+		return Name == Other.Name;
+	}
+
+	// The name of the variable
+	UPROPERTY(BlueprintReadOnly, Category = RigVMGraphVariableDescription)
+	FName Name;
+
+	// The C++ data type of the variable
+	UPROPERTY(BlueprintReadOnly, Category = RigVMGraphVariableDescription)
+	FString CPPType;
+
+	// The Struct of the C++ data type of the variable (or nullptr)
+	UPROPERTY(BlueprintReadOnly, Category = RigVMGraphVariableDescription)
+	UObject* CPPTypeObject = nullptr;
+
+	// The default value of the variable
+	UPROPERTY(BlueprintReadOnly, Category = RigVMGraphVariableDescription)
+	FString DefaultValue;
+
+	// Returns nullptr external variable matching this description
+	FORCEINLINE FRigVMExternalVariable ToExternalVariable() const
+	{
+		FRigVMExternalVariable ExternalVariable;
+		ExternalVariable.Name = Name;
+
+		if (CPPType.StartsWith(TEXT("TArray<")))
+		{
+			ExternalVariable.bIsArray = true;
+			ExternalVariable.TypeName = *CPPType.Mid(7, CPPType.Len() - 8);
+			ExternalVariable.TypeObject = CPPTypeObject;
+		}
+		else
+		{
+			ExternalVariable.bIsArray = false;
+			ExternalVariable.TypeName = *CPPType;
+			ExternalVariable.TypeObject = CPPTypeObject;
+		}
+
+		ExternalVariable.bIsPublic = false;
+		ExternalVariable.bIsReadOnly = false;
+		ExternalVariable.Memory = nullptr;
+		return ExternalVariable;
+	}
+
+	FORCEINLINE FEdGraphPinType ToPinType()
+	{
+		FEdGraphPinType PinType;
+		PinType.ResetToDefaults();
+		PinType.PinCategory = NAME_None;
+
+		if (CPPType == TEXT("bool"))
+		{
+			PinType.PinCategory = UEdGraphSchema_K2::PC_Boolean;
+		}
+		else if (CPPType == TEXT("int32"))
+		{
+			PinType.PinCategory = UEdGraphSchema_K2::PC_Int;
+		}
+		else if (CPPType == TEXT("float"))
+		{
+			PinType.PinCategory = UEdGraphSchema_K2::PC_Float;
+		}
+		else if (CPPType == TEXT("FName"))
+		{
+			PinType.PinCategory = UEdGraphSchema_K2::PC_Name;
+		}
+		else if (CPPType == TEXT("FString"))
+		{
+			PinType.PinCategory = UEdGraphSchema_K2::PC_String;
+		}
+		else if (Cast<UScriptStruct>(CPPTypeObject))
+		{
+			PinType.PinCategory = UEdGraphSchema_K2::PC_Struct;
+			PinType.PinSubCategoryObject = CPPTypeObject;
+		}
+		else if (Cast<UEnum>(CPPTypeObject))
+		{
+			PinType.PinCategory = UEdGraphSchema_K2::PC_Byte;
+			PinType.PinSubCategoryObject = CPPTypeObject;
+		}
+
+		if (CPPType.StartsWith(TEXT("TArray<")))
+		{
+			PinType.ContainerType = EPinContainerType::Array;
+		}
+		else
+		{
+			PinType.ContainerType = EPinContainerType::None;
+		}
+
+		return PinType;
+	}
+
+	static bool CPPTypeFromPinType(const FEdGraphPinType& InPinType, FString& OutCPPType, UObject*& OutCPPTypeObject)
+	{
+		FString Prefix = "";
+		FString Suffix = "";
+		if (InPinType.ContainerType == EPinContainerType::Array)
+		{
+			Prefix = TEXT("TArray<");
+			Suffix = TEXT(">");
+		}
+
+		OutCPPTypeObject = nullptr;
+		if (InPinType.PinCategory == UEdGraphSchema_K2::PC_Boolean)
+		{
+			OutCPPType = Prefix + TEXT("bool") + Suffix;
+		}
+		else if (InPinType.PinCategory == UEdGraphSchema_K2::PC_Int)
+		{
+			OutCPPType = Prefix + TEXT("int32") + Suffix;
+		}
+		else if (InPinType.PinCategory == UEdGraphSchema_K2::PC_Float)
+		{
+			OutCPPType = Prefix + TEXT("float") + Suffix;
+		}
+		else if (InPinType.PinCategory == UEdGraphSchema_K2::PC_Name)
+		{
+			OutCPPType = Prefix + TEXT("FName") + Suffix;
+		}
+		else if (InPinType.PinCategory == UEdGraphSchema_K2::PC_String)
+		{
+			OutCPPType = Prefix + TEXT("FString") + Suffix;
+		}
+		else if (InPinType.PinCategory == UEdGraphSchema_K2::PC_Struct)
+		{
+			OutCPPType = Prefix + TEXT("struct") + Suffix;
+			OutCPPTypeObject = InPinType.PinSubCategoryObject.Get();
+		}
+		else if (InPinType.PinCategory == UEdGraphSchema_K2::PC_Byte)
+		{
+			OutCPPType = Prefix + TEXT("byte") + Suffix;
+			OutCPPTypeObject = InPinType.PinSubCategoryObject.Get();
+		}
+		else
+		{
+			return false;
+		}
+		
+		return true;
+	}
+
+	FORCEINLINE bool ChangeType(const FEdGraphPinType& PinType)
+	{
+		return CPPTypeFromPinType(PinType, CPPType, CPPTypeObject);
+	}
+	
+};
diff --git a/Engine/Source/Editor/BlueprintGraph/Classes/EdGraphSchema_K2_Actions.h b/Engine/Source/Editor/BlueprintGraph/Classes/EdGraphSchema_K2_Actions.h
--- a/Engine/Source/Editor/BlueprintGraph/Classes/EdGraphSchema_K2_Actions.h
+++ b/Engine/Source/Editor/BlueprintGraph/Classes/EdGraphSchema_K2_Actions.h
@@ -10,6 +10,7 @@
 #include "Engine/Blueprint.h"
 #include "EdGraph/EdGraph.h"
 #include "EdGraph/EdGraphSchema.h"
+#include "EdGraphSchema_K2.h"
 
 #include "EdGraphSchema_K2_Actions.generated.h"
 
@@ -468,23 +469,23 @@ private:
 	FName VarName;
 
 	/** The struct that owns this item */
-	TWeakObjectPtr<UStruct> VariableSource;
+	TWeakObjectPtr<UObject> VariableSource;
 
 	/** TRUE if the variable's type is boolean */
 	bool bIsVarBool;
 
 public:
-	void SetVariableInfo(const FName& InVarName, const UStruct* InOwningScope, bool bInIsVarBool)
+	void SetVariableInfo(const FName& InVarName, const UObject* InOwningScope, bool bInIsVarBool)
 	{
 		VarName = InVarName;
 		bIsVarBool = bInIsVarBool;
 
 		check(InOwningScope);
-		VariableSource = MakeWeakObjectPtr(const_cast<UStruct*>(InOwningScope));
+		VariableSource = MakeWeakObjectPtr(const_cast<UObject*>(InOwningScope));
 	}
 
 	// Simple type info
-	static FName StaticGetTypeId() {static FName Type("FEdGraphSchemaAction_K2Var"); return Type;}
+	static FName StaticGetTypeId() {static FName Type("FEdGraphSchemaAction_BlueprintVariableBase"); return Type;}
 	virtual FName GetTypeId() const override { return StaticGetTypeId(); } 
 
 	FEdGraphSchemaAction_BlueprintVariableBase()
@@ -510,15 +511,31 @@ public:
 		return Cast<UClass>(GetVariableScope());
 	}
 
-	UStruct* GetVariableScope() const
+	UObject* GetVariableScope() const
 	{
 		return VariableSource.Get();
 	}
 
-	FProperty* GetProperty() const
+	virtual FProperty* GetProperty() const
 	{
-		return FindFProperty<FProperty>(GetVariableScope(), VarName);
+		if (UStruct* Scope = Cast<UStruct>(GetVariableScope()))
+		{
+			return FindFProperty<FProperty>(Scope, VarName);
+		}
+		return nullptr;
 	}
+
+	virtual FEdGraphPinType GetPinType() const{ return FEdGraphPinType();}
+
+	virtual void ChangeVariableType(const FEdGraphPinType& NewPinType){}
+
+	virtual void RenameVariable(const FName& NewName) { VarName = NewName; }
+
+	virtual bool IsValidName(const FName& NewName, FText& OutErrorMessage) const { return true; }
+
+	virtual void DeleteVariable() {}
+
+	virtual bool IsVariableUsed() { return false; }
 	
 	// FEdGraphSchemaAction interface
 	virtual void MovePersistentItemToCategory(const FText& NewCategoryName) override;
@@ -528,6 +545,8 @@ public:
 	// End of FEdGraphSchemaAction interface
 
 	UBlueprint* GetSourceBlueprint() const;
+
+	virtual void GetVariableTypeTree(TArray< TSharedPtr<UEdGraphSchema_K2::FPinTypeTreeInfo> >& TypeTree, ETypeTreeFilter TypeTreeFilter) const;
 };
 
 
@@ -553,6 +572,11 @@ public:
 	FEdGraphSchemaAction_K2Var(FText InNodeCategory, FText InMenuDesc, FText InToolTip, const int32 InGrouping, const int32 InSectionID)
 		: FEdGraphSchemaAction_BlueprintVariableBase(MoveTemp(InNodeCategory), MoveTemp(InMenuDesc), MoveTemp(InToolTip), InGrouping, InSectionID)
 	{}
+
+	virtual bool IsA(const FName& InType) const override
+	{
+		return InType == GetTypeId() || InType == FEdGraphSchemaAction_BlueprintVariableBase::StaticGetTypeId();
+	}
 };
 
 /*******************************************************************************
@@ -578,6 +602,11 @@ public:
 	FEdGraphSchemaAction_K2LocalVar(FText InNodeCategory, FText InMenuDesc, FText InToolTip, const int32 InGrouping, const int32 InSectionID)
 		: FEdGraphSchemaAction_BlueprintVariableBase(MoveTemp(InNodeCategory), MoveTemp(InMenuDesc), MoveTemp(InToolTip), InGrouping, InSectionID)
 	{}
+
+	virtual bool IsA(const FName& InType) const override
+    {
+    	return InType == GetTypeId() || InType == FEdGraphSchemaAction_BlueprintVariableBase::StaticGetTypeId();
+    }
 };
 
 /*******************************************************************************
@@ -783,6 +812,11 @@ public:
 	{
 		return FindFProperty<FMulticastDelegateProperty>(GetVariableClass(), GetVariableName());
 	}
+
+	virtual bool IsA(const FName& InType) const override
+	{
+		return InType == GetTypeId() || InType == FEdGraphSchemaAction_BlueprintVariableBase::StaticGetTypeId();
+	}
 };
 
 
diff --git a/Engine/Source/Editor/BlueprintGraph/Private/EdGraphSchema_K2_Actions.cpp b/Engine/Source/Editor/BlueprintGraph/Private/EdGraphSchema_K2_Actions.cpp
--- a/Engine/Source/Editor/BlueprintGraph/Private/EdGraphSchema_K2_Actions.cpp
+++ b/Engine/Source/Editor/BlueprintGraph/Private/EdGraphSchema_K2_Actions.cpp
@@ -40,7 +40,7 @@ namespace
 
 void FEdGraphSchemaAction_BlueprintVariableBase::MovePersistentItemToCategory(const FText& NewCategoryName)
 {
-	FBlueprintEditorUtils::SetBlueprintVariableCategory(GetSourceBlueprint(), VarName, GetVariableScope(), NewCategoryName);
+	FBlueprintEditorUtils::SetBlueprintVariableCategory(GetSourceBlueprint(), VarName, Cast<UStruct>(GetVariableScope()), NewCategoryName);
 }
 
 int32 FEdGraphSchemaAction_BlueprintVariableBase::GetReorderIndexInContainer() const
@@ -64,10 +64,10 @@ bool FEdGraphSchemaAction_BlueprintVariableBase::ReorderToBeforeAction(TSharedRe
 		FName TargetVarName = VarAction->GetVariableName();
 		if ((BP != nullptr) && (VarName != TargetVarName) && (VariableSource == VarAction->GetVariableScope()))
 		{
-			if (FBlueprintEditorUtils::MoveVariableBeforeVariable(BP, VarAction->GetVariableScope(), VarName, TargetVarName, true))
+			if (FBlueprintEditorUtils::MoveVariableBeforeVariable(BP, Cast<UStruct>(VarAction->GetVariableScope()), VarName, TargetVarName, true))
 			{
 				// Change category of var to match the one we dragged on to as well
-				FText TargetVarCategory = FBlueprintEditorUtils::GetBlueprintVariableCategory(BP, TargetVarName, GetVariableScope());
+				FText TargetVarCategory = FBlueprintEditorUtils::GetBlueprintVariableCategory(BP, TargetVarName, Cast<UStruct>(GetVariableScope()));
 				MovePersistentItemToCategory(TargetVarCategory);
 
 				// Update Blueprint after changes so they reflect in My Blueprint tab.
@@ -105,12 +105,19 @@ UBlueprint* FEdGraphSchemaAction_BlueprintVariableBase::GetSourceBlueprint() con
 	return UBlueprint::GetBlueprintFromClass(ClassToCheck);
 }
 
+void FEdGraphSchemaAction_BlueprintVariableBase::GetVariableTypeTree(
+	TArray<TSharedPtr<UEdGraphSchema_K2::FPinTypeTreeInfo>>& TypeTree, ETypeTreeFilter TypeTreeFilter) const
+{
+	const UEdGraphSchema_K2* Schema = GetDefault<UEdGraphSchema_K2>();
+	Schema->GetVariableTypeTree(TypeTree, TypeTreeFilter);
+}
+
 /////////////////////////////////////////////////////
 // FEdGraphSchemaAction_K2LocalVar
 
 int32 FEdGraphSchemaAction_K2LocalVar::GetReorderIndexInContainer() const
 {
-	return FBlueprintEditorUtils::FindLocalVariableIndex(GetSourceBlueprint(), GetVariableScope(), GetVariableName());
+	return FBlueprintEditorUtils::FindLocalVariableIndex(GetSourceBlueprint(), Cast<UStruct>(GetVariableScope()), GetVariableName());
 }
 
 /////////////////////////////////////////////////////
diff --git a/Engine/Source/Editor/Kismet/Private/BlueprintEditor.cpp b/Engine/Source/Editor/Kismet/Private/BlueprintEditor.cpp
--- a/Engine/Source/Editor/Kismet/Private/BlueprintEditor.cpp
+++ b/Engine/Source/Editor/Kismet/Private/BlueprintEditor.cpp
@@ -551,15 +551,22 @@ FSlateBrush const* FBlueprintEditor::GetVarIconAndColorFromProperty(const FPrope
 		FEdGraphPinType PinType;
 		if (K2Schema->ConvertPropertyToPinType(Property, PinType)) // use schema to get the color
 		{
-			IconColorOut = K2Schema->GetPinTypeColor(PinType);
-			SecondaryBrushOut = FBlueprintEditorUtils::GetSecondaryIconFromPin(PinType);
-			SecondaryColorOut = K2Schema->GetSecondaryPinTypeColor(PinType);
-			return FBlueprintEditorUtils::GetIconFromPin(PinType);
+			return GetVarIconAndColorFromPinType(PinType, IconColorOut, SecondaryBrushOut, SecondaryColorOut);
 		}
 	}
 	return FEditorStyle::GetBrush(TEXT("Kismet.AllClasses.VariableIcon"));
 }
 
+FSlateBrush const* FBlueprintEditor::GetVarIconAndColorFromPinType(const FEdGraphPinType& PinType,
+	FSlateColor& IconColorOut, FSlateBrush const*& SecondaryBrushOut, FSlateColor& SecondaryColorOut)
+{
+	const UEdGraphSchema_K2* K2Schema = GetDefault<UEdGraphSchema_K2>();
+	IconColorOut = K2Schema->GetPinTypeColor(PinType);
+	SecondaryBrushOut = FBlueprintEditorUtils::GetSecondaryIconFromPin(PinType);
+	SecondaryColorOut = K2Schema->GetSecondaryPinTypeColor(PinType);
+	return FBlueprintEditorUtils::GetIconFromPin(PinType);
+}
+
 bool FBlueprintEditor::IsInAScriptingMode() const
 {
 	return IsModeCurrent(FBlueprintEditorApplicationModes::StandardBlueprintEditorMode) || IsModeCurrent(FBlueprintEditorApplicationModes::BlueprintMacroMode);
diff --git a/Engine/Source/Editor/Kismet/Private/SBlueprintPalette.cpp b/Engine/Source/Editor/Kismet/Private/SBlueprintPalette.cpp
--- a/Engine/Source/Editor/Kismet/Private/SBlueprintPalette.cpp
+++ b/Engine/Source/Editor/Kismet/Private/SBlueprintPalette.cpp
@@ -384,13 +384,24 @@ static void GetPaletteItemIcon(TSharedPtr<FEdGraphSchemaAction> ActionIn, UBluep
 	{
 		FEdGraphSchemaAction_K2LocalVar* LocalVarAction = (FEdGraphSchemaAction_K2LocalVar*)ActionIn.Get();
 
-		UStruct* VarScope = LocalVarAction->GetVariableScope();
+		UStruct* VarScope = CastChecked<UStruct>(LocalVarAction->GetVariableScope());
 		BrushOut = FBlueprintEditor::GetVarIconAndColor(VarScope, LocalVarAction->GetVariableName(), ColorOut, SecondaryBrushOut, SecondaryColorOut);
 		ToolTipOut = FText::FromString(GetVarType(VarScope, LocalVarAction->GetVariableName(), true));
 
 		DocLinkOut = TEXT("Shared/Editor/Blueprint/VariableTypes");
 		DocExcerptOut = GetVarType(VarScope, LocalVarAction->GetVariableName(), false);
 	}
+	else if (ActionIn->IsA(FEdGraphSchemaAction_BlueprintVariableBase::StaticGetTypeId()))
+	{
+		FEdGraphSchemaAction_BlueprintVariableBase* BPVarAction = (FEdGraphSchemaAction_BlueprintVariableBase*)(ActionIn.Get());
+		const FEdGraphPinType PinType = BPVarAction->GetPinType();
+		
+		BrushOut = FBlueprintEditor::GetVarIconAndColorFromPinType(PinType, ColorOut, SecondaryBrushOut, SecondaryColorOut);
+		ToolTipOut = FText::FromString(UEdGraphSchema_K2::TypeToText(PinType).ToString());
+
+		DocLinkOut = TEXT("Shared/Editor/Blueprint/VariableTypes");
+		DocExcerptOut = UEdGraphSchema_K2::TypeToText(PinType).ToString();
+	}
 	else if (ActionIn->GetTypeId() == FEdGraphSchemaAction_K2Enum::StaticGetTypeId())
 	{
 		BrushOut = FEditorStyle::GetBrush(TEXT("GraphEditor.EnumGlyph"));
@@ -724,7 +735,25 @@ public:
 		BlueprintObj = InBlueprint;
 		BlueprintEditorPtr = InBlueprintEditor;
 		VariableProperty = InVariableProperty;
+		ActionPtr = nullptr;
+
+		ConstructInternal(InArgs);
+	}
+
+	void Construct(const FArguments& InArgs, TWeakPtr<FEdGraphSchemaAction_BlueprintVariableBase> InAction, UBlueprint* InBlueprint, TWeakPtr<FBlueprintEditor> InBlueprintEditor)
+	{
+		BlueprintObj = InBlueprint;
+		BlueprintEditorPtr = InBlueprintEditor;
+		VariableProperty = nullptr;
+		ActionPtr = InAction;
+
+		ConstructInternal(InArgs);
+	}
 
+private:
+
+	void ConstructInternal(const FArguments& InArgs)
+	{
 		TSharedPtr<IPinTypeSelectorFilter> CustomPinTypeFilter;
 		if (BlueprintEditorPtr.IsValid() && GetDefault<UBlueprintEditorSettings>()->bEnableNamespaceFilteringFeatures)
 		{
@@ -732,9 +761,12 @@ public:
 		}
 
 		const UEdGraphSchema_K2* Schema = GetDefault<UEdGraphSchema_K2>();
+		FGetPinTypeTree GetPinTypeTreeFunc = ActionPtr.IsValid() ? FGetPinTypeTree::CreateRaw(ActionPtr.Pin().Get(), &FEdGraphSchemaAction_BlueprintVariableBase::GetVariableTypeTree)
+			: FGetPinTypeTree::CreateUObject(Schema, &UEdGraphSchema_K2::GetVariableTypeTree);
+		
 		this->ChildSlot
 		[
-			SNew(SPinTypeSelector, FGetPinTypeTree::CreateUObject(Schema, &UEdGraphSchema_K2::GetVariableTypeTree))
+			SNew(SPinTypeSelector, GetPinTypeTreeFunc)
 			.ReadOnly(InArgs._ReadOnly)
 			.Schema(Schema)
 			.TargetPinType(this, &SPinTypeSelectorHelper::OnGetVarType)
@@ -744,8 +776,7 @@ public:
 			.CustomFilter(CustomPinTypeFilter)
 		];
 	}
-
-private:
+	
 	FEdGraphPinType OnGetVarType() const
 	{
 		if (FProperty* VarProp = const_cast<FProperty*>(VariableProperty.Get()))
@@ -755,6 +786,10 @@ private:
 			K2Schema->ConvertPropertyToPinType(VarProp, Type);
 			return Type;
 		}
+		else if(ActionPtr.IsValid())
+		{
+			return ActionPtr.Pin()->GetPinType();
+		}
 		return FEdGraphPinType();
 	}
 
@@ -781,12 +816,19 @@ private:
 					}
 				}
 			}
+			else if(ActionPtr.IsValid())
+			{
+				// Set the MyBP tab's last pin type used as this, for adding lots of variables of the same type
+				BlueprintEditorPtr.Pin()->GetMyBlueprintWidget()->GetLastPinTypeUsed() = InNewPinType;
+				
+				ActionPtr.Pin()->ChangeVariableType(InNewPinType);
+			}
 		}
 	}
 
 private:
 	/** The action that the owning palette entry represents */
-	TWeakPtr<FEdGraphSchemaAction_K2Var> ActionPtr;
+	TWeakPtr<FEdGraphSchemaAction_BlueprintVariableBase> ActionPtr;
 
 	/** Pointer back to the blueprint that is being displayed: */
 	UBlueprint* BlueprintObj;
@@ -1117,6 +1159,15 @@ void SBlueprintPaletteItem::Construct(const FArguments& InArgs, FCreateWidgetFor
 	{
 		GenerateVariableSettings(StaticCastSharedPtr<FEdGraphSchemaAction_K2LocalVar>(GraphAction)->GetProperty());
 	}
+	else if (GraphAction->IsA(FEdGraphSchemaAction_BlueprintVariableBase::StaticGetTypeId()))
+	{
+		ActionAccessSpecifier = EAccessSpecifier::Private;
+		const UEdGraphSchema_K2* Schema = GetDefault<UEdGraphSchema_K2>();
+		IconWidget = SNew(SPinTypeSelectorHelper, StaticCastSharedPtr<FEdGraphSchemaAction_BlueprintVariableBase>(GraphAction), Blueprint, BlueprintEditorPtr)
+			.IsEnabled(bIsEditingEnabled)
+			.ReadOnly_Lambda([this]() {return !IsHovered(); });
+	}
+	
 	// Determine the access level of this action if it is a function graph or for interface events
 	else if (bShouldCheckForAccessSpec && GraphAction->GetTypeId() == FEdGraphSchemaAction_K2Graph::StaticGetTypeId())
 	{
@@ -1179,7 +1230,7 @@ void SBlueprintPaletteItem::Construct(const FArguments& InArgs, FCreateWidgetFor
 	TSharedRef<SHorizontalBox> ActionBox = SNew(SHorizontalBox)		
 		.AddMetaData<FTutorialMetaData>(TagMeta);
 	
-	if (GraphAction->GetTypeId() == FEdGraphSchemaAction_K2Var::StaticGetTypeId() || GraphAction->GetTypeId() == FEdGraphSchemaAction_K2LocalVar::StaticGetTypeId())
+	if (GraphAction->IsA(FEdGraphSchemaAction_BlueprintVariableBase::StaticGetTypeId()))
 	{
 		ActionBox.Get().AddSlot()
 			.FillWidth(0.6f)
@@ -1403,7 +1454,7 @@ bool SBlueprintPaletteItem::OnNameTextVerifyChanged(const FText& InNewText, FTex
 		FEdGraphSchemaAction_K2LocalVar* LocalVarAction = (FEdGraphSchemaAction_K2LocalVar*)ActionPtr.Pin().Get();
 		OriginalName = (LocalVarAction->GetVariableName());
 		
-		ValidationScope = LocalVarAction->GetVariableScope();
+		ValidationScope = CastChecked<UStruct>(LocalVarAction->GetVariableScope());
 	}
 	else
 	{
@@ -1441,20 +1492,29 @@ bool SBlueprintPaletteItem::OnNameTextVerifyChanged(const FText& InNewText, FTex
 		}
 	}
 
-	TSharedPtr<INameValidatorInterface> NameValidator = MakeShareable(new FKismetNameValidator(BlueprintObj, OriginalName, ValidationScope));
-
-	EValidatorResult ValidatorResult = NameValidator->IsValid(TextAsString);
-	switch (ValidatorResult)
+	if (OriginalName.IsNone() && ActionPtr.Pin()->IsA(FEdGraphSchemaAction_BlueprintVariableBase::StaticGetTypeId()))
 	{
-	case EValidatorResult::Ok:
-	case EValidatorResult::ExistingName:
-		// These are fine, don't need to surface to the user, the rename can 'proceed' even if the name is the existing one
-		break;
-	default:
-		OutErrorMessage = INameValidatorInterface::GetErrorText(TextAsString, ValidatorResult);
-		break;
+		FEdGraphSchemaAction_BlueprintVariableBase* BPVar = (FEdGraphSchemaAction_BlueprintVariableBase*)ActionPtr.Pin().Get();
+		return BPVar->IsValidName(FName(TextAsString), OutErrorMessage);
 	}
+	else
+	{
+		TSharedPtr<INameValidatorInterface> NameValidator = MakeShareable(new FKismetNameValidator(BlueprintObj, OriginalName, ValidationScope));
 
+		EValidatorResult ValidatorResult = NameValidator->IsValid(TextAsString);
+		switch (ValidatorResult)
+		{
+		case EValidatorResult::Ok:
+		case EValidatorResult::ExistingName:
+			// These are fine, don't need to surface to the user, the rename can 'proceed' even if the name is the existing one
+			break;
+		default:
+			OutErrorMessage = INameValidatorInterface::GetErrorText(TextAsString, ValidatorResult);
+			break;
+		}
+	
+	}
+	
 	return OutErrorMessage.IsEmpty();
 }
 
@@ -1620,7 +1680,23 @@ void SBlueprintPaletteItem::OnNameTextCommitted(const FText& NewText, ETextCommi
 
 		BlueprintEditorPtr.Pin()->GetBlueprintObj()->Modify();
 
-		FBlueprintEditorUtils::RenameLocalVariable(BlueprintEditorPtr.Pin()->GetBlueprintObj(), LocalVarAction->GetVariableScope(), LocalVarAction->GetVariableName(), NewName);
+		FBlueprintEditorUtils::RenameLocalVariable(BlueprintEditorPtr.Pin()->GetBlueprintObj(), CastChecked<UStruct>(LocalVarAction->GetVariableScope()), LocalVarAction->GetVariableName(), NewName);
+	}
+	else if (ActionPtr.Pin()->IsA(FEdGraphSchemaAction_BlueprintVariableBase::StaticGetTypeId()))
+	{
+		FEdGraphSchemaAction_BlueprintVariableBase* BPVarAction = (FEdGraphSchemaAction_BlueprintVariableBase*)ActionPtr.Pin().Get();
+
+		// Check if the name is unchanged
+		if (NewName.IsEqual(BPVarAction->GetVariableName(), ENameCase::CaseSensitive))
+		{
+			return;
+		}
+
+		const FScopedTransaction Transaction( LOCTEXT( "RenameVariable", "Rename Variable" ) );
+		BlueprintEditorPtr.Pin()->GetBlueprintObj()->Modify();
+
+		BPVarAction->RenameVariable(NewName);
+		BlueprintEditorPtr.Pin()->GetMyBlueprintWidget()->Refresh();
 	}
 	BlueprintEditorPtr.Pin()->GetMyBlueprintWidget()->SelectItemByName(NewName, ESelectInfo::OnMouseClick);
 }
@@ -1715,7 +1791,8 @@ FText SBlueprintPaletteItem::GetToolTipText() const
 		{
 			FEdGraphSchemaAction_K2LocalVar* LocalVarAction = (FEdGraphSchemaAction_K2LocalVar*)PaletteAction.Get();
 			// The variable scope can not be found in intermediate graphs
-			if(LocalVarAction->GetVariableScope())
+			UStruct* LocalVarScope = Cast<UStruct>(LocalVarAction->GetVariableScope());
+			if(LocalVarScope)
 			{
 				UClass* VarClass = CastChecked<UClass>(LocalVarAction->GetVariableScope()->GetOuter());
 				if (bShowClassInTooltip && (VarClass != nullptr))
@@ -1726,9 +1803,9 @@ FText SBlueprintPaletteItem::GetToolTipText() const
 				else
 				{
 					FString Result;
-					FBlueprintEditorUtils::GetBlueprintVariableMetaData(Blueprint, LocalVarAction->GetVariableName(), LocalVarAction->GetVariableScope(), TEXT("tooltip"), Result);
+					FBlueprintEditorUtils::GetBlueprintVariableMetaData(Blueprint, LocalVarAction->GetVariableName(), LocalVarScope, TEXT("tooltip"), Result);
 					// Only use the variable tooltip if it has been filled out.
-					ToolTipText = FText::FromString( !Result.IsEmpty() ? Result : GetVarType(LocalVarAction->GetVariableScope(), LocalVarAction->GetVariableName(), true, true) );
+					ToolTipText = FText::FromString( !Result.IsEmpty() ? Result : GetVarType(LocalVarScope, LocalVarAction->GetVariableName(), true, true) );
 				}
 			}
 		}
diff --git a/Engine/Source/Editor/Kismet/Private/SMyBlueprint.cpp b/Engine/Source/Editor/Kismet/Private/SMyBlueprint.cpp
--- a/Engine/Source/Editor/Kismet/Private/SMyBlueprint.cpp
+++ b/Engine/Source/Editor/Kismet/Private/SMyBlueprint.cpp
@@ -620,7 +620,7 @@ void SMyBlueprint::OnCategoryNameCommitted(const FText& InNewText, ETextCommit::
 			{
 				FEdGraphSchemaAction_K2LocalVar* LocalVarAction = (FEdGraphSchemaAction_K2LocalVar*)Actions[i].Get();
 
-				FBlueprintEditorUtils::SetBlueprintVariableCategory(GetBlueprintObj(), LocalVarAction->GetVariableName(), LocalVarAction->GetVariableScope(), CategoryName, true);
+				FBlueprintEditorUtils::SetBlueprintVariableCategory(GetBlueprintObj(), LocalVarAction->GetVariableName(), CastChecked<UStruct>(LocalVarAction->GetVariableScope()), CategoryName, true);
 			}
 			else if (Actions[i]->GetTypeId() == FEdGraphSchemaAction_K2Delegate::StaticGetTypeId())
 			{
@@ -1142,21 +1142,18 @@ void SMyBlueprint::GetLocalVariables(FGraphActionSort& SortList) const
 	UEdGraph* TopLevelGraph = FBlueprintEditorUtils::GetTopLevelGraph(GetFocusedGraph());
 	if( TopLevelGraph )
 	{
+		TArray<FBPVariableDescription> LocalVariables;
+		bool bSchemaImplementsGetLocalVariables = false;
+	
 		// grab the parent graph's name
 		FGraphDisplayInfo EdGraphDisplayInfo;
 		if (UEdGraphSchema const* Schema = TopLevelGraph->GetSchema())
 		{
 			Schema->GetGraphDisplayInformation(*TopLevelGraph, EdGraphDisplayInfo);
-		}
-
-		TArray<UK2Node_FunctionEntry*> FunctionEntryNodes;
-		TopLevelGraph->GetNodesOfClass<UK2Node_FunctionEntry>(FunctionEntryNodes);
 
-		// Search in all FunctionEntry nodes for their local variables
-		FText ActionCategory;
-		for (UK2Node_FunctionEntry* const FunctionEntry : FunctionEntryNodes)
-		{
-			for (const FBPVariableDescription& Variable : FunctionEntry->LocalVariables)
+			// Try to get the local variables from the schema
+			bSchemaImplementsGetLocalVariables = Schema->GetLocalVariables(GetFocusedGraph(), LocalVariables);
+			for (const FBPVariableDescription& Variable : LocalVariables)
 			{
 				FText Category = Variable.Category;
 				if (Variable.Category.EqualTo(UEdGraphSchema_K2::VR_DefaultCategory))
@@ -1164,12 +1161,39 @@ void SMyBlueprint::GetLocalVariables(FGraphActionSort& SortList) const
 					Category = FText::GetEmpty();
 				}
 
-				UFunction* Func = FindUField<UFunction>(GetBlueprintObj()->SkeletonGeneratedClass, TopLevelGraph->GetFName());
-				if (Func)
+				TSharedPtr<FEdGraphSchemaAction> Action =  Schema->MakeActionFromVariableDescription(GetFocusedGraph(), Variable);
+				if (Action.IsValid())
+				{
+					SortList.AddAction(Action);
+				}						
+			}
+		}
+
+		// If the schema did not return any local variables, try to get them from the function entry
+		if (!bSchemaImplementsGetLocalVariables)
+		{
+			TArray<UK2Node_FunctionEntry*> FunctionEntryNodes;
+			TopLevelGraph->GetNodesOfClass<UK2Node_FunctionEntry>(FunctionEntryNodes);
+
+			// Search in all FunctionEntry nodes for their local variables
+			FText ActionCategory;
+			for (UK2Node_FunctionEntry* const FunctionEntry : FunctionEntryNodes)
+			{
+				for (const FBPVariableDescription& Variable : FunctionEntry->LocalVariables)
 				{
-					TSharedPtr<FEdGraphSchemaAction_K2LocalVar> NewVarAction = MakeShareable(new FEdGraphSchemaAction_K2LocalVar(Category, FText::FromName(Variable.VarName), FText::GetEmpty(), 0, NodeSectionID::LOCAL_VARIABLE));
-					NewVarAction->SetVariableInfo(Variable.VarName, Func, Variable.VarType.PinCategory == UEdGraphSchema_K2::PC_Boolean);
-					SortList.AddAction(NewVarAction);
+					FText Category = Variable.Category;
+					if (Variable.Category.EqualTo(UEdGraphSchema_K2::VR_DefaultCategory))
+					{
+						Category = FText::GetEmpty();
+					}
+
+					UFunction* Func = FindUField<UFunction>(GetBlueprintObj()->SkeletonGeneratedClass, TopLevelGraph->GetFName());
+					if (Func)
+					{
+						TSharedPtr<FEdGraphSchemaAction_K2LocalVar> NewVarAction = MakeShareable(new FEdGraphSchemaAction_K2LocalVar(Category, FText::FromName(Variable.VarName), FText::GetEmpty(), 0, NodeSectionID::LOCAL_VARIABLE));
+						NewVarAction->SetVariableInfo(Variable.VarName, Func, Variable.VarType.PinCategory == UEdGraphSchema_K2::PC_Boolean);
+						SortList.AddAction(NewVarAction);
+					}
 				}
 			}
 		}
@@ -1775,7 +1799,7 @@ FReply SMyBlueprint::OnActionDragged( const TArray< TSharedPtr<FEdGraphSchemaAct
 		else if( InAction->GetTypeId() == FEdGraphSchemaAction_K2LocalVar::StaticGetTypeId())
 		{
 			FEdGraphSchemaAction_K2LocalVar* VarAction = (FEdGraphSchemaAction_K2LocalVar*)InAction.Get();
-			if (UStruct* VariableScope = VarAction->GetVariableScope())
+			if (UStruct* VariableScope = Cast<UStruct>(VarAction->GetVariableScope()))
 			{
 				TSharedRef<FKismetVariableDragDropAction> DragOperation = FKismetVariableDragDropAction::New(InAction, VarAction->GetVariableName(), VariableScope, AnalyticsDelegate);
 				DragOperation->SetAltDrag(MouseEvent.IsAltDown());
@@ -1794,6 +1818,15 @@ FReply SMyBlueprint::OnActionDragged( const TArray< TSharedPtr<FEdGraphSchemaAct
 				return FReply::Handled().BeginDragDrop(DragOperation);
 			}
 		}
+		else if( InAction->IsA(FEdGraphSchemaAction_BlueprintVariableBase::StaticGetTypeId()))
+		{
+			FEdGraphSchemaAction_BlueprintVariableBase* VarAction = (FEdGraphSchemaAction_BlueprintVariableBase*)InAction.Get();
+
+			if (GetFocusedGraph()->GetSchema()->CanGraphBeDropped(InAction))
+			{
+				return GetFocusedGraph()->GetSchema()->BeginGraphDragAction(InAction, MouseEvent);
+			}		
+		}
 		else if (InAction->GetTypeId() == FEdGraphSchemaAction_K2Event::StaticGetTypeId())
 		{	
 			// Check if it's a custom event, it is preferable to drop a call function for custom events than to focus on the node
@@ -2067,6 +2100,22 @@ FEdGraphSchemaAction_K2LocalVar* SMyBlueprint::SelectionAsLocalVar() const
 	return SelectionAsType<FEdGraphSchemaAction_K2LocalVar>(GraphActionMenu);
 }
 
+FEdGraphSchemaAction_BlueprintVariableBase* SMyBlueprint::SelectionAsBlueprintVariable() const
+{
+	TArray<TSharedPtr<FEdGraphSchemaAction> > SelectedActions;
+	GraphActionMenu->GetSelectedActions(SelectedActions);
+
+	FEdGraphSchemaAction_BlueprintVariableBase* Selection = nullptr;
+
+	TSharedPtr<FEdGraphSchemaAction> SelectedAction( SelectedActions.Num() > 0 ? SelectedActions[0] : nullptr );
+	if ( SelectedAction.IsValid() && SelectedAction->IsA(FEdGraphSchemaAction_BlueprintVariableBase::StaticGetTypeId()))
+	{
+		Selection = (FEdGraphSchemaAction_BlueprintVariableBase*)SelectedAction.Get();
+	}
+
+	return Selection;
+}
+
 FEdGraphSchemaAction_K2Delegate* SMyBlueprint::SelectionAsDelegate() const
 {
 	return SelectionAsType<FEdGraphSchemaAction_K2Delegate>( GraphActionMenu );
@@ -2766,7 +2815,7 @@ void SMyBlueprint::OnDeleteEntry()
 	}
 	else if ( FEdGraphSchemaAction_K2LocalVar* LocalVarAction = SelectionAsLocalVar() )
 	{
-		if(FBlueprintEditorUtils::IsVariableUsed(GetBlueprintObj(), LocalVarAction->GetVariableName(), FBlueprintEditorUtils::FindScopeGraph(GetBlueprintObj(), LocalVarAction->GetVariableScope())))
+		if(FBlueprintEditorUtils::IsVariableUsed(GetBlueprintObj(), LocalVarAction->GetVariableName(), FBlueprintEditorUtils::FindScopeGraph(GetBlueprintObj(), CastChecked<UStruct>(LocalVarAction->GetVariableScope()))))
 		{
 			FText ConfirmDelete = FText::Format(LOCTEXT( "ConfirmDeleteLocalVariableInUse", "Local Variable {0} is in use! Do you really want to delete it?"),
 				FText::FromName( LocalVarAction->GetVariableName() ) );
@@ -2793,7 +2842,31 @@ void SMyBlueprint::OnDeleteEntry()
 		check(FunctionEntryNodes.Num() == 1);
 		FunctionEntryNodes[0]->Modify();
 
-		FBlueprintEditorUtils::RemoveLocalVariable(GetBlueprintObj(), LocalVarAction->GetVariableScope(), LocalVarAction->GetVariableName());
+		FBlueprintEditorUtils::RemoveLocalVariable(GetBlueprintObj(), CastChecked<UStruct>(LocalVarAction->GetVariableScope()), LocalVarAction->GetVariableName());
+	}
+	else if ( FEdGraphSchemaAction_BlueprintVariableBase* BPVarAction = SelectionAsBlueprintVariable() )
+	{
+		if(BPVarAction->IsVariableUsed())
+		{
+			FText ConfirmDelete = FText::Format(LOCTEXT( "ConfirmDeleteLocalVariableInUse", "Variable {0} is in use! Do you really want to delete it?"),
+				FText::FromName( BPVarAction->GetVariableName() ) );
+
+			// Warn the user that this may result in data loss
+			FSuppressableWarningDialog::FSetupInfo Info( ConfirmDelete, LOCTEXT("DeleteVar", "Delete Variable"), "DeleteVariableInUse_Warning" );
+			Info.ConfirmText = LOCTEXT( "DeleteVariable_Yes", "Yes");
+			Info.CancelText = LOCTEXT( "DeleteVariable_No", "No");	
+
+			FSuppressableWarningDialog DeleteVariableInUse( Info );
+			if ( DeleteVariableInUse.ShowModal() == FSuppressableWarningDialog::Cancel )
+			{
+				return;
+			}
+		}
+
+		const FScopedTransaction Transaction( LOCTEXT( "RemoveLocalVariable", "Remove Local Variable" ) );
+
+		GetBlueprintObj()->Modify();
+		BPVarAction->DeleteVariable();		
 	}
 	else if (FEdGraphSchemaAction_K2Event* EventAction = SelectionAsEvent())
 	{
@@ -2857,7 +2930,7 @@ void SMyBlueprint::OnDeleteEntry()
 				{
 					FEdGraphSchemaAction_K2LocalVar* K2LocalVarAction = (FEdGraphSchemaAction_K2LocalVar*)Actions[i].Get();
 
-					FBlueprintEditorUtils::RemoveLocalVariable(GetBlueprintObj(), K2LocalVarAction->GetVariableScope(), K2LocalVarAction->GetVariableName());
+					FBlueprintEditorUtils::RemoveLocalVariable(GetBlueprintObj(), CastChecked<UStruct>(K2LocalVarAction->GetVariableScope()), K2LocalVarAction->GetVariableName());
 					bModified = true;
 				}
 				else if (Actions[i]->GetTypeId() == FEdGraphSchemaAction_K2Graph::StaticGetTypeId())
@@ -2932,6 +3005,10 @@ bool SMyBlueprint::CanDeleteEntry() const
 	{
 		return true;
 	}
+	else if (FEdGraphSchemaAction_BlueprintVariableBase* BPVariable = SelectionAsBlueprintVariable())
+	{
+		return true;
+	}
 	else if (SelectionIsCategory())
 	{
 		// Can't delete categories if they can't be renamed, that means they are native
@@ -2990,7 +3067,7 @@ bool SMyBlueprint::CanDuplicateAction() const
 		}
 		return true;
 	}
-	else if(SelectionAsLocalVar())
+	else if(SelectionAsBlueprintVariable())
 	{
 		return true;
 	}
@@ -3062,7 +3139,7 @@ void SMyBlueprint::OnDuplicateAction()
 		const FScopedTransaction Transaction( LOCTEXT( "Duplicate Local Variable", "Duplicate Local Variable" ) );
 		GetBlueprintObj()->Modify();
 
-		DuplicateActionName = FBlueprintEditorUtils::DuplicateVariable(GetBlueprintObj(), LocalVarAction->GetVariableScope(), LocalVarAction->GetVariableName());
+		DuplicateActionName = FBlueprintEditorUtils::DuplicateVariable(GetBlueprintObj(), Cast<UStruct>(LocalVarAction->GetVariableScope()), LocalVarAction->GetVariableName());
 	}
 
 	// Select and rename the duplicated action
@@ -3229,7 +3306,7 @@ void SMyBlueprint::OnCopy()
 	}
 	else if (FEdGraphSchemaAction_K2LocalVar* LocalVarAction = SelectionAsLocalVar())
 	{
-		FBPVariableDescription* Description = FBlueprintEditorUtils::FindLocalVariable(Blueprint, LocalVarAction->GetVariableScope(), LocalVarAction->GetVariableName());
+		FBPVariableDescription* Description = FBlueprintEditorUtils::FindLocalVariable(Blueprint, CastChecked<UStruct>(LocalVarAction->GetVariableScope()), LocalVarAction->GetVariableName());
 
 		if (Description)
 		{
@@ -3261,7 +3338,7 @@ bool SMyBlueprint::CanCopy() const
 	}
 	if (FEdGraphSchemaAction_K2LocalVar* LocalVarAction = SelectionAsLocalVar())
 	{
-		return FBlueprintEditorUtils::FindLocalVariable(Blueprint, LocalVarAction->GetVariableScope(), LocalVarAction->GetVariableName()) != nullptr;
+		return FBlueprintEditorUtils::FindLocalVariable(Blueprint, Cast<UStruct>(LocalVarAction->GetVariableScope()), LocalVarAction->GetVariableName()) != nullptr;
 	}
 	if (FEdGraphSchemaAction_K2Graph* GraphAction = SelectionAsGraph())
 	{
diff --git a/Engine/Source/Editor/Kismet/Private/SMyBlueprint.h b/Engine/Source/Editor/Kismet/Private/SMyBlueprint.h
--- a/Engine/Source/Editor/Kismet/Private/SMyBlueprint.h
+++ b/Engine/Source/Editor/Kismet/Private/SMyBlueprint.h
@@ -100,6 +100,9 @@ public:
 	/** Accessor for getting the current selection as a K2 local var */
 	FEdGraphSchemaAction_K2LocalVar* SelectionAsLocalVar() const;
 
+	/** Accessor for getting the current selection as a K2 blueprint base variable */
+	FEdGraphSchemaAction_BlueprintVariableBase* SelectionAsBlueprintVariable() const;
+
 	/** Accessor for determining if the current selection is a category*/
 	bool SelectionIsCategory() const;
 	
diff --git a/Engine/Source/Editor/Kismet/Public/BlueprintEditor.h b/Engine/Source/Editor/Kismet/Public/BlueprintEditor.h
--- a/Engine/Source/Editor/Kismet/Public/BlueprintEditor.h
+++ b/Engine/Source/Editor/Kismet/Public/BlueprintEditor.h
@@ -437,6 +437,17 @@ public:
 	 */
 	static FSlateBrush const* GetVarIconAndColorFromProperty(const FProperty* Property, FSlateColor& IconColorOut, FSlateBrush const*& SecondaryBrushOut, FSlateColor& SecondaryColorOut);
 
+	/**
+	* Util for finding a glyph and color for a variable.
+	*
+	* @param PinType       The variable's pin type
+	* @param IconColorOut      The resulting color for the glyph
+	* @param SecondaryBrushOut The resulting secondary glyph brush (used for Map types)
+	* @param SecondaryColorOut The resulting secondary color for the glyph (used for Map types)
+	* @return					The resulting glyph brush
+	*/
+	static FSlateBrush const* GetVarIconAndColorFromPinType(const FEdGraphPinType& PinType, FSlateColor& IconColorOut, FSlateBrush const*& SecondaryBrushOut, FSlateColor& SecondaryColorOut);
+
 	/** Overridable function for determining if the current mode can script */
 	virtual bool IsInAScriptingMode() const;
 
@@ -530,7 +541,7 @@ public:
 	virtual bool CanAddNewLocalVariable() const;
 
 	/** Adds a new local variable to the focused function graph */
-	void OnAddNewLocalVariable();
+	virtual void OnAddNewLocalVariable();
 
 	// Type of new document/graph being created by a menu item
 	enum ECreatedDocumentType
diff --git a/Engine/Source/Runtime/Engine/Classes/EdGraph/EdGraphSchema.h b/Engine/Source/Runtime/Engine/Classes/EdGraph/EdGraphSchema.h
--- a/Engine/Source/Runtime/Engine/Classes/EdGraph/EdGraphSchema.h
+++ b/Engine/Source/Runtime/Engine/Classes/EdGraph/EdGraphSchema.h
@@ -14,6 +14,7 @@
 
 class FSlateRect;
 class UEdGraph;
+struct FBPVariableDescription;
 
 /** Distinguishes between different graph types. Graphs can have different properties; for example: functions have one entry point, ubergraphs can have multiples. */
 UENUM()
@@ -288,6 +289,12 @@ public:
 	// (e.g., both are variables in the same Blueprint)
 	virtual FEdGraphSchemaActionDefiningObject GetPersistentItemDefiningObject() const { return FEdGraphSchemaActionDefiningObject(nullptr); }
 
+	// Returns true if the action refers to a external or local variable
+	virtual bool IsA(const FName& InType) const
+	{
+		return InType == GetTypeId();
+	}
+
 private:
 	void UpdateSearchText();
 };
@@ -966,9 +973,9 @@ class ENGINE_API UEdGraphSchema : public UObject
 	virtual bool CanGraphBeDropped(TSharedPtr<FEdGraphSchemaAction> InAction) const { return false; }
 
 	/*
-	 * Begins a drag and drop action to drag a graph into another graph
+	 * Begins a drag and drop action to drag a graph action into another graph
 	 */
-	virtual FReply BeginGraphDragAction(TSharedPtr<FEdGraphSchemaAction> InAction) const { return FReply::Unhandled(); }
+	virtual FReply BeginGraphDragAction(TSharedPtr<FEdGraphSchemaAction> InAction, const FPointerEvent& MouseEvent = FPointerEvent() ) const { return FReply::Unhandled(); }
 
 	/**
 	 * Gets display information for a graph
@@ -1170,4 +1177,25 @@ class ENGINE_API UEdGraphSchema : public UObject
 	 * @return  true if the blueprint marking has been taken care off.
 	 */
 	virtual bool MarkBlueprintDirtyFromNewNode(UBlueprint* InBlueprint, UEdGraphNode* InEdGraphNode) const { return false; }
+
+	/**
+	* Returns the local variables related to the graph.
+	*
+	* @param   InGraph    The graph where to look for local variables
+	* @param   OutLocalVariables    The local variables found in the graph
+	* 
+	* @return  true if the graph can contain local variables (even if it has no local variables)
+	*/
+	virtual bool GetLocalVariables(const UEdGraph* InGraph, TArray<FBPVariableDescription>& OutLocalVariables) const { return false; }
+
+	/**
+	* Generates a graph schema action from a graph and a variable description.
+	*
+	* @param   InGraph    The graph where the variable is located
+	* @param   VariableDescription    The description of the variable from which to generate the action
+	* 
+	* @return  a shared pointer to the newly created action.
+	*/
+	virtual TSharedPtr<FEdGraphSchemaAction> MakeActionFromVariableDescription(const UEdGraph* InEdGraph, const FBPVariableDescription& VariableDescription) const { return nullptr; }
+
 };
diff --git a/Engine/Source/Runtime/RigVM/Private/RigVMCore/RigVMMemory.cpp b/Engine/Source/Runtime/RigVM/Private/RigVMCore/RigVMMemory.cpp
--- a/Engine/Source/Runtime/RigVM/Private/RigVMCore/RigVMMemory.cpp
+++ b/Engine/Source/Runtime/RigVM/Private/RigVMCore/RigVMMemory.cpp
@@ -1,6 +1,7 @@
 // Copyright Epic Games, Inc. All Rights Reserved.
 
 #include "RigVMCore/RigVMMemory.h"
+#include "RigVMCore/RigVMUtilities.h"
 #include "UObject/AnimObjectVersion.h"
 #include "UObject/ReleaseObjectVersion.h"
 #include "UObject/PropertyPortFlags.h"
@@ -1739,7 +1740,7 @@ void FRigVMMemoryContainer::SetRegisterValueFromString(const FRigVMOperand& InOp
 	}
 
 	FString CPPType = InCPPType;
-	if (CPPType.StartsWith(TEXT("TArray<")))
+	if (RigVMUtilities::IsArrayType(CPPType))
 	{
 		CPPType = CPPType.Mid(7, CPPType.Len() - 8);
 	}
diff --git a/Engine/Source/Runtime/RigVM/Public/RigVMCore/RigVMUtilities.h b/Engine/Source/Runtime/RigVM/Public/RigVMCore/RigVMUtilities.h
new file mode 100644
--- /dev/null
+++ b/Engine/Source/Runtime/RigVM/Public/RigVMCore/RigVMUtilities.h
@@ -0,0 +1,14 @@
+﻿// Copyright Epic Games, Inc. All Rights Reserved.
+
+#pragma once
+
+#include "CoreMinimal.h"
+
+namespace RigVMUtilities
+{
+	// Returns true if the type specified is an array
+	FORCEINLINE bool IsArrayType(const FString& InCPPType)
+	{
+		return InCPPType.StartsWith(TEXT("TArray<"));
+	}
+}
\ No newline at end of file
