Control Rig: Hierarchy tracing - follow up on 16558496 for customization on tracing precision

diff --git a/Engine/Plugins/Experimental/ControlRig/Source/ControlRig/Private/Rigs/RigHierarchy.cpp b/Engine/Plugins/Experimental/ControlRig/Source/ControlRig/Private/Rigs/RigHierarchy.cpp
--- a/Engine/Plugins/Experimental/ControlRig/Source/ControlRig/Private/Rigs/RigHierarchy.cpp
+++ b/Engine/Plugins/Experimental/ControlRig/Source/ControlRig/Private/Rigs/RigHierarchy.cpp
@@ -35,6 +35,9 @@ static void RigHierarchyCaptureCallStack(FString& OutCallstack, uint32 NumCallsT
 // CVar to record all transform changes 
 static TAutoConsoleVariable<int32> CVarControlRigHierarchyTraceAlways(TEXT("ControlRig.Hierarchy.TraceAlways"), 0, TEXT("if nonzero we will record all transform changes."));
 static TAutoConsoleVariable<int32> CVarControlRigHierarchyTraceCallstack(TEXT("ControlRig.Hierarchy.TraceCallstack"), 0, TEXT("if nonzero we will record the callstack for any trace entry.\nOnly works if(ControlRig.Hierarchy.TraceEnabled != 0)"));
+static TAutoConsoleVariable<int32> CVarControlRigHierarchyTracePrecision(TEXT("ControlRig.Hierarchy.TracePrecision"), 3, TEXT("sets the number digits in a float when tracing hierarchies."));
+static int32 sRigHierarchyLastTrace = INDEX_NONE;
+static TCHAR sRigHierarchyTraceFormat[16];
 
 // A console command to trace a single frame / single execution for a control rig anim node / control rig component
 FAutoConsoleCommandWithWorldAndArgs FCmdControlRigHierarchyTraceFrames
@@ -1890,13 +1893,26 @@ void URigHierarchy::StorePoseForTrace(const FString& InPrefix)
 	TracePoses.FindOrAdd(CurrentKey) = GetPose(false);
 }
 
+void URigHierarchy::CheckTraceFormatIfRequired()
+{
+	if(sRigHierarchyLastTrace != CVarControlRigHierarchyTracePrecision->GetInt())
+	{
+		sRigHierarchyLastTrace = CVarControlRigHierarchyTracePrecision->GetInt();
+		const FString Format = FString::Printf(TEXT("%%.%df"), sRigHierarchyLastTrace);
+		check(Format.Len() < 16);
+		sRigHierarchyTraceFormat[Format.Len()] = '\0';
+		FMemory::Memcpy(sRigHierarchyTraceFormat, *Format, Format.Len() * sizeof(TCHAR));
+	}
+}
+
 template <class CharType>
 struct TRigHierarchyJsonPrintPolicy
 	: public TPrettyJsonPrintPolicy<CharType>
 {
 	static inline void WriteDouble(  FArchive* Stream, double Value )
 	{
-		TJsonPrintPolicy<CharType>::WriteString(Stream, FString::Printf(TEXT("%.3f"), Value));
+		URigHierarchy::CheckTraceFormatIfRequired();
+		TJsonPrintPolicy<CharType>::WriteString(Stream, FString::Printf(sRigHierarchyTraceFormat, Value));
 	}
 };
 
diff --git a/Engine/Plugins/Experimental/ControlRig/Source/ControlRig/Public/Rigs/RigHierarchy.h b/Engine/Plugins/Experimental/ControlRig/Source/ControlRig/Public/Rigs/RigHierarchy.h
--- a/Engine/Plugins/Experimental/ControlRig/Source/ControlRig/Public/Rigs/RigHierarchy.h
+++ b/Engine/Plugins/Experimental/ControlRig/Source/ControlRig/Public/Rigs/RigHierarchy.h
@@ -2129,6 +2129,11 @@ public:
 	 */
 	void StorePoseForTrace(const FString& InPrefix);
 
+	/**
+	 * Updates the format for trace floating point numbers
+	 */
+	static void CheckTraceFormatIfRequired();
+	
 	/**
 	 * Dumps the content of the transform stack to a string
 	 */
