Updated FRigUnit_SpringInterp and FRigUnit_SpringInterpVector to take advantage of recent changes to spring interpolation:
- removed loop to prevent large delta times, the new spring interpolation can handle them
- added new parameter to control the desired target velocity as the target position changes, which wasn't previously an option.

diff --git a/Engine/Plugins/Experimental/ControlRig/Source/ControlRig/Private/Units/Simulation/RigUnit_SpringInterp.cpp b/Engine/Plugins/Experimental/ControlRig/Source/ControlRig/Private/Units/Simulation/RigUnit_SpringInterp.cpp
--- a/Engine/Plugins/Experimental/ControlRig/Source/ControlRig/Private/Units/Simulation/RigUnit_SpringInterp.cpp
+++ b/Engine/Plugins/Experimental/ControlRig/Source/ControlRig/Private/Units/Simulation/RigUnit_SpringInterp.cpp
@@ -3,11 +3,10 @@
 #include "Units/Simulation/RigUnit_SpringInterp.h"
 #include "Units/RigUnitContext.h"
 
+
 namespace RigUnitSpringInterpConstants
 {
-	static const float FixedTimeStep = 1.0f / 60.0f;
-	static const float MaxTimeStep = 0.1f;
-	static const float TargetVelocityAmount = 1.0f;
+	static const float Mass = 1.0f;
 }
 
 FRigUnit_SpringInterp_Execute()
@@ -20,20 +19,8 @@ FRigUnit_SpringInterp_Execute()
 	}
 	else
 	{
-		// Clamp to avoid large time deltas.
-		float RemainingTime = FMath::Min(Context.DeltaTime, RigUnitSpringInterpConstants::MaxTimeStep);
-
-		Result = Current;
-		while (RemainingTime >= RigUnitSpringInterpConstants::FixedTimeStep)
-		{
-			Result = UKismetMathLibrary::FloatSpringInterp(Result, Target, SpringState, Stiffness, CriticalDamping,
-				RigUnitSpringInterpConstants::FixedTimeStep, Mass, RigUnitSpringInterpConstants::TargetVelocityAmount,
-				false, 0.0f, 0.0f, true);
-			RemainingTime -= RigUnitSpringInterpConstants::FixedTimeStep;
-		}
-
-		Result = UKismetMathLibrary::FloatSpringInterp(Result, Target, SpringState, Stiffness, CriticalDamping,
-			RemainingTime, Mass, RigUnitSpringInterpConstants::TargetVelocityAmount, false, 0.0f, 0.0f, true);
+		Result = UKismetMathLibrary::FloatSpringInterp(Current, Target, SpringState, Stiffness, CriticalDamping,
+			Context.DeltaTime, RigUnitSpringInterpConstants::Mass, TargetVelocityAmount, false, 0.0f, 0.0f, true);
 	}
 }
 
@@ -47,19 +34,7 @@ FRigUnit_SpringInterpVector_Execute()
 	}
 	else
 	{
-		// Clamp to avoid large time deltas.
-		float RemainingTime = FMath::Min(Context.DeltaTime, RigUnitSpringInterpConstants::MaxTimeStep);
-
-		Result = Current;
-		while (RemainingTime >= RigUnitSpringInterpConstants::FixedTimeStep)
-		{
-			Result = UKismetMathLibrary::VectorSpringInterp(Result, Target, SpringState, Stiffness, CriticalDamping,
-				RigUnitSpringInterpConstants::FixedTimeStep, Mass, RigUnitSpringInterpConstants::TargetVelocityAmount,
-				false, FVector(), FVector(), true);
-			RemainingTime -= RigUnitSpringInterpConstants::FixedTimeStep;
-		}
-
-		Result = UKismetMathLibrary::VectorSpringInterp(Result, Target, SpringState, Stiffness, CriticalDamping,
-			RemainingTime, Mass, RigUnitSpringInterpConstants::TargetVelocityAmount, false, FVector(), FVector(), true);
+		Result = UKismetMathLibrary::VectorSpringInterp(Current, Target, SpringState, Stiffness, CriticalDamping,
+			Context.DeltaTime, RigUnitSpringInterpConstants::Mass, TargetVelocityAmount, false, FVector(), FVector(), true);
 	}
-}
\ No newline at end of file
+}
diff --git a/Engine/Plugins/Experimental/ControlRig/Source/ControlRig/Private/Units/Simulation/RigUnit_SpringInterp.h b/Engine/Plugins/Experimental/ControlRig/Source/ControlRig/Private/Units/Simulation/RigUnit_SpringInterp.h
--- a/Engine/Plugins/Experimental/ControlRig/Source/ControlRig/Private/Units/Simulation/RigUnit_SpringInterp.h
+++ b/Engine/Plugins/Experimental/ControlRig/Source/ControlRig/Private/Units/Simulation/RigUnit_SpringInterp.h
@@ -19,28 +19,41 @@ struct CONTROLRIG_API FRigUnit_SpringInterp : public FRigUnit_SimBase
 		Current = Target = Result = 0.0f;
 		Stiffness = 10.0f;
 		CriticalDamping = 2.0f;
-		Mass = 10.0f;
+		TargetVelocityAmount = 0.0f;
 		SpringState = FFloatSpringState();
 	}
 
 	RIGVM_METHOD()
 	virtual void Execute(const FRigUnitContext& Context) override;
 
+	/** Current position of the spring. */
 	UPROPERTY(meta = (Input))
 	float Current;
 
+	/** Rest position of the spring. */
 	UPROPERTY(meta=(Input))
 	float Target;
 
+	/** The spring stiffness determines how hard it will pull towards the target. */
 	UPROPERTY(meta=(Input))
 	float Stiffness;
 
+	/** 
+	 * Set it smaller than 1 to make the spring oscillate before stabilizing on the target. 
+	 * Set it equal to 1 to reach the target without overshooting. 
+	 * Set it higher than one to make the spring take longer to reach the target.
+	 */
 	UPROPERTY(meta=(Input))
 	float CriticalDamping;
 
-	UPROPERTY(meta=(Input))
-	float Mass;
+	/**
+	 * If 1, target changes turn into position offsets and the target is followed with less delay. As it gets close to
+	 * zero, the effect is disabled and only spring pull will affect the output.
+	 */
+	UPROPERTY(meta = (Input))
+	float TargetVelocityAmount;
 
+	/** New position of the spring after delta time. */
 	UPROPERTY(meta=(Output))
 	float Result;
 
@@ -61,28 +74,41 @@ struct CONTROLRIG_API FRigUnit_SpringInterpVector : public FRigUnit_SimBase
 		Current = Target = Result = FVector::ZeroVector;
 		Stiffness = 10.0f;
 		CriticalDamping = 2.0f;
-		Mass = 10.0f;
+		TargetVelocityAmount = 0.0f;
 		SpringState = FVectorSpringState();
 	}
 
 	RIGVM_METHOD()
 	virtual void Execute(const FRigUnitContext& Context) override;
 
+	/** Current position of the spring. */
 	UPROPERTY(meta = (Input))
 	FVector Current;
 
+	/** Rest position of the spring. */
 	UPROPERTY(meta=(Input))
 	FVector Target;
 
+	/** The spring stiffness determines how hard it will pull towards the target. */
 	UPROPERTY(meta=(Input))
 	float Stiffness;
 
+	/**
+	 * Set it smaller than 1 to make the spring oscillate before stabilizing on the target.
+	 * Set it equal to 1 to reach the target without overshooting.
+	 * Set it higher than one to make the spring take longer to reach the target.
+	 */
 	UPROPERTY(meta=(Input))
 	float CriticalDamping;
 
-	UPROPERTY(meta=(Input))
-	float Mass;
+	/** 
+	 * If 1, target changes turn into position offsets and the target is followed with less delay. As it gets close to 
+	 * zero, the effect is disabled and only spring pull will affect the output. 
+	 */
+	UPROPERTY(meta = (Input))
+	float TargetVelocityAmount;
 
+	/** New position of the spring after delta time. */
 	UPROPERTY(meta=(Output))
 	FVector Result;
 
