Control Rig: Only render null capsules if nulls are also displayed in editor

diff --git a/Engine/Plugins/Experimental/ControlRig/Source/ControlRigEditor/Private/Editor/ControlRigEditor.cpp b/Engine/Plugins/Experimental/ControlRig/Source/ControlRigEditor/Private/Editor/ControlRigEditor.cpp
--- a/Engine/Plugins/Experimental/ControlRig/Source/ControlRigEditor/Private/Editor/ControlRigEditor.cpp
+++ b/Engine/Plugins/Experimental/ControlRig/Source/ControlRigEditor/Private/Editor/ControlRigEditor.cpp
@@ -3091,6 +3091,7 @@ void FControlRigEditor::OnToolbarDrawNullsChanged(ECheckBoxState InNewValue)
 	if (FControlRigEditMode* EditMode = GetEditMode())
 	{
 		EditMode->Settings->bDisplayNulls = InNewValue == ECheckBoxState::Checked;
+		UpdateCapsules();
 	}
 }
 
@@ -5278,9 +5279,15 @@ void FControlRigEditor::UpdateCapsules()
 			const FRigElementKey Key = Hierarchy->GetKey(Index);
 			switch(Key.Type)
 			{
-				case ERigElementType::Bone:
 				case ERigElementType::Null:
-				//case ERigElementType::Control:
+				{
+					if(GetToolbarDrawNulls() != ECheckBoxState::Checked)
+					{
+						break;
+					}
+					// otherwise fall through to bone case
+				}
+				case ERigElementType::Bone:
 				{
 					const int32 CapsuleIndex = SelectionComponent->Add();
 					CapsuleToHierarchyIndex.Add(CapsuleIndex, Index);
