Control Rig Editor: Only draw capsules for bones if bones are supposed to be visible

diff --git a/Engine/Plugins/Experimental/ControlRig/Source/ControlRigEditor/Private/Editor/ControlRigEditor.cpp b/Engine/Plugins/Experimental/ControlRig/Source/ControlRigEditor/Private/Editor/ControlRigEditor.cpp
--- a/Engine/Plugins/Experimental/ControlRig/Source/ControlRigEditor/Private/Editor/ControlRigEditor.cpp
+++ b/Engine/Plugins/Experimental/ControlRig/Source/ControlRigEditor/Private/Editor/ControlRigEditor.cpp
@@ -2533,16 +2533,38 @@ void FControlRigEditor::Tick(float DeltaTime)
 			ControlRig->SetDeltaTime(DeltaTime);
 			ControlRig->Evaluate_AnyThread();
 			bDrawHierarchyBones = true;
-			
 		}
 	}
 
-	if (bDrawHierarchyBones)
+	if (FControlRigEditorEditMode* EditMode = GetEditMode())
 	{
-		if (FControlRigEditorEditMode* EditMode = GetEditMode())
+		if (bDrawHierarchyBones)
 		{
 			EditMode->bDrawHierarchyBones = bDrawHierarchyBones;
 		}
+
+		// check if we are supposed to have bones in our map
+		if ((EditMode->ConfigOption != nullptr) && (ControlRig != nullptr))
+		{
+			EBoneDrawMode::Type BoneDrawMode = (EBoneDrawMode::Type)EditMode->ConfigOption->DefaultBoneDrawSelection;
+			const bool bBonesExpectedInCapsules = BoneDrawMode != EBoneDrawMode::None;
+			bool bBonesFoundInCapsules = false;
+
+			for(const TPair<int32, int32>& Pair : CapsuleToHierarchyIndex)
+			{
+				const int32 HierarchyIndex = Pair.Value;
+				if(ControlRig->GetHierarchy()->GetKey(HierarchyIndex).Type == ERigElementType::Bone)
+				{
+					bBonesFoundInCapsules = true;
+					break;
+				}
+			}
+
+			if(bBonesExpectedInCapsules != bBonesFoundInCapsules)
+			{
+				UpdateCapsules();
+			}
+		}
 	}
 }
 
@@ -5269,7 +5291,17 @@ void FControlRigEditor::UpdateCapsules()
 
 	SelectionComponent->Reset();
 	CapsuleToHierarchyIndex.Reset();
-	
+
+	bool bCreateCapsulesForBones = false;
+	if (FControlRigEditorEditMode* EditMode = GetEditMode())
+	{
+		if (EditMode->ConfigOption)
+		{
+			EBoneDrawMode::Type BoneDrawMode = (EBoneDrawMode::Type)EditMode->ConfigOption->DefaultBoneDrawSelection;
+			bCreateCapsulesForBones = BoneDrawMode != EBoneDrawMode::None;
+		}
+	}
+
 	if(UControlRig* DebuggedControlRig = Cast<UControlRig>(GetBlueprintObj()->GetObjectBeingDebugged()))
 	{
 		URigHierarchy* Hierarchy = DebuggedControlRig->GetHierarchy();
@@ -5281,16 +5313,20 @@ void FControlRigEditor::UpdateCapsules()
 			{
 				case ERigElementType::Null:
 				{
-					if(GetToolbarDrawNulls() != ECheckBoxState::Checked)
+					if(GetToolbarDrawNulls() == ECheckBoxState::Checked)
 					{
-						break;
+						const int32 CapsuleIndex = SelectionComponent->Add();
+						CapsuleToHierarchyIndex.Add(CapsuleIndex, Index);
 					}
-					// otherwise fall through to bone case
+					break;
 				}
 				case ERigElementType::Bone:
 				{
-					const int32 CapsuleIndex = SelectionComponent->Add();
-					CapsuleToHierarchyIndex.Add(CapsuleIndex, Index);
+					if(bCreateCapsulesForBones)
+					{
+						const int32 CapsuleIndex = SelectionComponent->Add();
+						CapsuleToHierarchyIndex.Add(CapsuleIndex, Index);
+					}
 					break;
 				}
 				default:
