Control Rig: Set Bone Initials from Ref pose using a lazy pattern whenever the skeletal mesh or anim instance changes.

diff --git a/Engine/Plugins/Experimental/ControlRig/Source/ControlRig/Private/AnimNode_ControlRig.cpp b/Engine/Plugins/Experimental/ControlRig/Source/ControlRig/Private/AnimNode_ControlRig.cpp
--- a/Engine/Plugins/Experimental/ControlRig/Source/ControlRig/Private/AnimNode_ControlRig.cpp
+++ b/Engine/Plugins/Experimental/ControlRig/Source/ControlRig/Private/AnimNode_ControlRig.cpp
@@ -22,6 +22,7 @@ FAnimNode_ControlRig::FAnimNode_ControlRig()
 	, bSetRefPoseFromSkeleton(false)
 	, AlphaCurveName(NAME_None)
 	, LODThreshold(INDEX_NONE)
+	, RefPoseSetterHash()
 {
 }
 
@@ -34,17 +35,10 @@ void FAnimNode_ControlRig::OnInitializeAnimInstance(const FAnimInstanceProxy* In
 		ControlRig = NewObject<UControlRig>(InAnimInstance->GetOwningComponent(), ControlRigClass);
 		ControlRig->Initialize(true);
 		ControlRig->RequestInit();
+		RefPoseSetterHash.Reset();
+		ControlRig->OnInitialized_AnyThread().AddLambda([this](class UControlRig*, const EControlRigState, const FName&) { RefPoseSetterHash.Reset(); });
 
-		if (bSetRefPoseFromSkeleton)
-		{
-			if (InAnimInstance->GetSkelMeshComponent())
-			{
-				if (USkeletalMesh* SkeletalMesh = InAnimInstance->GetSkelMeshComponent()->SkeletalMesh)
-				{
-					ControlRig->SetBoneInitialTransformsFromSkeletalMesh(SkeletalMesh);
-				}
-			}
-		}
+		UpdateControlRigRefPoseIfNeeded(InProxy);
 	}
 
 	FAnimNode_ControlRigBase::OnInitializeAnimInstance(InProxy, InAnimInstance);
@@ -109,6 +103,7 @@ void FAnimNode_ControlRig::Update_AnyThread(const FAnimationUpdateContext& Conte
 		InternalBlendAlpha = 0.f;
 	}
 
+	UpdateControlRigRefPoseIfNeeded(Context.AnimInstanceProxy);
 	FAnimNode_ControlRigBase::Update_AnyThread(Context);
 
 	TRACE_ANIM_NODE_VALUE(Context, TEXT("Class"), *GetNameSafe(ControlRigClass.Get()));
@@ -135,6 +130,8 @@ void FAnimNode_ControlRig::CacheBones_AnyThread(const FAnimationCacheBonesContex
 
 	if(RequiredBones.IsValid())
 	{
+		RefPoseSetterHash.Reset();
+		
 		TArray<FName> const& UIDToNameLookUpTable = RequiredBones.GetUIDToNameLookupTable();
 
 		auto CacheCurveMappingUIDs = [&](const TMap<FName, FName>& Mapping, TArray<FName> const& InUIDToNameLookUpTable, 
@@ -266,6 +263,58 @@ void FAnimNode_ControlRig::UpdateOutput(UControlRig* InControlRig, FPoseContext&
 	}
 }
 
+void FAnimNode_ControlRig::UpdateControlRigRefPoseIfNeeded(const FAnimInstanceProxy* InProxy, bool bIncludePoseInHash)
+{
+	if(!bSetRefPoseFromSkeleton)
+	{
+		return;
+	}
+
+	int32 ExpectedHash = 0;
+	ExpectedHash = HashCombine(ExpectedHash, (int32)reinterpret_cast<uintptr_t>(InProxy->GetAnimInstanceObject()));
+	ExpectedHash = HashCombine(ExpectedHash, (int32)reinterpret_cast<uintptr_t>(InProxy->GetSkelMeshComponent()));
+
+	if(InProxy->GetSkelMeshComponent())
+	{
+		ExpectedHash = HashCombine(ExpectedHash, (int32)reinterpret_cast<uintptr_t>(InProxy->GetSkelMeshComponent()->SkeletalMesh.Get()));
+	}
+
+	if(bIncludePoseInHash)
+	{
+		FMemMark Mark(FMemStack::Get());
+		FCompactPose RefPose;
+		RefPose.ResetToRefPose(InProxy->GetRequiredBones());
+
+		for(const FCompactPoseBoneIndex& BoneIndex : RefPose.ForEachBoneIndex())
+		{
+			const FTransform& Transform = RefPose.GetRefPose(BoneIndex);
+			const FQuat Rotation = Transform.GetRotation();
+
+			ExpectedHash = HashCombine(ExpectedHash, GetTypeHash(Transform.GetTranslation()));
+			ExpectedHash = HashCombine(ExpectedHash, GetTypeHash(FVector4(Rotation.X, Rotation.Y, Rotation.Z, Rotation.W)));
+			ExpectedHash = HashCombine(ExpectedHash, GetTypeHash(Transform.GetScale3D()));
+		}
+
+		if(RefPoseSetterHash.IsSet() && (ExpectedHash == RefPoseSetterHash.GetValue()))
+		{
+			return;
+		}
+
+		ControlRig->SetBoneInitialTransformsFromCompactPose(&RefPose);
+	}
+	else
+	{
+		if(RefPoseSetterHash.IsSet() && (ExpectedHash == RefPoseSetterHash.GetValue()))
+		{
+			return;
+		}
+
+		ControlRig->SetBoneInitialTransformsFromAnimInstanceProxy(InProxy);
+	}
+	
+	RefPoseSetterHash = ExpectedHash;
+}
+
 void FAnimNode_ControlRig::SetIOMapping(bool bInput, const FName& SourceProperty, const FName& TargetCurve)
 {
 	DECLARE_SCOPE_HIERARCHICAL_COUNTER_FUNC()
diff --git a/Engine/Plugins/Experimental/ControlRig/Source/ControlRig/Private/ControlRig.cpp b/Engine/Plugins/Experimental/ControlRig/Source/ControlRig/Private/ControlRig.cpp
--- a/Engine/Plugins/Experimental/ControlRig/Source/ControlRig/Private/ControlRig.cpp
+++ b/Engine/Plugins/Experimental/ControlRig/Source/ControlRig/Private/ControlRig.cpp
@@ -2464,46 +2464,38 @@ bool UControlRig::ExecuteBreakpointAction(const ERigVMBreakpointAction Breakpoin
 
 #endif // WITH_EDITOR
 
-void UControlRig::SetBoneInitialTransformsFromSkeletalMeshComponent(USkeletalMeshComponent* InSkelMeshComp)
+void UControlRig::SetBoneInitialTransformsFromAnimInstance(UAnimInstance* InAnimInstance)
+{
+	FMemMark Mark(FMemStack::Get());
+	FCompactPose OutPose;
+	OutPose.ResetToRefPose(InAnimInstance->GetRequiredBones());
+	SetBoneInitialTransformsFromCompactPose(&OutPose);
+}
+
+void UControlRig::SetBoneInitialTransformsFromAnimInstanceProxy(const FAnimInstanceProxy* InAnimInstanceProxy)
+{
+	FMemMark Mark(FMemStack::Get());
+	FCompactPose OutPose;
+	OutPose.ResetToRefPose(InAnimInstanceProxy->GetRequiredBones());
+	SetBoneInitialTransformsFromCompactPose(&OutPose);
+}
+
+void UControlRig::SetBoneInitialTransformsFromSkeletalMeshComponent(USkeletalMeshComponent* InSkelMeshComp, bool bUseAnimInstance)
 {
 	check(InSkelMeshComp);
 	check(DynamicHierarchy);
-	if (InSkelMeshComp->GetAnimInstance() == nullptr)
+	
+	if (!bUseAnimInstance && (InSkelMeshComp->GetAnimInstance() != nullptr))
 	{
-		SetBoneInitialTransformsFromSkeletalMesh(InSkelMeshComp->SkeletalMesh);
-		return;
+		SetBoneInitialTransformsFromAnimInstance(InSkelMeshComp->GetAnimInstance());
 	}
 	else
 	{
-		FMemMark Mark(FMemStack::Get());
-		FCompactPose OutPose;
-		OutPose.ResetToRefPose(InSkelMeshComp->GetAnimInstance()->GetRequiredBones());
-		if(!OutPose.IsValid())
-		{
-			return;
-		}
-
-		DynamicHierarchy->ForEach<FRigBoneElement>([this, &OutPose, InSkelMeshComp](FRigBoneElement* BoneElement) -> bool
-			{
-				if (BoneElement->BoneType == ERigBoneType::Imported)
-				{
-					int32 MeshIndex = OutPose.GetBoneContainer().GetPoseBoneIndexForBoneName(BoneElement->GetName());
-					if (MeshIndex != INDEX_NONE)
-					{
-						FCompactPoseBoneIndex CPIndex = OutPose.GetBoneContainer().MakeCompactPoseIndex(FMeshPoseBoneIndex(MeshIndex));
-						if (CPIndex != INDEX_NONE)
-						{
-							FTransform LocalInitialTransform = OutPose.GetRefPose(CPIndex);
-							DynamicHierarchy->SetTransform(BoneElement, LocalInitialTransform, ERigTransformType::InitialLocal, true, false);
-						}
-					}
-				}
-				return true;
-			});
-		bResetInitialTransformsBeforeSetup = false;
+		SetBoneInitialTransformsFromSkeletalMesh(InSkelMeshComp->SkeletalMesh);
 	}
 }
 
+
 void UControlRig::SetBoneInitialTransformsFromSkeletalMesh(USkeletalMesh* InSkeletalMesh)
 {
 	check(InSkeletalMesh);
@@ -2528,6 +2520,44 @@ void UControlRig::SetBoneInitialTransformsFromRefSkeleton(const FReferenceSkelet
 		return true;
 	});
 	bResetInitialTransformsBeforeSetup = false;
+	RequestSetup();
+}
+
+void UControlRig::SetBoneInitialTransformsFromCompactPose(FCompactPose* InCompactPose)
+{
+	check(InCompactPose);
+
+	if(!InCompactPose->IsValid())
+	{
+		return;
+	}
+	if(!InCompactPose->GetBoneContainer().IsValid())
+	{
+		return;
+	}
+	
+	FMemMark Mark(FMemStack::Get());
+
+	DynamicHierarchy->ForEach<FRigBoneElement>([this, InCompactPose](FRigBoneElement* BoneElement) -> bool
+		{
+			if (BoneElement->BoneType == ERigBoneType::Imported)
+			{
+				int32 MeshIndex = InCompactPose->GetBoneContainer().GetPoseBoneIndexForBoneName(BoneElement->GetName());
+				if (MeshIndex != INDEX_NONE)
+				{
+					FCompactPoseBoneIndex CPIndex = InCompactPose->GetBoneContainer().MakeCompactPoseIndex(FMeshPoseBoneIndex(MeshIndex));
+					if (CPIndex != INDEX_NONE)
+					{
+						FTransform LocalInitialTransform = InCompactPose->GetRefPose(CPIndex);
+						DynamicHierarchy->SetTransform(BoneElement, LocalInitialTransform, ERigTransformType::InitialLocal, true, false);
+					}
+				}
+			}
+			return true;
+		});
+
+	bResetInitialTransformsBeforeSetup = false;
+	RequestSetup();
 }
 
 void UControlRig::OnHierarchyTransformUndoRedo(URigHierarchy* InHierarchy, const FRigElementKey& InKey, ERigTransformType::Type InTransformType, const FTransform& InTransform, bool bIsUndo)
diff --git a/Engine/Plugins/Experimental/ControlRig/Source/ControlRig/Private/Sequencer/MovieSceneControlRigParameterSection.cpp b/Engine/Plugins/Experimental/ControlRig/Source/ControlRig/Private/Sequencer/MovieSceneControlRigParameterSection.cpp
--- a/Engine/Plugins/Experimental/ControlRig/Source/ControlRig/Private/Sequencer/MovieSceneControlRigParameterSection.cpp
+++ b/Engine/Plugins/Experimental/ControlRig/Source/ControlRig/Private/Sequencer/MovieSceneControlRigParameterSection.cpp
@@ -2143,7 +2143,7 @@ bool UMovieSceneControlRigParameterSection::LoadAnimSequenceIntoThisSection(UAni
 	const FAnimationCurveData& CurveData = DataModel->GetCurveData();
 	const TArray<FBoneAnimationTrack>& BoneAnimationTracks = DataModel->GetBoneAnimationTracks();
 
-	ControlRig->SetBoneInitialTransformsFromSkeletalMeshComponent(SkelMeshComp);
+	ControlRig->SetBoneInitialTransformsFromSkeletalMeshComponent(SkelMeshComp, true);
 	ControlRig->RequestSetup();
 	ControlRig->Evaluate_AnyThread();
 
diff --git a/Engine/Plugins/Experimental/ControlRig/Source/ControlRig/Private/Sequencer/MovieSceneControlRigParameterTemplate.cpp b/Engine/Plugins/Experimental/ControlRig/Source/ControlRig/Private/Sequencer/MovieSceneControlRigParameterTemplate.cpp
--- a/Engine/Plugins/Experimental/ControlRig/Source/ControlRig/Private/Sequencer/MovieSceneControlRigParameterTemplate.cpp
+++ b/Engine/Plugins/Experimental/ControlRig/Source/ControlRig/Private/Sequencer/MovieSceneControlRigParameterTemplate.cpp
@@ -563,7 +563,7 @@ void FControlRigBindingHelper::BindToSequencerInstance(UControlRig* ControlRig)
 				AnimInstance->AddControlRigTrack(ControlRig->GetUniqueID(), ControlRig);
 				ControlRig->Initialize();
 				ControlRig->RequestInit();
-				ControlRig->SetBoneInitialTransformsFromSkeletalMeshComponent(SkeletalMeshComponent);
+				ControlRig->SetBoneInitialTransformsFromSkeletalMeshComponent(SkeletalMeshComponent, true);
 				ControlRig->Evaluate_AnyThread();
 			}
 		}
diff --git a/Engine/Plugins/Experimental/ControlRig/Source/ControlRig/Public/AnimNode_ControlRig.h b/Engine/Plugins/Experimental/ControlRig/Source/ControlRig/Public/AnimNode_ControlRig.h
--- a/Engine/Plugins/Experimental/ControlRig/Source/ControlRig/Public/AnimNode_ControlRig.h
+++ b/Engine/Plugins/Experimental/ControlRig/Source/ControlRig/Public/AnimNode_ControlRig.h
@@ -100,6 +100,12 @@ protected:
 	virtual void UpdateInput(UControlRig* InControlRig, const FPoseContext& InOutput) override;
 	virtual void UpdateOutput(UControlRig* InControlRig, FPoseContext& InOutput) override;
 
+	// Helper function to update the initial ref pose within the Control Rig if needed
+	void UpdateControlRigRefPoseIfNeeded(const FAnimInstanceProxy* InProxy, bool bIncludePoseInHash = false);
+	
+	// A hash to encode the pointer to anim instance, an
+	TOptional<int32> RefPoseSetterHash;
+
 public:
 
 	void PostSerialize(const FArchive& Ar);
diff --git a/Engine/Plugins/Experimental/ControlRig/Source/ControlRig/Public/ControlRig.h b/Engine/Plugins/Experimental/ControlRig/Source/ControlRig/Public/ControlRig.h
--- a/Engine/Plugins/Experimental/ControlRig/Source/ControlRig/Public/ControlRig.h
+++ b/Engine/Plugins/Experimental/ControlRig/Source/ControlRig/Public/ControlRig.h
@@ -380,16 +380,32 @@ public:
 	FControlRigExecuteEvent& OnExecuted_AnyThread() { return ExecutedEvent; }
 	FRigEventDelegate& OnRigEvent_AnyThread() { return RigEventDelegate; }
 
-	// Setup the initial transform / ref pose of the bones based upon skeletal mesh component
-	// Unlike the other two functions below this uses the current refpose instead of the RefSkeleton pose.
-	void SetBoneInitialTransformsFromSkeletalMeshComponent(USkeletalMeshComponent* InSkelMeshComp);
+	// Setup the initial transform / ref pose of the bones based upon an anim instance
+	// This uses the current refpose instead of the RefSkeleton pose.
+	void SetBoneInitialTransformsFromAnimInstance(UAnimInstance* InAnimInstance);
+
+	// Setup the initial transform / ref pose of the bones based upon an anim instance proxy
+	// This uses the current refpose instead of the RefSkeleton pose.
+	void SetBoneInitialTransformsFromAnimInstanceProxy(const FAnimInstanceProxy* InAnimInstanceProxy);
+
+	// Setup the initial transform / ref pose of the bones based upon skeletal mesh component (ref skeleton)
+	// This uses the RefSkeleton pose instead of the current refpose (or vice versae is bUseAnimInstance == true)
+	void SetBoneInitialTransformsFromSkeletalMeshComponent(USkeletalMeshComponent* InSkelMeshComp, bool bUseAnimInstance = false);
 
 	// Setup the initial transforms / ref pose of the bones based on a skeletal mesh
+	// This uses the RefSkeleton pose instead of the current refpose.
 	void SetBoneInitialTransformsFromSkeletalMesh(USkeletalMesh* InSkeletalMesh);
 
 	// Setup the initial transforms / ref pose of the bones based on a reference skeleton
+	// This uses the RefSkeleton pose instead of the current refpose.
 	void SetBoneInitialTransformsFromRefSkeleton(const FReferenceSkeleton& InReferenceSkeleton);
 
+private:
+
+	void SetBoneInitialTransformsFromCompactPose(FCompactPose* InCompactPose);
+
+public:
+	
 	const FControlRigDrawInterface& GetDrawInterface() const { return DrawInterface; };
 	FControlRigDrawInterface& GetDrawInterface() { return DrawInterface; };
 
@@ -694,6 +710,7 @@ public:
 	friend class UControlRigBlueprintGeneratedClass;
 	friend class FControlRigInteractionScope;
 	friend class UControlRigValidator;
+	friend struct FAnimNode_ControlRig;
 };
 
 class CONTROLRIG_API FControlRigBracketScope
diff --git a/Engine/Plugins/Experimental/ControlRig/Source/ControlRigEditor/Private/ControlRigSequencerEditorLibrary.cpp b/Engine/Plugins/Experimental/ControlRig/Source/ControlRigEditor/Private/ControlRigSequencerEditorLibrary.cpp
--- a/Engine/Plugins/Experimental/ControlRig/Source/ControlRigEditor/Private/ControlRigSequencerEditorLibrary.cpp
+++ b/Engine/Plugins/Experimental/ControlRig/Source/ControlRigEditor/Private/ControlRigSequencerEditorLibrary.cpp
@@ -744,7 +744,7 @@ bool UControlRigSequencerEditorLibrary::BakeToControlRig(UWorld* World, ULevelSe
 					ControlRig->GetDataSourceRegistry()->RegisterDataSource(UControlRig::OwnerComponent, ControlRig->GetObjectBinding()->GetBoundObject());
 					ControlRig->Initialize();
 					ControlRig->RequestInit();
-					ControlRig->SetBoneInitialTransformsFromSkeletalMeshComponent(SkeletalMeshComp);
+					ControlRig->SetBoneInitialTransformsFromSkeletalMeshComponent(SkeletalMeshComp, true);
 					ControlRig->Evaluate_AnyThread();
 
 					bool bSequencerOwnsControlRig = true;
diff --git a/Engine/Plugins/Experimental/ControlRig/Source/ControlRigEditor/Private/Sequencer/ControlRigParameterTrackEditor.cpp b/Engine/Plugins/Experimental/ControlRig/Source/ControlRigEditor/Private/Sequencer/ControlRigParameterTrackEditor.cpp
--- a/Engine/Plugins/Experimental/ControlRig/Source/ControlRigEditor/Private/Sequencer/ControlRigParameterTrackEditor.cpp
+++ b/Engine/Plugins/Experimental/ControlRig/Source/ControlRigEditor/Private/Sequencer/ControlRigParameterTrackEditor.cpp
@@ -946,7 +946,7 @@ void FControlRigParameterTrackEditor::BakeToControlRig(UClass* InClass, FGuid Ob
 					ControlRig->GetDataSourceRegistry()->RegisterDataSource(UControlRig::OwnerComponent, ControlRig->GetObjectBinding()->GetBoundObject());
 					ControlRig->Initialize();
 					ControlRig->RequestInit();
-					ControlRig->SetBoneInitialTransformsFromSkeletalMeshComponent(SkelMeshComp);
+					ControlRig->SetBoneInitialTransformsFromSkeletalMeshComponent(SkelMeshComp, true);
 					ControlRig->Evaluate_AnyThread();
 
 					UMovieSceneSection* NewSection = Track->CreateControlRigSection(0, ControlRig, bSequencerOwnsControlRig);
