Control Rig: Support description for functions

diff --git a/Engine/Plugins/Experimental/ControlRig/Source/ControlRigDeveloper/Private/ControlRigBlueprint.cpp b/Engine/Plugins/Experimental/ControlRig/Source/ControlRigDeveloper/Private/ControlRigBlueprint.cpp
--- a/Engine/Plugins/Experimental/ControlRig/Source/ControlRigDeveloper/Private/ControlRigBlueprint.cpp
+++ b/Engine/Plugins/Experimental/ControlRig/Source/ControlRigDeveloper/Private/ControlRigBlueprint.cpp
@@ -2622,6 +2622,7 @@ void UControlRigBlueprint::HandleModifiedEvent(ERigVMGraphNotifType InNotifType,
 			}
 			case ERigVMGraphNotifType::NodeCategoryChanged:
 			case ERigVMGraphNotifType::NodeKeywordsChanged:
+			case ERigVMGraphNotifType::NodeDescriptionChanged:
 			{
 				FBlueprintEditorUtils::MarkBlueprintAsStructurallyModified(this);
 				break;
diff --git a/Engine/Plugins/Experimental/ControlRig/Source/ControlRigEditor/Private/ControlRigGraphDetails.cpp b/Engine/Plugins/Experimental/ControlRig/Source/ControlRigEditor/Private/ControlRigGraphDetails.cpp
--- a/Engine/Plugins/Experimental/ControlRig/Source/ControlRigEditor/Private/ControlRigGraphDetails.cpp
+++ b/Engine/Plugins/Experimental/ControlRig/Source/ControlRigEditor/Private/ControlRigGraphDetails.cpp
@@ -13,6 +13,7 @@
 #include "Widgets/Text/STextBlock.h"
 #include "Widgets/Colors/SColorPicker.h"
 #include "Widgets/Text/SInlineEditableTextBlock.h"
+#include "Widgets/Text/SMultiLineEditableText.h"
 #include "PropertyCustomizationHelpers.h"
 #include "NodeFactory.h"
 #include "Graph/ControlRigGraphNode.h"
@@ -669,6 +670,22 @@ void FControlRigGraphDetails::CustomizeDetails(IDetailLayoutBuilder& DetailLayou
 			.OnTextCommitted(this, &FControlRigGraphDetails::SetNodeKeywords)
 		];
 
+		// description
+		SettingsCategory.AddCustomRow(FText::GetEmpty())
+		.NameContent()
+		[
+			SNew(STextBlock)
+			.Text(FText::FromString(TEXT("Description")))
+			.Font(IDetailLayoutBuilder::GetDetailFont())
+		]
+		.ValueContent()
+		[
+			SNew(SMultiLineEditableText)
+			.Font(IDetailLayoutBuilder::GetDetailFont())
+			.Text(this, &FControlRigGraphDetails::GetNodeDescription)
+			.OnTextCommitted(this, &FControlRigGraphDetails::SetNodeDescription)
+		];
+
 		if(AccessSpecifierStrings.IsEmpty())
 		{
 			AccessSpecifierStrings.Add(TSharedPtr<FString>(new FString(TEXT("Public"))));
@@ -868,6 +885,41 @@ void FControlRigGraphDetails::SetNodeKeywords(const FText& InNewText, ETextCommi
 	}
 }
 
+FText FControlRigGraphDetails::GetNodeDescription() const
+{
+	if (GraphPtr.IsValid() && ControlRigBlueprintPtr.IsValid())
+	{
+		UControlRigBlueprint* Blueprint = ControlRigBlueprintPtr.Get();
+		if (URigVMGraph* Model = Blueprint->GetModel(GraphPtr.Get()))
+		{
+			if (URigVMCollapseNode* OuterNode = Cast<URigVMCollapseNode>(Model->GetOuter()))
+			{
+				return FText::FromString(OuterNode->GetNodeDescription());
+			}
+		}
+	}
+
+	return FText();
+}
+
+void FControlRigGraphDetails::SetNodeDescription(const FText& InNewText, ETextCommit::Type InCommitType)
+{
+	if (GraphPtr.IsValid() && ControlRigBlueprintPtr.IsValid())
+	{
+		UControlRigBlueprint* Blueprint = ControlRigBlueprintPtr.Get();
+		if (URigVMGraph* Model = Blueprint->GetModel(GraphPtr.Get()))
+		{
+			if (URigVMCollapseNode* OuterNode = Cast<URigVMCollapseNode>(Model->GetOuter()))
+			{
+				if (URigVMController* Controller = Blueprint->GetOrCreateController(OuterNode->GetGraph()))
+				{
+					Controller->SetNodeDescription(OuterNode, InNewText.ToString());
+				}
+			}
+		}
+	}
+}
+
 FLinearColor FControlRigGraphDetails::GetNodeColor() const
 {
 	if (GraphPtr.IsValid() && ControlRigBlueprintPtr.IsValid())
diff --git a/Engine/Plugins/Experimental/ControlRig/Source/ControlRigEditor/Private/ControlRigGraphDetails.h b/Engine/Plugins/Experimental/ControlRig/Source/ControlRigEditor/Private/ControlRigGraphDetails.h
--- a/Engine/Plugins/Experimental/ControlRig/Source/ControlRigEditor/Private/ControlRigGraphDetails.h
+++ b/Engine/Plugins/Experimental/ControlRig/Source/ControlRigEditor/Private/ControlRigGraphDetails.h
@@ -160,6 +160,8 @@ public:
 	void SetNodeCategory(const FText& InNewText, ETextCommit::Type InCommitType);
 	FText GetNodeKeywords() const;
 	void SetNodeKeywords(const FText& InNewText, ETextCommit::Type InCommitType);
+	FText GetNodeDescription() const;
+	void SetNodeDescription(const FText& InNewText, ETextCommit::Type InCommitType);
 	FLinearColor GetNodeColor() const;
 	void SetNodeColor(FLinearColor InColor, bool bSetupUndoRedo);
 	void OnNodeColorBegin();
diff --git a/Engine/Source/Developer/RigVMDeveloper/Private/RigVMModel/Nodes/RigVMCollapseNode.cpp b/Engine/Source/Developer/RigVMDeveloper/Private/RigVMModel/Nodes/RigVMCollapseNode.cpp
--- a/Engine/Source/Developer/RigVMDeveloper/Private/RigVMModel/Nodes/RigVMCollapseNode.cpp
+++ b/Engine/Source/Developer/RigVMDeveloper/Private/RigVMModel/Nodes/RigVMCollapseNode.cpp
@@ -3,6 +3,16 @@
 #include "RigVMModel/Nodes/RigVMCollapseNode.h"
 #include "RigVMModel/RigVMFunctionLibrary.h"
 
+FText URigVMCollapseNode::GetToolTipText() const
+{
+	const FString ToolTipString = GetNodeDescription();
+	if(!ToolTipString.IsEmpty())
+	{
+		return FText::FromString(ToolTipString);
+	}
+	return Super::GetToolTipText();
+}
+
 URigVMFunctionLibrary* URigVMCollapseNode::GetLibrary() const
 {
 	return Cast<URigVMFunctionLibrary>(GetOuter());
diff --git a/Engine/Source/Developer/RigVMDeveloper/Private/RigVMModel/Nodes/RigVMFunctionReferenceNode.cpp b/Engine/Source/Developer/RigVMDeveloper/Private/RigVMModel/Nodes/RigVMFunctionReferenceNode.cpp
--- a/Engine/Source/Developer/RigVMDeveloper/Private/RigVMModel/Nodes/RigVMFunctionReferenceNode.cpp
+++ b/Engine/Source/Developer/RigVMDeveloper/Private/RigVMModel/Nodes/RigVMFunctionReferenceNode.cpp
@@ -21,6 +21,15 @@ FLinearColor URigVMFunctionReferenceNode::GetNodeColor() const
 	return Super::GetNodeColor();
 }
 
+FText URigVMFunctionReferenceNode::GetToolTipText() const
+{
+	if (URigVMLibraryNode* ReferencedNode = GetReferencedNode())
+	{
+		return ReferencedNode->GetToolTipText();
+	}
+	return Super::GetToolTipText();
+}
+
 FString URigVMFunctionReferenceNode::GetNodeCategory() const
 {
 	if (URigVMLibraryNode* ReferencedNode = GetReferencedNode())
diff --git a/Engine/Source/Developer/RigVMDeveloper/Private/RigVMModel/RigVMController.cpp b/Engine/Source/Developer/RigVMDeveloper/Private/RigVMModel/RigVMController.cpp
--- a/Engine/Source/Developer/RigVMDeveloper/Private/RigVMModel/RigVMController.cpp
+++ b/Engine/Source/Developer/RigVMDeveloper/Private/RigVMModel/RigVMController.cpp
@@ -4288,6 +4288,51 @@ bool URigVMController::SetNodeKeywordsByName(const FName& InNodeName, const FStr
 	return SetNodeKeywords(Node, InKeywords, bSetupUndoRedo, bMergeUndoAction);
 }
 
+bool URigVMController::SetNodeDescription(URigVMCollapseNode* InNode, const FString& InDescription, bool bSetupUndoRedo, bool bMergeUndoAction)
+{
+	if (!IsValidNodeForGraph(InNode))
+	{
+		return false;
+	}
+
+	if (InNode->GetNodeDescription() == InDescription)
+	{
+		return false;
+	}
+
+	FRigVMSetNodeDescriptionAction Action;
+	if (bSetupUndoRedo)
+	{
+		Action = FRigVMSetNodeDescriptionAction(InNode, InDescription);
+		Action.Title = FString::Printf(TEXT("Set Node Description"));
+		ActionStack->BeginAction(Action);
+	}
+
+	InNode->NodeDescription = InDescription;
+	Notify(ERigVMGraphNotifType::NodeDescriptionChanged, InNode);
+
+	if (bSetupUndoRedo)
+	{
+		ActionStack->EndAction(Action, bMergeUndoAction);
+	}
+
+	return true;
+}
+
+bool URigVMController::SetNodeDescriptionByName(const FName& InNodeName, const FString& InDescription, bool bSetupUndoRedo, bool bMergeUndoAction)
+{
+	if (!IsValidGraph())
+	{
+		return false;
+	}
+
+	URigVMGraph* Graph = GetGraph();
+	check(Graph);
+
+	URigVMCollapseNode* Node = Cast<URigVMCollapseNode>(Graph->FindNodeByName(InNodeName));
+	return SetNodeDescription(Node, InDescription, bSetupUndoRedo, bMergeUndoAction);
+}
+
 bool URigVMController::SetCommentText(URigVMNode* InNode, const FString& InCommentText, bool bSetupUndoRedo)
 {
 	if (!IsValidNodeForGraph(InNode))
diff --git a/Engine/Source/Developer/RigVMDeveloper/Private/RigVMModel/RigVMControllerActions.cpp b/Engine/Source/Developer/RigVMDeveloper/Private/RigVMModel/RigVMControllerActions.cpp
--- a/Engine/Source/Developer/RigVMDeveloper/Private/RigVMModel/RigVMControllerActions.cpp
+++ b/Engine/Source/Developer/RigVMDeveloper/Private/RigVMModel/RigVMControllerActions.cpp
@@ -995,6 +995,48 @@ bool FRigVMSetNodeKeywordsAction::Redo(URigVMController* InController)
 	return FRigVMBaseAction::Redo(InController);
 }
 
+FRigVMSetNodeDescriptionAction::FRigVMSetNodeDescriptionAction(URigVMCollapseNode* InNode, const FString& InNewDescription)
+	: NodePath(InNode->GetNodePath())
+	, OldDescription(InNode->GetNodeDescription())
+	, NewDescription(InNewDescription)
+{
+}
+
+bool FRigVMSetNodeDescriptionAction::Merge(const FRigVMBaseAction* Other)
+{
+	if (!FRigVMBaseAction::Merge(Other))
+	{
+		return false;
+	}
+
+	const FRigVMSetNodeDescriptionAction* Action = (const FRigVMSetNodeDescriptionAction*)Other;
+	if (NodePath != Action->NodePath)
+	{
+		return false;
+	}
+
+	NewDescription = Action->NewDescription;
+	return true;
+}
+
+bool FRigVMSetNodeDescriptionAction::Undo(URigVMController* InController)
+{
+	if (!FRigVMBaseAction::Undo(InController))
+	{
+		return false;
+	}
+	return InController->SetNodeDescriptionByName(*NodePath, OldDescription, false);
+}
+
+bool FRigVMSetNodeDescriptionAction::Redo(URigVMController* InController)
+{
+	if (!InController->SetNodeDescriptionByName(*NodePath, NewDescription, false))
+	{
+		return false;
+	}
+	return FRigVMBaseAction::Redo(InController);
+}
+
 FRigVMSetCommentTextAction::FRigVMSetCommentTextAction(URigVMCommentNode* InNode, const FString& InNewText)
 : NodePath(InNode->GetNodePath())
 , OldText(InNode->GetCommentText())
diff --git a/Engine/Source/Developer/RigVMDeveloper/Public/RigVMModel/Nodes/RigVMCollapseNode.h b/Engine/Source/Developer/RigVMDeveloper/Public/RigVMModel/Nodes/RigVMCollapseNode.h
--- a/Engine/Source/Developer/RigVMDeveloper/Public/RigVMModel/Nodes/RigVMCollapseNode.h
+++ b/Engine/Source/Developer/RigVMDeveloper/Public/RigVMModel/Nodes/RigVMCollapseNode.h
@@ -17,9 +17,13 @@ class RIGVMDEVELOPER_API URigVMCollapseNode : public URigVMLibraryNode
 
 public:
 
+	// RigVM node interface
+	virtual FText GetToolTipText() const override;
+
 	// Library node interface
 	virtual FString GetNodeCategory() const override { return NodeCategory; }
 	virtual FString GetNodeKeywords() const override { return NodeKeywords; }
+	virtual FString GetNodeDescription() const override { return NodeDescription; }
 	virtual URigVMFunctionLibrary* GetLibrary() const override;
 	virtual URigVMGraph* GetContainedGraph() const override { return ContainedGraph; }
 
@@ -36,6 +40,9 @@ private:
 	UPROPERTY()
 	FString NodeKeywords;
 
+	UPROPERTY()
+	FString NodeDescription;
+
 	friend class URigVMController;
 };
 
diff --git a/Engine/Source/Developer/RigVMDeveloper/Public/RigVMModel/Nodes/RigVMFunctionReferenceNode.h b/Engine/Source/Developer/RigVMDeveloper/Public/RigVMModel/Nodes/RigVMFunctionReferenceNode.h
--- a/Engine/Source/Developer/RigVMDeveloper/Public/RigVMModel/Nodes/RigVMFunctionReferenceNode.h
+++ b/Engine/Source/Developer/RigVMDeveloper/Public/RigVMModel/Nodes/RigVMFunctionReferenceNode.h
@@ -22,6 +22,7 @@ public:
 	// URigVMNode interface
 	virtual FString GetNodeTitle() const override;
 	virtual FLinearColor GetNodeColor() const override;
+	virtual FText GetToolTipText() const override;
 	// end URigVMNode interface
 
 	// URigVMLibraryNode interface
diff --git a/Engine/Source/Developer/RigVMDeveloper/Public/RigVMModel/Nodes/RigVMLibraryNode.h b/Engine/Source/Developer/RigVMDeveloper/Public/RigVMModel/Nodes/RigVMLibraryNode.h
--- a/Engine/Source/Developer/RigVMDeveloper/Public/RigVMModel/Nodes/RigVMLibraryNode.h
+++ b/Engine/Source/Developer/RigVMDeveloper/Public/RigVMModel/Nodes/RigVMLibraryNode.h
@@ -32,6 +32,7 @@ public:
 	// Library node interface
 	virtual FString GetNodeCategory() const { return FString(); }
 	virtual FString GetNodeKeywords() const { return FString(); }
+	virtual FString GetNodeDescription() const { return FString(); }
 	virtual URigVMFunctionLibrary* GetLibrary() const { return nullptr; }
 	virtual URigVMGraph* GetContainedGraph() const { return nullptr; }
 	virtual const TArray<URigVMNode*>& GetContainedNodes() const;
diff --git a/Engine/Source/Developer/RigVMDeveloper/Public/RigVMModel/RigVMController.h b/Engine/Source/Developer/RigVMDeveloper/Public/RigVMModel/RigVMController.h
--- a/Engine/Source/Developer/RigVMDeveloper/Public/RigVMModel/RigVMController.h
+++ b/Engine/Source/Developer/RigVMDeveloper/Public/RigVMModel/RigVMController.h
@@ -425,6 +425,16 @@ public:
 	UFUNCTION(BlueprintCallable, Category = RigVMController)
 	bool SetNodeKeywordsByName(const FName& InNodeName, const FString& InKeywords, bool bSetupUndoRedo = true, bool bMergeUndoAction = false);
 
+	// Sets the function description of a node in the graph.
+	// This causes a NodeDescriptionChanged modified event.
+	UFUNCTION(BlueprintCallable, Category = RigVMController)
+	bool SetNodeDescription(URigVMCollapseNode* InNode, const FString& InDescription, bool bSetupUndoRedo = true, bool bMergeUndoAction = false);
+
+	// Sets the keywords of a node in the graph.
+	// This causes a NodeDescriptionChanged modified event.
+	UFUNCTION(BlueprintCallable, Category = RigVMController)
+	bool SetNodeDescriptionByName(const FName& InNodeName, const FString& InDescription, bool bSetupUndoRedo = true, bool bMergeUndoAction = false);
+
 	// Sets the comment text of a comment node in the graph.
 	// This causes a CommentTextChanged modified event.
 	UFUNCTION(BlueprintCallable, Category = RigVMController)
diff --git a/Engine/Source/Developer/RigVMDeveloper/Public/RigVMModel/RigVMControllerActions.h b/Engine/Source/Developer/RigVMDeveloper/Public/RigVMModel/RigVMControllerActions.h
--- a/Engine/Source/Developer/RigVMDeveloper/Public/RigVMModel/RigVMControllerActions.h
+++ b/Engine/Source/Developer/RigVMDeveloper/Public/RigVMModel/RigVMControllerActions.h
@@ -840,6 +840,36 @@ public:
 	FString NewKeywords;
 };
 
+/**
+* An action setting a node's description in the graph.
+*/
+USTRUCT()
+struct FRigVMSetNodeDescriptionAction : public FRigVMBaseAction
+{
+	GENERATED_BODY()
+
+public:
+
+	FRigVMSetNodeDescriptionAction()
+	{
+		OldDescription = NewDescription = FString();
+	}
+	FRigVMSetNodeDescriptionAction(URigVMCollapseNode* InNode, const FString& InNewDescription);
+	virtual ~FRigVMSetNodeDescriptionAction() {};
+	virtual bool Merge(const FRigVMBaseAction* Other);
+	virtual bool Undo(URigVMController* InController) override;
+	virtual bool Redo(URigVMController* InController) override;
+
+	UPROPERTY()
+	FString NodePath;
+
+	UPROPERTY()
+	FString OldDescription;
+
+	UPROPERTY()
+	FString NewDescription;
+};
+
 /**
  * An action setting a comment node's text in the graph.
  */
diff --git a/Engine/Source/Developer/RigVMDeveloper/Public/RigVMModel/RigVMNotifications.h b/Engine/Source/Developer/RigVMDeveloper/Public/RigVMModel/RigVMNotifications.h
--- a/Engine/Source/Developer/RigVMDeveloper/Public/RigVMModel/RigVMNotifications.h
+++ b/Engine/Source/Developer/RigVMDeveloper/Public/RigVMModel/RigVMNotifications.h
@@ -52,6 +52,7 @@ enum class ERigVMGraphNotifType : uint8
 	NodeReferenceChanged, // A node has changed it's referenced function
 	NodeCategoryChanged, // A node's category has changed (Subject == URigVMNode)
 	NodeKeywordsChanged, // A node's keywords have changed (Subject == URigVMNode)
+	NodeDescriptionChanged, // A node's description has changed (Subject == URigVMNode)
 	VariableRemappingChanged, // A function reference node's remapping has changed (Subject == URigVMFunctionReferenceNode)
 	Invalid // The max for this enum (used for guarding)
 };
