Control Rig: Fix a regression in Spring IK node that caused infinite loop

diff --git a/Engine/Plugins/Experimental/ControlRig/Source/ControlRig/Private/Units/Highlevel/Hierarchy/RigUnit_SpringIK.cpp b/Engine/Plugins/Experimental/ControlRig/Source/ControlRig/Private/Units/Highlevel/Hierarchy/RigUnit_SpringIK.cpp
--- a/Engine/Plugins/Experimental/ControlRig/Source/ControlRig/Private/Units/Highlevel/Hierarchy/RigUnit_SpringIK.cpp
+++ b/Engine/Plugins/Experimental/ControlRig/Source/ControlRig/Private/Units/Highlevel/Hierarchy/RigUnit_SpringIK.cpp
@@ -40,7 +40,7 @@ FRigUnit_SpringIK_Execute()
 
 			while (EndBoneIndex != INDEX_NONE)
 			{
-				const FRigElementKey Key(EndBone, ERigElementType::Bone);
+				const FRigElementKey Key = Hierarchy->GetKey(EndBoneIndex);
 				CachedBones.Add(FCachedRigElement(Key, Hierarchy));
 				if (EndBoneIndex == StartBoneIndex)
 				{
