Control Rig: Fix unity build error in ContorlRig.h

diff --git a/Engine/Plugins/Experimental/ControlRig/Source/ControlRig/Public/ControlRig.h b/Engine/Plugins/Experimental/ControlRig/Source/ControlRig/Public/ControlRig.h
--- a/Engine/Plugins/Experimental/ControlRig/Source/ControlRig/Public/ControlRig.h
+++ b/Engine/Plugins/Experimental/ControlRig/Source/ControlRig/Public/ControlRig.h
@@ -18,6 +18,8 @@
 #include "Components/SceneComponent.h"
 #include "Engine/AssetUserData.h"
 #include "Interfaces/Interface_AssetUserData.h"
+#include "Animation/AnimInstance.h"
+#include "Animation/AnimInstanceProxy.h"
 
 #if WITH_EDITOR
 #include "RigVMModel/RigVMPin.h"
