Control Rig: Show tooltips in my blueprint for functions + graphs

diff --git a/Engine/Plugins/Experimental/ControlRig/Source/ControlRigDeveloper/Private/Graph/ControlRigGraphSchema.cpp b/Engine/Plugins/Experimental/ControlRig/Source/ControlRigDeveloper/Private/Graph/ControlRigGraphSchema.cpp
--- a/Engine/Plugins/Experimental/ControlRig/Source/ControlRigDeveloper/Private/Graph/ControlRigGraphSchema.cpp
+++ b/Engine/Plugins/Experimental/ControlRig/Source/ControlRigDeveloper/Private/Graph/ControlRigGraphSchema.cpp
@@ -715,6 +715,23 @@ void UControlRigGraphSchema::GetGraphDisplayInformation(const UEdGraph& Graph, /
 		{
 			DisplayInfo.DisplayName = FText::FromString(NodePathParts.Last());
 			DisplayInfo.PlainName = DisplayInfo.DisplayName;
+
+			static const FText LocalFunctionText = FText::FromString(TEXT("A local function.")); 
+			DisplayInfo.Tooltip = LocalFunctionText;
+
+			// if this is a riggraph within a collapse node - let's use that for the tooltip
+			if(URigVMGraph* Model = RigGraph->GetModel())
+			{
+				if(URigVMCollapseNode* CollapseNode = Model->GetTypedOuter<URigVMCollapseNode>())
+				{
+					DisplayInfo.Tooltip = CollapseNode->GetToolTipText();
+				}
+			}
+		}
+		else
+		{
+			static const FText MainGraphText = FText::FromString(TEXT("The main graph for the Control Rig."));
+			DisplayInfo.Tooltip = MainGraphText;
 		}
 	}
 }
