Control Rig: Frame camera upon import of hierarchy

diff --git a/Engine/Plugins/Experimental/ControlRig/Source/ControlRigEditor/Private/EditMode/ControlRigEditMode.cpp b/Engine/Plugins/Experimental/ControlRig/Source/ControlRigEditor/Private/EditMode/ControlRigEditMode.cpp
--- a/Engine/Plugins/Experimental/ControlRig/Source/ControlRigEditor/Private/EditMode/ControlRigEditMode.cpp
+++ b/Engine/Plugins/Experimental/ControlRig/Source/ControlRigEditor/Private/EditMode/ControlRigEditMode.cpp
@@ -401,6 +401,13 @@ void FControlRigEditMode::Tick(FEditorViewportClient* ViewportClient, float Delt
 {
 	FEdMode::Tick(ViewportClient, DeltaTime);
 
+	if(DeferredItemsToFrame.Num() > 0)
+	{
+		TGuardValue<FEditorViewportClient*> ViewportGuard(CurrentViewportClient, ViewportClient);
+		FrameItems(DeferredItemsToFrame);
+		DeferredItemsToFrame.Reset();
+	}
+
 	if (bSelectionChanged)
 	{
 		SetUpDetailPanel();
@@ -1722,6 +1729,43 @@ void FControlRigEditMode::FrameSelection()
 	}
 }
 
+void FControlRigEditMode::FrameItems(const TArray<FRigElementKey>& InItems)
+{
+	if(!OnGetRigElementTransformDelegate.IsBound())
+	{
+		return;
+	}
+
+	if(CurrentViewportClient == nullptr)
+	{
+		DeferredItemsToFrame = InItems;
+		return;
+	}
+
+	FBox Box(ForceInit);
+
+	for (int32 Index = 0; Index < InItems.Num(); ++Index)
+	{
+		static const float Radius = 20.f;
+		if (InItems[Index].Type == ERigElementType::Bone || InItems[Index].Type == ERigElementType::Null)
+		{
+			FTransform Transform = OnGetRigElementTransformDelegate.Execute(InItems[Index], false, true);
+			Box += Transform.TransformPosition(FVector::OneVector * Radius);
+			Box += Transform.TransformPosition(FVector::OneVector * -Radius);
+		}
+		else if (InItems[Index].Type == ERigElementType::Control)
+		{
+			FTransform Transform = OnGetRigElementTransformDelegate.Execute(InItems[Index], false, true);
+			Box += Transform.TransformPosition(FVector::OneVector * Radius);
+			Box += Transform.TransformPosition(FVector::OneVector * -Radius);
+		}
+	}
+
+	if(Box.IsValid)
+	{
+		CurrentViewportClient->FocusViewportOnBox(Box);
+	}
+}
 
 void FControlRigEditMode::IncreaseGizmoSize()
 {
diff --git a/Engine/Plugins/Experimental/ControlRig/Source/ControlRigEditor/Private/EditMode/ControlRigEditMode.h b/Engine/Plugins/Experimental/ControlRig/Source/ControlRigEditor/Private/EditMode/ControlRigEditMode.h
--- a/Engine/Plugins/Experimental/ControlRig/Source/ControlRigEditor/Private/EditMode/ControlRigEditMode.h
+++ b/Engine/Plugins/Experimental/ControlRig/Source/ControlRigEditor/Private/EditMode/ControlRigEditMode.h
@@ -175,12 +175,19 @@ protected:
 	/** Toggles visibility of manipulators in the viewport */
 	void ToggleManipulators();
 
+public:
+	
 	/** Clear Selection*/
 	void ClearSelection();
 
 	/** Frame to current Control Selection*/
 	void FrameSelection();
 
+	/** Frame a list of provided items*/
+   	void FrameItems(const TArray<FRigElementKey>& InItems);
+
+private:
+	
 	/** Whether or not we should Frame Selection or not*/
 	bool CanFrameSelection();
 
@@ -344,6 +351,8 @@ private:
 
 	TStrongObjectPtr<UControlRigEditModeDelegateHelper> DelegateHelper;
 
+	TArray<FRigElementKey> DeferredItemsToFrame;
+
 	friend class FControlRigEditorModule;
 	friend class FControlRigEditor;
 	friend class FControlRigEditModeGenericDetails;
diff --git a/Engine/Plugins/Experimental/ControlRig/Source/ControlRigEditor/Private/Editor/SRigHierarchy.cpp b/Engine/Plugins/Experimental/ControlRig/Source/ControlRigEditor/Private/Editor/SRigHierarchy.cpp
--- a/Engine/Plugins/Experimental/ControlRig/Source/ControlRigEditor/Private/Editor/SRigHierarchy.cpp
+++ b/Engine/Plugins/Experimental/ControlRig/Source/ControlRigEditor/Private/Editor/SRigHierarchy.cpp
@@ -37,6 +37,7 @@
 #include "SKismetInspector.h"
 #include "Types/WidgetActiveTimerDelegate.h"
 #include "Dialogs/CustomDialog.h"
+#include "EditMode/ControlRigEditMode.h"
 
 #define LOCTEXT_NAMESPACE "SRigHierarchy"
 
@@ -1592,11 +1593,16 @@ void SRigHierarchy::ImportHierarchy(const FAssetData& InAssetData)
 		URigHierarchyController* Controller = Hierarchy->GetController(true);
 		check(Controller);
 
-		Controller->ImportBones(RefSkeleton, NAME_None, false, false, bSelectBones, true);
+		TArray<FRigElementKey> ImportedBones = Controller->ImportBones(RefSkeleton, NAME_None, false, false, bSelectBones, true);
 		Controller->ImportCurves(Mesh->GetSkeleton(), NAME_None, true, true);
 
 		ControlRigBlueprint->SourceHierarchyImport = Mesh->GetSkeleton();
 		ControlRigBlueprint->SourceCurveImport = Mesh->GetSkeleton();
+
+		if(ImportedBones.Num() > 0)
+		{
+			ControlRigEditor.Pin()->GetEditMode()->FrameItems(ImportedBones);
+		}
 	}
 
 	ControlRigBlueprint->PropagateHierarchyFromBPToInstances();
