Control Rig: Added hotkey ctrl+period to enter gizmo transform editmode: use transform manipulators to change control's gizmo transform in viewport, a notification would show in the bottom right corner to indicate when the mode is active

diff --git a/Engine/Plugins/Experimental/ControlRig/Source/ControlRigEditor/Private/ControlRigEditorStyle.h b/Engine/Plugins/Experimental/ControlRig/Source/ControlRigEditor/Private/ControlRigEditorStyle.h
--- a/Engine/Plugins/Experimental/ControlRig/Source/ControlRigEditor/Private/ControlRigEditorStyle.h
+++ b/Engine/Plugins/Experimental/ControlRig/Source/ControlRigEditor/Private/ControlRigEditorStyle.h
@@ -106,6 +106,8 @@ public:
 				SetContentRoot(EngineEditorSlateDir);
 				// similar style to "LevelViewport.StartingPlayInEditorBorder"
 				Set( "ControlRig.Viewport.Border", new BOX_BRUSH( "Old/Window/ViewportDebugBorder", 0.8f, FLinearColor(1.0f,1.0f,1.0f,1.0f) ) );
+				// similar style to "AnimViewport.Notification.Warning"
+				Set( "ControlRig.Viewport.Notification.ChangeGizmoTransform", new BOX_BRUSH("Common/RoundedSelection_16x", 4.0f/16.0f, FLinearColor(FColor(169, 0, 148))));
 				// restore the previous content root
 				SetContentRoot(PreviousContentRoot);
 			}
diff --git a/Engine/Plugins/Experimental/ControlRig/Source/ControlRigEditor/Private/EditMode/ControlRigEditMode.cpp b/Engine/Plugins/Experimental/ControlRig/Source/ControlRigEditor/Private/EditMode/ControlRigEditMode.cpp
--- a/Engine/Plugins/Experimental/ControlRig/Source/ControlRigEditor/Private/EditMode/ControlRigEditMode.cpp
+++ b/Engine/Plugins/Experimental/ControlRig/Source/ControlRigEditor/Private/EditMode/ControlRigEditMode.cpp
@@ -124,7 +124,8 @@ namespace ControlRigSelectionConstants
 }
 
 FControlRigEditMode::FControlRigEditMode()
-	: InteractionScope(nullptr)
+	: bIsChangingGizmoTransform(false)
+	, InteractionScope(nullptr)
 	, bManipulatorMadeChange(false)
 	, bSelecting(false)
 	, bSelectionChanged(false)
@@ -1239,36 +1240,59 @@ bool FControlRigEditMode::InputDelta(FEditorViewportClient* InViewportClient, FV
 			if (AreRigElementsSelected(FRigElementTypeHelper::ToMask(ERigElementType::Control)))
 		{
 			FTransform ComponentTransform = GetHostingSceneComponentTransform();
-			bool bDoLocal = (CoordSystem == ECoordSystem::COORD_Local && Settings && Settings->bLocalTransformsInEachLocalSpace);
-			bool bUseLocal = false;
-			bool bCalcLocal = bDoLocal;
-			bool bFirstTime = true;
-			FTransform InOutLocal = FTransform::Identity;
-			for (AControlRigGizmoActor* GizmoActor : GizmoActors)
+
+			if (bIsChangingGizmoTransform)
 			{
-				if (GizmoActor->IsSelected())
+				for (AControlRigGizmoActor* GizmoActor : GizmoActors)
 				{
-					// test local vs global
-					if (bManipulatorMadeChange == false)
-					{
-						GEditor->BeginTransaction(LOCTEXT("MoveControlTransaction", "Move Control"));
-					}
-					if (bFirstTime)
+					if (GizmoActor->IsSelected())
 					{
-						bFirstTime = false;
+						if (bManipulatorMadeChange == false)
+						{
+							GEditor->BeginTransaction(LOCTEXT("ChangeControlGizmoTransaction", "Change Control Gizmo Transform"));
+						}
+
+						ChangeControlGizmoTransform(GizmoActor, bDoTranslation, InDrag, bDoRotation, InRot, bDoScale, InScale, ComponentTransform);
+						bManipulatorMadeChange = true;
+
+						// break here since we only support changing gizmo transform of a single control at a time
+						break;
 					}
-					else
+				}
+			}
+			else
+			{
+				bool bDoLocal = (CoordSystem == ECoordSystem::COORD_Local && Settings && Settings->bLocalTransformsInEachLocalSpace);
+				bool bUseLocal = false;
+				bool bCalcLocal = bDoLocal;
+				bool bFirstTime = true;
+				FTransform InOutLocal = FTransform::Identity;
+				for (AControlRigGizmoActor* GizmoActor : GizmoActors)
+				{
+					if (GizmoActor->IsSelected())
 					{
-						if (bDoLocal)
+						// test local vs global
+						if (bManipulatorMadeChange == false)
 						{
-							bUseLocal = true;
-							bDoLocal = false;
+							GEditor->BeginTransaction(LOCTEXT("MoveControlTransaction", "Move Control"));
+						}
+						if (bFirstTime)
+						{
+							bFirstTime = false;
+						}
+						else
+						{
+							if (bDoLocal)
+							{
+								bUseLocal = true;
+								bDoLocal = false;
+							}
 						}
-					}
 
-					MoveGizmo(GizmoActor, bDoTranslation, InDrag, bDoRotation, InRot, bDoScale, InScale, ComponentTransform,
-						bUseLocal,bDoLocal,InOutLocal);
-					bManipulatorMadeChange = true;
+						MoveGizmo(GizmoActor, bDoTranslation, InDrag, bDoRotation, InRot, bDoScale, InScale, ComponentTransform,
+							bUseLocal, bDoLocal, InOutLocal);
+						bManipulatorMadeChange = true;
+					}
 				}
 			}
 
@@ -1524,28 +1548,50 @@ void FControlRigEditMode::RecalcPivotTransform()
 				// todo?
 			}
 		}
+		
+		if (bIsChangingGizmoTransform)
+		{
+			if (UControlRig* ControlRig = GetControlRig(true))
+			{ 
+				for (const AControlRigGizmoActor* GizmoActor : GizmoActors)
+				{
+					if (GizmoActor->IsSelected())
+					{
+						if (FRigControlElement* ControlElement = ControlRig->GetHierarchy()->Find<FRigControlElement>(FRigElementKey(GizmoActor->ControlName, ERigElementType::Control)))
+						{
+							PivotTransform = ControlRig->GetHierarchy()->GetControlGizmoTransform(ControlElement, ERigTransformType::CurrentGlobal);
+						}
 
-		for (const AControlRigGizmoActor* GizmoActor : GizmoActors)
+						// break here since we don't want to change the gizmo transform of multiple controls.
+						break;
+					}
+				}
+			}
+		}
+		else
 		{
-			if (GizmoActor->IsSelected())
+			for (const AControlRigGizmoActor* GizmoActor : GizmoActors)
 			{
-				LastTransform = GizmoActor->GetActorTransform().GetRelativeTransform(ComponentTransform);
-				PivotLocation += LastTransform.GetLocation();
-				++NumSelectedControls;
-				if (Settings && Settings->bLocalTransformsInEachLocalSpace) //if in local just use first actors transform
+				if (GizmoActor->IsSelected())
 				{
-					break;
+					LastTransform = GizmoActor->GetActorTransform().GetRelativeTransform(ComponentTransform);
+					PivotLocation += LastTransform.GetLocation();
+					++NumSelectedControls;
+					if (Settings && Settings->bLocalTransformsInEachLocalSpace) //if in local just use first actors transform
+					{
+						break;
+					}
+					
 				}
 			}
-		}
 
-		PivotLocation /= (float)FMath::Max(1, NumSelectedControls);
-		PivotTransform.SetLocation(PivotLocation);
-		
-		// just use last rotation
-		FTransform WorldTransform = LastTransform * ComponentTransform;
-		PivotTransform.SetRotation(WorldTransform.GetRotation());
-	
+			PivotLocation /= (float)FMath::Max(1, NumSelectedControls);
+			PivotTransform.SetLocation(PivotLocation);
+			
+			// just use last rotation
+			FTransform WorldTransform = LastTransform * ComponentTransform;
+			PivotTransform.SetRotation(WorldTransform.GetRotation());
+		}
 	}
 	else if (AreRigElementSelectedAndMovable())
 	{
@@ -1596,6 +1642,14 @@ void FControlRigEditMode::HandleSelectionChanged()
 		}
 	}
 
+	// automatically exit gizmo transform edit mode if there is no gizmo selected
+	if (bIsChangingGizmoTransform)
+	{
+		if (!CanChangeControlGizmoTransform())
+		{
+			bIsChangingGizmoTransform = false;
+		}
+	}
 	
 	// update the pivot transform of our selected objects (they could be animating)
 	RecalcPivotTransform();
@@ -1635,6 +1689,10 @@ void FControlRigEditMode::BindCommands()
 	CommandBindings->MapAction(
 		Commands.ResetGizmoSize,
 		FExecuteAction::CreateRaw(this, &FControlRigEditMode::ResetGizmoSize));
+
+	CommandBindings->MapAction(
+		Commands.ToggleGizmoTransformEdit,
+		FExecuteAction::CreateRaw(this, &FControlRigEditMode::ToggleGizmoTransformEdit));
 }
 
 bool FControlRigEditMode::IsControlSelected() const
@@ -1785,6 +1843,24 @@ void FControlRigEditMode::ResetGizmoSize()
 	GetModeManager()->SetWidgetScale(Settings->GizmoScale);
 }
 
+void FControlRigEditMode::ToggleGizmoTransformEdit()
+{ 
+	if (bIsChangingGizmoTransform)
+	{
+		bIsChangingGizmoTransform = false;
+	}
+	else if (CanChangeControlGizmoTransform())
+	{
+		bIsChangingGizmoTransform = true;
+	}
+}
+
+FText FControlRigEditMode::GetToggleGizmoTransformEditHotKey() const
+{
+	const FControlRigEditModeCommands& Commands = FControlRigEditModeCommands::Get();
+	return Commands.ToggleGizmoTransformEdit->GetInputText();
+}
+
 void FControlRigEditMode::ToggleManipulators()
 {
 	// Toggle flag (is used in drawing code)
@@ -1998,6 +2074,14 @@ void FControlRigEditMode::OnHierarchyModified(ERigHierarchyNotification InNotif,
 		case ERigHierarchyNotification::ControlSettingChanged:
 		case ERigHierarchyNotification::ControlGizmoTransformChanged:
 		{
+			// in case the gizmo is turned off, automatically exit gizmo transform edit mode
+			if (bIsChangingGizmoTransform)
+			{
+				if (!CanChangeControlGizmoTransform())
+				{
+					bIsChangingGizmoTransform = false;
+				} 
+			}
 			RequestToRecreateGizmoActors();
 			break;
 		}
@@ -2126,6 +2210,39 @@ void FControlRigEditMode::OnCoordSystemChanged(ECoordSystem InCoordSystem)
 	}
 }
 
+bool FControlRigEditMode::CanChangeControlGizmoTransform()
+{
+	if (!IsInLevelEditor())
+	{
+		// do not allow multi-select
+		if (SelectedRigElements.Num() == 1)
+		{
+			if (AreRigElementsSelected(FRigElementTypeHelper::ToMask(ERigElementType::Control)))
+			{
+				if (UControlRig* ControlRig = GetControlRig(true))
+				{
+					// only enable for a Control with Gizmo enabled and visible
+					if (FRigControlElement* ControlElement = ControlRig->GetHierarchy()->Find<FRigControlElement>(SelectedRigElements[0]))
+					{
+						if (ControlElement->Settings.bGizmoEnabled && ControlElement->Settings.bGizmoVisible)
+						{
+							if (AControlRigGizmoActor* GizmoActor = GetGizmoFromControlName(SelectedRigElements[0].Name))
+							{
+								if (ensure(GizmoActor->IsSelected()))
+								{
+									return true;
+								}
+							}
+						}
+					}
+				}
+			} 
+		}
+	}
+
+	return false;
+}
+
 void FControlRigEditMode::SetGizmoTransform(AControlRigGizmoActor* GizmoActor, const FTransform& InTransform)
 {
 	if (UControlRig* ControlRig = GetControlRig(true, GizmoActor->ControlRigIndex))
@@ -2269,6 +2386,83 @@ void FControlRigEditMode::MoveGizmo(AControlRigGizmoActor* GizmoActor, const boo
 #endif
 }
 
+void FControlRigEditMode::ChangeControlGizmoTransform(AControlRigGizmoActor* GizmoActor, const bool bTranslation, FVector& InDrag,
+	const bool bRotation, FRotator& InRot, const bool bScale, FVector& InScale, const FTransform& ToWorldTransform)
+{
+	bool bTransformChanged = false; 
+
+	FTransform CurrentTransform = FTransform::Identity;
+
+	if (UControlRig* ControlRig = GetControlRig(true))
+	{
+		if (FRigControlElement* ControlElement = ControlRig->GetHierarchy()->Find<FRigControlElement>(FRigElementKey(GizmoActor->ControlName, ERigElementType::Control)))
+		{
+			CurrentTransform = ControlRig->GetHierarchy()->GetControlGizmoTransform(ControlElement, ERigTransformType::CurrentGlobal);
+			CurrentTransform = CurrentTransform * ToWorldTransform;
+		}
+	}
+
+	if (bRotation)
+	{
+		FQuat CurrentRotation = CurrentTransform.GetRotation();
+		CurrentRotation = (InRot.Quaternion() * CurrentRotation);
+		CurrentTransform.SetRotation(CurrentRotation);
+		bTransformChanged = true;
+	}
+
+	if (bTranslation)
+	{
+		FVector CurrentLocation = CurrentTransform.GetLocation();
+		CurrentLocation = CurrentLocation + InDrag;
+		CurrentTransform.SetLocation(CurrentLocation);
+		bTransformChanged = true;
+	}
+
+	if (bScale)
+	{
+		FVector CurrentScale = CurrentTransform.GetScale3D();
+		CurrentScale = CurrentScale + InScale;
+		CurrentTransform.SetScale3D(CurrentScale);
+		bTransformChanged = true;
+	}
+
+	if (bTransformChanged)
+	{
+		if (UControlRig* RuntimeControlRig = GetControlRig(false, GizmoActor->ControlRigIndex))
+		{
+			UControlRig* InteractionControlRig = GetControlRig(true, GizmoActor->ControlRigIndex);
+
+			FTransform NewTransform = CurrentTransform.GetRelativeTransform(ToWorldTransform);
+			FRigControlModifiedContext Context;
+
+			if (FRigControlElement* ControlElement = InteractionControlRig->GetHierarchy()->Find<FRigControlElement>(FRigElementKey(GizmoActor->ControlName, ERigElementType::Control)))
+			{
+				// do not setup undo for this first step since it is just used to calculate the local transform
+				InteractionControlRig->GetHierarchy()->SetControlGizmoTransform(ControlElement, NewTransform, ERigTransformType::CurrentGlobal, false);
+				FTransform CurrentLocalGizmoTransform = InteractionControlRig->GetHierarchy()->GetControlGizmoTransform(ControlElement, ERigTransformType::CurrentLocal);
+				// this call should trigger an instance-to-BP update in ControlRigEditor
+				InteractionControlRig->GetHierarchy()->SetControlGizmoTransform(ControlElement, CurrentLocalGizmoTransform, ERigTransformType::InitialLocal, true);
+
+				FTransform MeshTransform = FTransform::Identity;
+				FTransform GizmoTransform = CurrentLocalGizmoTransform;
+
+				if (UControlRig* ControlRig = GetControlRig(true))
+				{
+					if (const UControlRigGizmoLibrary* GizmoLibrary = GetControlRig(true)->GetGizmoLibrary())
+					{
+						if (const FControlRigGizmoDefinition* Gizmo = GizmoLibrary->GetGizmoByName(ControlElement->Settings.GizmoName, true /* use default */))
+						{
+							MeshTransform = Gizmo->Transform;
+						}
+					}
+				}
+
+				GizmoActor->StaticMeshComponent->SetRelativeTransform(MeshTransform * GizmoTransform);
+			}
+		}
+	} 
+}
+
 // temporarily we just support following types of gizmo
 bool IsSupportedControlType(const ERigControlType ControlType)
 {
@@ -2302,6 +2496,11 @@ bool FControlRigEditMode::ModeSupportedByGizmoActor(const AControlRigGizmoActor*
 		const FRigControlElement* ControlElement = ControlRig->FindControl(GizmoActor->ControlName);
 		if (ControlElement)
 		{
+			if (bIsChangingGizmoTransform)
+			{
+				return true;
+			}
+
 			if (IsSupportedControlType(ControlElement->Settings.ControlType))
 			{
 				switch (InMode)
diff --git a/Engine/Plugins/Experimental/ControlRig/Source/ControlRigEditor/Private/EditMode/ControlRigEditMode.h b/Engine/Plugins/Experimental/ControlRig/Source/ControlRigEditor/Private/EditMode/ControlRigEditMode.h
--- a/Engine/Plugins/Experimental/ControlRig/Source/ControlRigEditor/Private/EditMode/ControlRigEditMode.h
+++ b/Engine/Plugins/Experimental/ControlRig/Source/ControlRigEditor/Private/EditMode/ControlRigEditMode.h
@@ -203,6 +203,12 @@ private:
 	/** Reset Gizmo Size */
 	void ResetGizmoSize();
 
+	/** Toggle Gizmo Transform Edit*/
+	void ToggleGizmoTransformEdit();
+
+	/** The hotkey text is passed to a viewport notification to inform users how to toggle gizmo edit*/
+	FText GetToggleGizmoTransformEditHotKey() const;
+
 	/** Bind our keyboard commands */
 	void BindCommands();
 
@@ -214,6 +220,9 @@ private:
 	/** Let the preview scene know how we want to select components */
 	bool GizmoSelectionOverride(const UPrimitiveComponent* InComponent) const;
 
+	/** Enable editing of control's gizmo transform instead of control's transform*/
+	bool bIsChangingGizmoTransform;
+
 protected:
 
 	TWeakPtr<ISequencer> WeakSequencer;
@@ -304,7 +313,7 @@ private:
 	TArray<ECoordSystem> CoordSystemPerWidgetMode;
 	bool bIsChangingCoordSystem;
 
-
+	bool CanChangeControlGizmoTransform();
 public:
 	//Toolbar functions
 	void SetOnlySelectRigControls(bool val);
@@ -324,6 +333,10 @@ private:
 	void MoveGizmo(AControlRigGizmoActor* GizmoActor, const bool bTranslation, FVector& InDrag, 
 		const bool bRotation, FRotator& InRot, const bool bScale, FVector& InScale, const FTransform& ToWorldTransform,
 		bool bUseLocal, bool bCalcLocal, FTransform& InOutLocal);
+
+	void ChangeControlGizmoTransform(AControlRigGizmoActor* GizmoActor, const bool bTranslation, FVector& InDrag,
+		const bool bRotation, FRotator& InRot, const bool bScale, FVector& InScale, const FTransform& ToWorldTransform);
+
 	void TickGizmo(AControlRigGizmoActor* GizmoActor, const FTransform& ComponentTransform);
 	bool ModeSupportedByGizmoActor(const AControlRigGizmoActor* GizmoActor, UE::Widget::EWidgetMode InMode) const;
 
diff --git a/Engine/Plugins/Experimental/ControlRig/Source/ControlRigEditor/Private/EditMode/ControlRigEditModeCommands.cpp b/Engine/Plugins/Experimental/ControlRig/Source/ControlRigEditor/Private/EditMode/ControlRigEditModeCommands.cpp
--- a/Engine/Plugins/Experimental/ControlRig/Source/ControlRigEditor/Private/EditMode/ControlRigEditModeCommands.cpp
+++ b/Engine/Plugins/Experimental/ControlRig/Source/ControlRigEditor/Private/EditMode/ControlRigEditModeCommands.cpp
@@ -16,7 +16,7 @@ void FControlRigEditModeCommands::RegisterCommands()
 	UI_COMMAND(DecreaseGizmoSize, "Decrease Gizmo Size", "Decrease Size of the Gizmos In The Viewport", EUserInterfaceActionType::Button, FInputChord(EKeys::Hyphen, EModifierKey::Shift));
 	UI_COMMAND(ResetGizmoSize, "Reset Gizmo Size", "Resize Gizmo Size To Default", EUserInterfaceActionType::Button, FInputChord(EKeys::Equals));
 
-
+	UI_COMMAND(ToggleGizmoTransformEdit, "Toggle Gizmo Transform Edit", "Toggle Editing Selected Control's Gizmo Transform", EUserInterfaceActionType::ToggleButton, FInputChord(EKeys::Period, EModifierKey::Control)); 
 }
 
 #undef LOCTEXT_NAMESPACE
diff --git a/Engine/Plugins/Experimental/ControlRig/Source/ControlRigEditor/Private/EditMode/ControlRigEditModeCommands.h b/Engine/Plugins/Experimental/ControlRig/Source/ControlRigEditor/Private/EditMode/ControlRigEditModeCommands.h
--- a/Engine/Plugins/Experimental/ControlRig/Source/ControlRigEditor/Private/EditMode/ControlRigEditModeCommands.h
+++ b/Engine/Plugins/Experimental/ControlRig/Source/ControlRigEditor/Private/EditMode/ControlRigEditModeCommands.h
@@ -44,6 +44,8 @@ public:
 	/** Reset Gizmo Size */
 	TSharedPtr< FUICommandInfo > ResetGizmoSize;
 
+	/** Toggle Gizmo Transform Edit*/
+	TSharedPtr< FUICommandInfo > ToggleGizmoTransformEdit;
 
 	/**
 	 * Initialize commands
diff --git a/Engine/Plugins/Experimental/ControlRig/Source/ControlRigEditor/Private/Editor/ControlRigEditor.cpp b/Engine/Plugins/Experimental/ControlRig/Source/ControlRigEditor/Private/Editor/ControlRigEditor.cpp
--- a/Engine/Plugins/Experimental/ControlRig/Source/ControlRigEditor/Private/Editor/ControlRigEditor.cpp
+++ b/Engine/Plugins/Experimental/ControlRig/Source/ControlRigEditor/Private/Editor/ControlRigEditor.cpp
@@ -2787,6 +2787,32 @@ void FControlRigEditor::HandleViewportCreated(const TSharedRef<class IPersonaVie
 		return FEditorFontGlyphs::Eye;
 	};
 
+	auto GetChangingGizmoTransformText = [this]()
+	{
+		if (FControlRigEditMode* EditMode = GetEditMode())
+		{
+			FText HotKeyText = EditMode->GetToggleGizmoTransformEditHotKey();
+
+			if (!HotKeyText.IsEmpty())
+			{
+				FFormatNamedArguments Args;
+				Args.Add(TEXT("HotKey"), HotKeyText);
+				return FText::Format(LOCTEXT("ControlRigBPViewportGizmoTransformEditNotification", "Currently Manipulating Gizmo Transform - Press {HotKey} to Exit"), Args);
+			}
+		}
+		
+		return LOCTEXT("ControlRigBPViewportGizmoTransformEditNotification", "Currently Manipulating Gizmo Transform - Assign a Hotkey and Use It to Exit");
+	};
+
+	auto GetChangingGizmoTransformTextVisibility = [this]()
+	{
+		if (FControlRigEditMode* EditMode = GetEditMode())
+		{
+			return EditMode->bIsChangingGizmoTransform ? EVisibility::Visible : EVisibility::Collapsed;
+		}
+		return EVisibility::Collapsed;
+	};
+
 	InViewport->AddNotification(MakeAttributeLambda(GetErrorSeverity),
 		false,
 		SNew(SHorizontalBox)
@@ -2850,6 +2876,33 @@ void FControlRigEditor::HandleViewportCreated(const TSharedRef<class IPersonaVie
 		],
 		FPersonaViewportNotificationOptions(TAttribute<EVisibility>::Create(GetCompilationStateVisibility))
 	);
+	
+	FPersonaViewportNotificationOptions ChangeGizmoTransformNotificationOptions;
+	ChangeGizmoTransformNotificationOptions.OnGetVisibility = TAttribute<EVisibility>::Create(GetChangingGizmoTransformTextVisibility);
+	ChangeGizmoTransformNotificationOptions.OnGetBrushOverride = TAttribute<const FSlateBrush*>(FControlRigEditorStyle::Get().GetBrush("ControlRig.Viewport.Notification.ChangeGizmoTransform"));
+
+	// notification that shows when users enter the mode that allows them to change gizmo transform
+	InViewport->AddNotification(EMessageSeverity::Type::Info,
+		false,
+		SNew(SHorizontalBox)
+		.Visibility_Lambda(GetChangingGizmoTransformTextVisibility)
+		+SHorizontalBox::Slot()
+		.FillWidth(1.0f)
+		.Padding(4.0f, 4.0f)
+		[
+			SNew(SHorizontalBox)
+			.ToolTipText_Lambda(GetChangingGizmoTransformText)
+			+SHorizontalBox::Slot()
+			.FillWidth(1.0f)
+			.VAlign(VAlign_Center)
+			[
+				SNew(STextBlock)
+				.Text_Lambda(GetChangingGizmoTransformText)
+				.TextStyle(FEditorStyle::Get(), "AnimViewport.MessageText")
+			]
+		],
+		ChangeGizmoTransformNotificationOptions
+	);
 
 	InViewport->AddToolbarExtender(TEXT("AnimViewportDefaultCamera"), FMenuExtensionDelegate::CreateLambda(
 		[&](FMenuBuilder& InMenuBuilder)
@@ -4361,8 +4414,16 @@ void FControlRigEditor::OnHierarchyModified_AnyThread(ERigHierarchyNotification
 					FRigControlElement* SourceControlElement = Cast<FRigControlElement>(Element);
 					if(SourceControlElement)
 					{
+						FTransform InitialGizmoTransform = WeakHierarchy.Get()->GetControlGizmoTransform(SourceControlElement, ERigTransformType::InitialLocal);
+
+						// set current gizmo transform = initial gizmo transform so that the viewport reflects this change
+						WeakHierarchy.Get()->SetControlGizmoTransform(SourceControlElement, InitialGizmoTransform, ERigTransformType::CurrentLocal, false); 
+
 						RigBlueprint->Hierarchy->SetControlGizmoTransform(Key, WeakHierarchy.Get()->GetControlGizmoTransform(SourceControlElement, ERigTransformType::InitialLocal), true);
 						RigBlueprint->Hierarchy->SetControlGizmoTransform(Key, WeakHierarchy.Get()->GetControlGizmoTransform(SourceControlElement, ERigTransformType::CurrentLocal), false);
+
+						RigBlueprint->Modify();
+						RigBlueprint->MarkPackageDirty();
 					}
 				}
 				break;
