ControlRig: Restrict variable types to allowed CR types

diff --git a/Engine/Plugins/Experimental/ControlRig/Source/ControlRigDeveloper/Private/Graph/ControlRigGraphSchema.cpp b/Engine/Plugins/Experimental/ControlRig/Source/ControlRigDeveloper/Private/Graph/ControlRigGraphSchema.cpp
--- a/Engine/Plugins/Experimental/ControlRig/Source/ControlRigDeveloper/Private/Graph/ControlRigGraphSchema.cpp
+++ b/Engine/Plugins/Experimental/ControlRig/Source/ControlRigDeveloper/Private/Graph/ControlRigGraphSchema.cpp
@@ -127,26 +127,6 @@ bool FControlRigGraphSchemaAction_LocalVar::IsVariableUsed()
 	return false;
 }
 
-void FControlRigGraphSchemaAction_LocalVar::GetVariableTypeTree(
-	TArray<TSharedPtr<UEdGraphSchema_K2::FPinTypeTreeInfo>>& TypeTree, ETypeTreeFilter TypeTreeFilter) const
-{
-	const UEdGraphSchema_K2* Schema = GetDefault<UEdGraphSchema_K2>();
-	
-	// Clear the list
-	TypeTree.Empty();
-
-	TypeTree.Add( MakeShareable( new UEdGraphSchema_K2::FPinTypeTreeInfo(UEdGraphSchema_K2::GetCategoryText(UEdGraphSchema_K2::PC_Boolean, true), UEdGraphSchema_K2::PC_Boolean, Schema, LOCTEXT("BooleanType", "True or false value"))));
-	TypeTree.Add( MakeShareable( new UEdGraphSchema_K2::FPinTypeTreeInfo(UEdGraphSchema_K2::GetCategoryText(UEdGraphSchema_K2::PC_Int, true), UEdGraphSchema_K2::PC_Int, Schema, LOCTEXT("IntegerType", "Integer number")) ) );;
-	
-	TypeTree.Add(MakeShareable(new UEdGraphSchema_K2::FPinTypeTreeInfo(UEdGraphSchema_K2::GetCategoryText(UEdGraphSchema_K2::PC_Float, true), UEdGraphSchema_K2::PC_Float, Schema, LOCTEXT("FloatType", "Floating point number"))));
-	TypeTree.Add(MakeShareable(new UEdGraphSchema_K2::FPinTypeTreeInfo(UEdGraphSchema_K2::GetCategoryText(UEdGraphSchema_K2::PC_Name, true), UEdGraphSchema_K2::PC_Name, Schema, LOCTEXT("NameType", "A text name"))));
-	TypeTree.Add(MakeShareable(new UEdGraphSchema_K2::FPinTypeTreeInfo(UEdGraphSchema_K2::GetCategoryText(UEdGraphSchema_K2::PC_String, true), UEdGraphSchema_K2::PC_String, Schema, LOCTEXT("StringType", "A text string"))));
-
-	TypeTree.Add(MakeShareable(new UEdGraphSchema_K2::FPinTypeTreeInfo(UEdGraphSchema_K2::PC_Struct, TBaseStructure<FVector>::Get(), LOCTEXT("VectorType", "A 3D vector"))));
-	TypeTree.Add(MakeShareable(new UEdGraphSchema_K2::FPinTypeTreeInfo(UEdGraphSchema_K2::PC_Struct, TBaseStructure<FRotator>::Get(), LOCTEXT("RotatorType", "A 3D rotation"))));
-	TypeTree.Add(MakeShareable(new UEdGraphSchema_K2::FPinTypeTreeInfo(UEdGraphSchema_K2::PC_Struct, TBaseStructure<FTransform>::Get(), LOCTEXT("TransformType", "A 3D transformation, including translation, rotation and 3D scale."))));
-}
-
 FReply FControlRigFunctionDragDropAction::DroppedOnPanel(const TSharedRef< class SWidget >& Panel, FVector2D ScreenPosition, FVector2D GraphPosition, UEdGraph& Graph)
 {
 	// For local variables
@@ -538,6 +518,103 @@ FLinearColor UControlRigGraphSchema::GetPinTypeColor(const FEdGraphPinType& PinT
 	return GetDefault<UEdGraphSchema_K2>()->GetPinTypeColor(PinType);
 }
 
+bool UControlRigGraphSchema::SupportsPinType(UScriptStruct* ScriptStruct) const
+{
+	if(!ScriptStruct)
+	{
+		return false;
+	}
+
+	for (TFieldIterator<FProperty> It(ScriptStruct); It; ++It)
+	{
+		FName PropertyName = It->GetFName();
+		FProperty* Property = *It;
+
+		if (FArrayProperty* ArrayProperty = CastField<FArrayProperty>(Property))
+		{
+			Property = ArrayProperty->Inner;
+		}
+
+		FString CPPType = Property->GetCPPType();
+		if (CPPType == TEXT("bool") ||
+			CPPType == TEXT("float") ||
+			CPPType == TEXT("int32") ||
+			CPPType == TEXT("FString") ||
+			CPPType == TEXT("FName") ||
+			CPPType == TEXT("uint16"))
+		{
+			continue;
+		}		
+		
+		if (FStructProperty* StructProperty = CastField<FStructProperty>(Property))
+		{
+			if (SupportsPinType(StructProperty->Struct))
+			{
+				continue;
+			}
+		}
+		else if (FEnumProperty* EnumProperty = CastField<FEnumProperty>(Property))
+		{
+			continue;
+		}
+		else if (FByteProperty* ByteProperty = CastField<FByteProperty>(Property))
+		{
+			if (ByteProperty->Enum)
+			{
+				continue;
+			}
+		}
+	
+		return false;
+	}
+	
+	return true;
+}
+
+bool UControlRigGraphSchema::SupportsPinType(TWeakPtr<const FEdGraphSchemaAction> SchemaAction, const FEdGraphPinType& PinType) const
+{
+	if (PinType.IsContainer())
+	{
+		return false;
+	}
+	
+	const FName TypeName = PinType.PinCategory;
+
+	if (TypeName == UEdGraphSchema_K2::PC_Boolean ||
+		TypeName == UEdGraphSchema_K2::PC_Int ||
+		TypeName == UEdGraphSchema_K2::PC_Float ||
+		TypeName == UEdGraphSchema_K2::PC_Name ||
+		TypeName == UEdGraphSchema_K2::PC_String ||
+		TypeName == UEdGraphSchema_K2::PC_Enum)
+	{
+		return true;
+	}
+
+	if (PinType.PinCategory == UEdGraphSchema_K2::PC_Struct)
+	{
+		if (UScriptStruct* ScriptStruct = Cast<UScriptStruct>(PinType.PinSubCategoryObject))
+		{
+			return SupportsPinType(ScriptStruct);
+		}
+	}
+
+	if (PinType.PinCategory == UEdGraphSchema_K2::PC_Byte)
+	{
+		if (PinType.PinSubCategoryObject.IsValid())
+		{
+			return PinType.PinSubCategoryObject->IsA<UEnum>();
+		}
+	}
+
+	return false;
+}
+
+bool UControlRigGraphSchema::SupportsPinTypeContainer(TWeakPtr<const FEdGraphSchemaAction> SchemaAction,
+	const FEdGraphPinType& PinType, const EPinContainerType& ContainerType) const
+{
+	return ContainerType == EPinContainerType::None;
+}
+
 void UControlRigGraphSchema::BreakPinLinks(UEdGraphPin& TargetPin, bool bSendsNodeNotifcation) const
 {
 	//const FScopedTransaction Transaction( LOCTEXT("GraphEd_BreakPinLinks", "Break Pin Links") );
diff --git a/Engine/Plugins/Experimental/ControlRig/Source/ControlRigDeveloper/Public/Graph/ControlRigGraphSchema.h b/Engine/Plugins/Experimental/ControlRig/Source/ControlRigDeveloper/Public/Graph/ControlRigGraphSchema.h
--- a/Engine/Plugins/Experimental/ControlRig/Source/ControlRigDeveloper/Public/Graph/ControlRigGraphSchema.h
+++ b/Engine/Plugins/Experimental/ControlRig/Source/ControlRigDeveloper/Public/Graph/ControlRigGraphSchema.h
@@ -79,8 +79,6 @@ public:
 	virtual void DeleteVariable() override;
 
 	virtual bool IsVariableUsed() override;
-
-	virtual void GetVariableTypeTree(TArray< TSharedPtr<UEdGraphSchema_K2::FPinTypeTreeInfo> >& TypeTree, ETypeTreeFilter TypeTreeFilter) const override;
 };
 
 /** DragDropAction class for drag and dropping an item from the My Blueprints tree (e.g., variable or function) */
@@ -139,6 +137,8 @@ public:
 	virtual bool TryCreateConnection(UEdGraphPin* PinA, UEdGraphPin* PinB) const override;
 	virtual const FPinConnectionResponse CanCreateConnection(const UEdGraphPin* A, const UEdGraphPin* B) const override;
 	virtual FLinearColor GetPinTypeColor(const FEdGraphPinType& PinType) const override;
+	virtual bool SupportsPinType(TWeakPtr<const FEdGraphSchemaAction> SchemaAction, const FEdGraphPinType& PinType) const override;
+	virtual bool SupportsPinTypeContainer(TWeakPtr<const FEdGraphSchemaAction> SchemaAction, const FEdGraphPinType& PinType, const EPinContainerType& ContainerType) const override;
 	virtual void BreakPinLinks(UEdGraphPin& TargetPin, bool bSendsNodeNotifcation) const override;
 	virtual void BreakSinglePinLink(UEdGraphPin* SourcePin, UEdGraphPin* TargetPin) const override;
 	virtual bool CanGraphBeDropped(TSharedPtr<FEdGraphSchemaAction> InAction) const override;
@@ -173,6 +173,9 @@ public:
 	virtual bool SupportsDropPinOnNode(UEdGraphNode* InTargetNode, const FEdGraphPinType& InSourcePinType, EEdGraphPinDirection InSourcePinDirection, FText& OutErrorMessage) const override;
 	virtual void SetPinBeingDroppedOnNode(UEdGraphPin* InSourcePin) const override { PinBeingDropped = InSourcePin; }
 
+	/** Returns true if the schema supports the script type */
+	bool SupportsPinType(UScriptStruct* ScriptStruct) const;
+
 	/** Create a graph node for a rig */
 	UControlRigGraphNode* CreateGraphNode(UControlRigGraph* InGraph, const FName& InPropertyName) const;
 
diff --git a/Engine/Source/Editor/BlueprintGraph/Classes/EdGraphSchema_K2.h b/Engine/Source/Editor/BlueprintGraph/Classes/EdGraphSchema_K2.h
--- a/Engine/Source/Editor/BlueprintGraph/Classes/EdGraphSchema_K2.h
+++ b/Engine/Source/Editor/BlueprintGraph/Classes/EdGraphSchema_K2.h
@@ -510,7 +510,7 @@ public:
 	virtual bool ShouldHidePinDefaultValue(UEdGraphPin* Pin) const override;
 	virtual bool ShouldShowAssetPickerForPin(UEdGraphPin* Pin) const override;
 	virtual FLinearColor GetPinTypeColor(const FEdGraphPinType& PinType) const override;
-	FLinearColor GetSecondaryPinTypeColor(const FEdGraphPinType& PinType) const;
+	virtual FLinearColor GetSecondaryPinTypeColor(const FEdGraphPinType& PinType) const override;
 	virtual FText GetPinDisplayName(const UEdGraphPin* Pin) const override;
 	virtual void ConstructBasicPinTooltip(const UEdGraphPin& Pin, const FText& PinDescription, FString& TooltipOut) const override;
 	virtual EGraphType GetGraphType(const UEdGraph* TestEdGraph) const override;
diff --git a/Engine/Source/Editor/BlueprintGraph/Classes/EdGraphSchema_K2_Actions.h b/Engine/Source/Editor/BlueprintGraph/Classes/EdGraphSchema_K2_Actions.h
--- a/Engine/Source/Editor/BlueprintGraph/Classes/EdGraphSchema_K2_Actions.h
+++ b/Engine/Source/Editor/BlueprintGraph/Classes/EdGraphSchema_K2_Actions.h
@@ -545,8 +545,6 @@ public:
 	// End of FEdGraphSchemaAction interface
 
 	UBlueprint* GetSourceBlueprint() const;
-
-	virtual void GetVariableTypeTree(TArray< TSharedPtr<UEdGraphSchema_K2::FPinTypeTreeInfo> >& TypeTree, ETypeTreeFilter TypeTreeFilter) const;
 };
 
 
diff --git a/Engine/Source/Editor/BlueprintGraph/Private/EdGraphSchema_K2_Actions.cpp b/Engine/Source/Editor/BlueprintGraph/Private/EdGraphSchema_K2_Actions.cpp
--- a/Engine/Source/Editor/BlueprintGraph/Private/EdGraphSchema_K2_Actions.cpp
+++ b/Engine/Source/Editor/BlueprintGraph/Private/EdGraphSchema_K2_Actions.cpp
@@ -105,13 +105,6 @@ UBlueprint* FEdGraphSchemaAction_BlueprintVariableBase::GetSourceBlueprint() con
 	return UBlueprint::GetBlueprintFromClass(ClassToCheck);
 }
 
-void FEdGraphSchemaAction_BlueprintVariableBase::GetVariableTypeTree(
-	TArray<TSharedPtr<UEdGraphSchema_K2::FPinTypeTreeInfo>>& TypeTree, ETypeTreeFilter TypeTreeFilter) const
-{
-	const UEdGraphSchema_K2* Schema = GetDefault<UEdGraphSchema_K2>();
-	Schema->GetVariableTypeTree(TypeTree, TypeTreeFilter);
-}
-
 /////////////////////////////////////////////////////
 // FEdGraphSchemaAction_K2LocalVar
 
diff --git a/Engine/Source/Editor/Kismet/Private/BlueprintDetailsCustomization.cpp b/Engine/Source/Editor/Kismet/Private/BlueprintDetailsCustomization.cpp
--- a/Engine/Source/Editor/Kismet/Private/BlueprintDetailsCustomization.cpp
+++ b/Engine/Source/Editor/Kismet/Private/BlueprintDetailsCustomization.cpp
@@ -259,8 +259,6 @@ void FBlueprintVarActionDetails::CustomizeDetails( IDetailLayoutBuilder& DetailL
 		PropertyOwnerBlueprint = Cast<UBlueprint>(GeneratedClass->ClassGeneratedBy);
 	}
 
-	const UEdGraphSchema_K2* Schema = GetDefault<UEdGraphSchema_K2>();
-	
 	IDetailCategoryBuilder& Category = DetailLayout.EditCategory("Variable", LOCTEXT("VariableDetailsCategory", "Variable"));
 	const FSlateFontInfo DetailFontInfo = IDetailLayoutBuilder::GetDetailFont();
 	
@@ -295,6 +293,12 @@ void FBlueprintVarActionDetails::CustomizeDetails( IDetailLayoutBuilder& DetailL
 	{
 		CustomPinTypeFilter = BlueprintEditor.Pin()->GetOrCreateNamespaceHelperForBlueprint(GetBlueprintObj())->GetPinTypeSelectorFilter();
 	}
+	
+	const UEdGraphSchema* Schema = GetDefault<UEdGraphSchema_K2>();
+	if (BlueprintEditor.IsValid())
+	{
+		Schema = BlueprintEditor.Pin()->GetFocusedGraph()->GetSchema();
+	}
 
 	Category.AddCustomRow(LOCTEXT("VariableTypeLabel", "Variable Type"))
 		.NameContent()
@@ -307,7 +311,7 @@ void FBlueprintVarActionDetails::CustomizeDetails( IDetailLayoutBuilder& DetailL
 		.ValueContent()
 		.MaxDesiredWidth(980.f)
 		[
-			SNew(SPinTypeSelector, FGetPinTypeTree::CreateUObject(Schema, &UEdGraphSchema_K2::GetVariableTypeTree))
+			SNew(SPinTypeSelector, FGetPinTypeTree::CreateUObject(GetDefault<UEdGraphSchema_K2>(), &UEdGraphSchema_K2::GetVariableTypeTree))
 			.TargetPinType(this, &FBlueprintVarActionDetails::OnGetVarType)
 			.OnPinTypeChanged(this, &FBlueprintVarActionDetails::OnVarTypeChanged)
 			.IsEnabled(this, &FBlueprintVarActionDetails::GetVariableTypeChangeEnabled)
diff --git a/Engine/Source/Editor/Kismet/Private/SBlueprintPalette.cpp b/Engine/Source/Editor/Kismet/Private/SBlueprintPalette.cpp
--- a/Engine/Source/Editor/Kismet/Private/SBlueprintPalette.cpp
+++ b/Engine/Source/Editor/Kismet/Private/SBlueprintPalette.cpp
@@ -760,15 +760,18 @@ private:
 			CustomPinTypeFilter = BlueprintEditorPtr.Pin()->GetOrCreateNamespaceHelperForBlueprint(BlueprintObj)->GetPinTypeSelectorFilter();
 		}
 
-		const UEdGraphSchema_K2* Schema = GetDefault<UEdGraphSchema_K2>();
-		FGetPinTypeTree GetPinTypeTreeFunc = ActionPtr.IsValid() ? FGetPinTypeTree::CreateRaw(ActionPtr.Pin().Get(), &FEdGraphSchemaAction_BlueprintVariableBase::GetVariableTypeTree)
-			: FGetPinTypeTree::CreateUObject(Schema, &UEdGraphSchema_K2::GetVariableTypeTree);
+		const UEdGraphSchema* Schema = GetDefault<UEdGraphSchema_K2>();
+		if (BlueprintEditorPtr.IsValid())
+		{
+			Schema = BlueprintEditorPtr.Pin()->GetFocusedGraph()->GetSchema();
+		}
 		
 		this->ChildSlot
 		[
-			SNew(SPinTypeSelector, GetPinTypeTreeFunc)
+			SNew(SPinTypeSelector, FGetPinTypeTree::CreateUObject(GetDefault<UEdGraphSchema_K2>(), &UEdGraphSchema_K2::GetVariableTypeTree))
 			.ReadOnly(InArgs._ReadOnly)
 			.Schema(Schema)
+			.SchemaAction(ActionPtr)
 			.TargetPinType(this, &SPinTypeSelectorHelper::OnGetVarType)
 			.OnPinTypeChanged(this, &SPinTypeSelectorHelper::OnVarTypeChanged)
 			.TypeTreeFilter(ETypeTreeFilter::None)
diff --git a/Engine/Source/Editor/KismetWidgets/Private/SPinTypeSelector.cpp b/Engine/Source/Editor/KismetWidgets/Private/SPinTypeSelector.cpp
--- a/Engine/Source/Editor/KismetWidgets/Private/SPinTypeSelector.cpp
+++ b/Engine/Source/Editor/KismetWidgets/Private/SPinTypeSelector.cpp
@@ -291,7 +291,8 @@ void SPinTypeSelector::Construct(const FArguments& InArgs, FGetPinTypeTree GetPi
 	check(GetPinTypeTreeFunc.IsBound());
 	GetPinTypeTree = GetPinTypeTreeFunc;
 
-	Schema = (UEdGraphSchema_K2*)(InArgs._Schema);
+	Schema = InArgs._Schema;
+	SchemaAction = InArgs._SchemaAction;
 	TypeTreeFilter = InArgs._TypeTreeFilter;
 	TreeViewWidth = InArgs._TreeViewWidth;
 	TreeViewHeight = InArgs._TreeViewHeight;
@@ -977,6 +978,11 @@ TSharedRef<SWidget>	SPinTypeSelector::GetMenuContent(bool bForSecondaryType)
 		}
 	}
 
+	// Remove types not supported by schema
+	TArray<FPinTypeTreeItem> FilteredBySupportedTypes;
+	GetChildrenWithSupportedTypes(TypeTreeRoot, FilteredBySupportedTypes);
+	TypeTreeRoot = FilteredBySupportedTypes;
+
 	if (CustomFilter.IsValid())
 	{
 		NumFilteredPinTypeItems = 0;
@@ -1091,13 +1097,21 @@ TSharedRef<SWidget> SPinTypeSelector::GetPinContainerTypeMenuContent()
 {
 	FMenuBuilder MenuBuilder(true, nullptr);
 
-	if (PinTypeSelectorStatics::PinTypes.Num() == 0)
+	PinTypeSelectorStatics::PinTypes.Reset();
+	
+	PinTypeSelectorStatics::PinTypes.Add(MakeShared<EPinContainerType>(EPinContainerType::None));
+	PinTypeSelectorStatics::PinTypes.Add(MakeShared<EPinContainerType>(EPinContainerType::Array));
+	PinTypeSelectorStatics::PinTypes.Add(MakeShared<EPinContainerType>(EPinContainerType::Set));
+	PinTypeSelectorStatics::PinTypes.Add(MakeShared<EPinContainerType>(EPinContainerType::Map));
+
+
+	for (int32 i = PinTypeSelectorStatics::PinTypes.Num()-1; i >= 0; --i)
 	{
-		PinTypeSelectorStatics::PinTypes.Add(MakeShared<EPinContainerType>(EPinContainerType::None));
-		PinTypeSelectorStatics::PinTypes.Add(MakeShared<EPinContainerType>(EPinContainerType::Array));
-		PinTypeSelectorStatics::PinTypes.Add(MakeShared<EPinContainerType>(EPinContainerType::Set));
-		PinTypeSelectorStatics::PinTypes.Add(MakeShared<EPinContainerType>(EPinContainerType::Map));
-	}
+		if (!Schema->SupportsPinTypeContainer(SchemaAction, TargetPinType.Get(), *PinTypeSelectorStatics::PinTypes[i].Get()))
+		{
+			PinTypeSelectorStatics::PinTypes.RemoveAt(i);
+		}
+	}	
 
 	for (auto& PinType : PinTypeSelectorStatics::PinTypes)
 	{
@@ -1171,6 +1185,42 @@ void SPinTypeSelector::OnFilterTextCommitted(const FText& NewText, ETextCommit::
 	}
 }
 
+bool SPinTypeSelector::GetChildrenWithSupportedTypes(const TArray<FPinTypeTreeItem>& UnfilteredList,
+	TArray<FPinTypeTreeItem>& OutFilteredList)
+{
+	bool bReturnVal = false;
+	
+	for( auto it = UnfilteredList.CreateConstIterator(); it; ++it )
+	{
+		FPinTypeTreeItem Item = *it;
+		FPinTypeTreeItem NewInfo = MakeShareable( new UEdGraphSchema_K2::FPinTypeTreeInfo(Item) );
+		TArray<FPinTypeTreeItem> ValidChildren;
+
+		const bool bHasChildrenWithValidTypes = GetChildrenWithSupportedTypes(Item->Children, ValidChildren);
+		bool bSupportsType = true;
+
+		if (!bHasChildrenWithValidTypes)
+		{
+			bSupportsType = Schema->SupportsPinType(SchemaAction, Item->GetPinType(false));
+		}
+		
+		if (bHasChildrenWithValidTypes || bSupportsType)
+		{
+			NewInfo->Children = ValidChildren;
+			OutFilteredList.Add(NewInfo);
+
+			if (!NewInfo->bReadOnly)
+			{
+				++NumFilteredPinTypeItems;
+			}
+
+			bReturnVal = true;
+		}
+	}
+
+	return bReturnVal;
+}
+
 bool SPinTypeSelector::GetChildrenMatchingSearch(const FText& InSearchText, const TArray<FPinTypeTreeItem>& UnfilteredList, TArray<FPinTypeTreeItem>& OutFilteredList)
 {
 	TArray<FString> FilterTerms;
diff --git a/Engine/Source/Editor/KismetWidgets/Public/SPinTypeSelector.h b/Engine/Source/Editor/KismetWidgets/Public/SPinTypeSelector.h
--- a/Engine/Source/Editor/KismetWidgets/Public/SPinTypeSelector.h
+++ b/Engine/Source/Editor/KismetWidgets/Public/SPinTypeSelector.h
@@ -61,6 +61,7 @@ public:
 	SLATE_BEGIN_ARGS(SPinTypeSelector)
 		: _TargetPinType()
 		, _Schema(nullptr)
+		, _SchemaAction(nullptr)
 		, _TypeTreeFilter(ETypeTreeFilter::None)
 		, _bAllowArrays(true)
 		, _TreeViewWidth(300.f)
@@ -70,7 +71,8 @@ public:
 		, _ReadOnly(false)
 		{}
 		SLATE_ATTRIBUTE( FEdGraphPinType, TargetPinType )
-		SLATE_ARGUMENT( const UEdGraphSchema_K2*, Schema )
+		SLATE_ARGUMENT( const UEdGraphSchema*, Schema )
+		SLATE_ARGUMENT( TWeakPtr<const FEdGraphSchemaAction>, SchemaAction)
 		SLATE_ARGUMENT( ETypeTreeFilter, TypeTreeFilter )
 		SLATE_ARGUMENT( bool, bAllowArrays )
 		SLATE_ATTRIBUTE( FOptionalSize, TreeViewWidth )
@@ -126,7 +128,10 @@ protected:
 	FGetPinTypeTree				GetPinTypeTree;
 
 	/** Schema in charge of determining available types for this pin */
-	UEdGraphSchema_K2*			Schema;
+	const UEdGraphSchema*				Schema;
+
+	/** Schema action related to the pin selection */
+	TWeakPtr<const FEdGraphSchemaAction> SchemaAction;
 
 	/** UEdgraphSchema::ETypeTreeFilter flags for filtering available types*/
 	ETypeTreeFilter				TypeTreeFilter;
@@ -194,6 +199,9 @@ protected:
 	void OnFilterTextChanged(const FText& NewText);
 	void OnFilterTextCommitted(const FText& NewText, ETextCommit::Type CommitInfo);
 
+	/** Helper to generate the filtered list of types, based on the supported types of the schema */
+	bool GetChildrenWithSupportedTypes(const TArray<FPinTypeTreeItem>& UnfilteredList, TArray<FPinTypeTreeItem>& OutFilteredList);
+
 	/** Helper to generate the filtered list of types, based on the search string matching */
 	bool GetChildrenMatchingSearch(const FText& SearchText, const TArray<FPinTypeTreeItem>& UnfilteredList, TArray<FPinTypeTreeItem>& OutFilteredList);
 
diff --git a/Engine/Source/Runtime/Engine/Classes/EdGraph/EdGraphSchema.h b/Engine/Source/Runtime/Engine/Classes/EdGraph/EdGraphSchema.h
--- a/Engine/Source/Runtime/Engine/Classes/EdGraph/EdGraphSchema.h
+++ b/Engine/Source/Runtime/Engine/Classes/EdGraph/EdGraphSchema.h
@@ -800,6 +800,12 @@ class ENGINE_API UEdGraphSchema : public UObject
 	/** Should the Pin in question display an asset picker */
 	virtual bool ShouldShowAssetPickerForPin(UEdGraphPin* Pin) const { return true; }
 
+	/** Returns true if the schema supports the pin type through the schema action */
+	virtual bool SupportsPinType(TWeakPtr<const FEdGraphSchemaAction> SchemaAction, const FEdGraphPinType& PinType) const { return true; }
+
+	/** Returns true if the schema supports the pin type through the schema action */
+	virtual bool SupportsPinTypeContainer(TWeakPtr<const FEdGraphSchemaAction> SchemaAction, const FEdGraphPinType& PinType, const EPinContainerType& ContainerType) const { return true; }
+
 	/**
 	 * Gets the draw color of a pin based on it's type.
 	 *
@@ -809,6 +815,8 @@ class ENGINE_API UEdGraphSchema : public UObject
 	 */
 	virtual FLinearColor GetPinTypeColor(const FEdGraphPinType& PinType) const { return FLinearColor::Black; }
 
+	virtual FLinearColor GetSecondaryPinTypeColor(const FEdGraphPinType& PinType) const { return FLinearColor::White; };
+
 #if WITH_EDITORONLY_DATA
 	/** Get the name to show in the editor */
 	virtual FText GetPinDisplayName(const UEdGraphPin* Pin) const;
