Control Rig: Node Creation Hotkeys

diff --git a/Engine/Plugins/Experimental/ControlRig/Source/ControlRig/Private/Settings/ControlRigSettings.cpp b/Engine/Plugins/Experimental/ControlRig/Source/ControlRig/Private/Settings/ControlRigSettings.cpp
--- a/Engine/Plugins/Experimental/ControlRig/Source/ControlRig/Private/Settings/ControlRigSettings.cpp
+++ b/Engine/Plugins/Experimental/ControlRig/Source/ControlRig/Private/Settings/ControlRigSettings.cpp
@@ -1,6 +1,11 @@
 // Copyright Epic Games, Inc. All Rights Reserved.
 
 #include "Settings/ControlRigSettings.h"
+#include "RigVMModel/RigVMController.h"
+#include "RigVMModel/RigVMNode.h"
+
+#include "Units/Hierarchy/RigUnit_GetTransform.h"
+#include "Units/Hierarchy/RigUnit_SetTransform.h"
 
 UControlRigSettings::UControlRigSettings(const FObjectInitializer& ObjectInitializer)
 	: Super(ObjectInitializer)
@@ -16,7 +21,30 @@ UControlRigSettings::UControlRigSettings(const FObjectInitializer& ObjectInitial
 	SetupEventBorderColor = FLinearColor::Red;
 	BackwardsSolveBorderColor = FLinearColor::Yellow;
 	BackwardsAndForwardsBorderColor = FLinearColor::Blue;
+
+	NodeSnippet_1 = GetSnippetContentForUnitNode(FRigUnit_GetTransform::StaticStruct());
+	NodeSnippet_2 = GetSnippetContentForUnitNode(FRigUnit_SetTransform::StaticStruct());
+	 
 #endif
 }
 
+FString UControlRigSettings::GetSnippetContentForUnitNode(UScriptStruct* InUnitNodeStruct)
+{
+	URigVMGraph* Graph = NewObject<URigVMGraph>(GetTransientPackage(), NAME_None, RF_Transient);
+	URigVMController* Controller = NewObject<URigVMController>(GetTransientPackage(), NAME_None, RF_Transient);
+	Controller->UnfoldStructDelegate.BindLambda([](const UStruct* InStruct) -> bool
+	{
+		if (InStruct == TBaseStructure<FQuat>::Get())
+		{
+			return false;
+		}
 
+		return true;
+	});
+
+	Controller->SetGraph(Graph);
+	URigVMNode* Node = Controller->AddUnitNode(InUnitNodeStruct, FRigUnit::GetMethodName(), FVector2D::ZeroVector, FString(), false);
+	TArray<FName> NodeNames;
+	NodeNames.Add(Node->GetFName());
+	return Controller->ExportNodesToText(NodeNames);
+}
diff --git a/Engine/Plugins/Experimental/ControlRig/Source/ControlRig/Public/Settings/ControlRigSettings.h b/Engine/Plugins/Experimental/ControlRig/Source/ControlRig/Public/Settings/ControlRigSettings.h
--- a/Engine/Plugins/Experimental/ControlRig/Source/ControlRig/Public/Settings/ControlRigSettings.h
+++ b/Engine/Plugins/Experimental/ControlRig/Source/ControlRig/Public/Settings/ControlRigSettings.h
@@ -9,6 +9,8 @@
 #include "CoreMinimal.h"
 #include "Engine/DeveloperSettings.h"
 #include "ControlRigGizmoLibrary.h"
+#include "RigVMModel/RigVMGraph.h"
+
 #include "ControlRigSettings.generated.h"
 
 class UStaticMesh;
@@ -76,7 +78,72 @@ public:
 	 */
 	UPROPERTY(EditAnywhere, config, Category = Viewport)
 	FLinearColor BackwardsAndForwardsBorderColor;
+
+	/**
+	* The default node snippet to create when pressing 1 + Left Mouse Button
+	*/
+	UPROPERTY(EditAnywhere, config, Category = Snippets, meta = (DisplayName = "1"))
+	FString NodeSnippet_1;
+
+	/**
+	* The default node snippet to create when pressing 2 + Left Mouse Button
+	*/
+	UPROPERTY(EditAnywhere, config, Category = Snippets, meta = (DisplayName = "2"))
+	FString NodeSnippet_2;
+
+	/**
+	* The default node snippet to create when pressing 3 + Left Mouse Button
+	*/
+	UPROPERTY(EditAnywhere, config, Category = Snippets, meta = (DisplayName = "3"))
+	FString NodeSnippet_3;
+
+	/**
+	* The default node snippet to create when pressing 4 + Left Mouse Button
+	*/
+	UPROPERTY(EditAnywhere, config, Category = Snippets, meta = (DisplayName = "4"))
+	FString NodeSnippet_4;
+
+	/**
+	* The default node snippet to create when pressing 5 + Left Mouse Button
+	*/
+	UPROPERTY(EditAnywhere, config, Category = Snippets, meta = (DisplayName = "5"))
+	FString NodeSnippet_5;
+
+	/**
+	* The default node snippet to create when pressing 6 + Left Mouse Button
+	*/
+	UPROPERTY(EditAnywhere, config, Category = Snippets, meta = (DisplayName = "6"))
+	FString NodeSnippet_6;
+
+	/**
+	* The default node snippet to create when pressing 7 + Left Mouse Button
+	*/
+	UPROPERTY(EditAnywhere, config, Category = Snippets, meta = (DisplayName = "7"))
+	FString NodeSnippet_7;
+
+	/**
+	* The default node snippet to create when pressing 8 + Left Mouse Button
+	*/
+	UPROPERTY(EditAnywhere, config, Category = Snippets, meta = (DisplayName = "8"))
+	FString NodeSnippet_8;
+
+	/**
+	* The default node snippet to create when pressing 9 + Left Mouse Button
+	*/
+	UPROPERTY(EditAnywhere, config, Category = Snippets, meta = (DisplayName = "9"))
+	FString NodeSnippet_9;
+
+	/**
+	* The default node snippet to create when pressing 0 + Left Mouse Button
+	*/
+	UPROPERTY(EditAnywhere, config, Category = Snippets, meta = (DisplayName = "0"))
+	FString NodeSnippet_0;
+
 #endif
 
-	static UControlRigSettings * Get() { return CastChecked<UControlRigSettings>(UControlRigSettings::StaticClass()->GetDefaultObject()); }
+	static UControlRigSettings * Get() { return GetMutableDefault<UControlRigSettings>(); }
+
+private:
+
+	static FString GetSnippetContentForUnitNode(UScriptStruct* InUnitNodeStruct);
 };
diff --git a/Engine/Plugins/Experimental/ControlRig/Source/ControlRig/Public/Units/RigUnit.h b/Engine/Plugins/Experimental/ControlRig/Source/ControlRig/Public/Units/RigUnit.h
--- a/Engine/Plugins/Experimental/ControlRig/Source/ControlRig/Public/Units/RigUnit.h
+++ b/Engine/Plugins/Experimental/ControlRig/Source/ControlRig/Public/Units/RigUnit.h
@@ -28,6 +28,17 @@ struct CONTROLRIG_API FRigUnit : public FRigVMStruct
 	virtual void Execute(const FRigUnitContext& Context) {}
 
 	virtual FRigElementKey DetermineSpaceForPin(const FString& InPinPath, void* InUserContext) const { return FRigElementKey(); }
+
+	/** The name of the method used within each rig unit */
+	FORCEINLINE_DEBUGGABLE static FName GetMethodName()
+	{
+		static FName MethodName = FName(NAME_None);
+		if(MethodName.IsNone())
+		{
+			MethodName = TEXT("Execute");
+		}
+		return MethodName;
+	}
 };
 
 /** Base class for all rig units that can change data */
diff --git a/Engine/Plugins/Experimental/ControlRig/Source/ControlRigDeveloper/Private/ControlRigBlueprint.cpp b/Engine/Plugins/Experimental/ControlRig/Source/ControlRigDeveloper/Private/ControlRigBlueprint.cpp
--- a/Engine/Plugins/Experimental/ControlRig/Source/ControlRigDeveloper/Private/ControlRigBlueprint.cpp
+++ b/Engine/Plugins/Experimental/ControlRig/Source/ControlRigDeveloper/Private/ControlRigBlueprint.cpp
@@ -2012,7 +2012,7 @@ void UControlRigBlueprint::PopulateModelFromGraphForBackwardsCompatibility(UCont
 				UScriptStruct* UnitStruct = URigVMPin::FindObjectFromCPPTypeObjectPath<UScriptStruct>(StructPath);
 				if (UnitStruct && UnitStruct->IsChildOf(FRigVMStruct::StaticStruct()))
 				{ 
-					ModelNode = GetOrCreateController()->AddUnitNode(UnitStruct, TEXT("Execute"), NodePosition, PropertyName.ToString(), false);
+					ModelNode = GetOrCreateController()->AddUnitNode(UnitStruct, FRigUnit::GetMethodName(), NodePosition, PropertyName.ToString(), false);
 				}
 				else if (PropertyName != NAME_None) // check if this is a variable
 				{
diff --git a/Engine/Plugins/Experimental/ControlRig/Source/ControlRigEditor/Private/ControlRigBlueprintCommands.cpp b/Engine/Plugins/Experimental/ControlRig/Source/ControlRigEditor/Private/ControlRigBlueprintCommands.cpp
--- a/Engine/Plugins/Experimental/ControlRig/Source/ControlRigEditor/Private/ControlRigBlueprintCommands.cpp
+++ b/Engine/Plugins/Experimental/ControlRig/Source/ControlRigEditor/Private/ControlRigBlueprintCommands.cpp
@@ -21,6 +21,16 @@ void FControlRigBlueprintCommands::RegisterCommands()
 	UI_COMMAND(StepInto, "Step Into", "Steps into the collapsed/function node, when halted at a breakpoint.", EUserInterfaceActionType::Button, FInputChord(EKeys::F11));
 	UI_COMMAND(StepOut, "Step Out", "Steps out of the collapsed/function node, when halted at a breakpoint.", EUserInterfaceActionType::Button, FInputChord(EKeys::F11, EModifierKey::Shift));
 	UI_COMMAND(StepOver, "Step Over", "Steps over the node, when halted at a breakpoint.", EUserInterfaceActionType::Button, FInputChord(EKeys::F10));
+	UI_COMMAND(StoreNodeSnippet1, "Store Node Snippet 1", "Stores the selected node(s) into snippet 1", EUserInterfaceActionType::Button, FInputChord(EKeys::One, EModifierKey::Alt));
+	UI_COMMAND(StoreNodeSnippet2, "Store Node Snippet 2", "Stores the selected node(s) into snippet 2", EUserInterfaceActionType::Button, FInputChord(EKeys::Two, EModifierKey::Alt));
+	UI_COMMAND(StoreNodeSnippet3, "Store Node Snippet 3", "Stores the selected node(s) into snippet 3", EUserInterfaceActionType::Button, FInputChord(EKeys::Three, EModifierKey::Alt));
+	UI_COMMAND(StoreNodeSnippet4, "Store Node Snippet 4", "Stores the selected node(s) into snippet 4", EUserInterfaceActionType::Button, FInputChord(EKeys::Four, EModifierKey::Alt));
+	UI_COMMAND(StoreNodeSnippet5, "Store Node Snippet 5", "Stores the selected node(s) into snippet 5", EUserInterfaceActionType::Button, FInputChord(EKeys::Five, EModifierKey::Alt));
+	UI_COMMAND(StoreNodeSnippet6, "Store Node Snippet 6", "Stores the selected node(s) into snippet 6", EUserInterfaceActionType::Button, FInputChord(EKeys::Six, EModifierKey::Alt));
+	UI_COMMAND(StoreNodeSnippet7, "Store Node Snippet 7", "Stores the selected node(s) into snippet 7", EUserInterfaceActionType::Button, FInputChord(EKeys::Seven, EModifierKey::Alt));
+	UI_COMMAND(StoreNodeSnippet8, "Store Node Snippet 8", "Stores the selected node(s) into snippet 8", EUserInterfaceActionType::Button, FInputChord(EKeys::Eight, EModifierKey::Alt));
+	UI_COMMAND(StoreNodeSnippet9, "Store Node Snippet 9", "Stores the selected node(s) into snippet 9", EUserInterfaceActionType::Button, FInputChord(EKeys::Nine, EModifierKey::Alt));
+	UI_COMMAND(StoreNodeSnippet0, "Store Node Snippet 0", "Stores the selected node(s) into snippet 0", EUserInterfaceActionType::Button, FInputChord(EKeys::Zero, EModifierKey::Alt));
 }
 
 #undef LOCTEXT_NAMESPACE
diff --git a/Engine/Plugins/Experimental/ControlRig/Source/ControlRigEditor/Private/ControlRigBlueprintCommands.h b/Engine/Plugins/Experimental/ControlRig/Source/ControlRigEditor/Private/ControlRigBlueprintCommands.h
--- a/Engine/Plugins/Experimental/ControlRig/Source/ControlRigEditor/Private/ControlRigBlueprintCommands.h
+++ b/Engine/Plugins/Experimental/ControlRig/Source/ControlRigEditor/Private/ControlRigBlueprintCommands.h
@@ -63,6 +63,36 @@ public:
 	/** Steps out of the collapsed/function node, when halted at a breakpoint */
 	TSharedPtr< FUICommandInfo > StepOut;
 
+	/** Stores the selected node(s) into snippet 1. */
+	TSharedPtr< FUICommandInfo > StoreNodeSnippet1;
+
+	/** Stores the selected node(s) into snippet 2. */
+	TSharedPtr< FUICommandInfo > StoreNodeSnippet2;
+
+	/** Stores the selected node(s) into snippet 3. */
+	TSharedPtr< FUICommandInfo > StoreNodeSnippet3;
+
+	/** Stores the selected node(s) into snippet 4. */
+	TSharedPtr< FUICommandInfo > StoreNodeSnippet4;
+
+	/** Stores the selected node(s) into snippet 5. */
+	TSharedPtr< FUICommandInfo > StoreNodeSnippet5;
+
+	/** Stores the selected node(s) into snippet 6. */
+	TSharedPtr< FUICommandInfo > StoreNodeSnippet6;
+
+	/** Stores the selected node(s) into snippet 7. */
+	TSharedPtr< FUICommandInfo > StoreNodeSnippet7;
+
+	/** Stores the selected node(s) into snippet 8. */
+	TSharedPtr< FUICommandInfo > StoreNodeSnippet8;
+
+	/** Stores the selected node(s) into snippet 9. */
+	TSharedPtr< FUICommandInfo > StoreNodeSnippet9;
+
+	/** Stores the selected node(s) into snippet 0. */
+	TSharedPtr< FUICommandInfo > StoreNodeSnippet0;
+
 	/**
 	 * Initialize commands
 	 */
diff --git a/Engine/Plugins/Experimental/ControlRig/Source/ControlRigEditor/Private/ControlRigEditorModule.cpp b/Engine/Plugins/Experimental/ControlRig/Source/ControlRigEditor/Private/ControlRigEditorModule.cpp
--- a/Engine/Plugins/Experimental/ControlRig/Source/ControlRigEditor/Private/ControlRigEditorModule.cpp
+++ b/Engine/Plugins/Experimental/ControlRig/Source/ControlRigEditor/Private/ControlRigEditorModule.cpp
@@ -1447,7 +1447,7 @@ void FControlRigEditorModule::GetContextMenuActions(const UControlRigGraphSchema
 									LOCTEXT("AddAlphaInterp_Tooltip", "Injects an interpolate node"),
 									FSlateIcon(),
 									FUIAction(FExecuteAction::CreateLambda([Controller, InGraphPin, ModelPin, ScriptStruct]() {
-										URigVMInjectionInfo* Injection = Controller->AddInjectedNode(ModelPin->GetPinPath(), ModelPin->GetDirection() != ERigVMPinDirection::Output, ScriptStruct, TEXT("Execute"), TEXT("Value"), TEXT("Result"));
+										URigVMInjectionInfo* Injection = Controller->AddInjectedNode(ModelPin->GetPinPath(), ModelPin->GetDirection() != ERigVMPinDirection::Output, ScriptStruct, FRigUnit::GetMethodName(), TEXT("Value"), TEXT("Result"));
 										if (Injection)
 										{
 											TArray<FName> NodeNames;
@@ -1529,7 +1529,7 @@ void FControlRigEditorModule::GetContextMenuActions(const UControlRigGraphSchema
 									LOCTEXT("AddVisualDebug_Tooltip", "Injects a visual debugging node"),
 									FSlateIcon(),
 									FUIAction(FExecuteAction::CreateLambda([RigBlueprint, Controller, InGraphPin, ModelPin, ScriptStruct]() {
-										URigVMInjectionInfo* Injection = Controller->AddInjectedNode(ModelPin->GetPinPath(), ModelPin->GetDirection() != ERigVMPinDirection::Output, ScriptStruct, TEXT("Execute"), TEXT("Value"), TEXT("Value"));
+										URigVMInjectionInfo* Injection = Controller->AddInjectedNode(ModelPin->GetPinPath(), ModelPin->GetDirection() != ERigVMPinDirection::Output, ScriptStruct, FRigUnit::GetMethodName(), TEXT("Value"), TEXT("Value"));
 										if (Injection)
 										{
 											TArray<FName> NodeNames;
diff --git a/Engine/Plugins/Experimental/ControlRig/Source/ControlRigEditor/Private/Editor/ControlRigEditor.cpp b/Engine/Plugins/Experimental/ControlRig/Source/ControlRigEditor/Private/Editor/ControlRigEditor.cpp
--- a/Engine/Plugins/Experimental/ControlRig/Source/ControlRigEditor/Private/Editor/ControlRigEditor.cpp
+++ b/Engine/Plugins/Experimental/ControlRig/Source/ControlRigEditor/Private/Editor/ControlRigEditor.cpp
@@ -85,6 +85,7 @@
 #include "Engine/StaticMesh.h"
 #include "Engine/StaticMeshActor.h"
 #include "Components/StaticMeshComponent.h"
+#include "ControlRig/Private/Units/Execution/RigUnit_SequenceExecution.h"
 
 #define LOCTEXT_NAMESPACE "ControlRigEditor"
 
@@ -403,7 +404,7 @@ void FControlRigEditor::InitControlRigEditor(const EToolkitMode::Type Mode, cons
 		{
 			if (Model->GetNodes().Num() == 0)
 			{
-				URigVMNode* Node = InControlRigBlueprint->GetController()->AddUnitNode(FRigUnit_BeginExecution::StaticStruct(), TEXT("Execute"), FVector2D::ZeroVector, FString(), false);
+				URigVMNode* Node = InControlRigBlueprint->GetController()->AddUnitNode(FRigUnit_BeginExecution::StaticStruct(), FRigUnit::GetMethodName(), FVector2D::ZeroVector, FString(), false);
 				if (Node)
 				{
 					TArray<FName> NodeNames;
@@ -588,7 +589,6 @@ void FControlRigEditor::BindCommands()
 		FControlRigBlueprintCommands::Get().StepOut,
 		FExecuteAction::CreateSP(this, &FControlRigEditor::HandleBreakpointActionRequested, ERigVMBreakpointAction::StepOut),
 		FIsActionChecked::CreateSP(this, &FControlRigEditor::IsHaltedAtBreakpoint));
-
 }
 
 void FControlRigEditor::ToggleExecuteGraph()
@@ -1560,6 +1560,55 @@ void FControlRigEditor::CreateDefaultCommands()
 
 void FControlRigEditor::OnCreateGraphEditorCommands(TSharedPtr<FUICommandList> GraphEditorCommandsList)
 {
+	GraphEditorCommandsList->MapAction(
+		FControlRigBlueprintCommands::Get().StoreNodeSnippet1,
+		FExecuteAction::CreateSP(this, &FControlRigEditor::StoreNodeSnippet, 1),
+		FCanExecuteAction());
+
+	GraphEditorCommandsList->MapAction(
+		FControlRigBlueprintCommands::Get().StoreNodeSnippet2,
+		FExecuteAction::CreateSP(this, &FControlRigEditor::StoreNodeSnippet, 2),
+		FCanExecuteAction());
+
+	GraphEditorCommandsList->MapAction(
+		FControlRigBlueprintCommands::Get().StoreNodeSnippet3,
+		FExecuteAction::CreateSP(this, &FControlRigEditor::StoreNodeSnippet, 3),
+		FCanExecuteAction());
+
+	GraphEditorCommandsList->MapAction(
+		FControlRigBlueprintCommands::Get().StoreNodeSnippet4,
+		FExecuteAction::CreateSP(this, &FControlRigEditor::StoreNodeSnippet, 4),
+		FCanExecuteAction());
+
+	GraphEditorCommandsList->MapAction(
+		FControlRigBlueprintCommands::Get().StoreNodeSnippet5,
+		FExecuteAction::CreateSP(this, &FControlRigEditor::StoreNodeSnippet, 5),
+		FCanExecuteAction());
+
+	GraphEditorCommandsList->MapAction(
+		FControlRigBlueprintCommands::Get().StoreNodeSnippet6,
+		FExecuteAction::CreateSP(this, &FControlRigEditor::StoreNodeSnippet, 6),
+		FCanExecuteAction());
+
+	GraphEditorCommandsList->MapAction(
+		FControlRigBlueprintCommands::Get().StoreNodeSnippet7,
+		FExecuteAction::CreateSP(this, &FControlRigEditor::StoreNodeSnippet, 7),
+		FCanExecuteAction());
+
+	GraphEditorCommandsList->MapAction(
+		FControlRigBlueprintCommands::Get().StoreNodeSnippet8,
+		FExecuteAction::CreateSP(this, &FControlRigEditor::StoreNodeSnippet, 8),
+		FCanExecuteAction());
+
+	GraphEditorCommandsList->MapAction(
+		FControlRigBlueprintCommands::Get().StoreNodeSnippet9,
+		FExecuteAction::CreateSP(this, &FControlRigEditor::StoreNodeSnippet, 9),
+		FCanExecuteAction());
+
+	GraphEditorCommandsList->MapAction(
+		FControlRigBlueprintCommands::Get().StoreNodeSnippet0,
+		FExecuteAction::CreateSP(this, &FControlRigEditor::StoreNodeSnippet, 0),
+		FCanExecuteAction());
 
 }
 
@@ -1941,6 +1990,86 @@ bool FControlRigEditor::CanPasteNodes() const
 	return GetFocusedController()->CanImportNodesFromText(TextToImport);
 }
 
+FReply FControlRigEditor::OnSpawnGraphNodeByShortcut(FInputChord InChord, const FVector2D& InPosition,
+                                                     UEdGraph* InGraph)
+{
+	if(!InChord.HasAnyModifierKeys())
+	{
+		if(InChord.Key == EKeys::One)
+		{
+			RestoreNodeSnippet(1);
+			return FReply::Handled();
+		}
+		else if(InChord.Key == EKeys::Two)
+		{
+			RestoreNodeSnippet(2);
+			return FReply::Handled();
+		}
+		else if(InChord.Key == EKeys::Three)
+		{
+			RestoreNodeSnippet(3);
+			return FReply::Handled();
+		}
+		else if(InChord.Key == EKeys::Four)
+		{
+			RestoreNodeSnippet(4);
+			return FReply::Handled();
+		}
+		else if(InChord.Key == EKeys::Five)
+		{
+			RestoreNodeSnippet(5);
+			return FReply::Handled();
+		}
+		else if(InChord.Key == EKeys::Six)
+		{
+			RestoreNodeSnippet(6);
+			return FReply::Handled();
+		}
+		else if(InChord.Key == EKeys::Seven)
+		{
+			RestoreNodeSnippet(7);
+			return FReply::Handled();
+		}
+		else if(InChord.Key == EKeys::Eight)
+		{
+			RestoreNodeSnippet(8);
+			return FReply::Handled();
+		}
+		else if(InChord.Key == EKeys::Nine)
+		{
+			RestoreNodeSnippet(9);
+			return FReply::Handled();
+		}
+		else if(InChord.Key == EKeys::Zero)
+		{
+			RestoreNodeSnippet(0);
+			return FReply::Handled();
+		}
+		else if(InChord.Key == EKeys::Zero)
+		{
+			RestoreNodeSnippet(0);
+			return FReply::Handled();
+		}
+		
+		if(UControlRigGraph* RigGraph = Cast<UControlRigGraph>(InGraph))
+		{
+			if(URigVMController* Controller = RigGraph->GetController())
+			{
+				if(InChord.Key == EKeys::B)
+				{
+					Controller->AddBranchNode(InPosition);
+				}
+				else if(InChord.Key == EKeys::S)
+				{
+					Controller->AddUnitNode(FRigUnit_SequenceExecution::StaticStruct(), FRigUnit::GetMethodName(), InPosition);
+				}
+			}
+		}
+	}
+
+	return FReply::Unhandled();
+}
+
 void FControlRigEditor::PasteNodes()
 {
 	DECLARE_SCOPE_HIERARCHICAL_COUNTER_FUNC()
@@ -4703,7 +4832,7 @@ void FControlRigEditor::OnGraphNodeDropToPerform(TSharedPtr<FGraphNodeDragDropOp
 								{
 									Controller->OpenUndoBracket(TEXT("Create Collection from Items"));
 
-									if (URigVMNode* ItemsNode = Controller->AddUnitNode(FRigUnit_CollectionItems::StaticStruct(), TEXT("Execute"), NodePosition))
+									if (URigVMNode* ItemsNode = Controller->AddUnitNode(FRigUnit_CollectionItems::StaticStruct(), FRigUnit::GetMethodName(), NodePosition))
 									{
 										if (URigVMPin* ItemsPin = ItemsNode->FindPin(TEXT("Items")))
 										{
@@ -5014,7 +5143,7 @@ void FControlRigEditor::HandleMakeElementGetterSetter(ERigElementGetterSetterTyp
 		}
 
 		FName Name = FControlRigBlueprintUtils::ValidateName(Blueprint, StructTemplate->GetName());
-		if (URigVMUnitNode* ModelNode = GetFocusedController()->AddUnitNode(StructTemplate, TEXT("Execute"), NodePosition))
+		if (URigVMUnitNode* ModelNode = GetFocusedController()->AddUnitNode(StructTemplate, FRigUnit::GetMethodName(), NodePosition))
 		{
 			FString ItemTypeStr = StaticEnum<ERigElementType>()->GetDisplayNameTextByValue((int64)Key.Type).ToString();
 			NewNode.Name = ModelNode->GetFName();
@@ -5582,5 +5711,114 @@ TSharedPtr<SGraphEditor> FControlRigEditor::GetGraphEditor(UEdGraph* InEdGraph)
 	return TSharedPtr<SGraphEditor>();
 }
 
+void FControlRigEditor::StoreNodeSnippet(int32 InSnippetIndex)
+{
+	check((InSnippetIndex >= 0) && (InSnippetIndex < 10));
+
+	URigVMController* Controller = GetFocusedController();
+	if(Controller == nullptr)
+	{
+		return;
+	}
+
+	FString Snippet = Controller->ExportSelectedNodesToText();
+	if(Snippet.IsEmpty())
+	{
+		return;
+	}
+
+	FString* Setting = GetSnippetStorage(InSnippetIndex);
+	check(Setting != nullptr);
+
+	UControlRigSettings* Settings = UControlRigSettings::Get();
+	Settings->Modify();
+	(*Setting) = Snippet;
+
+	const FString PropertyName = FString::Printf(TEXT("NodeSnippet_%d"), InSnippetIndex);
+	FProperty* Property = UControlRigSettings::StaticClass()->FindPropertyByName(*PropertyName);
+	check(Property);
+	
+	Settings->UpdateSinglePropertyInConfigFile(Property, Settings->GetDefaultConfigFilename());
+}
+
+void FControlRigEditor::RestoreNodeSnippet(int32 InSnippetIndex)
+{
+	FString* Setting = GetSnippetStorage(InSnippetIndex);
+	check(Setting != nullptr);
+
+	if(Setting->IsEmpty())
+	{
+		return;
+	}
+
+	FPlatformApplicationMisc::ClipboardCopy(*(*Setting));
+	PasteNodes();
+}
+
+FString* FControlRigEditor::GetSnippetStorage(int32 InSnippetIndex)
+{
+	UControlRigSettings* Settings = UControlRigSettings::Get();
+	switch(InSnippetIndex)
+	{
+		case 1:
+		{
+			return &Settings->NodeSnippet_1;
+			break;
+		}
+		case 2:
+		{
+			return &Settings->NodeSnippet_2;
+			break;
+		}
+		case 3:
+		{
+			return &Settings->NodeSnippet_3;
+			break;
+		}
+		case 4:
+		{
+			return &Settings->NodeSnippet_4;
+			break;
+		}
+		case 5:
+		{
+			return &Settings->NodeSnippet_5;
+			break;
+		}
+		case 6:
+		{
+			return &Settings->NodeSnippet_6;
+			break;
+		}
+		case 7:
+		{
+			return &Settings->NodeSnippet_7;
+			break;
+		}
+		case 8:
+		{
+			return &Settings->NodeSnippet_8;
+			break;
+		}
+		case 9:
+		{
+			return &Settings->NodeSnippet_9;
+			break;
+		}
+		case 0:
+		{
+			return &Settings->NodeSnippet_0;
+			break;
+		}
+		default:
+		{
+			checkNoEntry();
+			break;
+		}
+	}
+
+	return nullptr;
+}
+
 
 #undef LOCTEXT_NAMESPACE
diff --git a/Engine/Plugins/Experimental/ControlRig/Source/ControlRigEditor/Private/Editor/ControlRigEditor.h b/Engine/Plugins/Experimental/ControlRig/Source/ControlRigEditor/Private/Editor/ControlRigEditor.h
--- a/Engine/Plugins/Experimental/ControlRig/Source/ControlRigEditor/Private/Editor/ControlRigEditor.h
+++ b/Engine/Plugins/Experimental/ControlRig/Source/ControlRigEditor/Private/Editor/ControlRigEditor.h
@@ -116,6 +116,8 @@ public:
 	virtual void PasteNodes() override;
 	virtual bool CanPasteNodes() const override;
 
+	virtual FReply OnSpawnGraphNodeByShortcut(FInputChord InChord, const FVector2D& InPosition, UEdGraph* InGraph) override;
+
 	// FEditorUndoClient Interface
 	virtual void PostUndo(bool bSuccess) override;
 	virtual void PostRedo(bool bSuccess) override;
@@ -392,6 +394,13 @@ protected:
 	URigVMGraph* GetFocusedModel() const;
 	URigVMController* GetFocusedController() const;
 	TSharedPtr<SGraphEditor> GetGraphEditor(UEdGraph* InEdGraph) const;
+	
+	// stores a node snippet into the setting
+	void StoreNodeSnippet(int32 InSnippetIndex);
+	// restores a node snippet from the setting
+	void RestoreNodeSnippet(int32 InSnippetIndex);
+	// returns the right setting storage for a given snippet index
+	static FString* GetSnippetStorage(int32 InSnippetIndex);
 
 	URigVMController* ActiveController;
 	bool bControlRigEditorInitialized;
diff --git a/Engine/Plugins/Experimental/ControlRig/Source/ControlRigEditor/Private/Graph/NodeSpawners/ControlRigUnitNodeSpawner.cpp b/Engine/Plugins/Experimental/ControlRig/Source/ControlRigEditor/Private/Graph/NodeSpawners/ControlRigUnitNodeSpawner.cpp
--- a/Engine/Plugins/Experimental/ControlRig/Source/ControlRigEditor/Private/Graph/NodeSpawners/ControlRigUnitNodeSpawner.cpp
+++ b/Engine/Plugins/Experimental/ControlRig/Source/ControlRigEditor/Private/Graph/NodeSpawners/ControlRigUnitNodeSpawner.cpp
@@ -144,7 +144,7 @@ UControlRigGraphNode* UControlRigUnitNodeSpawner::SpawnNode(UEdGraph* ParentGrap
 		FRigVMUnitNodeCreatedContext& UnitNodeCreatedContext = Controller->GetUnitNodeCreatedContext();
 		FRigVMUnitNodeCreatedContext::FScope ReasonScope(UnitNodeCreatedContext, ERigVMNodeCreatedReason::NodeSpawner);
 
-		if (URigVMUnitNode* ModelNode = Controller->AddUnitNode(StructTemplate, TEXT("Execute"), Location, Name.ToString(), bUndo))
+		if (URigVMUnitNode* ModelNode = Controller->AddUnitNode(StructTemplate, FRigUnit::GetMethodName(), Location, Name.ToString(), bUndo))
 		{
 			NewNode = Cast<UControlRigGraphNode>(RigGraph->FindNodeForModelNodeName(ModelNode->GetFName()));
 			check(NewNode);
diff --git a/Engine/Source/Editor/Kismet/Public/BlueprintEditor.h b/Engine/Source/Editor/Kismet/Public/BlueprintEditor.h
--- a/Engine/Source/Editor/Kismet/Public/BlueprintEditor.h
+++ b/Engine/Source/Editor/Kismet/Public/BlueprintEditor.h
@@ -586,7 +586,7 @@ public:
 	virtual void ClearSelectionStateFor(FName SelectionOwner);
 
 	/** Handles spawning a graph node in the current graph using the passed in chord */
-	FReply OnSpawnGraphNodeByShortcut(FInputChord InChord, const FVector2D& InPosition, UEdGraph* InGraph);
+	virtual FReply OnSpawnGraphNodeByShortcut(FInputChord InChord, const FVector2D& InPosition, UEdGraph* InGraph);
 
 	/**
 	 * Perform the actual promote to variable action on the given pin in the given blueprint.
