Control Rig: Fix build break introduced in CL 16601675

diff --git a/Engine/Plugins/Experimental/ControlRig/Source/ControlRig/Private/Settings/ControlRigSettings.cpp b/Engine/Plugins/Experimental/ControlRig/Source/ControlRig/Private/Settings/ControlRigSettings.cpp
--- a/Engine/Plugins/Experimental/ControlRig/Source/ControlRig/Private/Settings/ControlRigSettings.cpp
+++ b/Engine/Plugins/Experimental/ControlRig/Source/ControlRig/Private/Settings/ControlRigSettings.cpp
@@ -1,8 +1,11 @@
 // Copyright Epic Games, Inc. All Rights Reserved.
 
 #include "Settings/ControlRigSettings.h"
+
+#if WITH_EDITOR
 #include "RigVMModel/RigVMController.h"
 #include "RigVMModel/RigVMNode.h"
+#endif
 
 #include "Units/Hierarchy/RigUnit_GetTransform.h"
 #include "Units/Hierarchy/RigUnit_SetTransform.h"
@@ -28,6 +31,8 @@ UControlRigSettings::UControlRigSettings(const FObjectInitializer& ObjectInitial
 #endif
 }
 
+#if WITH_EDITOR
+
 FString UControlRigSettings::GetSnippetContentForUnitNode(UScriptStruct* InUnitNodeStruct)
 {
 	URigVMGraph* Graph = NewObject<URigVMGraph>(GetTransientPackage(), NAME_None, RF_Transient);
@@ -48,3 +53,5 @@ FString UControlRigSettings::GetSnippetContentForUnitNode(UScriptStruct* InUnitN
 	NodeNames.Add(Node->GetFName());
 	return Controller->ExportNodesToText(NodeNames);
 }
+
+#endif
\ No newline at end of file
diff --git a/Engine/Plugins/Experimental/ControlRig/Source/ControlRig/Public/Settings/ControlRigSettings.h b/Engine/Plugins/Experimental/ControlRig/Source/ControlRig/Public/Settings/ControlRigSettings.h
--- a/Engine/Plugins/Experimental/ControlRig/Source/ControlRig/Public/Settings/ControlRigSettings.h
+++ b/Engine/Plugins/Experimental/ControlRig/Source/ControlRig/Public/Settings/ControlRigSettings.h
@@ -9,7 +9,10 @@
 #include "CoreMinimal.h"
 #include "Engine/DeveloperSettings.h"
 #include "ControlRigGizmoLibrary.h"
+
+#if WITH_EDITOR
 #include "RigVMModel/RigVMGraph.h"
+#endif
 
 #include "ControlRigSettings.generated.h"
 
@@ -143,7 +146,10 @@ public:
 
 	static UControlRigSettings * Get() { return GetMutableDefault<UControlRigSettings>(); }
 
-private:
+#if WITH_EDITOR
 
+private:
 	static FString GetSnippetContentForUnitNode(UScriptStruct* InUnitNodeStruct);
+
+#endif
 };
