Control Rig: Improve undo / redo functionality for Controls / Bones in Setup Mode

diff --git a/Engine/Plugins/Experimental/ControlRig/Source/ControlRig/Private/Rigs/RigHierarchy.cpp b/Engine/Plugins/Experimental/ControlRig/Source/ControlRig/Private/Rigs/RigHierarchy.cpp
--- a/Engine/Plugins/Experimental/ControlRig/Source/ControlRig/Private/Rigs/RigHierarchy.cpp
+++ b/Engine/Plugins/Experimental/ControlRig/Source/ControlRig/Private/Rigs/RigHierarchy.cpp
@@ -2300,13 +2300,15 @@ void URigHierarchy::PushTransformToStack(const FRigElementKey& InKey, ERigTransf
 
 	if(bIsInteracting)
 	{
+		bool bCanMerge = LastInteractedKey == InKey;
+
 		FRigTransformStackEntry LastEntry;
 		if(!TransformUndoStack.IsEmpty())
 		{
 			LastEntry = TransformUndoStack.Last();
 		}
 
-		if(LastEntry.Key == InKey && LastEntry.EntryType == InEntryType && LastEntry.bAffectChildren == bAffectChildren)
+		if(bCanMerge && LastEntry.Key == InKey && LastEntry.EntryType == InEntryType && LastEntry.bAffectChildren == bAffectChildren)
 		{
 			// merge the entries on the stack
 			TransformUndoStack.Last() = 
@@ -2322,6 +2324,7 @@ void URigHierarchy::PushTransformToStack(const FRigElementKey& InKey, ERigTransf
 		}
 
 		TransformRedoStack.Reset();
+		LastInteractedKey = InKey;
 		return;
 	}
 
@@ -2384,11 +2387,23 @@ void URigHierarchy::PushCurveToStack(const FRigElementKey& InKey, float InOldCur
 bool URigHierarchy::ApplyTransformFromStack(const FRigTransformStackEntry& InEntry, bool bUndo)
 {
 #if WITH_EDITOR
-	
+
+	bool bApplyInitialForCurrent = false;
 	FRigBaseElement* Element = Find(InEntry.Key);
 	if(Element == nullptr)
 	{
-		return false;
+		// this might be a transient control which had been removed.
+		if(InEntry.Key.Type == ERigElementType::Control)
+		{
+			const FRigElementKey TargetKey = UControlRig::GetElementKeyFromTransientControl(InEntry.Key);
+			Element = Find(TargetKey);
+			bApplyInitialForCurrent = Element != nullptr;
+		}
+
+		if(Element == nullptr)
+		{
+			return false;
+		}
 	}
 
 	const FTransform& Transform = bUndo ? InEntry.OldTransform : InEntry.NewTransform;
@@ -2397,7 +2412,12 @@ bool URigHierarchy::ApplyTransformFromStack(const FRigTransformStackEntry& InEnt
 	{
 		case ERigTransformStackEntryType::TransformPose:
 		{
-			SetTransform(Cast<FRigTransformElement>(Element), Transform, InEntry.TransformType, InEntry.bAffectChildren, false); 
+			SetTransform(Cast<FRigTransformElement>(Element), Transform, InEntry.TransformType, InEntry.bAffectChildren, false);
+
+			if(ERigTransformType::IsCurrent(InEntry.TransformType) && bApplyInitialForCurrent)
+			{
+				SetTransform(Cast<FRigTransformElement>(Element), Transform, ERigTransformType::MakeInitial(InEntry.TransformType), InEntry.bAffectChildren, false);
+			}
 			break;
 		}
 		case ERigTransformStackEntryType::ControlOffset:
diff --git a/Engine/Plugins/Experimental/ControlRig/Source/ControlRig/Private/Settings/ControlRigSettings.cpp b/Engine/Plugins/Experimental/ControlRig/Source/ControlRig/Private/Settings/ControlRigSettings.cpp
--- a/Engine/Plugins/Experimental/ControlRig/Source/ControlRig/Private/Settings/ControlRigSettings.cpp
+++ b/Engine/Plugins/Experimental/ControlRig/Source/ControlRig/Private/Settings/ControlRigSettings.cpp
@@ -20,6 +20,7 @@ UControlRigSettings::UControlRigSettings(const FObjectInitializer& ObjectInitial
 	DefaultGizmoLibrary = LoadObject<UControlRigGizmoLibrary>(nullptr, TEXT("/ControlRig/Controls/DefaultGizmoLibrary.DefaultGizmoLibrary"));
 	bResetControlsOnCompile = true;
 	bResetControlsOnPinValueInteraction = false;
+	bEnableUndoForPoseInteraction = true;
 
 	SetupEventBorderColor = FLinearColor::Red;
 	BackwardsSolveBorderColor = FLinearColor::Yellow;
diff --git a/Engine/Plugins/Experimental/ControlRig/Source/ControlRig/Public/ControlRig.h b/Engine/Plugins/Experimental/ControlRig/Source/ControlRig/Public/ControlRig.h
--- a/Engine/Plugins/Experimental/ControlRig/Source/ControlRig/Public/ControlRig.h
+++ b/Engine/Plugins/Experimental/ControlRig/Source/ControlRig/Public/ControlRig.h
@@ -713,6 +713,7 @@ public:
 	friend class FControlRigInteractionScope;
 	friend class UControlRigValidator;
 	friend struct FAnimNode_ControlRig;
+	friend class URigHierarchy;
 };
 
 class CONTROLRIG_API FControlRigBracketScope
diff --git a/Engine/Plugins/Experimental/ControlRig/Source/ControlRig/Public/Rigs/RigHierarchy.h b/Engine/Plugins/Experimental/ControlRig/Source/ControlRig/Public/Rigs/RigHierarchy.h
--- a/Engine/Plugins/Experimental/ControlRig/Source/ControlRig/Public/Rigs/RigHierarchy.h
+++ b/Engine/Plugins/Experimental/ControlRig/Source/ControlRig/Public/Rigs/RigHierarchy.h
@@ -94,6 +94,7 @@ public:
 	, TransformStackIndex(0)
 	, bTransactingForTransformChange(false)
 	, bIsInteracting(false)
+	, LastInteractedKey()
 	, bSuspendNotifications(false)
 	, ResetPoseHash(INDEX_NONE)
 #if WITH_EDITOR
@@ -1672,7 +1673,11 @@ public:
 	 * Starts an interaction on the rig.
 	 * This will cause all transform actions happening to be merged
 	 */
-	FORCEINLINE void EndInteraction() { bIsInteracting = false; }
+	FORCEINLINE void EndInteraction()
+	{
+		bIsInteracting = false;
+		LastInteractedKey.Reset();
+	}
 
 	/**
 	 * Returns the transform stack index
@@ -2282,7 +2287,12 @@ private:
 	 * Manages merging transform actions into one during an interaction
 	 */
 	bool bIsInteracting;
-	
+
+	/**
+	* Stores the last key being interacted on
+	*/
+	FRigElementKey LastInteractedKey;
+
 	/** 
 	 * If set to true all notifs coming from this hierarchy will be suspended
 	 */
diff --git a/Engine/Plugins/Experimental/ControlRig/Source/ControlRig/Public/Settings/ControlRigSettings.h b/Engine/Plugins/Experimental/ControlRig/Source/ControlRig/Public/Settings/ControlRigSettings.h
--- a/Engine/Plugins/Experimental/ControlRig/Source/ControlRig/Public/Settings/ControlRigSettings.h
+++ b/Engine/Plugins/Experimental/ControlRig/Source/ControlRig/Public/Settings/ControlRigSettings.h
@@ -51,6 +51,11 @@ public:
 	UPROPERTY(EditAnywhere, config, Category = Interaction)
 	bool bResetControlsOnPinValueInteraction;
 
+	// When this is checked any hierarchy interaction within the Control Rig
+	// Editor will be stored on the undo stack
+	UPROPERTY(EditAnywhere, config, Category = Interaction)
+	bool bEnableUndoForPoseInteraction;
+
 	/**
 	 * When checked controls will be reset during a manual compilation
 	 * (when pressing the Compile button)
diff --git a/Engine/Plugins/Experimental/ControlRig/Source/ControlRigDeveloper/Private/ControlRigBlueprint.cpp b/Engine/Plugins/Experimental/ControlRig/Source/ControlRigDeveloper/Private/ControlRigBlueprint.cpp
--- a/Engine/Plugins/Experimental/ControlRig/Source/ControlRigDeveloper/Private/ControlRigBlueprint.cpp
+++ b/Engine/Plugins/Experimental/ControlRig/Source/ControlRigDeveloper/Private/ControlRigBlueprint.cpp
@@ -3072,6 +3072,22 @@ void UControlRigBlueprint::HandleHierarchyModified(ERigHierarchyNotification InN
 				}
 			}
 
+			if(bClearTransientControls)
+			{
+				if(UControlRig* RigBeingDebugged = Cast<UControlRig>(GetObjectBeingDebugged()))
+				{
+					const FName TransientControlName = UControlRig::GetNameForTransientControl(InElement->GetKey());
+					const FRigElementKey TransientControlKey(TransientControlName, ERigElementType::Control);
+					if (const FRigControlElement* ControlElement = RigBeingDebugged->GetHierarchy()->Find<FRigControlElement>(TransientControlKey))
+					{
+						if (ControlElement->Settings.bIsTransientControl)
+						{
+							bClearTransientControls = false;
+						}
+					}
+				}
+			}
+
 			if(bClearTransientControls)
 			{
 				ClearTransientControls();
diff --git a/Engine/Plugins/Experimental/ControlRig/Source/ControlRigEditor/Private/EditMode/ControlRigEditMode.cpp b/Engine/Plugins/Experimental/ControlRig/Source/ControlRigEditor/Private/EditMode/ControlRigEditMode.cpp
--- a/Engine/Plugins/Experimental/ControlRig/Source/ControlRigEditor/Private/EditMode/ControlRigEditMode.cpp
+++ b/Engine/Plugins/Experimental/ControlRig/Source/ControlRigEditor/Private/EditMode/ControlRigEditMode.cpp
@@ -51,6 +51,7 @@
 #include "Framework/Application/SlateApplication.h"
 #include "UnrealEdGlobals.h"
 #include "Editor/UnrealEdEngine.h"
+#include "Settings/ControlRigSettings.h"
 
 void UControlRigEditModeDelegateHelper::OnPoseInitialized()
 {
@@ -745,7 +746,7 @@ bool FControlRigEditMode::StartTracking(FEditorViewportClient* InViewportClient,
 					InteractionScope = new FControlRigInteractionScope(ControlRig);
 				}
 			}
-			else
+			else if(UControlRigSettings::Get()->bEnableUndoForPoseInteraction)
 			{
 				InteractionScope = new FControlRigInteractionScope(ControlRig);
 			}
diff --git a/Engine/Plugins/Experimental/ControlRig/Source/ControlRigEditor/Private/Editor/ControlRigEditor.cpp b/Engine/Plugins/Experimental/ControlRig/Source/ControlRigEditor/Private/Editor/ControlRigEditor.cpp
--- a/Engine/Plugins/Experimental/ControlRig/Source/ControlRigEditor/Private/Editor/ControlRigEditor.cpp
+++ b/Engine/Plugins/Experimental/ControlRig/Source/ControlRigEditor/Private/Editor/ControlRigEditor.cpp
@@ -2045,11 +2045,6 @@ FReply FControlRigEditor::OnSpawnGraphNodeByShortcut(FInputChord InChord, const
 			RestoreNodeSnippet(0);
 			return FReply::Handled();
 		}
-		else if(InChord.Key == EKeys::Zero)
-		{
-			RestoreNodeSnippet(0);
-			return FReply::Handled();
-		}
 		
 		if(UControlRigGraph* RigGraph = Cast<UControlRigGraph>(InGraph))
 		{
@@ -3555,6 +3550,13 @@ void FControlRigEditor::HandlePreviewSceneCreated(const TSharedRef<IPersonaPrevi
 				NewSelection.Add(ElementToSelect);
 				Hierarchy->GetController()->SetSelection(NewSelection);
 			}
+
+			/*
+			if (FControlRigEditMode* EditMode = GetEditMode())
+			{
+				EditMode->RequestToRecreateGizmoActors();
+			}
+			*/
 		}
 	});
 
@@ -5721,12 +5723,25 @@ void FControlRigEditor::StoreNodeSnippet(int32 InSnippetIndex)
 		return;
 	}
 
-	FString Snippet = Controller->ExportSelectedNodesToText();
+	const TArray<FName> SelectedNodeNames = Controller->GetGraph()->GetSelectNodes();
+	if(SelectedNodeNames.Num() == 0)
+	{
+		return;
+	}
+	
+	FString Snippet = Controller->ExportNodesToText(SelectedNodeNames);
 	if(Snippet.IsEmpty())
 	{
 		return;
 	}
 
+	TArray<FString> NodeNameStrings;
+	for(const FName& SelectedNodeName : SelectedNodeNames)
+	{
+		NodeNameStrings.Add(SelectedNodeName.ToString());
+	}
+	const FString NodeNamesJoined = FString::Join(NodeNameStrings, TEXT(", "));
+
 	FString* Setting = GetSnippetStorage(InSnippetIndex);
 	check(Setting != nullptr);
 
@@ -5739,6 +5754,18 @@ void FControlRigEditor::StoreNodeSnippet(int32 InSnippetIndex)
 	check(Property);
 	
 	Settings->UpdateSinglePropertyInConfigFile(Property, Settings->GetDefaultConfigFilename());
+
+	FNotificationInfo Info(FText::FromString(FString::Printf(
+		TEXT("A snippet has been stored to the Project Settings.\n")
+		TEXT("Nodes %s are now stored as snippet %d.\n")
+		TEXT("You can restore the snippet by pressing %d and left clicking into the graph."),
+		*NodeNamesJoined, InSnippetIndex, InSnippetIndex)));
+	Info.bFireAndForget = true;
+	Info.FadeOutDuration = 2.0f;
+	Info.ExpireDuration = 8.0f;
+
+	TSharedPtr<SNotificationItem> NotificationPtr = FSlateNotificationManager::Get().AddNotification(Info);
+	NotificationPtr->SetCompletionState(SNotificationItem::CS_Success);
 }
 
 void FControlRigEditor::RestoreNodeSnippet(int32 InSnippetIndex)
