Control Rig: Follow up on Snippets functionality

diff --git a/Engine/Plugins/Experimental/ControlRig/Source/ControlRig/Private/Settings/ControlRigSettings.cpp b/Engine/Plugins/Experimental/ControlRig/Source/ControlRig/Private/Settings/ControlRigSettings.cpp
--- a/Engine/Plugins/Experimental/ControlRig/Source/ControlRig/Private/Settings/ControlRigSettings.cpp
+++ b/Engine/Plugins/Experimental/ControlRig/Source/ControlRig/Private/Settings/ControlRigSettings.cpp
@@ -7,9 +7,6 @@
 #include "RigVMModel/RigVMNode.h"
 #endif
 
-#include "Units/Hierarchy/RigUnit_GetTransform.h"
-#include "Units/Hierarchy/RigUnit_SetTransform.h"
-
 UControlRigSettings::UControlRigSettings(const FObjectInitializer& ObjectInitializer)
 	: Super(ObjectInitializer)
 #if WITH_EDITORONLY_DATA
@@ -26,33 +23,5 @@ UControlRigSettings::UControlRigSettings(const FObjectInitializer& ObjectInitial
 	BackwardsSolveBorderColor = FLinearColor::Yellow;
 	BackwardsAndForwardsBorderColor = FLinearColor::Blue;
 
-	NodeSnippet_1 = GetSnippetContentForUnitNode(FRigUnit_GetTransform::StaticStruct());
-	NodeSnippet_2 = GetSnippetContentForUnitNode(FRigUnit_SetTransform::StaticStruct());
-	 
 #endif
 }
-
-#if WITH_EDITOR
-
-FString UControlRigSettings::GetSnippetContentForUnitNode(UScriptStruct* InUnitNodeStruct)
-{
-	URigVMGraph* Graph = NewObject<URigVMGraph>(GetTransientPackage(), NAME_None, RF_Transient);
-	URigVMController* Controller = NewObject<URigVMController>(GetTransientPackage(), NAME_None, RF_Transient);
-	Controller->UnfoldStructDelegate.BindLambda([](const UStruct* InStruct) -> bool
-	{
-		if (InStruct == TBaseStructure<FQuat>::Get())
-		{
-			return false;
-		}
-
-		return true;
-	});
-
-	Controller->SetGraph(Graph);
-	URigVMNode* Node = Controller->AddUnitNode(InUnitNodeStruct, FRigUnit::GetMethodName(), FVector2D::ZeroVector, FString(), false);
-	TArray<FName> NodeNames;
-	NodeNames.Add(Node->GetFName());
-	return Controller->ExportNodesToText(NodeNames);
-}
-
-#endif
\ No newline at end of file
diff --git a/Engine/Plugins/Experimental/ControlRig/Source/ControlRig/Public/Settings/ControlRigSettings.h b/Engine/Plugins/Experimental/ControlRig/Source/ControlRig/Public/Settings/ControlRigSettings.h
--- a/Engine/Plugins/Experimental/ControlRig/Source/ControlRig/Public/Settings/ControlRigSettings.h
+++ b/Engine/Plugins/Experimental/ControlRig/Source/ControlRig/Public/Settings/ControlRigSettings.h
@@ -150,11 +150,4 @@ public:
 #endif
 
 	static UControlRigSettings * Get() { return GetMutableDefault<UControlRigSettings>(); }
-
-#if WITH_EDITOR
-
-private:
-	static FString GetSnippetContentForUnitNode(UScriptStruct* InUnitNodeStruct);
-
-#endif
 };
diff --git a/Engine/Plugins/Experimental/ControlRig/Source/ControlRigEditor/Private/ControlRigBlueprintCommands.cpp b/Engine/Plugins/Experimental/ControlRig/Source/ControlRigEditor/Private/ControlRigBlueprintCommands.cpp
--- a/Engine/Plugins/Experimental/ControlRig/Source/ControlRigEditor/Private/ControlRigBlueprintCommands.cpp
+++ b/Engine/Plugins/Experimental/ControlRig/Source/ControlRigEditor/Private/ControlRigBlueprintCommands.cpp
@@ -21,16 +21,16 @@ void FControlRigBlueprintCommands::RegisterCommands()
 	UI_COMMAND(StepInto, "Step Into", "Steps into the collapsed/function node, when halted at a breakpoint.", EUserInterfaceActionType::Button, FInputChord(EKeys::F11));
 	UI_COMMAND(StepOut, "Step Out", "Steps out of the collapsed/function node, when halted at a breakpoint.", EUserInterfaceActionType::Button, FInputChord(EKeys::F11, EModifierKey::Shift));
 	UI_COMMAND(StepOver, "Step Over", "Steps over the node, when halted at a breakpoint.", EUserInterfaceActionType::Button, FInputChord(EKeys::F10));
-	UI_COMMAND(StoreNodeSnippet1, "Store Node Snippet 1", "Stores the selected node(s) into snippet 1", EUserInterfaceActionType::Button, FInputChord(EKeys::One, EModifierKey::Alt));
-	UI_COMMAND(StoreNodeSnippet2, "Store Node Snippet 2", "Stores the selected node(s) into snippet 2", EUserInterfaceActionType::Button, FInputChord(EKeys::Two, EModifierKey::Alt));
-	UI_COMMAND(StoreNodeSnippet3, "Store Node Snippet 3", "Stores the selected node(s) into snippet 3", EUserInterfaceActionType::Button, FInputChord(EKeys::Three, EModifierKey::Alt));
-	UI_COMMAND(StoreNodeSnippet4, "Store Node Snippet 4", "Stores the selected node(s) into snippet 4", EUserInterfaceActionType::Button, FInputChord(EKeys::Four, EModifierKey::Alt));
-	UI_COMMAND(StoreNodeSnippet5, "Store Node Snippet 5", "Stores the selected node(s) into snippet 5", EUserInterfaceActionType::Button, FInputChord(EKeys::Five, EModifierKey::Alt));
-	UI_COMMAND(StoreNodeSnippet6, "Store Node Snippet 6", "Stores the selected node(s) into snippet 6", EUserInterfaceActionType::Button, FInputChord(EKeys::Six, EModifierKey::Alt));
-	UI_COMMAND(StoreNodeSnippet7, "Store Node Snippet 7", "Stores the selected node(s) into snippet 7", EUserInterfaceActionType::Button, FInputChord(EKeys::Seven, EModifierKey::Alt));
-	UI_COMMAND(StoreNodeSnippet8, "Store Node Snippet 8", "Stores the selected node(s) into snippet 8", EUserInterfaceActionType::Button, FInputChord(EKeys::Eight, EModifierKey::Alt));
-	UI_COMMAND(StoreNodeSnippet9, "Store Node Snippet 9", "Stores the selected node(s) into snippet 9", EUserInterfaceActionType::Button, FInputChord(EKeys::Nine, EModifierKey::Alt));
-	UI_COMMAND(StoreNodeSnippet0, "Store Node Snippet 0", "Stores the selected node(s) into snippet 0", EUserInterfaceActionType::Button, FInputChord(EKeys::Zero, EModifierKey::Alt));
+	UI_COMMAND(StoreNodeSnippet1, "Store Node Snippet 1", "Stores the selected node(s) into snippet 1", EUserInterfaceActionType::Button, FInputChord(EKeys::One, EModifierKey::Alt | EModifierKey::Shift));
+	UI_COMMAND(StoreNodeSnippet2, "Store Node Snippet 2", "Stores the selected node(s) into snippet 2", EUserInterfaceActionType::Button, FInputChord(EKeys::Two, EModifierKey::Alt | EModifierKey::Shift));
+	UI_COMMAND(StoreNodeSnippet3, "Store Node Snippet 3", "Stores the selected node(s) into snippet 3", EUserInterfaceActionType::Button, FInputChord(EKeys::Three, EModifierKey::Alt | EModifierKey::Shift));
+	UI_COMMAND(StoreNodeSnippet4, "Store Node Snippet 4", "Stores the selected node(s) into snippet 4", EUserInterfaceActionType::Button, FInputChord(EKeys::Four, EModifierKey::Alt | EModifierKey::Shift));
+	UI_COMMAND(StoreNodeSnippet5, "Store Node Snippet 5", "Stores the selected node(s) into snippet 5", EUserInterfaceActionType::Button, FInputChord(EKeys::Five, EModifierKey::Alt | EModifierKey::Shift));
+	UI_COMMAND(StoreNodeSnippet6, "Store Node Snippet 6", "Stores the selected node(s) into snippet 6", EUserInterfaceActionType::Button, FInputChord(EKeys::Six, EModifierKey::Alt | EModifierKey::Shift));
+	UI_COMMAND(StoreNodeSnippet7, "Store Node Snippet 7", "Stores the selected node(s) into snippet 7", EUserInterfaceActionType::Button, FInputChord(EKeys::Seven, EModifierKey::Alt | EModifierKey::Shift));
+	UI_COMMAND(StoreNodeSnippet8, "Store Node Snippet 8", "Stores the selected node(s) into snippet 8", EUserInterfaceActionType::Button, FInputChord(EKeys::Eight, EModifierKey::Alt | EModifierKey::Shift));
+	UI_COMMAND(StoreNodeSnippet9, "Store Node Snippet 9", "Stores the selected node(s) into snippet 9", EUserInterfaceActionType::Button, FInputChord(EKeys::Nine, EModifierKey::Alt | EModifierKey::Shift));
+	UI_COMMAND(StoreNodeSnippet0, "Store Node Snippet 0", "Stores the selected node(s) into snippet 0", EUserInterfaceActionType::Button, FInputChord(EKeys::Zero, EModifierKey::Alt | EModifierKey::Shift));
 }
 
 #undef LOCTEXT_NAMESPACE
diff --git a/Engine/Plugins/Experimental/ControlRig/Source/ControlRigEditor/Private/Editor/ControlRigEditor.cpp b/Engine/Plugins/Experimental/ControlRig/Source/ControlRigEditor/Private/Editor/ControlRigEditor.cpp
--- a/Engine/Plugins/Experimental/ControlRig/Source/ControlRigEditor/Private/Editor/ControlRigEditor.cpp
+++ b/Engine/Plugins/Experimental/ControlRig/Source/ControlRigEditor/Private/Editor/ControlRigEditor.cpp
@@ -55,6 +55,9 @@
 #include "ControlRig/Private/Units/Hierarchy/RigUnit_GetControlTransform.h"
 #include "ControlRig/Private/Units/Hierarchy/RigUnit_SetControlTransform.h"
 #include "ControlRig/Private/Units/Execution/RigUnit_Collection.h"
+#include "ControlRig/Private/Units/Highlevel/Hierarchy/RigUnit_TransformConstraint.h"
+#include "ControlRig/Private/Units/Hierarchy/RigUnit_GetControlTransform.h"
+#include "ControlRig/Public/Units/Hierarchy/RigUnit_SetCurveValue.h"
 #include "Units/Hierarchy/RigUnit_AddBoneTransform.h"
 #include "Graph/NodeSpawners/ControlRigUnitNodeSpawner.h"
 #include "Graph/ControlRigGraphSchema.h"
@@ -1994,6 +1997,43 @@ FReply FControlRigEditor::OnSpawnGraphNodeByShortcut(FInputChord InChord, const
                                                      UEdGraph* InGraph)
 {
 	if(!InChord.HasAnyModifierKeys())
+	{
+		if(UControlRigGraph* RigGraph = Cast<UControlRigGraph>(InGraph))
+		{
+			if(URigVMController* Controller = RigGraph->GetController())
+			{
+				if(InChord.Key == EKeys::B)
+				{
+					Controller->AddBranchNode(InPosition);
+				}
+				else if(InChord.Key == EKeys::S)
+				{
+					Controller->AddUnitNode(FRigUnit_SequenceExecution::StaticStruct(), FRigUnit::GetMethodName(), InPosition);
+				}
+				else if(InChord.Key == EKeys::One)
+				{
+					Controller->AddUnitNode(FRigUnit_GetTransform::StaticStruct(), FRigUnit::GetMethodName(), InPosition);
+				}
+				else if(InChord.Key == EKeys::Two)
+				{
+					Controller->AddUnitNode(FRigUnit_SetTransform::StaticStruct(), FRigUnit::GetMethodName(), InPosition);
+				}
+				else if(InChord.Key == EKeys::Three)
+				{
+					Controller->AddUnitNode(FRigUnit_ParentConstraint::StaticStruct(), FRigUnit::GetMethodName(), InPosition);
+				}
+				else if(InChord.Key == EKeys::Four)
+				{
+					Controller->AddUnitNode(FRigUnit_GetControlFloat::StaticStruct(), FRigUnit::GetMethodName(), InPosition);
+				}
+				else if(InChord.Key == EKeys::Five)
+				{
+					Controller->AddUnitNode(FRigUnit_SetCurveValue::StaticStruct(), FRigUnit::GetMethodName(), InPosition);
+				}
+			}
+		}
+	}
+	else if(InChord.NeedsAlt() && !InChord.NeedsControl() && !InChord.NeedsShift())
 	{
 		if(InChord.Key == EKeys::One)
 		{
@@ -2045,21 +2085,6 @@ FReply FControlRigEditor::OnSpawnGraphNodeByShortcut(FInputChord InChord, const
 			RestoreNodeSnippet(0);
 			return FReply::Handled();
 		}
-		
-		if(UControlRigGraph* RigGraph = Cast<UControlRigGraph>(InGraph))
-		{
-			if(URigVMController* Controller = RigGraph->GetController())
-			{
-				if(InChord.Key == EKeys::B)
-				{
-					Controller->AddBranchNode(InPosition);
-				}
-				else if(InChord.Key == EKeys::S)
-				{
-					Controller->AddUnitNode(FRigUnit_SequenceExecution::StaticStruct(), FRigUnit::GetMethodName(), InPosition);
-				}
-			}
-		}
 	}
 
 	return FReply::Unhandled();
@@ -5780,6 +5805,16 @@ void FControlRigEditor::RestoreNodeSnippet(int32 InSnippetIndex)
 
 	FPlatformApplicationMisc::ClipboardCopy(*(*Setting));
 	PasteNodes();
+	
+	FNotificationInfo Info(FText::FromString(FString::Printf(
+		TEXT("Snippet %d has been restored."),
+		InSnippetIndex)));
+	Info.bFireAndForget = true;
+	Info.FadeOutDuration = 0.5f;
+	Info.ExpireDuration = 3.0f;
+
+	TSharedPtr<SNotificationItem> NotificationPtr = FSlateNotificationManager::Get().AddNotification(Info);
+	NotificationPtr->SetCompletionState(SNotificationItem::CS_Success);
 }
 
 FString* FControlRigEditor::GetSnippetStorage(int32 InSnippetIndex)
