Control Rig: Follow up on Snippets functionality

diff --git a/Engine/Plugins/Experimental/ControlRig/Source/ControlRigEditor/Private/Editor/ControlRigEditor.cpp b/Engine/Plugins/Experimental/ControlRig/Source/ControlRigEditor/Private/Editor/ControlRigEditor.cpp
--- a/Engine/Plugins/Experimental/ControlRig/Source/ControlRigEditor/Private/Editor/ControlRigEditor.cpp
+++ b/Engine/Plugins/Experimental/ControlRig/Source/ControlRigEditor/Private/Editor/ControlRigEditor.cpp
@@ -5783,7 +5783,7 @@ void FControlRigEditor::StoreNodeSnippet(int32 InSnippetIndex)
 	FNotificationInfo Info(FText::FromString(FString::Printf(
 		TEXT("A snippet has been stored to the Project Settings.\n")
 		TEXT("Nodes %s are now stored as snippet %d.\n")
-		TEXT("You can restore the snippet by pressing %d and left clicking into the graph."),
+		TEXT("You can restore the snippet by pressing Alt-%d and left clicking into the graph."),
 		*NodeNamesJoined, InSnippetIndex, InSnippetIndex)));
 	Info.bFireAndForget = true;
 	Info.FadeOutDuration = 2.0f;
