Control Rig: Refactor RigHierarchyContextMenu to use UToolMenu APIs for python extension support

diff --git a/Engine/Plugins/Experimental/ControlRig/Source/ControlRigEditor/Private/EditMode/ControlRigEditMode.cpp b/Engine/Plugins/Experimental/ControlRig/Source/ControlRigEditor/Private/EditMode/ControlRigEditMode.cpp
--- a/Engine/Plugins/Experimental/ControlRig/Source/ControlRigEditor/Private/EditMode/ControlRigEditMode.cpp
+++ b/Engine/Plugins/Experimental/ControlRig/Source/ControlRigEditor/Private/EditMode/ControlRigEditMode.cpp
@@ -52,6 +52,7 @@
 #include "UnrealEdGlobals.h"
 #include "Editor/UnrealEdEngine.h"
 #include "Settings/ControlRigSettings.h"
+#include "ToolMenus.h"
 
 void UControlRigEditModeDelegateHelper::OnPoseInitialized()
 {
@@ -978,12 +979,16 @@ void FControlRigEditMode::OpenContextMenu(FEditorViewportClient* InViewportClien
 		Commands = OnContextMenuCommandsDelegate.Execute();
 	}
 
-	if (OnContextMenuDelegate.IsBound())
+	if (OnGetContextMenuDelegate.IsBound())
 	{
-		FMenuBuilder MenuBuilder(true, Commands);
-		OnContextMenuDelegate.Execute(MenuBuilder);
+		TSharedPtr<SWidget> MenuWidget = SNullWidget::NullWidget;
+		
+		if (UToolMenu* ContextMenu = OnGetContextMenuDelegate.Execute())
+		{
+			UToolMenus* ToolMenus = UToolMenus::Get();
+			MenuWidget = ToolMenus->GenerateWidget(ContextMenu);
+		}
 
-		TSharedPtr<SWidget> MenuWidget = MenuBuilder.MakeWidget();
 		TSharedPtr<SWidget> ParentWidget = InViewportClient->GetEditorViewportWidget();
 
 		if (MenuWidget.IsValid() && ParentWidget.IsValid())
diff --git a/Engine/Plugins/Experimental/ControlRig/Source/ControlRigEditor/Private/EditMode/ControlRigEditMode.h b/Engine/Plugins/Experimental/ControlRig/Source/ControlRigEditor/Private/EditMode/ControlRigEditMode.h
--- a/Engine/Plugins/Experimental/ControlRig/Source/ControlRigEditor/Private/EditMode/ControlRigEditMode.h
+++ b/Engine/Plugins/Experimental/ControlRig/Source/ControlRigEditor/Private/EditMode/ControlRigEditMode.h
@@ -38,11 +38,13 @@ struct FRigControl;
 class IControlRigManipulatable;
 class ISequencer;
 enum class EControlRigSetKey : uint8;
+class UToolMenu;
 
 DECLARE_DELEGATE_RetVal_ThreeParams(FTransform, FOnGetRigElementTransform, const FRigElementKey& /*RigElementKey*/, bool /*bLocal*/, bool /*bOnDebugInstance*/);
 DECLARE_DELEGATE_ThreeParams(FOnSetRigElementTransform, const FRigElementKey& /*RigElementKey*/, const FTransform& /*Transform*/, bool /*bLocal*/);
 DECLARE_DELEGATE_RetVal(TSharedPtr<FUICommandList>, FNewMenuCommandsDelegate);
 DECLARE_MULTICAST_DELEGATE_TwoParams(FControlRigAddedOrRemoved, UControlRig*, bool /*true if added, false if removed*/);
+DECLARE_DELEGATE_RetVal(UToolMenu*, FOnGetContextMenu);
 
 class FControlRigEditMode;
 
@@ -130,7 +132,7 @@ public:
 	FOnSetRigElementTransform& OnSetRigElementTransform() { return OnSetRigElementTransformDelegate; }
 
 	/** Context Menu Delegates */
-	FNewMenuDelegate& OnContextMenu() { return OnContextMenuDelegate; }
+	FOnGetContextMenu& OnGetContextMenu() { return OnGetContextMenuDelegate; }
 	FNewMenuCommandsDelegate& OnContextMenuCommands() { return OnContextMenuCommandsDelegate; }
 	FSimpleMulticastDelegate& OnAnimSystemInitialized() { return OnAnimSystemInitializedDelegate; }
 
@@ -256,7 +258,7 @@ protected:
 
 	FOnGetRigElementTransform OnGetRigElementTransformDelegate;
 	FOnSetRigElementTransform OnSetRigElementTransformDelegate;
-	FNewMenuDelegate OnContextMenuDelegate;
+	FOnGetContextMenu OnGetContextMenuDelegate;
 	FNewMenuCommandsDelegate OnContextMenuCommandsDelegate;
 	FSimpleMulticastDelegate OnAnimSystemInitializedDelegate;
 	FControlRigAddedOrRemoved OnControlRigAddedOrRemovedDelegate;
diff --git a/Engine/Plugins/Experimental/ControlRig/Source/ControlRigEditor/Private/Editor/ControlRigContextMenuContext.cpp b/Engine/Plugins/Experimental/ControlRig/Source/ControlRigEditor/Private/Editor/ControlRigContextMenuContext.cpp
new file mode 100644
--- /dev/null
+++ b/Engine/Plugins/Experimental/ControlRig/Source/ControlRigEditor/Private/Editor/ControlRigContextMenuContext.cpp
@@ -0,0 +1,16 @@
+// Copyright Epic Games, Inc. All Rights Reserved.
+
+#include "ControlRigContextMenuContext.h"
+
+#include "ControlRigBlueprint.h"
+#include "ControlRigEditor.h"
+
+UControlRigBlueprint* UControlRigContextMenuContext::GetControlRigBlueprint() const
+{
+	if (ControlRigEditor.IsValid())
+	{
+		return ControlRigEditor.Pin()->GetControlRigBlueprint();
+	}
+
+	return nullptr;
+}
\ No newline at end of file
diff --git a/Engine/Plugins/Experimental/ControlRig/Source/ControlRigEditor/Private/Editor/ControlRigContextMenuContext.h b/Engine/Plugins/Experimental/ControlRig/Source/ControlRigEditor/Private/Editor/ControlRigContextMenuContext.h
new file mode 100644
--- /dev/null
+++ b/Engine/Plugins/Experimental/ControlRig/Source/ControlRigEditor/Private/Editor/ControlRigContextMenuContext.h
@@ -0,0 +1,24 @@
+﻿// Copyright Epic Games, Inc. All Rights Reserved.
+
+#pragma once
+
+#include "CoreMinimal.h"
+
+#include "ControlRigContextMenuContext.generated.h"
+
+class UControlRig;
+class UControlRigBlueprint;
+class FControlRigEditor;
+
+UCLASS(BlueprintType)
+class UControlRigContextMenuContext : public UObject
+{
+	GENERATED_BODY()
+public:
+	/** Get the control rig blueprint that we are editing */
+	UFUNCTION(BlueprintCallable, Category = ControlRigEditorExtensions)
+    UControlRigBlueprint* GetControlRigBlueprint() const;
+	
+	/** Our owning control rig editor */
+	TWeakPtr<FControlRigEditor> ControlRigEditor;
+};
\ No newline at end of file
diff --git a/Engine/Plugins/Experimental/ControlRig/Source/ControlRigEditor/Private/Editor/ControlRigEditor.cpp b/Engine/Plugins/Experimental/ControlRig/Source/ControlRigEditor/Private/Editor/ControlRigEditor.cpp
--- a/Engine/Plugins/Experimental/ControlRig/Source/ControlRigEditor/Private/Editor/ControlRigEditor.cpp
+++ b/Engine/Plugins/Experimental/ControlRig/Source/ControlRigEditor/Private/Editor/ControlRigEditor.cpp
@@ -364,7 +364,7 @@ void FControlRigEditor::InitControlRigEditor(const EToolkitMode::Type Mode, cons
 	{
 		EditMode->OnGetRigElementTransform() = FOnGetRigElementTransform::CreateSP(this, &FControlRigEditor::GetRigElementTransform);
 		EditMode->OnSetRigElementTransform() = FOnSetRigElementTransform::CreateSP(this, &FControlRigEditor::SetRigElementTransform);
-		EditMode->OnContextMenu() = FNewMenuDelegate::CreateSP(this, &FControlRigEditor::HandleOnViewportContextMenuDelegate);
+		EditMode->OnGetContextMenu() = FOnGetContextMenu::CreateSP(this, &FControlRigEditor::HandleOnGetViewportContextMenuDelegate);
 		EditMode->OnContextMenuCommands() = FNewMenuCommandsDelegate::CreateSP(this, &FControlRigEditor::HandleOnViewportContextMenuCommandsDelegate);
 		EditMode->OnAnimSystemInitialized().Add(FSimpleMulticastDelegate::FDelegate::CreateRaw(this, &FControlRigEditor::OnAnimInitialized));
 		
@@ -5682,12 +5682,13 @@ void FControlRigEditor::UpdateGraphCompilerErrors()
 
 }
 
-void FControlRigEditor::HandleOnViewportContextMenuDelegate(class FMenuBuilder& MenuBuilder)
+UToolMenu* FControlRigEditor::HandleOnGetViewportContextMenuDelegate()
 {
-	if (OnViewportContextMenuDelegate.IsBound())
+	if (OnGetViewportContextMenuDelegate.IsBound())
 	{
-		OnViewportContextMenuDelegate.Execute(MenuBuilder);
+		return OnGetViewportContextMenuDelegate.Execute();
 	}
+	return nullptr;
 }
 
 TSharedPtr<FUICommandList> FControlRigEditor::HandleOnViewportContextMenuCommandsDelegate()
diff --git a/Engine/Plugins/Experimental/ControlRig/Source/ControlRigEditor/Private/Editor/ControlRigEditor.h b/Engine/Plugins/Experimental/ControlRig/Source/ControlRigEditor/Private/Editor/ControlRigEditor.h
--- a/Engine/Plugins/Experimental/ControlRig/Source/ControlRigEditor/Private/Editor/ControlRigEditor.h
+++ b/Engine/Plugins/Experimental/ControlRig/Source/ControlRigEditor/Private/Editor/ControlRigEditor.h
@@ -27,6 +27,7 @@ class SWidget;
 class SBorder;
 class USkeletalMesh;
 class FStructOnScope;
+class UToolMenu;
 
 UENUM()
 enum class EControlRigEditorEventQueue : uint8
@@ -169,7 +170,7 @@ public:
 	void OnGraphNodeDropToPerform(TSharedPtr<FGraphNodeDragDropOp> DragDropOp, UEdGraph* Graph, const FVector2D& NodePosition, const FVector2D& ScreenPosition);
 
 	FPersonaViewportKeyDownDelegate& GetKeyDownDelegate() { return OnKeyDownDelegate; }
-	FNewMenuDelegate& OnViewportContextMenu() { return OnViewportContextMenuDelegate; }
+	FOnGetContextMenu& OnGetViewportContextMenu() { return OnGetViewportContextMenuDelegate; }
 	FNewMenuCommandsDelegate& OnViewportContextMenuCommands() { return OnViewportContextMenuCommandsDelegate; }
 
 	DECLARE_EVENT_OneParam(FControlRigEditor, FPreviewControlRigUpdated, FControlRigEditor*);
@@ -373,8 +374,8 @@ protected:
 	FPersonaViewportKeyDownDelegate OnKeyDownDelegate;
 
 	/** Delgate to build the context menu for the viewport */
-	FNewMenuDelegate OnViewportContextMenuDelegate;
-	void HandleOnViewportContextMenuDelegate(class FMenuBuilder& MenuBuilder);
+	FOnGetContextMenu OnGetViewportContextMenuDelegate;
+	UToolMenu* HandleOnGetViewportContextMenuDelegate();
 	FNewMenuCommandsDelegate OnViewportContextMenuCommandsDelegate;
 	TSharedPtr<FUICommandList> HandleOnViewportContextMenuCommandsDelegate();
 
diff --git a/Engine/Plugins/Experimental/ControlRig/Source/ControlRigEditor/Private/Editor/SRigHierarchy.cpp b/Engine/Plugins/Experimental/ControlRig/Source/ControlRigEditor/Private/Editor/SRigHierarchy.cpp
--- a/Engine/Plugins/Experimental/ControlRig/Source/ControlRigEditor/Private/Editor/SRigHierarchy.cpp
+++ b/Engine/Plugins/Experimental/ControlRig/Source/ControlRigEditor/Private/Editor/SRigHierarchy.cpp
@@ -38,6 +38,8 @@
 #include "Types/WidgetActiveTimerDelegate.h"
 #include "Dialogs/CustomDialog.h"
 #include "EditMode/ControlRigEditMode.h"
+#include "ToolMenus.h"
+#include "ControlRigContextMenuContext.h"
 
 #define LOCTEXT_NAMESPACE "SRigHierarchy"
 
@@ -289,7 +291,7 @@ SRigHierarchy::~SRigHierarchy()
 	if (ControlRigEditor.IsValid())
 	{
 		ControlRigEditor.Pin()->GetKeyDownDelegate().Unbind();
-		ControlRigEditor.Pin()->OnViewportContextMenu().Unbind();
+		ControlRigEditor.Pin()->OnGetViewportContextMenu().Unbind();
 		ControlRigEditor.Pin()->OnViewportContextMenuCommands().Unbind();
 	}
 
@@ -420,7 +422,7 @@ void SRigHierarchy::Construct(const FArguments& InArgs, TSharedRef<FControlRigEd
 				.OnGenerateRow(this, &SRigHierarchy::MakeTableRowWidget)
 				.OnGetChildren(this, &SRigHierarchy::HandleGetChildrenForTree)
 				.OnSelectionChanged(this, &SRigHierarchy::OnSelectionChanged)
-				.OnContextMenuOpening(this, &SRigHierarchy::CreateContextMenu)
+				.OnContextMenuOpening(this, &SRigHierarchy::CreateContextMenuWidget)
 				.OnMouseButtonClick(this, &SRigHierarchy::OnItemClicked)
 				.OnMouseButtonDoubleClick(this, &SRigHierarchy::OnItemDoubleClicked)
 				.OnSetExpansionRecursive(this, &SRigHierarchy::OnSetExpansionRecursive)
@@ -461,7 +463,7 @@ void SRigHierarchy::Construct(const FArguments& InArgs, TSharedRef<FControlRigEd
 		ControlRigEditor.Pin()->GetKeyDownDelegate().BindLambda([&](const FGeometry& MyGeometry, const FKeyEvent& InKeyEvent)->FReply {
 			return OnKeyDown(MyGeometry, InKeyEvent);
 		});
-		ControlRigEditor.Pin()->OnViewportContextMenu().BindSP(this, &SRigHierarchy::FillContextMenu);
+		ControlRigEditor.Pin()->OnGetViewportContextMenu().BindSP(this, &SRigHierarchy::GetOrCreateContextMenu);
 		ControlRigEditor.Pin()->OnViewportContextMenuCommands().BindSP(this, &SRigHierarchy::GetContextMenuCommands);
 	}
 }
@@ -545,7 +547,10 @@ void SRigHierarchy::BindCommands()
 	CommandList->MapAction(
 		Commands.ControlBoneTransform,
 		FExecuteAction::CreateSP(this, &SRigHierarchy::HandleControlBoneOrSpaceTransform),
-		FCanExecuteAction::CreateSP(this, &SRigHierarchy::IsSingleBoneSelected));
+		FCanExecuteAction::CreateSP(this, &SRigHierarchy::IsSingleBoneSelected),
+		FIsActionChecked(),
+		FIsActionButtonVisible::CreateSP(this, &SRigHierarchy::IsSingleBoneSelected)
+		);
 
 	CommandList->MapAction(
 		Commands.Unparent,
@@ -723,6 +728,19 @@ void SRigHierarchy::RefreshTreeView(bool bRebuildContent)
 				AddSpacerElement();
 			}
 		}
+		else
+		{
+			if (RootElements.Num()> 0)
+			{
+				// elements may be added at the end of the list after a spacer element
+				// we need to remove the spacer element and re-add it at the end
+				RootElements.RemoveAll([](TSharedPtr<FRigTreeElement> InElement)
+				{
+					return InElement.Get()->Key == FRigElementKey();
+				});
+				AddSpacerElement();
+			}
+		}
 
 		TreeView->RequestTreeRefresh();
 		{
@@ -1318,12 +1336,16 @@ TSharedRef< SWidget > SRigHierarchy::CreateFilterMenu()
 	return MenuBuilder.MakeWidget();
 }
 
-TSharedPtr< SWidget > SRigHierarchy::CreateContextMenu()
+TSharedPtr< SWidget > SRigHierarchy::CreateContextMenuWidget()
 {
-	const bool CloseAfterSelection = true;
-	FMenuBuilder MenuBuilder(CloseAfterSelection, CommandList);
-	FillContextMenu(MenuBuilder);
-	return MenuBuilder.MakeWidget();
+	UToolMenus* ToolMenus = UToolMenus::Get();
+
+	if (UToolMenu* Menu = GetOrCreateContextMenu())
+	{
+		return ToolMenus->GenerateWidget(Menu);
+	}
+	
+	return SNullWidget::NullWidget;
 }
 
 void SRigHierarchy::OnItemClicked(TSharedPtr<FRigTreeElement> InItem)
@@ -1380,10 +1402,18 @@ void SRigHierarchy::OnSetExpansionRecursive(TSharedPtr<FRigTreeElement> InItem,
 	SetExpansionRecursive(InItem, false, bShouldBeExpanded);
 }
 
-void SRigHierarchy::FillContextMenu(class FMenuBuilder& MenuBuilder)
+UToolMenu* SRigHierarchy::GetOrCreateContextMenu()
 {
-	const FControlRigHierarchyCommands& Actions = FControlRigHierarchyCommands::Get();
+	const FName MenuName = TEXT("ControlRigEditor.RigHierarchy.ContextMenu");
+	const FName InteractionSectionName = TEXT("Interaction");
+	
+	UToolMenus* ToolMenus = UToolMenus::Get();
+	const FControlRigHierarchyCommands& Commands = FControlRigHierarchyCommands::Get();
+	
+	if (!ToolMenus->IsMenuRegistered(MenuName))
 	{
+		UToolMenu* Menu = ToolMenus->RegisterMenu(MenuName);
+
 		struct FLocalMenuBuilder
 		{
 			static void FillNewMenu(FMenuBuilder& InSubMenuBuilder, TSharedPtr<SRigHierarchyTreeView> InTreeView)
@@ -1406,64 +1436,69 @@ void SRigHierarchy::FillContextMenu(class FMenuBuilder& MenuBuilder)
 			}
 		};
 
-		MenuBuilder.BeginSection("Elements", LOCTEXT("ElementsHeader", "Elements"));
-		MenuBuilder.AddSubMenu(LOCTEXT("New", "New"), LOCTEXT("New_ToolTip", "Create New Elements"),
-			FNewMenuDelegate::CreateStatic(&FLocalMenuBuilder::FillNewMenu, TreeView));
-
-		MenuBuilder.AddMenuEntry(Actions.DeleteItem);
-		MenuBuilder.AddMenuEntry(Actions.DuplicateItem);
-		MenuBuilder.AddMenuEntry(Actions.RenameItem);
-		MenuBuilder.AddMenuEntry(Actions.MirrorItem);
-		MenuBuilder.EndSection();
-
-		if (IsSingleBoneSelected())
-		{
-			MenuBuilder.BeginSection("Interaction", LOCTEXT("InteractionHeader", "Interaction"));
-			MenuBuilder.AddMenuEntry(Actions.ControlBoneTransform);
-			MenuBuilder.EndSection();
-		}
-
-		/*
-		if (IsSingleNullSelected())
-		{
-			MenuBuilder.BeginSection("Interaction", LOCTEXT("InteractionHeader", "Interaction"));
-			MenuBuilder.AddMenuEntry(Actions.ControlSpaceTransform);
-			MenuBuilder.EndSection();
-		}
-		*/
-
-		MenuBuilder.BeginSection("Copy&Paste", LOCTEXT("Copy&PasteHeader", "Copy & Paste"));
-		MenuBuilder.AddMenuEntry(Actions.CopyItems);
-		MenuBuilder.AddMenuEntry(Actions.PasteItems);
-		MenuBuilder.AddMenuEntry(Actions.PasteLocalTransforms);
-		MenuBuilder.AddMenuEntry(Actions.PasteGlobalTransforms);
-		MenuBuilder.EndSection();
-
-		MenuBuilder.BeginSection("Transforms", LOCTEXT("TransformsHeader", "Transforms"));
-		MenuBuilder.AddMenuEntry(Actions.ResetTransform);
-		MenuBuilder.AddMenuEntry(Actions.ResetAllTransforms);
-		MenuBuilder.AddMenuEntry(Actions.SetInitialTransformFromCurrentTransform);
-		MenuBuilder.AddMenuEntry(Actions.SetInitialTransformFromClosestBone);
-		MenuBuilder.AddMenuEntry(Actions.SetGizmoTransformFromCurrent);
-		MenuBuilder.AddMenuEntry(Actions.Unparent);
-		MenuBuilder.EndSection();
-
-		MenuBuilder.BeginSection("Assets", LOCTEXT("AssetsHeader", "Assets"));
-		MenuBuilder.AddSubMenu(
-			LOCTEXT("ImportSubMenu", "Import"),
-			LOCTEXT("ImportSubMenu_ToolTip", "Import hierarchy to the current rig. This only imports non-existing node. For example, if there is hand_r, it won't import hand_r. If you want to reimport whole new hiearchy, delete all nodes, and use import hierarchy."),
-			FNewMenuDelegate::CreateSP(this, &SRigHierarchy::CreateImportMenu)
+		FToolMenuSection& ElementsSection = Menu->AddSection(TEXT("Elements"), LOCTEXT("ElementsHeader", "Elements"));
+		ElementsSection.AddSubMenu(TEXT("New"), LOCTEXT("New", "New"), LOCTEXT("New_ToolTip", "Create New Elements"),
+                                   FNewMenuDelegate::CreateStatic(&FLocalMenuBuilder::FillNewMenu, TreeView));
+		ElementsSection.AddMenuEntry(Commands.DeleteItem);
+		ElementsSection.AddMenuEntry(Commands.DuplicateItem);
+		ElementsSection.AddMenuEntry(Commands.RenameItem);
+		ElementsSection.AddMenuEntry(Commands.MirrorItem);
+
+		// dynamic section is used here so that the whole section can hide when a condition is not met
+		Menu->AddDynamicSection(InteractionSectionName,
+			FNewToolMenuDelegate::CreateLambda([InteractionSectionName, Commands, WeakThis = TWeakPtr<SRigHierarchy>(SharedThis(this)) ](UToolMenu* InMenu)
+			{
+				if (const TSharedPtr<SRigHierarchy> RigHierarchyPanel = WeakThis.Pin())
+				{
+					if (RigHierarchyPanel->IsSingleBoneSelected())
+					{
+						FToolMenuSection& InteractionSection = InMenu->AddSection(InteractionSectionName, LOCTEXT("InteractionHeader", "Interaction"));
+						InteractionSection.AddMenuEntry(Commands.ControlBoneTransform);
+					}
+				}
+			}),
+			FToolMenuInsert(TEXT("Elements"), EToolMenuInsertType::After)
 		);
-		MenuBuilder.AddSubMenu(
-			LOCTEXT("RefreshSubMenu", "Refresh"),
+		
+		FToolMenuSection& CopyPasteSection = Menu->AddSection(TEXT("Copy&Paste"), LOCTEXT("Copy&PasteHeader", "Copy & Paste"));
+		CopyPasteSection.AddMenuEntry(Commands.CopyItems);
+		CopyPasteSection.AddMenuEntry(Commands.PasteItems);
+		CopyPasteSection.AddMenuEntry(Commands.PasteLocalTransforms);
+		CopyPasteSection.AddMenuEntry(Commands.PasteGlobalTransforms);
+		
+		FToolMenuSection& TransformsSection = Menu->AddSection(TEXT("Transforms"), LOCTEXT("TransformsHeader", "Transforms"));
+		TransformsSection.AddMenuEntry(Commands.ResetTransform);
+		TransformsSection.AddMenuEntry(Commands.ResetAllTransforms);
+		TransformsSection.AddMenuEntry(Commands.SetInitialTransformFromCurrentTransform);
+		TransformsSection.AddMenuEntry(Commands.SetInitialTransformFromClosestBone);
+		TransformsSection.AddMenuEntry(Commands.SetGizmoTransformFromCurrent);
+		TransformsSection.AddMenuEntry(Commands.Unparent);
+
+		FToolMenuSection& AssetsSection = Menu->AddSection(TEXT("Assets"), LOCTEXT("AssetsHeader", "Assets"));
+		AssetsSection.AddSubMenu(TEXT("Import"), LOCTEXT("ImportSubMenu", "Import"),
+            LOCTEXT("ImportSubMenu_ToolTip", "Import hierarchy to the current rig. This only imports non-existing node. For example, if there is hand_r, it won't import hand_r. If you want to reimport whole new hiearchy, delete all nodes, and use import hierarchy."),
+            FNewMenuDelegate::CreateSP(this, &SRigHierarchy::CreateImportMenu)
+        );
+		
+		AssetsSection.AddSubMenu(TEXT("Refresh"), LOCTEXT("RefreshSubMenu", "Refresh"),
 			LOCTEXT("RefreshSubMenu_ToolTip", "Refresh the existing initial transform from the selected mesh. This only updates if the node is found."),
-			FNewMenuDelegate::CreateSP(this, &SRigHierarchy::CreateRefreshMenu)
-		);
-		MenuBuilder.EndSection();
+            FNewMenuDelegate::CreateSP(this, &SRigHierarchy::CreateImportMenu)
+        );
 	}
+
+	// individual entries in this menu can access members of this context, particularly useful for editor scripting
+	UControlRigContextMenuContext* ContextMenuContext = NewObject<UControlRigContextMenuContext>();
+	ContextMenuContext->ControlRigEditor = ControlRigEditor;
+	
+	FToolMenuContext MenuContext(CommandList);
+	MenuContext.AddObject(ContextMenuContext);
+	
+	UToolMenu* Menu = ToolMenus->GenerateMenu(MenuName, MenuContext);
+
+	return Menu;
 }
 
-TSharedPtr<FUICommandList> SRigHierarchy::GetContextMenuCommands()
+TSharedPtr<FUICommandList> SRigHierarchy::GetContextMenuCommands() const
 {
 	return CommandList;
 }
diff --git a/Engine/Plugins/Experimental/ControlRig/Source/ControlRigEditor/Private/Editor/SRigHierarchy.h b/Engine/Plugins/Experimental/ControlRig/Source/ControlRigEditor/Private/Editor/SRigHierarchy.h
--- a/Engine/Plugins/Experimental/ControlRig/Source/ControlRigEditor/Private/Editor/SRigHierarchy.h
+++ b/Engine/Plugins/Experimental/ControlRig/Source/ControlRigEditor/Private/Editor/SRigHierarchy.h
@@ -19,6 +19,7 @@ class UControlRig;
 struct FAssetData;
 class FMenuBuilder;
 class SRigHierarchyItem;
+class UToolMenu;
 
 DECLARE_DELEGATE_RetVal_TwoParams(FName, FOnRenameElement, const FRigElementKey& /*OldKey*/, const FString& /*NewName*/);
 DECLARE_DELEGATE_RetVal_ThreeParams(bool, FOnVerifyElementNameChanged, const FRigElementKey& /*OldKey*/, const FString& /*NewName*/, FText& /*OutErrorMessage*/);
@@ -227,7 +228,7 @@ private:
 	void OnSelectionChanged(TSharedPtr<FRigTreeElement> Selection, ESelectInfo::Type SelectInfo);
 
 	TSharedRef< SWidget > CreateFilterMenu();
-	TSharedPtr< SWidget > CreateContextMenu();
+	TSharedPtr< SWidget > CreateContextMenuWidget();
 	void OnItemClicked(TSharedPtr<FRigTreeElement> InItem);
 	void OnItemDoubleClicked(TSharedPtr<FRigTreeElement> InItem);
 	void OnSetExpansionRecursive(TSharedPtr<FRigTreeElement> InItem, bool bShouldBeExpanded);
@@ -241,8 +242,10 @@ private:
 	TOptional<EItemDropZone> OnCanAcceptDrop(const FDragDropEvent& DragDropEvent, EItemDropZone DropZone, TSharedPtr<FRigTreeElement> TargetItem);
 	FReply OnAcceptDrop(const FDragDropEvent& DragDropEvent, EItemDropZone DropZone, TSharedPtr<FRigTreeElement> TargetItem);
 
-	void FillContextMenu(class FMenuBuilder& MenuBuilder);
-	TSharedPtr<FUICommandList> GetContextMenuCommands();
+	UToolMenu* GetOrCreateContextMenu();
+	TSharedPtr<FUICommandList> GetContextMenuCommands() const;
+
+
 
 	/** Our owning control rig editor */
 	TWeakPtr<FControlRigEditor> ControlRigEditor;
