Control Rig: Added hierarchy listeners to propagate rig element property changes between hierarchies. This is part 1/2 for the jira below. part 2/2 will be delievered in the form of a python script

diff --git a/Engine/Plugins/Experimental/ControlRig/Source/ControlRig/Private/ControlRig.cpp b/Engine/Plugins/Experimental/ControlRig/Source/ControlRig/Private/ControlRig.cpp
--- a/Engine/Plugins/Experimental/ControlRig/Source/ControlRig/Private/ControlRig.cpp
+++ b/Engine/Plugins/Experimental/ControlRig/Source/ControlRig/Private/ControlRig.cpp
@@ -99,6 +99,16 @@ void UControlRig::BeginDestroy()
 		VM->ExecutionReachedExit().RemoveAll(this);
 	}
 
+#if WITH_EDITOR
+	if (!HasAnyFlags(RF_ClassDefaultObject))
+	{
+		if(UControlRig* CDO = GetClass()->GetDefaultObject<UControlRig>())
+		{
+			CDO->GetHierarchy()->UnregisterListeningHierarchy(GetHierarchy());
+		}
+	}
+#endif
+
 #if WITH_EDITORONLY_DATA
 	if (VMSnapshotBeforeExecution)
 	{
@@ -179,6 +189,9 @@ void UControlRig::InitializeFromCDO()
 	// copy CDO property you need to here
 	if (!HasAnyFlags(RF_ClassDefaultObject))
 	{
+		// similar to FControlRigBlueprintCompilerContext::CopyTermDefaultsToDefaultObject,
+		// where CDO is initialized from BP there,
+		// we initialize all other instances of Control Rig from the CDO here
 		UControlRig* CDO = GetClass()->GetDefaultObject<UControlRig>();
 
 		// copy hierarchy
@@ -188,6 +201,11 @@ void UControlRig::InitializeFromCDO()
 			GetHierarchy()->ResetPoseToInitial(ERigElementType::All);
 		}
 
+#if WITH_EDITOR
+		// current hierarchy should always mirror CDO's hierarchy whenever a change of interest happens
+		CDO->GetHierarchy()->RegisterListeningHierarchy(GetHierarchy());
+#endif
+
 		// notify clients that the hierarchy has changed
 		GetHierarchy()->Notify(ERigHierarchyNotification::HierarchyReset, nullptr);
 
diff --git a/Engine/Plugins/Experimental/ControlRig/Source/ControlRig/Private/Rigs/RigHierarchy.cpp b/Engine/Plugins/Experimental/ControlRig/Source/ControlRig/Private/Rigs/RigHierarchy.cpp
--- a/Engine/Plugins/Experimental/ControlRig/Source/ControlRig/Private/Rigs/RigHierarchy.cpp
+++ b/Engine/Plugins/Experimental/ControlRig/Source/ControlRig/Private/Rigs/RigHierarchy.cpp
@@ -268,6 +268,58 @@ void URigHierarchy::CopyHierarchy(URigHierarchy* InHierarchy)
 	UpdateAllCachedChildren();
 }
 
+#if WITH_EDITOR
+void URigHierarchy::RegisterListeningHierarchy(URigHierarchy* InHierarchy)
+{
+	if (ensure(InHierarchy))
+	{
+		bool bFoundListener = false;
+		for(int32 ListenerIndex = ListeningHierarchies.Num() - 1; ListenerIndex >= 0; ListenerIndex--)
+		{
+			const URigHierarchy::FRigHierarchyListener& Listener = ListeningHierarchies[ListenerIndex];
+			if(Listener.Hierarchy.IsValid())
+			{
+				if(Listener.Hierarchy.Get() == InHierarchy)
+				{
+					bFoundListener = true;
+					break;
+				}
+			}
+		}
+
+		if(!bFoundListener)
+		{
+			URigHierarchy::FRigHierarchyListener Listener;
+			Listener.Hierarchy = InHierarchy; 
+			ListeningHierarchies.Add(Listener);
+		}
+	}
+}
+
+void URigHierarchy::UnregisterListeningHierarchy(URigHierarchy* InHierarchy)
+{
+	if (ensure(InHierarchy))
+	{
+		for(int32 ListenerIndex = ListeningHierarchies.Num() - 1; ListenerIndex >= 0; ListenerIndex--)
+		{
+			const URigHierarchy::FRigHierarchyListener& Listener = ListeningHierarchies[ListenerIndex];
+			if(Listener.Hierarchy.IsValid())
+			{
+				if(Listener.Hierarchy.Get() == InHierarchy)
+				{
+					ListeningHierarchies.RemoveAt(ListenerIndex);
+				}
+			}
+		}
+	}
+}
+
+void URigHierarchy::ClearListeningHierarchy()
+{
+	ListeningHierarchies.Reset();
+}
+#endif
+
 void URigHierarchy::CopyPose(URigHierarchy* InHierarchy, bool bCurrent, bool bInitial)
 {
 	check(InHierarchy);
@@ -871,6 +923,30 @@ bool URigHierarchy::SetParentWeight(FRigBaseElement* InChild, int32 InParentInde
 			}
 
 			PropagateDirtyFlags(MultiParentElement, ERigTransformType::IsInitial(LocalType), bAffectChildren);
+			
+#if WITH_EDITOR
+			if (ensure(!bPropagatingChange))
+			{
+				TGuardValue<bool> bPropagatingChangeGuardValue(bPropagatingChange, true);
+			
+				for(FRigHierarchyListener& Listener : ListeningHierarchies)
+				{
+					if(!Listener.ShouldReactToChange(LocalType))
+					{
+						continue;
+					}
+
+					URigHierarchy* ListeningHierarchy = Listener.Hierarchy.Get();
+					if (ensure(ListeningHierarchy))
+					{
+						if(FRigBaseElement* ListeningElement = ListeningHierarchy->Find(InChild->GetKey()))
+						{
+							ListeningHierarchy->SetParentWeight(ListeningElement, InParentIndex, InWeight, bInitial, bAffectChildren);
+						}
+					}
+				}	
+			}
+#endif
 			return true;
 		}
 	}
@@ -1331,6 +1407,29 @@ void URigHierarchy::SetTransform(FRigTransformElement* InTransformElement, const
 			bAffectChildren,
 			bSetupUndo);
 	}
+
+	if (ensure(!bPropagatingChange))
+	{
+		TGuardValue<bool> bPropagatingChangeGuardValue(bPropagatingChange, true);
+			
+		for(FRigHierarchyListener& Listener : ListeningHierarchies)
+		{
+			if(!Listener.ShouldReactToChange(InTransformType))
+			{
+				continue;
+			}
+
+			URigHierarchy* ListeningHierarchy = Listener.Hierarchy.Get();
+			if (ensure(ListeningHierarchy))
+			{			
+				if(FRigTransformElement* ListeningElement = Cast<FRigTransformElement>(ListeningHierarchy->Find(InTransformElement->GetKey())))
+				{
+					// bSetupUndo = false such that all listening hierarchies performs undo at the same time the root hierachy undos
+					ListeningHierarchy->SetTransform(ListeningElement, InTransform, InTransformType, bAffectChildren, false, bForce);
+				}
+			}
+		}
+	}
 #endif
 }
 
@@ -1389,6 +1488,15 @@ void URigHierarchy::SetControlOffsetTransform(FRigControlElement* InControlEleme
 	InControlElement->Offset.MarkDirty(OpposedType);
 	InControlElement->Gizmo.MarkDirty(MakeGlobal(InTransformType));
 
+	if (ERigTransformType::IsInitial(InTransformType))
+	{
+		// control's offset transform is considered a special type of transform
+		// whenever its initial value is changed, we want to make sure the current is kept in sync
+		// such that the viewport can reflect this change
+		SetControlOffsetTransform(InControlElement, InTransform, ERigTransformType::MakeCurrent(InTransformType), bAffectChildren, false, bForce);
+	}
+	
+
 #if WITH_EDITOR
 	if(bSetupUndo || IsTracingChanges())
 	{
@@ -1401,6 +1509,29 @@ void URigHierarchy::SetControlOffsetTransform(FRigControlElement* InControlEleme
             bAffectChildren,
             bSetupUndo);
 	}
+
+	if (ensure(!bPropagatingChange))
+	{
+		TGuardValue<bool> bPropagatingChangeGuardValue(bPropagatingChange, true);
+			
+		for(FRigHierarchyListener& Listener : ListeningHierarchies)
+		{
+			if(!Listener.ShouldReactToChange(InTransformType))
+			{
+				continue;
+			}
+
+			URigHierarchy* ListeningHierarchy = Listener.Hierarchy.Get();
+			if (ensure(ListeningHierarchy))
+			{	
+				if(FRigControlElement* ListeningElement = Cast<FRigControlElement>(ListeningHierarchy->Find(InControlElement->GetKey())))
+				{
+					// bSetupUndo = false such that all listening hierarchies performs undo at the same time the root hierachy undos
+					ListeningHierarchy->SetControlOffsetTransform(ListeningElement, InTransform, InTransformType, bAffectChildren, false, bForce);
+				}
+			}
+		}
+	}
 #endif
 }
 
@@ -1453,6 +1584,14 @@ void URigHierarchy::SetControlGizmoTransform(FRigControlElement* InControlElemen
 	InControlElement->Gizmo.Set(InTransformType, InTransform);
 	InControlElement->Gizmo.MarkDirty(OpposedType);
 
+	if (IsInitial(InTransformType))
+	{
+		// control's gizmo transform, similar to offset transform, is considered a special type of transform
+		// whenever its initial value is changed, we want to make sure the current is kept in sync
+		// such that the viewport can reflect this change
+		SetControlGizmoTransform(InControlElement, InTransform, ERigTransformType::MakeCurrent(InTransformType), false, bForce);
+	}
+	
 #if WITH_EDITOR
 	if(bSetupUndo || IsTracingChanges())
 	{
@@ -1471,6 +1610,31 @@ void URigHierarchy::SetControlGizmoTransform(FRigControlElement* InControlElemen
 	{
 		Notify(ERigHierarchyNotification::ControlGizmoTransformChanged, InControlElement);
 	}
+
+#if WITH_EDITOR
+	if (ensure(!bPropagatingChange))
+	{
+		TGuardValue<bool> bPropagatingChangeGuardValue(bPropagatingChange, true);
+			
+		for(FRigHierarchyListener& Listener : ListeningHierarchies)
+		{
+			if(!Listener.ShouldReactToChange(InTransformType))
+			{
+				continue;
+			}
+
+			URigHierarchy* ListeningHierarchy = Listener.Hierarchy.Get();
+			if (ensure(ListeningHierarchy))
+			{	
+				if(FRigControlElement* ListeningElement = Cast<FRigControlElement>(ListeningHierarchy->Find(InControlElement->GetKey())))
+				{
+					// bSetupUndo = false such that all listening hierarchies performs undo at the same time the root hierachy undos
+					ListeningHierarchy->SetControlGizmoTransform(ListeningElement, InTransform, InTransformType, false, bForce);
+				}
+			}
+		}
+	}
+#endif
 }
 
 FTransform URigHierarchy::GetParentTransform(FRigBaseElement* InElement, const ERigTransformType::Type InTransformType) const
@@ -1747,17 +1911,40 @@ void URigHierarchy::SetControlValue(FRigControlElement* InControlElement, const
 				break;
 			}
 			case ERigControlValueType::Minimum:
-			{
-				InControlElement->Settings.MinimumValue = InValue;
-				InControlElement->Settings.ApplyLimits(InControlElement->Settings.MinimumValue);
-				Notify(ERigHierarchyNotification::ControlSettingChanged, InControlElement);
-				break;
-			}
 			case ERigControlValueType::Maximum:
 			{
-				InControlElement->Settings.MaximumValue = InValue;
-				InControlElement->Settings.ApplyLimits(InControlElement->Settings.MaximumValue);
+				if(InValueType == ERigControlValueType::Minimum)
+				{
+					InControlElement->Settings.MinimumValue = InValue;
+					InControlElement->Settings.ApplyLimits(InControlElement->Settings.MinimumValue);
+				}
+				else
+				{
+					InControlElement->Settings.MaximumValue = InValue;
+					InControlElement->Settings.ApplyLimits(InControlElement->Settings.MaximumValue);
+				}
+				
 				Notify(ERigHierarchyNotification::ControlSettingChanged, InControlElement);
+
+#if WITH_EDITOR
+				if (ensure(!bPropagatingChange))
+				{
+					TGuardValue<bool> bPropagatingChangeGuardValue(bPropagatingChange, true);
+			
+					for(FRigHierarchyListener& Listener : ListeningHierarchies)
+					{
+						URigHierarchy* ListeningHierarchy = Listener.Hierarchy.Get();
+					
+						if (ensure(ListeningHierarchy))
+						{
+							if(FRigControlElement* ListeningElement = Cast<FRigControlElement>(ListeningHierarchy->Find(InControlElement->GetKey())))
+							{
+								ListeningHierarchy->SetControlValue(ListeningElement, InValue, InValueType, false, bForce);
+							}
+						}
+					}
+				}
+#endif
 				break;
 			}
 		}	
@@ -1773,6 +1960,25 @@ void URigHierarchy::SetControlVisibility(FRigControlElement* InControlElement, b
 
 	InControlElement->Settings.bGizmoVisible = bVisibility;
 	Notify(ERigHierarchyNotification::ControlVisibilityChanged, InControlElement);
+
+#if WITH_EDITOR
+	if (ensure(!bPropagatingChange))
+	{
+		TGuardValue<bool> bPropagatingChangeGuardValue(bPropagatingChange, true);
+			
+		for(FRigHierarchyListener& Listener : ListeningHierarchies)
+		{
+			URigHierarchy* ListeningHierarchy = Listener.Hierarchy.Get();
+			if (ensure(ListeningHierarchy))
+			{
+				if(FRigControlElement* ListeningElement = Cast<FRigControlElement>(ListeningHierarchy->Find(InControlElement->GetKey())))
+				{
+					ListeningHierarchy->SetControlVisibility(ListeningElement, bVisibility);
+				}
+			}
+		}
+	}
+#endif
 }
 
 float URigHierarchy::GetCurveValue(FRigCurveElement* InCurveElement) const
@@ -1804,6 +2010,29 @@ void URigHierarchy::SetCurveValue(FRigCurveElement* InCurveElement, float InValu
 	{
 		PushCurveToStack(InCurveElement->GetKey(), PreviousValue, InCurveElement->Value, bSetupUndo);
 	}
+
+	if (ensure(!bPropagatingChange))
+	{
+		TGuardValue<bool> bPropagatingChangeGuardValue(bPropagatingChange, true);
+			
+		for(FRigHierarchyListener& Listener : ListeningHierarchies)
+		{
+			if(!Listener.Hierarchy.IsValid())
+			{
+				continue;
+			}
+
+			URigHierarchy* ListeningHierarchy = Listener.Hierarchy.Get();
+			if (ensure(ListeningHierarchy))
+			{
+				if(FRigCurveElement* ListeningElement = Cast<FRigCurveElement>(ListeningHierarchy->Find(InCurveElement->GetKey())))
+				{
+					// bSetupUndo = false such that all listening hierarchies performs undo at the same time the root hierachy undos
+					ListeningHierarchy->SetCurveValue(ListeningElement, InValue, false, bForce);
+				}
+			}
+		}
+	}
 #endif
 }
 
diff --git a/Engine/Plugins/Experimental/ControlRig/Source/ControlRig/Public/Rigs/RigHierarchy.h b/Engine/Plugins/Experimental/ControlRig/Source/ControlRig/Public/Rigs/RigHierarchy.h
--- a/Engine/Plugins/Experimental/ControlRig/Source/ControlRig/Public/Rigs/RigHierarchy.h
+++ b/Engine/Plugins/Experimental/ControlRig/Source/ControlRig/Public/Rigs/RigHierarchy.h
@@ -98,6 +98,7 @@ public:
 	, bSuspendNotifications(false)
 	, ResetPoseHash(INDEX_NONE)
 #if WITH_EDITOR
+	, bPropagatingChange(false)
 	, TraceFramesLeft(0)
 	, TraceFramesCaptured(0)
 #endif
@@ -127,6 +128,22 @@ public:
 	UFUNCTION(BlueprintCallable, Category = URigHierarchy)
 	void CopyHierarchy(URigHierarchy* InHierarchy);
 
+#if WITH_EDITOR
+	/**
+	* Add dependent hierarchies that listens to changes made to this hierarchy
+	* Note: By default, only changes to the initial states of this hierarchy is mirrored to the listening hierarchies
+	*/	
+	void RegisterListeningHierarchy(URigHierarchy* InHierarchy);
+	
+	/**
+	* Remove dependent hierarchies that listens to changes made to this hierarchy
+	*/	
+	void UnregisterListeningHierarchy(URigHierarchy* InHierarchy);
+	
+	void ClearListeningHierarchy();
+#endif
+
+public:
 	/**
 	 * Copies the contents of a hierarchy onto this one
 	 */
@@ -701,6 +718,7 @@ public:
 	 * @InTypeFilter The types to retrieve the selection for
 	 * @return An array of the currently selected elements
 	 */
+	UFUNCTION(BlueprintCallable, Category = URigHierarchy)
 	TArray<FRigElementKey> GetSelectedKeys(ERigElementType InTypeFilter = ERigElementType::All) const;
 
 	/**
@@ -735,6 +753,7 @@ public:
 	 * @param InKeys The keys to sort
 	 * @return The sorted keys
 	 */
+	UFUNCTION(BlueprintCallable, Category = URigHierarchy)
 	FORCEINLINE TArray<FRigElementKey> SortKeys(const TArray<FRigElementKey>& InKeys) const
 	{
 		TArray<FRigElementKey> Result;
@@ -2011,7 +2030,7 @@ public:
 	 * Sets the offset transform for a given control element
 	 * @param InControlElement The element to set the transform for
 	 * @param InTransform The offset transform to set
-	 * @param InTransformType The type of transform to set
+	 * @param InTransformType The type of transform to set. Note: for offset transform, setting the initial value also updates the current value
 	 * @param bAffectChildren If set to false children will not move (maintain global).
 	 * @param bSetupUndo If true the transform stack will be setup for undo / redo
 	 * @param bForce Set the transform even if it is the same as the previously set one
@@ -2030,7 +2049,7 @@ public:
 	 * Sets the gizmo transform for a given control element
 	 * @param InControlElement The element to set the transform for
 	 * @param InTransform The gizmo transform to set
-	 * @param InTransformType The type of transform to set
+	 * @param InTransformType The type of transform to set. Note: for gizmo transform, setting the initial value also updates the current value
 	 * @param bSetupUndo If true the transform stack will be setup for undo / redo
 	 * @param bForce Set the transform even if it is the same as the previously set one
 	 */
@@ -2312,6 +2331,47 @@ private:
 	int32 ResetPoseHash;
 	TArray<bool> ResetPoseHasFilteredChildren;
 
+#if WITH_EDITOR
+
+	// this is mainly used for propagating changes between hierarchies in the direction of blueprint -> CDO -> other instances
+	struct FRigHierarchyListener
+	{
+		FRigHierarchyListener()
+			: Hierarchy(nullptr)
+			, bShouldReactToInitialChanges(true)
+			// default to false since changing the current values in a BP/CDO is not meaningful
+			, bShouldReactToCurrentChanges(false) 
+		{}
+
+		bool ShouldReactToChange(ERigTransformType::Type InTransformType) const
+		{
+			if(Hierarchy.IsValid())
+			{
+				if(ERigTransformType::IsInitial(InTransformType))
+				{
+					return bShouldReactToInitialChanges;
+				}
+
+				if(ERigTransformType::IsCurrent(InTransformType))
+				{
+					return bShouldReactToCurrentChanges;
+				}
+			}
+			return false;
+		}
+		
+		TWeakObjectPtr<URigHierarchy> Hierarchy;
+		bool bShouldReactToInitialChanges;
+		bool bShouldReactToCurrentChanges;
+	};
+	
+	TArray<FRigHierarchyListener> ListeningHierarchies;
+
+	// a bool to guard against circular dependencies among listening hierarchies
+	bool bPropagatingChange;
+	
+#endif
+
 #if WITH_EDITOR
 
 protected:
@@ -2321,7 +2381,7 @@ protected:
 	TMap<FName, FRigPose> TracePoses;
 
 #endif
-	
+
 	friend class URigHierarchyController;
 	friend class UControlRig;
 	friend class FControlRigEditor;
diff --git a/Engine/Plugins/Experimental/ControlRig/Source/ControlRigDeveloper/Private/ControlRigBlueprint.cpp b/Engine/Plugins/Experimental/ControlRig/Source/ControlRigDeveloper/Private/ControlRigBlueprint.cpp
--- a/Engine/Plugins/Experimental/ControlRig/Source/ControlRigDeveloper/Private/ControlRigBlueprint.cpp
+++ b/Engine/Plugins/Experimental/ControlRig/Source/ControlRigDeveloper/Private/ControlRigBlueprint.cpp
@@ -2915,6 +2915,8 @@ void UControlRigBlueprint::PropagatePoseFromBPToInstances()
 	{
 		if (UControlRig* DefaultObject = Cast<UControlRig>(MyControlRigClass->GetDefaultObject(false)))
 		{
+			DefaultObject->GetHierarchy()->CopyPose(Hierarchy, true, true);
+
 			TArray<UObject*> ArchetypeInstances;
 			DefaultObject->GetArchetypeInstances(ArchetypeInstances);
 			for (UObject* ArchetypeInstance : ArchetypeInstances)
diff --git a/Engine/Plugins/Experimental/ControlRig/Source/ControlRigDeveloper/Private/ControlRigBlueprintCompiler.cpp b/Engine/Plugins/Experimental/ControlRig/Source/ControlRigDeveloper/Private/ControlRigBlueprintCompiler.cpp
--- a/Engine/Plugins/Experimental/ControlRig/Source/ControlRigDeveloper/Private/ControlRigBlueprintCompiler.cpp
+++ b/Engine/Plugins/Experimental/ControlRig/Source/ControlRigDeveloper/Private/ControlRigBlueprintCompiler.cpp
@@ -90,6 +90,9 @@ void FControlRigBlueprintCompilerContext::CopyTermDefaultsToDefaultObject(UObjec
 	UControlRigBlueprint* ControlRigBlueprint = Cast<UControlRigBlueprint>(Blueprint);
 	if (ControlRigBlueprint)
 	{
+		// here, CDO is initialized from BP,
+		// and in UControlRig::InitializeFromCDO,
+		// other Control Rig Instances are then initialized From the CDO 
 		UControlRig* ControlRig = CastChecked<UControlRig>(DefaultObject);
 		
 		// copy hierarchy
@@ -97,6 +100,14 @@ void FControlRigBlueprintCompilerContext::CopyTermDefaultsToDefaultObject(UObjec
 			TGuardValue<bool> Guard(ControlRig->GetHierarchy()->GetSuspendNotificationsFlag(), false);
 			ControlRig->GetHierarchy()->CopyHierarchy(ControlRigBlueprint->Hierarchy);
 			ControlRig->GetHierarchy()->ResetPoseToInitial(ERigElementType::All);
+
+#if WITH_EDITOR
+			// link CDO hierarchy to BP's hierarchy
+			// so that whenever BP's hierarchy is modified, CDO's hierarchy is kept in sync
+			// Other instances of Control Rig can make use of CDO to reset to the initial state
+			ControlRigBlueprint->Hierarchy->ClearListeningHierarchy();
+			ControlRigBlueprint->Hierarchy->RegisterListeningHierarchy(ControlRig->GetHierarchy());
+#endif
 		}
 
 		// notify clients that the hierarchy has changed
