Control Rig: Add pre-initialize event to Control Rig Component

diff --git a/Engine/Plugins/Experimental/ControlRig/Source/ControlRig/Private/ControlRigComponent.cpp b/Engine/Plugins/Experimental/ControlRig/Source/ControlRig/Private/ControlRigComponent.cpp
--- a/Engine/Plugins/Experimental/ControlRig/Source/ControlRig/Private/ControlRigComponent.cpp
+++ b/Engine/Plugins/Experimental/ControlRig/Source/ControlRig/Private/ControlRigComponent.cpp
@@ -54,6 +54,7 @@ UControlRigComponent::UControlRigComponent(const FObjectInitializer& ObjectIniti
 	bUpdateInEditor = true;
 	bDrawBones = true;
 	bShowDebugDrawing = false;
+	bIsInsideInitializeBracket = false;
 }
 
 #if WITH_EDITOR
@@ -252,6 +253,11 @@ float UControlRigComponent::GetAbsoluteTime() const
 	return 0.f;
 }
 
+void UControlRigComponent::OnPreInitialize_Implementation(UControlRigComponent* Component)
+{
+	OnPreInitializeDelegate.Broadcast(Component);
+}
+
 void UControlRigComponent::OnPostInitialize_Implementation(UControlRigComponent* Component)
 {
 	ValidateMappingData();
@@ -282,6 +288,25 @@ void UControlRigComponent::OnPostUpdate_Implementation(UControlRigComponent* Com
 
 void UControlRigComponent::Initialize()
 {
+	if(bIsInsideInitializeBracket)
+	{
+		return;
+	}
+
+	TGuardValue<bool> InitializeBracket(bIsInsideInitializeBracket, true);
+	
+#if WITH_EDITOR
+	if (bUpdateInEditor)
+	{
+		FEditorScriptExecutionGuard AllowScripts;
+		OnPreInitialize(this);
+	}
+	else
+#endif
+	{
+		OnPreInitialize(this);
+	}
+	
 	if(UControlRig* CR = SetupControlRigIfRequired())
 	{
 		if (CR->IsInitializing())
@@ -453,7 +478,7 @@ void UControlRigComponent::AddMappedElements(TArray<FControlRigComponentMappedEl
 
 void UControlRigComponent::AddMappedComponents(TArray<FControlRigComponentMappedComponent> Components)
 {
-	if (!EnsureCalledOutsideOfBracket(TEXT("AddMappedComponents")))
+	if (!EnsureCalledOutsideOfBracket(TEXT("AddMapprepedComponents")))
 	{
 		return;
 	}
diff --git a/Engine/Plugins/Experimental/ControlRig/Source/ControlRig/Public/ControlRigComponent.h b/Engine/Plugins/Experimental/ControlRig/Source/ControlRig/Public/ControlRigComponent.h
--- a/Engine/Plugins/Experimental/ControlRig/Source/ControlRig/Public/ControlRigComponent.h
+++ b/Engine/Plugins/Experimental/ControlRig/Source/ControlRig/Public/ControlRigComponent.h
@@ -182,6 +182,10 @@ public:
 	UPROPERTY(EditAnywhere, Category = ControlRig, meta = (DisallowCreateNew))
 	TSubclassOf<UControlRig> ControlRigClass;
 
+	/** Event fired just before this component's ControlRig is initialized */
+	UPROPERTY(BlueprintAssignable, Category = "ControlRig", meta = (DisplayName = "On Pre Initialize"))
+	FControlRigComponentDelegate OnPreInitializeDelegate;
+
 	/** Event fired after this component's ControlRig is initialized */
 	UPROPERTY(BlueprintAssignable, Category = "ControlRig", meta = (DisplayName = "On Post Initialize"))
 	FControlRigComponentDelegate OnPostInitializeDelegate;
@@ -233,6 +237,9 @@ public:
 	UFUNCTION(BlueprintPure, Category = "ControlRig")
 	float GetAbsoluteTime() const;
 
+	UFUNCTION(BlueprintNativeEvent, Category = "ControlRig", CallInEditor)
+	void OnPreInitialize(UControlRigComponent* Component);
+
 	UFUNCTION(BlueprintNativeEvent, Category = "ControlRig", CallInEditor)
 	void OnPostInitialize(UControlRigComponent* Component);
 
@@ -627,6 +634,7 @@ private:
 	TMap<USkeletalMeshComponent*, FCachedSkeletalMeshComponentSettings> CachedSkeletalMeshComponentSettings;
 
 	FControlRigComponentEvent ControlRigCreatedEvent;
+	bool bIsInsideInitializeBracket;
 
 	friend class FControlRigSceneProxy;
 };
