Control Rig: Ensure to create rig graph independent from "Spawn Default Blueprint Nodes" setting

diff --git a/Engine/Plugins/Experimental/ControlRig/Source/ControlRigEditor/Private/ControlRigBlueprintFactory.cpp b/Engine/Plugins/Experimental/ControlRig/Source/ControlRigEditor/Private/ControlRigBlueprintFactory.cpp
--- a/Engine/Plugins/Experimental/ControlRig/Source/ControlRigEditor/Private/ControlRigBlueprintFactory.cpp
+++ b/Engine/Plugins/Experimental/ControlRig/Source/ControlRigEditor/Private/ControlRigBlueprintFactory.cpp
@@ -21,6 +21,9 @@
 #include "ControlRigBlueprint.h"
 #include "ControlRigBlueprintActions.h"
 #include "ControlRigBlueprintGeneratedClass.h"
+#include "Graph/ControlRigGraphSchema.h"
+#include "Graph/ControlRigGraph.h"
+#include "Kismet2/BlueprintEditorUtils.h"
 
 #define LOCTEXT_NAMESPACE "ControlRigBlueprintFactory"
 
@@ -284,6 +287,7 @@ UObject* UControlRigBlueprintFactory::FactoryCreateNew(UClass* Class, UObject* I
 	else
 	{
 		UControlRigBlueprint* ControlRigBlueprint = CastChecked<UControlRigBlueprint>(FKismetEditorUtilities::CreateBlueprint(ParentClass, InParent, Name, BPTYPE_Normal, UControlRigBlueprint::StaticClass(), UControlRigBlueprintGeneratedClass::StaticClass(), CallingContext));
+		CreateRigGraphIfRequired(ControlRigBlueprint);
 		return ControlRigBlueprint;
 	}
 }
@@ -303,4 +307,23 @@ UControlRigBlueprint* UControlRigBlueprintFactory::CreateControlRigFromSkeletalM
 	return FControlRigBlueprintActions::CreateControlRigFromSkeletalMeshOrSkeleton(InSelectedObject);
 }
 
+void UControlRigBlueprintFactory::CreateRigGraphIfRequired(UControlRigBlueprint* InBlueprint)
+{
+	for(UEdGraph* EdGraph : InBlueprint->UbergraphPages)
+	{
+		if(EdGraph->IsA<UControlRigGraph>())
+		{
+			return;
+		}
+	}
+	
+	// add an initial graph for us to work in
+	const UControlRigGraphSchema* ControlRigGraphSchema = GetDefault<UControlRigGraphSchema>();
+	UEdGraph* ControlRigGraph = FBlueprintEditorUtils::CreateNewGraph(InBlueprint, ControlRigGraphSchema->GraphName_ControlRig, UControlRigGraph::StaticClass(), UControlRigGraphSchema::StaticClass());
+	ControlRigGraph->bAllowDeletion = false;
+	FBlueprintEditorUtils::AddUbergraphPage(InBlueprint, ControlRigGraph);
+	InBlueprint->LastEditedDocuments.AddUnique(ControlRigGraph);
+	InBlueprint->PostLoad();
+}
+
 #undef LOCTEXT_NAMESPACE
diff --git a/Engine/Plugins/Experimental/ControlRig/Source/ControlRigEditor/Private/ControlRigBlueprintFactory.h b/Engine/Plugins/Experimental/ControlRig/Source/ControlRigEditor/Private/ControlRigBlueprintFactory.h
--- a/Engine/Plugins/Experimental/ControlRig/Source/ControlRigEditor/Private/ControlRigBlueprintFactory.h
+++ b/Engine/Plugins/Experimental/ControlRig/Source/ControlRigEditor/Private/ControlRigBlueprintFactory.h
@@ -42,5 +42,11 @@ public:
 	 */
 	UFUNCTION(BlueprintCallable, Category = "Control Rig")
 	static UControlRigBlueprint* CreateControlRigFromSkeletalMeshOrSkeleton(UObject* InSelectedObject);
+
+	/**
+	 * Creates the riggraph if necessary
+	 * @param InBlueprint The blueprint to setup up
+	 */
+	static void CreateRigGraphIfRequired(UControlRigBlueprint* InBlueprint);
 };
 
diff --git a/Engine/Plugins/Experimental/ControlRig/Source/ControlRigEditor/Private/ControlRigEditorModule.cpp b/Engine/Plugins/Experimental/ControlRig/Source/ControlRigEditor/Private/ControlRigEditorModule.cpp
--- a/Engine/Plugins/Experimental/ControlRig/Source/ControlRigEditor/Private/ControlRigEditorModule.cpp
+++ b/Engine/Plugins/Experimental/ControlRig/Source/ControlRigEditor/Private/ControlRigEditorModule.cpp
@@ -110,6 +110,7 @@
 #include "ControlRig/Private/Units/Execution/RigUnit_InverseExecution.h"
 #include "Graph/SControlRigGraphPinVariableBinding.h"
 #include "AssetTypeActions_ControlRigPose.h"
+#include "ControlRigBlueprintFactory.h"
 #include "EditMode/SControlRigBaseListWidget.h"
 #include "EditMode/SControlRigTweenWidget.h"
 #include "EditMode/SControlRigSnapper.h"
@@ -1016,14 +1017,7 @@ void FControlRigEditorModule::ExtendAnimSequenceMenu()
 
 void FControlRigEditorModule::HandleNewBlueprintCreated(UBlueprint* InBlueprint)
 {
-	// add an initial graph for us to work in
-	const UControlRigGraphSchema* ControlRigGraphSchema = GetDefault<UControlRigGraphSchema>();
-
-	UEdGraph* ControlRigGraph = FBlueprintEditorUtils::CreateNewGraph(InBlueprint, ControlRigGraphSchema->GraphName_ControlRig, UControlRigGraph::StaticClass(), UControlRigGraphSchema::StaticClass());
-	ControlRigGraph->bAllowDeletion = false;
-	FBlueprintEditorUtils::AddUbergraphPage(InBlueprint, ControlRigGraph);
-	InBlueprint->LastEditedDocuments.AddUnique(ControlRigGraph);
-	InBlueprint->PostLoad();
+	UControlRigBlueprintFactory::CreateRigGraphIfRequired(Cast<UControlRigBlueprint>(InBlueprint));
 }
 
 TSharedRef<IControlRigEditor> FControlRigEditorModule::CreateControlRigEditor(const EToolkitMode::Type Mode, const TSharedPtr< IToolkitHost >& InitToolkitHost, class UControlRigBlueprint* InBlueprint)
