Control Rig: Fix interaction with gizmo, offset transform with and without setup mode in details panel

diff --git a/Engine/Plugins/Experimental/ControlRig/Source/ControlRig/Private/Rigs/RigHierarchy.cpp b/Engine/Plugins/Experimental/ControlRig/Source/ControlRig/Private/Rigs/RigHierarchy.cpp
--- a/Engine/Plugins/Experimental/ControlRig/Source/ControlRig/Private/Rigs/RigHierarchy.cpp
+++ b/Engine/Plugins/Experimental/ControlRig/Source/ControlRig/Private/Rigs/RigHierarchy.cpp
@@ -931,7 +931,7 @@ bool URigHierarchy::SetParentWeight(FRigBaseElement* InChild, int32 InParentInde
 			
 				for(FRigHierarchyListener& Listener : ListeningHierarchies)
 				{
-					if(!Listener.ShouldReactToChange(LocalType))
+					if(!bForcePropagation && !Listener.ShouldReactToChange(LocalType))
 					{
 						continue;
 					}
@@ -1414,7 +1414,7 @@ void URigHierarchy::SetTransform(FRigTransformElement* InTransformElement, const
 			
 		for(FRigHierarchyListener& Listener : ListeningHierarchies)
 		{
-			if(!Listener.ShouldReactToChange(InTransformType))
+			if(!bForcePropagation && !Listener.ShouldReactToChange(InTransformType))
 			{
 				continue;
 			}
@@ -1516,11 +1516,6 @@ void URigHierarchy::SetControlOffsetTransform(FRigControlElement* InControlEleme
 			
 		for(FRigHierarchyListener& Listener : ListeningHierarchies)
 		{
-			if(!Listener.ShouldReactToChange(InTransformType))
-			{
-				continue;
-			}
-
 			URigHierarchy* ListeningHierarchy = Listener.Hierarchy.Get();
 			if (ensure(ListeningHierarchy))
 			{	
@@ -1618,18 +1613,50 @@ void URigHierarchy::SetControlGizmoTransform(FRigControlElement* InControlElemen
 			
 		for(FRigHierarchyListener& Listener : ListeningHierarchies)
 		{
-			if(!Listener.ShouldReactToChange(InTransformType))
-			{
-				continue;
+			URigHierarchy* ListeningHierarchy = Listener.Hierarchy.Get();
+			if (ensure(ListeningHierarchy))
+			{	
+				if(FRigControlElement* ListeningElement = Cast<FRigControlElement>(ListeningHierarchy->Find(InControlElement->GetKey())))
+				{
+					// bSetupUndo = false such that all listening hierarchies performs undo at the same time the root hierachy undos
+					ListeningHierarchy->SetControlGizmoTransform(ListeningElement, InTransform, InTransformType, false, bForce);
+				}
 			}
+		}
+	}
+#endif
+}
 
+void URigHierarchy::SetControlSettings(FRigControlElement* InControlElement, FRigControlSettings InSettings, bool bSetupUndo, bool bForce)
+{
+	if(InControlElement == nullptr)
+	{
+		return;
+	}
+
+	const FRigControlSettings PreviousSettings = InControlElement->Settings;
+	if(!bForce && PreviousSettings == InSettings)
+	{
+		return;
+	}
+
+	InControlElement->Settings = InSettings;
+	Notify(ERigHierarchyNotification::ControlSettingChanged, InControlElement);
+	
+#if WITH_EDITOR
+	if (ensure(!bPropagatingChange))
+	{
+		TGuardValue<bool> bPropagatingChangeGuardValue(bPropagatingChange, true);
+			
+		for(FRigHierarchyListener& Listener : ListeningHierarchies)
+		{
 			URigHierarchy* ListeningHierarchy = Listener.Hierarchy.Get();
 			if (ensure(ListeningHierarchy))
 			{	
 				if(FRigControlElement* ListeningElement = Cast<FRigControlElement>(ListeningHierarchy->Find(InControlElement->GetKey())))
 				{
 					// bSetupUndo = false such that all listening hierarchies performs undo at the same time the root hierachy undos
-					ListeningHierarchy->SetControlGizmoTransform(ListeningElement, InTransform, InTransformType, false, bForce);
+					ListeningHierarchy->SetControlSettings(ListeningElement, InSettings, false, bForce);
 				}
 			}
 		}
diff --git a/Engine/Plugins/Experimental/ControlRig/Source/ControlRig/Private/Rigs/RigHierarchyElements.cpp b/Engine/Plugins/Experimental/ControlRig/Source/ControlRig/Private/Rigs/RigHierarchyElements.cpp
--- a/Engine/Plugins/Experimental/ControlRig/Source/ControlRig/Private/Rigs/RigHierarchyElements.cpp
+++ b/Engine/Plugins/Experimental/ControlRig/Source/ControlRig/Private/Rigs/RigHierarchyElements.cpp
@@ -448,6 +448,86 @@ void FRigControlSettings::Load(FArchive& Ar)
 // FRigControlElement
 ////////////////////////////////////////////////////////////////////////////////
 
+bool FRigControlSettings::operator==(const FRigControlSettings& InOther) const
+{
+	if(ControlType != InOther.ControlType)
+	{
+		return false;
+	}
+	if(DisplayName != InOther.DisplayName)
+	{
+		return false;
+	}
+	if(PrimaryAxis != InOther.PrimaryAxis)
+	{
+		return false;
+	}
+	if(bIsCurve != InOther.bIsCurve)
+	{
+		return false;
+	}
+	if(bAnimatable != InOther.bAnimatable)
+	{
+		return false;
+	}
+	if(bLimitTranslation != InOther.bLimitTranslation)
+	{
+		return false;
+	}
+	if(bLimitRotation != InOther.bLimitRotation)
+	{
+		return false;
+	}
+	if(bLimitScale != InOther.bLimitScale)
+	{
+		return false;
+	}
+	if(bDrawLimits != InOther.bDrawLimits)
+	{
+		return false;
+	}
+	if(bGizmoEnabled != InOther.bGizmoEnabled)
+	{
+		return false;
+	}
+	if(bGizmoVisible != InOther.bGizmoVisible)
+	{
+		return false;
+	}
+	if(GizmoName != InOther.GizmoName)
+	{
+		return false;
+	}
+	if(bIsTransientControl != InOther.bIsTransientControl)
+	{
+		return false;
+	}
+	if( ControlEnum != InOther. ControlEnum)
+	{
+		return false;
+	}
+	if(!GizmoColor.Equals(InOther.GizmoColor, 0.001))
+	{
+		return false;
+	}
+
+	const FTransform MinimumTransform = MinimumValue.GetAsTransform(ControlType, PrimaryAxis);
+	const FTransform OtherMinimumTransform = InOther.MinimumValue.GetAsTransform(ControlType, PrimaryAxis);
+	if(!MinimumTransform.Equals(OtherMinimumTransform, 0.001))
+	{
+		return false;
+	}
+
+	const FTransform MaximumTransform = MaximumValue.GetAsTransform(ControlType, PrimaryAxis);
+	const FTransform OtherMaximumTransform = InOther.MaximumValue.GetAsTransform(ControlType, PrimaryAxis);
+	if(!MaximumTransform.Equals(OtherMaximumTransform, 0.001))
+	{
+		return false;
+	}
+
+	return true;
+}
+
 void FRigControlElement::Save(FArchive& Ar, URigHierarchy* Hierarchy, ESerializationPhase SerializationPhase)
 {
 	Super::Save(Ar, Hierarchy, SerializationPhase);
diff --git a/Engine/Plugins/Experimental/ControlRig/Source/ControlRig/Public/Rigs/RigHierarchy.h b/Engine/Plugins/Experimental/ControlRig/Source/ControlRig/Public/Rigs/RigHierarchy.h
--- a/Engine/Plugins/Experimental/ControlRig/Source/ControlRig/Public/Rigs/RigHierarchy.h
+++ b/Engine/Plugins/Experimental/ControlRig/Source/ControlRig/Public/Rigs/RigHierarchy.h
@@ -99,6 +99,7 @@ public:
 	, ResetPoseHash(INDEX_NONE)
 #if WITH_EDITOR
 	, bPropagatingChange(false)
+	, bForcePropagation(false)
 	, TraceFramesLeft(0)
 	, TraceFramesCaptured(0)
 #endif
@@ -1413,6 +1414,36 @@ public:
 		}
 	}
 
+	/**
+	 * Sets the control settings for a given control element by key
+	 * @param InKey The key of the control element to set the settings for
+	 * @param InSettings The new control settings value to set
+	 * @param bSetupUndo If true the transform stack will be setup for undo / redo
+	 */
+	UFUNCTION(BlueprintCallable, Category = URigHierarchy)
+	FORCEINLINE_DEBUGGABLE void SetControlSettings(FRigElementKey InKey, FRigControlSettings InSettings, bool bSetupUndo = false, bool bForce = false)
+	{
+		return SetControlSettingsByIndex(GetIndex(InKey), InSettings, bSetupUndo, bForce);
+	}
+
+	/**
+	 * Sets the control settings for a given control element by index
+	 * @param InElementIndex The index of the control element to set the settings for
+	 * @param InSettings The new control settings value to set
+	 * @param bSetupUndo If true the transform stack will be setup for undo / redo
+	 */
+	UFUNCTION(BlueprintCallable, Category = URigHierarchy)
+	FORCEINLINE_DEBUGGABLE void SetControlSettingsByIndex(int32 InElementIndex, FRigControlSettings InSettings, bool bSetupUndo = false, bool bForce = false)
+	{
+		if(Elements.IsValidIndex(InElementIndex))
+		{
+			if(FRigControlElement* ControlElement = Cast<FRigControlElement>(Elements[InElementIndex]))
+			{
+				SetControlSettings(ControlElement, InSettings, bSetupUndo, bForce);
+			}
+		}
+	}
+
 	/**
 	 * Returns the global current or initial value for a given key.
 	 * If the element does not have a parent FTransform::Identity will be returned.
@@ -2055,6 +2086,14 @@ public:
 	 */
 	void SetControlGizmoTransform(FRigControlElement* InControlElement, const FTransform& InTransform, const ERigTransformType::Type InTransformType, bool bSetupUndo = false, bool bForce = false);
 
+	/**
+	 * Sets the control settings for a given control element
+	 * @param InControlElement The element to set the settings for
+	 * @param InSettings The new control settings value to set
+	 * @param bSetupUndo If true the transform stack will be setup for undo / redo
+	 */
+	void SetControlSettings(FRigControlElement* InControlElement, FRigControlSettings InSettings, bool bSetupUndo = false, bool bForce = false);
+
 	/**
 	 * Returns a control's current value
 	 * @param InControlElement The element to retrieve the current value for
@@ -2366,9 +2405,13 @@ private:
 	};
 	
 	TArray<FRigHierarchyListener> ListeningHierarchies;
+	friend class FRigHierarchyListenerGuard;
 
 	// a bool to guard against circular dependencies among listening hierarchies
 	bool bPropagatingChange;
+
+	// a bool to disable any propagation checks and force propagation
+	bool bForcePropagation;
 	
 #endif
 
@@ -2386,3 +2429,99 @@ protected:
 	friend class UControlRig;
 	friend class FControlRigEditor;
 };
+
+#if WITH_EDITOR
+
+class FRigHierarchyListenerGuard
+{
+public:
+	FORCEINLINE_DEBUGGABLE FRigHierarchyListenerGuard(
+		URigHierarchy* InHierarchy, 
+		bool bInEnableInitialChanges, 
+		bool bInEnableCurrentChanges,
+		URigHierarchy* InListeningHierarchy = nullptr)
+			: Hierarchy(InHierarchy)
+			, bEnableInitialChanges(bInEnableInitialChanges)
+			, bEnableCurrentChanges(bInEnableCurrentChanges)
+			, ListeningHierarchy(InListeningHierarchy)
+	{
+		check(Hierarchy);
+
+		if(ListeningHierarchy == nullptr)
+		{
+			InitialFlags.AddUninitialized(Hierarchy->ListeningHierarchies.Num());
+			CurrentFlags.AddUninitialized(Hierarchy->ListeningHierarchies.Num());
+
+			for(int32 Index = 0; Index < Hierarchy->ListeningHierarchies.Num(); Index++)
+			{
+				URigHierarchy::FRigHierarchyListener& Listener = Hierarchy->ListeningHierarchies[Index];
+				InitialFlags[Index] = Listener.bShouldReactToInitialChanges;
+				CurrentFlags[Index] = Listener.bShouldReactToCurrentChanges;
+
+				Listener.bShouldReactToInitialChanges = bInEnableInitialChanges; 
+				Listener.bShouldReactToCurrentChanges = bInEnableCurrentChanges; 
+			}
+		}
+		else
+		{
+			for(int32 Index = 0; Index < Hierarchy->ListeningHierarchies.Num(); Index++)
+			{
+				URigHierarchy::FRigHierarchyListener& Listener = Hierarchy->ListeningHierarchies[Index];
+
+				if(Listener.Hierarchy.Get() == ListeningHierarchy)
+				{
+					InitialFlags.Add(Listener.bShouldReactToInitialChanges);
+					CurrentFlags.Add(Listener.bShouldReactToCurrentChanges);
+
+					Listener.bShouldReactToInitialChanges = bInEnableInitialChanges; 
+					Listener.bShouldReactToCurrentChanges = bInEnableCurrentChanges;
+					break;
+				}
+			}
+		}
+	}
+
+	FORCEINLINE_DEBUGGABLE ~FRigHierarchyListenerGuard()
+	{
+		if(ListeningHierarchy == nullptr)
+		{
+			check(Hierarchy->ListeningHierarchies.Num() == InitialFlags.Num());
+			check(Hierarchy->ListeningHierarchies.Num() == CurrentFlags.Num());
+			
+			for(int32 Index = 0; Index < Hierarchy->ListeningHierarchies.Num(); Index++)
+			{
+				URigHierarchy::FRigHierarchyListener& Listener = Hierarchy->ListeningHierarchies[Index];
+				Listener.bShouldReactToInitialChanges = InitialFlags[Index];
+				Listener.bShouldReactToCurrentChanges = CurrentFlags[Index];
+			}
+		}
+		else
+		{
+			for(int32 Index = 0; Index < Hierarchy->ListeningHierarchies.Num(); Index++)
+			{
+				URigHierarchy::FRigHierarchyListener& Listener = Hierarchy->ListeningHierarchies[Index];
+
+				if(Listener.Hierarchy.Get() == ListeningHierarchy)
+				{
+					check(InitialFlags.Num() == 1);
+					check(CurrentFlags.Num() == 1);
+
+					Listener.bShouldReactToInitialChanges = InitialFlags[0];
+					Listener.bShouldReactToCurrentChanges = CurrentFlags[0];
+					break;
+				}
+			}
+		}
+	}
+
+private:
+
+	URigHierarchy* Hierarchy; 
+	bool bEnableInitialChanges; 
+	bool bEnableCurrentChanges;
+	URigHierarchy* ListeningHierarchy;
+	TArray<bool> InitialFlags;
+	TArray<bool> CurrentFlags;
+};
+
+#endif
\ No newline at end of file
diff --git a/Engine/Plugins/Experimental/ControlRig/Source/ControlRig/Public/Rigs/RigHierarchyElements.h b/Engine/Plugins/Experimental/ControlRig/Source/ControlRig/Public/Rigs/RigHierarchyElements.h
--- a/Engine/Plugins/Experimental/ControlRig/Source/ControlRig/Public/Rigs/RigHierarchyElements.h
+++ b/Engine/Plugins/Experimental/ControlRig/Source/ControlRig/Public/Rigs/RigHierarchyElements.h
@@ -864,6 +864,13 @@ struct CONTROLRIG_API FRigControlSettings
 		Value.SetFromTransform(FTransform::Identity, ControlType, PrimaryAxis);
 		return Value;
 	}
+
+	bool operator == (const FRigControlSettings& InOther) const;
+
+	FORCEINLINE bool operator != (const FRigControlSettings& InOther) const
+	{
+		return !(*this == InOther);
+	}
 };
 
 USTRUCT(BlueprintType)
diff --git a/Engine/Plugins/Experimental/ControlRig/Source/ControlRigEditor/Private/Editor/ControlRigEditor.cpp b/Engine/Plugins/Experimental/ControlRig/Source/ControlRigEditor/Private/Editor/ControlRigEditor.cpp
--- a/Engine/Plugins/Experimental/ControlRig/Source/ControlRigEditor/Private/Editor/ControlRigEditor.cpp
+++ b/Engine/Plugins/Experimental/ControlRig/Source/ControlRigEditor/Private/Editor/ControlRigEditor.cpp
@@ -4013,7 +4013,14 @@ void FControlRigEditor::OnFinishedChangingProperties(const FPropertyChangedEvent
 				FRigControlElement* DebuggedControlElement = DebuggedControlRig->GetHierarchy()->Find<FRigControlElement>(RigElementInDetailPanel);
 				if(DebuggedControlElement)
 				{
-					ControlElement->Settings = DebuggedControlElement->Settings;
+					if(bSetupModeEnabled)
+					{
+						DebuggedControlElement->Settings = ControlElement->Settings;
+					}
+					else
+					{
+						ControlElement->Settings = DebuggedControlElement->Settings;
+					}
 				}
 			
 				if (PropertyChangedEvent.Property->GetName() == TEXT("GizmoColor"))
@@ -4063,33 +4070,35 @@ void FControlRigEditor::OnBlueprintPropertyChainEvent(FPropertyChangedChainEvent
 		{
 			TransformType = ERigTransformType::MakeInitial(TransformType); 
 		}
+		
 		if(bIsLocal)
 		{
 			TransformType = ERigTransformType::MakeLocal(TransformType); 
 		}
 
+		UControlRig* DefaultRig = Cast<UControlRig>(ControlRigBP->GeneratedClass->GetDefaultObject(true /* create if needed */));
+		URigHierarchy* DefaultHierarchy = DefaultRig->GetHierarchy();
 		URigHierarchy* DebuggedHierarchy = DebuggedControlRig->GetHierarchy();
-		URigHierarchy* Hierarchy = DebuggedHierarchy;
-		if(bSetupModeEnabled || bIsInitial)
-		{
-			Hierarchy = ControlRigBP->Hierarchy;
-			DebuggedHierarchy = Hierarchy;
-		}
+		URigHierarchy* BlueprintHierarchy = ControlRigBP->Hierarchy;
+
+		FRigBaseElement* SourceElement = bSetupModeEnabled ?
+			BlueprintHierarchy->Find(RigElementInDetailPanel) :
+			DebuggedHierarchy->Find(RigElementInDetailPanel);
 
-		FRigBaseElement* SourceElement = DebuggedHierarchy->Find(RigElementInDetailPanel);
-		if(SourceElement == nullptr)
-		{
-			SourceElement = Hierarchy->Find(RigElementInDetailPanel);;
-		}
 		if(SourceElement == nullptr)
 		{
 			return;
 		}
-		FRigBaseElement* TargetElement = Hierarchy->Find(RigElementInDetailPanel);
+		
+		FRigBaseElement* TargetElement = BlueprintHierarchy->Find(RigElementInDetailPanel);
 		if(TargetElement == nullptr)
 		{
 			return;
 		}
+
+		// force propagation of any change from the blueprint hierarchy
+		FRigHierarchyListenerGuard PropagateToDefault(BlueprintHierarchy, true, true);
+		FRigHierarchyListenerGuard PropagateToDebugged(DefaultHierarchy, true, true, DebuggedHierarchy);
 		
 		if(HeadProperty->GetValue() == FRigTransformElement::StaticStruct()->FindPropertyByName(TEXT("Pose")))
 		{
@@ -4106,18 +4115,27 @@ void FControlRigEditor::OnBlueprintPropertyChainEvent(FPropertyChangedChainEvent
 						FRigControlValue Value;
 						Value.SetFromTransform(Transform, TargetControlElement->Settings.ControlType, TargetControlElement->Settings.PrimaryAxis);
 						
-						if(ERigTransformType::IsInitial(TransformType))
+						if(ERigTransformType::IsInitial(TransformType) || bSetupModeEnabled)
 						{
-							Hierarchy->SetControlValue(TargetControlElement, Value, ERigControlValueType::Initial, true, true);
+							BlueprintHierarchy->SetControlValue(TargetControlElement, Value, ERigControlValueType::Initial, true, true);
+							BlueprintHierarchy->SetControlValue(TargetControlElement, Value, ERigControlValueType::Current, true, true);
 						}
 						else
 						{
-							Hierarchy->SetControlValue(TargetControlElement, Value, ERigControlValueType::Current, true, true);
+							BlueprintHierarchy->SetControlValue(TargetControlElement, Value, ERigControlValueType::Current, true, true);
 						}
 					}
 					else
 					{
-						Hierarchy->SetTransform(TargetTransformElement, Transform, TransformType, true, true, true);
+						if(ERigTransformType::IsInitial(TransformType) || bSetupModeEnabled)
+						{
+							BlueprintHierarchy->SetTransform(TargetTransformElement, Transform, ERigTransformType::MakeInitial(TransformType), true, true, true);
+							BlueprintHierarchy->SetTransform(TargetTransformElement, Transform, ERigTransformType::MakeCurrent(TransformType), true, true, true);
+						}
+						else
+						{
+							BlueprintHierarchy->SetTransform(TargetTransformElement, Transform, TransformType, true, true, true);
+						}
 					}
 				}
 			}
@@ -4130,8 +4148,8 @@ void FControlRigEditor::OnBlueprintPropertyChainEvent(FPropertyChangedChainEvent
 				if(FRigControlElement* TargetControlElement = Cast<FRigControlElement>(TargetElement))
 				{
 					const FTransform Transform = SourceControlElement->Offset.Get(TransformType);
-					Hierarchy->SetControlOffsetTransform(TargetControlElement, Transform, TransformType, true, true, true);
-					Hierarchy->SetControlOffsetTransform(TargetControlElement, Transform, ERigTransformType::MakeCurrent(TransformType), true, true, true);
+					// this will call the current version internally as well
+					BlueprintHierarchy->SetControlOffsetTransform(TargetControlElement, Transform, ERigTransformType::MakeInitial(TransformType), true, true, true);
 				}
 			}
 		}
@@ -4143,8 +4161,8 @@ void FControlRigEditor::OnBlueprintPropertyChainEvent(FPropertyChangedChainEvent
 				if(FRigControlElement* TargetControlElement = Cast<FRigControlElement>(TargetElement))
 				{
 					const FTransform Transform = SourceControlElement->Gizmo.Get(TransformType);
-					DebuggedHierarchy->SetControlGizmoTransform(SourceControlElement, Transform, TransformType, false, true);
-					DebuggedHierarchy->SetControlGizmoTransform(SourceControlElement, Transform, ERigTransformType::MakeCurrent(TransformType), false, true);
+					// this will call the current version internally as well
+					BlueprintHierarchy->SetControlGizmoTransform(TargetControlElement, Transform, ERigTransformType::MakeInitial(TransformType), false, true);
 					return;
 				}
 			}
diff --git a/Engine/Plugins/Experimental/ControlRig/Source/ControlRigEditor/Private/Editor/SControlRigGizmoNameList.cpp b/Engine/Plugins/Experimental/ControlRig/Source/ControlRigEditor/Private/Editor/SControlRigGizmoNameList.cpp
--- a/Engine/Plugins/Experimental/ControlRig/Source/ControlRigEditor/Private/Editor/SControlRigGizmoNameList.cpp
+++ b/Engine/Plugins/Experimental/ControlRig/Source/ControlRigEditor/Private/Editor/SControlRigGizmoNameList.cpp
@@ -75,36 +75,16 @@ void SControlRigGizmoNameList::SetNameListText(const FText& NewTypeInValue, ETex
 
 		bool bIsOnInstance = false;
 		URigHierarchy* Hierarchy = Blueprint->Hierarchy;
-		UControlRig* DebuggedControlRig = Cast<UControlRig>(Blueprint->GetObjectBeingDebugged());
-
-		if (DebuggedControlRig)
-		{
-			if (!DebuggedControlRig->IsSetupModeEnabled())
-			{
-				Hierarchy = DebuggedControlRig->GetHierarchy();
-				bIsOnInstance = true;
-			}
-		}
 
 		FRigControlElement* ControlElement = Hierarchy->Get<FRigControlElement>(ControlIndex);
 		if ((ControlElement != nullptr) && (ControlElement->Settings.GizmoName != NewName))
 		{
 			const FScopedTransaction Transaction(NSLOCTEXT("ControlRigEditor", "ChangeGizmoName", "Change Gizmo Name"));
-			Blueprint->Hierarchy->Modify();
-
-			ControlElement->Settings.GizmoName = NewName;
-
-			if (bIsOnInstance && DebuggedControlRig)
-			{
-				FRigControlElement* OtherControlElement = Blueprint->Hierarchy->Get<FRigControlElement>(ControlIndex);
-				if(OtherControlElement)
-				{
-					OtherControlElement->Settings.GizmoName = NewName;
-					Blueprint->Hierarchy->Notify(ERigHierarchyNotification::ControlSettingChanged, OtherControlElement);
-				}
-			}
+			Hierarchy->Modify();
 
-			Hierarchy->Notify(ERigHierarchyNotification::ControlSettingChanged, ControlElement);
+			FRigControlSettings Settings = ControlElement->Settings;
+			Settings.GizmoName = NewName;
+			Hierarchy->SetControlSettings(ControlElement, Settings, true, true);
 		}
 	}
 }
