Control Rig: Fix multi node movement

diff --git a/Engine/Plugins/Experimental/ControlRig/Source/ControlRigDeveloper/Private/Graph/ControlRigGraphSchema.cpp b/Engine/Plugins/Experimental/ControlRig/Source/ControlRigDeveloper/Private/Graph/ControlRigGraphSchema.cpp
--- a/Engine/Plugins/Experimental/ControlRig/Source/ControlRigDeveloper/Private/Graph/ControlRigGraphSchema.cpp
+++ b/Engine/Plugins/Experimental/ControlRig/Source/ControlRigDeveloper/Private/Graph/ControlRigGraphSchema.cpp
@@ -775,9 +775,16 @@ bool UControlRigGraphSchema::IsStructEditable(UStruct* InStruct) const
 
 void UControlRigGraphSchema::SetNodePosition(UEdGraphNode* Node, const FVector2D& Position) const
 {
+	return SetNodePosition(Node, Position, true);
+}
+
+void UControlRigGraphSchema::SetNodePosition(UEdGraphNode* Node, const FVector2D& Position, bool bSetupUndo) const
+{
+	StartGraphNodeInteraction(Node);
+
 	if (UControlRigGraphNode* RigNode = Cast<UControlRigGraphNode>(Node))
 	{
-		RigNode->GetController()->SetNodePosition(RigNode->GetModelNode(), Position, true, false);
+		RigNode->GetController()->SetNodePosition(RigNode->GetModelNode(), Position, bSetupUndo, false);
 	}
 }
 
@@ -1302,10 +1309,20 @@ bool UControlRigGraphSchema::RequestVariableDropOnPin(UEdGraph* InGraph, FProper
 	return false;
 }
 
-void UControlRigGraphSchema::EndGraphNodeInteraction(UEdGraphNode* InNode) const
+void UControlRigGraphSchema::StartGraphNodeInteraction(UEdGraphNode* InNode) const
 {
 #if WITH_EDITOR
 
+	check(InNode);
+
+	if(NodesBeingInteracted.Contains(InNode))
+	{
+		return;
+	}
+	
+	NodePositionsDuringStart.Reset();
+	NodesBeingInteracted.Reset();
+
 	UControlRigGraph* Graph = Cast<UControlRigGraph>(InNode->GetOuter());
 	if (Graph == nullptr)
 	{
@@ -1315,34 +1332,35 @@ void UControlRigGraphSchema::EndGraphNodeInteraction(UEdGraphNode* InNode) const
 	check(Graph->GetController());
 	check(Graph->GetModel());
 
-	TArray<UEdGraphNode*> NodesToMove;
-	NodesToMove.Add(InNode);
+	NodesBeingInteracted = GetNodesToMoveForNode(InNode);
 
-	for (UEdGraphNode* SelectedGraphNode : Graph->Nodes)
+	for (UEdGraphNode* NodeToMove : NodesBeingInteracted)
 	{
-		if (SelectedGraphNode->IsSelected())
+		FName NodeName = NodeToMove->GetFName();
+		if (URigVMNode* ModelNode = Graph->GetModel()->FindNodeByName(NodeName))
 		{
-			NodesToMove.AddUnique(SelectedGraphNode);
+			NodePositionsDuringStart.FindOrAdd(NodeName, ModelNode->GetPosition());
 		}
 	}
 
-	for (int32 NodeIndex = 0; NodeIndex < NodesToMove.Num(); NodeIndex++)
+#endif
+}
+
+void UControlRigGraphSchema::EndGraphNodeInteraction(UEdGraphNode* InNode) const
+{
+#if WITH_EDITOR
+
+	UControlRigGraph* Graph = Cast<UControlRigGraph>(InNode->GetOuter());
+	if (Graph == nullptr)
 	{
-		if (UEdGraphNode_Comment* CommentNode = Cast<UEdGraphNode_Comment>(NodesToMove[NodeIndex]))
-		{
-			if (CommentNode->MoveMode == ECommentBoxMode::GroupMovement)
-			{
-				for (FCommentNodeSet::TConstIterator NodeIt(CommentNode->GetNodesUnderComment()); NodeIt; ++NodeIt)
-				{
-					if (UEdGraphNode* NodeUnderComment = Cast<UEdGraphNode>(*NodeIt))
-					{
-						NodesToMove.AddUnique(NodeUnderComment);
-					}
-				}
-			}
-		}
+		return;
 	}
 
+	check(Graph->GetController());
+	check(Graph->GetModel());
+
+	TArray<UEdGraphNode*> NodesToMove = GetNodesToMoveForNode(InNode);
+	
 	bool bMovedSomething = false;
 
 	Graph->GetController()->OpenUndoBracket(TEXT("Move Nodes"));
@@ -1352,8 +1370,15 @@ void UControlRigGraphSchema::EndGraphNodeInteraction(UEdGraphNode* InNode) const
 		FName NodeName = NodeToMove->GetFName();
 		if (URigVMNode* ModelNode = Graph->GetModel()->FindNodeByName(NodeName))
 		{
-			FVector2D Position(NodeToMove->NodePosX, NodeToMove->NodePosY);
-			if (Graph->GetController()->SetNodePositionByName(NodeName, Position, true, false))
+			if(FVector2D* OldPosition = NodePositionsDuringStart.Find(NodeName))
+			{
+				Graph->GetController()->SuspendNotifications(true);
+				Graph->GetController()->SetNodePositionByName(NodeName, *OldPosition, false, false);
+				Graph->GetController()->SuspendNotifications(false);
+			}
+			
+			FVector2D NewPosition(NodeToMove->NodePosX, NodeToMove->NodePosY);
+			if(Graph->GetController()->SetNodePositionByName(NodeName, NewPosition, true, false))
 			{
 				bMovedSomething = true;
 			}
@@ -1369,7 +1394,55 @@ void UControlRigGraphSchema::EndGraphNodeInteraction(UEdGraphNode* InNode) const
 		Graph->GetController()->CancelUndoBracket();
 	}
 
+	NodesBeingInteracted.Reset();
+	NodePositionsDuringStart.Reset();
+
+#endif
+}
+
+TArray<UEdGraphNode*> UControlRigGraphSchema::GetNodesToMoveForNode(UEdGraphNode* InNode)
+{
+	TArray<UEdGraphNode*> NodesToMove;
+
+#if WITH_EDITOR
+
+	UControlRigGraph* Graph = Cast<UControlRigGraph>(InNode->GetOuter());
+	if (Graph == nullptr)
+	{
+		return NodesToMove;
+	}
+
+	NodesToMove.Add(InNode);
+
+	for (UEdGraphNode* SelectedGraphNode : Graph->Nodes)
+	{
+		if (SelectedGraphNode->IsSelected())
+		{
+			NodesToMove.AddUnique(SelectedGraphNode);
+		}
+	}
+
+	for (int32 NodeIndex = 0; NodeIndex < NodesToMove.Num(); NodeIndex++)
+	{
+		if (UEdGraphNode_Comment* CommentNode = Cast<UEdGraphNode_Comment>(NodesToMove[NodeIndex]))
+		{
+			if (CommentNode->MoveMode == ECommentBoxMode::GroupMovement)
+			{
+				for (FCommentNodeSet::TConstIterator NodeIt(CommentNode->GetNodesUnderComment()); NodeIt; ++NodeIt)
+				{
+					if (UEdGraphNode* NodeUnderComment = Cast<UEdGraphNode>(*NodeIt))
+					{
+						NodesToMove.AddUnique(NodeUnderComment);
+					}
+				}
+			}
+		}
+	}
+
 #endif
+
+	return NodesToMove;
 }
 
+
 #undef LOCTEXT_NAMESPACE
diff --git a/Engine/Plugins/Experimental/ControlRig/Source/ControlRigDeveloper/Public/Graph/ControlRigGraphSchema.h b/Engine/Plugins/Experimental/ControlRig/Source/ControlRigDeveloper/Public/Graph/ControlRigGraphSchema.h
--- a/Engine/Plugins/Experimental/ControlRig/Source/ControlRigDeveloper/Public/Graph/ControlRigGraphSchema.h
+++ b/Engine/Plugins/Experimental/ControlRig/Source/ControlRigDeveloper/Public/Graph/ControlRigGraphSchema.h
@@ -161,6 +161,7 @@ public:
 	virtual bool RequestVariableDropOnPin(UEdGraph* InGraph, FProperty* InVariableToDrop, UEdGraphPin* InPin, const FVector2D& InDropPosition, const FVector2D& InScreenPosition) override;
 	virtual bool IsStructEditable(UStruct* InStruct) const;
 	virtual void SetNodePosition(UEdGraphNode* Node, const FVector2D& Position) const override;
+	void SetNodePosition(UEdGraphNode* Node, const FVector2D& Position, bool bSetupUndo) const;
 	virtual void GetGraphDisplayInformation(const UEdGraph& Graph, /*out*/ FGraphDisplayInfo& DisplayInfo) const override;
 	virtual bool GetLocalVariables(const UEdGraph* InGraph, TArray<FBPVariableDescription>& OutLocalVariables) const override;
 	virtual TSharedPtr<FEdGraphSchemaAction> MakeActionFromVariableDescription(const UEdGraph* InEdGraph, const FBPVariableDescription& Variable) const override;
@@ -191,7 +192,9 @@ public:
 	/** Returns all of the applicable pin types for variables within a control rig */
 	virtual void GetVariablePinTypes(TArray<FEdGraphPinType>& PinTypes) const;
 
+	void StartGraphNodeInteraction(UEdGraphNode* InNode) const;
 	void EndGraphNodeInteraction(UEdGraphNode* InNode) const;
+	static TArray<UEdGraphNode*> GetNodesToMoveForNode(UEdGraphNode* InNode);
 
 private:
 
@@ -199,6 +202,11 @@ private:
 	bool bLastPinWasInput;
 	mutable UEdGraphPin* PinBeingDropped = nullptr;
 
+#if WITH_EDITOR
+	mutable TArray<UEdGraphNode*> NodesBeingInteracted;
+	mutable TMap<FName, FVector2D> NodePositionsDuringStart;
+#endif
+
 	friend class UControlRigRerouteNodeSpawner;
 	friend class UControlRigIfNodeSpawner;
 	friend class UControlRigSelectNodeSpawner;
diff --git a/Engine/Plugins/Experimental/ControlRig/Source/ControlRigEditor/Private/Graph/SControlRigGraphNode.cpp b/Engine/Plugins/Experimental/ControlRig/Source/ControlRigEditor/Private/Graph/SControlRigGraphNode.cpp
--- a/Engine/Plugins/Experimental/ControlRig/Source/ControlRigEditor/Private/Graph/SControlRigGraphNode.cpp
+++ b/Engine/Plugins/Experimental/ControlRig/Source/ControlRigEditor/Private/Graph/SControlRigGraphNode.cpp
@@ -570,7 +570,7 @@ void SControlRigGraphNode::MoveTo(const FVector2D& NewPosition, FNodeSet& NodeFi
 		{
 			if (const UControlRigGraphSchema* RigSchema = Cast<UControlRigGraphSchema>(GraphNode->GetSchema()))
 			{
-				RigSchema->SetNodePosition(GraphNode, NewPosition);
+				RigSchema->SetNodePosition(GraphNode, NewPosition, false);
 			}
 		}
 	}
