Control Rig: Fix multi node movement

diff --git a/Engine/Plugins/Experimental/ControlRig/Source/ControlRigDeveloper/Private/Graph/ControlRigGraphSchema.cpp b/Engine/Plugins/Experimental/ControlRig/Source/ControlRigDeveloper/Private/Graph/ControlRigGraphSchema.cpp
--- a/Engine/Plugins/Experimental/ControlRig/Source/ControlRigDeveloper/Private/Graph/ControlRigGraphSchema.cpp
+++ b/Engine/Plugins/Experimental/ControlRig/Source/ControlRigDeveloper/Private/Graph/ControlRigGraphSchema.cpp
@@ -786,6 +786,14 @@ void UControlRigGraphSchema::SetNodePosition(UEdGraphNode* Node, const FVector2D
 	{
 		RigNode->GetController()->SetNodePosition(RigNode->GetModelNode(), Position, bSetupUndo, false);
 	}
+	
+	if (UEdGraphNode_Comment* CommentNode = Cast<UEdGraphNode_Comment>(Node))
+	{
+		if(UControlRigGraph* Graph = CommentNode->GetTypedOuter<UControlRigGraph>())
+		{
+			Graph->GetController()->SetNodePositionByName(CommentNode->GetFName(), Position, bSetupUndo, false);
+		}
+	}
 }
 
 void UControlRigGraphSchema::GetGraphDisplayInformation(const UEdGraph& Graph, /*out*/ FGraphDisplayInfo& DisplayInfo) const
@@ -1370,14 +1378,14 @@ void UControlRigGraphSchema::EndGraphNodeInteraction(UEdGraphNode* InNode) const
 		FName NodeName = NodeToMove->GetFName();
 		if (URigVMNode* ModelNode = Graph->GetModel()->FindNodeByName(NodeName))
 		{
+			FVector2D NewPosition(NodeToMove->NodePosX, NodeToMove->NodePosY);
+
 			if(FVector2D* OldPosition = NodePositionsDuringStart.Find(NodeName))
 			{
-				Graph->GetController()->SuspendNotifications(true);
+				TGuardValue<bool> SuspendNotification(Graph->bSuspendModelNotifications, true);
 				Graph->GetController()->SetNodePositionByName(NodeName, *OldPosition, false, false);
-				Graph->GetController()->SuspendNotifications(false);
 			}
 			
-			FVector2D NewPosition(NodeToMove->NodePosX, NodeToMove->NodePosY);
 			if(Graph->GetController()->SetNodePositionByName(NodeName, NewPosition, true, false))
 			{
 				bMovedSomething = true;
@@ -1387,6 +1395,11 @@ void UControlRigGraphSchema::EndGraphNodeInteraction(UEdGraphNode* InNode) const
 
 	if (bMovedSomething)
 	{
+		if (GEditor)
+		{
+			GEditor->CancelTransaction(0);
+		}
+
 		Graph->GetController()->CloseUndoBracket();
 	}
 	else
@@ -1444,5 +1457,22 @@ TArray<UEdGraphNode*> UControlRigGraphSchema::GetNodesToMoveForNode(UEdGraphNode
 	return NodesToMove;
 }
 
+FVector2D UControlRigGraphSchema::GetNodePositionAtStartOfInteraction(UEdGraphNode* InNode) const
+{
+#if WITH_EDITOR
+	if(InNode)
+	{
+		if(const FVector2D* Position = NodePositionsDuringStart.Find(InNode->GetFName()))
+		{
+			return *Position;
+		}
+
+		return FVector2D(InNode->NodePosX, InNode->NodePosY);
+	}
+#endif
+
+	return FVector2D::ZeroVector;
+}
+
 
 #undef LOCTEXT_NAMESPACE
diff --git a/Engine/Plugins/Experimental/ControlRig/Source/ControlRigDeveloper/Public/Graph/ControlRigGraphSchema.h b/Engine/Plugins/Experimental/ControlRig/Source/ControlRigDeveloper/Public/Graph/ControlRigGraphSchema.h
--- a/Engine/Plugins/Experimental/ControlRig/Source/ControlRigDeveloper/Public/Graph/ControlRigGraphSchema.h
+++ b/Engine/Plugins/Experimental/ControlRig/Source/ControlRigDeveloper/Public/Graph/ControlRigGraphSchema.h
@@ -195,6 +195,7 @@ public:
 	void StartGraphNodeInteraction(UEdGraphNode* InNode) const;
 	void EndGraphNodeInteraction(UEdGraphNode* InNode) const;
 	static TArray<UEdGraphNode*> GetNodesToMoveForNode(UEdGraphNode* InNode);
+	FVector2D GetNodePositionAtStartOfInteraction(UEdGraphNode* InNode) const;
 
 private:
 
diff --git a/Engine/Plugins/Experimental/ControlRig/Source/ControlRigEditor/Private/Graph/SControlRigGraphNodeComment.cpp b/Engine/Plugins/Experimental/ControlRig/Source/ControlRigEditor/Private/Graph/SControlRigGraphNodeComment.cpp
--- a/Engine/Plugins/Experimental/ControlRig/Source/ControlRigEditor/Private/Graph/SControlRigGraphNodeComment.cpp
+++ b/Engine/Plugins/Experimental/ControlRig/Source/ControlRigEditor/Private/Graph/SControlRigGraphNodeComment.cpp
@@ -129,6 +129,33 @@ void SControlRigGraphNodeComment::Tick(const FGeometry& AllottedGeometry, const
 	SGraphNodeComment::Tick(AllottedGeometry, InCurrentTime, InDeltaTime);
 }
 
+void SControlRigGraphNodeComment::MoveTo(const FVector2D& NewPosition, FNodeSet& NodeFilter, bool bMarkDirty)
+{
+	if (!NodeFilter.Find(SharedThis(this)))
+	{
+		if (GraphNode && !RequiresSecondPassLayout())
+		{
+			if (const UControlRigGraphSchema* RigSchema = Cast<UControlRigGraphSchema>(GraphNode->GetSchema()))
+			{
+				const FVector2D CommentOldPosition = RigSchema->GetNodePositionAtStartOfInteraction(GraphNode);
+				const FVector2D Delta = NewPosition - CommentOldPosition;
+
+				TArray<UEdGraphNode*> NodesToMove = UControlRigGraphSchema::GetNodesToMoveForNode(GraphNode);
+				for(UEdGraphNode* NodeToMove : NodesToMove)
+				{
+					const FVector2D OldPosition = RigSchema->GetNodePositionAtStartOfInteraction(NodeToMove);
+
+					if((NodeToMove == GraphNode) || IsNodeUnderComment(CommentOldPosition, OldPosition))
+					{
+						const FVector2D AccumulatedPosition = OldPosition + Delta; 
+						RigSchema->SetNodePosition(NodeToMove, AccumulatedPosition, false);
+					}
+				}
+			}
+		}
+	}
+}
+
 /*
 void SControlRigGraphNodeComment::OnCommentTextCommitted(const FText& NewComment, ETextCommit::Type CommitInfo)
 {
@@ -156,11 +183,25 @@ void SControlRigGraphNodeComment::OnCommentTextCommitted(const FText& NewComment
 
 bool SControlRigGraphNodeComment::IsNodeUnderComment(UEdGraphNode_Comment* InCommentNode, const TSharedRef<SGraphNode> InNodeWidget) const
 {
-	const FVector2D NodePosition = GetPosition();
+	return IsNodeUnderComment(InNodeWidget->GetPosition());
+}
+
+
+bool SControlRigGraphNodeComment::IsNodeUnderComment(UEdGraphNode* InNode) const
+{
+	return IsNodeUnderComment(FVector2D(InNode->NodePosX, InNode->NodePosY));
+}
+
+bool SControlRigGraphNodeComment::IsNodeUnderComment(const FVector2D& InNodePosition) const
+{
+	return IsNodeUnderComment(GetPosition(), InNodePosition);
+}
+
+bool SControlRigGraphNodeComment::IsNodeUnderComment(const FVector2D& InCommentPosition, const FVector2D& InNodePosition) const
+{
 	const FVector2D NodeSize = GetDesiredSize();
-	const FSlateRect CommentRect(NodePosition.X, NodePosition.Y, NodePosition.X + NodeSize.X, NodePosition.Y + NodeSize.Y);
+	const FSlateRect CommentRect(InCommentPosition.X, InCommentPosition.Y, InCommentPosition.X + NodeSize.X, InCommentPosition.Y + NodeSize.Y);
 
-	const FVector2D InNodePosition = InNodeWidget->GetPosition();
 	return CommentRect.ContainsPoint(InNodePosition);
 }
 
diff --git a/Engine/Plugins/Experimental/ControlRig/Source/ControlRigEditor/Private/Graph/SControlRigGraphNodeComment.h b/Engine/Plugins/Experimental/ControlRig/Source/ControlRigEditor/Private/Graph/SControlRigGraphNodeComment.h
--- a/Engine/Plugins/Experimental/ControlRig/Source/ControlRigEditor/Private/Graph/SControlRigGraphNodeComment.h
+++ b/Engine/Plugins/Experimental/ControlRig/Source/ControlRigEditor/Private/Graph/SControlRigGraphNodeComment.h
@@ -16,10 +16,14 @@ public:
 	virtual FReply OnMouseButtonUp(const FGeometry& MyGeometry, const FPointerEvent& MouseEvent) override;
 	virtual void EndUserInteraction() const override;
 	virtual void Tick(const FGeometry& AllottedGeometry, const double InCurrentTime, const float InDeltaTime) override;
+	virtual void MoveTo(const FVector2D& NewPosition, FNodeSet& NodeFilter, bool bMarkDirty = true) override;
 
 protected:
 
 	virtual bool IsNodeUnderComment(UEdGraphNode_Comment* InCommentNode, const TSharedRef<SGraphNode> InNodeWidget) const override;
+	bool IsNodeUnderComment(UEdGraphNode* InNode) const;
+	bool IsNodeUnderComment(const FVector2D& InNodePosition) const;
+	bool IsNodeUnderComment(const FVector2D& InCommentPosition, const FVector2D& InNodePosition) const;
 
 private:
 
diff --git a/Engine/Plugins/Experimental/ControlRig/Source/ControlRigEditor/Private/Graph/SControlRigGraphNodeKnot.cpp b/Engine/Plugins/Experimental/ControlRig/Source/ControlRigEditor/Private/Graph/SControlRigGraphNodeKnot.cpp
--- a/Engine/Plugins/Experimental/ControlRig/Source/ControlRigEditor/Private/Graph/SControlRigGraphNodeKnot.cpp
+++ b/Engine/Plugins/Experimental/ControlRig/Source/ControlRigEditor/Private/Graph/SControlRigGraphNodeKnot.cpp
@@ -34,4 +34,18 @@ void SControlRigGraphNodeKnot::EndUserInteraction() const
 	SGraphNodeKnot::EndUserInteraction();
 }
 
-#undef LOCTEXT_NAMESPACE
\ No newline at end of file
+void SControlRigGraphNodeKnot::MoveTo(const FVector2D& NewPosition, FNodeSet& NodeFilter, bool bMarkDirty)
+{
+	if (!NodeFilter.Find(SharedThis(this)))
+	{
+		if (GraphNode && !RequiresSecondPassLayout())
+		{
+			if (const UControlRigGraphSchema* RigSchema = Cast<UControlRigGraphSchema>(GraphNode->GetSchema()))
+			{
+				RigSchema->SetNodePosition(GraphNode, NewPosition, false);
+			}
+		}
+	}
+}
+
+#undef LOCTEXT_NAMESPACE
diff --git a/Engine/Plugins/Experimental/ControlRig/Source/ControlRigEditor/Private/Graph/SControlRigGraphNodeKnot.h b/Engine/Plugins/Experimental/ControlRig/Source/ControlRigEditor/Private/Graph/SControlRigGraphNodeKnot.h
--- a/Engine/Plugins/Experimental/ControlRig/Source/ControlRigEditor/Private/Graph/SControlRigGraphNodeKnot.h
+++ b/Engine/Plugins/Experimental/ControlRig/Source/ControlRigEditor/Private/Graph/SControlRigGraphNodeKnot.h
@@ -15,4 +15,5 @@ class SControlRigGraphNodeKnot : public SGraphNodeKnot
 
 	// SGraphNode interface
 	virtual void EndUserInteraction() const override;
+	virtual void MoveTo(const FVector2D& NewPosition, FNodeSet& NodeFilter, bool bMarkDirty = true) override;
 };
\ No newline at end of file
