Control Rig: Crash with Control Rig Actors due to them not removing delegate when deleted.

diff --git a/Engine/Plugins/Experimental/ControlRig/Source/ControlRig/Private/ControlRigControlActor.cpp b/Engine/Plugins/Experimental/ControlRig/Source/ControlRig/Private/ControlRigControlActor.cpp
--- a/Engine/Plugins/Experimental/ControlRig/Source/ControlRig/Private/ControlRigControlActor.cpp
+++ b/Engine/Plugins/Experimental/ControlRig/Source/ControlRig/Private/ControlRigControlActor.cpp
@@ -32,6 +32,26 @@ AControlRigControlActor::AControlRigControlActor(const FObjectInitializer& Objec
 	Refresh();
 }
 
+AControlRigControlActor::~AControlRigControlActor()
+{
+	RemoveUnbindDelegate();
+}
+
+void AControlRigControlActor::RemoveUnbindDelegate()
+{
+	if (ControlRig)
+	{
+		if (TSharedPtr<IControlRigObjectBinding> Binding = ControlRig->GetObjectBinding())
+		{
+			if (OnUnbindDelegate.IsValid())
+			{
+				Binding->OnControlRigUnbind().Remove(OnUnbindDelegate);
+				OnUnbindDelegate.Reset();
+			}
+		}
+	}
+}
+
 #if WITH_EDITOR
 
 void AControlRigControlActor::PostEditChangeProperty(FPropertyChangedEvent& PropertyChangedEvent)
@@ -45,6 +65,8 @@ void AControlRigControlActor::PostEditChangeProperty(FPropertyChangedEvent& Prop
 			PropertyChangedEvent.Property->GetFName() == GET_MEMBER_NAME_CHECKED(AControlRigControlActor, ColorParameter) ||
 			PropertyChangedEvent.Property->GetFName() == GET_MEMBER_NAME_CHECKED(AControlRigControlActor, bCastShadows)))
 	{
+		RemoveUnbindDelegate();
+		ControlRig = nullptr;
 		Clear();
 		Refresh();
 	}
@@ -65,29 +87,38 @@ void AControlRigControlActor::Tick(float DeltaSeconds)
 void AControlRigControlActor::Clear()
 {
 	TArray<USceneComponent*> ChildComponents;
-	ActorRootComponent->GetChildrenComponents(true, ChildComponents);
-	for (USceneComponent* Child : ChildComponents)
+	if (ActorRootComponent)
 	{
-		if (UStaticMeshComponent* StaticMeshComponent = Cast<UStaticMeshComponent>(Child))
+		ActorRootComponent->GetChildrenComponents(true, ChildComponents);
+		for (USceneComponent* Child : ChildComponents)
 		{
-			Components.AddUnique(StaticMeshComponent);
+			if (UStaticMeshComponent* StaticMeshComponent = Cast<UStaticMeshComponent>(Child))
+			{
+				Components.AddUnique(StaticMeshComponent);
+			}
 		}
-	}
 
-	for (UStaticMeshComponent* Component : Components)
-	{
-		Component->DetachFromComponent(FDetachmentTransformRules::KeepRelativeTransform);
-		Component->UnregisterComponent();
-		Component->DestroyComponent();
+		for (UStaticMeshComponent* Component : Components)
+		{
+			Component->DetachFromComponent(FDetachmentTransformRules::KeepRelativeTransform);
+			Component->UnregisterComponent();
+			Component->DestroyComponent();
+		}
 	}
 
-	ControlRig = nullptr;
 	ControlNames.Reset();
 	GizmoTransforms.Reset();
 	Components.Reset();
 	Materials.Reset();
 }
 
+void AControlRigControlActor::EndPlay(const EEndPlayReason::Type EndPlayReason)
+{
+	RemoveUnbindDelegate();
+	ControlRig = nullptr;
+	Super::EndPlay(EndPlayReason);
+}
+
 void AControlRigControlActor::Refresh()
 {
 	if (ActorToTrack == nullptr)
@@ -109,7 +140,8 @@ void AControlRigControlActor::Refresh()
 						if (Actor == ActorToTrack)
 						{
 							ControlRig = RigInstance;
-							Binding->OnControlRigUnbind().AddLambda([ this ]( ) { this->Clear(); this->Refresh(); });
+							RemoveUnbindDelegate();
+							OnUnbindDelegate =Binding->OnControlRigUnbind().AddLambda([ this ]( ) { this->Clear(); this->Refresh(); });
 							break;
 						}
 					}
diff --git a/Engine/Plugins/Experimental/ControlRig/Source/ControlRig/Public/ControlRigControlActor.h b/Engine/Plugins/Experimental/ControlRig/Source/ControlRig/Public/ControlRigControlActor.h
--- a/Engine/Plugins/Experimental/ControlRig/Source/ControlRig/Public/ControlRigControlActor.h
+++ b/Engine/Plugins/Experimental/ControlRig/Source/ControlRig/Public/ControlRigControlActor.h
@@ -22,10 +22,12 @@ class CONTROLRIG_API AControlRigControlActor : public AActor
 public:
 
 	AControlRigControlActor(const FObjectInitializer& ObjectInitializer);
-
+	~AControlRigControlActor();
 	// AACtor overrides
 	virtual void Tick(float DeltaSeconds) override;
 	virtual void BeginPlay() override { Clear(); }
+	virtual void EndPlay(const EEndPlayReason::Type EndPlayReason) override;
+
 
 #if WITH_EDITOR
 	virtual void PostEditChangeProperty(FPropertyChangedEvent& PropertyChangedEvent) override;
@@ -82,4 +84,8 @@ private:
 
 	UPROPERTY(transient)
 	FName ColorParameterName;
+
+private:
+	void RemoveUnbindDelegate();
+	FDelegateHandle OnUnbindDelegate;
 };
\ No newline at end of file
