Control Rig: Moving some settings to Editor Preference

diff --git a/Engine/Plugins/Experimental/ControlRig/Source/ControlRig/Private/Settings/ControlRigSettings.cpp b/Engine/Plugins/Experimental/ControlRig/Source/ControlRig/Private/Settings/ControlRigSettings.cpp
--- a/Engine/Plugins/Experimental/ControlRig/Source/ControlRig/Private/Settings/ControlRigSettings.cpp
+++ b/Engine/Plugins/Experimental/ControlRig/Source/ControlRig/Private/Settings/ControlRigSettings.cpp
@@ -9,12 +9,19 @@
 
 UControlRigSettings::UControlRigSettings(const FObjectInitializer& ObjectInitializer)
 	: Super(ObjectInitializer)
+{
+#if WITH_EDITORONLY_DATA
+	DefaultGizmoLibrary = LoadObject<UControlRigGizmoLibrary>(nullptr, TEXT("/ControlRig/Controls/DefaultGizmoLibrary.DefaultGizmoLibrary"));
+#endif
+}
+
+UControlRigEditorSettings::UControlRigEditorSettings(const FObjectInitializer& ObjectInitializer)
+    : Super(ObjectInitializer)
 #if WITH_EDITORONLY_DATA
-	, bResetControlTransformsOnCompile(true)
+    , bResetControlTransformsOnCompile(true)
 #endif
 {
 #if WITH_EDITORONLY_DATA
-	DefaultGizmoLibrary = LoadObject<UControlRigGizmoLibrary>(nullptr, TEXT("/ControlRig/Controls/DefaultGizmoLibrary.DefaultGizmoLibrary"));
 	bResetControlsOnCompile = true;
 	bResetControlsOnPinValueInteraction = false;
 	bEnableUndoForPoseInteraction = true;
@@ -22,6 +29,5 @@ UControlRigSettings::UControlRigSettings(const FObjectInitializer& ObjectInitial
 	SetupEventBorderColor = FLinearColor::Red;
 	BackwardsSolveBorderColor = FLinearColor::Yellow;
 	BackwardsAndForwardsBorderColor = FLinearColor::Blue;
-
 #endif
 }
diff --git a/Engine/Plugins/Experimental/ControlRig/Source/ControlRig/Public/Settings/ControlRigSettings.h b/Engine/Plugins/Experimental/ControlRig/Source/ControlRig/Public/Settings/ControlRigSettings.h
--- a/Engine/Plugins/Experimental/ControlRig/Source/ControlRig/Public/Settings/ControlRigSettings.h
+++ b/Engine/Plugins/Experimental/ControlRig/Source/ControlRig/Public/Settings/ControlRigSettings.h
@@ -40,7 +40,24 @@ public:
 #if WITH_EDITORONLY_DATA
 	UPROPERTY(EditAnywhere, config, Category = DefaultGizmo)
 	TSoftObjectPtr<UControlRigGizmoLibrary> DefaultGizmoLibrary;
+#endif
+	
+	static UControlRigSettings * Get() { return GetMutableDefault<UControlRigSettings>(); }
+};
 
+/**
+ * Customize Control Rig Editor.
+ */
+UCLASS(config = EditorPerProjectUserSettings, meta=(DisplayName="Control Rig Editor"))
+class CONTROLRIG_API UControlRigEditorSettings : public UDeveloperSettings
+{
+	GENERATED_UCLASS_BODY()
+
+public:
+	/** Gets the category for the settings, some high level grouping like, Editor, Engine, Game...etc. */
+	virtual FName GetCategoryName() const override { return TEXT("ContentEditors"); }
+	
+#if WITH_EDITORONLY_DATA
 	// When this is checked all controls will return to their initial
 	// value as the user hits the Compile button.
 	UPROPERTY(EditAnywhere, config, Category = Interaction)
@@ -149,5 +166,7 @@ public:
 
 #endif
 
-	static UControlRigSettings * Get() { return GetMutableDefault<UControlRigSettings>(); }
+	static UControlRigEditorSettings * Get() { return GetMutableDefault<UControlRigEditorSettings>(); }
 };
+
+
diff --git a/Engine/Plugins/Experimental/ControlRig/Source/ControlRigDeveloper/Private/ControlRigBlueprint.cpp b/Engine/Plugins/Experimental/ControlRig/Source/ControlRigDeveloper/Private/ControlRigBlueprint.cpp
--- a/Engine/Plugins/Experimental/ControlRig/Source/ControlRigDeveloper/Private/ControlRigBlueprint.cpp
+++ b/Engine/Plugins/Experimental/ControlRig/Source/ControlRigDeveloper/Private/ControlRigBlueprint.cpp
@@ -1721,7 +1721,7 @@ FName UControlRigBlueprint::AddMemberVariable(const FName& InName, const FString
 FName UControlRigBlueprint::AddTransientControl(URigVMPin* InPin)
 {
 	TUniquePtr<FControlValueScope> ValueScope;
-	if (!UControlRigSettings::Get()->bResetControlsOnPinValueInteraction) // if we need to retain the controls
+	if (!UControlRigEditorSettings::Get()->bResetControlsOnPinValueInteraction) // if we need to retain the controls
 	{
 		ValueScope = MakeUnique<FControlValueScope>(this);
 	}
@@ -1771,7 +1771,7 @@ FName UControlRigBlueprint::AddTransientControl(URigVMPin* InPin)
 FName UControlRigBlueprint::RemoveTransientControl(URigVMPin* InPin)
 {
 	TUniquePtr<FControlValueScope> ValueScope;
-	if (!UControlRigSettings::Get()->bResetControlsOnPinValueInteraction) // if we need to retain the controls
+	if (!UControlRigEditorSettings::Get()->bResetControlsOnPinValueInteraction) // if we need to retain the controls
 	{
 		ValueScope = MakeUnique<FControlValueScope>(this);
 	}
@@ -1801,7 +1801,7 @@ FName UControlRigBlueprint::RemoveTransientControl(URigVMPin* InPin)
 FName UControlRigBlueprint::AddTransientControl(const FRigElementKey& InElement)
 {
 	TUniquePtr<FControlValueScope> ValueScope;
-	if (!UControlRigSettings::Get()->bResetControlsOnPinValueInteraction) // if we need to retain the controls
+	if (!UControlRigEditorSettings::Get()->bResetControlsOnPinValueInteraction) // if we need to retain the controls
 	{
 		ValueScope = MakeUnique<FControlValueScope>(this);
 	}
@@ -1860,7 +1860,7 @@ FName UControlRigBlueprint::AddTransientControl(const FRigElementKey& InElement)
 FName UControlRigBlueprint::RemoveTransientControl(const FRigElementKey& InElement)
 {
 	TUniquePtr<FControlValueScope> ValueScope;
-	if (!UControlRigSettings::Get()->bResetControlsOnPinValueInteraction) // if we need to retain the controls
+	if (!UControlRigEditorSettings::Get()->bResetControlsOnPinValueInteraction) // if we need to retain the controls
 	{
 		ValueScope = MakeUnique<FControlValueScope>(this);
 	}
@@ -1890,7 +1890,7 @@ FName UControlRigBlueprint::RemoveTransientControl(const FRigElementKey& InEleme
 void UControlRigBlueprint::ClearTransientControls()
 {
 	TUniquePtr<FControlValueScope> ValueScope;
-	if (!UControlRigSettings::Get()->bResetControlsOnPinValueInteraction) // if we need to retain the controls
+	if (!UControlRigEditorSettings::Get()->bResetControlsOnPinValueInteraction) // if we need to retain the controls
 	{
 		ValueScope = MakeUnique<FControlValueScope>(this);
 	}
diff --git a/Engine/Plugins/Experimental/ControlRig/Source/ControlRigEditor/Private/ControlRigEditorModule.cpp b/Engine/Plugins/Experimental/ControlRig/Source/ControlRigEditor/Private/ControlRigEditorModule.cpp
--- a/Engine/Plugins/Experimental/ControlRig/Source/ControlRigEditor/Private/ControlRigEditorModule.cpp
+++ b/Engine/Plugins/Experimental/ControlRig/Source/ControlRigEditor/Private/ControlRigEditorModule.cpp
@@ -1856,9 +1856,9 @@ void FControlRigEditorModule::GetContextMenuActions(const UControlRigGraphSchema
 #if WITH_EDITORONLY_DATA
 
 								FScopedTransaction Transaction(LOCTEXT("RigUnitDefaultExpansionStateChanged", "Changed Rig Unit Default Expansion State"));
-								UControlRigSettings::Get()->Modify();
+								UControlRigEditorSettings::Get()->Modify();
 
-								FControlRigSettingsPerPinBool& ExpansionMap = UControlRigSettings::Get()->RigUnitPinExpansion.FindOrAdd(UnitNode->GetScriptStruct()->GetName());
+								FControlRigSettingsPerPinBool& ExpansionMap = UControlRigEditorSettings::Get()->RigUnitPinExpansion.FindOrAdd(UnitNode->GetScriptStruct()->GetName());
 								ExpansionMap.Values.Empty();
 
 								TArray<URigVMPin*> Pins = UnitNode->GetAllPinsRecursively();
diff --git a/Engine/Plugins/Experimental/ControlRig/Source/ControlRigEditor/Private/EditMode/ControlRigEditMode.cpp b/Engine/Plugins/Experimental/ControlRig/Source/ControlRigEditor/Private/EditMode/ControlRigEditMode.cpp
--- a/Engine/Plugins/Experimental/ControlRig/Source/ControlRigEditor/Private/EditMode/ControlRigEditMode.cpp
+++ b/Engine/Plugins/Experimental/ControlRig/Source/ControlRigEditor/Private/EditMode/ControlRigEditMode.cpp
@@ -747,7 +747,7 @@ bool FControlRigEditMode::StartTracking(FEditorViewportClient* InViewportClient,
 					InteractionScope = new FControlRigInteractionScope(ControlRig);
 				}
 			}
-			else if(UControlRigSettings::Get()->bEnableUndoForPoseInteraction)
+			else if(UControlRigEditorSettings::Get()->bEnableUndoForPoseInteraction)
 			{
 				InteractionScope = new FControlRigInteractionScope(ControlRig);
 			}
diff --git a/Engine/Plugins/Experimental/ControlRig/Source/ControlRigEditor/Private/Editor/ControlRigEditor.cpp b/Engine/Plugins/Experimental/ControlRig/Source/ControlRigEditor/Private/Editor/ControlRigEditor.cpp
--- a/Engine/Plugins/Experimental/ControlRig/Source/ControlRigEditor/Private/Editor/ControlRigEditor.cpp
+++ b/Engine/Plugins/Experimental/ControlRig/Source/ControlRigEditor/Private/Editor/ControlRigEditor.cpp
@@ -1628,7 +1628,7 @@ void FControlRigEditor::Compile()
 		DECLARE_SCOPE_HIERARCHICAL_COUNTER_FUNC()
 
 		TUniquePtr<UControlRigBlueprint::FControlValueScope> ValueScope;
-		if (!UControlRigSettings::Get()->bResetControlsOnCompile) // if we need to retain the controls
+		if (!UControlRigEditorSettings::Get()->bResetControlsOnCompile) // if we need to retain the controls
 		{
 			ValueScope = MakeUnique<UControlRigBlueprint::FControlValueScope>(GetControlRigBlueprint());
 		}
@@ -1785,7 +1785,7 @@ void FControlRigEditor::Compile()
 			}
 		}
 
-		if (UControlRigSettings::Get()->bResetControlTransformsOnCompile)
+		if (UControlRigEditorSettings::Get()->bResetControlTransformsOnCompile)
 		{
 			RigBlueprint->Hierarchy->ForEach<FRigControlElement>([RigBlueprint](FRigControlElement* ControlElement) -> bool
             {
@@ -3183,13 +3183,13 @@ void FControlRigEditor::HandleViewportCreated(const TSharedRef<class IPersonaVie
 		switch (GetEventQueue())
 		{
 			case EControlRigEditorEventQueue::Setup:
-				Color = UControlRigSettings::Get()->SetupEventBorderColor;
+				Color = UControlRigEditorSettings::Get()->SetupEventBorderColor;
 				break;
 			case EControlRigEditorEventQueue::Inverse:
-				Color = UControlRigSettings::Get()->BackwardsSolveBorderColor;
+				Color = UControlRigEditorSettings::Get()->BackwardsSolveBorderColor;
 				break;
 			case EControlRigEditorEventQueue::InverseAndUpdate:
-				Color = UControlRigSettings::Get()->BackwardsAndForwardsBorderColor;
+				Color = UControlRigEditorSettings::Get()->BackwardsAndForwardsBorderColor;
 				break;
 			default:
 				Color = FLinearColor::Transparent;
@@ -5789,12 +5789,12 @@ void FControlRigEditor::StoreNodeSnippet(int32 InSnippetIndex)
 	FString* Setting = GetSnippetStorage(InSnippetIndex);
 	check(Setting != nullptr);
 
-	UControlRigSettings* Settings = UControlRigSettings::Get();
+	UControlRigEditorSettings* Settings = UControlRigEditorSettings::Get();
 	Settings->Modify();
 	(*Setting) = Snippet;
 
 	const FString PropertyName = FString::Printf(TEXT("NodeSnippet_%d"), InSnippetIndex);
-	FProperty* Property = UControlRigSettings::StaticClass()->FindPropertyByName(*PropertyName);
+	FProperty* Property = UControlRigEditorSettings::StaticClass()->FindPropertyByName(*PropertyName);
 	check(Property);
 	
 	Settings->UpdateSinglePropertyInConfigFile(Property, Settings->GetDefaultConfigFilename());
@@ -5838,7 +5838,7 @@ void FControlRigEditor::RestoreNodeSnippet(int32 InSnippetIndex)
 
 FString* FControlRigEditor::GetSnippetStorage(int32 InSnippetIndex)
 {
-	UControlRigSettings* Settings = UControlRigSettings::Get();
+	UControlRigEditorSettings* Settings = UControlRigEditorSettings::Get();
 	switch(InSnippetIndex)
 	{
 		case 1:
diff --git a/Engine/Plugins/Experimental/ControlRig/Source/ControlRigEditor/Private/Graph/NodeSpawners/ControlRigUnitNodeSpawner.cpp b/Engine/Plugins/Experimental/ControlRig/Source/ControlRigEditor/Private/Graph/NodeSpawners/ControlRigUnitNodeSpawner.cpp
--- a/Engine/Plugins/Experimental/ControlRig/Source/ControlRigEditor/Private/Graph/NodeSpawners/ControlRigUnitNodeSpawner.cpp
+++ b/Engine/Plugins/Experimental/ControlRig/Source/ControlRigEditor/Private/Graph/NodeSpawners/ControlRigUnitNodeSpawner.cpp
@@ -160,7 +160,7 @@ UControlRigGraphNode* UControlRigUnitNodeSpawner::SpawnNode(UEdGraph* ParentGrap
 #if WITH_EDITORONLY_DATA
 			if (!bIsTemplateNode)
 			{
-				const FControlRigSettingsPerPinBool* ExpansionMapPtr = UControlRigSettings::Get()->RigUnitPinExpansion.Find(ModelNode->GetScriptStruct()->GetName());
+				const FControlRigSettingsPerPinBool* ExpansionMapPtr = UControlRigEditorSettings::Get()->RigUnitPinExpansion.Find(ModelNode->GetScriptStruct()->GetName());
 				if (ExpansionMapPtr)
 				{
 					const FControlRigSettingsPerPinBool& ExpansionMap = *ExpansionMapPtr;
