Control Rig Component: Fixed a typo

diff --git a/Engine/Plugins/Experimental/ControlRig/Source/ControlRig/Private/ControlRigComponent.cpp b/Engine/Plugins/Experimental/ControlRig/Source/ControlRig/Private/ControlRigComponent.cpp
--- a/Engine/Plugins/Experimental/ControlRig/Source/ControlRig/Private/ControlRigComponent.cpp
+++ b/Engine/Plugins/Experimental/ControlRig/Source/ControlRig/Private/ControlRigComponent.cpp
@@ -478,7 +478,7 @@ void UControlRigComponent::AddMappedElements(TArray<FControlRigComponentMappedEl
 
 void UControlRigComponent::AddMappedComponents(TArray<FControlRigComponentMappedComponent> Components)
 {
-	if (!EnsureCalledOutsideOfBracket(TEXT("AddMapprepedComponents")))
+	if (!EnsureCalledOutsideOfBracket(TEXT("AddMappedComponents")))
 	{
 		return;
 	}
