Control Rig: Enable python dependency by default

diff --git a/Engine/Plugins/Experimental/ControlRig/ControlRig.uplugin b/Engine/Plugins/Experimental/ControlRig/ControlRig.uplugin
--- a/Engine/Plugins/Experimental/ControlRig/ControlRig.uplugin
+++ b/Engine/Plugins/Experimental/ControlRig/ControlRig.uplugin
@@ -16,6 +16,13 @@
 	"Installed" : false,
 	"Plugins":
 	[
+		{
+			"Name": "PythonScriptPlugin",
+			"Enabled": true,
+			"WhitelistTargets": [
+				"Editor"
+			]
+		},
 		{
 			"Name": "LevelSequenceEditor",
 			"Enabled": true
