Control Rig: Ensure to maintain default values on complex structures

diff --git a/Engine/Plugins/Experimental/ControlRig/Source/ControlRigEditor/Private/ControlRigElementDetails.cpp b/Engine/Plugins/Experimental/ControlRig/Source/ControlRigEditor/Private/ControlRigElementDetails.cpp
--- a/Engine/Plugins/Experimental/ControlRig/Source/ControlRigEditor/Private/ControlRigElementDetails.cpp
+++ b/Engine/Plugins/Experimental/ControlRig/Source/ControlRigEditor/Private/ControlRigElementDetails.cpp
@@ -505,6 +505,13 @@ void FRigUnitDetails::CustomizeDetails(IDetailLayoutBuilder& DetailBuilder)
 				continue;
 			}
 		}
+		else if (FStructProperty* StructProperty = CastField<FStructProperty>(Property))
+		{
+			const FSimpleDelegate OnStructContentsChangedDelegate = FSimpleDelegate::CreateSP(this, &FRigUnitDetails::OnStructContentsChanged, Property, DetailBuilder.GetPropertyUtilities());
+			PropertyHandle->SetOnPropertyValueChanged(OnStructContentsChangedDelegate);
+			PropertyHandle->SetOnChildPropertyValueChanged(OnStructContentsChangedDelegate);
+		}
+
 		CategoryBuilder.AddProperty(PropertyHandle);
 	}
 }
@@ -565,6 +572,12 @@ void FRigUnitDetails::OnNameListComboBox(TSharedPtr<FStructOnScope> InStructOnSc
 	Widget->SetSelectedItem(CurrentlySelected);
 }
 
+void FRigUnitDetails::OnStructContentsChanged(FProperty* InProperty, const TSharedRef<IPropertyUtilities> PropertyUtilities)
+{
+	const FPropertyChangedEvent ChangeEvent(InProperty, EPropertyChangeType::ValueSet);
+	PropertyUtilities->NotifyFinishedChangingProperties(ChangeEvent);
+}
+
 void FRigComputedTransformDetails::CustomizeHeader(TSharedRef<IPropertyHandle> InStructPropertyHandle, FDetailWidgetRow& HeaderRow, IPropertyTypeCustomizationUtils& StructCustomizationUtils)
 {
 	BlueprintBeingCustomized = nullptr;
diff --git a/Engine/Plugins/Experimental/ControlRig/Source/ControlRigEditor/Private/ControlRigElementDetails.h b/Engine/Plugins/Experimental/ControlRig/Source/ControlRigEditor/Private/ControlRigElementDetails.h
--- a/Engine/Plugins/Experimental/ControlRig/Source/ControlRigEditor/Private/ControlRigElementDetails.h
+++ b/Engine/Plugins/Experimental/ControlRig/Source/ControlRigEditor/Private/ControlRigElementDetails.h
@@ -75,6 +75,7 @@ protected:
 	void SetNameListText(const FText& NewTypeInValue, ETextCommit::Type /*CommitInfo*/, TSharedPtr<FStructOnScope> InStructOnScope, FNameProperty* InProperty, TSharedRef<IPropertyUtilities> PropertyUtilities);
 	void OnNameListChanged(TSharedPtr<FString> NewSelection, ESelectInfo::Type SelectInfo, TSharedPtr<FStructOnScope> InStructOnScope, FNameProperty* InProperty, TSharedRef<IPropertyUtilities> PropertyUtilities);
 	void OnNameListComboBox(TSharedPtr<FStructOnScope> InStructOnScope, FNameProperty* InProperty, const TArray<TSharedPtr<FString>>* InNameList);
+	void OnStructContentsChanged(FProperty* InProperty, const TSharedRef<IPropertyUtilities> PropertyUtilities);
 
 	UControlRigBlueprint* BlueprintBeingCustomized;
 	UControlRigGraph* GraphBeingCustomized;
diff --git a/Engine/Source/Runtime/RigVM/Private/RigVMCore/RigVMStruct.cpp b/Engine/Source/Runtime/RigVM/Private/RigVMCore/RigVMStruct.cpp
--- a/Engine/Source/Runtime/RigVM/Private/RigVMCore/RigVMStruct.cpp
+++ b/Engine/Source/Runtime/RigVM/Private/RigVMCore/RigVMStruct.cpp
@@ -434,6 +434,11 @@ FString FRigVMStruct::ExportToFullyQualifiedText(UScriptStruct* InStruct, const
 	TArray<FString> FieldValues;
 	for (TFieldIterator<FProperty> It(InStruct); It; ++It)
 	{
+		if(It->HasAnyPropertyFlags(CPF_Transient))
+		{
+			continue;
+		}
+		
 		FString PropertyName = It->GetName();
 		const uint8* StructMemberMemoryPtr = It->ContainerPtrToValuePtr<uint8>(InStructMemoryPtr);
 		FString DefaultValue = ExportToFullyQualifiedText(*It, StructMemberMemoryPtr);
