Control Rig: Maintain content of default values for variables

diff --git a/Engine/Plugins/Experimental/ControlRig/Source/ControlRigDeveloper/Private/ControlRigBlueprint.cpp b/Engine/Plugins/Experimental/ControlRig/Source/ControlRigDeveloper/Private/ControlRigBlueprint.cpp
--- a/Engine/Plugins/Experimental/ControlRig/Source/ControlRigDeveloper/Private/ControlRigBlueprint.cpp
+++ b/Engine/Plugins/Experimental/ControlRig/Source/ControlRigDeveloper/Private/ControlRigBlueprint.cpp
@@ -3161,10 +3161,6 @@ void UControlRigBlueprint::OnPostVariableChange(UBlueprint* InBlueprint)
 	TMap<FGuid, int32> NewVariablesByGuid;
 	for (int32 VarIndex = 0; VarIndex < NewVariables.Num(); VarIndex++)
 	{
-		// we use the storage within the CDO for the default values,
-		// no need to maintain the default value as a string
-		NewVariables[VarIndex].DefaultValue = FString();
-
 		NewVariablesByGuid.Add(NewVariables[VarIndex].VarGuid, VarIndex);
 	}
 
diff --git a/Engine/Plugins/Experimental/ControlRig/Source/ControlRigEditor/Private/Editor/ControlRigEditor.cpp b/Engine/Plugins/Experimental/ControlRig/Source/ControlRigEditor/Private/Editor/ControlRigEditor.cpp
--- a/Engine/Plugins/Experimental/ControlRig/Source/ControlRigEditor/Private/Editor/ControlRigEditor.cpp
+++ b/Engine/Plugins/Experimental/ControlRig/Source/ControlRigEditor/Private/Editor/ControlRigEditor.cpp
@@ -1749,24 +1749,7 @@ void FControlRigEditor::Compile()
 		// store the defaults from the CDO back on the new variables list
 		for(FBPVariableDescription& NewVariable : RigBlueprint->NewVariables)
 		{
-			UClass* GeneratedClass = RigBlueprint->GeneratedClass;
-			UObject* GeneratedCDO = GeneratedClass->GetDefaultObject();
-			FProperty* TargetProperty = FindFProperty<FProperty>(GeneratedClass, NewVariable.VarName);
-
-			if (TargetProperty)
-			{
-				// Grab the address of where the property is actually stored (UObject* base, plus the offset defined in the property)
-				void* OldPropertyAddr = TargetProperty->ContainerPtrToValuePtr<void>(GeneratedCDO);
-				if (OldPropertyAddr)
-				{
-					FString NewDefaultValue;
-					TargetProperty->ExportTextItem(NewDefaultValue, OldPropertyAddr, OldPropertyAddr, nullptr, PPF_SerializedAsImportText);
-					if(NewDefaultValue.Len() < 80)
-					{
-						NewVariable.DefaultValue = NewDefaultValue;
-					}
-				}
-			}
+			UpdateDefaultValueForVariable(NewVariable, true);
 		}
 
 		if(SelectedKey.IsValid())
@@ -4043,6 +4026,20 @@ void FControlRigEditor::OnFinishedChangingProperties(const FPropertyChangedEvent
 			ControlRigBP->MarkPackageDirty();
 		}
 	}
+
+	// this might be a property value for a variable
+	if(ControlRig && ControlRigBP)
+	{
+		const FName VarName = PropertyChangedEvent.MemberProperty->GetFName();
+		for(FBPVariableDescription& NewVariable : ControlRigBP->NewVariables)
+		{
+			if(NewVariable.VarName == VarName)
+			{
+				UpdateDefaultValueForVariable(NewVariable, false);
+				break;
+			}
+		}
+	}
 }
 
 void FControlRigEditor::OnBlueprintPropertyChainEvent(FPropertyChangedChainEvent& PropertyChangedChainEvent)
@@ -4230,6 +4227,34 @@ bool FControlRigEditor::OnRequestBulkEditDialog(UControlRigBlueprint* InBlueprin
 	return BulkEditDialog->ShowModal() != EAppReturnType::Cancel;
 }
 
+void FControlRigEditor::UpdateDefaultValueForVariable(FBPVariableDescription& InVariable, bool bUseCDO)
+{
+	if (UControlRigBlueprint* ControlRigBP = GetControlRigBlueprint())
+	{
+		UClass* GeneratedClass = ControlRigBP->GeneratedClass;
+		UObject* ObjectContainer = bUseCDO ? GeneratedClass->GetDefaultObject() : ControlRigBP->GetObjectBeingDebugged();
+		if(ObjectContainer)
+		{
+			FProperty* TargetProperty = FindFProperty<FProperty>(GeneratedClass, InVariable.VarName);
+
+			if (TargetProperty)
+			{
+				// Grab the address of where the property is actually stored (UObject* base, plus the offset defined in the property)
+				void* OldPropertyAddr = TargetProperty->ContainerPtrToValuePtr<void>(ObjectContainer);
+				if (OldPropertyAddr)
+				{
+					FString NewDefaultValue;
+					TargetProperty->ExportTextItem(NewDefaultValue, OldPropertyAddr, OldPropertyAddr, nullptr, PPF_None);
+					if(NewDefaultValue.Len() < 80)
+					{
+						InVariable.DefaultValue = NewDefaultValue;
+					}
+				}
+			}
+		}
+	}
+}
+
 void FControlRigEditor::OnCreateComment()
 {
 	TSharedPtr<SGraphEditor> GraphEditor = FocusedGraphEdPtr.Pin();
diff --git a/Engine/Plugins/Experimental/ControlRig/Source/ControlRigEditor/Private/Editor/ControlRigEditor.h b/Engine/Plugins/Experimental/ControlRig/Source/ControlRigEditor/Private/Editor/ControlRigEditor.h
--- a/Engine/Plugins/Experimental/ControlRig/Source/ControlRigEditor/Private/Editor/ControlRigEditor.h
+++ b/Engine/Plugins/Experimental/ControlRig/Source/ControlRigEditor/Private/Editor/ControlRigEditor.h
@@ -391,6 +391,7 @@ protected:
 	void OnBlueprintPropertyChainEvent(FPropertyChangedChainEvent& PropertyChangedChainEvent);
 	void OnRequestLocalizeFunctionDialog(URigVMLibraryNode* InFunction, UControlRigBlueprint* InTargetBlueprint, bool bForce);
 	bool OnRequestBulkEditDialog(UControlRigBlueprint* InBlueprint, URigVMController* InController, URigVMLibraryNode* InFunction, ERigVMControllerBulkEditType InEditType);
+	void UpdateDefaultValueForVariable(FBPVariableDescription& InVariable, bool bUseCDO);
 
 	URigVMGraph* GetFocusedModel() const;
 	URigVMController* GetFocusedController() const;
