Control Rig: Prevent crash when closing and opening editor because of add breakpoint callback

diff --git a/Engine/Plugins/Experimental/ControlRig/Source/ControlRigEditor/Private/Editor/ControlRigEditor.cpp b/Engine/Plugins/Experimental/ControlRig/Source/ControlRigEditor/Private/Editor/ControlRigEditor.cpp
--- a/Engine/Plugins/Experimental/ControlRig/Source/ControlRigEditor/Private/Editor/ControlRigEditor.cpp
+++ b/Engine/Plugins/Experimental/ControlRig/Source/ControlRigEditor/Private/Editor/ControlRigEditor.cpp
@@ -205,7 +205,7 @@ FControlRigEditor::FControlRigEditor()
 	, bFirstTimeSelecting(true)
 	, bAnyErrorsLeft(false)
 	, LastEventQueue(EControlRigEditorEventQueue::Setup)
-	, bIsInDebugMode(false)
+	, ExecutionMode(EControlRigExecutionModeType::EControlRigExecutionModeType_Release)
 	, LastDebuggedRig()
 	, RigHierarchyTabCount(0)
 	, HaltedAtNode(nullptr)
@@ -455,10 +455,7 @@ void FControlRigEditor::InitControlRigEditor(const EToolkitMode::Type Mode, cons
 		InControlRigBlueprint->Hierarchy->OnModified().AddSP(this, &FControlRigEditor::OnHierarchyModified);
 		InControlRigBlueprint->OnRefreshEditor().AddSP(this, &FControlRigEditor::HandleRefreshEditorFromBlueprint);
 		InControlRigBlueprint->OnVariableDropped().AddSP(this, &FControlRigEditor::HandleVariableDroppedFromBlueprint);
-		InControlRigBlueprint->OnBreakpointAdded().AddLambda([this]()
-		{
-			SetExecutionMode(true);
-		});
+		InControlRigBlueprint->OnBreakpointAdded().AddSP(this, &FControlRigEditor::HandleBreakpointAdded);
 
 		if (FControlRigEditMode* EditMode = GetEditMode())
 		{
@@ -565,12 +562,12 @@ void FControlRigEditor::BindCommands()
 
 	GetToolkitCommands()->MapAction(
 		FControlRigBlueprintCommands::Get().ReleaseMode,
-		FExecuteAction::CreateSP(this, &FControlRigEditor::SetExecutionMode, false),
+		FExecuteAction::CreateSP(this, &FControlRigEditor::SetExecutionMode, EControlRigExecutionModeType_Release),
 		FCanExecuteAction());
 
 	GetToolkitCommands()->MapAction(
 		FControlRigBlueprintCommands::Get().DebugMode,
-		FExecuteAction::CreateSP(this, &FControlRigEditor::SetExecutionMode, true),
+		FExecuteAction::CreateSP(this, &FControlRigEditor::SetExecutionMode, EControlRigExecutionModeType_Debug),
 		FCanExecuteAction());
 
 	GetToolkitCommands()->MapAction(
@@ -640,7 +637,9 @@ void FControlRigEditor::ToggleEventQueue()
 
 void FControlRigEditor::ToggleExecutionMode()
 {
-	SetExecutionMode(!bIsInDebugMode);
+	SetExecutionMode((ExecutionMode == EControlRigExecutionModeType_Debug) ?
+		EControlRigExecutionModeType_Release
+		: EControlRigExecutionModeType_Debug);
 }
 
 TSharedRef<SWidget> FControlRigEditor::GenerateEventQueueMenuContent()
@@ -665,8 +664,8 @@ TSharedRef<SWidget> FControlRigEditor::GenerateExecutionModeMenuContent()
 	FMenuBuilder MenuBuilder(true, GetToolkitCommands());
 
 	MenuBuilder.BeginSection(TEXT("Events"));
-	MenuBuilder.AddMenuEntry(FControlRigBlueprintCommands::Get().ReleaseMode, TEXT("Release"), TAttribute<FText>(), TAttribute<FText>(), GetExecutionModeIcon(false));
-	MenuBuilder.AddMenuEntry(FControlRigBlueprintCommands::Get().DebugMode, TEXT("Debug"), TAttribute<FText>(), TAttribute<FText>(), GetExecutionModeIcon(true));
+	MenuBuilder.AddMenuEntry(FControlRigBlueprintCommands::Get().ReleaseMode, TEXT("Release"), TAttribute<FText>(), TAttribute<FText>(), GetExecutionModeIcon(EControlRigExecutionModeType_Release));
+	MenuBuilder.AddMenuEntry(FControlRigBlueprintCommands::Get().DebugMode, TEXT("Debug"), TAttribute<FText>(), TAttribute<FText>(), GetExecutionModeIcon(EControlRigExecutionModeType_Debug));
 	MenuBuilder.EndSection();
 
 	return MenuBuilder.MakeWidget();
@@ -939,20 +938,20 @@ FSlateIcon FControlRigEditor::GetEventQueueIcon() const
 	return GetEventQueueIcon(GetEventQueue());
 }
 
-void FControlRigEditor::SetExecutionMode(const bool bSetDebugMode)
+void FControlRigEditor::SetExecutionMode(const EControlRigExecutionModeType InExecutionMode)
 {
-	if (bIsInDebugMode == bSetDebugMode)
+	if (ExecutionMode == InExecutionMode)
 	{
 		return;
 	}
 
-	bIsInDebugMode = bSetDebugMode;
-	GetControlRigBlueprint()->SetDebugMode(bSetDebugMode);
+	ExecutionMode = InExecutionMode;
+	GetControlRigBlueprint()->SetDebugMode(InExecutionMode == EControlRigExecutionModeType_Debug);
 	Compile();
 	
 	if (ControlRig)
 	{
-		ControlRig->bIsInDebugMode = bSetDebugMode;
+		ControlRig->bIsInDebugMode = InExecutionMode == EControlRigExecutionModeType_Debug;
 	}
 
 	if (HaltedAtNode)
@@ -964,30 +963,30 @@ void FControlRigEditor::SetExecutionMode(const bool bSetDebugMode)
 
 int32 FControlRigEditor::GetExecutionModeComboValue() const
 {
-	return (int32) bIsInDebugMode;
+	return (int32) ExecutionMode;
 }
 
 FText FControlRigEditor::GetExecutionModeLabel() const
 {
-	if (bIsInDebugMode)
+	if (ExecutionMode == EControlRigExecutionModeType_Debug)
 	{
 		return FText::FromString(TEXT("DebugMode"));
 	}
 	return FText::FromString(TEXT("ReleaseMode"));
 }
 
-FSlateIcon FControlRigEditor::GetExecutionModeIcon(bool bDebugMode)
+FSlateIcon FControlRigEditor::GetExecutionModeIcon(EControlRigExecutionModeType InExecutionMode)
 {
-	if (bDebugMode)
+	if (InExecutionMode == EControlRigExecutionModeType_Debug)
 	{
-		return FSlateIcon(FControlRigEditorStyle::Get().GetStyleSetName(), "ControlRig.DebugMode");
+		return FSlateIcon(FEditorStyle::GetStyleSetName(), "LevelEditor.Tabs.Debug");
 	}
 	return FSlateIcon(FControlRigEditorStyle::Get().GetStyleSetName(), "ControlRig.ReleaseMode");
 }
 
 FSlateIcon FControlRigEditor::GetExecutionModeIcon() const
 {
-	return GetExecutionModeIcon(bIsInDebugMode);
+	return GetExecutionModeIcon(ExecutionMode);
 }
 
 void FControlRigEditor::ToggleSetupMode()
@@ -5461,6 +5460,11 @@ void FControlRigEditor::HandleVariableDroppedFromBlueprint(UObject* InSubject, F
 	);
 }
 
+void FControlRigEditor::HandleBreakpointAdded()
+{
+	SetExecutionMode(EControlRigExecutionModeType_Debug);
+}
+
 void FControlRigEditor::OnGraphNodeClicked(UControlRigGraphNode* InNode)
 {
 	if (InNode)
diff --git a/Engine/Plugins/Experimental/ControlRig/Source/ControlRigEditor/Private/Editor/ControlRigEditor.h b/Engine/Plugins/Experimental/ControlRig/Source/ControlRigEditor/Private/Editor/ControlRigEditor.h
--- a/Engine/Plugins/Experimental/ControlRig/Source/ControlRigEditor/Private/Editor/ControlRigEditor.h
+++ b/Engine/Plugins/Experimental/ControlRig/Source/ControlRigEditor/Private/Editor/ControlRigEditor.h
@@ -254,10 +254,16 @@ private:
 	static FSlateIcon GetEventQueueIcon(EControlRigEditorEventQueue InEventQueue);
 	FSlateIcon GetEventQueueIcon() const;
 
-	void SetExecutionMode(const bool bSetDebugMode);
+	enum EControlRigExecutionModeType
+	{
+		EControlRigExecutionModeType_Release,
+		EControlRigExecutionModeType_Debug
+	};
+
+	void SetExecutionMode(const EControlRigExecutionModeType InExecutionMode);
 	int32 GetExecutionModeComboValue() const;
 	FText GetExecutionModeLabel() const;
-	static FSlateIcon GetExecutionModeIcon(const bool bInIsDebugMode);
+	static FSlateIcon GetExecutionModeIcon(const EControlRigExecutionModeType InExecutionMode);
 	FSlateIcon GetExecutionModeIcon() const;
 	
 	virtual void GetCustomDebugObjects(TArray<FCustomDebugObject>& DebugList) const override;
@@ -330,6 +336,8 @@ private:
 
 	 void HandleVariableDroppedFromBlueprint(UObject* InSubject, FProperty* InVariableToDrop, const FVector2D& InDropPosition, const FVector2D& InScreenPosition);
 
+	 void HandleBreakpointAdded();
+
 	 void OnGraphNodeClicked(UControlRigGraphNode* InNode);
 
 	 void OnNodeDoubleClicked(UControlRigBlueprint* InBlueprint, URigVMNode* InNode);
@@ -443,7 +451,7 @@ protected:
 	bool bAnyErrorsLeft;
 
 	EControlRigEditorEventQueue LastEventQueue;
-	bool bIsInDebugMode;
+	EControlRigExecutionModeType ExecutionMode;
 	FString LastDebuggedRig;
 	int32 RigHierarchyTabCount;
 	TMap<FString, FString> KnownInstructionLimitWarnings;
