Control Rig Component: Also call ::Initialize during InitializeComponent, which is called during Spawn.

diff --git a/Engine/Plugins/Animation/ControlRig/Source/ControlRig/Private/ControlRigComponent.cpp b/Engine/Plugins/Animation/ControlRig/Source/ControlRig/Private/ControlRigComponent.cpp
--- a/Engine/Plugins/Animation/ControlRig/Source/ControlRig/Private/ControlRigComponent.cpp
+++ b/Engine/Plugins/Animation/ControlRig/Source/ControlRig/Private/ControlRigComponent.cpp
@@ -55,6 +55,7 @@ UControlRigComponent::UControlRigComponent(const FObjectInitializer& ObjectIniti
 	bDrawBones = true;
 	bShowDebugDrawing = false;
 	bIsInsideInitializeBracket = false;
+	bWantsInitializeComponent = true;
 }
 
 #if WITH_EDITOR
@@ -82,6 +83,11 @@ void UControlRigComponent::BeginDestroy()
 	gPendingSkeletalMeshes.Remove(this);
 }
 
+void UControlRigComponent::InitializeComponent()
+{
+	Initialize();
+}
+
 void UControlRigComponent::OnRegister()
 {
 	Super::OnRegister();
diff --git a/Engine/Plugins/Animation/ControlRig/Source/ControlRig/Public/ControlRigComponent.h b/Engine/Plugins/Animation/ControlRig/Source/ControlRig/Public/ControlRigComponent.h
--- a/Engine/Plugins/Animation/ControlRig/Source/ControlRig/Public/ControlRigComponent.h
+++ b/Engine/Plugins/Animation/ControlRig/Source/ControlRig/Public/ControlRigComponent.h
@@ -215,6 +215,7 @@ public:
 	//~ End UObject Interface
 	
 	//~ Begin UActorComponent Interface
+	virtual void InitializeComponent() override;
 	virtual void OnRegister() override;
 	virtual void OnUnregister() override;
 	virtual void TickComponent(float DeltaTime, enum ELevelTick TickType, FActorComponentTickFunction* ThisTickFunction) override;
