Control Rig: Fix kalman rig unit for rotation values

diff --git a/Engine/Plugins/Animation/ControlRig/Source/ControlRig/Private/Units/Simulation/RigUnit_Kalman.cpp b/Engine/Plugins/Animation/ControlRig/Source/ControlRig/Private/Units/Simulation/RigUnit_Kalman.cpp
--- a/Engine/Plugins/Animation/ControlRig/Source/ControlRig/Private/Units/Simulation/RigUnit_Kalman.cpp
+++ b/Engine/Plugins/Animation/ControlRig/Source/ControlRig/Private/Units/Simulation/RigUnit_Kalman.cpp
@@ -147,7 +147,6 @@ FRigUnit_KalmanTransform_Execute()
 			Location += F.GetLocation();
 			AxisX += F.TransformVectorNoScale(FVector(1.f, 0.f, 0.f));
 			AxisY += F.TransformVectorNoScale(FVector(0.f, 1.f, 0.f));
-			AxisY += F.GetLocation();
 			Scale += F.GetScale3D();
 			NumberValidEntries++;
 		}
