Control Rig: Disable selectability on GizmoActor StaticMeshComponent as needed

diff --git a/Engine/Plugins/Animation/ControlRig/Source/ControlRig/Private/ControlRigGizmoActor.cpp b/Engine/Plugins/Animation/ControlRig/Source/ControlRig/Private/ControlRigGizmoActor.cpp
--- a/Engine/Plugins/Animation/ControlRig/Source/ControlRig/Private/ControlRigGizmoActor.cpp
+++ b/Engine/Plugins/Animation/ControlRig/Source/ControlRig/Private/ControlRigGizmoActor.cpp
@@ -31,6 +31,7 @@ AControlRigGizmoActor::AControlRigGizmoActor(const FObjectInitializer& ObjectIni
 	StaticMeshComponent->SetupAttachment(RootComponent);
 	StaticMeshComponent->bCastStaticShadow = false;
 	StaticMeshComponent->bCastDynamicShadow = false;
+	StaticMeshComponent->bSelectable = bSelectable && bEnabled;
 }
 
 void AControlRigGizmoActor::SetEnabled(bool bInEnabled)
@@ -38,6 +39,7 @@ void AControlRigGizmoActor::SetEnabled(bool bInEnabled)
 	if(bEnabled != bInEnabled)
 	{
 		bEnabled = bInEnabled;
+		StaticMeshComponent->bSelectable = bSelectable && bEnabled;
 		FEditorScriptExecutionGuard Guard;
 		OnEnabledChanged(bEnabled);
 	}
@@ -68,6 +70,7 @@ void AControlRigGizmoActor::SetSelectable(bool bInSelectable)
 	if (bSelectable != bInSelectable)
 	{
 		bSelectable = bInSelectable;
+		StaticMeshComponent->bSelectable = bSelectable && bEnabled;
 		if (!bSelectable)
 		{
 			SetSelected(false);
