Control Rig: Avoid ensures when checking ListeningHierarchy weak pointers

diff --git a/Engine/Plugins/Animation/ControlRig/Source/ControlRig/Private/Rigs/RigHierarchy.cpp b/Engine/Plugins/Animation/ControlRig/Source/ControlRig/Private/Rigs/RigHierarchy.cpp
--- a/Engine/Plugins/Animation/ControlRig/Source/ControlRig/Private/Rigs/RigHierarchy.cpp
+++ b/Engine/Plugins/Animation/ControlRig/Source/ControlRig/Private/Rigs/RigHierarchy.cpp
@@ -937,7 +937,7 @@ bool URigHierarchy::SetParentWeight(FRigBaseElement* InChild, int32 InParentInde
 					}
 
 					URigHierarchy* ListeningHierarchy = Listener.Hierarchy.Get();
-					if (ensure(ListeningHierarchy))
+					if (ListeningHierarchy)
 					{
 						if(FRigBaseElement* ListeningElement = ListeningHierarchy->Find(InChild->GetKey()))
 						{
@@ -1420,7 +1420,7 @@ void URigHierarchy::SetTransform(FRigTransformElement* InTransformElement, const
 			}
 
 			URigHierarchy* ListeningHierarchy = Listener.Hierarchy.Get();
-			if (ensure(ListeningHierarchy))
+			if (ListeningHierarchy)
 			{			
 				if(FRigTransformElement* ListeningElement = Cast<FRigTransformElement>(ListeningHierarchy->Find(InTransformElement->GetKey())))
 				{
@@ -1517,7 +1517,7 @@ void URigHierarchy::SetControlOffsetTransform(FRigControlElement* InControlEleme
 		for(FRigHierarchyListener& Listener : ListeningHierarchies)
 		{
 			URigHierarchy* ListeningHierarchy = Listener.Hierarchy.Get();
-			if (ensure(ListeningHierarchy))
+			if (ListeningHierarchy)
 			{	
 				if(FRigControlElement* ListeningElement = Cast<FRigControlElement>(ListeningHierarchy->Find(InControlElement->GetKey())))
 				{
@@ -1614,7 +1614,7 @@ void URigHierarchy::SetControlGizmoTransform(FRigControlElement* InControlElemen
 		for(FRigHierarchyListener& Listener : ListeningHierarchies)
 		{
 			URigHierarchy* ListeningHierarchy = Listener.Hierarchy.Get();
-			if (ensure(ListeningHierarchy))
+			if (ListeningHierarchy)
 			{	
 				if(FRigControlElement* ListeningElement = Cast<FRigControlElement>(ListeningHierarchy->Find(InControlElement->GetKey())))
 				{
@@ -1651,7 +1651,7 @@ void URigHierarchy::SetControlSettings(FRigControlElement* InControlElement, FRi
 		for(FRigHierarchyListener& Listener : ListeningHierarchies)
 		{
 			URigHierarchy* ListeningHierarchy = Listener.Hierarchy.Get();
-			if (ensure(ListeningHierarchy))
+			if (ListeningHierarchy)
 			{	
 				if(FRigControlElement* ListeningElement = Cast<FRigControlElement>(ListeningHierarchy->Find(InControlElement->GetKey())))
 				{
@@ -1962,7 +1962,7 @@ void URigHierarchy::SetControlValue(FRigControlElement* InControlElement, const
 					{
 						URigHierarchy* ListeningHierarchy = Listener.Hierarchy.Get();
 					
-						if (ensure(ListeningHierarchy))
+						if (ListeningHierarchy)
 						{
 							if(FRigControlElement* ListeningElement = Cast<FRigControlElement>(ListeningHierarchy->Find(InControlElement->GetKey())))
 							{
@@ -1996,7 +1996,7 @@ void URigHierarchy::SetControlVisibility(FRigControlElement* InControlElement, b
 		for(FRigHierarchyListener& Listener : ListeningHierarchies)
 		{
 			URigHierarchy* ListeningHierarchy = Listener.Hierarchy.Get();
-			if (ensure(ListeningHierarchy))
+			if (ListeningHierarchy)
 			{
 				if(FRigControlElement* ListeningElement = Cast<FRigControlElement>(ListeningHierarchy->Find(InControlElement->GetKey())))
 				{
@@ -2050,7 +2050,7 @@ void URigHierarchy::SetCurveValue(FRigCurveElement* InCurveElement, float InValu
 			}
 
 			URigHierarchy* ListeningHierarchy = Listener.Hierarchy.Get();
-			if (ensure(ListeningHierarchy))
+			if (ListeningHierarchy)
 			{
 				if(FRigCurveElement* ListeningElement = Cast<FRigCurveElement>(ListeningHierarchy->Find(InCurveElement->GetKey())))
 				{
