Control Rig: Limit Scale on EulerTransform affects Location

diff --git a/Engine/Plugins/Animation/ControlRig/Source/ControlRig/Public/Rigs/RigHierarchyDefines.h b/Engine/Plugins/Animation/ControlRig/Source/ControlRig/Public/Rigs/RigHierarchyDefines.h
--- a/Engine/Plugins/Animation/ControlRig/Source/ControlRig/Public/Rigs/RigHierarchyDefines.h
+++ b/Engine/Plugins/Animation/ControlRig/Source/ControlRig/Public/Rigs/RigHierarchyDefines.h
@@ -705,7 +705,7 @@ public:
 				}
 				if (bLimitScale)
 				{
-					ValueRef.Location = FVector(
+					ValueRef.Scale = FVector(
 						Local::Clamp(ValueRef.Scale.X, Min.Scale.X, Max.Scale.X),
 						Local::Clamp(ValueRef.Scale.Y, Min.Scale.Y, Max.Scale.Y),
 						Local::Clamp(ValueRef.Scale.Z, Min.Scale.Z, Max.Scale.Z)
