Control Rig: Improve default limit ranges for Vector2D control

diff --git a/Engine/Plugins/Animation/ControlRig/Source/ControlRigEditor/Private/ControlRigElementDetails.cpp b/Engine/Plugins/Animation/ControlRig/Source/ControlRigEditor/Private/ControlRigElementDetails.cpp
--- a/Engine/Plugins/Animation/ControlRig/Source/ControlRigEditor/Private/ControlRigElementDetails.cpp
+++ b/Engine/Plugins/Animation/ControlRig/Source/ControlRigEditor/Private/ControlRigElementDetails.cpp
@@ -1441,7 +1441,7 @@ void FRigControlElementDetails::CustomizeDetails(IDetailLayoutBuilder& DetailBui
 						ValueToSet = FRigControlValue::Make<float>(0.f);
 						ControlElement->Settings.bLimitTranslation = true;
 						ControlElement->Settings.MinimumValue = FRigControlValue::Make<float>(0.f);
-						ControlElement->Settings.MaximumValue = FRigControlValue::Make<float>(1.f);
+						ControlElement->Settings.MaximumValue = FRigControlValue::Make<float>(100.f);
                         break;
 					}
 					case ERigControlType::Integer:
@@ -1449,7 +1449,7 @@ void FRigControlElementDetails::CustomizeDetails(IDetailLayoutBuilder& DetailBui
 						ValueToSet = FRigControlValue::Make<int32>(0);
 						ControlElement->Settings.bLimitTranslation = true;
 						ControlElement->Settings.MinimumValue = FRigControlValue::Make<int32>(0);
-						ControlElement->Settings.MaximumValue = FRigControlValue::Make<int32>(10);
+						ControlElement->Settings.MaximumValue = FRigControlValue::Make<int32>(100);
 						break;
 					}
 					case ERigControlType::Vector2D:
@@ -1457,7 +1457,7 @@ void FRigControlElementDetails::CustomizeDetails(IDetailLayoutBuilder& DetailBui
 						ValueToSet = FRigControlValue::Make<FVector2D>(FVector2D::ZeroVector);
 						ControlElement->Settings.bLimitTranslation = true;
 						ControlElement->Settings.MinimumValue = FRigControlValue::Make<FVector2D>(FVector2D::ZeroVector);
-						ControlElement->Settings.MaximumValue = FRigControlValue::Make<FVector2D>(FVector2D(1.f, 1.f));
+						ControlElement->Settings.MaximumValue = FRigControlValue::Make<FVector2D>(FVector2D(100.f, 100.f));
 						break;
 					}
 					case ERigControlType::Position:
