Control Rig: Fix crash when expanding collapsed/function node

diff --git a/Engine/Plugins/Animation/ControlRig/Source/ControlRigDeveloper/Private/Graph/ControlRigGraph.cpp b/Engine/Plugins/Animation/ControlRig/Source/ControlRigDeveloper/Private/Graph/ControlRigGraph.cpp
--- a/Engine/Plugins/Animation/ControlRig/Source/ControlRigDeveloper/Private/Graph/ControlRigGraph.cpp
+++ b/Engine/Plugins/Animation/ControlRig/Source/ControlRigDeveloper/Private/Graph/ControlRigGraph.cpp
@@ -253,7 +253,7 @@ void UControlRigGraph::HandleModifiedEvent(ERigVMGraphNotifType InNotifType, URi
 
 			for (URigVMNode* Node : InGraph->GetNodes())
 			{
-				UEdGraphNode* EdNode = FindNodeForModelNodeName(Node->GetFName());
+				UEdGraphNode* EdNode = FindNodeForModelNodeName(Node->GetFName(), false);
 				if (EdNode != nullptr)
 				{
 					RemoveNode(EdNode);
@@ -393,7 +393,7 @@ void UControlRigGraph::HandleModifiedEvent(ERigVMGraphNotifType InNotifType, URi
 					break;
 				}
 
-				UEdGraphNode* EdNode = FindNodeForModelNodeName(ModelNode->GetFName());
+				UEdGraphNode* EdNode = FindNodeForModelNodeName(ModelNode->GetFName(), false);
 				if (EdNode)
 				{
 					RemoveNode(EdNode, true);
@@ -743,7 +743,7 @@ int32 UControlRigGraph::GetInstructionIndex(UControlRigGraphNode* InNode)
 	return INDEX_NONE;
 }
 
-UEdGraphNode* UControlRigGraph::FindNodeForModelNodeName(const FName& InModelNodeName)
+UEdGraphNode* UControlRigGraph::FindNodeForModelNodeName(const FName& InModelNodeName, const bool bCacheIfRequired)
 {
 	DECLARE_SCOPE_HIERARCHICAL_COUNTER_FUNC()
 
@@ -759,7 +759,10 @@ UEdGraphNode* UControlRigGraph::FindNodeForModelNodeName(const FName& InModelNod
 		{
 			if (RigNode->ModelNodePath == InModelNodePath)
 			{
-				ModelNodePathToEdNode.Add(InModelNodeName, EdNode);
+				if (bCacheIfRequired)
+				{
+					ModelNodePathToEdNode.Add(InModelNodeName, EdNode);
+				}
 				return EdNode;
 			}
 		}
@@ -767,7 +770,10 @@ UEdGraphNode* UControlRigGraph::FindNodeForModelNodeName(const FName& InModelNod
 		{
 			if (EdNode->GetFName() == InModelNodeName)
 			{
-				ModelNodePathToEdNode.Add(InModelNodeName, EdNode);
+				if (bCacheIfRequired)
+				{
+					ModelNodePathToEdNode.Add(InModelNodeName, EdNode);
+				}
 				return EdNode;
 			}
 		}
diff --git a/Engine/Plugins/Animation/ControlRig/Source/ControlRigDeveloper/Public/Graph/ControlRigGraph.h b/Engine/Plugins/Animation/ControlRig/Source/ControlRigDeveloper/Public/Graph/ControlRigGraph.h
--- a/Engine/Plugins/Animation/ControlRig/Source/ControlRigDeveloper/Public/Graph/ControlRigGraph.h
+++ b/Engine/Plugins/Animation/ControlRig/Source/ControlRigDeveloper/Public/Graph/ControlRigGraph.h
@@ -66,7 +66,7 @@ public:
 	bool bSuspendModelNotifications;
 	bool bIsTemporaryGraphForCopyPaste;
 
-	UEdGraphNode* FindNodeForModelNodeName(const FName& InModelNodeName);
+	UEdGraphNode* FindNodeForModelNodeName(const FName& InModelNodeName, const bool bCacheIfRequired = true);
 
 	UControlRigBlueprint* GetBlueprint() const;
 	URigVMGraph* GetModel() const;
diff --git a/Engine/Source/Developer/RigVMDeveloper/Private/RigVMModel/RigVMController.cpp b/Engine/Source/Developer/RigVMDeveloper/Private/RigVMModel/RigVMController.cpp
--- a/Engine/Source/Developer/RigVMDeveloper/Private/RigVMModel/RigVMController.cpp
+++ b/Engine/Source/Developer/RigVMDeveloper/Private/RigVMModel/RigVMController.cpp
@@ -2962,10 +2962,15 @@ TArray<URigVMNode*> URigVMController::ExpandLibraryNode(URigVMLibraryNode* InNod
 	check(ExpandedNodeNames.Num() >= NodeNames.Num());
 
 	TMap<FName, FName> NodeNameMap;
-	for (int32 NodeNameIndex = 0; NodeNameIndex < NodeNames.Num(); NodeNameIndex++)
+	for (int32 NodeNameIndex = 0, ExpandedNodeNameIndex = 0; NodeNameIndex < NodeNames.Num(); ExpandedNodeNameIndex++)
 	{
-		NodeNameMap.Add(NodeNames[NodeNameIndex], ExpandedNodeNames[NodeNameIndex]);
-		SetNodePosition(ExpandedNodes[NodeNameIndex], InNode->Position + ContainedNodes[NodeNameIndex]->Position - Center, false, false);
+		if (ExpandedNodes[ExpandedNodeNameIndex]->IsInjected())
+		{
+			continue;
+		}
+		NodeNameMap.Add(NodeNames[NodeNameIndex], ExpandedNodeNames[ExpandedNodeNameIndex]);
+		SetNodePosition(ExpandedNodes[ExpandedNodeNameIndex], InNode->Position + ContainedNodes[NodeNameIndex]->Position - Center, false, false);
+		NodeNameIndex++;
 	}
 
 	// a) store all of the pin defaults off the library node
