Control Rig: Fix a deprecated node, GetJointTransform

diff --git a/Engine/Plugins/Animation/ControlRig/Source/ControlRig/Private/Units/Deprecated/RigUnit_GetJointTransform.cpp b/Engine/Plugins/Animation/ControlRig/Source/ControlRig/Private/Units/Deprecated/RigUnit_GetJointTransform.cpp
--- a/Engine/Plugins/Animation/ControlRig/Source/ControlRig/Private/Units/Deprecated/RigUnit_GetJointTransform.cpp
+++ b/Engine/Plugins/Animation/ControlRig/Source/ControlRig/Private/Units/Deprecated/RigUnit_GetJointTransform.cpp
@@ -7,4 +7,30 @@
 FRigUnit_GetJointTransform_Execute()
 {
     DECLARE_SCOPE_HIERARCHICAL_COUNTER_RIGUNIT()
+    
+	if (URigHierarchy* Hierarchy = Context.Hierarchy)
+	{
+		const FRigElementKey Key(Joint, ERigElementType::Bone);
+		
+		switch (Type)
+		{
+		case ETransformGetterType::Current:
+			{
+				const FTransform ComputedBaseTransform = UtilityHelpers::GetBaseTransformByMode(TransformSpace, [Hierarchy](const FRigElementKey& JointKey) { return Hierarchy->GetGlobalTransform(JointKey); },
+				Hierarchy->GetFirstParent(Key), FRigElementKey(BaseJoint, ERigElementType::Bone), BaseTransform);
+
+				Output = Hierarchy->GetInitialGlobalTransform(Key).GetRelativeTransform(ComputedBaseTransform);
+				break;
+			}
+		case ETransformGetterType::Initial:
+		default:
+			{
+				const FTransform ComputedBaseTransform = UtilityHelpers::GetBaseTransformByMode(TransformSpace, [Hierarchy](const FRigElementKey& JointKey) { return Hierarchy->GetInitialGlobalTransform(JointKey); },
+				Hierarchy->GetFirstParent(Key), FRigElementKey(BaseJoint, ERigElementType::Bone), BaseTransform);
+
+				Output = Hierarchy->GetInitialGlobalTransform(Key).GetRelativeTransform(ComputedBaseTransform);
+				break;
+			}
+		}
+	}
 }
