Control Rig: Prevent unsupported types to be selected as inputs/outputs of functions

diff --git a/Engine/Plugins/Animation/ControlRig/Source/ControlRigDeveloper/Public/ControlRigBlueprint.h b/Engine/Plugins/Animation/ControlRig/Source/ControlRigDeveloper/Public/ControlRigBlueprint.h
--- a/Engine/Plugins/Animation/ControlRig/Source/ControlRigDeveloper/Public/ControlRigBlueprint.h
+++ b/Engine/Plugins/Animation/ControlRig/Source/ControlRigDeveloper/Public/ControlRigBlueprint.h
@@ -239,7 +239,7 @@ public:
 	URigVMController* GetTemplateController();
 
 #if WITH_EDITOR
-	UEdGraph* GetEdGraph(URigVMGraph* InModel = nullptr) const;
+	UEdGraph* GetEdGraph(URigVMGraph* InModel) const;
 	UEdGraph* GetEdGraph(const FString& InNodePath) const;
 #endif
 
diff --git a/Engine/Plugins/Animation/ControlRig/Source/ControlRigEditor/Private/ControlRigGraphDetails.cpp b/Engine/Plugins/Animation/ControlRig/Source/ControlRigEditor/Private/ControlRigGraphDetails.cpp
--- a/Engine/Plugins/Animation/ControlRig/Source/ControlRigEditor/Private/ControlRigGraphDetails.cpp
+++ b/Engine/Plugins/Animation/ControlRig/Source/ControlRigEditor/Private/ControlRigGraphDetails.cpp
@@ -20,6 +20,7 @@
 #include "ControlRig.h"
 #include "RigVMCore/RigVMExternalVariable.h"
 #include "Kismet2/BlueprintEditorUtils.h"
+#include "Graph/ControlRigGraphSchema.h"
 
 #define LOCTEXT_NAMESPACE "ControlRigGraphDetails"
 
@@ -125,7 +126,7 @@ void FControlRigArgumentGroupLayout::HandleModifiedEvent(ERigVMGraphNotifType In
 
 void FControlRigArgumentLayout::GenerateHeaderRowContent(FDetailWidgetRow& NodeRow)
 {
-	const UEdGraphSchema_K2* K2Schema = GetDefault<UEdGraphSchema_K2>();
+	const UEdGraphSchema* Schema = GetDefault<UControlRigGraphSchema>();
 
 	ETypeTreeFilter TypeTreeFilter = ETypeTreeFilter::None;
 	TypeTreeFilter |= ETypeTreeFilter::AllowExec;
@@ -135,76 +136,76 @@ void FControlRigArgumentLayout::GenerateHeaderRowContent(FDetailWidgetRow& NodeR
 		[
 			SNew(SHorizontalBox)
 			+ SHorizontalBox::Slot()
-		.FillWidth(1)
-		.VAlign(VAlign_Center)
-		[
-			SAssignNew(ArgumentNameWidget, SEditableTextBox)
-			.Text(this, &FControlRigArgumentLayout::OnGetArgNameText)
-		.OnTextChanged(this, &FControlRigArgumentLayout::OnArgNameChange)
-		.OnTextCommitted(this, &FControlRigArgumentLayout::OnArgNameTextCommitted)
-		.ToolTipText(this, &FControlRigArgumentLayout::OnGetArgToolTipText)
-		.Font(IDetailLayoutBuilder::GetDetailFont())
-		.IsEnabled(!ShouldPinBeReadOnly())
-		]
+			.FillWidth(1)
+			.VAlign(VAlign_Center)
+			[
+				SAssignNew(ArgumentNameWidget, SEditableTextBox)
+				.Text(this, &FControlRigArgumentLayout::OnGetArgNameText)
+				.OnTextChanged(this, &FControlRigArgumentLayout::OnArgNameChange)
+				.OnTextCommitted(this, &FControlRigArgumentLayout::OnArgNameTextCommitted)
+				.ToolTipText(this, &FControlRigArgumentLayout::OnGetArgToolTipText)
+				.Font(IDetailLayoutBuilder::GetDetailFont())
+				.IsEnabled(!ShouldPinBeReadOnly())
+			]
 		]
-	.ValueContent()
+		.ValueContent()
 		.MaxDesiredWidth(980.f)
 		[
 			SNew(SHorizontalBox)
 			+ SHorizontalBox::Slot()
-		.VAlign(VAlign_Center)
-		.Padding(0.0f, 0.0f, 4.0f, 0.0f)
-		.AutoWidth()
-		[
-			SNew(SPinTypeSelector, FGetPinTypeTree::CreateUObject(K2Schema, &UEdGraphSchema_K2::GetVariableTypeTree))
-			.TargetPinType(this, &FControlRigArgumentLayout::OnGetPinInfo)
-		.OnPinTypePreChanged(this, &FControlRigArgumentLayout::OnPrePinInfoChange)
-		.OnPinTypeChanged(this, &FControlRigArgumentLayout::PinInfoChanged)
-		.Schema(K2Schema)
-		.TypeTreeFilter(TypeTreeFilter)
-		.bAllowArrays(!ShouldPinBeReadOnly())
-		.IsEnabled(!ShouldPinBeReadOnly(true))
-		.Font(IDetailLayoutBuilder::GetDetailFont())
-		]
-	+ SHorizontalBox::Slot()
-		.AutoWidth()
-		[
-			SNew(SButton)
-			.ButtonStyle(FAppStyle::Get(), TEXT("SimpleButton"))
-		.ContentPadding(0)
-		.IsEnabled(!IsPinEditingReadOnly())
-		.OnClicked(this, &FControlRigArgumentLayout::OnArgMoveUp)
-		.ToolTipText(LOCTEXT("FunctionArgDetailsArgMoveUpTooltip", "Move this parameter up in the list."))
-		[
-			SNew(SImage)
-			.Image(FEditorStyle::GetBrush("Icons.ChevronUp"))
-		.ColorAndOpacity(FSlateColor::UseForeground())
-		]
-		]
-	+ SHorizontalBox::Slot()
-		.AutoWidth()
-		.Padding(2, 0)
-		[
-			SNew(SButton)
-			.ButtonStyle(FAppStyle::Get(), TEXT("SimpleButton"))
-		.ContentPadding(0)
-		.IsEnabled(!IsPinEditingReadOnly())
-		.OnClicked(this, &FControlRigArgumentLayout::OnArgMoveDown)
-		.ToolTipText(LOCTEXT("FunctionArgDetailsArgMoveDownTooltip", "Move this parameter down in the list."))
-		[
-			SNew(SImage)
-			.Image(FEditorStyle::GetBrush("Icons.ChevronDown"))
-		.ColorAndOpacity(FSlateColor::UseForeground())
-		]
-		]
-	+ SHorizontalBox::Slot()
-		.HAlign(HAlign_Right)
-		.VAlign(VAlign_Center)
-		.Padding(10, 0, 0, 0)
-		.AutoWidth()
-		[
-			PropertyCustomizationHelpers::MakeClearButton(FSimpleDelegate::CreateSP(this, &FControlRigArgumentLayout::OnRemoveClicked), LOCTEXT("FunctionArgDetailsClearTooltip", "Remove this parameter."), !IsPinEditingReadOnly())
-		]
+			.VAlign(VAlign_Center)
+			.Padding(0.0f, 0.0f, 4.0f, 0.0f)
+			.AutoWidth()
+			[
+				SNew(SPinTypeSelector, FGetPinTypeTree::CreateUObject(GetDefault<UEdGraphSchema_K2>(), &UEdGraphSchema_K2::GetVariableTypeTree))
+				.TargetPinType(this, &FControlRigArgumentLayout::OnGetPinInfo)
+				.OnPinTypePreChanged(this, &FControlRigArgumentLayout::OnPrePinInfoChange)
+				.OnPinTypeChanged(this, &FControlRigArgumentLayout::PinInfoChanged)
+				.Schema(Schema)
+				.TypeTreeFilter(TypeTreeFilter)
+				.bAllowArrays(!ShouldPinBeReadOnly())
+				.IsEnabled(!ShouldPinBeReadOnly(true))
+				.Font(IDetailLayoutBuilder::GetDetailFont())
+			]
+			+ SHorizontalBox::Slot()
+			.AutoWidth()
+			[
+				SNew(SButton)
+				.ButtonStyle(FAppStyle::Get(), TEXT("SimpleButton"))
+				.ContentPadding(0)
+				.IsEnabled(!IsPinEditingReadOnly())
+				.OnClicked(this, &FControlRigArgumentLayout::OnArgMoveUp)
+				.ToolTipText(LOCTEXT("FunctionArgDetailsArgMoveUpTooltip", "Move this parameter up in the list."))
+				[
+					SNew(SImage)
+					.Image(FEditorStyle::GetBrush("Icons.ChevronUp"))
+				.ColorAndOpacity(FSlateColor::UseForeground())
+				]
+			]
+			+ SHorizontalBox::Slot()
+			.AutoWidth()
+			.Padding(2, 0)
+			[
+				SNew(SButton)
+				.ButtonStyle(FAppStyle::Get(), TEXT("SimpleButton"))
+				.ContentPadding(0)
+				.IsEnabled(!IsPinEditingReadOnly())
+				.OnClicked(this, &FControlRigArgumentLayout::OnArgMoveDown)
+				.ToolTipText(LOCTEXT("FunctionArgDetailsArgMoveDownTooltip", "Move this parameter down in the list."))
+				[
+					SNew(SImage)
+					.Image(FEditorStyle::GetBrush("Icons.ChevronDown"))
+					.ColorAndOpacity(FSlateColor::UseForeground())
+				]
+			]
+			+ SHorizontalBox::Slot()
+			.HAlign(HAlign_Right)
+			.VAlign(VAlign_Center)
+			.Padding(10, 0, 0, 0)
+			.AutoWidth()
+			[
+				PropertyCustomizationHelpers::MakeClearButton(FSimpleDelegate::CreateSP(this, &FControlRigArgumentLayout::OnRemoveClicked), LOCTEXT("FunctionArgDetailsClearTooltip", "Remove this parameter."), !IsPinEditingReadOnly())
+			]
 		];
 }
 
