Control Rig: Changes to support a new python extendable context menu for drag & drop actions in rig hierarchy panel

diff --git a/Engine/Plugins/Animation/ControlRig/Source/ControlRig/Public/Rigs/RigHierarchy.h b/Engine/Plugins/Animation/ControlRig/Source/ControlRig/Public/Rigs/RigHierarchy.h
--- a/Engine/Plugins/Animation/ControlRig/Source/ControlRig/Public/Rigs/RigHierarchy.h
+++ b/Engine/Plugins/Animation/ControlRig/Source/ControlRig/Public/Rigs/RigHierarchy.h
@@ -452,7 +452,51 @@ public:
 	{
 		return Get<T>(GetIndex(InKey));
 	}
+	
+private:
+	/**
+	* Returns bone element for a given key, for scripting purpose only, for cpp usage, use Find<FRigBoneElement>()
+	* @param InKey The key of the bone element to retrieve. 
+	*/
+	UFUNCTION(BlueprintCallable, Category = URigHierarchy, meta = (DisplayName = "Find Bone", ScriptName = "FindBone"))
+    FRigBoneElement FindBone_ForBlueprintOnly(const FRigElementKey& InKey) const
+	{
+		if (const FRigBoneElement* Bone = Find<FRigBoneElement>(InKey))
+		{
+			return *Bone;
+		}
+		return FRigBoneElement();
+	}	
+	
+	/**
+	* Returns control element for a given key, for scripting purpose only, for cpp usage, use Find<FRigControlElement>()
+	* @param InKey The key of the control element to retrieve. 
+	*/
+	UFUNCTION(BlueprintCallable, Category = URigHierarchy, meta = (DisplayName = "Find Control", ScriptName = "FindControl"))
+    FRigControlElement FindControl_ForBlueprintOnly(const FRigElementKey& InKey) const
+	{
+		if (const FRigControlElement* Control = Find<FRigControlElement>(InKey))
+		{
+			return *Control;
+		}
+		return FRigControlElement();
+	}	
 
+	/**
+	* Returns null element for a given key, for scripting purpose only, for cpp usage, use Find<FRigControlElement>()
+	* @param InKey The key of the null element to retrieve. 
+	*/
+	UFUNCTION(BlueprintCallable, Category = URigHierarchy, meta = (DisplayName = "Find Null", ScriptName = "FindNull"))
+    FRigNullElement FindNull_ForBlueprintOnly(const FRigElementKey& InKey) const
+	{
+		if (const FRigNullElement* Null = Find<FRigNullElement>(InKey))
+		{
+			return *Null;
+		}
+		return FRigNullElement();
+	}
+	
+public:	
 	/**
 	 * Returns an element for a given key.
 	 * This templated method also casts to the chosen
diff --git a/Engine/Plugins/Animation/ControlRig/Source/ControlRigEditor/Private/Editor/ControlRigContextMenuContext.cpp b/Engine/Plugins/Animation/ControlRig/Source/ControlRigEditor/Private/Editor/ControlRigContextMenuContext.cpp
--- a/Engine/Plugins/Animation/ControlRig/Source/ControlRigEditor/Private/Editor/ControlRigContextMenuContext.cpp
+++ b/Engine/Plugins/Animation/ControlRig/Source/ControlRigEditor/Private/Editor/ControlRigContextMenuContext.cpp
@@ -4,6 +4,7 @@
 
 #include "ControlRigBlueprint.h"
 #include "ControlRigEditor.h"
+#include "Slate/Public/Framework/Application/SlateApplication.h"
 
 UControlRigBlueprint* UControlRigContextMenuContext::GetControlRigBlueprint() const
 {
@@ -13,4 +14,26 @@ UControlRigBlueprint* UControlRigContextMenuContext::GetControlRigBlueprint() co
 	}
 
 	return nullptr;
-}
\ No newline at end of file
+}
+
+UControlRig* UControlRigContextMenuContext::GetControlRig() const
+{
+	if (UControlRigBlueprint* Blueprint = GetControlRigBlueprint())
+	{
+		if (UControlRig* ControlRig = Cast<UControlRig>(Blueprint->GetObjectBeingDebugged()))
+		{
+			return ControlRig;
+		}
+	}
+	return nullptr;
+}
+
+bool UControlRigContextMenuContext::IsAltDown() const
+{
+	return FSlateApplication::Get().GetModifierKeys().IsAltDown();
+}
+
+FControlRigRigHierarchyDragAndDropContext UControlRigContextMenuContext::GetDragAndDropContext()
+{
+	return DragAndDropContext;
+}
diff --git a/Engine/Plugins/Animation/ControlRig/Source/ControlRigEditor/Private/Editor/ControlRigContextMenuContext.h b/Engine/Plugins/Animation/ControlRig/Source/ControlRigEditor/Private/Editor/ControlRigContextMenuContext.h
--- a/Engine/Plugins/Animation/ControlRig/Source/ControlRigEditor/Private/Editor/ControlRigContextMenuContext.h
+++ b/Engine/Plugins/Animation/ControlRig/Source/ControlRigEditor/Private/Editor/ControlRigContextMenuContext.h
@@ -3,6 +3,7 @@
 #pragma once
 
 #include "CoreMinimal.h"
+#include "Rigs/RigHierarchyDefines.h"
 
 #include "ControlRigContextMenuContext.generated.h"
 
@@ -10,15 +11,62 @@ class UControlRig;
 class UControlRigBlueprint;
 class FControlRigEditor;
 
+USTRUCT(BlueprintType)
+struct FControlRigRigHierarchyDragAndDropContext
+{
+	GENERATED_BODY()
+	
+	FControlRigRigHierarchyDragAndDropContext() = default;
+	
+	FControlRigRigHierarchyDragAndDropContext(const TArray<FRigElementKey> InDraggedElementKeys, FRigElementKey InTargetElementKey)
+		: DraggedElementKeys(InDraggedElementKeys)
+		, TargetElementKey(InTargetElementKey)
+	{
+	}
+	
+	UPROPERTY(BlueprintReadOnly, Category = ControlRigEditorExtensions)
+	TArray<FRigElementKey> DraggedElementKeys;
+
+	UPROPERTY(BlueprintReadOnly, Category = ControlRigEditorExtensions)
+	FRigElementKey TargetElementKey;
+};
+
 UCLASS(BlueprintType)
 class UControlRigContextMenuContext : public UObject
 {
 	GENERATED_BODY()
+
 public:
+	/**
+	 *	Initialize the Context
+	 * @param InControlRigEditor 		The Control Rig editor that hosts the menu
+	 * @param InDragAndDropContext 		Optional, only used for the menu that shows up after a drag & drop action
+	*/
+	void Init(TWeakPtr<FControlRigEditor> InControlRigEditor, FControlRigRigHierarchyDragAndDropContext InDragAndDropContext = FControlRigRigHierarchyDragAndDropContext())
+	{
+		ControlRigEditor = InControlRigEditor;
+		DragAndDropContext = InDragAndDropContext;
+	}
+	
 	/** Get the control rig blueprint that we are editing */
 	UFUNCTION(BlueprintCallable, Category = ControlRigEditorExtensions)
     UControlRigBlueprint* GetControlRigBlueprint() const;
 	
+	/** Get the active control rig instance in the viewport */
+	UFUNCTION(BlueprintCallable, Category = ControlRigEditorExtensions)
+    UControlRig* GetControlRig() const;
+
+	/** Returns true if either alt key is down */
+	UFUNCTION(BlueprintCallable, Category = ControlRigEditorExtensions)
+	bool IsAltDown() const;
+	
+	/** Returns the source and target element keys of a drag & drop action */
+	UFUNCTION(BlueprintCallable, Category = ControlRigEditorExtensions)
+	FControlRigRigHierarchyDragAndDropContext GetDragAndDropContext();
+
+private:
 	/** Our owning control rig editor */
 	TWeakPtr<FControlRigEditor> ControlRigEditor;
+
+	FControlRigRigHierarchyDragAndDropContext DragAndDropContext;
 };
\ No newline at end of file
diff --git a/Engine/Plugins/Animation/ControlRig/Source/ControlRigEditor/Private/Editor/SRigHierarchy.cpp b/Engine/Plugins/Animation/ControlRig/Source/ControlRigEditor/Private/Editor/SRigHierarchy.cpp
--- a/Engine/Plugins/Animation/ControlRig/Source/ControlRigEditor/Private/Editor/SRigHierarchy.cpp
+++ b/Engine/Plugins/Animation/ControlRig/Source/ControlRigEditor/Private/Editor/SRigHierarchy.cpp
@@ -1402,6 +1402,51 @@ void SRigHierarchy::OnSetExpansionRecursive(TSharedPtr<FRigTreeElement> InItem,
 	SetExpansionRecursive(InItem, false, bShouldBeExpanded);
 }
 
+UToolMenu* SRigHierarchy::GetOrCreateDragDropMenu(const TArray<FRigElementKey>& DraggedKeys, FRigElementKey TargetKey)
+{
+	const FName MenuName = TEXT("ControlRigEditor.RigHierarchy.DragDropMenu");
+	UToolMenus* ToolMenus = UToolMenus::Get();
+
+	if (!ToolMenus->IsMenuRegistered(MenuName))
+	{
+		UToolMenu* Menu = ToolMenus->RegisterMenu(MenuName);
+
+		FToolMenuEntry ParentEntry = FToolMenuEntry::InitMenuEntry(
+			TEXT("Parent"),
+			LOCTEXT("DragDropMenu_Parent", "Parent"),
+			LOCTEXT("DragDropMenu_Parent_ToolTip", "Parent Selected Items to the Target Item"),
+			FSlateIcon(),
+			FToolMenuExecuteAction::CreateSP(this, &SRigHierarchy::HandleParent)
+		);
+		Menu->AddMenuEntry(NAME_None, ParentEntry);
+
+		UToolMenu* AlignMenu = Menu->AddSubMenu(
+			ToolMenus->CurrentOwner(),
+			NAME_None,
+			TEXT("Align"),
+			LOCTEXT("DragDropMenu_Align", "Align"),
+			LOCTEXT("DragDropMenu_Align_ToolTip", "Align Selected Items' Transforms to Target Item's Transform")
+		);
+
+		FToolMenuEntry AlignAllEntry = FToolMenuEntry::InitMenuEntry(
+			TEXT("All"),
+			LOCTEXT("DragDropMenu_Align_All", "All"),
+			LOCTEXT("DragDropMenu_Align_All_ToolTip", "Align Selected Items' Transforms to Target Item's Transform"),
+			FSlateIcon(),
+			FToolMenuExecuteAction::CreateSP(this, &SRigHierarchy::HandleAlign)
+		);
+
+		AlignMenu->AddMenuEntry(NAME_None, AlignAllEntry);
+	}
+
+	UControlRigContextMenuContext* MenuContext = NewObject<UControlRigContextMenuContext>();
+	MenuContext->Init(ControlRigEditor, FControlRigRigHierarchyDragAndDropContext(DraggedKeys, TargetKey));
+	
+	UToolMenu* Menu = ToolMenus->GenerateMenu(MenuName, FToolMenuContext(MenuContext));
+
+	return Menu;
+}
+
 UToolMenu* SRigHierarchy::GetOrCreateContextMenu()
 {
 	const FName MenuName = TEXT("ControlRigEditor.RigHierarchy.ContextMenu");
@@ -1488,7 +1533,7 @@ UToolMenu* SRigHierarchy::GetOrCreateContextMenu()
 
 	// individual entries in this menu can access members of this context, particularly useful for editor scripting
 	UControlRigContextMenuContext* ContextMenuContext = NewObject<UControlRigContextMenuContext>();
-	ContextMenuContext->ControlRigEditor = ControlRigEditor;
+	ContextMenuContext->Init(ControlRigEditor);
 	
 	FToolMenuContext MenuContext(CommandList);
 	MenuContext.AddObject(ContextMenuContext);
@@ -2253,137 +2298,34 @@ TOptional<EItemDropZone> SRigHierarchy::OnCanAcceptDrop(const FDragDropEvent& Dr
 
 FReply SRigHierarchy::OnAcceptDrop(const FDragDropEvent& DragDropEvent, EItemDropZone DropZone, TSharedPtr<FRigTreeElement> TargetItem)
 {
+	bool bSummonDragDropMenu = DragDropEvent.GetModifierKeys().IsAltDown() && DragDropEvent.GetModifierKeys().IsShiftDown(); 
 	bool bMatchTransforms = DragDropEvent.GetModifierKeys().IsAltDown();
 	bool bReparentItems = !bMatchTransforms;
 
 	TSharedPtr<FRigElementHierarchyDragDropOp> RigDragDropOp = DragDropEvent.GetOperationAs<FRigElementHierarchyDragDropOp>();
 	if (RigDragDropOp.IsValid())
 	{
-		URigHierarchy* Hierarchy = GetHierarchy();
-		URigHierarchy* DebuggedHierarchy = GetDebuggedHierarchy();
-
-		if (Hierarchy && ControlRigBlueprint.IsValid())
+		if (bSummonDragDropMenu)
 		{
-			URigHierarchyController* Controller = Hierarchy->GetController(true);
-			if(Controller == nullptr)
-			{
-				return FReply::Unhandled();
-			}
-
-			TGuardValue<bool> GuardRigHierarchyChanges(bIsChangingRigHierarchy, true);
-			TGuardValue<bool> SuspendBlueprintNotifs(ControlRigBlueprint->bSuspendAllNotifications, true);
-			FScopedTransaction Transaction(LOCTEXT("HierarchyDragAndDrop", "Drag & Drop"));
-
-			FTransform TargetGlobalTransform = DebuggedHierarchy->GetGlobalTransform(TargetItem->Key);
-
-			for (const FRigElementKey& DraggedKey : RigDragDropOp->GetElements())
-			{
-				if (DraggedKey == TargetItem->Key)
-				{
-					return FReply::Unhandled();
-				}
-
-				if (bReparentItems)
-				{
-					if(Hierarchy->IsParentedTo(TargetItem->Key, DraggedKey))
-					{
-						return FReply::Unhandled();
-					}
-				}
+			const FVector2D& SummonLocation = DragDropEvent.GetScreenSpacePosition();
 
-				if (DraggedKey.Type == ERigElementType::Bone)
-				{
-					if(FRigBoneElement* BoneElement = Hierarchy->Find<FRigBoneElement>(DraggedKey))
-					{
-						if (BoneElement->BoneType == ERigBoneType::Imported && BoneElement->ParentElement != nullptr)
-						{
-							FText ConfirmText = bMatchTransforms ?
-								LOCTEXT("ConfirmMatchTransform", "Matching transforms of imported(white) bones can cause issues with animation - are you sure ?") :
-								LOCTEXT("ConfirmReparentBoneHierarchy", "Reparenting imported(white) bones can cause issues with animation - are you sure ?");
+			// Get the context menu content. If NULL, don't open a menu.
+			UToolMenu* DragDropMenu = GetOrCreateDragDropMenu(RigDragDropOp->GetElements(), TargetItem->Key);
+			const TSharedPtr<SWidget> MenuContent = UToolMenus::Get()->GenerateWidget(DragDropMenu);
 
-							FText TitleText = bMatchTransforms ?
-								LOCTEXT("MatchTransformImportedBone", "Match Transform on Imported Bone") :
-								LOCTEXT("ReparentImportedBone", "Reparent Imported Bone");
-
-							FSuppressableWarningDialog::FSetupInfo Info(ConfirmText, TitleText, "SRigHierarchy_Warning");
-							Info.ConfirmText = LOCTEXT("SRigHierarchy_Warning_Yes", "Yes");
-							Info.CancelText = LOCTEXT("SRigHierarchy_Warning_No", "No");
-
-							FSuppressableWarningDialog ChangeImportedBonesInHierarchy(Info);
-							if (ChangeImportedBonesInHierarchy.ShowModal() == FSuppressableWarningDialog::Cancel)
-							{
-								return FReply::Unhandled();
-							}
-						}
-					}
-				}
-			}
-
-			for (const FRigElementKey& DraggedKey : RigDragDropOp->GetElements())
+			if (MenuContent.IsValid())
 			{
-				if (bMatchTransforms)
-				{
-					if (DraggedKey.Type == ERigElementType::Control)
-					{
-						int32 ControlIndex = DebuggedHierarchy->GetIndex(DraggedKey);
-						if (ControlIndex == INDEX_NONE)
-						{
-							continue;
-						}
-
-						FTransform ParentTransform = DebuggedHierarchy->GetParentTransformByIndex(ControlIndex, false);
-						FTransform OffsetTransform = TargetGlobalTransform.GetRelativeTransform(ParentTransform);
-
-						Hierarchy->SetControlOffsetTransformByIndex(ControlIndex, OffsetTransform, ERigTransformType::InitialLocal, true, true);
-						Hierarchy->SetControlOffsetTransformByIndex(ControlIndex, OffsetTransform, ERigTransformType::CurrentLocal, true, true);
-						Hierarchy->SetLocalTransform(DraggedKey, FTransform::Identity, true, true, true);
-						Hierarchy->SetInitialLocalTransform(DraggedKey, FTransform::Identity, true, true);
-						DebuggedHierarchy->SetControlOffsetTransformByIndex(ControlIndex, OffsetTransform, ERigTransformType::InitialLocal, true, true);
-						DebuggedHierarchy->SetControlOffsetTransformByIndex(ControlIndex, OffsetTransform, ERigTransformType::CurrentLocal, true, true);
-						DebuggedHierarchy->SetLocalTransform(DraggedKey, FTransform::Identity, true, true, true);
-						DebuggedHierarchy->SetInitialLocalTransform(DraggedKey, FTransform::Identity, true, true);
-					}
-					else
-					{
-						Hierarchy->SetInitialGlobalTransform(DraggedKey, TargetGlobalTransform, true, true);
-						Hierarchy->SetGlobalTransform(DraggedKey, TargetGlobalTransform, false, true, true);
-						DebuggedHierarchy->SetInitialGlobalTransform(DraggedKey, TargetGlobalTransform, true, true);
-						DebuggedHierarchy->SetGlobalTransform(DraggedKey, TargetGlobalTransform, false, true, true);
-					}
-					continue;
-				}
-
-				FRigElementKey ParentKey = TargetItem->Key;
-
-				FTransform InitialTransform = DebuggedHierarchy->GetInitialGlobalTransform(DraggedKey);
-				FTransform GlobalTransform = DebuggedHierarchy->GetGlobalTransform(DraggedKey);
-
-				if(ParentKey.IsValid())
-				{
-					Controller->SetParent(DraggedKey, ParentKey, true, true);
-				}
-				else
-				{
-					Controller->RemoveAllParents(DraggedKey, true, true);
-				}
-
-				DebuggedHierarchy->SetInitialGlobalTransform(DraggedKey, InitialTransform, true, true);
-				DebuggedHierarchy->SetGlobalTransform(DraggedKey, GlobalTransform, false, true, true);
-				Hierarchy->SetInitialGlobalTransform(DraggedKey, InitialTransform, true, true);
-				Hierarchy->SetGlobalTransform(DraggedKey, GlobalTransform, false, true, true);
+				const FWidgetPath WidgetPath = DragDropEvent.GetEventPath() != nullptr ? *DragDropEvent.GetEventPath() : FWidgetPath();
+				FSlateApplication::Get().PushMenu(AsShared(), WidgetPath, MenuContent.ToSharedRef(), SummonLocation, FPopupTransitionEffect(FPopupTransitionEffect::ContextMenu));
 			}
+			
+			return FReply::Handled();
 		}
-
-		ControlRigBlueprint->PropagateHierarchyFromBPToInstances();
-
-		if(bReparentItems)
+		else
 		{
-			TGuardValue<bool> GuardRigHierarchyChanges(bIsChangingRigHierarchy, true);
-			ControlRigBlueprint->BroadcastRefreshEditor();
-			RefreshTreeView();
+			return ReparentOrMatchTransform(RigDragDropOp->GetElements(), TargetItem->Key, bReparentItems);
 		}
 
-		return FReply::Handled();
 	}
 
 	return FReply::Unhandled();
@@ -2705,6 +2647,157 @@ bool SRigHierarchy::FindClosestBone(const FVector& Point, FName& OutRigElementNa
 	return false;
 }
 
+void SRigHierarchy::HandleParent(const FToolMenuContext& Context)
+{
+	if (UControlRigContextMenuContext* MenuContext = Cast<UControlRigContextMenuContext>(Context.FindByClass(UControlRigContextMenuContext::StaticClass())))
+	{
+		const FControlRigRigHierarchyDragAndDropContext DragAndDropContext = MenuContext->GetDragAndDropContext();
+		ReparentOrMatchTransform(DragAndDropContext.DraggedElementKeys, DragAndDropContext.TargetElementKey, true);
+	}
+}
+
+void SRigHierarchy::HandleAlign(const FToolMenuContext& Context)
+{
+	if (UControlRigContextMenuContext* MenuContext = Cast<UControlRigContextMenuContext>(Context.FindByClass(UControlRigContextMenuContext::StaticClass())))
+	{
+		const FControlRigRigHierarchyDragAndDropContext DragAndDropContext = MenuContext->GetDragAndDropContext();
+		ReparentOrMatchTransform(DragAndDropContext.DraggedElementKeys, DragAndDropContext.TargetElementKey, false);
+	}
+}
+
+FReply SRigHierarchy::ReparentOrMatchTransform(const TArray<FRigElementKey>& DraggedKeys, FRigElementKey TargetKey, bool bReparentItems)
+{
+	bool bMatchTransforms = !bReparentItems;
+
+	URigHierarchy* Hierarchy = GetHierarchy();
+	URigHierarchy* DebuggedHierarchy = GetDebuggedHierarchy();
+
+	if (Hierarchy && ControlRigBlueprint.IsValid())
+	{
+		URigHierarchyController* Controller = Hierarchy->GetController(true);
+		if(Controller == nullptr)
+		{
+			return FReply::Unhandled();
+		}
+
+		TGuardValue<bool> GuardRigHierarchyChanges(bIsChangingRigHierarchy, true);
+		TGuardValue<bool> SuspendBlueprintNotifs(ControlRigBlueprint->bSuspendAllNotifications, true);
+		FScopedTransaction Transaction(LOCTEXT("HierarchyDragAndDrop", "Drag & Drop"));
+
+		FTransform TargetGlobalTransform = DebuggedHierarchy->GetGlobalTransform(TargetKey);
+
+		for (const FRigElementKey& DraggedKey : DraggedKeys)
+		{
+			if (DraggedKey == TargetKey)
+			{
+				return FReply::Unhandled();
+			}
+
+			if (bReparentItems)
+			{
+				if(Hierarchy->IsParentedTo(TargetKey, DraggedKey))
+				{
+					return FReply::Unhandled();
+				}
+			}
+
+			if (DraggedKey.Type == ERigElementType::Bone)
+			{
+				if(FRigBoneElement* BoneElement = Hierarchy->Find<FRigBoneElement>(DraggedKey))
+				{
+					if (BoneElement->BoneType == ERigBoneType::Imported && BoneElement->ParentElement != nullptr)
+					{
+						FText ConfirmText = bMatchTransforms ?
+							LOCTEXT("ConfirmMatchTransform", "Matching transforms of imported(white) bones can cause issues with animation - are you sure ?") :
+							LOCTEXT("ConfirmReparentBoneHierarchy", "Reparenting imported(white) bones can cause issues with animation - are you sure ?");
+
+						FText TitleText = bMatchTransforms ?
+							LOCTEXT("MatchTransformImportedBone", "Match Transform on Imported Bone") :
+							LOCTEXT("ReparentImportedBone", "Reparent Imported Bone");
+
+						FSuppressableWarningDialog::FSetupInfo Info(ConfirmText, TitleText, "SRigHierarchy_Warning");
+						Info.ConfirmText = LOCTEXT("SRigHierarchy_Warning_Yes", "Yes");
+						Info.CancelText = LOCTEXT("SRigHierarchy_Warning_No", "No");
+
+						FSuppressableWarningDialog ChangeImportedBonesInHierarchy(Info);
+						if (ChangeImportedBonesInHierarchy.ShowModal() == FSuppressableWarningDialog::Cancel)
+						{
+							return FReply::Unhandled();
+						}
+					}
+				}
+			}
+		}
+
+		for (const FRigElementKey& DraggedKey : DraggedKeys)
+		{
+			if (bMatchTransforms)
+			{
+				if (DraggedKey.Type == ERigElementType::Control)
+				{
+					int32 ControlIndex = DebuggedHierarchy->GetIndex(DraggedKey);
+					if (ControlIndex == INDEX_NONE)
+					{
+						continue;
+					}
+
+					FTransform ParentTransform = DebuggedHierarchy->GetParentTransformByIndex(ControlIndex, false);
+					FTransform OffsetTransform = TargetGlobalTransform.GetRelativeTransform(ParentTransform);
+
+					Hierarchy->SetControlOffsetTransformByIndex(ControlIndex, OffsetTransform, ERigTransformType::InitialLocal, true, true);
+					Hierarchy->SetControlOffsetTransformByIndex(ControlIndex, OffsetTransform, ERigTransformType::CurrentLocal, true, true);
+					Hierarchy->SetLocalTransform(DraggedKey, FTransform::Identity, true, true, true);
+					Hierarchy->SetInitialLocalTransform(DraggedKey, FTransform::Identity, true, true);
+					DebuggedHierarchy->SetControlOffsetTransformByIndex(ControlIndex, OffsetTransform, ERigTransformType::InitialLocal, true, true);
+					DebuggedHierarchy->SetControlOffsetTransformByIndex(ControlIndex, OffsetTransform, ERigTransformType::CurrentLocal, true, true);
+					DebuggedHierarchy->SetLocalTransform(DraggedKey, FTransform::Identity, true, true, true);
+					DebuggedHierarchy->SetInitialLocalTransform(DraggedKey, FTransform::Identity, true, true);
+				}
+				else
+				{
+					Hierarchy->SetInitialGlobalTransform(DraggedKey, TargetGlobalTransform, true, true);
+					Hierarchy->SetGlobalTransform(DraggedKey, TargetGlobalTransform, false, true, true);
+					DebuggedHierarchy->SetInitialGlobalTransform(DraggedKey, TargetGlobalTransform, true, true);
+					DebuggedHierarchy->SetGlobalTransform(DraggedKey, TargetGlobalTransform, false, true, true);
+				}
+				continue;
+			}
+
+			FRigElementKey ParentKey = TargetKey;
+
+			FTransform InitialTransform = DebuggedHierarchy->GetInitialGlobalTransform(DraggedKey);
+			FTransform GlobalTransform = DebuggedHierarchy->GetGlobalTransform(DraggedKey);
+
+			if(ParentKey.IsValid())
+			{
+				Controller->SetParent(DraggedKey, ParentKey, true, true);
+			}
+			else
+			{
+				Controller->RemoveAllParents(DraggedKey, true, true);
+			}
+
+			DebuggedHierarchy->SetInitialGlobalTransform(DraggedKey, InitialTransform, true, true);
+			DebuggedHierarchy->SetGlobalTransform(DraggedKey, GlobalTransform, false, true, true);
+			Hierarchy->SetInitialGlobalTransform(DraggedKey, InitialTransform, true, true);
+			Hierarchy->SetGlobalTransform(DraggedKey, GlobalTransform, false, true, true);
+		}
+	}
+
+	ControlRigBlueprint->PropagateHierarchyFromBPToInstances();
+
+	if(bReparentItems)
+	{
+		TGuardValue<bool> GuardRigHierarchyChanges(bIsChangingRigHierarchy, true);
+		ControlRigBlueprint->BroadcastRefreshEditor();
+		RefreshTreeView();
+	}
+
+		
+	return FReply::Handled();
+
+}
+
 void SRigHierarchy::HandleSetInitialTransformFromClosestBone()
 {
 	if (IsControlOrNullSelected())
diff --git a/Engine/Plugins/Animation/ControlRig/Source/ControlRigEditor/Private/Editor/SRigHierarchy.h b/Engine/Plugins/Animation/ControlRig/Source/ControlRigEditor/Private/Editor/SRigHierarchy.h
--- a/Engine/Plugins/Animation/ControlRig/Source/ControlRigEditor/Private/Editor/SRigHierarchy.h
+++ b/Engine/Plugins/Animation/ControlRig/Source/ControlRigEditor/Private/Editor/SRigHierarchy.h
@@ -245,7 +245,7 @@ private:
 	UToolMenu* GetOrCreateContextMenu();
 	TSharedPtr<FUICommandList> GetContextMenuCommands() const;
 
-
+	UToolMenu* GetOrCreateDragDropMenu(const TArray<FRigElementKey>& DraggedKeys, FRigElementKey TargetKey);
 
 	/** Our owning control rig editor */
 	TWeakPtr<FControlRigEditor> ControlRigEditor;
@@ -331,6 +331,10 @@ private:
 	void HandleUnparent();
 	bool FindClosestBone(const FVector& Point, FName& OutRigElementName, FTransform& OutGlobalTransform) const;
 
+	void HandleParent(const FToolMenuContext& Context);
+	void HandleAlign(const FToolMenuContext& Context);
+	FReply ReparentOrMatchTransform(const TArray<FRigElementKey>& DraggedKeys, FRigElementKey TargetKey, bool bReparentItems);
+
 	FName CreateUniqueName(const FName& InBaseName, ERigElementType InElementType) const;
 
 	void SetExpansionRecursive(TSharedPtr<FRigTreeElement> InElement, bool bTowardsParent, bool bShouldBeExpanded);
