Control Rig: Recreate Hierarchy during undo if needed

diff --git a/Engine/Plugins/Animation/ControlRig/Source/ControlRigEditor/Private/Editor/ControlRigEditor.cpp b/Engine/Plugins/Animation/ControlRig/Source/ControlRigEditor/Private/Editor/ControlRigEditor.cpp
--- a/Engine/Plugins/Animation/ControlRig/Source/ControlRigEditor/Private/Editor/ControlRigEditor.cpp
+++ b/Engine/Plugins/Animation/ControlRig/Source/ControlRigEditor/Private/Editor/ControlRigEditor.cpp
@@ -2161,6 +2161,17 @@ void FControlRigEditor::PostUndo(bool bSuccess)
 			GetPersonaToolkit()->SetPreviewMesh(PreviewMesh, true);
 		}
 
+		if (UControlRig* DebuggedControlRig = Cast<UControlRig>(RigBlueprint->GetObjectBeingDebugged()))
+		{
+			if(URigHierarchy* Hierarchy = DebuggedControlRig->GetHierarchy())
+			{
+				if(Hierarchy->Num() == 0)
+				{
+					OnHierarchyChanged();
+				}
+			}
+		}
+
 		if (FControlRigEditMode* EditMode = GetEditMode())
 		{
 			EditMode->RequestToRecreateGizmoActors();
