Control Rig: Optimized memory usage of RigHierarchy due to invalid enum value

diff --git a/Engine/Plugins/Animation/ControlRig/Source/ControlRig/Private/Rigs/RigHierarchy.cpp b/Engine/Plugins/Animation/ControlRig/Source/ControlRig/Private/Rigs/RigHierarchy.cpp
--- a/Engine/Plugins/Animation/ControlRig/Source/ControlRig/Private/Rigs/RigHierarchy.cpp
+++ b/Engine/Plugins/Animation/ControlRig/Source/ControlRig/Private/Rigs/RigHierarchy.cpp
@@ -221,7 +221,7 @@ void URigHierarchy::Reset()
 	}
 	Elements.Reset();
 	ElementsPerType.Reset();
-	for(int32 TypeIndex=0;TypeIndex<(int32)ERigElementType::All;TypeIndex++)
+	for(int32 TypeIndex=0;TypeIndex<(int32)ERigElementType::Last;TypeIndex++)
 	{
 		ElementsPerType.Add(TArray<FRigBaseElement*>());
 	}
diff --git a/Engine/Plugins/Animation/ControlRig/Source/ControlRig/Public/Rigs/RigHierarchyDefines.h b/Engine/Plugins/Animation/ControlRig/Source/ControlRig/Public/Rigs/RigHierarchyDefines.h
--- a/Engine/Plugins/Animation/ControlRig/Source/ControlRig/Public/Rigs/RigHierarchyDefines.h
+++ b/Engine/Plugins/Animation/ControlRig/Source/ControlRig/Public/Rigs/RigHierarchyDefines.h
@@ -24,6 +24,7 @@ enum class ERigElementType : uint8
 	Curve = 0x008,
 	RigidBody = 0x010,
 	Socket = 0x020,
+	Last = 0x040 UMETA(Hidden),
 	All = Bone | Null | Control | Curve | RigidBody | Socket,
 };
 
