Control Rig: Add option to automatically hook up mutable nodes

diff --git a/Engine/Plugins/Animation/ControlRig/Source/ControlRig/Private/Settings/ControlRigSettings.cpp b/Engine/Plugins/Animation/ControlRig/Source/ControlRig/Private/Settings/ControlRigSettings.cpp
--- a/Engine/Plugins/Animation/ControlRig/Source/ControlRig/Private/Settings/ControlRigSettings.cpp
+++ b/Engine/Plugins/Animation/ControlRig/Source/ControlRig/Private/Settings/ControlRigSettings.cpp
@@ -22,6 +22,7 @@ UControlRigEditorSettings::UControlRigEditorSettings(const FObjectInitializer& O
 #endif
 {
 #if WITH_EDITORONLY_DATA
+	bAutoLinkMutableNodes = false;
 	bResetControlsOnCompile = true;
 	bResetControlsOnPinValueInteraction = false;
 	bEnableUndoForPoseInteraction = true;
diff --git a/Engine/Plugins/Animation/ControlRig/Source/ControlRig/Public/Settings/ControlRigSettings.h b/Engine/Plugins/Animation/ControlRig/Source/ControlRig/Public/Settings/ControlRigSettings.h
--- a/Engine/Plugins/Animation/ControlRig/Source/ControlRig/Public/Settings/ControlRigSettings.h
+++ b/Engine/Plugins/Animation/ControlRig/Source/ControlRig/Public/Settings/ControlRigSettings.h
@@ -58,6 +58,12 @@ public:
 	virtual FName GetCategoryName() const override { return TEXT("ContentEditors"); }
 	
 #if WITH_EDITORONLY_DATA
+
+	// When this is checked mutable nodes (nodes with an execute pin)
+	// will be hooked up automatically.
+	UPROPERTY(EditAnywhere, config, Category = Interaction)
+	bool bAutoLinkMutableNodes;
+
 	// When this is checked all controls will return to their initial
 	// value as the user hits the Compile button.
 	UPROPERTY(EditAnywhere, config, Category = Interaction)
diff --git a/Engine/Plugins/Animation/ControlRig/Source/ControlRigEditor/Private/Graph/NodeSpawners/ControlRigUnitNodeSpawner.cpp b/Engine/Plugins/Animation/ControlRig/Source/ControlRigEditor/Private/Graph/NodeSpawners/ControlRigUnitNodeSpawner.cpp
--- a/Engine/Plugins/Animation/ControlRig/Source/ControlRigEditor/Private/Graph/NodeSpawners/ControlRigUnitNodeSpawner.cpp
+++ b/Engine/Plugins/Animation/ControlRig/Source/ControlRigEditor/Private/Graph/NodeSpawners/ControlRigUnitNodeSpawner.cpp
@@ -231,6 +231,11 @@ UControlRigGraphNode* UControlRigUnitNodeSpawner::SpawnNode(UEdGraph* ParentGrap
 
 void UControlRigUnitNodeSpawner::HookupMutableNode(URigVMNode* InModelNode, UControlRigBlueprint* InRigBlueprint)
 {
+	if(!UControlRigEditorSettings::Get()->bAutoLinkMutableNodes)
+	{
+		return;
+	}
+	
 	URigVMController* Controller = InRigBlueprint->GetController(InModelNode->GetGraph());
 
 	Controller->ClearNodeSelection(true);
