Control Rig: Disable auto selecting nodes when adding them to the panel

diff --git a/Engine/Plugins/Animation/ControlRig/Source/ControlRigDeveloper/Private/Graph/ControlRigGraph.cpp b/Engine/Plugins/Animation/ControlRig/Source/ControlRigDeveloper/Private/Graph/ControlRigGraph.cpp
--- a/Engine/Plugins/Animation/ControlRig/Source/ControlRigDeveloper/Private/Graph/ControlRigGraph.cpp
+++ b/Engine/Plugins/Animation/ControlRig/Source/ControlRigDeveloper/Private/Graph/ControlRigGraph.cpp
@@ -312,7 +312,7 @@ void UControlRigGraph::HandleModifiedEvent(ERigVMGraphNotifType InNotifType, URi
 				if (URigVMCommentNode* CommentModelNode = Cast<URigVMCommentNode>(ModelNode))
 				{
 					UEdGraphNode_Comment* NewNode = NewObject<UEdGraphNode_Comment>(this, CommentModelNode->GetFName());
-					AddNode(NewNode, false);
+					AddNode(NewNode, false, false);
 
 					NewNode->CreateNewGuid();
 					NewNode->PostPlacedNewNode();
@@ -330,7 +330,7 @@ void UControlRigGraph::HandleModifiedEvent(ERigVMGraphNotifType InNotifType, URi
 				else if (URigVMRerouteNode* RerouteModelNode = Cast<URigVMRerouteNode>(ModelNode))
 				{
 					UControlRigGraphNode* NewNode = NewObject<UControlRigGraphNode>(this, ModelNode->GetFName());
-					AddNode(NewNode, false);
+					AddNode(NewNode, false, false);
 
 					NewNode->ModelNodePath = ModelNode->GetNodePath();
 					NewNode->CreateNewGuid();
@@ -351,7 +351,7 @@ void UControlRigGraph::HandleModifiedEvent(ERigVMGraphNotifType InNotifType, URi
 				else // struct, library, parameter + variable
 				{
 					UControlRigGraphNode* NewNode = NewObject<UControlRigGraphNode>(this, ModelNode->GetFName());
-					AddNode(NewNode, false);
+					AddNode(NewNode, false, false);
 
 					NewNode->ModelNodePath = ModelNode->GetNodePath();
 					NewNode->CreateNewGuid();
