Control Rig: Item comparison nodes

diff --git a/Engine/Plugins/Animation/ControlRig/Source/ControlRig/Private/Units/Execution/RigUnit_Item.cpp b/Engine/Plugins/Animation/ControlRig/Source/ControlRig/Private/Units/Execution/RigUnit_Item.cpp
--- a/Engine/Plugins/Animation/ControlRig/Source/ControlRig/Private/Units/Execution/RigUnit_Item.cpp
+++ b/Engine/Plugins/Animation/ControlRig/Source/ControlRig/Private/Units/Execution/RigUnit_Item.cpp
@@ -37,6 +37,26 @@ FRigUnit_ItemReplace_Execute()
 	FRigUnit_NameReplace::StaticExecute(RigVMExecuteContext, Item.Name, Old, New, Result.Name, Context);
 }
 
+FRigUnit_ItemEquals_Execute()
+{
+	Result = (A == B);
+}
+
+FRigUnit_ItemNotEquals_Execute()
+{
+	Result = (A != B);
+}
+
+FRigUnit_ItemTypeEquals_Execute()
+{
+	Result = (A.Type == B.Type);
+}
+
+FRigUnit_ItemTypeNotEquals_Execute()
+{
+	Result = (A.Type != B.Type);
+}
+
 #if WITH_DEV_AUTOMATION_TESTS
 #include "Units/RigUnitTest.h"
 
diff --git a/Engine/Plugins/Animation/ControlRig/Source/ControlRig/Private/Units/Execution/RigUnit_Item.h b/Engine/Plugins/Animation/ControlRig/Source/ControlRig/Private/Units/Execution/RigUnit_Item.h
--- a/Engine/Plugins/Animation/ControlRig/Source/ControlRig/Private/Units/Execution/RigUnit_Item.h
+++ b/Engine/Plugins/Animation/ControlRig/Source/ControlRig/Private/Units/Execution/RigUnit_Item.h
@@ -75,4 +75,112 @@ struct CONTROLRIG_API FRigUnit_ItemReplace : public FRigUnit_ItemBase
 
 	UPROPERTY(meta = (Output))
 	FRigElementKey Result;
+};
+
+/**
+* Returns true if the two items are equal
+*/
+USTRUCT(meta=(DisplayName="Item Equals", Keywords=""))
+struct CONTROLRIG_API FRigUnit_ItemEquals : public FRigUnit_ItemBase
+{
+	GENERATED_BODY()
+
+	FRigUnit_ItemEquals()
+	{
+		A = B = FRigElementKey();
+		Result = false;
+	}
+
+	RIGVM_METHOD()
+	virtual void Execute(const FRigUnitContext& Context) override;
+
+	UPROPERTY(meta = (Input, ExpandByDefault))
+	FRigElementKey A;
+
+	UPROPERTY(meta = (Input, ExpandByDefault))
+	FRigElementKey B;
+
+	UPROPERTY(meta = (Output))
+	bool Result;
+};
+
+/**
+* Returns true if the two items are not equal
+*/
+USTRUCT(meta=(DisplayName="Item Not Equals", Keywords=""))
+struct CONTROLRIG_API FRigUnit_ItemNotEquals : public FRigUnit_ItemBase
+{
+	GENERATED_BODY()
+
+	FRigUnit_ItemNotEquals()
+	{
+		A = B = FRigElementKey();
+		Result = false;
+	}
+
+	RIGVM_METHOD()
+	virtual void Execute(const FRigUnitContext& Context) override;
+
+	UPROPERTY(meta = (Input, ExpandByDefault))
+	FRigElementKey A;
+
+	UPROPERTY(meta = (Input, ExpandByDefault))
+	FRigElementKey B;
+
+	UPROPERTY(meta = (Output))
+	bool Result;
+};
+
+/**
+* Returns true if the two items' types are equal
+*/
+USTRUCT(meta=(DisplayName="Item Type Equals", Keywords=""))
+struct CONTROLRIG_API FRigUnit_ItemTypeEquals : public FRigUnit_ItemBase
+{
+	GENERATED_BODY()
+
+	FRigUnit_ItemTypeEquals()
+	{
+		A = B = FRigElementKey();
+		Result = false;
+	}
+
+	RIGVM_METHOD()
+	virtual void Execute(const FRigUnitContext& Context) override;
+
+	UPROPERTY(meta = (Input, ExpandByDefault))
+	FRigElementKey A;
+
+	UPROPERTY(meta = (Input, ExpandByDefault))
+	FRigElementKey B;
+
+	UPROPERTY(meta = (Output))
+	bool Result;
+};
+
+/**
+* Returns true if the two items's types are not equal
+*/
+USTRUCT(meta=(DisplayName="Item Type Not Equals", Keywords=""))
+struct CONTROLRIG_API FRigUnit_ItemTypeNotEquals : public FRigUnit_ItemBase
+{
+	GENERATED_BODY()
+
+	FRigUnit_ItemTypeNotEquals()
+	{
+		A = B = FRigElementKey();
+		Result = false;
+	}
+
+	RIGVM_METHOD()
+	virtual void Execute(const FRigUnitContext& Context) override;
+
+	UPROPERTY(meta = (Input, ExpandByDefault))
+	FRigElementKey A;
+
+	UPROPERTY(meta = (Input, ExpandByDefault))
+	FRigElementKey B;
+
+	UPROPERTY(meta = (Output))
+	bool Result;
 };
\ No newline at end of file
