Control Rig: Variables not present after undoing delete variable

diff --git a/Engine/Plugins/Animation/ControlRig/Source/ControlRigEditor/Private/Editor/ControlRigEditor.cpp b/Engine/Plugins/Animation/ControlRig/Source/ControlRigEditor/Private/Editor/ControlRigEditor.cpp
--- a/Engine/Plugins/Animation/ControlRig/Source/ControlRigEditor/Private/Editor/ControlRigEditor.cpp
+++ b/Engine/Plugins/Animation/ControlRig/Source/ControlRigEditor/Private/Editor/ControlRigEditor.cpp
@@ -2177,6 +2177,8 @@ void FControlRigEditor::PostUndo(bool bSuccess)
 			EditMode->RequestToRecreateGizmoActors();
 		}
 	}
+	
+	Compile();
 }
 
 void FControlRigEditor::PostRedo(bool bSuccess)
@@ -2198,6 +2200,8 @@ void FControlRigEditor::PostRedo(bool bSuccess)
 			EditMode->RequestToRecreateGizmoActors();
 		}
 	}
+
+	Compile();
 }
 
 void FControlRigEditor::EnsureValidRigElementInDetailPanel()
