Control Rig: Implement frame of selection / complete graph using the F hotkey

diff --git a/Engine/Plugins/Animation/ControlRig/Source/ControlRigEditor/Private/ControlRigBlueprintCommands.cpp b/Engine/Plugins/Animation/ControlRig/Source/ControlRigEditor/Private/ControlRigBlueprintCommands.cpp
--- a/Engine/Plugins/Animation/ControlRig/Source/ControlRigEditor/Private/ControlRigBlueprintCommands.cpp
+++ b/Engine/Plugins/Animation/ControlRig/Source/ControlRigEditor/Private/ControlRigBlueprintCommands.cpp
@@ -31,6 +31,7 @@ void FControlRigBlueprintCommands::RegisterCommands()
 	UI_COMMAND(StoreNodeSnippet8, "Store Node Snippet 8", "Stores the selected node(s) into snippet 8", EUserInterfaceActionType::Button, FInputChord(EKeys::Eight, EModifierKey::Alt | EModifierKey::Shift));
 	UI_COMMAND(StoreNodeSnippet9, "Store Node Snippet 9", "Stores the selected node(s) into snippet 9", EUserInterfaceActionType::Button, FInputChord(EKeys::Nine, EModifierKey::Alt | EModifierKey::Shift));
 	UI_COMMAND(StoreNodeSnippet0, "Store Node Snippet 0", "Stores the selected node(s) into snippet 0", EUserInterfaceActionType::Button, FInputChord(EKeys::Zero, EModifierKey::Alt | EModifierKey::Shift));
+	UI_COMMAND(FrameSelection, "Frame Selection", "Frames the selected nodes in the Graph View.", EUserInterfaceActionType::Button, FInputChord(EKeys::F));
 }
 
 #undef LOCTEXT_NAMESPACE
diff --git a/Engine/Plugins/Animation/ControlRig/Source/ControlRigEditor/Private/ControlRigBlueprintCommands.h b/Engine/Plugins/Animation/ControlRig/Source/ControlRigEditor/Private/ControlRigBlueprintCommands.h
--- a/Engine/Plugins/Animation/ControlRig/Source/ControlRigEditor/Private/ControlRigBlueprintCommands.h
+++ b/Engine/Plugins/Animation/ControlRig/Source/ControlRigEditor/Private/ControlRigBlueprintCommands.h
@@ -93,6 +93,9 @@ public:
 	/** Stores the selected node(s) into snippet 0. */
 	TSharedPtr< FUICommandInfo > StoreNodeSnippet0;
 
+	/** Frames the selected nodes */
+	TSharedPtr< FUICommandInfo > FrameSelection;
+
 	/**
 	 * Initialize commands
 	 */
diff --git a/Engine/Plugins/Animation/ControlRig/Source/ControlRigEditor/Private/Editor/ControlRigEditor.cpp b/Engine/Plugins/Animation/ControlRig/Source/ControlRigEditor/Private/Editor/ControlRigEditor.cpp
--- a/Engine/Plugins/Animation/ControlRig/Source/ControlRigEditor/Private/Editor/ControlRigEditor.cpp
+++ b/Engine/Plugins/Animation/ControlRig/Source/ControlRigEditor/Private/Editor/ControlRigEditor.cpp
@@ -589,6 +589,11 @@ void FControlRigEditor::BindCommands()
 		FControlRigBlueprintCommands::Get().StepOut,
 		FExecuteAction::CreateSP(this, &FControlRigEditor::HandleBreakpointActionRequested, ERigVMBreakpointAction::StepOut),
 		FIsActionChecked::CreateSP(this, &FControlRigEditor::IsHaltedAtBreakpoint));
+
+	GetToolkitCommands()->MapAction(
+		FControlRigBlueprintCommands::Get().FrameSelection,
+		FExecuteAction::CreateSP(this, &FControlRigEditor::FrameSelection),
+		FCanExecuteAction());
 }
 
 void FControlRigEditor::ToggleExecuteGraph()
@@ -5617,6 +5622,18 @@ void FControlRigEditor::UpdateCapsules()
 	}
 }
 
+void FControlRigEditor::FrameSelection()
+{
+	if (SGraphEditor* GraphEd = FocusedGraphEdPtr.Pin().Get())
+	{
+		if(URigVMGraph* Model = GetFocusedModel())
+		{
+			const bool bFrameAll = Model->GetSelectNodes().Num() == 0;
+			GraphEd->ZoomToFit(!bFrameAll);
+		}
+	}
+}
+
 void FControlRigEditor::UpdateGraphCompilerErrors()
 {
 	DECLARE_SCOPE_HIERARCHICAL_COUNTER_FUNC()
diff --git a/Engine/Plugins/Animation/ControlRig/Source/ControlRigEditor/Private/Editor/ControlRigEditor.h b/Engine/Plugins/Animation/ControlRig/Source/ControlRigEditor/Private/Editor/ControlRigEditor.h
--- a/Engine/Plugins/Animation/ControlRig/Source/ControlRigEditor/Private/Editor/ControlRigEditor.h
+++ b/Engine/Plugins/Animation/ControlRig/Source/ControlRigEditor/Private/Editor/ControlRigEditor.h
@@ -352,6 +352,8 @@ private:
 
 	void UpdateCapsules();
 
+	void FrameSelection();
+
 protected:
 
 	/** Toolbox hosting widget */
