Control Rig: Ensure proper ordering of items in the context menus

diff --git a/Engine/Plugins/Animation/ControlRig/Source/ControlRigEditor/Private/Editor/SRigHierarchy.cpp b/Engine/Plugins/Animation/ControlRig/Source/ControlRigEditor/Private/Editor/SRigHierarchy.cpp
--- a/Engine/Plugins/Animation/ControlRig/Source/ControlRigEditor/Private/Editor/SRigHierarchy.cpp
+++ b/Engine/Plugins/Animation/ControlRig/Source/ControlRigEditor/Private/Editor/SRigHierarchy.cpp
@@ -1418,6 +1418,8 @@ UToolMenu* SRigHierarchy::GetOrCreateDragDropMenu(const TArray<FRigElementKey>&
 			FSlateIcon(),
 			FToolMenuExecuteAction::CreateSP(this, &SRigHierarchy::HandleParent)
 		);
+
+		ParentEntry.InsertPosition.Position = EToolMenuInsertType::First;
 		Menu->AddMenuEntry(NAME_None, ParentEntry);
 
 		UToolMenu* AlignMenu = Menu->AddSubMenu(
@@ -1428,6 +1430,15 @@ UToolMenu* SRigHierarchy::GetOrCreateDragDropMenu(const TArray<FRigElementKey>&
 			LOCTEXT("DragDropMenu_Align_ToolTip", "Align Selected Items' Transforms to Target Item's Transform")
 		);
 
+		if (FToolMenuSection* DefaultSection = Menu->FindSection(NAME_None))
+		{
+			if (FToolMenuEntry* AlignMenuEntry = DefaultSection->FindEntry(TEXT("Align")))
+			{
+				AlignMenuEntry->InsertPosition.Name = ParentEntry.Name;
+				AlignMenuEntry->InsertPosition.Position = EToolMenuInsertType::After;
+			}
+		}
+
 		FToolMenuEntry AlignAllEntry = FToolMenuEntry::InitMenuEntry(
 			TEXT("All"),
 			LOCTEXT("DragDropMenu_Align_All", "All"),
@@ -1435,6 +1446,7 @@ UToolMenu* SRigHierarchy::GetOrCreateDragDropMenu(const TArray<FRigElementKey>&
 			FSlateIcon(),
 			FToolMenuExecuteAction::CreateSP(this, &SRigHierarchy::HandleAlign)
 		);
+		AlignAllEntry.InsertPosition.Position = EToolMenuInsertType::First;
 
 		AlignMenu->AddMenuEntry(NAME_None, AlignAllEntry);
 	}
