Control Rig: Only perform user action if OK is being pressed in dialogs

diff --git a/Engine/Plugins/Animation/ControlRig/Source/ControlRigEditor/Private/ControlRigDrawingDetails.cpp b/Engine/Plugins/Animation/ControlRig/Source/ControlRigEditor/Private/ControlRigDrawingDetails.cpp
--- a/Engine/Plugins/Animation/ControlRig/Source/ControlRigEditor/Private/ControlRigDrawingDetails.cpp
+++ b/Engine/Plugins/Animation/ControlRig/Source/ControlRigEditor/Private/ControlRigDrawingDetails.cpp
@@ -111,10 +111,17 @@ FReply FControlRigDrawContainerDetails::OnImportCurvesFromFBXClicked()
 		TSharedRef<SKismetInspector> KismetInspector = SNew(SKismetInspector);
 		KismetInspector->ShowSingleStruct(StructToDisplay);
 
-		SGenericDialogWidget::OpenDialog(LOCTEXT("ControlRigEditorImportFBXCurves", "Import FBX Curves"), KismetInspector, SGenericDialogWidget::FArguments(), true);
+		SGenericDialogWidget::FArguments DialogArguments;
+		DialogArguments.OnOkPressed_Lambda([&Settings, this, &OutFilenames] ()
+		{
+			if(OutFilenames.Num() > 0)
+			{
+				ImportCurvesFromFBX(OutFilenames[0], BlueprintBeingCustomized, Settings);
+			}
+		});
+	
+		SGenericDialogWidget::OpenDialog(LOCTEXT("ControlRigEditorImportFBXCurves", "Import FBX Curves"), KismetInspector, DialogArguments, true);
 
-		// todo: show dialog for settings
-		ImportCurvesFromFBX(OutFilenames[0], BlueprintBeingCustomized, Settings);
 	}
 	return FReply::Handled();
 }
diff --git a/Engine/Plugins/Animation/ControlRig/Source/ControlRigEditor/Private/Editor/SRigHierarchy.cpp b/Engine/Plugins/Animation/ControlRig/Source/ControlRigEditor/Private/Editor/SRigHierarchy.cpp
--- a/Engine/Plugins/Animation/ControlRig/Source/ControlRigEditor/Private/Editor/SRigHierarchy.cpp
+++ b/Engine/Plugins/Animation/ControlRig/Source/ControlRigEditor/Private/Editor/SRigHierarchy.cpp
@@ -660,12 +660,16 @@ FReply SRigHierarchy::OnImportSkeletonClicked()
 	TSharedRef<SKismetInspector> KismetInspector = SNew(SKismetInspector);
 	KismetInspector->ShowSingleStruct(StructToDisplay);
 
-	SGenericDialogWidget::OpenDialog(LOCTEXT("ControlRigHierarchyImport", "Import Hierarchy"), KismetInspector, SGenericDialogWidget::FArguments(), true);
-
-	if (Settings.Mesh != nullptr)
+	SGenericDialogWidget::FArguments DialogArguments;
+	DialogArguments.OnOkPressed_Lambda([&Settings, this] ()
 	{
-		ImportHierarchy(FAssetData(Settings.Mesh));
-	}
+		if (Settings.Mesh != nullptr)
+		{
+			ImportHierarchy(FAssetData(Settings.Mesh));
+		}
+	});
+	
+	SGenericDialogWidget::OpenDialog(LOCTEXT("ControlRigHierarchyImport", "Import Hierarchy"), KismetInspector, DialogArguments, true);
 
 	return FReply::Handled();
 }
diff --git a/Engine/Source/Editor/UnrealEd/Private/Dialogs/Dialogs.cpp b/Engine/Source/Editor/UnrealEd/Private/Dialogs/Dialogs.cpp
--- a/Engine/Source/Editor/UnrealEd/Private/Dialogs/Dialogs.cpp
+++ b/Engine/Source/Editor/UnrealEd/Private/Dialogs/Dialogs.cpp
@@ -910,6 +910,8 @@ void SGenericDialogWidget::Construct( const FArguments& InArgs )
 			.OnClicked(this, &SGenericDialogWidget::OnOK_Clicked)
 		]
 	];
+
+	OkPressedDelegate = InArgs._OnOkPressed;
 }
 
 void SGenericDialogWidget::OpenDialog(const FText& InDialogTitle, const TSharedRef< SWidget >& DisplayContent, const FArguments& InArgs, bool bAsModalDialog)
@@ -930,6 +932,7 @@ void SGenericDialogWidget::OpenDialog(const FText& InDialogTitle, const TSharedR
 				SAssignNew(GenericDialogWidget, SGenericDialogWidget)
 				.UseScrollBox(InArgs._UseScrollBox)
 				.ScrollBoxMaxHeight(InArgs._ScrollBoxMaxHeight)
+				.OnOkPressed(InArgs._OnOkPressed)
 				[
 					DisplayContent
 				]
@@ -950,6 +953,8 @@ void SGenericDialogWidget::OpenDialog(const FText& InDialogTitle, const TSharedR
 
 FReply SGenericDialogWidget::OnOK_Clicked(void)
 {
+	OkPressedDelegate.ExecuteIfBound();
+	
 	MyWindow.Pin()->RequestDestroyWindow();
 
 	return FReply::Handled();
diff --git a/Engine/Source/Editor/UnrealEd/Public/Dialogs/Dialogs.h b/Engine/Source/Editor/UnrealEd/Public/Dialogs/Dialogs.h
--- a/Engine/Source/Editor/UnrealEd/Public/Dialogs/Dialogs.h
+++ b/Engine/Source/Editor/UnrealEd/Public/Dialogs/Dialogs.h
@@ -189,6 +189,9 @@ public:
 		/** Content for the dialog */
 		SLATE_DEFAULT_SLOT( FArguments, Content )
 
+		/** Called when the OK button is pressed */
+		SLATE_EVENT( FSimpleDelegate, OnOkPressed )
+
 	SLATE_END_ARGS()
 
 	void Construct( const FArguments& InArgs );
@@ -207,6 +210,8 @@ private:
 private:
 	/** Pointer to the containing window. */
 	TWeakPtr< SWindow > MyWindow;
+
+	FSimpleDelegate OkPressedDelegate;
 };
 
 UE_DEPRECATED(4.26, "Creating groups (nested packages) is no longer supported. Use PromptUserIfExistingObject overload that does not take the Group paramater.")
