Control Rig: Fix static analysis warning about BoneType always being Imported

diff --git a/Engine/Plugins/Animation/ControlRig/Source/ControlRig/Private/Rigs/RigHierarchyController.cpp b/Engine/Plugins/Animation/ControlRig/Source/ControlRig/Private/Rigs/RigHierarchyController.cpp
--- a/Engine/Plugins/Animation/ControlRig/Source/ControlRig/Private/Rigs/RigHierarchyController.cpp
+++ b/Engine/Plugins/Animation/ControlRig/Source/ControlRig/Private/Rigs/RigHierarchyController.cpp
@@ -1173,15 +1173,12 @@ TArray<FString> URigHierarchyController::GetAddBonePythonCommands(FRigBoneElemen
 		ParentKeyStr = Bone->ParentElement->GetKey().ToPythonString();
 	}
 	
-	if (Bone->BoneType == ERigBoneType::Imported)
-	{
-		// AddBone(FName InName, FRigElementKey InParent, FTransform InTransform, bool bTransformInGlobal = true, ERigBoneType InBoneType = ERigBoneType::User, bool bSetupUndo = false);
-		Commands.Add(FString::Printf(TEXT("hierarchy_controller.add_bone('%s', %s, %s, False, unreal.RigBoneType.%s)"),
-			*Bone->GetName().ToString(),
-			*ParentKeyStr,
-			*TransformStr,
-			Bone->BoneType == ERigBoneType::Imported ? TEXT("IMPORTED") : TEXT("USER")));
-	}
+	// AddBone(FName InName, FRigElementKey InParent, FTransform InTransform, bool bTransformInGlobal = true, ERigBoneType InBoneType = ERigBoneType::User, bool bSetupUndo = false);
+	Commands.Add(FString::Printf(TEXT("hierarchy_controller.add_bone('%s', %s, %s, False, unreal.RigBoneType.%s)"),
+		*Bone->GetName().ToString(),
+		*ParentKeyStr,
+		*TransformStr,
+		Bone->BoneType == ERigBoneType::Imported ? TEXT("IMPORTED") : TEXT("USER")));
 
 	return Commands;
 }
