Control Rig: Fix regression on RefreshHierarchy

diff --git a/Engine/Plugins/Animation/ControlRig/Source/ControlRigEditor/Private/Editor/SRigHierarchy.cpp b/Engine/Plugins/Animation/ControlRig/Source/ControlRigEditor/Private/Editor/SRigHierarchy.cpp
--- a/Engine/Plugins/Animation/ControlRig/Source/ControlRigEditor/Private/Editor/SRigHierarchy.cpp
+++ b/Engine/Plugins/Animation/ControlRig/Source/ControlRigEditor/Private/Editor/SRigHierarchy.cpp
@@ -1543,7 +1543,7 @@ UToolMenu* SRigHierarchy::GetOrCreateContextMenu()
 		
 		AssetsSection.AddSubMenu(TEXT("Refresh"), LOCTEXT("RefreshSubMenu", "Refresh"),
 			LOCTEXT("RefreshSubMenu_ToolTip", "Refresh the existing initial transform from the selected mesh. This only updates if the node is found."),
-            FNewMenuDelegate::CreateSP(this, &SRigHierarchy::CreateImportMenu)
+            FNewMenuDelegate::CreateSP(this, &SRigHierarchy::CreateRefreshMenu)
         );
 	}
 
