Control Rig: Deprecated Twist Bones Unit

diff --git a/Engine/Plugins/Animation/ControlRig/Source/ControlRig/Private/Units/Highlevel/Hierarchy/RigUnit_TwistBones.h b/Engine/Plugins/Animation/ControlRig/Source/ControlRig/Private/Units/Highlevel/Hierarchy/RigUnit_TwistBones.h
--- a/Engine/Plugins/Animation/ControlRig/Source/ControlRig/Private/Units/Highlevel/Hierarchy/RigUnit_TwistBones.h
+++ b/Engine/Plugins/Animation/ControlRig/Source/ControlRig/Private/Units/Highlevel/Hierarchy/RigUnit_TwistBones.h
@@ -99,7 +99,7 @@ struct CONTROLRIG_API FRigUnit_TwistBones : public FRigUnit_HighlevelBaseMutable
 /**
  * Creates a gradient of twist rotation along a collection of items.
  */
-USTRUCT(meta=(DisplayName="Twist Bones", Category="Hierarchy", Keywords="TwistBones"))
+USTRUCT(meta=(DisplayName="Twist Bones", Category="Hierarchy", Keywords="TwistBones", Deprecated = "5.0"))
 struct CONTROLRIG_API FRigUnit_TwistBonesPerItem : public FRigUnit_HighlevelBaseMutable
 {
 	GENERATED_BODY()
