Control Rig: Fix reset execution mode when editor is reopened

diff --git a/Engine/Plugins/Animation/ControlRig/Source/ControlRigDeveloper/Private/ControlRigBlueprint.cpp b/Engine/Plugins/Animation/ControlRig/Source/ControlRigDeveloper/Private/ControlRigBlueprint.cpp
--- a/Engine/Plugins/Animation/ControlRig/Source/ControlRigDeveloper/Private/ControlRigBlueprint.cpp
+++ b/Engine/Plugins/Animation/ControlRig/Source/ControlRigDeveloper/Private/ControlRigBlueprint.cpp
@@ -729,6 +729,11 @@ void UControlRigBlueprint::HandleReportFromCompiler(EMessageSeverity::Type InSev
 
 void UControlRigBlueprint::ClearBreakpoints()
 {
+	for(URigVMNode* Node : RigVMBreakpointNodes)
+	{
+		Node->SetHasBreakpoint(false);		
+	}
+	
 	RigVMBreakpointNodes.Empty();
 	RefreshControlRigBreakpoints();
 }
diff --git a/Engine/Plugins/Animation/ControlRig/Source/ControlRigEditor/Private/Editor/ControlRigEditor.cpp b/Engine/Plugins/Animation/ControlRig/Source/ControlRigEditor/Private/Editor/ControlRigEditor.cpp
--- a/Engine/Plugins/Animation/ControlRig/Source/ControlRigEditor/Private/Editor/ControlRigEditor.cpp
+++ b/Engine/Plugins/Animation/ControlRig/Source/ControlRigEditor/Private/Editor/ControlRigEditor.cpp
@@ -237,6 +237,16 @@ FControlRigEditor::~FControlRigEditor()
 		RigBlueprint->OnRequestLocalizeFunctionDialog().RemoveAll(this);
 		RigBlueprint->OnRequestBulkEditDialog().Unbind();
 		RigBlueprint->OnReportCompilerMessage().RemoveAll(this);
+
+#if WITH_EDITOR
+		RigBlueprint->SetDebugMode(false);
+		RigBlueprint->ClearBreakpoints();
+		if (HaltedAtNode)
+		{
+			HaltedAtNode->SetExecutionIsHaltedAtThisNode(false);
+			HaltedAtNode = nullptr;
+		}
+#endif
 	}
 
 	if (NodeDetailBuffer.Num() > 0 && NodeDetailStruct != nullptr)
@@ -3626,6 +3636,10 @@ void FControlRigEditor::UpdateControlRig()
 			ControlRig->PreviewInstance = PreviewInstance;
 			ControlRig->bSetupModeEnabled = bSetupModeEnabled;
 
+#if WITH_EDITOR
+			ControlRig->bIsInDebugMode = ExecutionMode == EControlRigExecutionModeType_Debug;
+#endif
+
 			if (UControlRig* CDO = Cast<UControlRig>(Class->GetDefaultObject()))
 			{
 				CDO->GizmoLibrary = GetControlRigBlueprint()->GizmoLibrary;
diff --git a/Engine/Source/Developer/RigVMDeveloper/Private/RigVMModel/RigVMNode.cpp b/Engine/Source/Developer/RigVMDeveloper/Private/RigVMModel/RigVMNode.cpp
--- a/Engine/Source/Developer/RigVMDeveloper/Private/RigVMModel/RigVMNode.cpp
+++ b/Engine/Source/Developer/RigVMDeveloper/Private/RigVMModel/RigVMNode.cpp
@@ -12,6 +12,8 @@ URigVMNode::URigVMNode()
 , Position(FVector2D::ZeroVector)
 , Size(FVector2D::ZeroVector)
 , NodeColor(FLinearColor::White)
+, bHasBreakpoint(false)
+, bHaltedAtThisNode(false)
 {
 
 }
diff --git a/Engine/Source/Runtime/RigVM/Private/RigVMCore/RigVMDebugInfo.cpp b/Engine/Source/Runtime/RigVM/Private/RigVMCore/RigVMDebugInfo.cpp
--- a/Engine/Source/Runtime/RigVM/Private/RigVMCore/RigVMDebugInfo.cpp
+++ b/Engine/Source/Runtime/RigVM/Private/RigVMCore/RigVMDebugInfo.cpp
@@ -7,6 +7,8 @@ void FRigVMDebugInfo::ResetState()
 	BreakpointHits.Empty();
 	BreakpointActivationOnHit.Empty();
 	TemporaryBreakpoint = nullptr;
+	CurrentActiveBreakpoint = nullptr;
+	CurrentActiveBreakpointCallstack.Empty();
 }
 
 void FRigVMDebugInfo::StartExecution()
diff --git a/Engine/Source/Runtime/RigVM/Public/RigVMCore/RigVMDebugInfo.h b/Engine/Source/Runtime/RigVM/Public/RigVMCore/RigVMDebugInfo.h
--- a/Engine/Source/Runtime/RigVM/Public/RigVMCore/RigVMDebugInfo.h
+++ b/Engine/Source/Runtime/RigVM/Public/RigVMCore/RigVMDebugInfo.h
@@ -82,6 +82,12 @@ struct RIGVM_API FRigVMDebugInfo
 	void SetBreakpoints(const TArray<TSharedRef<FRigVMBreakpoint>>& InBreakpoints)
 	{
 		Breakpoints = InBreakpoints;
+
+		if (CurrentActiveBreakpoint.IsValid() && !Breakpoints.Contains(CurrentActiveBreakpoint))
+		{
+			CurrentActiveBreakpoint = nullptr;
+			CurrentActiveBreakpointCallstack.Reset();
+		}
 	}
 
 	bool IsTemporaryBreakpoint(TSharedPtr<FRigVMBreakpoint> Breakpoint) const
