Control Rig: Support editor side mirroring on collections

diff --git a/Engine/Plugins/Animation/ControlRig/Source/ControlRigEditor/Private/ControlRigEditorModule.cpp b/Engine/Plugins/Animation/ControlRig/Source/ControlRigEditor/Private/ControlRigEditorModule.cpp
--- a/Engine/Plugins/Animation/ControlRig/Source/ControlRigEditor/Private/ControlRigEditorModule.cpp
+++ b/Engine/Plugins/Animation/ControlRig/Source/ControlRigEditor/Private/ControlRigEditorModule.cpp
@@ -1642,7 +1642,8 @@ void FControlRigEditorModule::GetContextMenuActions(const UControlRigGraphSchema
 							StructMemory->Execute(RigUnitContext);
 						}
 
-						for (const URigVMPin* Pin : ModelNode->GetAllPinsRecursively())
+						const TArray<URigVMPin*> AllPins = ModelNode->GetAllPinsRecursively();
+						for (const URigVMPin* Pin : AllPins)
 						{
 							if (Pin->GetCPPType() == TEXT("FName"))
 							{
@@ -1671,7 +1672,7 @@ void FControlRigEditorModule::GetContextMenuActions(const UControlRigGraphSchema
 								RigElementsToSelect.AddUnique(Key);
 								PinToKey.Add(Pin, Key);
 							}
-							else if (Pin->GetCPPTypeObject() == FRigElementKey::StaticStruct())
+							else if (Pin->GetCPPTypeObject() == FRigElementKey::StaticStruct() && !Pin->IsArray())
 							{
 								if (StructMemory == nullptr)
 								{
@@ -1693,10 +1694,60 @@ void FControlRigEditorModule::GetContextMenuActions(const UControlRigGraphSchema
 								else
 								{
 									check(ScriptStruct);
-									if (const FProperty* Property = ScriptStruct->FindPropertyByName(Pin->GetFName()))
+
+									TArray<FString> PropertyNames; 
+									if(!URigVMPin::SplitPinPath(Pin->GetSegmentPath(true), PropertyNames))
+									{
+										PropertyNames.Add(Pin->GetName());
+									}
+
+									UScriptStruct* Struct = ScriptStruct;
+									uint8* Memory = (uint8*)StructMemory; 
+
+									while(!PropertyNames.IsEmpty())
 									{
-										const FRigElementKey& Key = *Property->ContainerPtrToValuePtr<FRigElementKey>(StructMemory);
+										FString PropertyName;
+										PropertyNames.HeapPop(PropertyName);
+										
+										const FProperty* Property = ScriptStruct->FindPropertyByName(*PropertyName);
+										if(Property == nullptr)
+										{
+											Memory = nullptr;
+											break;
+										}
 
+										Memory = Property->ContainerPtrToValuePtr<uint8>(Memory);
+										
+										if(PropertyNames.IsEmpty())
+										{
+											continue;
+										}
+										
+										if(const FArrayProperty* ArrayProperty = CastField<FArrayProperty>(Property))
+										{
+											PropertyNames.HeapPop(PropertyName);
+
+											int32 ArrayIndex = FCString::Atoi(*PropertyName);
+											FScriptArrayHelper Helper(ArrayProperty, Memory);
+											if(!Helper.IsValidIndex(ArrayIndex))
+											{
+												Memory = nullptr;
+												break;
+											}
+
+											Memory = Helper.GetRawPtr(ArrayIndex);
+											Property = ArrayProperty->Inner;
+										}
+
+										if(const FStructProperty* StructProperty = CastField<FStructProperty>(Property))
+										{
+											Struct = StructProperty->Struct;
+										}
+									}
+
+									if(Memory)
+									{
+										const FRigElementKey& Key = *(const FRigElementKey*)Memory;
 										if (Key.IsValid())
 										{
 											RigElementsToSelect.AddUnique(Key);
