Control Rig: Allow CJK characters in rig element name

diff --git a/Engine/Plugins/Animation/ControlRig/Source/ControlRig/Private/Rigs/RigHierarchy.cpp b/Engine/Plugins/Animation/ControlRig/Source/ControlRig/Private/Rigs/RigHierarchy.cpp
--- a/Engine/Plugins/Animation/ControlRig/Source/ControlRig/Private/Rigs/RigHierarchy.cpp
+++ b/Engine/Plugins/Animation/ControlRig/Source/ControlRig/Private/Rigs/RigHierarchy.cpp
@@ -549,7 +549,7 @@ void URigHierarchy::SanitizeName(FString& InOutName)
 		TCHAR& C = InOutName[i];
 
 		const bool bGoodChar =
-			((C >= 'A') && (C <= 'Z')) || ((C >= 'a') && (C <= 'z')) ||		// A-Z (upper and lowercase) anytime
+			(!iscntrl(C) && !isdigit(C) && !ispunct(C) && !isspace(C)) ||	// Any alpha character (whatever the language)
 			(C == '_') || (C == '-') || (C == '.') ||						// _  - . anytime
 			((i > 0) && (C >= '0') && (C <= '9'));							// 0-9 after the first character
 
