Control Rig: Allow CJK characters in rig element name

diff --git a/Engine/Plugins/Animation/ControlRig/Source/ControlRig/Private/Rigs/RigHierarchy.cpp b/Engine/Plugins/Animation/ControlRig/Source/ControlRig/Private/Rigs/RigHierarchy.cpp
--- a/Engine/Plugins/Animation/ControlRig/Source/ControlRig/Private/Rigs/RigHierarchy.cpp
+++ b/Engine/Plugins/Animation/ControlRig/Source/ControlRig/Private/Rigs/RigHierarchy.cpp
@@ -548,10 +548,9 @@ void URigHierarchy::SanitizeName(FString& InOutName)
 	{
 		TCHAR& C = InOutName[i];
 
-		const bool bGoodChar =
-			(!iscntrl(C) && !isdigit(C) && !ispunct(C) && !isspace(C)) ||	// Any alpha character (whatever the language)
+		const bool bGoodChar = FChar::IsAlpha(C) ||							// Any letter
 			(C == '_') || (C == '-') || (C == '.') ||						// _  - . anytime
-			((i > 0) && (C >= '0') && (C <= '9'));							// 0-9 after the first character
+			((i > 0) && (FChar::IsDigit(C)));								// 0-9 after the first character
 
 		if (!bGoodChar)
 		{
