Control Rig: Fix not valid categories for functions and variables

diff --git a/Engine/Plugins/Animation/ControlRig/Source/ControlRigEditor/Private/ControlRigGraphDetails.cpp b/Engine/Plugins/Animation/ControlRig/Source/ControlRigEditor/Private/ControlRigGraphDetails.cpp
--- a/Engine/Plugins/Animation/ControlRig/Source/ControlRigEditor/Private/ControlRigGraphDetails.cpp
+++ b/Engine/Plugins/Animation/ControlRig/Source/ControlRigEditor/Private/ControlRigGraphDetails.cpp
@@ -21,6 +21,7 @@
 #include "RigVMCore/RigVMExternalVariable.h"
 #include "Kismet2/BlueprintEditorUtils.h"
 #include "Graph/ControlRigGraphSchema.h"
+#include "EditorCategoryUtils.h"
 
 #define LOCTEXT_NAMESPACE "ControlRigGraphDetails"
 
@@ -653,6 +654,24 @@ void FControlRigGraphDetails::CustomizeDetails(IDetailLayoutBuilder& DetailLayou
 			.Font(IDetailLayoutBuilder::GetDetailFont())
 			.Text(this, &FControlRigGraphDetails::GetNodeCategory)
 			.OnTextCommitted(this, &FControlRigGraphDetails::SetNodeCategory)
+			.OnVerifyTextChanged_Lambda([&](const FText& InNewText, FText& OutErrorMessage) -> bool
+			{
+				const FText NewText = FEditorCategoryUtils::GetCategoryDisplayString(InNewText);
+				if (NewText.IsEmpty())
+				{
+					OutErrorMessage = LOCTEXT("CategoryEmpty", "Cannot add a category with an empty string.");
+					return false;
+				}
+				if (ControlRigBlueprintPtr.IsValid())
+				{
+					if (NewText.EqualTo(FText::FromString(ControlRigBlueprintPtr.Get()->GetName())))
+					{
+						OutErrorMessage = LOCTEXT("CategoryEqualsBlueprintName", "Cannot add a category with the same name as the blueprint.");
+						return false;
+					}
+				}
+				return true;
+			})
 		];
 
 		// node keywords
