Control Rig: Force creation of DataSourceRegistry during control rig initialization to avoid UObject creation from worker thread, which can prevent the DataSourceRegistry from getting GCed properly upon exiting PIE

diff --git a/Engine/Plugins/Animation/ControlRig/Source/ControlRig/Private/ControlRig.cpp b/Engine/Plugins/Animation/ControlRig/Source/ControlRig/Private/ControlRig.cpp
--- a/Engine/Plugins/Animation/ControlRig/Source/ControlRig/Private/ControlRig.cpp
+++ b/Engine/Plugins/Animation/ControlRig/Source/ControlRig/Private/ControlRig.cpp
@@ -167,6 +167,9 @@ void UControlRig::Initialize(bool bInitRigUnits)
 	InitializeFromCDO();
 	InstantiateVMFromCDO();
 
+	// Create the data source registry here to avoid UObject creation from Non-Game Threads
+	GetDataSourceRegistry();
+
 	// should refresh mapping 
 	RequestSetup();
 
