Control Rig: Renaming Control from details panel causes details panel to show invalid values for the modified control.
This is happening because of this old call sequence: PropagateBPtoInstance() -> SetDetail() -> PropagateBPtoInstance(), the new sequence should be PropagateBPtoInstance() -> PropagateBPtoInstance() -> SetDetail(). One of the propagation may be redundant but I choose to keep it for now to avoid regression. For details see comment in InitControlRigEditor()

diff --git a/Engine/Plugins/Animation/ControlRig/Source/ControlRigDeveloper/Private/ControlRigBlueprint.cpp b/Engine/Plugins/Animation/ControlRig/Source/ControlRigDeveloper/Private/ControlRigBlueprint.cpp
--- a/Engine/Plugins/Animation/ControlRig/Source/ControlRigDeveloper/Private/ControlRigBlueprint.cpp
+++ b/Engine/Plugins/Animation/ControlRig/Source/ControlRigDeveloper/Private/ControlRigBlueprint.cpp
@@ -107,6 +107,7 @@ UControlRigBlueprint::UControlRigBlueprint(const FObjectInitializer& ObjectIniti
 	Hierarchy = CreateDefaultSubobject<URigHierarchy>(TEXT("Hierarchy"));
 	HierarchyController = CreateDefaultSubobject<URigHierarchyController>(TEXT("HierarchyController"));
 	HierarchyController->SetHierarchy(Hierarchy);
+	// give BP a chance to propagate hierarchy changes to available control rig instances
 	HierarchyController->OnModified().AddUObject(this, &UControlRigBlueprint::HandleHierarchyModified);
 }
 
@@ -3282,6 +3283,8 @@ void UControlRigBlueprint::HandleHierarchyModified(ERigHierarchyNotification InN
 			break;
 		}
 	}
+
+	HierarchyModifiedEvent.Broadcast(InNotification, InHierarchy, InElement);
 	
 #endif
 }
diff --git a/Engine/Plugins/Animation/ControlRig/Source/ControlRigDeveloper/Public/ControlRigBlueprint.h b/Engine/Plugins/Animation/ControlRig/Source/ControlRigDeveloper/Public/ControlRigBlueprint.h
--- a/Engine/Plugins/Animation/ControlRig/Source/ControlRigDeveloper/Public/ControlRigBlueprint.h
+++ b/Engine/Plugins/Animation/ControlRig/Source/ControlRigDeveloper/Public/ControlRigBlueprint.h
@@ -475,11 +475,20 @@ public:
 	void PropagatePropertyFromBPToInstances(FRigElementKey InRigElement, const FProperty* InProperty);
 	void PropagatePropertyFromInstanceToBP(FRigElementKey InRigElement, const FProperty* InProperty, UControlRig* InInstance);
 
+	/**
+	* Returns the modified event, which can be used to 
+	* subscribe to topological changes happening within the hierarchy. The event is broadcasted only after all hierarchy instances are up to date
+	* @return The event used for subscription.
+	*/
+	FRigHierarchyModifiedEvent& OnHierarchyModified() { return HierarchyModifiedEvent; }
+
 private:
 
 	UPROPERTY()
 	UControlRigValidator* Validator;
 
+	FRigHierarchyModifiedEvent	HierarchyModifiedEvent;
+
 	void HandleHierarchyModified(ERigHierarchyNotification InNotification, URigHierarchy* InHierarchy, const FRigBaseElement* InElement);
 
 	// Class used to temporarily cache all 
diff --git a/Engine/Plugins/Animation/ControlRig/Source/ControlRigEditor/Private/Editor/ControlRigEditor.cpp b/Engine/Plugins/Animation/ControlRig/Source/ControlRigEditor/Private/Editor/ControlRigEditor.cpp
--- a/Engine/Plugins/Animation/ControlRig/Source/ControlRigEditor/Private/Editor/ControlRigEditor.cpp
+++ b/Engine/Plugins/Animation/ControlRig/Source/ControlRigEditor/Private/Editor/ControlRigEditor.cpp
@@ -223,10 +223,10 @@ FControlRigEditor::~FControlRigEditor()
 
 		UControlRigBlueprint::sCurrentlyOpenedRigBlueprints.Remove(RigBlueprint);
 
-		RigBlueprint->Hierarchy->OnModified().RemoveAll(this);
+		RigBlueprint->OnHierarchyModified().RemoveAll(this);
 		if (FControlRigEditMode* EditMode = GetEditMode())
 		{
-			RigBlueprint->Hierarchy->OnModified().RemoveAll(EditMode);
+			RigBlueprint->OnHierarchyModified().RemoveAll(EditMode);
 		}
 		RigBlueprint->OnRefreshEditor().RemoveAll(this);
 		RigBlueprint->OnVariableDropped().RemoveAll(this);
@@ -462,14 +462,27 @@ void FControlRigEditor::InitControlRigEditor(const EToolkitMode::Type Mode, cons
 			}
 		}
 
-		InControlRigBlueprint->Hierarchy->OnModified().AddSP(this, &FControlRigEditor::OnHierarchyModified);
+		// listening to the BP's event instead of BP's Hierarchy's Event ensure a propagation order of
+		// 1. Hierarchy change in BP
+		// 2. BP propagate to instances
+		// 3. Editor forces propagation again, and reflects hierarchy change in either instances or BP
+		// 
+		// if directly listening to BP's Hierarchy's Event, this ordering is not guaranteed due to multicast,
+		// a problematic order we have encountered looks like:
+		// 1. Hierarchy change in BP
+		// 2. FControlRigEditor::OnHierarchyModified performs propagation from BP to instances, refresh UI
+		// 3. BP performs propagation again in UControlRigBlueprint::HandleHierarchyModified, invalidates the rig element
+		//    that the UI is observing
+		// 4. Editor UI shows an invalid rig element
+		InControlRigBlueprint->OnHierarchyModified().AddSP(this, &FControlRigEditor::OnHierarchyModified);
+		
 		InControlRigBlueprint->OnRefreshEditor().AddSP(this, &FControlRigEditor::HandleRefreshEditorFromBlueprint);
 		InControlRigBlueprint->OnVariableDropped().AddSP(this, &FControlRigEditor::HandleVariableDroppedFromBlueprint);
 		InControlRigBlueprint->OnBreakpointAdded().AddSP(this, &FControlRigEditor::HandleBreakpointAdded);
 
 		if (FControlRigEditMode* EditMode = GetEditMode())
 		{
-			InControlRigBlueprint->Hierarchy->OnModified().AddSP(EditMode, &FControlRigEditMode::OnHierarchyModified);
+			InControlRigBlueprint->OnHierarchyModified().AddSP(EditMode, &FControlRigEditMode::OnHierarchyModified);
 		}
 
 		InControlRigBlueprint->OnNodeDoubleClicked().AddSP(this, &FControlRigEditor::OnNodeDoubleClicked);
@@ -4341,9 +4354,17 @@ void FControlRigEditor::OnHierarchyChanged()
 				ClearDetailObject();
 			}
 		}
-		else if (!ControlRigBP->Hierarchy->Contains(RigElementInDetailPanel))
+		else
 		{
+			// PropagateHierarchyFromBPToInstances() recreates the rig elements in a instance's hierarchy
+			// update the details viewer to observe the newly created elements
+			const FRigElementKey CurrentElementInDetailsPanel = RigElementInDetailPanel;
 			ClearDetailObject();
+
+			if (ControlRigBP->Hierarchy->Contains(CurrentElementInDetailsPanel))
+			{
+				SetDetailStruct(CurrentElementInDetailsPanel);
+			}
 		}
 	}
 	else
