Add missing includes

#rnx

diff --git a/Engine/Plugins/Animation/ControlRig/Source/ControlRigEditor/Private/Graph/SControlRigGraphPinVariableBinding.cpp b/Engine/Plugins/Animation/ControlRig/Source/ControlRigEditor/Private/Graph/SControlRigGraphPinVariableBinding.cpp
--- a/Engine/Plugins/Animation/ControlRig/Source/ControlRigEditor/Private/Graph/SControlRigGraphPinVariableBinding.cpp
+++ b/Engine/Plugins/Animation/ControlRig/Source/ControlRigEditor/Private/Graph/SControlRigGraphPinVariableBinding.cpp
@@ -3,6 +3,10 @@
 
 #include "Graph/SControlRigGraphPinVariableBinding.h"
 #include "Widgets/Layout/SBox.h"
+#include "Widgets/Layout/SSpacer.h"
+#include "Widgets/Images/SImage.h"
+#include "Framework/MultiBox/MultiBoxExtender.h"
+#include "Framework/MultiBox/MultiBoxBuilder.h"
 #include "EdGraphSchema_K2.h"
 #include "DetailLayoutBuilder.h"
 #include "RigVMModel/RigVMController.h"
