Property display name change to make behaviour more obvious to users

diff --git a/Engine/Plugins/Animation/ControlRig/Source/ControlRig/Public/AnimNode_ControlRig.h b/Engine/Plugins/Animation/ControlRig/Source/ControlRig/Public/AnimNode_ControlRig.h
--- a/Engine/Plugins/Animation/ControlRig/Source/ControlRig/Public/AnimNode_ControlRig.h
+++ b/Engine/Plugins/Animation/ControlRig/Source/ControlRig/Public/AnimNode_ControlRig.h
@@ -55,7 +55,8 @@ private:
 	UPROPERTY(EditAnywhere, Category = Settings, meta = (PinShownByDefault, DisplayName = "bEnabled", DisplayAfter = "AlphaScaleBias"))
 	uint8 bAlphaBoolEnabled : 1;
 
-	UPROPERTY(EditAnywhere, Category=Settings)
+	// Override the initial transforms with those taken from the mesh component
+	UPROPERTY(EditAnywhere, Category=Settings, meta = (DisplayName = "Set Initial Transforms From Mesh"))
 	uint8 bSetRefPoseFromSkeleton : 1;
 
 	UPROPERTY(EditAnywhere, Category=Settings)
