Control Rig: Fix Control Rig name labeled None from Animation Window

diff --git a/Engine/Plugins/Animation/ControlRig/Source/ControlRigEditor/Private/EditMode/SControlRigEditModeTools.cpp b/Engine/Plugins/Animation/ControlRig/Source/ControlRigEditor/Private/EditMode/SControlRigEditModeTools.cpp
--- a/Engine/Plugins/Animation/ControlRig/Source/ControlRigEditor/Private/EditMode/SControlRigEditModeTools.cpp
+++ b/Engine/Plugins/Animation/ControlRig/Source/ControlRigEditor/Private/EditMode/SControlRigEditModeTools.cpp
@@ -107,7 +107,7 @@ public:
 			TSharedPtr<IPropertyHandle> ValuePropertyHandle = DetailLayout.GetProperty(ValuePropertyName, Proxy->GetClass());
 			if (ValuePropertyHandle)
 			{
-				ValuePropertyHandle->SetPropertyDisplayName(FText::FromName(ControlElement->Settings.DisplayName));
+				ValuePropertyHandle->SetPropertyDisplayName(FText::FromName(ControlElement->GetDisplayName()));
 			}
 
 			URigHierarchy* Hierarchy = Proxy->ControlRig->GetHierarchy();
