Control Rig: Fix crash when rig element is renamed with empty string

diff --git a/Engine/Plugins/Animation/ControlRig/Source/ControlRigEditor/Private/Editor/SRigHierarchy.cpp b/Engine/Plugins/Animation/ControlRig/Source/ControlRigEditor/Private/Editor/SRigHierarchy.cpp
--- a/Engine/Plugins/Animation/ControlRig/Source/ControlRigEditor/Private/Editor/SRigHierarchy.cpp
+++ b/Engine/Plugins/Animation/ControlRig/Source/ControlRigEditor/Private/Editor/SRigHierarchy.cpp
@@ -2385,6 +2385,12 @@ bool SRigHierarchy::OnVerifyNameChanged(const FRigElementKey& OldKey, const FStr
 		return true;
 	}
 
+	if (NewName.IsEmpty())
+	{
+		OutErrorMessage = FText::FromString(TEXT("Name is empty."));
+		return false;
+	}
+
 	// make sure there is no duplicate
 	if (ControlRigBlueprint.IsValid())
 	{
