Control Rig: Guard against nullptr within SRigHierarchy

diff --git a/Engine/Plugins/Animation/ControlRig/Source/ControlRigEditor/Private/Editor/SRigHierarchy.cpp b/Engine/Plugins/Animation/ControlRig/Source/ControlRigEditor/Private/Editor/SRigHierarchy.cpp
--- a/Engine/Plugins/Animation/ControlRig/Source/ControlRigEditor/Private/Editor/SRigHierarchy.cpp
+++ b/Engine/Plugins/Animation/ControlRig/Source/ControlRigEditor/Private/Editor/SRigHierarchy.cpp
@@ -1305,7 +1305,10 @@ void SRigHierarchy::HandleSetObjectBeingDebugged(UObject* InObject)
 
 void SRigHierarchy::ClearDetailPanel() const
 {
-	ControlRigEditor.Pin()->ClearDetailObject();
+	if(ControlRigEditor.IsValid())
+	{
+		ControlRigEditor.Pin()->ClearDetailObject();
+	}
 }
 
 TSharedRef< SWidget > SRigHierarchy::CreateFilterMenu()
@@ -1599,6 +1602,11 @@ void SRigHierarchy::RefreshHierarchy(const FAssetData& InAssetData)
 	}
 	TGuardValue<bool> GuardRigHierarchyChanges(bIsChangingRigHierarchy, true);
 
+	if(!ControlRigEditor.IsValid())
+	{
+		return;
+	}
+
 	URigHierarchy* Hierarchy = GetHierarchy();
 	USkeletalMesh* Mesh = Cast<USkeletalMesh>(InAssetData.GetAsset());
 	if (Mesh && Hierarchy)
@@ -1667,6 +1675,11 @@ void SRigHierarchy::ImportHierarchy(const FAssetData& InAssetData)
 	}
 	TGuardValue<bool> GuardRigHierarchyChanges(bIsChangingRigHierarchy, true);
 
+	if (!ControlRigEditor.IsValid())
+	{
+		return;
+	}
+
 	URigHierarchy* Hierarchy = GetHierarchy();
 	USkeletalMesh* Mesh = Cast<USkeletalMesh>(InAssetData.GetAsset());
 	if (Mesh && Hierarchy)
@@ -1780,6 +1793,11 @@ bool SRigHierarchy::IsControlOrNullSelected() const
 
 void SRigHierarchy::HandleDeleteItem()
 {
+	if(!ControlRigEditor.IsValid())
+	{
+		return;
+	}
+
 	URigHierarchy* Hierarchy = GetHierarchy();
  	if (Hierarchy)
  	{
@@ -1849,6 +1867,11 @@ bool SRigHierarchy::CanDeleteItem() const
 /** Delete Item */
 void SRigHierarchy::HandleNewItem(ERigElementType InElementType)
 {
+	if(!ControlRigEditor.IsValid())
+	{
+		return;
+	}
+
 	URigHierarchy* Hierarchy = GetHierarchy();
 	if (Hierarchy)
 	{
@@ -1919,6 +1942,11 @@ bool SRigHierarchy::CanDuplicateItem() const
 /** Duplicate Item */
 void SRigHierarchy::HandleDuplicateItem()
 {
+	if(!ControlRigEditor.IsValid())
+	{
+		return;
+	}
+
 	URigHierarchy* Hierarchy = GetHierarchy();
 	if (Hierarchy)
 	{
@@ -1951,6 +1979,11 @@ void SRigHierarchy::HandleDuplicateItem()
 /** Mirror Item */
 void SRigHierarchy::HandleMirrorItem()
 {
+	if(!ControlRigEditor.IsValid())
+	{
+		return;
+	}
+	
 	URigHierarchy* Hierarchy = GetHierarchy();
 	if (Hierarchy)
 	{
@@ -2025,6 +2058,11 @@ bool SRigHierarchy::CanRenameItem() const
 /** Delete Item */
 void SRigHierarchy::HandleRenameItem()
 {
+	if(!ControlRigEditor.IsValid())
+	{
+		return;
+	}
+
 	if (!CanRenameItem())
 	{
 		return;
