Control Rig: Change default gizmo library

diff --git a/Engine/Plugins/Animation/ControlRig/Source/ControlRigDeveloper/Public/ControlRigBlueprint.h b/Engine/Plugins/Animation/ControlRig/Source/ControlRigDeveloper/Public/ControlRigBlueprint.h
--- a/Engine/Plugins/Animation/ControlRig/Source/ControlRigDeveloper/Public/ControlRigBlueprint.h
+++ b/Engine/Plugins/Animation/ControlRig/Source/ControlRigDeveloper/Public/ControlRigBlueprint.h
@@ -338,7 +338,7 @@ public:
 #endif
 
 #if WITH_EDITORONLY_DATA
-	UPROPERTY(EditAnywhere, config, Category = DefaultGizmo)
+	UPROPERTY(EditAnywhere, Category = DefaultGizmo)
 	TSoftObjectPtr<UControlRigGizmoLibrary> GizmoLibrary;
 #endif
 
