Control Rig: Remove unnecessary ensure in URigHierarchy::UnregisterListeningHirearch

diff --git a/Engine/Plugins/Animation/ControlRig/Source/ControlRig/Private/Rigs/RigHierarchy.cpp b/Engine/Plugins/Animation/ControlRig/Source/ControlRig/Private/Rigs/RigHierarchy.cpp
--- a/Engine/Plugins/Animation/ControlRig/Source/ControlRig/Private/Rigs/RigHierarchy.cpp
+++ b/Engine/Plugins/Animation/ControlRig/Source/ControlRig/Private/Rigs/RigHierarchy.cpp
@@ -271,7 +271,7 @@ void URigHierarchy::CopyHierarchy(URigHierarchy* InHierarchy)
 #if WITH_EDITOR
 void URigHierarchy::RegisterListeningHierarchy(URigHierarchy* InHierarchy)
 {
-	if (ensure(InHierarchy))
+	if (InHierarchy)
 	{
 		bool bFoundListener = false;
 		for(int32 ListenerIndex = ListeningHierarchies.Num() - 1; ListenerIndex >= 0; ListenerIndex--)
@@ -298,7 +298,7 @@ void URigHierarchy::RegisterListeningHierarchy(URigHierarchy* InHierarchy)
 
 void URigHierarchy::UnregisterListeningHierarchy(URigHierarchy* InHierarchy)
 {
-	if (ensure(InHierarchy))
+	if (InHierarchy)
 	{
 		for(int32 ListenerIndex = ListeningHierarchies.Num() - 1; ListenerIndex >= 0; ListenerIndex--)
 		{
