Control Rig: Temporarily disable interaction scope during a click to support selection across interaction rigs

diff --git a/Engine/Plugins/Animation/ControlRig/Source/ControlRig/Public/ControlRig.h b/Engine/Plugins/Animation/ControlRig/Source/ControlRig/Public/ControlRig.h
--- a/Engine/Plugins/Animation/ControlRig/Source/ControlRig/Public/ControlRig.h
+++ b/Engine/Plugins/Animation/ControlRig/Source/ControlRig/Public/ControlRig.h
@@ -309,20 +309,6 @@ public:
 		}
 	}
 
-	/** Turn On Interact, MUST Call SetInteractOff*/
-	FORCEINLINE_DEBUGGABLE void SetInteractOn()
-	{
-		++InteractionBracket;
-		++InterRigSyncBracket;
-	}
-
-	/** Turn Off Interact, MUST Have Called SetInteractOn*/
-	FORCEINLINE_DEBUGGABLE void SetInteractOff()
-	{
-		--InteractionBracket;
-		--InterRigSyncBracket;
-	}
-
 	bool SetControlGlobalTransform(const FName& InControlName, const FTransform& InGlobalTransform, bool bNotify = true, const FRigControlModifiedContext& Context = FRigControlModifiedContext(), bool bSetupUndo = true);
 
 	virtual FRigControlValue GetControlValueFromGlobalTransform(const FName& InControlName, const FTransform& InGlobalTransform);
@@ -740,7 +726,7 @@ class CONTROLRIG_API FControlRigInteractionScope
 {
 public:
 
-	FORCEINLINE FControlRigInteractionScope(UControlRig* InControlRig)
+	FORCEINLINE_DEBUGGABLE FControlRigInteractionScope(UControlRig* InControlRig)
 		: ControlRig(InControlRig)
 		, InteractionBracketScope(InControlRig->InteractionBracket)
 		, SyncBracketScope(InControlRig->InterRigSyncBracket)
@@ -748,9 +734,9 @@ public:
 		InControlRig->GetHierarchy()->StartInteraction();
 	}
 
-	FORCEINLINE ~FControlRigInteractionScope()
+	FORCEINLINE_DEBUGGABLE ~FControlRigInteractionScope()
 	{
-		if(ControlRig.IsValid())
+		if(ensure(ControlRig.IsValid()))
 		{
 			ControlRig->GetHierarchy()->EndInteraction();
 		}
diff --git a/Engine/Plugins/Animation/ControlRig/Source/ControlRigEditor/Private/EditMode/ControlRigEditMode.cpp b/Engine/Plugins/Animation/ControlRig/Source/ControlRigEditor/Private/EditMode/ControlRigEditMode.cpp
--- a/Engine/Plugins/Animation/ControlRig/Source/ControlRigEditor/Private/EditMode/ControlRigEditMode.cpp
+++ b/Engine/Plugins/Animation/ControlRig/Source/ControlRigEditor/Private/EditMode/ControlRigEditMode.cpp
@@ -836,7 +836,15 @@ bool FControlRigEditMode::HandleClick(FEditorViewportClient* InViewportClient, H
 				if (GizmoActor->IsSelectable())
 				{
 					FScopedTransaction ScopedTransaction(LOCTEXT("SelectControlTransaction", "Select Control"), IsInLevelEditor() && !GIsTransacting);
-						
+
+					// temporarily disable the interaction scope
+					const bool bInteractionScopePresent = InteractionScope != nullptr; 
+					if (bInteractionScopePresent)
+					{
+						delete InteractionScope;
+						InteractionScope = nullptr;
+					}
+					
 					const FName& ControlName = GizmoActor->ControlName;
 					if (Click.IsShiftDown()) //guess we just select
 					{
@@ -856,7 +864,15 @@ bool FControlRigEditMode::HandleClick(FEditorViewportClient* InViewportClient, H
 						ClearRigElementSelection(FRigElementTypeHelper::ToMask(ERigElementType::Control));
 						SetRigElementSelection(ERigElementType::Control, ControlName, true);
 					}
-	
+
+					if (bInteractionScopePresent)
+					{
+						if(UControlRig* ControlRig = GetControlRig(true))
+						{
+							InteractionScope = new FControlRigInteractionScope(ControlRig);
+						}
+					}
+
 					// for now we show this menu all the time if body is selected
 					// if we want some global menu, we'll have to move this
 					if (Click.GetKey() == EKeys::RightMouseButton)
