ControlRigComponent: Changed MappedElement to use ComponentProperty instead of PathToComponenet to record component ref since PathToComponent cannot be serialized. Also added clear mapping elements at pre init to avoid mapped elements accumulating

diff --git a/Engine/Plugins/Animation/ControlRig/Source/ControlRig/Private/ControlRigComponent.cpp b/Engine/Plugins/Animation/ControlRig/Source/ControlRig/Private/ControlRigComponent.cpp
--- a/Engine/Plugins/Animation/ControlRig/Source/ControlRig/Private/ControlRigComponent.cpp
+++ b/Engine/Plugins/Animation/ControlRig/Source/ControlRig/Private/ControlRigComponent.cpp
@@ -301,6 +301,8 @@ void UControlRigComponent::Initialize()
 
 	TGuardValue<bool> InitializeBracket(bIsInsideInitializeBracket, true);
 	
+	ClearMappedElements();
+
 #if WITH_EDITOR
 	if (bUpdateInEditor)
 	{
@@ -503,7 +505,7 @@ void UControlRigComponent::AddMappedComponents(TArray<FControlRigComponentMapped
 
 		FControlRigComponentMappedElement ElementToMap;
 		ElementToMap.ComponentReference.OtherActor = Component->GetOwner() != GetOwner() ? Component->GetOwner() : nullptr;
-		ElementToMap.ComponentReference.PathToComponent = Component->GetName();
+		ElementToMap.ComponentReference.ComponentProperty = Component->GetFName();
 
 		ElementToMap.ElementName = ComponentToMap.ElementName;
 		ElementToMap.ElementType = ComponentToMap.ElementType;
@@ -609,7 +611,7 @@ void UControlRigComponent::AddMappedSkeletalMesh(USkeletalMeshComponent* Skeleta
 
 		FControlRigComponentMappedElement ElementToMap;
 		ElementToMap.ComponentReference.OtherActor = SkeletalMeshComponent->GetOwner() != GetOwner() ? SkeletalMeshComponent->GetOwner() : nullptr;
-		ElementToMap.ComponentReference.PathToComponent = SkeletalMeshComponent->GetName();
+		ElementToMap.ComponentReference.ComponentProperty = SkeletalMeshComponent->GetFName();
 
 		ElementToMap.ElementName = BoneToMap.Source;
 		ElementToMap.ElementType = ERigElementType::Bone;
@@ -628,7 +630,7 @@ void UControlRigComponent::AddMappedSkeletalMesh(USkeletalMeshComponent* Skeleta
 
 		FControlRigComponentMappedElement ElementToMap;
 		ElementToMap.ComponentReference.OtherActor = SkeletalMeshComponent->GetOwner() != GetOwner() ? SkeletalMeshComponent->GetOwner() : nullptr;
-		ElementToMap.ComponentReference.PathToComponent = SkeletalMeshComponent->GetName();
+		ElementToMap.ComponentReference.ComponentProperty = SkeletalMeshComponent->GetFName();
 
 		ElementToMap.ElementName = CurveToMap.Source;
 		ElementToMap.ElementType = ERigElementType::Curve;
