Control Rig: Added Control Rig Context Menu Context to GraphNodeContextMenu + Added GraphPanel specific context to ControlRigContextMenuContext

diff --git a/Engine/Plugins/Animation/ControlRig/Source/ControlRigEditor/Private/Editor/ControlRigContextMenuContext.cpp b/Engine/Plugins/Animation/ControlRig/Source/ControlRigEditor/Private/Editor/ControlRigContextMenuContext.cpp
--- a/Engine/Plugins/Animation/ControlRig/Source/ControlRigEditor/Private/Editor/ControlRigContextMenuContext.cpp
+++ b/Engine/Plugins/Animation/ControlRig/Source/ControlRigEditor/Private/Editor/ControlRigContextMenuContext.cpp
@@ -6,21 +6,23 @@
 #include "ControlRigEditor.h"
 #include "Slate/Public/Framework/Application/SlateApplication.h"
 
-UControlRigBlueprint* UControlRigContextMenuContext::GetControlRigBlueprint() const
+void UControlRigContextMenuContext::Init(TWeakObjectPtr<UControlRigBlueprint> InControlRigBlueprint, const FControlRigRigHierarchyDragAndDropContext& InDragAndDropContext, const FControlRigGraphNodeContextMenuContext& InGraphNodeContext)
 {
-	if (ControlRigEditor.IsValid())
-	{
-		return ControlRigEditor.Pin()->GetControlRigBlueprint();
-	}
+	ControlRigBlueprint = InControlRigBlueprint;
+	DragAndDropContext = InDragAndDropContext;
+	GraphNodeContextMenuContext = InGraphNodeContext;
+}
 
-	return nullptr;
+UControlRigBlueprint* UControlRigContextMenuContext::GetControlRigBlueprint() const
+{
+	return ControlRigBlueprint.Get();
 }
 
 UControlRig* UControlRigContextMenuContext::GetControlRig() const
 {
-	if (UControlRigBlueprint* Blueprint = GetControlRigBlueprint())
+	if (UControlRigBlueprint* RigBlueprint = GetControlRigBlueprint())
 	{
-		if (UControlRig* ControlRig = Cast<UControlRig>(Blueprint->GetObjectBeingDebugged()))
+		if (UControlRig* ControlRig = Cast<UControlRig>(RigBlueprint->GetObjectBeingDebugged()))
 		{
 			return ControlRig;
 		}
@@ -33,7 +35,12 @@ bool UControlRigContextMenuContext::IsAltDown() const
 	return FSlateApplication::Get().GetModifierKeys().IsAltDown();
 }
 
-FControlRigRigHierarchyDragAndDropContext UControlRigContextMenuContext::GetDragAndDropContext()
+FControlRigRigHierarchyDragAndDropContext UControlRigContextMenuContext::GetRigHierarchyDragAndDropContext()
 {
 	return DragAndDropContext;
 }
+
+FControlRigGraphNodeContextMenuContext UControlRigContextMenuContext::GetGraphNodeContextMenuContext()
+{
+	return GraphNodeContextMenuContext;
+}
diff --git a/Engine/Plugins/Animation/ControlRig/Source/ControlRigEditor/Private/Editor/ControlRigContextMenuContext.h b/Engine/Plugins/Animation/ControlRig/Source/ControlRigEditor/Private/Editor/ControlRigContextMenuContext.h
--- a/Engine/Plugins/Animation/ControlRig/Source/ControlRigEditor/Private/Editor/ControlRigContextMenuContext.h
+++ b/Engine/Plugins/Animation/ControlRig/Source/ControlRigEditor/Private/Editor/ControlRigContextMenuContext.h
@@ -3,6 +3,7 @@
 #pragma once
 
 #include "CoreMinimal.h"
+
 #include "Rigs/RigHierarchyDefines.h"
 
 #include "ControlRigContextMenuContext.generated.h"
@@ -10,6 +11,9 @@
 class UControlRig;
 class UControlRigBlueprint;
 class FControlRigEditor;
+class URigVMGraph;
+class URigVMNode;
+class URigVMPin;
 
 USTRUCT(BlueprintType)
 struct FControlRigRigHierarchyDragAndDropContext
@@ -31,6 +35,33 @@ struct FControlRigRigHierarchyDragAndDropContext
 	FRigElementKey TargetElementKey;
 };
 
+USTRUCT(BlueprintType)
+struct FControlRigGraphNodeContextMenuContext
+{
+	GENERATED_BODY()
+
+	FControlRigGraphNodeContextMenuContext() = default;
+	
+	FControlRigGraphNodeContextMenuContext(TObjectPtr<const URigVMGraph> InGraph, TObjectPtr<const URigVMNode> InNode, TObjectPtr<const URigVMPin> InPin)
+		: Graph(InGraph)
+		, Node(InNode)
+		, Pin(InPin)
+	{
+	}
+	
+	/** The graph associated with this context. */
+	UPROPERTY(BlueprintReadOnly, Category = ControlRigEditorExtensions)
+	TObjectPtr<const URigVMGraph> Graph;
+
+	/** The node associated with this context. */
+	UPROPERTY(BlueprintReadOnly, Category = ControlRigEditorExtensions)
+	TObjectPtr<const URigVMNode> Node;
+
+	/** The pin associated with this context; may be NULL when over a node. */
+	UPROPERTY(BlueprintReadOnly, Category = ControlRigEditorExtensions)
+	TObjectPtr<const URigVMPin> Pin;
+};
+
 UCLASS(BlueprintType)
 class UControlRigContextMenuContext : public UObject
 {
@@ -39,15 +70,12 @@ class UControlRigContextMenuContext : public UObject
 public:
 	/**
 	 *	Initialize the Context
-	 * @param InControlRigEditor 		The Control Rig editor that hosts the menu
+	 * @param InControlRigBlueprint 	The Control Rig Bluerpint currently opened in the Editor
 	 * @param InDragAndDropContext 		Optional, only used for the menu that shows up after a drag & drop action
+	 * @param InGraphNodeContext 		Optional, only used for the graph node context menu
 	*/
-	void Init(TWeakPtr<FControlRigEditor> InControlRigEditor, FControlRigRigHierarchyDragAndDropContext InDragAndDropContext = FControlRigRigHierarchyDragAndDropContext())
-	{
-		ControlRigEditor = InControlRigEditor;
-		DragAndDropContext = InDragAndDropContext;
-	}
-	
+	void Init(TWeakObjectPtr<UControlRigBlueprint> InControlRigBlueprint, const FControlRigRigHierarchyDragAndDropContext& InDragAndDropContext = FControlRigRigHierarchyDragAndDropContext(), const FControlRigGraphNodeContextMenuContext& InGraphNodeContext = FControlRigGraphNodeContextMenuContext());
+
 	/** Get the control rig blueprint that we are editing */
 	UFUNCTION(BlueprintCallable, Category = ControlRigEditorExtensions)
     UControlRigBlueprint* GetControlRigBlueprint() const;
@@ -60,13 +88,18 @@ public:
 	UFUNCTION(BlueprintCallable, Category = ControlRigEditorExtensions)
 	bool IsAltDown() const;
 	
-	/** Returns the source and target element keys of a drag & drop action */
+	/** Returns context for a drag & drop action that contains source and target element keys */ 
 	UFUNCTION(BlueprintCallable, Category = ControlRigEditorExtensions)
-	FControlRigRigHierarchyDragAndDropContext GetDragAndDropContext();
+	FControlRigRigHierarchyDragAndDropContext GetRigHierarchyDragAndDropContext();
+
+	/** Returns context for graph node context menu */
+	UFUNCTION(BlueprintCallable, Category = ControlRigEditorExtensions)
+	FControlRigGraphNodeContextMenuContext GetGraphNodeContextMenuContext();
 
 private:
-	/** Our owning control rig editor */
-	TWeakPtr<FControlRigEditor> ControlRigEditor;
+	TWeakObjectPtr<UControlRigBlueprint> ControlRigBlueprint;
 
 	FControlRigRigHierarchyDragAndDropContext DragAndDropContext;
+
+	FControlRigGraphNodeContextMenuContext GraphNodeContextMenuContext;
 };
\ No newline at end of file
diff --git a/Engine/Plugins/Animation/ControlRig/Source/ControlRigEditor/Private/Editor/ControlRigEditor.cpp b/Engine/Plugins/Animation/ControlRig/Source/ControlRigEditor/Private/Editor/ControlRigEditor.cpp
--- a/Engine/Plugins/Animation/ControlRig/Source/ControlRigEditor/Private/Editor/ControlRigEditor.cpp
+++ b/Engine/Plugins/Animation/ControlRig/Source/ControlRigEditor/Private/Editor/ControlRigEditor.cpp
@@ -89,6 +89,7 @@
 #include "Engine/StaticMeshActor.h"
 #include "Components/StaticMeshComponent.h"
 #include "ControlRig/Private/Units/Execution/RigUnit_SequenceExecution.h"
+#include "ControlRigContextMenuContext.h"
 
 #define LOCTEXT_NAMESPACE "ControlRigEditor"
 
@@ -1892,6 +1893,40 @@ FLinearColor FControlRigEditor::GetWorldCentricTabColorScale() const
 	return FLinearColor( 0.5f, 0.25f, 0.35f, 0.5f );
 }
 
+void FControlRigEditor::InitToolMenuContext(FToolMenuContext& MenuContext)
+{
+	FAssetEditorToolkit::InitToolMenuContext(MenuContext);
+
+	if (UControlRigBlueprint* RigBlueprint = GetControlRigBlueprint())
+	{
+		URigVMGraph* Model = nullptr;
+		URigVMNode* Node = nullptr;
+		URigVMPin* Pin = nullptr;
+		
+		if (UGraphNodeContextMenuContext* GraphNodeContext = MenuContext.FindContext<UGraphNodeContextMenuContext>())
+		{
+		
+			if (GraphNodeContext->Node)
+			{
+				Model = RigBlueprint->GetModel(GraphNodeContext->Graph);
+				Node = Model->FindNodeByName(GraphNodeContext->Node->GetFName());
+			}
+		
+			if (GraphNodeContext->Pin && Node)
+			{
+				Pin = Model->FindPin(GraphNodeContext->Pin->GetName());
+			}
+		}
+		
+		UControlRigContextMenuContext* ControlRigContext = NewObject<UControlRigContextMenuContext>();
+		ControlRigContext->Init(RigBlueprint,
+			FControlRigRigHierarchyDragAndDropContext(),
+			FControlRigGraphNodeContextMenuContext(Model, Node, Pin));
+
+		MenuContext.AddObject(ControlRigContext);
+	}
+}
+
 bool FControlRigEditor::TransactionObjectAffectsBlueprint(UObject* InTransactedObject)
 {
 	UControlRigBlueprint* RigBlueprint = Cast<UControlRigBlueprint>(GetBlueprintObj());
diff --git a/Engine/Plugins/Animation/ControlRig/Source/ControlRigEditor/Private/Editor/ControlRigEditor.h b/Engine/Plugins/Animation/ControlRig/Source/ControlRigEditor/Private/Editor/ControlRigEditor.h
--- a/Engine/Plugins/Animation/ControlRig/Source/ControlRigEditor/Private/Editor/ControlRigEditor.h
+++ b/Engine/Plugins/Animation/ControlRig/Source/ControlRigEditor/Private/Editor/ControlRigEditor.h
@@ -102,6 +102,7 @@ public:
 	{
 		return FString(TEXT("Engine/Animation/ControlRig"));
 	}
+	virtual void InitToolMenuContext(FToolMenuContext& MenuContext) override;
 
 	// BlueprintEditor interface
 	virtual bool TransactionObjectAffectsBlueprint(UObject* InTransactedObject) override;
@@ -468,4 +469,4 @@ protected:
 	friend class SControlRigStackView;
 	friend class SRigHierarchy;
 	friend struct FRigHierarchyTabSummoner;
-};
\ No newline at end of file
+};
diff --git a/Engine/Plugins/Animation/ControlRig/Source/ControlRigEditor/Private/Editor/SRigHierarchy.cpp b/Engine/Plugins/Animation/ControlRig/Source/ControlRigEditor/Private/Editor/SRigHierarchy.cpp
--- a/Engine/Plugins/Animation/ControlRig/Source/ControlRigEditor/Private/Editor/SRigHierarchy.cpp
+++ b/Engine/Plugins/Animation/ControlRig/Source/ControlRigEditor/Private/Editor/SRigHierarchy.cpp
@@ -1459,7 +1459,7 @@ UToolMenu* SRigHierarchy::GetOrCreateDragDropMenu(const TArray<FRigElementKey>&
 	}
 
 	UControlRigContextMenuContext* MenuContext = NewObject<UControlRigContextMenuContext>();
-	MenuContext->Init(ControlRigEditor, FControlRigRigHierarchyDragAndDropContext(DraggedKeys, TargetKey));
+	MenuContext->Init(ControlRigBlueprint, FControlRigRigHierarchyDragAndDropContext(DraggedKeys, TargetKey));
 	
 	UToolMenu* Menu = ToolMenus->GenerateMenu(MenuName, FToolMenuContext(MenuContext));
 
@@ -1552,7 +1552,7 @@ UToolMenu* SRigHierarchy::GetOrCreateContextMenu()
 
 	// individual entries in this menu can access members of this context, particularly useful for editor scripting
 	UControlRigContextMenuContext* ContextMenuContext = NewObject<UControlRigContextMenuContext>();
-	ContextMenuContext->Init(ControlRigEditor);
+	ContextMenuContext->Init(ControlRigBlueprint);
 	
 	FToolMenuContext MenuContext(CommandList);
 	MenuContext.AddObject(ContextMenuContext);
@@ -2711,7 +2711,7 @@ void SRigHierarchy::HandleParent(const FToolMenuContext& Context)
 {
 	if (UControlRigContextMenuContext* MenuContext = Cast<UControlRigContextMenuContext>(Context.FindByClass(UControlRigContextMenuContext::StaticClass())))
 	{
-		const FControlRigRigHierarchyDragAndDropContext DragAndDropContext = MenuContext->GetDragAndDropContext();
+		const FControlRigRigHierarchyDragAndDropContext DragAndDropContext = MenuContext->GetRigHierarchyDragAndDropContext();
 		ReparentOrMatchTransform(DragAndDropContext.DraggedElementKeys, DragAndDropContext.TargetElementKey, true);
 	}
 }
@@ -2720,7 +2720,7 @@ void SRigHierarchy::HandleAlign(const FToolMenuContext& Context)
 {
 	if (UControlRigContextMenuContext* MenuContext = Cast<UControlRigContextMenuContext>(Context.FindByClass(UControlRigContextMenuContext::StaticClass())))
 	{
-		const FControlRigRigHierarchyDragAndDropContext DragAndDropContext = MenuContext->GetDragAndDropContext();
+		const FControlRigRigHierarchyDragAndDropContext DragAndDropContext = MenuContext->GetRigHierarchyDragAndDropContext();
 		ReparentOrMatchTransform(DragAndDropContext.DraggedElementKeys, DragAndDropContext.TargetElementKey, false);
 	}
 }
diff --git a/Engine/Source/Editor/GraphEditor/Private/SGraphEditorImpl.cpp b/Engine/Source/Editor/GraphEditor/Private/SGraphEditorImpl.cpp
--- a/Engine/Source/Editor/GraphEditor/Private/SGraphEditorImpl.cpp
+++ b/Engine/Source/Editor/GraphEditor/Private/SGraphEditorImpl.cpp
@@ -1169,6 +1169,11 @@ FActionMenuContent SGraphEditorImpl::GraphEd_OnGetContextMenuFor(const FGraphCon
 				ToolkitMenuContext->Toolkit = AssetEditorToolkit;
 				Context.AddObject(ToolkitMenuContext);
 
+				if (TSharedPtr<FAssetEditorToolkit> SharedToolKit = AssetEditorToolkit.Pin())
+				{
+					SharedToolKit->InitToolMenuContext(Context);
+				}
+
 				// Need to additionally pass through the asset toolkit to hook up those commands?
 
 				UToolMenus* ToolMenus = UToolMenus::Get();
diff --git a/Engine/Source/Editor/Kismet/Private/BlueprintEditor.cpp b/Engine/Source/Editor/Kismet/Private/BlueprintEditor.cpp
--- a/Engine/Source/Editor/Kismet/Private/BlueprintEditor.cpp
+++ b/Engine/Source/Editor/Kismet/Private/BlueprintEditor.cpp
@@ -1567,7 +1567,8 @@ TSharedRef<SGraphEditor> FBlueprintEditor::CreateGraphEditorWidget(TSharedRef<FT
 		.GraphToEdit(InGraph)
 		.GraphEvents(InEvents)
 		.OnNavigateHistoryBack(FSimpleDelegate::CreateSP(this, &FBlueprintEditor::NavigateTab, FDocumentTracker::NavigateBackwards))
-		.OnNavigateHistoryForward(FSimpleDelegate::CreateSP(this, &FBlueprintEditor::NavigateTab, FDocumentTracker::NavigateForwards));
+		.OnNavigateHistoryForward(FSimpleDelegate::CreateSP(this, &FBlueprintEditor::NavigateTab, FDocumentTracker::NavigateForwards))
+		.AssetEditorToolkit(this->AsShared());
 		//@TODO: Crashes in command list code during the callback .OnGraphModuleReloaded(FEdGraphEvent::CreateSP(this, &FBlueprintEditor::ChangeOpenGraphInDocumentEditorWidget, WeakParent))
 		;
 
