Control Rig: Initilize all member variable for FControlRigGraphNodeContextMenuContext

diff --git a/Engine/Plugins/Animation/ControlRig/Source/ControlRigEditor/Private/Editor/ControlRigContextMenuContext.h b/Engine/Plugins/Animation/ControlRig/Source/ControlRigEditor/Private/Editor/ControlRigContextMenuContext.h
--- a/Engine/Plugins/Animation/ControlRig/Source/ControlRigEditor/Private/Editor/ControlRigContextMenuContext.h
+++ b/Engine/Plugins/Animation/ControlRig/Source/ControlRigEditor/Private/Editor/ControlRigContextMenuContext.h
@@ -40,7 +40,12 @@ struct FControlRigGraphNodeContextMenuContext
 {
 	GENERATED_BODY()
 
-	FControlRigGraphNodeContextMenuContext() = default;
+	FControlRigGraphNodeContextMenuContext()
+		: Graph(nullptr)
+		, Node(nullptr)
+		, Pin(nullptr)
+	{
+	}
 	
 	FControlRigGraphNodeContextMenuContext(TObjectPtr<const URigVMGraph> InGraph, TObjectPtr<const URigVMNode> InNode, TObjectPtr<const URigVMPin> InPin)
 		: Graph(InGraph)
