Control Rig: Add tracing capabilities during spawn to RigHierarchy

diff --git a/Engine/Plugins/Animation/ControlRig/Source/ControlRig/Private/Rigs/RigHierarchy.cpp b/Engine/Plugins/Animation/ControlRig/Source/ControlRig/Private/Rigs/RigHierarchy.cpp
--- a/Engine/Plugins/Animation/ControlRig/Source/ControlRig/Private/Rigs/RigHierarchy.cpp
+++ b/Engine/Plugins/Animation/ControlRig/Source/ControlRig/Private/Rigs/RigHierarchy.cpp
@@ -36,6 +36,7 @@ static void RigHierarchyCaptureCallStack(FString& OutCallstack, uint32 NumCallsT
 static TAutoConsoleVariable<int32> CVarControlRigHierarchyTraceAlways(TEXT("ControlRig.Hierarchy.TraceAlways"), 0, TEXT("if nonzero we will record all transform changes."));
 static TAutoConsoleVariable<int32> CVarControlRigHierarchyTraceCallstack(TEXT("ControlRig.Hierarchy.TraceCallstack"), 0, TEXT("if nonzero we will record the callstack for any trace entry.\nOnly works if(ControlRig.Hierarchy.TraceEnabled != 0)"));
 static TAutoConsoleVariable<int32> CVarControlRigHierarchyTracePrecision(TEXT("ControlRig.Hierarchy.TracePrecision"), 3, TEXT("sets the number digits in a float when tracing hierarchies."));
+static TAutoConsoleVariable<int32> CVarControlRigHierarchyTraceOnSpawn(TEXT("ControlRig.Hierarchy.TraceOnSpawn"), 0, TEXT("sets the number of frames to trace when a new hierarchy is spawned"));
 static int32 sRigHierarchyLastTrace = INDEX_NONE;
 static TCHAR sRigHierarchyTraceFormat[16];
 
@@ -80,6 +81,30 @@ FAutoConsoleCommandWithWorldAndArgs FCmdControlRigHierarchyTraceFrames
 
 const TArray<FRigBaseElement*> URigHierarchy::EmptyElementArray;
 
+URigHierarchy::URigHierarchy()
+: TopologyVersion(0)
+, bEnableDirtyPropagation(true)
+, Elements()
+, IndexLookup()
+, TransformStackIndex(0)
+, bTransactingForTransformChange(false)
+, bIsInteracting(false)
+, LastInteractedKey()
+, bSuspendNotifications(false)
+, ResetPoseHash(INDEX_NONE)
+#if WITH_EDITOR
+, bPropagatingChange(false)
+, bForcePropagation(false)
+, TraceFramesLeft(0)
+, TraceFramesCaptured(0)
+#endif
+{
+	Reset();
+#if WITH_EDITOR
+	TraceFrames(CVarControlRigHierarchyTraceOnSpawn->GetInt());
+#endif
+}
+
 URigHierarchy::~URigHierarchy()
 {
 	Reset();
diff --git a/Engine/Plugins/Animation/ControlRig/Source/ControlRig/Public/ControlRig.h b/Engine/Plugins/Animation/ControlRig/Source/ControlRig/Public/ControlRig.h
--- a/Engine/Plugins/Animation/ControlRig/Source/ControlRig/Public/ControlRig.h
+++ b/Engine/Plugins/Animation/ControlRig/Source/ControlRig/Public/ControlRig.h
@@ -370,23 +370,23 @@ public:
 
 	// Setup the initial transform / ref pose of the bones based upon an anim instance
 	// This uses the current refpose instead of the RefSkeleton pose.
-	void SetBoneInitialTransformsFromAnimInstance(UAnimInstance* InAnimInstance);
+	virtual void SetBoneInitialTransformsFromAnimInstance(UAnimInstance* InAnimInstance);
 
 	// Setup the initial transform / ref pose of the bones based upon an anim instance proxy
 	// This uses the current refpose instead of the RefSkeleton pose.
-	void SetBoneInitialTransformsFromAnimInstanceProxy(const FAnimInstanceProxy* InAnimInstanceProxy);
+	virtual void SetBoneInitialTransformsFromAnimInstanceProxy(const FAnimInstanceProxy* InAnimInstanceProxy);
 
 	// Setup the initial transform / ref pose of the bones based upon skeletal mesh component (ref skeleton)
 	// This uses the RefSkeleton pose instead of the current refpose (or vice versae is bUseAnimInstance == true)
-	void SetBoneInitialTransformsFromSkeletalMeshComponent(USkeletalMeshComponent* InSkelMeshComp, bool bUseAnimInstance = false);
+	virtual void SetBoneInitialTransformsFromSkeletalMeshComponent(USkeletalMeshComponent* InSkelMeshComp, bool bUseAnimInstance = false);
 
 	// Setup the initial transforms / ref pose of the bones based on a skeletal mesh
 	// This uses the RefSkeleton pose instead of the current refpose.
-	void SetBoneInitialTransformsFromSkeletalMesh(USkeletalMesh* InSkeletalMesh);
+	virtual void SetBoneInitialTransformsFromSkeletalMesh(USkeletalMesh* InSkeletalMesh);
 
 	// Setup the initial transforms / ref pose of the bones based on a reference skeleton
 	// This uses the RefSkeleton pose instead of the current refpose.
-	void SetBoneInitialTransformsFromRefSkeleton(const FReferenceSkeleton& InReferenceSkeleton);
+	virtual void SetBoneInitialTransformsFromRefSkeleton(const FReferenceSkeleton& InReferenceSkeleton);
 
 private:
 
diff --git a/Engine/Plugins/Animation/ControlRig/Source/ControlRig/Public/Rigs/RigHierarchy.h b/Engine/Plugins/Animation/ControlRig/Source/ControlRig/Public/Rigs/RigHierarchy.h
--- a/Engine/Plugins/Animation/ControlRig/Source/ControlRig/Public/Rigs/RigHierarchy.h
+++ b/Engine/Plugins/Animation/ControlRig/Source/ControlRig/Public/Rigs/RigHierarchy.h
@@ -86,27 +86,7 @@ class CONTROLRIG_API URigHierarchy : public UObject
 
 public:
 
-	URigHierarchy()
-	: TopologyVersion(0)
-	, bEnableDirtyPropagation(true)
-	, Elements()
-	, IndexLookup()
-	, TransformStackIndex(0)
-	, bTransactingForTransformChange(false)
-	, bIsInteracting(false)
-	, LastInteractedKey()
-	, bSuspendNotifications(false)
-	, ResetPoseHash(INDEX_NONE)
-#if WITH_EDITOR
-	, bPropagatingChange(false)
-	, bForcePropagation(false)
-	, TraceFramesLeft(0)
-	, TraceFramesCaptured(0)
-#endif
-	{
-		Reset();
-	}
-
+	URigHierarchy();
 	virtual ~URigHierarchy();
 
 	// UObject interface
