Control Rig: Fix blueprint buttons in toolbar missing in control rig editor

diff --git a/Engine/Plugins/Animation/ControlRig/Source/ControlRigEditor/Private/Editor/ControlRigEditor.cpp b/Engine/Plugins/Animation/ControlRig/Source/ControlRigEditor/Private/Editor/ControlRigEditor.cpp
--- a/Engine/Plugins/Animation/ControlRig/Source/ControlRigEditor/Private/Editor/ControlRigEditor.cpp
+++ b/Engine/Plugins/Animation/ControlRig/Source/ControlRigEditor/Private/Editor/ControlRigEditor.cpp
@@ -1895,7 +1895,7 @@ FLinearColor FControlRigEditor::GetWorldCentricTabColorScale() const
 
 void FControlRigEditor::InitToolMenuContext(FToolMenuContext& MenuContext)
 {
-	FAssetEditorToolkit::InitToolMenuContext(MenuContext);
+	FBlueprintEditor::InitToolMenuContext(MenuContext);
 
 	if (UControlRigBlueprint* RigBlueprint = GetControlRigBlueprint())
 	{
