Control Rig: FKControl Rig to use offset transform correctly

* centralized the code for computing control offsets within the FKControlRig
* added an empty animation instance (needed?) to be able to create a temp skel mesh comp
* changed baking process to access a pose coming out of the anim sequences evaluate rather than accessing the raw anim tracks

diff --git a/Engine/Plugins/Animation/ControlRig/Source/ControlRig/ControlRig.Build.cs b/Engine/Plugins/Animation/ControlRig/Source/ControlRig/ControlRig.Build.cs
--- a/Engine/Plugins/Animation/ControlRig/Source/ControlRig/ControlRig.Build.cs
+++ b/Engine/Plugins/Animation/ControlRig/Source/ControlRig/ControlRig.Build.cs
@@ -49,6 +49,7 @@ namespace UnrealBuildTool.Rules
                         "Json",
                         "Serialization",
                         "JsonUtilities",
+                        "AnimationBlueprintLibrary",
                     }
                 );
 
diff --git a/Engine/Plugins/Animation/ControlRig/Source/ControlRig/Private/Rigs/FKControlRig.cpp b/Engine/Plugins/Animation/ControlRig/Source/ControlRig/Private/Rigs/FKControlRig.cpp
--- a/Engine/Plugins/Animation/ControlRig/Source/ControlRig/Private/Rigs/FKControlRig.cpp
+++ b/Engine/Plugins/Animation/ControlRig/Source/ControlRig/Private/Rigs/FKControlRig.cpp
@@ -5,6 +5,7 @@
 #include "Engine/SkeletalMesh.h"
 #include "IControlRigObjectBinding.h"
 #include "Components/SkeletalMeshComponent.h"
+#include "Execution/RigUnit_PrepareForExecution.h"
 #include "Units/Execution/RigUnit_BeginExecution.h"
 #include "Units/Execution/RigUnit_InverseExecution.h"
 
@@ -36,6 +37,16 @@ void UFKControlRig::ExecuteUnits(FRigUnitContext& InOutContext, const FName& InE
 		return;
 	}
 
+#if WITH_EDITOR
+	if(URigHierarchy* Hierarchy = GetHierarchy())
+	{
+		if(Hierarchy->IsTracingChanges())
+		{
+			Hierarchy->StorePoseForTrace(TEXT("UFKControlRig::BeginExecuteUnits"));
+		}
+	}
+#endif	
+
 	if (InEventName == FRigUnit_BeginExecution::EventName)
 	{
 		FRigVMExecuteContext VMContext;
@@ -53,13 +64,19 @@ void UFKControlRig::ExecuteUnits(FRigUnitContext& InOutContext, const FName& InE
 				{
 					case EControlRigFKRigExecuteMode::Replace:
 					{
-						GetHierarchy()->SetTransform(BoneElement, LocalTransform, ERigTransformType::CurrentLocal, true, false);
+						const FTransform OffsetTransform = GetHierarchy()->GetControlOffsetTransform(Control, ERigTransformType::InitialLocal);
+						FTransform Transform = LocalTransform * OffsetTransform;
+						Transform.NormalizeRotation();
+						Transform.SetScale3D(FVector::OneVector);
+						GetHierarchy()->SetTransform(BoneElement, Transform, ERigTransformType::CurrentLocal, true, false);
 						break;
 					}
 					case EControlRigFKRigExecuteMode::Additive:
 					{
 						const FTransform PreviousTransform = GetHierarchy()->GetTransform(BoneElement, ERigTransformType::CurrentLocal);
-						const FTransform Transform = LocalTransform * PreviousTransform;
+						FTransform Transform = LocalTransform * PreviousTransform;
+						Transform.NormalizeRotation();
+						Transform.SetScale3D(FVector::OneVector);
 						GetHierarchy()->SetTransform(BoneElement, Transform, ERigTransformType::CurrentLocal, true, false);
 						break;
 					}
@@ -102,20 +119,33 @@ void UFKControlRig::ExecuteUnits(FRigUnitContext& InOutContext, const FName& InE
 		const bool bNotify = true;
 		const FRigControlModifiedContext Context = FRigControlModifiedContext();
 		const bool bSetupUndo = false;
-		GetHierarchy()->ForEach<FRigBoneElement>([&](FRigBoneElement* BoneElement) -> bool
-        {
-            const FName ControlName = GetControlName(BoneElement->GetName());
+
+		GetHierarchy()->Traverse([&](FRigBaseElement* InElement, bool& bContinue)
+		{
+			if(!InElement->IsA<FRigBoneElement>())
+			{
+				bContinue = false;
+				return;
+			}
+
+			FRigBoneElement* BoneElement = CastChecked<FRigBoneElement>(InElement);
+
+			const FName ControlName = GetControlName(BoneElement->GetName());
 			const FRigElementKey ControlKey(ControlName, ERigElementType::Control);
 			const int32 ControlIndex = GetHierarchy()->GetIndex(ControlKey);
 			
-			if (IsControlActive[ControlIndex])
-			{
-				const FEulerTransform EulerTransform(GetHierarchy()->GetTransform(BoneElement, ERigTransformType::CurrentLocal));
-				SetControlValue(ControlName, FRigControlValue::Make(EulerTransform), bNotify, Context, bSetupUndo);
-			}
+			// during inversion we assume Replace Mode
+			FRigControlElement* Control = GetHierarchy()->GetChecked<FRigControlElement>(ControlIndex);
+			const FTransform Offset = GetHierarchy()->GetControlOffsetTransform(Control, ERigTransformType::InitialLocal);
+			const FTransform Current = GetHierarchy()->GetTransform(BoneElement, ERigTransformType::CurrentLocal);
 			
-			return true;
-		});
+			FTransform Transform = Current.GetRelativeTransform(Offset);
+			Transform.NormalizeRotation();
+			Transform.SetScale3D(FVector::OneVector);
+
+			GetHierarchy()->SetTransform(Control, Transform, ERigTransformType::CurrentLocal, true);
+
+		}, true);
 
 		GetHierarchy()->ForEach<FRigCurveElement>([&](FRigCurveElement* CurveElement) -> bool
         {
@@ -123,15 +153,41 @@ void UFKControlRig::ExecuteUnits(FRigUnitContext& InOutContext, const FName& InE
 			const FRigElementKey ControlKey(ControlName, ERigElementType::Control);
 			const int32 ControlIndex = GetHierarchy()->GetIndex(ControlKey);
 
-			if (IsControlActive[ControlIndex])
-			{
-				const float CurveValue = GetHierarchy()->GetCurveValue(CurveElement);
-				SetControlValue(ControlName, FRigControlValue::Make(CurveValue), bNotify, Context, bSetupUndo);
-			}
+			const float CurveValue = GetHierarchy()->GetCurveValue(CurveElement);
+			SetControlValue(ControlName, FRigControlValue::Make(CurveValue), bNotify, Context, bSetupUndo);
 
 			return true;
 		});
 	}
+
+#if WITH_EDITOR
+	if(URigHierarchy* Hierarchy = GetHierarchy())
+	{
+		if(Hierarchy->IsTracingChanges())
+		{
+			Hierarchy->StorePoseForTrace(TEXT("UFKControlRig::EndExecuteUnits"));
+			Hierarchy->DumpTransformStackToFile();
+		}
+	}
+#endif	
+}
+
+void UFKControlRig::SetBoneInitialTransformsFromSkeletalMeshComponent(USkeletalMeshComponent* InSkelMeshComp,
+	bool bUseAnimInstance)
+{
+	Super::SetBoneInitialTransformsFromSkeletalMeshComponent(InSkelMeshComp, bUseAnimInstance);
+
+#if WITH_EDITOR
+	if(URigHierarchy* Hierarchy = GetHierarchy())
+	{
+		if(Hierarchy->IsTracingChanges())
+		{
+			Hierarchy->StorePoseForTrace(TEXT("UFKControlRig::SetBoneInitialTransformsFromSkeletalMeshComponent"));
+		}
+	}
+#endif	
+
+	SetControlOffsetsFromBoneInitials();
 }
 
 void UFKControlRig::Initialize(bool bInitRigUnits /*= true*/)
@@ -220,27 +276,23 @@ void UFKControlRig::CreateRigElements(const FReferenceSkeleton& InReferenceSkele
 	{
 		const FName BoneName = BoneElement->GetName();
 		const FName ControlName = GetControlName(BoneName); // name conflict?
-		FRigElementKey ParentKey = GetHierarchy()->GetFirstParent(BoneElement->GetKey());
+		const FRigElementKey ParentKey = GetHierarchy()->GetFirstParent(BoneElement->GetKey());
 
-		FTransform OffsetTransform;
-		if (ParentKey.IsValid())
-		{
-			FTransform GlobalTransform = GetHierarchy()->GetGlobalTransform(BoneElement->GetIndex());
-			FTransform ParentTransform = GetHierarchy()->GetGlobalTransform(ParentKey);
-			OffsetTransform = GlobalTransform.GetRelativeTransform(ParentTransform);
-		}
-		else
-		{
-			OffsetTransform = GetHierarchy()->GetLocalTransform(BoneElement->GetIndex());
-		}
+		FTransform OffsetTransform = FTransform::Identity;
 
 		FRigControlSettings Settings;
 		Settings.ControlType = ERigControlType::EulerTransform;
 		Settings.DisplayName = BoneName;
+
+		OffsetTransform.NormalizeRotation();
+		OffsetTransform.SetScale3D(FVector::OneVector);
+		
 		Controller->AddControl(ControlName, ParentKey, Settings, FRigControlValue::Make(FEulerTransform::Identity), OffsetTransform, FTransform::Identity, false);
 
 		return true;
 	});
+
+	SetControlOffsetsFromBoneInitials();
 	
 	GetHierarchy()->ForEach<FRigCurveElement>([&](FRigCurveElement* CurveElement) -> bool
     {
@@ -265,6 +317,47 @@ void UFKControlRig::CreateRigElements(const FReferenceSkeleton& InReferenceSkele
 	}
 }
 
+void UFKControlRig::SetControlOffsetsFromBoneInitials()
+{
+	GetHierarchy()->Traverse([&](FRigBaseElement* InElement, bool& bContinue)
+	{
+		if(!InElement->IsA<FRigBoneElement>())
+		{
+			bContinue = false;
+			return;
+		}
+
+		FRigBoneElement* BoneElement = CastChecked<FRigBoneElement>(InElement);
+		const FName BoneName = BoneElement->GetName();
+		const FName ControlName = GetControlName(BoneName); // name conflict?
+
+		FRigControlElement* ControlElement = GetHierarchy()->Find<FRigControlElement>(FRigElementKey(ControlName, ERigElementType::Control));
+		if(ControlElement == nullptr)
+		{
+			return;
+		}
+			
+		const FRigElementKey ParentKey = GetHierarchy()->GetFirstParent(BoneElement->GetKey());
+
+		FTransform OffsetTransform;
+		if (ParentKey.IsValid())
+		{
+			FTransform GlobalTransform = GetHierarchy()->GetGlobalTransformByIndex(BoneElement->GetIndex(), true);
+			FTransform ParentTransform = GetHierarchy()->GetGlobalTransform(ParentKey, true);
+			OffsetTransform = GlobalTransform.GetRelativeTransform(ParentTransform);
+		}
+		else
+		{
+			OffsetTransform = GetHierarchy()->GetLocalTransformByIndex(BoneElement->GetIndex(), true);
+		}
+
+		OffsetTransform.NormalizeRotation();
+
+		GetHierarchy()->SetControlOffsetTransform(ControlElement, OffsetTransform, ERigTransformType::InitialLocal, false, false, true);
+
+	}, true);
+}
+
 void UFKControlRig::CreateRigElements(const USkeletalMesh* InReferenceMesh)
 {
 	if (InReferenceMesh)
diff --git a/Engine/Plugins/Animation/ControlRig/Source/ControlRig/Private/Sequencer/MovieSceneControlRigParameterSection.cpp b/Engine/Plugins/Animation/ControlRig/Source/ControlRig/Private/Sequencer/MovieSceneControlRigParameterSection.cpp
--- a/Engine/Plugins/Animation/ControlRig/Source/ControlRig/Private/Sequencer/MovieSceneControlRigParameterSection.cpp
+++ b/Engine/Plugins/Animation/ControlRig/Source/ControlRig/Private/Sequencer/MovieSceneControlRigParameterSection.cpp
@@ -16,6 +16,10 @@
 #include "Animation/AnimSequenceHelpers.h"
 #include "Components/SkeletalMeshComponent.h"
 
+#if WITH_EDITOR
+#include "AnimPose.h"
+#endif
+
 #define LOCTEXT_NAMESPACE "MovieSceneControlParameterRigSection"
 
 #if WITH_EDITOR
@@ -2160,15 +2164,22 @@ bool UMovieSceneControlRigParameterSection::LoadAnimSequenceIntoThisSection(UAni
 			SourceHierarchy->SetCurveValue(FRigElementKey(Curve.Name.DisplayName, ERigElementType::Curve), Val);
 		}
 
-		for(int32 TrackIndex = 0; TrackIndex < BoneAnimationTracks.Num(); ++TrackIndex)
-		{
-			const FBoneAnimationTrack& AnimationTrack = BoneAnimationTracks[TrackIndex];
+		// retrieve the pose using the services that persona and sequencer rely on
+		// rather than accessing the low level raw tracks.
+		FAnimPose AnimPose;
+		UAnimPoseExtensions::GetAnimPoseAtTime(AnimSequence, SequenceSecond, FAnimPoseEvaluationOptions(), AnimPose);
 
-			FName BoneName = SkelMeshComp->SkeletalMesh->GetSkeleton()->GetReferenceSkeleton().GetBoneName(AnimationTrack.BoneTreeIndex);
-			FTransform BoneTransform;
-			UE::Anim::GetBoneTransformFromModel(DataModel, BoneTransform, TrackIndex, SequenceSecond, EAnimInterpolationType::Linear);
-			SourceHierarchy->SetLocalTransform(FRigElementKey(BoneName, ERigElementType::Bone), BoneTransform);
+		TArray<FName> BoneNames;
+		UAnimPoseExtensions::GetBoneNames(AnimPose, BoneNames);
+		for(const FName& BoneName : BoneNames)
+		{
+			if(FRigBoneElement* BoneElement = SourceHierarchy->Find<FRigBoneElement>(FRigElementKey(BoneName, ERigElementType::Bone)))
+			{
+				FTransform LocalTransform = UAnimPoseExtensions::GetBonePose(AnimPose, BoneName, EAnimPoseSpaces::Local);
+				SourceHierarchy->SetLocalTransform(BoneElement->GetIndex(), LocalTransform, true, false);
+			}
 		}
+
 		if (Index == 0)
 		{
 			//to make sure the first frame looks good we need to do this first. UE-100069
diff --git a/Engine/Plugins/Animation/ControlRig/Source/ControlRig/Public/Rigs/FKControlRig.h b/Engine/Plugins/Animation/ControlRig/Source/ControlRig/Public/Rigs/FKControlRig.h
--- a/Engine/Plugins/Animation/ControlRig/Source/ControlRig/Public/Rigs/FKControlRig.h
+++ b/Engine/Plugins/Animation/ControlRig/Source/ControlRig/Public/Rigs/FKControlRig.h
@@ -42,6 +42,7 @@ public:
 	// BEGIN ControlRig
 	virtual void Initialize(bool bInitRigUnits = true) override;
 	virtual void ExecuteUnits(FRigUnitContext& InOutContext, const FName& InEventName) override;
+	virtual void SetBoneInitialTransformsFromSkeletalMeshComponent(USkeletalMeshComponent* InSkelMeshComp, bool bUseAnimInstance = false) override;
 	// END ControlRig
 
 	// utility function to 
@@ -61,6 +62,7 @@ private:
 	/** Create RigElements - bone hierarchy and curves - from incoming skeleton */
 	void CreateRigElements(const USkeletalMesh* InReferenceMesh);
 	void CreateRigElements(const FReferenceSkeleton& InReferenceSkeleton, const FSmartNameMapping* InSmartNameMapping);
+	void SetControlOffsetsFromBoneInitials();
 
 	UPROPERTY()
 	TArray<bool> IsControlActive;
