Control Rig: Fix rename leaves unvalid element in details panel

diff --git a/Engine/Plugins/Animation/ControlRig/Source/ControlRig/Private/Rigs/RigHierarchyController.cpp b/Engine/Plugins/Animation/ControlRig/Source/ControlRig/Private/Rigs/RigHierarchyController.cpp
--- a/Engine/Plugins/Animation/ControlRig/Source/ControlRig/Private/Rigs/RigHierarchyController.cpp
+++ b/Engine/Plugins/Animation/ControlRig/Source/ControlRig/Private/Rigs/RigHierarchyController.cpp
@@ -1591,6 +1591,11 @@ FRigElementKey URigHierarchyController::RenameElement(FRigElementKey InElement,
 		TransactionPtr->Cancel();
 	}
 	TransactionPtr.Reset();
+
+	if (bRenamed)
+	{
+		ClearSelection();
+	}
 #endif
 
 	return bRenamed ? Element->GetKey() : FRigElementKey();
