Control Rig & Sequencer: 1. setup event now always runs from a determinstic initial state. 2. any changes to the control values prior to the setup event is saved before the event executes and reapplied to the rig right after setup event. 3. Pre-forward solve no longer runs before setup event 4. Similar logic on the sequencer side, save value first before calling control rig->initialize and avoid unnecessary reinitialization

diff --git a/Engine/Config/BaseEngine.ini b/Engine/Config/BaseEngine.ini
--- a/Engine/Config/BaseEngine.ini
+++ b/Engine/Config/BaseEngine.ini
@@ -917,6 +917,9 @@ NearClipPlane=10.0
 +StructRedirects=(OldName="AnimInstanceSubsystemData",NewName="/Script/Engine.AnimSubsystemInstance")
 +ClassRedirects=(OldName="/Script/AnimationModifiers.AnimationBlueprintLibrary",NewName="/Script/AnimationBlueprintLibrary.AnimationBlueprintLibrary")
 +PropertyRedirects=(OldName="GeometryCollectionComponent.Simulating", NewName="SimulatePhysics")
++PropertyRedirects=(OldName="ControlRigComponent.OnPreUpdateDelegate", NewName="OnPreForwardsSolveDelegate")
++PropertyRedirects=(OldName="ControlRigComponent.OnPostUpdateDelegate", NewName="OnPostForwardsSolveDelegate")
+
 
 [CoreUObject.Metadata]
 ; Note: UnrealHeaderTool should be rerun after making changes to MetadataRedirects to catch any keys specified in class headers
diff --git a/Engine/Plugins/Animation/ControlRig/Source/ControlRig/Private/ControlRig.cpp b/Engine/Plugins/Animation/ControlRig/Source/ControlRig/Private/ControlRig.cpp
--- a/Engine/Plugins/Animation/ControlRig/Source/ControlRig/Private/ControlRig.cpp
+++ b/Engine/Plugins/Animation/ControlRig/Source/ControlRig/Private/ControlRig.cpp
@@ -23,6 +23,7 @@
 #include "Engine/Blueprint.h"
 #include "EdGraphSchema_K2.h"
 #include "Kismet2/BlueprintEditorUtils.h"
+#include "Settings/ControlRigSettings.h"
 #endif// WITH_EDITOR
 
 #define LOCTEXT_NAMESPACE "ControlRig"
@@ -91,6 +92,8 @@ void UControlRig::BeginDestroy()
 	InitializedEvent.Clear();
 	PreSetupEvent.Clear();
 	PostSetupEvent.Clear();
+	PreForwardsSolveEvent.Clear();
+	PostForwardsSolveEvent.Clear();
 	ExecutedEvent.Clear();
 	SetInteractionRig(nullptr);
 
@@ -680,19 +683,34 @@ void UControlRig::Execute(const EControlRigState InState, const FName& InEventNa
 					GetHierarchy()->CopyPose(CDO->GetHierarchy(), false, true);
 				}
 			}
-
-			if (PreSetupEvent.IsBound())
+			
 			{
-				FControlRigBracketScope BracketScope(PreSetupBracket);
-				PreSetupEvent.Broadcast(this, EControlRigState::Update, FRigUnit_PrepareForExecution::EventName);
-			}
+				// save the current control values to preserve user intention
+				// the control values are reapplied to the rig after setup event as it goes out of scope
+				TUniquePtr<UControlRig::FControlValueScope> ValueScope;
+				if (!bSetupModeEnabled && !UControlRigEditorSettings::Get()->bResetControlsOnCompile)
+				{
+					// only do this in non-setup mode because 
+					// when setup mode is enabled, the control values are cleared before reaching here (too late to save them)
+					ValueScope = MakeUnique<UControlRig::FControlValueScope>(this);
+				}
 
-			ExecuteUnits(Context, FRigUnit_PrepareForExecution::EventName);
+				// reset the pose to initial such that setup event can run from a deterministic initial state
+				GetHierarchy()->ResetPoseToInitial();
 
-			if (PostSetupEvent.IsBound())
-			{
-				FControlRigBracketScope BracketScope(PostSetupBracket);
-				PostSetupEvent.Broadcast(this, EControlRigState::Update, FRigUnit_PrepareForExecution::EventName);
+				if (PreSetupEvent.IsBound())
+				{
+					FControlRigBracketScope BracketScope(PreSetupBracket);
+					PreSetupEvent.Broadcast(this, EControlRigState::Update, FRigUnit_PrepareForExecution::EventName);
+				}
+
+				ExecuteUnits(Context, FRigUnit_PrepareForExecution::EventName);
+
+				if (PostSetupEvent.IsBound())
+				{
+					FControlRigBracketScope BracketScope(PostSetupBracket);
+					PostSetupEvent.Broadcast(this, EControlRigState::Update, FRigUnit_PrepareForExecution::EventName);
+				}
 			}
 
 			if (bSetupModeEnabled)
@@ -708,7 +726,6 @@ void UControlRig::Execute(const EControlRigState InState, const FName& InEventNa
 
 	if (!bSetupModeEnabled)
 	{
-
 		if(!IsExecuting())
 		{ 
 
@@ -721,8 +738,26 @@ void UControlRig::Execute(const EControlRigState InState, const FName& InEventNa
 			}
 #endif
 
+			if (InState == EControlRigState::Update && InEventName == FRigUnit_BeginExecution::EventName)
+			{
+				if (PreForwardsSolveEvent.IsBound())
+				{
+					FControlRigBracketScope BracketScope(PreForwardsSolveBracket);
+					PreForwardsSolveEvent.Broadcast(this, EControlRigState::Update, FRigUnit_BeginExecution::EventName);
+				}
+			}
+
 			ExecuteUnits(Context, InEventName);
 
+			if (InState == EControlRigState::Update && InEventName == FRigUnit_BeginExecution::EventName)
+			{
+				if (PostForwardsSolveEvent.IsBound())
+				{
+					FControlRigBracketScope BracketScope(PostForwardsSolveBracket);
+					PostForwardsSolveEvent.Broadcast(this, EControlRigState::Update, FRigUnit_BeginExecution::EventName);
+				}
+			}
+			
 			if (InState == EControlRigState::Init)
 			{
 				ExecuteUnits(Context, FRigUnit_BeginExecution::EventName);
@@ -2601,6 +2636,33 @@ void UControlRig::OnHierarchyTransformUndoRedo(URigHierarchy* InHierarchy, const
 	}
 }
 
+UControlRig::FControlValueScope::FControlValueScope(UControlRig* InControlRig)
+{
+	check(InControlRig);
+
+	TArray<FRigControlElement*> Controls = InControlRig->AvailableControls();
+	for (FRigControlElement* ControlElement : Controls)
+	{
+		ControlValues.Add(ControlElement->GetName(), InControlRig->GetControlValue(ControlElement->GetName()));
+	}
+
+	ControlRig = InControlRig;
+}
+
+UControlRig::FControlValueScope::~FControlValueScope()
+{
+	check(ControlRig);
+
+	for (const TPair<FName, FRigControlValue>& Pair : ControlValues)
+	{
+		if (ControlRig->FindControl(Pair.Key))
+		{
+			// bNotify = false such that it won't trigger sequencer auto key
+			ControlRig->SetControlValue(Pair.Key, Pair.Value, false, FRigControlModifiedContext(), false);
+		}
+	}
+}
+ 
 #undef LOCTEXT_NAMESPACE
 
 
diff --git a/Engine/Plugins/Animation/ControlRig/Source/ControlRig/Private/ControlRigComponent.cpp b/Engine/Plugins/Animation/ControlRig/Source/ControlRig/Private/ControlRigComponent.cpp
--- a/Engine/Plugins/Animation/ControlRig/Source/ControlRig/Private/ControlRigComponent.cpp
+++ b/Engine/Plugins/Animation/ControlRig/Source/ControlRig/Private/ControlRigComponent.cpp
@@ -280,16 +280,14 @@ void UControlRigComponent::OnPostSetup_Implementation(UControlRigComponent* Comp
 	OnPostSetupDelegate.Broadcast(Component);
 }
 
-void UControlRigComponent::OnPreUpdate_Implementation(UControlRigComponent* Component)
+void UControlRigComponent::OnPreForwardsSolve_Implementation(UControlRigComponent* Component)
 {
-	TransferInputs();
-	OnPreUpdateDelegate.Broadcast(Component);
+	OnPreForwardsSolveDelegate.Broadcast(Component);
 }
 
-void UControlRigComponent::OnPostUpdate_Implementation(UControlRigComponent* Component)
+void UControlRigComponent::OnPostForwardsSolve_Implementation(UControlRigComponent* Component)
 {
-	TransferOutputs();
-	OnPostUpdateDelegate.Broadcast(Component);
+	OnPostForwardsSolveDelegate.Broadcast(Component);
 }
 
 void UControlRigComponent::Initialize()
@@ -365,17 +363,9 @@ void UControlRigComponent::Update(float DeltaTime)
 					Hierarchy->ResetTransformStack();
 				}
 			}
-
-			if (bUpdateInEditor)
-			{
-				FEditorScriptExecutionGuard AllowScripts;
-				OnPreUpdate(this);
-			}
-			else
 #endif
-			{
-				OnPreUpdate(this);
-			}
+
+			TransferInputs();
 
 #if WITH_EDITOR
 			if(URigHierarchy* Hierarchy = CR->GetHierarchy())
@@ -1123,6 +1113,8 @@ UControlRig* UControlRigComponent::SetupControlRigIfRequired()
 			ControlRig->OnInitialized_AnyThread().RemoveAll(this);
 			ControlRig->OnPreSetup_AnyThread().RemoveAll(this);
 			ControlRig->OnPostSetup_AnyThread().RemoveAll(this);
+			ControlRig->OnPreForwardsSolve_AnyThread().RemoveAll(this);
+			ControlRig->OnPostForwardsSolve_AnyThread().RemoveAll(this);
 			ControlRig->OnExecuted_AnyThread().RemoveAll(this);
 			ControlRig = nullptr;
 		}
@@ -1155,6 +1147,8 @@ void UControlRigComponent::SetControlRig(UControlRig* InControlRig)
 	ControlRig->OnInitialized_AnyThread().AddUObject(this, &UControlRigComponent::HandleControlRigInitializedEvent);
 	ControlRig->OnPreSetup_AnyThread().AddUObject(this, &UControlRigComponent::HandleControlRigPreSetupEvent);
 	ControlRig->OnPostSetup_AnyThread().AddUObject(this, &UControlRigComponent::HandleControlRigPostSetupEvent);
+	ControlRig->OnPreForwardsSolve_AnyThread().AddUObject(this, &UControlRigComponent::HandleControlRigPreForwardsSolveEvent);
+	ControlRig->OnPostForwardsSolve_AnyThread().AddUObject(this, &UControlRigComponent::HandleControlRigPostForwardsSolveEvent);
 	ControlRig->OnExecuted_AnyThread().AddUObject(this, &UControlRigComponent::HandleControlRigExecutedEvent);
 
 	ControlRig->GetDataSourceRegistry()->RegisterDataSource(UControlRig::OwnerComponent, this);
@@ -1550,24 +1544,41 @@ void UControlRigComponent::HandleControlRigPostSetupEvent(UControlRig* InControl
 	}
 }
 
-void UControlRigComponent::HandleControlRigExecutedEvent(UControlRig* InControlRig, const EControlRigState InState, const FName& InEventName)
+void UControlRigComponent::HandleControlRigPreForwardsSolveEvent(UControlRig* InControlRig, const EControlRigState InState, const FName& InEventName)
 {
-	if (InEventName == FRigUnit_BeginExecution::EventName)
+#if WITH_EDITOR
+	if (bUpdateInEditor)
 	{
+		FEditorScriptExecutionGuard AllowScripts;
+		OnPreForwardsSolve(this);
+	}
+	else
+#endif
+	{
+		OnPreForwardsSolve(this);
+	}
+}
+
+void UControlRigComponent::HandleControlRigPostForwardsSolveEvent(UControlRig* InControlRig, const EControlRigState InState, const FName& InEventName)
+{
 #if WITH_EDITOR
-		if (bUpdateInEditor)
-		{
-			FEditorScriptExecutionGuard AllowScripts;
-			OnPostUpdate(this);
-		}
-		else
+	if (bUpdateInEditor)
+	{
+		FEditorScriptExecutionGuard AllowScripts;
+		OnPostForwardsSolve(this);
+	}
+	else
 #endif
-		{
-			OnPostUpdate(this);
-		}
+	{
+		OnPostForwardsSolve(this);
 	}
 }
 
+void UControlRigComponent::HandleControlRigExecutedEvent(UControlRig* InControlRig, const EControlRigState InState, const FName& InEventName)
+{
+	TransferOutputs();
+}
+
 void UControlRigComponent::ConvertTransformToRigSpace(FTransform& InOutTransform, EControlRigComponentSpace FromSpace)
 {
 	switch (FromSpace)
diff --git a/Engine/Plugins/Animation/ControlRig/Source/ControlRig/Public/ControlRig.h b/Engine/Plugins/Animation/ControlRig/Source/ControlRig/Public/ControlRig.h
--- a/Engine/Plugins/Animation/ControlRig/Source/ControlRig/Public/ControlRig.h
+++ b/Engine/Plugins/Animation/ControlRig/Source/ControlRig/Public/ControlRig.h
@@ -365,6 +365,8 @@ public:
 	FControlRigExecuteEvent& OnInitialized_AnyThread() { return InitializedEvent; }
 	FControlRigExecuteEvent& OnPreSetup_AnyThread() { return PreSetupEvent; }
 	FControlRigExecuteEvent& OnPostSetup_AnyThread() { return PostSetupEvent; }
+	FControlRigExecuteEvent& OnPreForwardsSolve_AnyThread() { return PreForwardsSolveEvent; }
+	FControlRigExecuteEvent& OnPostForwardsSolve_AnyThread() { return PostForwardsSolveEvent; }
 	FControlRigExecuteEvent& OnExecuted_AnyThread() { return ExecutedEvent; }
 	FRigEventDelegate& OnRigEvent_AnyThread() { return RigEventDelegate; }
 
@@ -479,6 +481,12 @@ private:
 	/** Broadcasts a notification whenever the controlrig has been setup. */
 	FControlRigExecuteEvent PostSetupEvent;
 
+	/** Broadcasts a notification before a forward solve has been initiated */
+	FControlRigExecuteEvent PreForwardsSolveEvent;
+	
+	/** Broadcasts a notification after a forward solve has been initiated */
+	FControlRigExecuteEvent PostForwardsSolveEvent;
+	
 	/** Broadcasts a notification whenever the controlrig is executed / updated. */
 	FControlRigExecuteEvent ExecutedEvent;
 
@@ -605,6 +613,8 @@ protected:
 	int32 UpdateBracket;
 	int32 PreSetupBracket;
 	int32 PostSetupBracket;
+	int32 PreForwardsSolveBracket;
+	int32 PostForwardsSolveBracket;
 	int32 InteractionBracket;
 	int32 InterRigSyncBracket;
 
@@ -685,6 +695,22 @@ public:
 	URigVM* GetSnapshotVM(bool bCreateIfNeeded = true);
 #endif	
 
+private:
+
+	// Class used to temporarily cache all 
+	// current control values and reapply them
+	// on destruction, similar to UControlRigBlueprint::FControlValueScope
+	class FControlValueScope
+	{
+	public:
+		FControlValueScope(UControlRig* InControlRig);
+		~FControlValueScope();
+
+	private:
+
+		UControlRig* ControlRig;
+		TMap<FName, FRigControlValue> ControlValues;
+	};
 
 	friend class FControlRigBlueprintCompilerContext;
 	friend struct FRigHierarchyRef;
diff --git a/Engine/Plugins/Animation/ControlRig/Source/ControlRig/Public/ControlRigComponent.h b/Engine/Plugins/Animation/ControlRig/Source/ControlRig/Public/ControlRigComponent.h
--- a/Engine/Plugins/Animation/ControlRig/Source/ControlRig/Public/ControlRigComponent.h
+++ b/Engine/Plugins/Animation/ControlRig/Source/ControlRig/Public/ControlRigComponent.h
@@ -198,13 +198,13 @@ public:
 	UPROPERTY(BlueprintAssignable, Category = "ControlRig", meta = (DisplayName = "On Post Setup"))
 	FControlRigComponentDelegate OnPostSetupDelegate;
 
-	/** Event fired before this component's ControlRig is updated */
+	/** Event fired before this component's ControlRig's forwards solve */
 	UPROPERTY(BlueprintAssignable, Category = "ControlRig", meta = (DisplayName = "On Pre Forwards Solve"))
-	FControlRigComponentDelegate OnPreUpdateDelegate;
+	FControlRigComponentDelegate OnPreForwardsSolveDelegate;
 
-	/** Event fired after this component's ControlRig is updated */
+	/** Event fired after this component's ControlRig's forwards solve */
 	UPROPERTY(BlueprintAssignable, Category = "ControlRig", meta = (DisplayName = "On Post Forwards Solve"))
-	FControlRigComponentDelegate OnPostUpdateDelegate;
+	FControlRigComponentDelegate OnPostForwardsSolveDelegate;
 
 	//~ Begin UObject interface
 #if WITH_EDITOR
@@ -251,10 +251,10 @@ public:
 	void OnPostSetup(UControlRigComponent* Component);
 
 	UFUNCTION(BlueprintNativeEvent, Category = "ControlRig", CallInEditor, meta = (DisplayName = "On Pre Forwards Solve"))
-	void OnPreUpdate(UControlRigComponent* Component);
+	void OnPreForwardsSolve(UControlRigComponent* Component);
 
 	UFUNCTION(BlueprintNativeEvent, Category = "ControlRig", CallInEditor, meta = (DisplayName = "On Post Forwards Solve"))
-	void OnPostUpdate(UControlRigComponent* Component);
+	void OnPostForwardsSolve(UControlRigComponent* Component);
 
 	/** Initializes the rig's memory and calls the setup event */
 	UFUNCTION(BlueprintCallable, Category = "ControlRig")
@@ -617,6 +617,8 @@ private:
 	void HandleControlRigInitializedEvent(UControlRig* InControlRig, const EControlRigState InState, const FName& InEventName);
 	void HandleControlRigPreSetupEvent(UControlRig* InControlRig, const EControlRigState InState, const FName& InEventName);
 	void HandleControlRigPostSetupEvent(UControlRig* InControlRig, const EControlRigState InState, const FName& InEventName);
+	void HandleControlRigPreForwardsSolveEvent(UControlRig* InControlRig, const EControlRigState InState, const FName& InEventName);
+	void HandleControlRigPostForwardsSolveEvent(UControlRig* InControlRig, const EControlRigState InState, const FName& InEventName);
 	void HandleControlRigExecutedEvent(UControlRig* InControlRig, const EControlRigState InState, const FName& InEventName);
 
 	void ConvertTransformToRigSpace(FTransform& InOutTransform, EControlRigComponentSpace FromSpace);
diff --git a/Engine/Plugins/Animation/ControlRig/Source/ControlRigEditor/Private/Sequencer/ControlRigParameterTrackEditor.cpp b/Engine/Plugins/Animation/ControlRig/Source/ControlRigEditor/Private/Sequencer/ControlRigParameterTrackEditor.cpp
--- a/Engine/Plugins/Animation/ControlRig/Source/ControlRigEditor/Private/Sequencer/ControlRigParameterTrackEditor.cpp
+++ b/Engine/Plugins/Animation/ControlRig/Source/ControlRigEditor/Private/Sequencer/ControlRigParameterTrackEditor.cpp
@@ -1260,7 +1260,11 @@ void FControlRigParameterTrackEditor::AddControlRig(UClass* InClass, UObject* Bo
 			ControlRig->SetObjectBinding(MakeShared<FControlRigObjectBinding>());
 			ControlRig->GetObjectBinding()->BindToObject(BoundActor);
 			ControlRig->GetDataSourceRegistry()->RegisterDataSource(UControlRig::OwnerComponent, ControlRig->GetObjectBinding()->GetBoundObject());
-			ControlRig->Initialize();
+			// Do not re-initialize existing control rig
+			if (!InExistingControlRig)
+			{ 
+				ControlRig->Initialize();
+			}
 			ControlRig->Evaluate_AnyThread();
 
 			SequencerParent->NotifyMovieSceneDataChanged(EMovieSceneDataChangeType::MovieSceneStructureItemsChanged);
