Control Rig: removed an unnecessary reference to UControlRigEditorSettings causing build issues

diff --git a/Engine/Plugins/Animation/ControlRig/Source/ControlRig/Private/ControlRig.cpp b/Engine/Plugins/Animation/ControlRig/Source/ControlRig/Private/ControlRig.cpp
--- a/Engine/Plugins/Animation/ControlRig/Source/ControlRig/Private/ControlRig.cpp
+++ b/Engine/Plugins/Animation/ControlRig/Source/ControlRig/Private/ControlRig.cpp
@@ -23,7 +23,6 @@
 #include "Engine/Blueprint.h"
 #include "EdGraphSchema_K2.h"
 #include "Kismet2/BlueprintEditorUtils.h"
-#include "Settings/ControlRigSettings.h"
 #endif// WITH_EDITOR
 
 #define LOCTEXT_NAMESPACE "ControlRig"
@@ -688,7 +687,7 @@ void UControlRig::Execute(const EControlRigState InState, const FName& InEventNa
 				// save the current control values to preserve user intention
 				// the control values are reapplied to the rig after setup event as it goes out of scope
 				TUniquePtr<UControlRig::FControlValueScope> ValueScope;
-				if (!bSetupModeEnabled && !UControlRigEditorSettings::Get()->bResetControlsOnCompile)
+				if (!bSetupModeEnabled)
 				{
 					// only do this in non-setup mode because 
 					// when setup mode is enabled, the control values are cleared before reaching here (too late to save them)
