Control Rig: Properly detach from the OnEditorClose delegate on destrtuct.

#rnx

diff --git a/Engine/Plugins/Animation/ControlRig/Source/ControlRigEditor/Private/Editor/SRigCurveContainer.cpp b/Engine/Plugins/Animation/ControlRig/Source/ControlRigEditor/Private/Editor/SRigCurveContainer.cpp
--- a/Engine/Plugins/Animation/ControlRig/Source/ControlRigEditor/Private/Editor/SRigCurveContainer.cpp
+++ b/Engine/Plugins/Animation/ControlRig/Source/ControlRigEditor/Private/Editor/SRigCurveContainer.cpp
@@ -205,6 +205,9 @@ void SRigCurveContainer::Construct(const FArguments& InArgs, TSharedRef<FControl
 
 SRigCurveContainer::~SRigCurveContainer()
 {
+	// Make sure we don't get called on editor exit if we've already been shut down.
+	GEditor->OnEditorClose().RemoveAll(this);
+
 	if (ControlRigEditor.IsValid())
 	{
 		ControlRigBlueprint = ControlRigEditor.Pin()->GetControlRigBlueprint();
