Control Rig: Introduce component wise weights for constraints

diff --git a/Engine/Plugins/Animation/ControlRig/Source/ControlRig/Private/Rigs/RigHierarchy.cpp b/Engine/Plugins/Animation/ControlRig/Source/ControlRig/Private/Rigs/RigHierarchy.cpp
--- a/Engine/Plugins/Animation/ControlRig/Source/ControlRig/Private/Rigs/RigHierarchy.cpp
+++ b/Engine/Plugins/Animation/ControlRig/Source/ControlRig/Private/Rigs/RigHierarchy.cpp
@@ -766,9 +766,9 @@ TArray<FRigBaseElement*> URigHierarchy::GetParents(const FRigBaseElement* InElem
 	}
 	else if(const FRigMultiParentElement* MultiParentElement = Cast<FRigMultiParentElement>(InElement))
 	{
-		for(FRigTransformElement* ParentElement : MultiParentElement->ParentElements)
+		for(const FRigElementParentConstraint& ParentConstraint : MultiParentElement->ParentConstraints)
 		{
-			Parents.AddUnique(ParentElement);
+			Parents.AddUnique(ParentConstraint.ParentElement);
 		}
 	}
 
@@ -814,9 +814,9 @@ FRigBaseElement* URigHierarchy::GetFirstParent(const FRigBaseElement* InElement)
 	}
 	else if(const FRigMultiParentElement* MultiParentElement = Cast<FRigMultiParentElement>(InElement))
 	{
-		if(MultiParentElement->ParentElements.Num() > 0)
+		if(MultiParentElement->ParentConstraints.Num() > 0)
 		{
-			return MultiParentElement->ParentElements[0];
+			return MultiParentElement->ParentConstraints[0].ParentElement;
 		}
 	}
 	
@@ -846,18 +846,18 @@ int32 URigHierarchy::GetNumberOfParents(const FRigBaseElement* InElement) const
 	}
 	else if(const FRigMultiParentElement* MultiParentElement = Cast<FRigMultiParentElement>(InElement))
 	{
-		return MultiParentElement->ParentElements.Num();
+		return MultiParentElement->ParentConstraints.Num();
 	}
 
 	return 0;
 }
 
-float URigHierarchy::GetParentWeight(FRigElementKey InChild, FRigElementKey InParent, bool bInitial) const
+FRigElementWeight URigHierarchy::GetParentWeight(FRigElementKey InChild, FRigElementKey InParent, bool bInitial) const
 {
 	return GetParentWeight(Find(InChild), Find(InParent), bInitial);
 }
 
-float URigHierarchy::GetParentWeight(const FRigBaseElement* InChild, const FRigBaseElement* InParent, bool bInitial) const
+FRigElementWeight URigHierarchy::GetParentWeight(const FRigBaseElement* InChild, const FRigBaseElement* InParent, bool bInitial) const
 {
 	if(const FRigMultiParentElement* MultiParentElement = Cast<FRigMultiParentElement>(InChild))
 	{
@@ -866,34 +866,34 @@ float URigHierarchy::GetParentWeight(const FRigBaseElement* InChild, const FRigB
 			return GetParentWeight(InChild, *ParentIndexPtr, bInitial);
 		}
 	}
-	return FLT_MAX;
+	return FRigElementWeight(FLT_MAX);
 }
 
-float URigHierarchy::GetParentWeight(const FRigBaseElement* InChild, int32 InParentIndex, bool bInitial) const
+FRigElementWeight URigHierarchy::GetParentWeight(const FRigBaseElement* InChild, int32 InParentIndex, bool bInitial) const
 {
 	if(const FRigMultiParentElement* MultiParentElement = Cast<FRigMultiParentElement>(InChild))
 	{
-		if(MultiParentElement->ParentWeights.IsValidIndex(InParentIndex))
+		if(MultiParentElement->ParentConstraints.IsValidIndex(InParentIndex))
 		{
 			if(bInitial)
 			{
-				return MultiParentElement->ParentWeightsInitial[InParentIndex];
+				return MultiParentElement->ParentConstraints[InParentIndex].InitialWeight;
 			}
 			else
 			{
-				return MultiParentElement->ParentWeights[InParentIndex];
+				return MultiParentElement->ParentConstraints[InParentIndex].Weight;
 			}
 		}
 	}
-	return FLT_MAX;
+	return FRigElementWeight(FLT_MAX);
 }
 
-bool URigHierarchy::SetParentWeight(FRigElementKey InChild, FRigElementKey InParent, float InWeight, bool bInitial, bool bAffectChildren)
+bool URigHierarchy::SetParentWeight(FRigElementKey InChild, FRigElementKey InParent, FRigElementWeight InWeight, bool bInitial, bool bAffectChildren)
 {
 	return SetParentWeight(Find(InChild), Find(InParent), InWeight, bInitial, bAffectChildren);
 }
 
-bool URigHierarchy::SetParentWeight(FRigBaseElement* InChild, const FRigBaseElement* InParent, float InWeight, bool bInitial, bool bAffectChildren)
+bool URigHierarchy::SetParentWeight(FRigBaseElement* InChild, const FRigBaseElement* InParent, FRigElementWeight InWeight, bool bInitial, bool bAffectChildren)
 {
 	if(const FRigMultiParentElement* MultiParentElement = Cast<FRigMultiParentElement>(InChild))
 	{
@@ -905,21 +905,25 @@ bool URigHierarchy::SetParentWeight(FRigBaseElement* InChild, const FRigBaseElem
 	return false;
 }
 
-bool URigHierarchy::SetParentWeight(FRigBaseElement* InChild, int32 InParentIndex, float InWeight, bool bInitial, bool bAffectChildren)
+bool URigHierarchy::SetParentWeight(FRigBaseElement* InChild, int32 InParentIndex, FRigElementWeight InWeight, bool bInitial, bool bAffectChildren)
 {
 	using namespace ERigTransformType;
 
 	if(FRigMultiParentElement* MultiParentElement = Cast<FRigMultiParentElement>(InChild))
 	{
-		if(MultiParentElement->ParentWeights.IsValidIndex(InParentIndex))
+		if(MultiParentElement->ParentConstraints.IsValidIndex(InParentIndex))
 		{
-			const float InputWeight = FMath::Max(InWeight, 0.f);
+			InWeight.Location = FMath::Max(InWeight.Location, 0.f);
+			InWeight.Rotation = FMath::Max(InWeight.Rotation, 0.f);
+			InWeight.Scale = FMath::Max(InWeight.Scale, 0.f);
 
-			float& TargetWeight = bInitial?
-				MultiParentElement->ParentWeightsInitial[InParentIndex] :
-			MultiParentElement->ParentWeights[InParentIndex];
+			FRigElementWeight& TargetWeight = bInitial?
+				MultiParentElement->ParentConstraints[InParentIndex].InitialWeight :
+				MultiParentElement->ParentConstraints[InParentIndex].Weight;
 
-			if(FMath::IsNearlyZero(InputWeight - TargetWeight))
+			if(FMath::IsNearlyZero(InWeight.Location - TargetWeight.Location) &&
+				FMath::IsNearlyZero(InWeight.Rotation - TargetWeight.Rotation) &&
+				FMath::IsNearlyZero(InWeight.Scale - TargetWeight.Scale))
 			{
 				return false;
 			}
@@ -938,7 +942,7 @@ bool URigHierarchy::SetParentWeight(FRigBaseElement* InChild, int32 InParentInde
 				MultiParentElement->Pose.MarkDirty(LocalType);
 			}
 
-			TargetWeight = InputWeight;
+			TargetWeight = InWeight;
 			MultiParentElement->Parent.MarkDirty(GlobalType);
 
 			if(FRigControlElement* ControlElement = Cast<FRigControlElement>(MultiParentElement))
@@ -1700,19 +1704,276 @@ FTransform URigHierarchy::GetParentTransform(FRigBaseElement* InElement, const E
 
 		if(Output.bDirty)
 		{
-			Output.Set(FTransform::Identity);
+			FTransform OutputTransform = FTransform::Identity;
 
 			const bool bInitial = IsInitial(InTransformType);
-			const TArray<float>& ParentWeights = bInitial ? MultiParentElement->ParentWeightsInitial : MultiParentElement->ParentWeights;
-			const TArray<FRigTransformElement*>& ParentElements = MultiParentElement->ParentElements;
+			const TArray<FRigElementParentConstraint>& ParentConstraints = MultiParentElement->ParentConstraints;
 
+			struct Local
+			{
+				static float GetWeightForLerp(const float WeightA, const float WeightB)
+				{
+					float Weight = 0.f;
+					const float ClampedWeightA = FMath::Max(WeightA, 0.f);
+					const float ClampedWeightB = FMath::Max(WeightB, 0.f);
+					const float OverallWeight = ClampedWeightA + ClampedWeightB;
+					if(OverallWeight > SMALL_NUMBER)
+					{
+						Weight = ClampedWeightB / OverallWeight;
+					}
+					return Weight;
+				}
+			};
+
+			struct FTransformIndex
+			{
+				int32 Location;
+				int32 Rotation;
+				int32 Scale;
+
+				FTransformIndex()
+					: Location(INDEX_NONE)
+					, Rotation(INDEX_NONE)
+					, Scale(INDEX_NONE)
+				{}
+
+				FTransformIndex(int32 InIndex)
+					: Location(InIndex)
+					, Rotation(InIndex)
+					, Scale(InIndex)
+				{}
+			};
+
+			// collect all of the weights
+			FTransformIndex FirstConstraint;
+			FTransformIndex SecondConstraint;
+			FTransformIndex NumConstraintsAffecting(0);
+			FRigElementWeight TotalWeight(0.f);
+
+			// find all of the weights affecting this output
+			for(int32 ConstraintIndex = 0; ConstraintIndex < ParentConstraints.Num(); ConstraintIndex++)
+			{
+				const FRigElementWeight& Weight = ParentConstraints[ConstraintIndex].GetWeight(bInitial);
+				if(Weight.AffectsLocation())
+				{
+					NumConstraintsAffecting.Location++;
+					TotalWeight.Location += Weight.Location;
+
+					if(FirstConstraint.Location == INDEX_NONE)
+					{
+						FirstConstraint.Location = ConstraintIndex;
+					}
+					else if(SecondConstraint.Location == INDEX_NONE)
+					{
+						SecondConstraint.Location = ConstraintIndex;
+					}
+				}
+				if(Weight.AffectsRotation())
+				{
+					NumConstraintsAffecting.Rotation++;
+					TotalWeight.Rotation += Weight.Rotation;
+
+					if(FirstConstraint.Rotation == INDEX_NONE)
+					{
+						FirstConstraint.Rotation = ConstraintIndex;
+					}
+					else if(SecondConstraint.Rotation == INDEX_NONE)
+					{
+						SecondConstraint.Rotation = ConstraintIndex;
+					}
+				}
+				if(Weight.AffectsScale())
+				{
+					NumConstraintsAffecting.Scale++;
+					TotalWeight.Scale += Weight.Scale;
+
+					if(FirstConstraint.Scale == INDEX_NONE)
+					{
+						FirstConstraint.Scale = ConstraintIndex;
+					}
+					else if(SecondConstraint.Scale == INDEX_NONE)
+					{
+						SecondConstraint.Scale = ConstraintIndex;
+					}
+				}
+			}
+
+			if(NumConstraintsAffecting.Location == 1)
+			{
+				check(FirstConstraint.Location != INDEX_NONE);
+
+				const FRigElementParentConstraint& ParentConstraint = ParentConstraints[FirstConstraint.Location];
+				const FRigElementWeight& Weight = ParentConstraint.GetWeight(bInitial);
+				const FTransform ParentTransform = GetTransform(ParentConstraint.ParentElement, InTransformType);
+				check(Weight.AffectsLocation());
+				OutputTransform.SetLocation(ParentTransform.GetLocation());
+			}
+			else if(NumConstraintsAffecting.Location == 2)
+			{
+				check(FirstConstraint.Location != INDEX_NONE);
+				check(SecondConstraint.Location != INDEX_NONE);
+
+				const FRigElementParentConstraint& ParentConstraintA = ParentConstraints[FirstConstraint.Location];
+				const FRigElementParentConstraint& ParentConstraintB = ParentConstraints[SecondConstraint.Location];
+
+				const FRigElementWeight& WeightA = ParentConstraintA.GetWeight(bInitial); 
+				const FRigElementWeight& WeightB = ParentConstraintB.GetWeight(bInitial);
+				check(WeightA.AffectsLocation());
+				check(WeightB.AffectsLocation());
+				const float Weight = Local::GetWeightForLerp(WeightA.Location, WeightB.Location);
+
+				const FVector ParentLocationA = GetTransform(ParentConstraintA.ParentElement, InTransformType).GetLocation();
+				const FVector ParentLocationB = GetTransform(ParentConstraintB.ParentElement, InTransformType).GetLocation();
+				OutputTransform.SetLocation(FMath::Lerp<FVector>(ParentLocationA, ParentLocationB, Weight));
+			}
+			else if(NumConstraintsAffecting.Location > 2)
+			{
+				check(TotalWeight.Location > SMALL_NUMBER);
+				
+				FVector Location = FVector::ZeroVector;
+				
+				for(int32 ConstraintIndex = 0; ConstraintIndex < ParentConstraints.Num(); ConstraintIndex++)
+				{
+					const FRigElementParentConstraint& ParentConstraint = ParentConstraints[ConstraintIndex];
+					const FRigElementWeight& Weight = ParentConstraint.GetWeight(bInitial);
+					if(!Weight.AffectsLocation())
+					{
+						continue;
+					}
+
+					const FVector ParentLocation = GetTransform(ParentConstraint.ParentElement, InTransformType).GetLocation();
+					Location += ParentLocation * Weight.Location / TotalWeight.Location; 
+				}
+
+				OutputTransform.SetLocation(Location);
+			}
+
+			if(NumConstraintsAffecting.Rotation == 1)
+			{
+				check(FirstConstraint.Rotation != INDEX_NONE);
+
+				const FRigElementParentConstraint& ParentConstraint = ParentConstraints[FirstConstraint.Rotation];
+				const FRigElementWeight& Weight = ParentConstraint.GetWeight(bInitial);
+				const FTransform ParentTransform = GetTransform(ParentConstraint.ParentElement, InTransformType);
+				check(Weight.AffectsRotation());
+				OutputTransform.SetRotation(ParentTransform.GetRotation());
+			}
+			else if(NumConstraintsAffecting.Rotation == 2)
+			{
+				check(FirstConstraint.Rotation != INDEX_NONE);
+				check(SecondConstraint.Rotation != INDEX_NONE);
+
+				const FRigElementParentConstraint& ParentConstraintA = ParentConstraints[FirstConstraint.Rotation];
+				const FRigElementParentConstraint& ParentConstraintB = ParentConstraints[SecondConstraint.Rotation];
+
+				const FRigElementWeight& WeightA = ParentConstraintA.GetWeight(bInitial); 
+				const FRigElementWeight& WeightB = ParentConstraintB.GetWeight(bInitial);
+				check(WeightA.AffectsRotation());
+				check(WeightB.AffectsRotation());
+				const float Weight = Local::GetWeightForLerp(WeightA.Rotation, WeightB.Rotation);
+
+				const FQuat ParentRotationA = GetTransform(ParentConstraintA.ParentElement, InTransformType).GetRotation();
+				const FQuat ParentRotationB = GetTransform(ParentConstraintB.ParentElement, InTransformType).GetRotation();
+				OutputTransform.SetRotation(FQuat::Slerp(ParentRotationA, ParentRotationB, Weight));
+			}
+			else if(NumConstraintsAffecting.Rotation > 2)
+			{
+				check(TotalWeight.Rotation > SMALL_NUMBER);
+				
+				int NumMixedRotations = 0;
+				FQuat FirstRotation = FQuat::Identity;
+				FQuat MixedRotation = FQuat(0.f, 0.f, 0.f, 0.f);
+
+				for(int32 ConstraintIndex = 0; ConstraintIndex < ParentConstraints.Num(); ConstraintIndex++)
+				{
+					const FRigElementParentConstraint& ParentConstraint = ParentConstraints[ConstraintIndex];
+					const FRigElementWeight& Weight = ParentConstraint.GetWeight(bInitial);
+					if(!Weight.AffectsRotation())
+					{
+						continue;
+					}
+
+					FQuat ParentRotation = GetTransform(ParentConstraint.ParentElement, InTransformType).GetRotation().GetNormalized();
+
+					float NormalizedWeight = Weight.Rotation / TotalWeight.Rotation;
+
+					if(NumMixedRotations == 0)
+					{
+						FirstRotation = ParentRotation; 
+					}
+					else if ((ParentRotation | FirstRotation) <= 0.f)
+					{
+						NormalizedWeight = -NormalizedWeight;
+					}
+
+					MixedRotation.X += NormalizedWeight * ParentRotation.X;
+					MixedRotation.Y += NormalizedWeight * ParentRotation.Y;
+					MixedRotation.Z += NormalizedWeight * ParentRotation.Z;
+					MixedRotation.W += NormalizedWeight * ParentRotation.W;
+					NumMixedRotations++;
+				}
+
+				OutputTransform.SetRotation(MixedRotation.GetNormalized());
+			}
+
+			if(NumConstraintsAffecting.Scale == 1)
+			{
+				check(FirstConstraint.Scale != INDEX_NONE);
+
+				const FRigElementParentConstraint& ParentConstraint = ParentConstraints[FirstConstraint.Scale];
+				const FRigElementWeight& Weight = ParentConstraint.GetWeight(bInitial);
+				const FTransform ParentTransform = GetTransform(ParentConstraint.ParentElement, InTransformType);
+				check(Weight.AffectsScale());
+				OutputTransform.SetScale3D(ParentTransform.GetScale3D());
+			}
+			else if(NumConstraintsAffecting.Scale == 2)
+			{
+				check(FirstConstraint.Scale != INDEX_NONE);
+				check(SecondConstraint.Scale != INDEX_NONE);
+
+				const FRigElementParentConstraint& ParentConstraintA = ParentConstraints[FirstConstraint.Scale];
+				const FRigElementParentConstraint& ParentConstraintB = ParentConstraints[SecondConstraint.Scale];
+
+				const FRigElementWeight& WeightA = ParentConstraintA.GetWeight(bInitial); 
+				const FRigElementWeight& WeightB = ParentConstraintB.GetWeight(bInitial);
+				check(WeightA.AffectsScale());
+				check(WeightB.AffectsScale());
+				const float Weight = Local::GetWeightForLerp(WeightA.Scale, WeightB.Scale);
+
+				const FVector ParentScaleA = GetTransform(ParentConstraintA.ParentElement, InTransformType).GetScale3D();
+				const FVector ParentScaleB = GetTransform(ParentConstraintB.ParentElement, InTransformType).GetScale3D();
+				OutputTransform.SetScale3D(FMath::Lerp<FVector>(ParentScaleA, ParentScaleB, Weight));
+			}
+			else if(NumConstraintsAffecting.Scale > 2)
+			{
+				check(TotalWeight.Scale > SMALL_NUMBER);
+				
+				FVector Scale = FVector::ZeroVector;
+				
+				for(int32 ConstraintIndex = 0; ConstraintIndex < ParentConstraints.Num(); ConstraintIndex++)
+				{
+					const FRigElementParentConstraint& ParentConstraint = ParentConstraints[ConstraintIndex];
+					const FRigElementWeight& Weight = ParentConstraint.GetWeight(bInitial);
+					if(!Weight.AffectsScale())
+					{
+						continue;
+					}
+
+					const FVector ParentScale = GetTransform(ParentConstraint.ParentElement, InTransformType).GetScale3D();
+					Scale += ParentScale * Weight.Scale / TotalWeight.Scale; 
+				}
+
+				OutputTransform.SetScale3D(Scale);
+			}
+
+			/*
 			// if we have only one parent we behave like a single parent element
-			if(ParentElements.Num() == 1)
+			if(ParentConstraints.Num() == 1)
 			{
-				Output.Set(GetTransform(ParentElements[0], InTransformType));
+				Output.Set(GetTransform(ParentConstraints[0], InTransformType));
 			}
 			// for exactly two parents we'll lerp
-			else if(ParentElements.Num() == 2)
+			else if(ParentConstraints.Num() == 2)
 			{
 				float Weight = 0.f;
 				const float WeightA = ParentWeights[0];
@@ -1727,22 +1988,22 @@ FTransform URigHierarchy::GetParentTransform(FRigBaseElement* InElement, const E
 
 				if(Weight <= SMALL_NUMBER)
 				{
-					Output.Set(GetTransform(ParentElements[0], InTransformType));
+					Output.Set(GetTransform(ParentConstraints[0], InTransformType));
 				}
 				else if(Weight >= (1.0f - SMALL_NUMBER))
 				{
-					Output.Set(GetTransform(ParentElements[1], InTransformType));
+					Output.Set(GetTransform(ParentConstraints[1], InTransformType));
 				}
 				else
 				{
-					const FTransform ParentTransformA = GetTransform(ParentElements[0], InTransformType);
-					const FTransform ParentTransformB = GetTransform(ParentElements[1], InTransformType);
+					const FTransform ParentTransformA = GetTransform(ParentConstraints[0], InTransformType);
+					const FTransform ParentTransformB = GetTransform(ParentConstraints[1], InTransformType);
 					Output.Set(FControlRigMathLibrary::LerpTransform(ParentTransformA, ParentTransformB, Weight));
 				}
 			}
-			else if(ParentElements.Num() > 2)
+			else if(ParentConstraints.Num() > 2)
 			{
-				ensure(ParentElements.Num() == ParentWeights.Num());
+				ensure(ParentConstraints.Num() == ParentWeights.Num());
 
 				// determine if there are more than one parent weighted
 				float OverallWeight = 0.f;
@@ -1772,7 +2033,7 @@ FTransform URigHierarchy::GetParentTransform(FRigBaseElement* InElement, const E
 				// if there is only weight on one parent, behave like a single parent element
 				if(NumWeightedParents == 1)
 				{
-					Output.Set(GetTransform(ParentElements[FirstWeightedParent], InTransformType));
+					Output.Set(GetTransform(ParentConstraints[FirstWeightedParent], InTransformType));
 				}
 				else if(NumWeightedParents == 2) // for two weighted parents, special case once more
 				{
@@ -1786,8 +2047,8 @@ FTransform URigHierarchy::GetParentTransform(FRigBaseElement* InElement, const E
 						Weight = ClampedWeightB / OverallWeight;
 					}
 
-					const FTransform ParentTransformA = GetTransform(ParentElements[FirstWeightedParent], InTransformType);
-					const FTransform ParentTransformB = GetTransform(ParentElements[SecondWeightedParent], InTransformType);
+					const FTransform ParentTransformA = GetTransform(ParentConstraints[FirstWeightedParent], InTransformType);
+					const FTransform ParentTransformB = GetTransform(ParentConstraints[SecondWeightedParent], InTransformType);
 					Output.Set(FControlRigMathLibrary::LerpTransform(ParentTransformA, ParentTransformB, Weight));
 				}
 				else if(OverallWeight > SMALL_NUMBER)
@@ -1809,7 +2070,7 @@ FTransform URigHierarchy::GetParentTransform(FRigBaseElement* InElement, const E
 						}
 						const float NormalizedWeight = ClampedWeight / OverallWeight;
 
-						const FTransform ParentTransform = GetTransform(ParentElements[ParentIndex], InTransformType);
+						const FTransform ParentTransform = GetTransform(ParentConstraints[ParentIndex], InTransformType);
 
 						// mix translation
 						MixedTranslation += ParentTransform.GetTranslation() * NormalizedWeight; 
@@ -1868,7 +2129,9 @@ FTransform URigHierarchy::GetParentTransform(FRigBaseElement* InElement, const E
 					Output.Set(MixedTransform);
 				}
 			}
+			*/
 
+			Output.Set(OutputTransform);
 			Output.bDirty = false;
 		}
 		return Output.Transform;
@@ -2123,13 +2386,13 @@ bool URigHierarchy::IsParentedTo(FRigBaseElement* InChild, FRigBaseElement* InPa
 
 	if(FRigMultiParentElement* MultiParentElement = Cast<FRigMultiParentElement>(InChild))
 	{
-		for(FRigTransformElement* MultiParentElementParent : MultiParentElement->ParentElements)
+		for(const FRigElementParentConstraint& ParentConstraint : MultiParentElement->ParentConstraints)
 		{
-			if(MultiParentElementParent == InParent)
+			if(ParentConstraint.ParentElement == InParent)
 			{
 				return true;
 			}
-			if(IsParentedTo(MultiParentElementParent, InParent))
+			if(IsParentedTo(ParentConstraint.ParentElement, InParent))
 			{
 				return true;
 			}
@@ -2314,9 +2577,9 @@ void URigHierarchy::UpdateCachedChildren(const FRigBaseElement* InElement, bool
 		}
 		else if(FRigMultiParentElement* MultiParentElement = Cast<FRigMultiParentElement>(Element))
 		{
-			for(FRigTransformElement* ParentElement : MultiParentElement->ParentElements)
+			for(const FRigElementParentConstraint& ParentConstraint : MultiParentElement->ParentConstraints)
 			{
-				if(ParentElement == InElement)
+				if(ParentConstraint.ParentElement == InElement)
 				{
 					InElement->CachedChildren.Add(MultiParentElement);
 					break;
@@ -2352,16 +2615,16 @@ void URigHierarchy::UpdateAllCachedChildren() const
 		}
 		else if(FRigMultiParentElement* MultiParentElement = Cast<FRigMultiParentElement>(Element))
 		{
-			for(FRigTransformElement* ParentElement : MultiParentElement->ParentElements)
+			for(const FRigElementParentConstraint& ParentConstraint : MultiParentElement->ParentConstraints)
 			{
-				if(ParentElement)
+				if(ParentConstraint.ParentElement)
 				{
-					if(!ParentVisited[ParentElement->Index])
+					if(!ParentVisited[ParentConstraint.ParentElement->Index])
 					{
-						ParentElement->CachedChildren.Reset();
-						ParentVisited[ParentElement->Index] = true;
+						ParentConstraint.ParentElement->CachedChildren.Reset();
+						ParentVisited[ParentConstraint.ParentElement->Index] = true;
 					}
-					ParentElement->CachedChildren.Add(Element);
+					ParentConstraint.ParentElement->CachedChildren.Add(Element);
 				}
 			}
 		}
diff --git a/Engine/Plugins/Animation/ControlRig/Source/ControlRig/Private/Rigs/RigHierarchyController.cpp b/Engine/Plugins/Animation/ControlRig/Source/ControlRig/Private/Rigs/RigHierarchyController.cpp
--- a/Engine/Plugins/Animation/ControlRig/Source/ControlRig/Private/Rigs/RigHierarchyController.cpp
+++ b/Engine/Plugins/Animation/ControlRig/Source/ControlRig/Private/Rigs/RigHierarchyController.cpp
@@ -709,8 +709,12 @@ FString URigHierarchyController::ExportToText(TArray<FRigElementKey> InKeys) con
 
 		if(FRigMultiParentElement* MultiParentElement = Cast<FRigMultiParentElement>(Element))
 		{
-			ensure(PerElementData.Parents.Num() == MultiParentElement->ParentWeights.Num());
-			PerElementData.ParentWeights = MultiParentElement->ParentWeights;
+			ensure(PerElementData.Parents.Num() == MultiParentElement->ParentConstraints.Num());
+
+			for(const FRigElementParentConstraint& ParentConstraint : MultiParentElement->ParentConstraints)
+			{
+				PerElementData.ParentWeights.Add(ParentConstraint.Weight);
+			}
 		}
 		else
 		{
@@ -1188,9 +1192,9 @@ TArray<FString> URigHierarchyController::GetAddNullPythonCommands(FRigNullElemen
 	FString TransformStr = RigVMPythonUtils::TransformToPythonString(Null->Pose.Initial.Local.Transform);
 
 	FString ParentKeyStr = "''";
-	if (Null->ParentElements.Num() > 0)
+	if (Null->ParentConstraints.Num() > 0)
 	{
-		ParentKeyStr = Null->ParentElements[0]->GetKey().ToPythonString();		
+		ParentKeyStr = Null->ParentConstraints[0].ParentElement->GetKey().ToPythonString();		
 	}
 		
 	// AddNull(FName InName, FRigElementKey InParent, FTransform InTransform, bool bTransformInGlobal = true, bool bSetupUndo = false);
@@ -1206,9 +1210,9 @@ TArray<FString> URigHierarchyController::GetAddControlPythonCommands(FRigControl
 	FString TransformStr = RigVMPythonUtils::TransformToPythonString(Control->Pose.Initial.Local.Transform);
 
 	FString ParentKeyStr = "''";
-	if (Control->ParentElements.Num() > 0)
+	if (Control->ParentConstraints.Num() > 0)
 	{
-		ParentKeyStr = Control->ParentElements[0]->GetKey().ToPythonString();
+		ParentKeyStr = Control->ParentConstraints[0].ParentElement->GetKey().ToPythonString();
 	}
 
 	FRigControlSettings& Settings = Control->Settings;
@@ -1491,9 +1495,13 @@ bool URigHierarchyController::RemoveElement(FRigBaseElement* InElement)
 			}
 			else if(FRigMultiParentElement* MultiParentElement = Cast<FRigMultiParentElement>(ElementToDirty.Element))
 			{
-				if(MultiParentElement->ParentElements.Contains(InElement))
+				for(const FRigElementParentConstraint& ParentConstraint : MultiParentElement->ParentConstraints)
 				{
-					RemoveParent(MultiParentElement, InElement, true);
+					if(ParentConstraint.ParentElement == InElement)
+					{
+						RemoveParent(MultiParentElement, InElement, true);
+						break;
+					}
 				}
 			}
 		}
@@ -1530,9 +1538,9 @@ bool URigHierarchyController::RemoveElement(FRigBaseElement* InElement)
 	}
 	else if(FRigMultiParentElement* MultiParentElement = Cast<FRigMultiParentElement>(InElement))
 	{
-		for(FRigTransformElement* ParentElement : MultiParentElement->ParentElements)
+		for(const FRigElementParentConstraint& ParentConstraint : MultiParentElement->ParentConstraints)
 		{
-			RemoveElementToDirty(ParentElement, InElement);
+			RemoveElementToDirty(ParentConstraint.ParentElement, InElement);
 		}
 	}
 
@@ -1704,9 +1712,12 @@ bool URigHierarchyController::AddParent(FRigBaseElement* InChild, FRigBaseElemen
 
 	else if(FRigMultiParentElement* MultiParentElement = Cast<FRigMultiParentElement>(InChild))
 	{
-		if(MultiParentElement->ParentElements.Contains(InParent))
+		for(const FRigElementParentConstraint& ParentConstraint : MultiParentElement->ParentConstraints)
 		{
-			return false;
+			if(ParentConstraint.ParentElement == InParent)
+			{
+				return false;
+			}
 		}
 	}
 
@@ -1747,58 +1758,57 @@ bool URigHierarchyController::AddParent(FRigBaseElement* InChild, FRigBaseElemen
 		Hierarchy->GetControlGizmoTransform(ControlElement, ERigTransformType::InitialLocal);
 	}
 
-	if(FRigSingleParentElement* SingleParentElement = Cast<FRigSingleParentElement>(InChild))
+	FRigElementParentConstraint Constraint;
+	Constraint.ParentElement = Cast<FRigTransformElement>(InParent);
+	if(Constraint.ParentElement == nullptr)
 	{
-		if(FRigTransformElement* NewTransformParent = Cast<FRigTransformElement>(InParent))
-		{
-			AddElementToDirty(NewTransformParent, SingleParentElement);
-			SingleParentElement->ParentElement = NewTransformParent;
+		return false;
+	}
 
-			Hierarchy->TopologyVersion++;
+	if(FRigSingleParentElement* SingleParentElement = Cast<FRigSingleParentElement>(InChild))
+	{
+		AddElementToDirty(Constraint.ParentElement, SingleParentElement);
+		SingleParentElement->ParentElement = Constraint.ParentElement;
 
-			if(!bMaintainGlobalTransform)
-			{
-				Hierarchy->PropagateDirtyFlags(SingleParentElement, true, true);
-				Hierarchy->PropagateDirtyFlags(SingleParentElement, false, true);
-			}
+		Hierarchy->TopologyVersion++;
 
-			Notify(ERigHierarchyNotification::ParentChanged, SingleParentElement);
-			return true;
+		if(!bMaintainGlobalTransform)
+		{
+			Hierarchy->PropagateDirtyFlags(SingleParentElement, true, true);
+			Hierarchy->PropagateDirtyFlags(SingleParentElement, false, true);
 		}
+
+		Notify(ERigHierarchyNotification::ParentChanged, SingleParentElement);
+		return true;
 	}
 	else if(FRigMultiParentElement* MultiParentElement = Cast<FRigMultiParentElement>(InChild))
 	{
-		if(FRigTransformElement* NewTransformParent = Cast<FRigTransformElement>(InParent))
-		{
-			AddElementToDirty(NewTransformParent, MultiParentElement);
+		AddElementToDirty(Constraint.ParentElement, MultiParentElement);
 
-			const int32 ParentIndex = MultiParentElement->ParentElements.Add(NewTransformParent);
-			MultiParentElement->ParentWeights.Add(1.f);
-			MultiParentElement->ParentWeightsInitial.Add(1.f);
-			MultiParentElement->IndexLookup.Add(NewTransformParent->GetKey(), ParentIndex);
+		const int32 ParentIndex = MultiParentElement->ParentConstraints.Add(Constraint);
+		MultiParentElement->IndexLookup.Add(Constraint.ParentElement->GetKey(), ParentIndex);
 
-			MultiParentElement->Parent.MarkDirty(ERigTransformType::CurrentGlobal);  
-			MultiParentElement->Parent.MarkDirty(ERigTransformType::InitialGlobal);
-
-			if(FRigControlElement* ControlElement = Cast<FRigControlElement>(MultiParentElement))
-			{
-				ControlElement->Offset.MarkDirty(ERigTransformType::CurrentGlobal);  
-				ControlElement->Offset.MarkDirty(ERigTransformType::InitialGlobal);
-				ControlElement->Gizmo.MarkDirty(ERigTransformType::CurrentGlobal);  
-				ControlElement->Gizmo.MarkDirty(ERigTransformType::InitialGlobal);
-			}
+		MultiParentElement->Parent.MarkDirty(ERigTransformType::CurrentGlobal);  
+		MultiParentElement->Parent.MarkDirty(ERigTransformType::InitialGlobal);
 
-			Hierarchy->TopologyVersion++;
+		if(FRigControlElement* ControlElement = Cast<FRigControlElement>(MultiParentElement))
+		{
+			ControlElement->Offset.MarkDirty(ERigTransformType::CurrentGlobal);  
+			ControlElement->Offset.MarkDirty(ERigTransformType::InitialGlobal);
+			ControlElement->Gizmo.MarkDirty(ERigTransformType::CurrentGlobal);  
+			ControlElement->Gizmo.MarkDirty(ERigTransformType::InitialGlobal);
+		}
 
-			if(!bMaintainGlobalTransform)
-			{
-				Hierarchy->PropagateDirtyFlags(MultiParentElement, true, true);
-				Hierarchy->PropagateDirtyFlags(MultiParentElement, false, true);
-			}
+		Hierarchy->TopologyVersion++;
 
-			Notify(ERigHierarchyNotification::ParentChanged, MultiParentElement);
-			return true;
+		if(!bMaintainGlobalTransform)
+		{
+			Hierarchy->PropagateDirtyFlags(MultiParentElement, true, true);
+			Hierarchy->PropagateDirtyFlags(MultiParentElement, false, true);
 		}
+
+		Notify(ERigHierarchyNotification::ParentChanged, MultiParentElement);
+		return true;
 	}
 	
 	return false;
@@ -1903,8 +1913,17 @@ bool URigHierarchyController::RemoveParent(FRigBaseElement* InChild, FRigBaseEle
 	// single parent children can't be parented multiple times
 	else if(FRigMultiParentElement* MultiParentElement = Cast<FRigMultiParentElement>(InChild))
 	{
-		const int32 ParentIndex = MultiParentElement->ParentElements.Find(ParentTransformElement);
-		if(MultiParentElement->ParentElements.IsValidIndex(ParentIndex))
+		int32 ParentIndex = INDEX_NONE;
+		for(int32 ConstraintIndex = 0; ConstraintIndex < MultiParentElement->ParentConstraints.Num(); ConstraintIndex++)
+		{
+			if(MultiParentElement->ParentConstraints[ConstraintIndex].ParentElement == ParentTransformElement)
+			{
+				ParentIndex = ConstraintIndex;
+				break;
+			}
+		}
+				
+		if(MultiParentElement->ParentConstraints.IsValidIndex(ParentIndex))
 		{
 			if(bMaintainGlobalTransform)
 			{
@@ -1924,12 +1943,10 @@ bool URigHierarchyController::RemoveParent(FRigBaseElement* InChild, FRigBaseEle
 			// remove the previous parent
 			RemoveElementToDirty(InParent, MultiParentElement); 
 
-			const FRigElementKey PreviousParentKey = MultiParentElement->ParentElements[ParentIndex]->GetKey();
+			const FRigElementKey PreviousParentKey = MultiParentElement->ParentConstraints[ParentIndex].ParentElement->GetKey();
 			Hierarchy->PreviousParentMap.FindOrAdd(MultiParentElement->GetKey()) = PreviousParentKey;
 
-			MultiParentElement->ParentElements.RemoveAt(ParentIndex);
-			MultiParentElement->ParentWeights.RemoveAt(ParentIndex);
-			MultiParentElement->ParentWeightsInitial.RemoveAt(ParentIndex);
+			MultiParentElement->ParentConstraints.RemoveAt(ParentIndex);
 			MultiParentElement->IndexLookup.Remove(ParentTransformElement->GetKey());
 			for(TPair<FRigElementKey, int32>& Pair : MultiParentElement->IndexLookup)
 			{
@@ -2013,10 +2030,10 @@ bool URigHierarchyController::RemoveAllParents(FRigBaseElement* InChild, bool bM
 	{
 		bool bSuccess = true;
 
-		TArray<FRigTransformElement*> Parents = MultiParentElement->ParentElements;
-		for(FRigTransformElement* Parent : Parents)
+		TArray<FRigElementParentConstraint> ParentConstraints = MultiParentElement->ParentConstraints;
+		for(const FRigElementParentConstraint& ParentConstraint : ParentConstraints)
 		{
-			if(!RemoveParent(MultiParentElement, Parent, bMaintainGlobalTransform))
+			if(!RemoveParent(MultiParentElement, ParentConstraint.ParentElement, bMaintainGlobalTransform))
 			{
 				bSuccess = false;
 			}
diff --git a/Engine/Plugins/Animation/ControlRig/Source/ControlRig/Private/Rigs/RigHierarchyElements.cpp b/Engine/Plugins/Animation/ControlRig/Source/ControlRig/Private/Rigs/RigHierarchyElements.cpp
--- a/Engine/Plugins/Animation/ControlRig/Source/ControlRig/Private/Rigs/RigHierarchyElements.cpp
+++ b/Engine/Plugins/Animation/ControlRig/Source/ControlRig/Private/Rigs/RigHierarchyElements.cpp
@@ -3,7 +3,7 @@
 #include "Rigs/RigHierarchyElements.h"
 #include "Rigs/RigHierarchy.h"
 #include "Units/RigUnitContext.h"
-#include "UObject/AnimObjectVersion.h"
+#include "ControlRigObjectVersion.h"
 
 ////////////////////////////////////////////////////////////////////////////////
 // FRigBaseElement
@@ -11,7 +11,7 @@
 
 void FRigBaseElement::Serialize(FArchive& Ar, URigHierarchy* Hierarchy, ESerializationPhase SerializationPhase)
 {
-	Ar.UsingCustomVersion(FAnimObjectVersion::GUID);
+	Ar.UsingCustomVersion(FControlRigObjectVersion::GUID);
 
 	if (Ar.IsSaving() || Ar.IsObjectReferenceCollector() || Ar.IsCountingMemory())
 	{
@@ -220,22 +220,22 @@ void FRigMultiParentElement::Save(FArchive& Ar, URigHierarchy* Hierarchy, ESeria
 	{
 		Parent.Save(Ar);
 
-		int32 NumParents = ParentElements.Num();
+		int32 NumParents = ParentConstraints.Num();
 		Ar << NumParents;
 	}
 	else if(SerializationPhase == ESerializationPhase::InterElementData)
 	{
-		for(int32 ParentIndex = 0; ParentIndex < ParentElements.Num(); ParentIndex++)
+		for(int32 ParentIndex = 0; ParentIndex < ParentConstraints.Num(); ParentIndex++)
 		{
 			FRigElementKey ParentKey;
-			if(ParentElements[ParentIndex])
+			if(ParentConstraints[ParentIndex].ParentElement)
 			{
-				ParentKey = ParentElements[ParentIndex]->GetKey();
+				ParentKey = ParentConstraints[ParentIndex].ParentElement->GetKey();
 			}
 
 			Ar << ParentKey;
-			Ar << ParentWeightsInitial[ParentIndex];
-			Ar << ParentWeights[ParentIndex];
+			Ar << ParentConstraints[ParentIndex].InitialWeight;
+			Ar << ParentConstraints[ParentIndex].Weight;
 		}
 	}
 }
@@ -251,22 +251,33 @@ void FRigMultiParentElement::Load(FArchive& Ar, URigHierarchy* Hierarchy, ESeria
 		int32 NumParents = 0;
 		Ar << NumParents;
 
-		ParentElements.SetNumZeroed(NumParents);
-		ParentWeights.SetNumZeroed(NumParents);
-		ParentWeightsInitial.SetNumZeroed(NumParents);
+		ParentConstraints.SetNum(NumParents);
 	}
 	else if(SerializationPhase == ESerializationPhase::InterElementData)
 	{
-		for(int32 ParentIndex = 0; ParentIndex < ParentElements.Num(); ParentIndex++)
+		for(int32 ParentIndex = 0; ParentIndex < ParentConstraints.Num(); ParentIndex++)
 		{
 			FRigElementKey ParentKey;
 			Ar << ParentKey;
 			ensure(ParentKey.IsValid());
 
-			ParentElements[ParentIndex] = Hierarchy->FindChecked<FRigTransformElement>(ParentKey);
+			ParentConstraints[ParentIndex].ParentElement = Hierarchy->FindChecked<FRigTransformElement>(ParentKey);
 
-			Ar << ParentWeightsInitial[ParentIndex];
-			Ar << ParentWeights[ParentIndex];
+			if (Ar.CustomVer(FControlRigObjectVersion::GUID) >= FControlRigObjectVersion::RigHierarchyMultiParentConstraints)
+			{
+				Ar << ParentConstraints[ParentIndex].InitialWeight;
+				Ar << ParentConstraints[ParentIndex].Weight;
+			}
+			else
+			{
+				float InitialWeight = 0.f;
+				Ar << InitialWeight;
+				ParentConstraints[ParentIndex].InitialWeight = FRigElementWeight(InitialWeight);
+
+				float Weight = 0.f;
+				Ar << Weight;
+				ParentConstraints[ParentIndex].Weight = FRigElementWeight(Weight);
+			}
 
 			IndexLookup.Add(ParentKey, ParentIndex);
 		}
@@ -274,21 +285,21 @@ void FRigMultiParentElement::Load(FArchive& Ar, URigHierarchy* Hierarchy, ESeria
 }
 
 void FRigMultiParentElement::CopyFrom(URigHierarchy* InHierarchy, FRigBaseElement* InOther,
-	URigHierarchy* InOtherHierarchy)
+                                      URigHierarchy* InOtherHierarchy)
 {
 	Super::CopyFrom(InHierarchy, InOther, InOtherHierarchy);
 	
 	const FRigMultiParentElement* Source = CastChecked<FRigMultiParentElement>(InOther);
 	Parent = Source->Parent;
-	ParentElements.Reset();
-	ParentWeights = Source->ParentWeights;
-	ParentWeightsInitial = Source->ParentWeightsInitial;
+	ParentConstraints.Reset();
 
-	for(int32 ParentIndex = 0; ParentIndex < Source->ParentElements.Num(); ParentIndex++)
+	for(int32 ParentIndex = 0; ParentIndex < Source->ParentConstraints.Num(); ParentIndex++)
 	{
-		const FRigTransformElement* SourceParentElement = Source->ParentElements[ParentIndex];
-		ParentElements.Add(CastChecked<FRigTransformElement>(InHierarchy->Get(SourceParentElement->Index)));
-		check(ParentElements[ParentIndex]->GetKey() == SourceParentElement->GetKey());
+		FRigElementParentConstraint ParentConstraint = Source->ParentConstraints[ParentIndex];
+		const FRigTransformElement* SourceParentElement = ParentConstraint.ParentElement;
+		ParentConstraint.ParentElement = CastChecked<FRigTransformElement>(InHierarchy->Get(SourceParentElement->Index));
+		ParentConstraints.Add(ParentConstraint);
+		check(ParentConstraints[ParentIndex].ParentElement->GetKey() == SourceParentElement->GetKey());
 	}
 }
 
diff --git a/Engine/Plugins/Animation/ControlRig/Source/ControlRig/Public/ControlRigObjectVersion.h b/Engine/Plugins/Animation/ControlRig/Source/ControlRig/Public/ControlRigObjectVersion.h
--- a/Engine/Plugins/Animation/ControlRig/Source/ControlRig/Public/ControlRigObjectVersion.h
+++ b/Engine/Plugins/Animation/ControlRig/Source/ControlRig/Public/ControlRigObjectVersion.h
@@ -33,6 +33,9 @@ struct CONTROLRIG_API FControlRigObjectVersion
 		// Hierarchy V2.0
 		RigHierarchyV2,
 
+		// RigHierarchy to support multi component parent constraints
+		RigHierarchyMultiParentConstraints,
+		
 		// -----<new versions can be added above this line>-------------------------------------------------
 		VersionPlusOne,
 		LatestVersion = VersionPlusOne - 1
diff --git a/Engine/Plugins/Animation/ControlRig/Source/ControlRig/Public/Rigs/RigHierarchy.h b/Engine/Plugins/Animation/ControlRig/Source/ControlRig/Public/Rigs/RigHierarchy.h
--- a/Engine/Plugins/Animation/ControlRig/Source/ControlRig/Public/Rigs/RigHierarchy.h
+++ b/Engine/Plugins/Animation/ControlRig/Source/ControlRig/Public/Rigs/RigHierarchy.h
@@ -1607,7 +1607,7 @@ public:
 	 * @return Returns the weight of a parent below a multi parent element, or FLT_MAX if the parent is invalid
 	 */
 	UFUNCTION(BlueprintCallable, Category = URigHierarchy)
-    float GetParentWeight(FRigElementKey InChild, FRigElementKey InParent, bool bInitial = false) const;
+    FRigElementWeight GetParentWeight(FRigElementKey InChild, FRigElementKey InParent, bool bInitial = false) const;
 
 	/**
 	 * Returns the weight of a parent below a multi parent element
@@ -1616,7 +1616,7 @@ public:
 	 * @param bInitial If true the initial weights will be used
 	 * @return Returns the weight of a parent below a multi parent element, or FLT_MAX if the parent is invalid
 	 */
-    float GetParentWeight(const FRigBaseElement* InChild, const FRigBaseElement* InParent, bool bInitial = false) const;
+    FRigElementWeight GetParentWeight(const FRigBaseElement* InChild, const FRigBaseElement* InParent, bool bInitial = false) const;
 
 	/**
 	 * Returns the weight of a parent below a multi parent element
@@ -1625,7 +1625,7 @@ public:
 	 * @param bInitial If true the initial weights will be used
 	 * @return Returns the weight of a parent below a multi parent element, or FLT_MAX if the parent is invalid
 	 */
-	float GetParentWeight(const FRigBaseElement* InChild, int32 InParentIndex, bool bInitial = false) const;
+	FRigElementWeight GetParentWeight(const FRigBaseElement* InChild, int32 InParentIndex, bool bInitial = false) const;
 
 	/**
 	 * Sets the weight of a parent below a multi parent element
@@ -1637,7 +1637,7 @@ public:
 	 * @return Returns true if changing the weight was successful
 	 */
 	UFUNCTION(BlueprintCallable, Category = URigHierarchy)
-    bool SetParentWeight(FRigElementKey InChild, FRigElementKey InParent, float InWeight, bool bInitial = false, bool bAffectChildren = true);
+    bool SetParentWeight(FRigElementKey InChild, FRigElementKey InParent, FRigElementWeight InWeight, bool bInitial = false, bool bAffectChildren = true);
 
 	/**
 	 * Sets the weight of a parent below a multi parent element
@@ -1648,7 +1648,7 @@ public:
 	 * @param bAffectChildren If set to false children will not move (maintain global).
 	 * @return Returns true if changing the weight was successful
 	 */
-	bool SetParentWeight(FRigBaseElement* InChild, const FRigBaseElement* InParent, float InWeight, bool bInitial = false, bool bAffectChildren = true);
+	bool SetParentWeight(FRigBaseElement* InChild, const FRigBaseElement* InParent, FRigElementWeight InWeight, bool bInitial = false, bool bAffectChildren = true);
 
 	/**
 	 * Sets the weight of a parent below a multi parent element
@@ -1659,7 +1659,7 @@ public:
 	 * @param bAffectChildren If set to false children will not move (maintain global).
 	 * @return Returns true if changing the weight was successful
 	 */
-	bool SetParentWeight(FRigBaseElement* InChild, int32 InParentIndex, float InWeight, bool bInitial = false, bool bAffectChildren = true);
+	bool SetParentWeight(FRigBaseElement* InChild, int32 InParentIndex, FRigElementWeight InWeight, bool bInitial = false, bool bAffectChildren = true);
 
 	/**
 	 * Returns true if an element is parented to another element
diff --git a/Engine/Plugins/Animation/ControlRig/Source/ControlRig/Public/Rigs/RigHierarchyElements.h b/Engine/Plugins/Animation/ControlRig/Source/ControlRig/Public/Rigs/RigHierarchyElements.h
--- a/Engine/Plugins/Animation/ControlRig/Source/ControlRig/Public/Rigs/RigHierarchyElements.h
+++ b/Engine/Plugins/Animation/ControlRig/Source/ControlRig/Public/Rigs/RigHierarchyElements.h
@@ -386,49 +386,6 @@
 	FRigLocalAndGlobalTransform Initial;
 };
 
-USTRUCT()
-struct CONTROLRIG_API FRigHierarchyCopyPasteContentPerElement
-{
-	GENERATED_USTRUCT_BODY()
-
-    UPROPERTY()
-	FRigElementKey Key;
-
-	UPROPERTY()
-	FString Content;
-
-	UPROPERTY()
-	TArray<FRigElementKey> Parents;
-
-	UPROPERTY()
-	TArray<float> ParentWeights;
-
-	UPROPERTY()
-	FRigCurrentAndInitialTransform Pose;
-};
-
-USTRUCT()
-struct CONTROLRIG_API FRigHierarchyCopyPasteContent
-{
-	GENERATED_USTRUCT_BODY()
-
-    UPROPERTY()
-	TArray<FRigHierarchyCopyPasteContentPerElement> Elements;
-
-	// Maintain properties below for backwards compatibility pre-5.0
-	UPROPERTY()
-	TArray<ERigElementType> Types;
-
-	UPROPERTY()
-	TArray<FString> Contents;
-
-	UPROPERTY()
-	TArray<FTransform> LocalTransforms;
-
-	UPROPERTY()
-	TArray<FTransform> GlobalTransforms;
-};
-
 USTRUCT(BlueprintType)
 struct CONTROLRIG_API FRigBaseElement
 {
@@ -654,5 +611,98 @@
 };
 
 USTRUCT(BlueprintType)
+struct CONTROLRIG_API FRigElementWeight
+{
+public:
+	GENERATED_BODY()
+
+	UPROPERTY()
+	float Location;
+
+	UPROPERTY()
+	float Rotation;
+
+	UPROPERTY()
+	float Scale;
+
+	FORCEINLINE FRigElementWeight()
+		: Location(1.f)
+		, Rotation(1.f)
+		, Scale(1.f)
+	{}
+
+	FORCEINLINE FRigElementWeight(float InWeight)
+		: Location(InWeight)
+		, Rotation(InWeight)
+		, Scale(InWeight)
+	{}
+
+	FORCEINLINE FRigElementWeight(float InLocation, float InRotation, float InScale)
+		: Location(InLocation)
+		, Rotation(InRotation)
+		, Scale(InScale)
+	{}
+
+	FORCEINLINE friend FArchive& operator <<(FArchive& Ar, FRigElementWeight& Weight)
+	{
+		Ar << Weight.Location;
+		Ar << Weight.Rotation;
+		Ar << Weight.Scale;
+		return Ar;
+	}
+
+	FORCEINLINE bool AffectsLocation() const
+	{
+		return Location > SMALL_NUMBER;
+	}
+
+	FORCEINLINE bool AffectsRotation() const
+	{
+		return Rotation > SMALL_NUMBER;
+	}
+
+	FORCEINLINE bool AffectsScale() const
+	{
+		return Scale > SMALL_NUMBER;
+	}
+
+	FORCEINLINE bool IsAlmostZero() const
+	{
+		return !AffectsLocation() && !AffectsRotation() && !AffectsScale();
+	}
+
+	FORCEINLINE friend FRigElementWeight operator *(FRigElementWeight Weight, float Scale)
+	{
+		return FRigElementWeight(Weight.Location * Scale, Weight.Rotation * Scale, Weight.Scale * Scale);
+	}
+
+	FORCEINLINE friend FRigElementWeight operator *(float Scale, FRigElementWeight Weight)
+	{
+		return FRigElementWeight(Weight.Location * Scale, Weight.Rotation * Scale, Weight.Scale * Scale);
+	}
+};
+
+USTRUCT()
+struct CONTROLRIG_API FRigElementParentConstraint
+{
+public:
+	GENERATED_BODY()
+
+	FRigTransformElement* ParentElement;
+	FRigElementWeight Weight;
+	FRigElementWeight InitialWeight;
+		
+	FORCEINLINE FRigElementParentConstraint()
+		: ParentElement(nullptr)
+	{
+	}
+
+	FORCEINLINE const FRigElementWeight& GetWeight(bool bInitial = false) const
+	{
+		return bInitial ? InitialWeight : Weight;
+	}
+};
+
+USTRUCT(BlueprintType)
 struct CONTROLRIG_API FRigMultiParentElement : public FRigTransformElement
 {
@@ -673,11 +723,9 @@
 	UPROPERTY(BlueprintReadOnly, Category = RigElement)
 	FRigCurrentAndInitialTransform Parent;
 
-	TArray<FRigTransformElement*> ParentElements;
-	TArray<float> ParentWeights;
-	TArray<float> ParentWeightsInitial;
+	TArray<FRigElementParentConstraint> ParentConstraints;
 	TMap<FRigElementKey, int32> IndexLookup;
-
+	
 protected:
 
 	virtual void CopyFrom(URigHierarchy* InHierarchy, FRigBaseElement* InOther, URigHierarchy* InOtherHierarchy) override;
@@ -1064,3 +1112,46 @@
 	friend struct FRigBaseElement;
 	friend class URigHierarchyController;
 };
+
+USTRUCT()
+struct CONTROLRIG_API FRigHierarchyCopyPasteContentPerElement
+{
+	GENERATED_USTRUCT_BODY()
+
+	UPROPERTY()
+	FRigElementKey Key;
+
+	UPROPERTY()
+	FString Content;
+
+	UPROPERTY()
+	TArray<FRigElementKey> Parents;
+
+	UPROPERTY()
+	TArray<FRigElementWeight> ParentWeights;
+
+	UPROPERTY()
+	FRigCurrentAndInitialTransform Pose;
+};
+
+USTRUCT()
+struct CONTROLRIG_API FRigHierarchyCopyPasteContent
+{
+	GENERATED_USTRUCT_BODY()
+
+	UPROPERTY()
+	TArray<FRigHierarchyCopyPasteContentPerElement> Elements;
+
+	// Maintain properties below for backwards compatibility pre-5.0
+	UPROPERTY()
+	TArray<ERigElementType> Types;
+
+	UPROPERTY()
+	TArray<FString> Contents;
+
+	UPROPERTY()
+	TArray<FTransform> LocalTransforms;
+
+	UPROPERTY()
+	TArray<FTransform> GlobalTransforms;
+};
diff --git a/Engine/Plugins/Animation/ControlRig/Source/ControlRigEditor/Private/Sequencer/ControlRigParameterTrackEditor.cpp b/Engine/Plugins/Animation/ControlRig/Source/ControlRigEditor/Private/Sequencer/ControlRigParameterTrackEditor.cpp
--- a/Engine/Plugins/Animation/ControlRig/Source/ControlRigEditor/Private/Sequencer/ControlRigParameterTrackEditor.cpp
+++ b/Engine/Plugins/Animation/ControlRig/Source/ControlRigEditor/Private/Sequencer/ControlRigParameterTrackEditor.cpp
@@ -1914,12 +1914,18 @@ void FControlRigParameterTrackEditor::HandleControlSelected(UControlRig* Subject
 	{
 		Hierarchy->ForEach<FRigControlElement>([ControlElement, Controller, bSelected](FRigControlElement* OtherControlElement) -> bool
 		{
-			if (OtherControlElement->ParentElements.Contains(ControlElement) &&
-                (OtherControlElement->Settings.ControlType == ERigControlType::Bool ||
+			if(OtherControlElement->Settings.ControlType == ERigControlType::Bool ||
                     OtherControlElement->Settings.ControlType == ERigControlType::Float ||
-                    OtherControlElement->Settings.ControlType == ERigControlType::Integer))
+                    OtherControlElement->Settings.ControlType == ERigControlType::Integer)
 			{
-				Controller->SelectElement(OtherControlElement->GetKey(), bSelected);
+				for(const FRigElementParentConstraint& ParentConstraint : OtherControlElement->ParentConstraints)
+				{
+					if(ParentConstraint.ParentElement == ControlElement)
+					{
+						Controller->SelectElement(OtherControlElement->GetKey(), bSelected);
+						break;
+					}
+				}
 			}
 
 			return true;
