Control Rig: Build break fix for (674,82): error: declaration shadows a field of 'FRigElementWeight'

diff --git a/Engine/Plugins/Animation/ControlRig/Source/ControlRig/Public/Rigs/RigHierarchyElements.h b/Engine/Plugins/Animation/ControlRig/Source/ControlRig/Public/Rigs/RigHierarchyElements.h
--- a/Engine/Plugins/Animation/ControlRig/Source/ControlRig/Public/Rigs/RigHierarchyElements.h
+++ b/Engine/Plugins/Animation/ControlRig/Source/ControlRig/Public/Rigs/RigHierarchyElements.h
@@ -671,14 +671,14 @@ public:
 		return !AffectsLocation() && !AffectsRotation() && !AffectsScale();
 	}
 
-	FORCEINLINE friend FRigElementWeight operator *(FRigElementWeight Weight, float Scale)
+	FORCEINLINE friend FRigElementWeight operator *(FRigElementWeight InWeight, float InScale)
 	{
-		return FRigElementWeight(Weight.Location * Scale, Weight.Rotation * Scale, Weight.Scale * Scale);
+		return FRigElementWeight(InWeight.Location * InScale, InWeight.Rotation * InScale, InWeight.Scale * InScale);
 	}
 
-	FORCEINLINE friend FRigElementWeight operator *(float Scale, FRigElementWeight Weight)
+	FORCEINLINE friend FRigElementWeight operator *(float InScale, FRigElementWeight InWeight)
 	{
-		return FRigElementWeight(Weight.Location * Scale, Weight.Rotation * Scale, Weight.Scale * Scale);
+		return FRigElementWeight(InWeight.Location * InScale, InWeight.Rotation * InScale, InWeight.Scale * InScale);
 	}
 };
 
