Control Rig: follow up to CL16849069, only reset controls to their initial position as resetting the entire hierarchy can override incoming poses

diff --git a/Engine/Plugins/Animation/ControlRig/Source/ControlRig/Private/ControlRig.cpp b/Engine/Plugins/Animation/ControlRig/Source/ControlRig/Private/ControlRig.cpp
--- a/Engine/Plugins/Animation/ControlRig/Source/ControlRig/Private/ControlRig.cpp
+++ b/Engine/Plugins/Animation/ControlRig/Source/ControlRig/Private/ControlRig.cpp
@@ -695,7 +695,7 @@ void UControlRig::Execute(const EControlRigState InState, const FName& InEventNa
 				}
 
 				// reset the pose to initial such that setup event can run from a deterministic initial state
-				GetHierarchy()->ResetPoseToInitial();
+				GetHierarchy()->ResetPoseToInitial(ERigElementType::Control);
 
 				if (PreSetupEvent.IsBound())
 				{
