Control Rig: save and restore the entire pose as we enter/exit setup event, beyond just controls

diff --git a/Engine/Plugins/Animation/ControlRig/Source/ControlRig/Private/ControlRig.cpp b/Engine/Plugins/Animation/ControlRig/Source/ControlRig/Private/ControlRig.cpp
--- a/Engine/Plugins/Animation/ControlRig/Source/ControlRig/Private/ControlRig.cpp
+++ b/Engine/Plugins/Animation/ControlRig/Source/ControlRig/Private/ControlRig.cpp
@@ -684,18 +684,18 @@ void UControlRig::Execute(const EControlRigState InState, const FName& InEventNa
 			}
 			
 			{
-				// save the current control values to preserve user intention
-				// the control values are reapplied to the rig after setup event as it goes out of scope
-				TUniquePtr<UControlRig::FControlValueScope> ValueScope;
+				// save the current state of all pose elements to preserve user intention
+				// the saved pose is reapplied to the rig after setup event as the pose scope goes out of scope
+				TUniquePtr<UControlRig::FPoseScope> PoseScope;
 				if (!bSetupModeEnabled)
 				{
 					// only do this in non-setup mode because 
 					// when setup mode is enabled, the control values are cleared before reaching here (too late to save them)
-					ValueScope = MakeUnique<UControlRig::FControlValueScope>(this);
+					PoseScope = MakeUnique<UControlRig::FPoseScope>(this);
 				}
 
 				// reset the pose to initial such that setup event can run from a deterministic initial state
-				GetHierarchy()->ResetPoseToInitial(ERigElementType::Control);
+				GetHierarchy()->ResetPoseToInitial(ERigElementType::All);
 
 				if (PreSetupEvent.IsBound())
 				{
@@ -2635,31 +2635,20 @@ void UControlRig::OnHierarchyTransformUndoRedo(URigHierarchy* InHierarchy, const
 	}
 }
 
-UControlRig::FControlValueScope::FControlValueScope(UControlRig* InControlRig)
+UControlRig::FPoseScope::FPoseScope(UControlRig* InControlRig)
 {
 	check(InControlRig);
 
-	TArray<FRigControlElement*> Controls = InControlRig->AvailableControls();
-	for (FRigControlElement* ControlElement : Controls)
-	{
-		ControlValues.Add(ControlElement->GetName(), InControlRig->GetControlValue(ControlElement->GetName()));
-	}
+	CachedPose = InControlRig->GetHierarchy()->GetPose();
 
 	ControlRig = InControlRig;
 }
 
-UControlRig::FControlValueScope::~FControlValueScope()
+UControlRig::FPoseScope::~FPoseScope()
 {
 	check(ControlRig);
 
-	for (const TPair<FName, FRigControlValue>& Pair : ControlValues)
-	{
-		if (ControlRig->FindControl(Pair.Key))
-		{
-			// bNotify = false such that it won't trigger sequencer auto key
-			ControlRig->SetControlValue(Pair.Key, Pair.Value, false, FRigControlModifiedContext(), false);
-		}
-	}
+	ControlRig->GetHierarchy()->SetPose(CachedPose);
 }
  
 #undef LOCTEXT_NAMESPACE
diff --git a/Engine/Plugins/Animation/ControlRig/Source/ControlRig/Public/ControlRig.h b/Engine/Plugins/Animation/ControlRig/Source/ControlRig/Public/ControlRig.h
--- a/Engine/Plugins/Animation/ControlRig/Source/ControlRig/Public/ControlRig.h
+++ b/Engine/Plugins/Animation/ControlRig/Source/ControlRig/Public/ControlRig.h
@@ -697,19 +697,18 @@ public:
 
 private:
 
-	// Class used to temporarily cache all 
-	// current control values and reapply them
-	// on destruction, similar to UControlRigBlueprint::FControlValueScope
-	class FControlValueScope
+	// Class used to temporarily cache current pose of the hierarchy
+	// restore it on destruction, similar to UControlRigBlueprint::FControlValueScope
+	class FPoseScope
 	{
 	public:
-		FControlValueScope(UControlRig* InControlRig);
-		~FControlValueScope();
+		FPoseScope(UControlRig* InControlRig);
+		~FPoseScope();
 
 	private:
 
 		UControlRig* ControlRig;
-		TMap<FName, FRigControlValue> ControlValues;
+		FRigPose CachedPose;
 	};
 
 	friend class FControlRigBlueprintCompilerContext;
