Add missing functionality to Control Rig SpringInterp so Verlet can be (and is now) deprecated

diff --git a/Engine/Plugins/Animation/ControlRig/Source/ControlRig/Private/Units/Simulation/RigUnit_SpringInterp.cpp b/Engine/Plugins/Animation/ControlRig/Source/ControlRig/Private/Units/Simulation/RigUnit_SpringInterp.cpp
--- a/Engine/Plugins/Animation/ControlRig/Source/ControlRig/Private/Units/Simulation/RigUnit_SpringInterp.cpp
+++ b/Engine/Plugins/Animation/ControlRig/Source/ControlRig/Private/Units/Simulation/RigUnit_SpringInterp.cpp
@@ -21,10 +21,20 @@ FRigUnit_SpringInterp_Execute()
 		// Treat the input as a frequency in Hz
 		float AngularFrequency = Strength * 2.0f * PI;
 		float Stiffness = AngularFrequency * AngularFrequency;
+		float AdjustedTarget = Target;
+		if (!FMath::IsNearlyZero(Stiffness))
+		{
+			AdjustedTarget += Force / (Stiffness * RigUnitSpringInterpConstants::Mass);
+		}
+		else
+		{
+			SpringState.Velocity += Force * (Context.DeltaTime / RigUnitSpringInterpConstants::Mass);
+		}
 		Result = UKismetMathLibrary::FloatSpringInterp(
-			bUseCurrentInput ? Current : Result, Target, SpringState, Stiffness, CriticalDamping,
+			bUseCurrentInput ? Current : Result, AdjustedTarget, SpringState, Stiffness, CriticalDamping,
 			Context.DeltaTime, RigUnitSpringInterpConstants::Mass, TargetVelocityAmount, 
 			false, 0.0f, 0.0f, !bUseCurrentInput || bInitializeFromTarget);
+		Velocity = SpringState.Velocity;
 	}
 }
 
@@ -41,11 +51,21 @@ FRigUnit_SpringInterpVector_Execute()
 		// Treat the input as a frequency in Hz
 		float AngularFrequency = Strength * 2.0f * PI;
 		float Stiffness = AngularFrequency * AngularFrequency;
+		FVector AdjustedTarget = Target;
+		if (!FMath::IsNearlyZero(Stiffness))
+		{
+			AdjustedTarget += Force / (Stiffness * RigUnitSpringInterpConstants::Mass);
+		}
+		else
+		{
+			SpringState.Velocity += Force * (Context.DeltaTime / RigUnitSpringInterpConstants::Mass);
+		}
 		Result = UKismetMathLibrary::VectorSpringInterp(
-			bUseCurrentInput ? Current : Result, Target, SpringState, Stiffness, CriticalDamping,
-			Context.DeltaTime, RigUnitSpringInterpConstants::Mass, TargetVelocityAmount, 
+			bUseCurrentInput ? Current : Result, AdjustedTarget, SpringState, Stiffness, CriticalDamping,
+			Context.DeltaTime, RigUnitSpringInterpConstants::Mass, TargetVelocityAmount,
 			false, FVector(), FVector(), !bUseCurrentInput || bInitializeFromTarget);
 	}
+	Velocity = SpringState.Velocity;
 }
 
 FRigUnit_SpringInterpQuaternion_Execute()
@@ -61,10 +81,12 @@ FRigUnit_SpringInterpQuaternion_Execute()
 		// Treat the input as a frequency in Hz
 		float AngularFrequency = Strength * 2.0f * PI;
 		float Stiffness = AngularFrequency * AngularFrequency;
+		SpringState.AngularVelocity += Torque * (Context.DeltaTime / RigUnitSpringInterpConstants::Mass);
 		Result = UKismetMathLibrary::QuaternionSpringInterp(
 			bUseCurrentInput ? Current : Result, Target, SpringState, Stiffness, CriticalDamping,
 			Context.DeltaTime, RigUnitSpringInterpConstants::Mass, TargetVelocityAmount, 
 			!bUseCurrentInput || bInitializeFromTarget);
+		AngularVelocity = SpringState.AngularVelocity;
 	}
 }
 
diff --git a/Engine/Plugins/Animation/ControlRig/Source/ControlRig/Private/Units/Simulation/RigUnit_SpringInterp.h b/Engine/Plugins/Animation/ControlRig/Source/ControlRig/Private/Units/Simulation/RigUnit_SpringInterp.h
--- a/Engine/Plugins/Animation/ControlRig/Source/ControlRig/Private/Units/Simulation/RigUnit_SpringInterp.h
+++ b/Engine/Plugins/Animation/ControlRig/Source/ControlRig/Private/Units/Simulation/RigUnit_SpringInterp.h
@@ -16,12 +16,13 @@
 
 	FRigUnit_SpringInterp()
 	{
-		bUseCurrentInput = true;
-		Current = Target = Result = 0.0f;
+		bUseCurrentInput = false;
+		Current = Target = Velocity = Result = 0.0f;
 		Strength = 4.0f;
 		CriticalDamping = 1.0f;
 		TargetVelocityAmount = 0.0f;
 		bInitializeFromTarget = false;
+		Force = 0.0f;
 		SpringState = FFloatSpringState();
 	}
 
@@ -33,17 +34,5 @@
 	float Target;
 
 	/**
-	 * If true, then the Current input will be used to initialize the state, and is required to be a variable that 
-	 * holds the current state. If false then the Target value will be used to initialize the state and the Current 
-	 * input will be ignored/unnecessary as a state will be maintained by this node.
-	 */
-	UPROPERTY(meta=(Input, Constant))
-	bool bUseCurrentInput;
-
-	/** Current position of the spring. */
-	UPROPERTY(meta = (Input, EditCondition = "bUseCurrentInput"))
-	float Current;
-
-	/**
 	 * The spring strength determines how hard it will pull towards the target. The value is the frequency
 	 * at which it will oscillate when there is no damping.
@@ -61,5 +50,23 @@
 	float CriticalDamping;
 
 	/**
+	 * Extra force to apply (since the mass is 1, this is also the acceleration).
+	 */
+	UPROPERTY(meta = (Input))
+	float Force;
+
+	/**
+	 * If true, then the Current input will be used to initialize the state, and is required to be a variable that 
+	 * holds the current state. If false then the Target value will be used to initialize the state and the Current 
+	 * input will be ignored/unnecessary as a state will be maintained by this node.
+	 */
+	UPROPERTY(meta=(Input, Constant))
+	bool bUseCurrentInput;
+
+	/** Current position of the spring. */
+	UPROPERTY(meta = (Input, EditCondition = "bUseCurrentInput"))
+	float Current;
+
+	/**
 	 * The amount that the velocity should be passed through to the spring. A value of 1 will result in more
 	 * responsive output, but if the input is noisy or has step changes, these discontinuities will be passed 
@@ -78,6 +85,10 @@
 	UPROPERTY(meta=(Output))
 	float Result;
 
+	/** Velocity */
+	UPROPERTY(meta = (Output))
+	float Velocity;
+
 	UPROPERTY()
 	FFloatSpringState SpringState;
 };
@@ -92,12 +103,13 @@
 	
 	FRigUnit_SpringInterpVector()
 	{
-		bUseCurrentInput = true;
-		Current = Target = Result = FVector::ZeroVector;
+		bUseCurrentInput = false;
+		Current = Target = Velocity = Result = FVector::ZeroVector;
 		Strength = 4.0f;
 		CriticalDamping = 1.0f;
 		TargetVelocityAmount = 0.0f;
 		bInitializeFromTarget = false;
+		Force = FVector::ZeroVector;
 		SpringState = FVectorSpringState();
 	}
 
@@ -109,17 +121,5 @@
 	FVector Target;
 
 	/**
-	 * If true, then the Current input will be used to initialize the state, and is required to be a variable that 
-	 * holds the current state. If false then the Target value will be used to initialize the state and the Current 
-	 * input will be ignored/unnecessary as a state will be maintained by this node.
-	 */
-	UPROPERTY(meta=(Input, Constant))
-	bool bUseCurrentInput;
-
-	/** Current position of the spring. */
-	UPROPERTY(meta = (Input, EditCondition = "bUseCurrentInput"))
-	FVector Current;
-
-	/**
 	 * The spring strength determines how hard it will pull towards the target. The value is the frequency
 	 * at which it will oscillate when there is no damping.
@@ -137,5 +137,23 @@
 	float CriticalDamping;
 
 	/**
+	 * Extra force to apply (since the mass is 1, this is also the acceleration).
+	 */
+	UPROPERTY(meta = (Input))
+	FVector Force;
+
+	/**
+	 * If true, then the Current input will be used to initialize the state, and is required to be a variable that 
+	 * holds the current state. If false then the Target value will be used to initialize the state and the Current 
+	 * input will be ignored/unnecessary as a state will be maintained by this node.
+	 */
+	UPROPERTY(meta=(Input, Constant))
+	bool bUseCurrentInput;
+
+	/** Current position of the spring. */
+	UPROPERTY(meta = (Input, EditCondition = "bUseCurrentInput"))
+	FVector Current;
+
+	/**
 	 * The amount that the velocity should be passed through to the spring. A value of 1 will result in more
 	 * responsive output, but if the input is noisy or has step changes, these discontinuities will be passed 
@@ -154,6 +172,10 @@
 	UPROPERTY(meta=(Output))
 	FVector Result;
 
+	/** Velocity */
+	UPROPERTY(meta = (Output))
+	FVector Velocity;
+
 	UPROPERTY()
 	FVectorSpringState SpringState;
 };
@@ -168,12 +190,14 @@
 	
 	FRigUnit_SpringInterpQuaternion()
 	{
-		bUseCurrentInput = true;
+		bUseCurrentInput = false;
 		Current = Target = Result = FQuat::Identity;
+		AngularVelocity = FVector::ZeroVector;
 		Strength = 4.0f;
 		CriticalDamping = 1.0f;
 		TargetVelocityAmount = 0.0f;
 		bInitializeFromTarget = false;
+		Torque = FVector::ZeroVector;
 		SpringState = FQuaternionSpringState();
 	}
 
@@ -185,17 +209,5 @@
 	FQuat Target;
 
 	/**
-	 * If true, then the Current input will be used to initialize the state, and is required to be a variable that 
-	 * holds the current state. If false then the Target value will be used to initialize the state and the Current 
-	 * input will be ignored/unnecessary as a state will be maintained by this node.
-	 */
-	UPROPERTY(meta=(Input, Constant))
-	bool bUseCurrentInput;
-
-	/** Current position of the spring. */
-	UPROPERTY(meta = (Input, EditCondition = "bUseCurrentInput"))
-	FQuat Current;
-
-	/**
 	 * The spring strength determines how hard it will pull towards the target. The value is the frequency
 	 * at which it will oscillate when there is no damping.
@@ -213,5 +225,23 @@
 	float CriticalDamping;
 
 	/**
+	 * Extra torque to apply (since the moment of inertia is 1, this is also the angular acceleration).
+	 */
+	UPROPERTY(meta = (Input))
+	FVector Torque;
+
+	/**
+	 * If true, then the Current input will be used to initialize the state, and is required to be a variable that 
+	 * holds the current state. If false then the Target value will be used to initialize the state and the Current 
+	 * input will be ignored/unnecessary as a state will be maintained by this node.
+	 */
+	UPROPERTY(meta=(Input, Constant))
+	bool bUseCurrentInput;
+
+	/** Current position of the spring. */
+	UPROPERTY(meta = (Input, EditCondition = "bUseCurrentInput"))
+	FQuat Current;
+
+	/**
 	 * The amount that the velocity should be passed through to the spring. A value of 1 will result in more
 	 * responsive output, but if the input is noisy or has step changes, these discontinuities will be passed 
@@ -230,6 +260,10 @@
 	UPROPERTY(meta=(Output))
 	FQuat Result;
 
+	/** Angular velocity */
+	UPROPERTY(meta = (Output))
+	FVector AngularVelocity;
+
 	UPROPERTY()
 	FQuaternionSpringState SpringState;
 };
diff --git a/Engine/Plugins/Animation/ControlRig/Source/ControlRig/Private/Units/Simulation/RigUnit_Verlet.h b/Engine/Plugins/Animation/ControlRig/Source/ControlRig/Private/Units/Simulation/RigUnit_Verlet.h
--- a/Engine/Plugins/Animation/ControlRig/Source/ControlRig/Private/Units/Simulation/RigUnit_Verlet.h
+++ b/Engine/Plugins/Animation/ControlRig/Source/ControlRig/Private/Units/Simulation/RigUnit_Verlet.h
@@ -7,9 +7,10 @@
 #include "RigUnit_Verlet.generated.h"
 
 /**
- * Simulates a single position over time using verlet integration
+ * Simulates a single position over time using Verlet integration. It is recommended to use SpringInterp instead as it
+ * is more accurate and stable, and has more meaningful parameters.
  */
-USTRUCT(meta=(DisplayName="Verlet (Vector)", PrototypeName="Verlet", Keywords="Simulate,Integrate"))
+USTRUCT(meta=(DisplayName="Verlet (Vector)", PrototypeName="Verlet", Keywords="Simulate,Integrate", Deprecated = "5.0"))
 struct CONTROLRIG_API FRigUnit_VerletIntegrateVector : public FRigUnit_SimBase
 {
 	GENERATED_BODY()
diff --git a/Engine/Source/Runtime/Engine/Classes/Kismet/KismetMathLibrary.h b/Engine/Source/Runtime/Engine/Classes/Kismet/KismetMathLibrary.h
--- a/Engine/Source/Runtime/Engine/Classes/Kismet/KismetMathLibrary.h
+++ b/Engine/Source/Runtime/Engine/Classes/Kismet/KismetMathLibrary.h
@@ -159,12 +159,12 @@ struct ENGINE_API FQuaternionSpringState
 	GENERATED_BODY()
 
 	FQuat   PrevTarget;
-	FVector Velocity; // Angular velocity
+	FVector AngularVelocity;
 	bool    bPrevTargetValid;
 
 	FQuaternionSpringState()
 	: PrevTarget(FQuat::Identity)
-	, Velocity(FVector::ZeroVector)
+	, AngularVelocity(FVector::ZeroVector)
 	, bPrevTargetValid(false)
 	{
 	}
@@ -172,7 +172,7 @@ struct ENGINE_API FQuaternionSpringState
 	void Reset()
 	{
 		PrevTarget = FQuat::Identity;
-		Velocity = FVector::ZeroVector;
+		AngularVelocity = FVector::ZeroVector;
 		bPrevTargetValid = false;
 	}
 };
@@ -3931,7 +3931,7 @@ class ENGINE_API UKismetMathLibrary : public UBlueprintFunctionLibrary
 
 	/** Sets the state angular velocity of a quaternion spring */
 	UFUNCTION(BlueprintCallable, Category = "Math|Interpolation")
-	static void SetQuaternionSpringStateVelocity(UPARAM(ref) FQuaternionSpringState& SpringState, FVector Velocity);
+	static void SetQuaternionSpringStateAngularVelocity(UPARAM(ref) FQuaternionSpringState& SpringState, FVector AngularVelocity);
 
 	//
 	// Random stream functions
diff --git a/Engine/Source/Runtime/Engine/Private/KismetMathLibrary.cpp b/Engine/Source/Runtime/Engine/Private/KismetMathLibrary.cpp
--- a/Engine/Source/Runtime/Engine/Private/KismetMathLibrary.cpp
+++ b/Engine/Source/Runtime/Engine/Private/KismetMathLibrary.cpp
@@ -451,7 +451,7 @@ FQuat UKismetMathLibrary::QuaternionSpringInterp(FQuat Current, FQuat Target, UP
 	{
 		Target = -Target;
 	}
-	FQuat VelocityQuat(SpringState.Velocity.X, SpringState.Velocity.Y, SpringState.Velocity.Z, 0.0f);
+	FQuat VelocityQuat(SpringState.AngularVelocity.X, SpringState.AngularVelocity.Y, SpringState.AngularVelocity.Z, 0.0f);
 	FQuat Velocity = 0.5f * Current * VelocityQuat;
 
   	GenericSpringInterp(Current, Target, SpringState.PrevTarget, SpringState.bPrevTargetValid, Velocity, Stiffness,
@@ -459,7 +459,7 @@ FQuat UKismetMathLibrary::QuaternionSpringInterp(FQuat Current, FQuat Target, UP
 	Current.Normalize();
 
 	VelocityQuat = Current.Inverse() * 2.0f * Velocity;
-	SpringState.Velocity.Set(VelocityQuat.X, VelocityQuat.Y, VelocityQuat.Z);
+	SpringState.AngularVelocity.Set(VelocityQuat.X, VelocityQuat.Y, VelocityQuat.Z);
 
 	return Current;
 }
@@ -490,9 +490,9 @@ void UKismetMathLibrary::SetVectorSpringStateVelocity(FVectorSpringState& Spring
 	SpringState.Velocity = Velocity;
 }
 
-void UKismetMathLibrary::SetQuaternionSpringStateVelocity(FQuaternionSpringState& SpringState, FVector Velocity)
+void UKismetMathLibrary::SetQuaternionSpringStateAngularVelocity(FQuaternionSpringState& SpringState, FVector AngularVelocity)
 {
-	SpringState.Velocity = Velocity;
+	SpringState.AngularVelocity = AngularVelocity;
 }
 
 FVector UKismetMathLibrary::RandomUnitVector()
