Undeprecate Verlet control rig node, and add Verlet to the list of SpringInterp keywords

diff --git a/Engine/Plugins/Animation/ControlRig/Source/ControlRig/Private/Units/Simulation/RigUnit_SpringInterp.h b/Engine/Plugins/Animation/ControlRig/Source/ControlRig/Private/Units/Simulation/RigUnit_SpringInterp.h
--- a/Engine/Plugins/Animation/ControlRig/Source/ControlRig/Private/Units/Simulation/RigUnit_SpringInterp.h
+++ b/Engine/Plugins/Animation/ControlRig/Source/ControlRig/Private/Units/Simulation/RigUnit_SpringInterp.h
@@ -9,7 +9,7 @@
 /**
  * Uses a simple spring model to interpolate a float from Current to Target.
  */
-USTRUCT(meta=(DisplayName="Spring Interpolate", Keywords="Alpha,SpringInterpolate", PrototypeName = "SpringInterp"))
+USTRUCT(meta=(DisplayName="Spring Interpolate", Keywords="Alpha,SpringInterpolate,Verlet", PrototypeName = "SpringInterp"))
 struct CONTROLRIG_API FRigUnit_SpringInterp : public FRigUnit_SimBase
 {
 	GENERATED_BODY()
@@ -96,7 +96,7 @@ struct CONTROLRIG_API FRigUnit_SpringInterp : public FRigUnit_SimBase
 /**
  * Uses a simple spring model to interpolate a vector from Current to Target.
  */
-USTRUCT(meta=(DisplayName="Spring Interpolate", Keywords="Alpha,SpringInterpolate", PrototypeName = "SpringInterp", MenuDescSuffix = "(Vector)"))
+USTRUCT(meta=(DisplayName="Spring Interpolate", Keywords="Alpha,SpringInterpolate,Verlet", PrototypeName = "SpringInterp", MenuDescSuffix = "(Vector)"))
 struct CONTROLRIG_API FRigUnit_SpringInterpVector : public FRigUnit_SimBase
 {
 	GENERATED_BODY()
@@ -183,7 +183,7 @@ struct CONTROLRIG_API FRigUnit_SpringInterpVector : public FRigUnit_SimBase
 /**
  * Uses a simple spring model to interpolate a quaternion from Current to Target.
  */
-USTRUCT(meta=(DisplayName="Spring Interpolate", Keywords="Alpha,SpringInterpolate", PrototypeName = "SpringInterp", MenuDescSuffix = "(Quaternion)"))
+USTRUCT(meta=(DisplayName="Spring Interpolate", Keywords="Alpha,SpringInterpolate,Verlet", PrototypeName = "SpringInterp", MenuDescSuffix = "(Quaternion)"))
 struct CONTROLRIG_API FRigUnit_SpringInterpQuaternion : public FRigUnit_SimBase
 {
 	GENERATED_BODY()
diff --git a/Engine/Plugins/Animation/ControlRig/Source/ControlRig/Private/Units/Simulation/RigUnit_Verlet.h b/Engine/Plugins/Animation/ControlRig/Source/ControlRig/Private/Units/Simulation/RigUnit_Verlet.h
--- a/Engine/Plugins/Animation/ControlRig/Source/ControlRig/Private/Units/Simulation/RigUnit_Verlet.h
+++ b/Engine/Plugins/Animation/ControlRig/Source/ControlRig/Private/Units/Simulation/RigUnit_Verlet.h
@@ -10,7 +10,7 @@
  * Simulates a single position over time using Verlet integration. It is recommended to use SpringInterp instead as it
  * is more accurate and stable, and has more meaningful parameters.
  */
-USTRUCT(meta=(DisplayName="Verlet (Vector)", PrototypeName="Verlet", Keywords="Simulate,Integrate", Deprecated = "5.0"))
+USTRUCT(meta=(DisplayName="Verlet (Vector)", PrototypeName="Verlet", Keywords="Simulate,Integrate"))
 struct CONTROLRIG_API FRigUnit_VerletIntegrateVector : public FRigUnit_SimBase
 {
 	GENERATED_BODY()
