Control Rig: UObject Shim Layer for Details Panel

diff --git a/Engine/Plugins/Animation/ControlRig/Source/ControlRig/Private/Rigs/RigHierarchyElements.cpp b/Engine/Plugins/Animation/ControlRig/Source/ControlRig/Private/Rigs/RigHierarchyElements.cpp
--- a/Engine/Plugins/Animation/ControlRig/Source/ControlRig/Private/Rigs/RigHierarchyElements.cpp
+++ b/Engine/Plugins/Animation/ControlRig/Source/ControlRig/Private/Rigs/RigHierarchyElements.cpp
@@ -9,6 +9,42 @@
 // FRigBaseElement
 ////////////////////////////////////////////////////////////////////////////////
 
+UScriptStruct* FRigBaseElement::GetElementStruct() const
+{
+	switch(GetType())
+	{
+		case ERigElementType::Bone:
+		{
+			return FRigBoneElement::StaticStruct();
+		}
+		case ERigElementType::Null:
+		{
+			return FRigNullElement::StaticStruct();
+		}
+		case ERigElementType::Control:
+		{
+			return FRigControlElement::StaticStruct();
+		}
+		case ERigElementType::Curve:
+		{
+			return FRigCurveElement::StaticStruct();
+		}
+		case ERigElementType::Socket:
+		{
+			return FRigSocketElement::StaticStruct();
+		}
+		case ERigElementType::RigidBody:
+		{
+			return FRigRigidBodyElement::StaticStruct();
+		}
+		default:
+		{
+				break;
+		}
+	}
+	return FRigBaseElement::StaticStruct();
+}
+
 void FRigBaseElement::Serialize(FArchive& Ar, URigHierarchy* Hierarchy, ESerializationPhase SerializationPhase)
 {
 	Ar.UsingCustomVersion(FControlRigObjectVersion::GUID);
diff --git a/Engine/Plugins/Animation/ControlRig/Source/ControlRig/Public/Rigs/RigHierarchyElements.h b/Engine/Plugins/Animation/ControlRig/Source/ControlRig/Public/Rigs/RigHierarchyElements.h
--- a/Engine/Plugins/Animation/ControlRig/Source/ControlRig/Public/Rigs/RigHierarchyElements.h
+++ b/Engine/Plugins/Animation/ControlRig/Source/ControlRig/Public/Rigs/RigHierarchyElements.h
@@ -22,6 +22,16 @@ template<typename T> \
 friend FORCEINLINE T* Cast(ElementType* InElement) \
 { \
    return Cast<T>((FRigBaseElement*) InElement); \
+} \
+template<typename T> \
+friend FORCEINLINE const T* CastChecked(const ElementType* InElement) \
+{ \
+	return CastChecked<T>((const FRigBaseElement*) InElement); \
+} \
+template<typename T> \
+friend FORCEINLINE T* CastChecked(ElementType* InElement) \
+{ \
+	return CastChecked<T>((FRigBaseElement*) InElement); \
 }
 
 UENUM()
@@ -430,6 +440,7 @@ protected:
 
 public:
 
+	UScriptStruct* GetElementStruct() const;
 	void Serialize(FArchive& Ar, URigHierarchy* Hierarchy, ESerializationPhase SerializationPhase);
 	virtual void Save(FArchive& Ar, URigHierarchy* Hierarchy, ESerializationPhase SerializationPhase);
 	virtual void Load(FArchive& Ar, URigHierarchy* Hierarchy, ESerializationPhase SerializationPhase);
diff --git a/Engine/Plugins/Animation/ControlRig/Source/ControlRigEditor/Private/ControlRigEditorModule.cpp b/Engine/Plugins/Animation/ControlRig/Source/ControlRigEditor/Private/ControlRigEditorModule.cpp
--- a/Engine/Plugins/Animation/ControlRig/Source/ControlRigEditor/Private/ControlRigEditorModule.cpp
+++ b/Engine/Plugins/Animation/ControlRig/Source/ControlRigEditor/Private/ControlRigEditorModule.cpp
@@ -199,18 +199,6 @@ void FControlRigEditorModule::StartupModule()
 	ClassesToUnregisterOnShutdown.Add(UMovieSceneControlRigParameterSection::StaticClass()->GetFName());
 	PropertyEditorModule.RegisterCustomClassLayout(ClassesToUnregisterOnShutdown.Last(), FOnGetDetailCustomizationInstance::CreateStatic(&FMovieSceneControlRigSectionDetailsCustomization::MakeInstance));
 
-	ClassesToUnregisterOnShutdown.Add(FRigBoneElement::StaticStruct()->GetFName());
-	PropertyEditorModule.RegisterCustomClassLayout(ClassesToUnregisterOnShutdown.Last(), FOnGetDetailCustomizationInstance::CreateStatic(&FRigBoneElementDetails::MakeInstance));
-	
-	ClassesToUnregisterOnShutdown.Add(FRigControlElement::StaticStruct()->GetFName());
-	PropertyEditorModule.RegisterCustomClassLayout(ClassesToUnregisterOnShutdown.Last(), FOnGetDetailCustomizationInstance::CreateStatic(&FRigControlElementDetails::MakeInstance));
-
-	ClassesToUnregisterOnShutdown.Add(FRigNullElement::StaticStruct()->GetFName());
-	PropertyEditorModule.RegisterCustomClassLayout(ClassesToUnregisterOnShutdown.Last(), FOnGetDetailCustomizationInstance::CreateStatic(&FRigNullElementDetails::MakeInstance));
-
-	ClassesToUnregisterOnShutdown.Add(FRigInfluenceMapPerEvent::StaticStruct()->GetFName());
-	PropertyEditorModule.RegisterCustomClassLayout(ClassesToUnregisterOnShutdown.Last(), FOnGetDetailCustomizationInstance::CreateStatic(&FRigInfluenceMapPerEventDetails::MakeInstance));
-
 	ClassesToUnregisterOnShutdown.Add(UControlRig::StaticClass()->GetFName());
 
 	// same as ClassesToUnregisterOnShutdown but for properties, there is none right now
@@ -231,6 +219,15 @@ void FControlRigEditorModule::StartupModule()
 	PropertiesToUnregisterOnShutdown.Add(FRigComputedTransform::StaticStruct()->GetFName());
 	PropertyEditorModule.RegisterCustomPropertyTypeLayout(PropertiesToUnregisterOnShutdown.Last(), FOnGetPropertyTypeCustomizationInstance::CreateStatic(&FRigComputedTransformDetails::MakeInstance));
 
+	PropertiesToUnregisterOnShutdown.Add(FRigBoneElement::StaticStruct()->GetFName());
+	PropertyEditorModule.RegisterCustomPropertyTypeLayout(PropertiesToUnregisterOnShutdown.Last(), FOnGetPropertyTypeCustomizationInstance::CreateStatic(&FRigBoneElementDetails::MakeInstance));
+
+	PropertiesToUnregisterOnShutdown.Add(FRigControlElement::StaticStruct()->GetFName());
+	PropertyEditorModule.RegisterCustomPropertyTypeLayout(PropertiesToUnregisterOnShutdown.Last(), FOnGetPropertyTypeCustomizationInstance::CreateStatic(&FRigControlElementDetails::MakeInstance));
+
+	PropertiesToUnregisterOnShutdown.Add(FRigNullElement::StaticStruct()->GetFName());
+	PropertyEditorModule.RegisterCustomPropertyTypeLayout(PropertiesToUnregisterOnShutdown.Last(), FOnGetPropertyTypeCustomizationInstance::CreateStatic(&FRigNullElementDetails::MakeInstance));
+
 	// Register asset tools
 	auto RegisterAssetTypeAction = [this](const TSharedRef<IAssetTypeActions>& InAssetTypeAction)
 	{
diff --git a/Engine/Plugins/Animation/ControlRig/Source/ControlRigEditor/Private/ControlRigElementDetails.cpp b/Engine/Plugins/Animation/ControlRig/Source/ControlRigEditor/Private/ControlRigElementDetails.cpp
--- a/Engine/Plugins/Animation/ControlRig/Source/ControlRigEditor/Private/ControlRigElementDetails.cpp
+++ b/Engine/Plugins/Animation/ControlRig/Source/ControlRigEditor/Private/ControlRigElementDetails.cpp
@@ -16,6 +16,7 @@
 #include "Widgets/Input/SButton.h"
 #include "Editor/SControlRigGizmoNameList.h"
 #include "ControlRigBlueprint.h"
+#include "IDetailGroup.h"
 #include "Graph/ControlRigGraph.h"
 #include "PropertyCustomizationHelpers.h"
 #include "SEnumCombo.h"
@@ -26,6 +27,8 @@
 
 #define LOCTEXT_NAMESPACE "ControlRigElementDetails"
 
+static const FText ControlRigDetailsMultipleValues = LOCTEXT("MultipleValues", "Multiple Values");
+
 namespace FRigElementKeyDetailsDefs
 {
 	// Active foreground pin alpha
@@ -263,9 +266,20 @@ FString FRigElementKeyDetails::GetElementName() const
 	FString ElementNameStr;
 	if (NameHandle.IsValid())
 	{
-		FName ElementName;
-		NameHandle->GetValue(ElementName);
-		ElementNameStr = ElementName.ToString();
+		for(int32 ObjectIndex = 0; ObjectIndex < NameHandle->GetNumPerObjectValues(); ObjectIndex++)
+		{
+			FString PerObjectValue;
+			NameHandle->GetPerObjectValue(ObjectIndex, PerObjectValue);
+
+			if(ObjectIndex == 0)
+			{
+				ElementNameStr = PerObjectValue;
+			}
+			else if(ElementNameStr != PerObjectValue)
+			{
+				return ControlRigDetailsMultipleValues.ToString();
+			}
+		}
 	}
 	return ElementNameStr;
 }
@@ -703,70 +717,70 @@ void FRigComputedTransformDetails::OnTransformChanged(FEditPropertyChain* InProp
 	}
 }
 
-void FRigBaseElementDetails::CustomizeDetails(IDetailLayoutBuilder& DetailBuilder)
+void FRigBaseElementDetails::CustomizeHeader(TSharedRef<IPropertyHandle> InStructPropertyHandle, FDetailWidgetRow& HeaderRow, IPropertyTypeCustomizationUtils& StructCustomizationUtils)
 {
-	DetailBuilder.HideProperty(TEXT("Key"), FRigBaseElement::StaticStruct());
-	DetailBuilder.HideProperty(TEXT("Index"), FRigBaseElement::StaticStruct());
-	DetailBuilder.HideProperty(TEXT("SubIndex"), FRigBaseElement::StaticStruct());
+	BlueprintBeingCustomized = nullptr;
+	HierarchyBeingCustomized = nullptr;
+	ObjectsBeingCustomized.Reset();
+	
+	RigElementDetails_GetCustomizedInfo(InStructPropertyHandle, BlueprintBeingCustomized);
 
-	TArray<TSharedPtr<FStructOnScope>> StructsBeingCustomized;
-	DetailBuilder.GetStructsBeingCustomized(StructsBeingCustomized);
-	for (TSharedPtr<FStructOnScope> StructBeingCustomized : StructsBeingCustomized)
+	TArray<UObject*> Objects;
+	InStructPropertyHandle->GetOuterObjects(Objects);
+	for (UObject* Object : Objects)
 	{
-		if (UPackage* Package = StructBeingCustomized->GetPackage())
-		{
-			TArray<UObject*> SubObjects;
-			Package->GetDefaultSubobjects(SubObjects);
-
-			for (UObject* SubObject : SubObjects)
-			{
-				if (UControlRig* Rig = Cast<UControlRig>(SubObject))
-				{
-					BlueprintBeingCustomized = Cast<UControlRigBlueprint>(Rig->GetClass()->ClassGeneratedBy);
-					if(BlueprintBeingCustomized)
-					{
-						HierarchyBeingCustomized = BlueprintBeingCustomized->Hierarchy;
-						if (UControlRig* DebuggedControlRig = Cast<UControlRig>(BlueprintBeingCustomized->GetObjectBeingDebugged()))
-						{
-							if (!DebuggedControlRig->IsSetupModeEnabled())
-							{
-								HierarchyBeingCustomized = DebuggedControlRig->GetHierarchy();
-							}
-						}
-						break;
-					}
-				}
-			}
+		UDetailsViewWrapperObject* WrapperObject = CastChecked<UDetailsViewWrapperObject>(Object);
 
-			if (HierarchyBeingCustomized)
-			{
-				ElementKeyBeingCustomized = ((const FRigBaseElement*)StructBeingCustomized->GetStructMemory())->GetKey();
-				break;
-			}
+		if(HierarchyBeingCustomized == nullptr)
+		{
+			HierarchyBeingCustomized = Cast<URigHierarchy>(WrapperObject->GetOuter());
 		}
+
+		ObjectsBeingCustomized.Add(WrapperObject);
 	}
+}
+
+void FRigBaseElementDetails::CustomizeChildren(TSharedRef<class IPropertyHandle> InStructPropertyHandle, class IDetailChildrenBuilder& StructBuilder, IPropertyTypeCustomizationUtils& StructCustomizationUtils)
+{
+	StructBuilder.AddCustomRow(FText::FromString(TEXT("Name")))
+	.NameContent()
+	[
+		SNew(STextBlock)
+		.Text(FText::FromString(TEXT("Name")))
+		.Font(IDetailLayoutBuilder::GetDetailFont())
+	]
+	.ValueContent()
+	[
+		SNew(SEditableTextBox)
+		.Font(IDetailLayoutBuilder::GetDetailFont())
+		.Text(this, &FRigBaseElementDetails::GetName)
+		.OnTextCommitted(this, &FRigBaseElementDetails::SetName)
+		.IsEnabled(ObjectsBeingCustomized.Num() == 1)
+	];
+}
 
-	IDetailCategoryBuilder& Category = DetailBuilder.EditCategory(TEXT("RigElement"));
-	Category.InitiallyCollapsed(false);
+FRigElementKey FRigBaseElementDetails::GetElementKey() const
+{
+	check(ObjectsBeingCustomized.Num() == 1);
+	return ObjectsBeingCustomized[0]->GetContent<FRigBaseElement>()->GetKey(); 
+}
 
-	Category.AddCustomRow(FText::FromString(TEXT("Name")))
-    .NameContent()
-    [
-        SNew(STextBlock)
-        .Text(FText::FromString(TEXT("Name")))
-        .Font(IDetailLayoutBuilder::GetDetailFont())
-    ]
-    .ValueContent()
-    [
-        SNew(SEditableTextBox)
-        .Font(IDetailLayoutBuilder::GetDetailFont())
-        .Text(this, &FRigBaseElementDetails::GetName)
-        .OnTextCommitted(this, &FRigBaseElementDetails::SetName)
-    ];
+FText FRigBaseElementDetails::GetName() const
+{
+	if(ObjectsBeingCustomized.Num() > 1)
+	{
+		return ControlRigDetailsMultipleValues;
+	}
+	return FText::FromName(GetElementKey().Name);
 }
 
 void FRigBaseElementDetails::SetName(const FText& InNewText, ETextCommit::Type InCommitType)
 {
+	if(ObjectsBeingCustomized.Num() > 1)
+	{
+		return;
+	}
+	
 	URigHierarchy* Hierarchy = nullptr;
 	if (BlueprintBeingCustomized)
 	{
@@ -781,7 +795,7 @@ void FRigBaseElementDetails::SetName(const FText& InNewText, ETextCommit::Type I
 	{
 		URigHierarchyController* Controller = Hierarchy->GetController(true);
 		check(Controller);
-		Controller->RenameElement(ElementKeyBeingCustomized, *InNewText.ToString(), true, true);
+		Controller->RenameElement(GetElementKey(), *InNewText.ToString(), true, true);
 	}
 }
 
@@ -803,115 +817,109 @@ bool FRigBaseElementDetails::IsSetupModeEnabled() const
 	return false;
 }
 
-void FRigTransformElementDetails::CustomizeDetails(IDetailLayoutBuilder& DetailBuilder)
+TArray<FRigElementKey> FRigBaseElementDetails::GetElementKeys() const
 {
-	FRigBaseElementDetails::CustomizeDetails(DetailBuilder);
+	TArray<FRigElementKey> Keys;
+	for(UDetailsViewWrapperObject* ObjectBeingCustomized : ObjectsBeingCustomized)
+	{
+		Keys.Add(ObjectBeingCustomized->GetContent<FRigBaseElement>()->GetKey());
+	}
+	return Keys;
+}
 
-	//if(ElementKeyBeingCustomized.Type != ERigElementType::Control)
+bool FRigBaseElementDetails::IsAnyControlOfType(ERigControlType InType) const
+{
+	for(UDetailsViewWrapperObject* ObjectBeingCustomized : ObjectsBeingCustomized)
 	{
-		IDetailCategoryBuilder& Category = DetailBuilder.EditCategory(TEXT("Pose"));
-		if(ElementKeyBeingCustomized.Type == ERigElementType::Control)
+		if(FRigControlElement* ControlElement = Cast<FRigControlElement>(ObjectBeingCustomized->GetContent<FRigBaseElement>()))
 		{
-			Category
-			.InitiallyCollapsed(true);
+			if(ControlElement->Settings.ControlType == InType)
+			{
+				return true;
+			}
 		}
+	}
+	return false;
+}
 
-		// setup initial global
+bool FRigBaseElementDetails::IsAnyControlNotOfType(ERigControlType InType) const
+{
+	for(UDetailsViewWrapperObject* ObjectBeingCustomized : ObjectsBeingCustomized)
+	{
+		if(FRigControlElement* ControlElement = Cast<FRigControlElement>(ObjectBeingCustomized->GetContent<FRigBaseElement>()))
 		{
-			const TSharedRef<IPropertyHandle> PropertyHandle = DetailBuilder.GetProperty(TEXT("Pose.Initial.Global.Transform"), FRigTransformElement::StaticStruct());
-			Category.AddProperty(PropertyHandle, EPropertyLocation::Advanced).DisplayName(FText::FromString(TEXT("Initial Global")))
-			.IsEnabled(TAttribute<bool>::Create(TAttribute<bool>::FGetter::CreateSP(this, &FRigBaseElementDetails::IsSetupModeEnabled)));
+			if(ControlElement->Settings.ControlType != InType)
+			{
+				return true;
+			}
+		}
+	}
+	return false;
+}
 
-			PoseInitialGlobal.AddHead(FRigControlElement::StaticStruct()->FindPropertyByName(TEXT("Pose")));
-			PoseInitialGlobal.AddTail(FRigCurrentAndInitialTransform::StaticStruct()->FindPropertyByName(TEXT("Initial")));
-			PoseInitialGlobal.AddTail(FRigLocalAndGlobalTransform::StaticStruct()->FindPropertyByName(TEXT("Global")));
-			PoseInitialGlobal.AddTail(FRigComputedTransform::StaticStruct()->FindPropertyByName(TEXT("Transform")));
-			PoseInitialGlobal.SetActiveMemberPropertyNode(PoseInitialGlobal.GetTail()->GetValue());
+void FRigTransformElementDetails::CustomizeChildren(TSharedRef<class IPropertyHandle> InStructPropertyHandle, class IDetailChildrenBuilder& StructBuilder, IPropertyTypeCustomizationUtils& StructCustomizationUtils)
+{
+	FRigBaseElementDetails::CustomizeChildren(InStructPropertyHandle, StructBuilder, StructCustomizationUtils);
 
-			const FSimpleDelegate OnTransformChangedDelegate = FSimpleDelegate::CreateStatic(&FRigTransformElementDetails::OnTransformChanged, &PoseInitialGlobal, BlueprintBeingCustomized);
-			PropertyHandle->SetOnPropertyValueChanged(OnTransformChangedDelegate);
-			PropertyHandle->SetOnChildPropertyValueChanged(OnTransformChangedDelegate);
-			PropertyHandle->SetToolTipText(FText::FromString(TEXT("The initial / ref pose global transform in the space of the rig / actor.")));
-		}
+	IDetailGroup& PoseGroupDefault = StructBuilder.AddGroup(TEXT("Pose"), LOCTEXT("Pose", "Pose")); 
+	IDetailGroup& PoseGroupAdvanced = StructBuilder.GetParentCategory().AddGroup(TEXT("Pose"), LOCTEXT("Pose", "Pose"), true, false); 
+	
+	{
+		const TSharedPtr<IPropertyHandle> PoseHandle = InStructPropertyHandle->GetChildHandle(TEXT("Pose"));
+		const TSharedPtr<IPropertyHandle> InitialHandle = PoseHandle->GetChildHandle(TEXT("Initial"));
+		const TSharedPtr<IPropertyHandle> CurrentHandle = PoseHandle->GetChildHandle(TEXT("Current"));
 
-		// setup initial local
+		// setup initial global
 		{
-			const TSharedRef<IPropertyHandle> PropertyHandle = DetailBuilder.GetProperty(TEXT("Pose.Initial.Local.Transform"), FRigTransformElement::StaticStruct());
-			Category.AddProperty(PropertyHandle, EPropertyLocation::Advanced).DisplayName(FText::FromString(TEXT("Initial Local")))
+			const TSharedPtr<IPropertyHandle> GlobalHandle = InitialHandle->GetChildHandle(TEXT("Global"));
+			const TSharedPtr<IPropertyHandle> TransformHandle = GlobalHandle->GetChildHandle(TEXT("Transform"));
+
+			PoseGroupAdvanced.AddPropertyRow(TransformHandle.ToSharedRef()).DisplayName(FText::FromString(TEXT("Initial Global")))
 			.IsEnabled(TAttribute<bool>::Create(TAttribute<bool>::FGetter::CreateSP(this, &FRigBaseElementDetails::IsSetupModeEnabled)));
+		}
 
-			PoseInitialLocal.AddHead(FRigControlElement::StaticStruct()->FindPropertyByName(TEXT("Pose")));
-			PoseInitialLocal.AddTail(FRigCurrentAndInitialTransform::StaticStruct()->FindPropertyByName(TEXT("Initial")));
-			PoseInitialLocal.AddTail(FRigLocalAndGlobalTransform::StaticStruct()->FindPropertyByName(TEXT("Local")));
-			PoseInitialLocal.AddTail(FRigComputedTransform::StaticStruct()->FindPropertyByName(TEXT("Transform")));
-			PoseInitialLocal.SetActiveMemberPropertyNode(PoseInitialLocal.GetTail()->GetValue());
+		// setup initial local
+		{
+			const TSharedPtr<IPropertyHandle> LocalHandle = InitialHandle->GetChildHandle(TEXT("Local"));
+			const TSharedPtr<IPropertyHandle> TransformHandle = LocalHandle->GetChildHandle(TEXT("Transform"));
 
-			const FSimpleDelegate OnTransformChangedDelegate = FSimpleDelegate::CreateStatic(&FRigTransformElementDetails::OnTransformChanged, &PoseInitialLocal, BlueprintBeingCustomized);
-			PropertyHandle->SetOnPropertyValueChanged(OnTransformChangedDelegate);
-			PropertyHandle->SetOnChildPropertyValueChanged(OnTransformChangedDelegate);
-			PropertyHandle->SetToolTipText(FText::FromString(TEXT("The initial / ref pose local transform in the space of the parent.\nFor Controls the initial value is relative to the offset.")));
+			PoseGroupAdvanced.AddPropertyRow(TransformHandle.ToSharedRef()).DisplayName(FText::FromString(TEXT("Initial Local")))
+			.IsEnabled(TAttribute<bool>::Create(TAttribute<bool>::FGetter::CreateSP(this, &FRigBaseElementDetails::IsSetupModeEnabled)));
 		}
 
 		// setup current global
 		{
-			const TSharedRef<IPropertyHandle> PropertyHandle = DetailBuilder.GetProperty(TEXT("Pose.Current.Global.Transform"), FRigTransformElement::StaticStruct());
-			Category.AddProperty(PropertyHandle, EPropertyLocation::Advanced).DisplayName(FText::FromString(TEXT("Current Global")))
-			.IsEnabled(TAttribute<bool>::Create(TAttribute<bool>::FGetter::CreateSP(this, &FRigBaseElementDetails::IsSetupModeEnabled)));
-
-			PoseCurrentGlobal.AddHead(FRigControlElement::StaticStruct()->FindPropertyByName(TEXT("Pose")));
-			PoseCurrentGlobal.AddTail(FRigCurrentAndInitialTransform::StaticStruct()->FindPropertyByName(TEXT("Current")));
-			PoseCurrentGlobal.AddTail(FRigLocalAndGlobalTransform::StaticStruct()->FindPropertyByName(TEXT("Global")));
-			PoseCurrentGlobal.AddTail(FRigComputedTransform::StaticStruct()->FindPropertyByName(TEXT("Transform")));
-			PoseCurrentGlobal.SetActiveMemberPropertyNode(PoseCurrentGlobal.GetTail()->GetValue());
+			const TSharedPtr<IPropertyHandle> GlobalHandle = CurrentHandle->GetChildHandle(TEXT("Global"));
+			const TSharedPtr<IPropertyHandle> TransformHandle = GlobalHandle->GetChildHandle(TEXT("Transform"));
 
-			const FSimpleDelegate OnTransformChangedDelegate = FSimpleDelegate::CreateStatic(&FRigTransformElementDetails::OnTransformChanged, &PoseCurrentGlobal, BlueprintBeingCustomized);
-			PropertyHandle->SetOnPropertyValueChanged(OnTransformChangedDelegate);
-			PropertyHandle->SetOnChildPropertyValueChanged(OnTransformChangedDelegate);
-			PropertyHandle->SetToolTipText(FText::FromString(TEXT("The current global transform in the space of the rig / actor.")));
+			PoseGroupAdvanced.AddPropertyRow(TransformHandle.ToSharedRef()).DisplayName(FText::FromString(TEXT("Current Global")))
+			.IsEnabled(TAttribute<bool>::Create(TAttribute<bool>::FGetter::CreateSP(this, &FRigBaseElementDetails::IsSetupModeEnabled)));
 		}
 
 		// setup current local
 		{
-			const TSharedRef<IPropertyHandle> PropertyHandle = DetailBuilder.GetProperty(TEXT("Pose.Current.Local.Transform"), FRigTransformElement::StaticStruct());
-			Category.AddProperty(PropertyHandle).DisplayName(FText::FromString(TEXT("Current Local")))
-			.IsEnabled(ElementKeyBeingCustomized.Type != ERigElementType::Control);
-
-			PoseCurrentLocal.AddHead(FRigControlElement::StaticStruct()->FindPropertyByName(TEXT("Pose")));
-			PoseCurrentLocal.AddTail(FRigCurrentAndInitialTransform::StaticStruct()->FindPropertyByName(TEXT("Current")));
-			PoseCurrentLocal.AddTail(FRigLocalAndGlobalTransform::StaticStruct()->FindPropertyByName(TEXT("Local")));
-			PoseCurrentLocal.AddTail(FRigComputedTransform::StaticStruct()->FindPropertyByName(TEXT("Transform")));
-			PoseCurrentLocal.SetActiveMemberPropertyNode(PoseCurrentLocal.GetTail()->GetValue());
-
-			const FSimpleDelegate OnTransformChangedDelegate = FSimpleDelegate::CreateStatic(&FRigTransformElementDetails::OnTransformChanged, &PoseCurrentLocal, BlueprintBeingCustomized);
-			PropertyHandle->SetOnPropertyValueChanged(OnTransformChangedDelegate);
-			PropertyHandle->SetOnChildPropertyValueChanged(OnTransformChangedDelegate);
-			PropertyHandle->SetToolTipText(FText::FromString(TEXT("The current local transform in the space of the parent.\nFor Controls the initial value is relative to the offset.")));
-		}
+			const TSharedPtr<IPropertyHandle> LocalHandle = CurrentHandle->GetChildHandle(TEXT("Local"));
+			const TSharedPtr<IPropertyHandle> TransformHandle = LocalHandle->GetChildHandle(TEXT("Transform"));
 
-		DetailBuilder.HideProperty(TEXT("Pose"), FRigTransformElement::StaticStruct());
+			PoseGroupDefault.AddPropertyRow(TransformHandle.ToSharedRef()).DisplayName(FText::FromString(TEXT("Current Local")))
+			.IsEnabled(TAttribute<bool>::Create(TAttribute<bool>::FGetter::CreateSP(this, &FRigTransformElementDetails::IsCurrentLocalEnabled)));
+		}
 	}
 }
 
-void FRigTransformElementDetails::OnTransformChanged(FEditPropertyChain* InPropertyChain, UControlRigBlueprint* InBlueprint)
+bool FRigTransformElementDetails::IsCurrentLocalEnabled() const
 {
-	if(InBlueprint && InPropertyChain)
+	for(UDetailsViewWrapperObject* WrapperObject : ObjectsBeingCustomized)
 	{
-		if(InPropertyChain->Num() > 1)
+		if(ObjectsBeingCustomized[0]->GetContent<FRigBaseElement>()->GetType() == ERigElementType::Control)
 		{
-			FPropertyChangedEvent ChangeEvent(InPropertyChain->GetHead()->GetValue(), EPropertyChangeType::ValueSet);
-			ChangeEvent.SetActiveMemberProperty(InPropertyChain->GetTail()->GetValue());
-			FPropertyChangedChainEvent ChainEvent(*InPropertyChain, ChangeEvent);
-			InBlueprint->BroadcastPostEditChangeChainProperty(ChainEvent);
+			return false;
 		}
 	}
+	return true;
 }
 
-void FRigBoneElementDetails::CustomizeDetails( IDetailLayoutBuilder& DetailBuilder )
-{
-	FRigTransformElementDetails::CustomizeDetails(DetailBuilder);
-}
-
-void FRigControlElementDetails_SetupBoolValueWidget(IDetailCategoryBuilder& InCategory, ERigControlValueType InValueType, FRigControlElement* InControlElement, URigHierarchy* InHierarchy)
+void FRigControlElementDetails_SetupBoolValueWidget(IDetailGroup& InGroup, IDetailChildrenBuilder& InStructBuilder, ERigControlValueType InValueType, FRigControlElement* InControlElement, URigHierarchy* InHierarchy)
 {
 	UEnum* ControlTypeEnum = StaticEnum<ERigControlType>();
 	UEnum* ValueTypeEnum = StaticEnum<ERigControlValueType>();
@@ -921,7 +929,7 @@ void FRigControlElementDetails_SetupBoolValueWidget(IDetailCategoryBuilder& InCa
 	TWeakObjectPtr<URigHierarchy> HierarchyPtr = InHierarchy;
 	const FRigElementKey Key = InControlElement->GetKey();
 
-	InCategory.AddCustomRow(PropertyLabel)
+	InStructBuilder.AddCustomRow(PropertyLabel)
 	.NameContent()
 	[
 	    SNew(STextBlock)
@@ -979,7 +987,7 @@ void FRigControlElementDetails_SetupBoolValueWidget(IDetailCategoryBuilder& InCa
     })));
 }
 
-void FRigControlElementDetails_SetupIntegerValueWidget(IDetailCategoryBuilder& InCategory, ERigControlValueType InValueType, FRigControlElement* InControlElement, URigHierarchy* InHierarchy)
+void FRigControlElementDetails_SetupIntegerValueWidget(IDetailGroup& InGroup, IDetailChildrenBuilder& InStructBuilder, ERigControlValueType InValueType, FRigControlElement* InControlElement, URigHierarchy* InHierarchy)
 {
 	UEnum* ControlTypeEnum = StaticEnum<ERigControlType>();
 	UEnum* ValueTypeEnum = StaticEnum<ERigControlValueType>();
@@ -1008,7 +1016,7 @@ void FRigControlElementDetails_SetupIntegerValueWidget(IDetailCategoryBuilder& I
 
 	if (InControlElement->Settings.ControlEnum)
 	{
-		InCategory.AddCustomRow(PropertyLabel)
+		InStructBuilder.AddCustomRow(PropertyLabel)
 		.Visibility(VisibilityAttribute)
         .NameContent()
         [
@@ -1060,7 +1068,7 @@ void FRigControlElementDetails_SetupIntegerValueWidget(IDetailCategoryBuilder& I
 	}
 	else
 	{
-		InCategory.AddCustomRow(PropertyLabel)
+		InStructBuilder.AddCustomRow(PropertyLabel)
 		.Visibility(VisibilityAttribute)
         .NameContent()
         [
@@ -1144,7 +1152,7 @@ void FRigControlElementDetails_SetupIntegerValueWidget(IDetailCategoryBuilder& I
 	}
 }
 
-void FRigControlElementDetails_SetupFloatValueWidget(IDetailCategoryBuilder& InCategory, ERigControlValueType InValueType, FRigControlElement* InControlElement, URigHierarchy* InHierarchy)
+void FRigControlElementDetails_SetupFloatValueWidget(IDetailGroup& InGroup, IDetailChildrenBuilder& InStructBuilder, ERigControlValueType InValueType, FRigControlElement* InControlElement, URigHierarchy* InHierarchy)
 {
 	UEnum* ControlTypeEnum = StaticEnum<ERigControlType>();
 	UEnum* ValueTypeEnum = StaticEnum<ERigControlValueType>();
@@ -1171,7 +1179,7 @@ void FRigControlElementDetails_SetupFloatValueWidget(IDetailCategoryBuilder& InC
         return EnabledAttribute.Get() ? EVisibility::Visible : EVisibility::Hidden;
     }));
 
-	InCategory.AddCustomRow(PropertyLabel)
+	InStructBuilder.AddCustomRow(PropertyLabel)
 	.Visibility(VisibilityAttribute)
     .NameContent()
     [
@@ -1255,7 +1263,7 @@ void FRigControlElementDetails_SetupFloatValueWidget(IDetailCategoryBuilder& InC
 }
 
 template<typename T>
-void FRigControlElementDetails_SetupStructValueWidget(IDetailCategoryBuilder& InCategory, ERigControlValueType InValueType, FRigControlElement* InControlElement, URigHierarchy* InHierarchy)
+void FRigControlElementDetails_SetupStructValueWidget(IDetailGroup& InGroup, IDetailChildrenBuilder& InStructBuilder, ERigControlValueType InValueType, FRigControlElement* InControlElement, URigHierarchy* InHierarchy)
 {
 	UEnum* ControlTypeEnum = StaticEnum<ERigControlType>();
 	UEnum* ValueTypeEnum = StaticEnum<ERigControlValueType>();
@@ -1301,7 +1309,7 @@ void FRigControlElementDetails_SetupStructValueWidget(IDetailCategoryBuilder& In
 		return EVisibility::Hidden;
     }));
 
-	IDetailPropertyRow* Row = InCategory.AddExternalStructure(StructToDisplay);
+	IDetailPropertyRow* Row = InStructBuilder.AddExternalStructure(StructToDisplay.ToSharedRef());
 	Row->DisplayName(PropertyLabel);
 	Row->ShouldAutoExpand(true);
 	Row->IsEnabled(EnabledAttribute);
@@ -1331,7 +1339,7 @@ void FRigControlElementDetails_SetupStructValueWidget(IDetailCategoryBuilder& In
 	Handle->SetOnChildPropertyValueChanged(OnStructContentsChangedDelegate);
 }
 
-void FRigControlElementDetails_SetupValueWidget(IDetailCategoryBuilder& InCategory, ERigControlValueType InValueType, FRigControlElement* InControlElement, URigHierarchy* InHierarchy)
+void FRigControlElementDetails_SetupValueWidget(IDetailGroup& InGroup, IDetailChildrenBuilder& InStructBuilder, ERigControlValueType InValueType, FRigControlElement* InControlElement, URigHierarchy* InHierarchy)
 {
 	switch(InControlElement->Settings.ControlType)
 	{
@@ -1341,48 +1349,48 @@ void FRigControlElementDetails_SetupValueWidget(IDetailCategoryBuilder& InCatego
 			{
 				return;
 			}
-			FRigControlElementDetails_SetupBoolValueWidget(InCategory, InValueType, InControlElement, InHierarchy);
+			FRigControlElementDetails_SetupBoolValueWidget(InGroup, InStructBuilder, InValueType, InControlElement, InHierarchy);
 			break;
 		}
 		case ERigControlType::Integer:
 		{
-			FRigControlElementDetails_SetupIntegerValueWidget(InCategory, InValueType, InControlElement, InHierarchy);
+			FRigControlElementDetails_SetupIntegerValueWidget(InGroup, InStructBuilder, InValueType, InControlElement, InHierarchy);
 			break;
 		}
 		case ERigControlType::Float:
 		{
-			FRigControlElementDetails_SetupFloatValueWidget(InCategory, InValueType, InControlElement, InHierarchy);
+			FRigControlElementDetails_SetupFloatValueWidget(InGroup, InStructBuilder, InValueType, InControlElement, InHierarchy);
 			break;
 		}
 		case ERigControlType::Vector2D:
 		{
-			FRigControlElementDetails_SetupStructValueWidget<FVector2D>(InCategory, InValueType, InControlElement, InHierarchy);
+			FRigControlElementDetails_SetupStructValueWidget<FVector2D>(InGroup, InStructBuilder, InValueType, InControlElement, InHierarchy);
 			break;
 		}
 		case ERigControlType::Position:
 		case ERigControlType::Scale:
 		{
-			FRigControlElementDetails_SetupStructValueWidget<FVector>(InCategory, InValueType, InControlElement, InHierarchy);
+			FRigControlElementDetails_SetupStructValueWidget<FVector>(InGroup, InStructBuilder, InValueType, InControlElement, InHierarchy);
 			break;
 		}
 		case ERigControlType::Rotator:
 		{
-			FRigControlElementDetails_SetupStructValueWidget<FRotator>(InCategory, InValueType, InControlElement, InHierarchy);
+			FRigControlElementDetails_SetupStructValueWidget<FRotator>(InGroup, InStructBuilder, InValueType, InControlElement, InHierarchy);
 			break;
 		}
 		case ERigControlType::TransformNoScale:
 		{
-			FRigControlElementDetails_SetupStructValueWidget<FTransformNoScale>(InCategory, InValueType, InControlElement, InHierarchy);
+			FRigControlElementDetails_SetupStructValueWidget<FTransformNoScale>(InGroup, InStructBuilder, InValueType, InControlElement, InHierarchy);
 			break;
 		}
 		case ERigControlType::EulerTransform:
 		{
-			FRigControlElementDetails_SetupStructValueWidget<FEulerTransform>(InCategory, InValueType, InControlElement, InHierarchy);
+			FRigControlElementDetails_SetupStructValueWidget<FEulerTransform>(InGroup, InStructBuilder, InValueType, InControlElement, InHierarchy);
 			break;
 		}
 		case ERigControlType::Transform:
 		{
-			FRigControlElementDetails_SetupStructValueWidget<FTransform>(InCategory, InValueType, InControlElement, InHierarchy);
+			FRigControlElementDetails_SetupStructValueWidget<FTransform>(InGroup, InStructBuilder, InValueType, InControlElement, InHierarchy);
 			break;
 		}
 		default:
@@ -1395,9 +1403,9 @@ void FRigControlElementDetails_SetupValueWidget(IDetailCategoryBuilder& InCatego
 
 TArray<TSharedPtr<FString>> FRigControlElementDetails::ControlTypeList;
 
-void FRigControlElementDetails::CustomizeDetails(IDetailLayoutBuilder& DetailBuilder)
+void FRigControlElementDetails::CustomizeChildren(TSharedRef<class IPropertyHandle> InStructPropertyHandle, class IDetailChildrenBuilder& StructBuilder, IPropertyTypeCustomizationUtils& StructCustomizationUtils)
 {
-	FRigTransformElementDetails::CustomizeDetails(DetailBuilder);
+	FRigTransformElementDetails::CustomizeChildren(InStructPropertyHandle, StructBuilder, StructCustomizationUtils);
 
 	GizmoNameList.Reset();
 	if (BlueprintBeingCustomized)
@@ -1416,34 +1424,30 @@ void FRigControlElementDetails::CustomizeDetails(IDetailLayoutBuilder& DetailBui
 		}
 	}
 
-	if (HierarchyBeingCustomized == nullptr || !ElementKeyBeingCustomized)
+	if (HierarchyBeingCustomized == nullptr)
 	{
 		return;
 	}
 
-	IDetailCategoryBuilder& ControlCategory = DetailBuilder.EditCategory(TEXT("Control"), LOCTEXT("ControlCategory", "Control"));
-	IDetailCategoryBuilder& LimitsCategory = DetailBuilder.EditCategory(TEXT("Limits"), LOCTEXT("LimitsCategory", "Limits"));
-	IDetailCategoryBuilder& GizmoCategory = DetailBuilder.EditCategory(TEXT("Gizmo"), LOCTEXT("GizmoCategory", "Gizmo"));
+	IDetailGroup* ControlGroup = &StructBuilder.AddGroup(TEXT("Control"), LOCTEXT("Control", "Control"));
+	IDetailGroup* GizmoGroup = nullptr;
+	IDetailGroup* LimitsGroup = nullptr;
 
-	ControlCategory.InitiallyCollapsed(false);
-	LimitsCategory.InitiallyCollapsed(false);
-	GizmoCategory.InitiallyCollapsed(false);
+	const TSharedPtr<IPropertyHandle> SettingsHandle = InStructPropertyHandle->GetChildHandle(TEXT("Settings"));
+	const TSharedPtr<IPropertyHandle> DisplayNameHandle = SettingsHandle->GetChildHandle(TEXT("DisplayName"));
 
-	const TSharedRef<IPropertyHandle> DisplayNameHandle = DetailBuilder.GetProperty(TEXT("Settings.DisplayName"));
-	DetailBuilder.HideProperty(DisplayNameHandle);
-	ControlCategory.AddProperty(DisplayNameHandle).CustomWidget()
+	ControlGroup->AddWidgetRow()
 	.NameContent()
-    [
-        SNew(STextBlock)
-        .Text(FText::FromString(TEXT("Display Name")))
-        .Font(IDetailLayoutBuilder::GetDetailFont())
-    ]
+	[
+		DisplayNameHandle->CreatePropertyNameWidget()
+	]
 	.ValueContent()
-    [
+	[
 		SNew(SEditableTextBox)
 		.Font(IDetailLayoutBuilder::GetDetailFont())
 		.Text(this, &FRigControlElementDetails::GetDisplayName)
-		.OnTextCommitted(this, &FRigControlElementDetails::SetDisplayName, DetailBuilder.GetPropertyUtilities())
+		.OnTextCommitted(this, &FRigControlElementDetails::SetDisplayName)
+		.IsEnabled(ObjectsBeingCustomized.Num() == 1)
 	];
 
 	if (ControlTypeList.Num() == 0)
@@ -1455,326 +1459,291 @@ void FRigControlElementDetails::CustomizeDetails(IDetailLayoutBuilder& DetailBui
 		}
 	}
 
-	FRigControlElement* ControlElement = HierarchyBeingCustomized->FindChecked<FRigControlElement>(ElementKeyBeingCustomized);
-
 	// when control type changes, we have to refresh detail panel
-	const TSharedRef<IPropertyHandle> ControlTypeHandle = DetailBuilder.GetProperty(TEXT("Settings.ControlType"));
+	const TSharedPtr<IPropertyHandle> ControlTypeHandle = SettingsHandle->GetChildHandle(TEXT("ControlType"));
 	ControlTypeHandle->SetOnPropertyValueChanged(FSimpleDelegate::CreateLambda(
-		[this, &DetailBuilder]()
+		[this, &StructCustomizationUtils]()
 		{
-			DetailBuilder.ForceRefreshDetails();
+			TArray<FRigControlElement*> ControlElementsInView = GetElementsInDetailsView<FRigControlElement>();
+			TArray<FRigControlElement*> ControlElementsInHierarchy = GetElementsInHierarchy<FRigControlElement>();
+			check(ControlElementsInView.Num() == ControlElementsInHierarchy.Num());
 
-			if (this->HierarchyBeingCustomized && this->ElementKeyBeingCustomized.IsValid())
+			if (this->HierarchyBeingCustomized && ControlElementsInHierarchy.Num() > 0)
 			{
-				FRigControlElement* ControlElement = this->HierarchyBeingCustomized->FindChecked<FRigControlElement>(ElementKeyBeingCustomized);
-				FRigControlValue ValueToSet;
+				for(int32 ControlIndex = 0; ControlIndex< ControlElementsInView.Num(); ControlIndex++)
+				{
+					FRigControlElement* ViewElement = ControlElementsInView[ControlIndex];
+					FRigControlElement* ControlElement = ControlElementsInHierarchy[ControlIndex];
+					
+					FRigControlValue ValueToSet;
 
-				ControlElement->Settings.bLimitTranslation = false;
-				ControlElement->Settings.bLimitRotation = false;
-				ControlElement->Settings.bLimitScale = false;
+					ControlElement->Settings.ControlType = ViewElement->Settings.ControlType;
+					ControlElement->Settings.bLimitTranslation = false;
+					ControlElement->Settings.bLimitRotation = false;
+					ControlElement->Settings.bLimitScale = false;
 
-				switch (ControlElement->Settings.ControlType)
-				{
-					case ERigControlType::Bool:
-					{
-						ValueToSet = FRigControlValue::Make<bool>(false);
-						break;
-					}
-					case ERigControlType::Float:
+					switch (ControlElement->Settings.ControlType)
 					{
-						ValueToSet = FRigControlValue::Make<float>(0.f);
-						ControlElement->Settings.bLimitTranslation = true;
-						ControlElement->Settings.MinimumValue = FRigControlValue::Make<float>(0.f);
-						ControlElement->Settings.MaximumValue = FRigControlValue::Make<float>(100.f);
-                        break;
-					}
-					case ERigControlType::Integer:
-					{
-						ValueToSet = FRigControlValue::Make<int32>(0);
-						ControlElement->Settings.bLimitTranslation = true;
-						ControlElement->Settings.MinimumValue = FRigControlValue::Make<int32>(0);
-						ControlElement->Settings.MaximumValue = FRigControlValue::Make<int32>(100);
-						break;
-					}
-					case ERigControlType::Vector2D:
-					{
-						ValueToSet = FRigControlValue::Make<FVector2D>(FVector2D::ZeroVector);
-						ControlElement->Settings.bLimitTranslation = true;
-						ControlElement->Settings.MinimumValue = FRigControlValue::Make<FVector2D>(FVector2D::ZeroVector);
-						ControlElement->Settings.MaximumValue = FRigControlValue::Make<FVector2D>(FVector2D(100.f, 100.f));
-						break;
-					}
-					case ERigControlType::Position:
-					{
-						ValueToSet = FRigControlValue::Make<FVector>(FVector::ZeroVector);
-						ControlElement->Settings.MinimumValue = FRigControlValue::Make<FVector>(-FVector::OneVector);
-						ControlElement->Settings.MaximumValue = FRigControlValue::Make<FVector>(FVector::OneVector);
-						break;
-					}
-					case ERigControlType::Scale:
-					{
-						ValueToSet = FRigControlValue::Make<FVector>(FVector::OneVector);
-						ControlElement->Settings.MinimumValue = FRigControlValue::Make<FVector>(FVector::ZeroVector);
-						ControlElement->Settings.MaximumValue = FRigControlValue::Make<FVector>(FVector::OneVector);
-						break;
-					}
-					case ERigControlType::Rotator:
-					{
-						ValueToSet = FRigControlValue::Make<FRotator>(FRotator::ZeroRotator);
-						ControlElement->Settings.MinimumValue = FRigControlValue::Make<FRotator>(FRotator::ZeroRotator);
-						ControlElement->Settings.MaximumValue = FRigControlValue::Make<FRotator>(FRotator(180.f, 180.f, 180.f));
-						break;
-					}
-					case ERigControlType::Transform:
-					{
-						ValueToSet = FRigControlValue::Make<FTransform>(FTransform::Identity);
-						ControlElement->Settings.MinimumValue = ValueToSet;
-						ControlElement->Settings.MaximumValue = ValueToSet;
-						break;
-					}
-					case ERigControlType::TransformNoScale:
-					{
-						FTransformNoScale Identity = FTransform::Identity;
-						ValueToSet = FRigControlValue::Make<FTransformNoScale>(Identity);
-						ControlElement->Settings.MinimumValue = ValueToSet;
-                        ControlElement->Settings.MaximumValue = ValueToSet;
-                        break;
-					}
-					case ERigControlType::EulerTransform:
-					{
-						FEulerTransform Identity = FEulerTransform::Identity;
-						ValueToSet = FRigControlValue::Make<FEulerTransform>(Identity);
-						ControlElement->Settings.MinimumValue = ValueToSet;
-						ControlElement->Settings.MaximumValue = ValueToSet;
-						break;
-					}
-					default:
-					{
-						ensure(false);
-						break;
+						case ERigControlType::Bool:
+						{
+							ValueToSet = FRigControlValue::Make<bool>(false);
+							break;
+						}
+						case ERigControlType::Float:
+						{
+							ValueToSet = FRigControlValue::Make<float>(0.f);
+							ControlElement->Settings.bLimitTranslation = true;
+							ControlElement->Settings.MinimumValue = FRigControlValue::Make<float>(0.f);
+							ControlElement->Settings.MaximumValue = FRigControlValue::Make<float>(100.f);
+							break;
+						}
+						case ERigControlType::Integer:
+						{
+							ValueToSet = FRigControlValue::Make<int32>(0);
+							ControlElement->Settings.bLimitTranslation = true;
+							ControlElement->Settings.MinimumValue = FRigControlValue::Make<int32>(0);
+							ControlElement->Settings.MaximumValue = FRigControlValue::Make<int32>(100);
+							break;
+						}
+						case ERigControlType::Vector2D:
+						{
+							ValueToSet = FRigControlValue::Make<FVector2D>(FVector2D::ZeroVector);
+							ControlElement->Settings.bLimitTranslation = true;
+							ControlElement->Settings.MinimumValue = FRigControlValue::Make<FVector2D>(FVector2D::ZeroVector);
+							ControlElement->Settings.MaximumValue = FRigControlValue::Make<FVector2D>(FVector2D(100.f, 100.f));
+							break;
+						}
+						case ERigControlType::Position:
+						{
+							ValueToSet = FRigControlValue::Make<FVector>(FVector::ZeroVector);
+							ControlElement->Settings.MinimumValue = FRigControlValue::Make<FVector>(-FVector::OneVector);
+							ControlElement->Settings.MaximumValue = FRigControlValue::Make<FVector>(FVector::OneVector);
+							break;
+						}
+						case ERigControlType::Scale:
+						{
+							ValueToSet = FRigControlValue::Make<FVector>(FVector::OneVector);
+							ControlElement->Settings.MinimumValue = FRigControlValue::Make<FVector>(FVector::ZeroVector);
+							ControlElement->Settings.MaximumValue = FRigControlValue::Make<FVector>(FVector::OneVector);
+							break;
+						}
+						case ERigControlType::Rotator:
+						{
+							ValueToSet = FRigControlValue::Make<FRotator>(FRotator::ZeroRotator);
+							ControlElement->Settings.MinimumValue = FRigControlValue::Make<FRotator>(FRotator::ZeroRotator);
+							ControlElement->Settings.MaximumValue = FRigControlValue::Make<FRotator>(FRotator(180.f, 180.f, 180.f));
+							break;
+						}
+						case ERigControlType::Transform:
+						{
+							ValueToSet = FRigControlValue::Make<FTransform>(FTransform::Identity);
+							ControlElement->Settings.MinimumValue = ValueToSet;
+							ControlElement->Settings.MaximumValue = ValueToSet;
+							break;
+						}
+						case ERigControlType::TransformNoScale:
+						{
+							FTransformNoScale Identity = FTransform::Identity;
+							ValueToSet = FRigControlValue::Make<FTransformNoScale>(Identity);
+							ControlElement->Settings.MinimumValue = ValueToSet;
+							ControlElement->Settings.MaximumValue = ValueToSet;
+							break;
+						}
+						case ERigControlType::EulerTransform:
+						{
+							FEulerTransform Identity = FEulerTransform::Identity;
+							ValueToSet = FRigControlValue::Make<FEulerTransform>(Identity);
+							ControlElement->Settings.MinimumValue = ValueToSet;
+							ControlElement->Settings.MaximumValue = ValueToSet;
+							break;
+						}
+						default:
+						{
+							ensure(false);
+							break;
+						}
 					}
-				}
 
-				this->HierarchyBeingCustomized->Notify(ERigHierarchyNotification::ControlSettingChanged, ControlElement);
-				this->HierarchyBeingCustomized->SetControlValue(ControlElement, ValueToSet, ERigControlValueType::Initial, true);
-				this->HierarchyBeingCustomized->SetControlValue(ControlElement, ValueToSet, ERigControlValueType::Current, true);
+					this->HierarchyBeingCustomized->Notify(ERigHierarchyNotification::ControlSettingChanged, ControlElement);
+					this->HierarchyBeingCustomized->SetControlValue(ControlElement, ValueToSet, ERigControlValueType::Initial, true);
+					this->HierarchyBeingCustomized->SetControlValue(ControlElement, ValueToSet, ERigControlValueType::Current, true);
 
-				if (this->HierarchyBeingCustomized != this->BlueprintBeingCustomized->Hierarchy)
-				{
-					if(FRigControlElement* OtherControlElement = this->BlueprintBeingCustomized->Hierarchy->Find<FRigControlElement>(ControlElement->GetKey()))
+					FRigControlElement::StaticStruct()->CopyScriptStruct(ViewElement, ControlElement);
+
+					if (this->HierarchyBeingCustomized != this->BlueprintBeingCustomized->Hierarchy)
 					{
-						OtherControlElement->Settings = ControlElement->Settings;
-						this->BlueprintBeingCustomized->Hierarchy->Notify(ERigHierarchyNotification::ControlSettingChanged, ControlElement);
-                        this->BlueprintBeingCustomized->Hierarchy->SetControlValue(ControlElement, ValueToSet, ERigControlValueType::Initial, true);
-                        this->BlueprintBeingCustomized->Hierarchy->SetControlValue(ControlElement, ValueToSet, ERigControlValueType::Current, true);
-                   }
+						if(FRigControlElement* OtherControlElement = this->BlueprintBeingCustomized->Hierarchy->Find<FRigControlElement>(ControlElement->GetKey()))
+						{
+							OtherControlElement->Settings = ControlElement->Settings;
+							this->BlueprintBeingCustomized->Hierarchy->Notify(ERigHierarchyNotification::ControlSettingChanged, ControlElement);
+							this->BlueprintBeingCustomized->Hierarchy->SetControlValue(ControlElement, ValueToSet, ERigControlValueType::Initial, true);
+							this->BlueprintBeingCustomized->Hierarchy->SetControlValue(ControlElement, ValueToSet, ERigControlValueType::Current, true);
+					}
+					}
 				}
+				
+				StructCustomizationUtils.GetPropertyUtilities()->ForceRefresh();
 			}
 		}
 	));
 
-	if(ControlElement->Settings.ControlType == ERigControlType::Bool)
-	{
-		DetailBuilder.HideCategory(TEXT("Pose"));
-		DetailBuilder.HideProperty(TEXT("Offset"));
-		DetailBuilder.HideProperty(TEXT("Gizmo"));
-	}
-	else
+	ControlGroup->AddPropertyRow(ControlTypeHandle.ToSharedRef());
+	ControlGroup->AddPropertyRow(SettingsHandle->GetChildHandle(TEXT("bAnimatable")).ToSharedRef());
+
+	// any but bool controls show the offset + gizmo
+	if(IsAnyControlNotOfType(ERigControlType::Bool))
 	{
-		// setup offset transform
+		// setup offset
 		{
-			const TSharedRef<IPropertyHandle> OffsetInitialLocalTransformHandle = DetailBuilder.GetProperty(TEXT("Offset.Initial.Local.Transform"));
-			ControlCategory.AddProperty(OffsetInitialLocalTransformHandle).DisplayName(FText::FromString(TEXT("Offset Transform")));
-			DetailBuilder.HideProperty(TEXT("Offset"));
-
-			OffsetPropertyChain.AddHead(FRigControlElement::StaticStruct()->FindPropertyByName(TEXT("Offset")));
-			OffsetPropertyChain.AddTail(FRigCurrentAndInitialTransform::StaticStruct()->FindPropertyByName(TEXT("Initial")));
-			OffsetPropertyChain.AddTail(FRigLocalAndGlobalTransform::StaticStruct()->FindPropertyByName(TEXT("Local")));
-			OffsetPropertyChain.AddTail(FRigComputedTransform::StaticStruct()->FindPropertyByName(TEXT("Transform")));
-			OffsetPropertyChain.SetActiveMemberPropertyNode(OffsetPropertyChain.GetTail()->GetValue());
-
-			const FSimpleDelegate OnTransformChangedDelegate = FSimpleDelegate::CreateStatic(&FRigTransformElementDetails::OnTransformChanged, &OffsetPropertyChain, BlueprintBeingCustomized);
-			OffsetInitialLocalTransformHandle->SetOnPropertyValueChanged(OnTransformChangedDelegate);
-			OffsetInitialLocalTransformHandle->SetOnChildPropertyValueChanged(OnTransformChangedDelegate);
-			OffsetInitialLocalTransformHandle->SetToolTipText(FText::FromString(TEXT("The offset transform is used as a middle man between the parent and the local transform.\nYou can use it to offset a control for visual adjustment.")));
+			const TSharedPtr<IPropertyHandle> OffsetHandle = InStructPropertyHandle->GetChildHandle(TEXT("Offset"));
+			const TSharedPtr<IPropertyHandle> InitialHandle = OffsetHandle->GetChildHandle(TEXT("Initial"));
+			const TSharedPtr<IPropertyHandle> LocalHandle = InitialHandle->GetChildHandle(TEXT("Local"));
+			const TSharedPtr<IPropertyHandle> TransformHandle = LocalHandle->GetChildHandle(TEXT("Transform"));
+			ControlGroup->AddPropertyRow(TransformHandle.ToSharedRef()).DisplayName(FText::FromString(TEXT("Offset Transform")));
 		}
 
-		// setup gizmo transform
-		{
-			const TSharedRef<IPropertyHandle> GizmoInitialLocalTransformHandle = DetailBuilder.GetProperty(TEXT("Gizmo.Initial.Local.Transform"));
-			GizmoCategory.AddProperty(GizmoInitialLocalTransformHandle).DisplayName(FText::FromString(TEXT("Gizmo Transform")));
-			DetailBuilder.HideProperty(TEXT("Gizmo"));
-
-			GizmoPropertyChain.AddHead(FRigControlElement::StaticStruct()->FindPropertyByName(TEXT("Gizmo")));
-			GizmoPropertyChain.AddTail(FRigCurrentAndInitialTransform::StaticStruct()->FindPropertyByName(TEXT("Initial")));
-			GizmoPropertyChain.AddTail(FRigLocalAndGlobalTransform::StaticStruct()->FindPropertyByName(TEXT("Local")));
-			GizmoPropertyChain.AddTail(FRigComputedTransform::StaticStruct()->FindPropertyByName(TEXT("Transform")));
-			GizmoPropertyChain.SetActiveMemberPropertyNode(GizmoPropertyChain.GetTail()->GetValue());
-
-			const FSimpleDelegate OnTransformChangedDelegate = FSimpleDelegate::CreateStatic(&FRigTransformElementDetails::OnTransformChanged, &GizmoPropertyChain, BlueprintBeingCustomized);
-			GizmoInitialLocalTransformHandle->SetOnPropertyValueChanged(OnTransformChangedDelegate);
-			GizmoInitialLocalTransformHandle->SetOnChildPropertyValueChanged(OnTransformChangedDelegate);
-			GizmoInitialLocalTransformHandle->SetToolTipText(FText::FromString(TEXT("The gizmo transform is used as a transform applied only to the UI element on the screen.\nIt doesn't affect animation.")));
-		}
+		GizmoGroup = &StructBuilder.AddGroup(TEXT("Gizmo"), LOCTEXT("Gizmo", "Gizmo"));
 	}
 
-	FRigControlElementDetails_SetupValueWidget(ControlCategory, ERigControlValueType::Current, ControlElement, HierarchyBeingCustomized);
-	switch (ControlElement->Settings.ControlType)
+	if(IsAnyControlOfType(ERigControlType::Float) ||
+		IsAnyControlOfType(ERigControlType::Integer) ||
+		IsAnyControlOfType(ERigControlType::Vector2D) ||
+		IsAnyControlOfType(ERigControlType::Position) ||
+		IsAnyControlOfType(ERigControlType::Rotator) ||
+		IsAnyControlOfType(ERigControlType::Scale) ||
+		IsAnyControlOfType(ERigControlType::Transform) ||
+		IsAnyControlOfType(ERigControlType::TransformNoScale) ||
+		IsAnyControlOfType(ERigControlType::EulerTransform))
 	{
-		case ERigControlType::Bool:
-		case ERigControlType::Float:
-		case ERigControlType::Integer:
-		case ERigControlType::Vector2D:
-		{
-			FRigControlElementDetails_SetupValueWidget(ControlCategory, ERigControlValueType::Initial, ControlElement, HierarchyBeingCustomized);
-			break;
-		}
-		default:
-		{
-			break;
-		}
+		LimitsGroup = &StructBuilder.AddGroup(TEXT("Limits"), LOCTEXT("Limits", "Limits"));
 	}
-	FRigControlElementDetails_SetupValueWidget(LimitsCategory, ERigControlValueType::Minimum, ControlElement, HierarchyBeingCustomized);
-	FRigControlElementDetails_SetupValueWidget(LimitsCategory, ERigControlValueType::Maximum, ControlElement, HierarchyBeingCustomized);
-
-	switch (ControlElement->Settings.ControlType)
+	
+	if(IsAnyControlOfType(ERigControlType::Float) ||
+		IsAnyControlOfType(ERigControlType::Integer) ||
+		IsAnyControlOfType(ERigControlType::Vector2D) ||
+		IsAnyControlOfType(ERigControlType::Position) ||
+		IsAnyControlOfType(ERigControlType::Transform) ||
+		IsAnyControlOfType(ERigControlType::TransformNoScale) ||
+		IsAnyControlOfType(ERigControlType::EulerTransform))
 	{
-		case ERigControlType::Float:
-		case ERigControlType::Integer:
-		case ERigControlType::Vector2D:
-		case ERigControlType::Position:
-		case ERigControlType::Scale:
-		case ERigControlType::Rotator:
-		case ERigControlType::Transform:
-		case ERigControlType::TransformNoScale:
-		case ERigControlType::EulerTransform:
-		{
-			const TSharedRef<IPropertyHandle> GizmoNameHandle = DetailBuilder.GetProperty(TEXT("Settings.GizmoName"));
-			DetailBuilder.HideProperty(GizmoNameHandle);
-			GizmoCategory.AddProperty(GizmoNameHandle).CustomWidget()
-			.NameContent()
-			[
-				SNew(STextBlock)
-				.Text(FText::FromString(TEXT("Gizmo Name")))
-				.Font(IDetailLayoutBuilder::GetDetailFont())
-				.IsEnabled(this, &FRigControlElementDetails::IsGizmoEnabled)
-			]
-			.ValueContent()
-			[
-				SNew(SControlRigGizmoNameList, ControlElement, BlueprintBeingCustomized)
-				.OnGetNameListContent(this, &FRigControlElementDetails::GetGizmoNameList)
-				.IsEnabled(this, &FRigControlElementDetails::IsGizmoEnabled)
-			];
-			break;
-		}
-		default:
-		{
-			DetailBuilder.HideProperty(TEXT("Settings.bGizmoEnabled"));
-			DetailBuilder.HideProperty(TEXT("Settings.bGizmoVisible"));
-			DetailBuilder.HideProperty(TEXT("Settings.GizmoName"));
-			DetailBuilder.HideProperty(TEXT("Settings.GizmoColor"));
-			break;
-		}
+		const TSharedPtr<IPropertyHandle> LimitHandle = SettingsHandle->GetChildHandle(TEXT("bLimitTranslation"));
+		LimitsGroup->AddPropertyRow(LimitHandle.ToSharedRef()).DisplayName(FText::FromString(TEXT("Limit Translation")));
 	}
 
-	if(ControlElement->Settings.ControlType != ERigControlType::Integer)
+	if(IsAnyControlOfType(ERigControlType::Rotator) ||
+		IsAnyControlOfType(ERigControlType::Transform) ||
+		IsAnyControlOfType(ERigControlType::TransformNoScale) ||
+		IsAnyControlOfType(ERigControlType::EulerTransform))
 	{
-		DetailBuilder.HideProperty(TEXT("Settings.ControlEnum"));
+		const TSharedPtr<IPropertyHandle> LimitHandle = SettingsHandle->GetChildHandle(TEXT("bLimitRotation"));
+		LimitsGroup->AddPropertyRow(LimitHandle.ToSharedRef()).DisplayName(FText::FromString(TEXT("Limit Rotation")));
 	}
 
-	bool bShowLimitTranslation = false;
-	bool bShowLimitRotation = false;
-	bool bShowLimitScale = false;
+	if(IsAnyControlOfType(ERigControlType::Scale) ||
+		IsAnyControlOfType(ERigControlType::Transform) ||
+		IsAnyControlOfType(ERigControlType::EulerTransform))
+	{
+		const TSharedPtr<IPropertyHandle> LimitHandle = SettingsHandle->GetChildHandle(TEXT("bLimitScale"));
+		LimitsGroup->AddPropertyRow(LimitHandle.ToSharedRef()).DisplayName(FText::FromString(TEXT("Limit Scale")));
+	}
 
-	switch (ControlElement->Settings.ControlType)
+	if(!(IsAnyControlNotOfType(ERigControlType::Integer) &&
+		IsAnyControlNotOfType(ERigControlType::Float) &&
+		IsAnyControlNotOfType(ERigControlType::Vector2D)))
 	{
-		case ERigControlType::Float:
-    	case ERigControlType::Integer:
-    	case ERigControlType::Vector2D:
-    	case ERigControlType::Position:
-    	case ERigControlType::Transform:
-    	case ERigControlType::TransformNoScale:
-    	case ERigControlType::EulerTransform:
-		{
-				bShowLimitTranslation = true;
-			break;
-		}
-		default:
-		{
-			break;
-		}
+		const TSharedPtr<IPropertyHandle> PrimaryAxisHandle = SettingsHandle->GetChildHandle(TEXT("PrimaryAxis"));
+		ControlGroup->AddPropertyRow(PrimaryAxisHandle.ToSharedRef()).DisplayName(FText::FromString(TEXT("Primary Axis")));
 	}
 
-	switch (ControlElement->Settings.ControlType)
+	TArray<FRigControlElement*> ControlElements;
+	for(UDetailsViewWrapperObject* ObjectBeingCustomized : ObjectsBeingCustomized)
 	{
-    	case ERigControlType::Rotator:
-    	case ERigControlType::Transform:
-    	case ERigControlType::TransformNoScale:
-    	case ERigControlType::EulerTransform:
+		if(FRigControlElement* ControlElement = Cast<FRigControlElement>(ObjectBeingCustomized->GetContent<FRigBaseElement>()))
 		{
-			bShowLimitRotation = true;
-			break;
-		}
-		default:
-		{
-			break;
+			ControlElements.Add(ControlElement);
 		}
 	}
 
-	switch (ControlElement->Settings.ControlType)
+	// only setup value widgets if there is only ony control selected
+	if(ControlElements.Num() == 1)
 	{
-		case ERigControlType::Scale:
-    	case ERigControlType::Transform:
-    	case ERigControlType::EulerTransform:
-		{
-			bShowLimitScale = true;
-			break;
-		}
-		default:
+		FRigControlElementDetails_SetupValueWidget(*ControlGroup, StructBuilder, ERigControlValueType::Current, ControlElements[0], HierarchyBeingCustomized);
+
+		switch (ControlElements[0]->Settings.ControlType)
 		{
-			break;
+			case ERigControlType::Bool:
+			case ERigControlType::Float:
+			case ERigControlType::Integer:
+			case ERigControlType::Vector2D:
+			{
+				FRigControlElementDetails_SetupValueWidget(*ControlGroup, StructBuilder, ERigControlValueType::Initial, ControlElements[0], HierarchyBeingCustomized);
+				break;
+			}
+			default:
+			{
+				break;
+			}
 		}
+		
+		FRigControlElementDetails_SetupValueWidget(*ControlGroup, StructBuilder, ERigControlValueType::Minimum, ControlElements[0], HierarchyBeingCustomized);
+		FRigControlElementDetails_SetupValueWidget(*ControlGroup, StructBuilder, ERigControlValueType::Maximum, ControlElements[0], HierarchyBeingCustomized);
 	}
 
-	if(!bShowLimitTranslation)
-	{
-		DetailBuilder.HideProperty(TEXT("Settings.bLimitTranslation"));
-	}
-	if(!bShowLimitRotation)
-	{
-		DetailBuilder.HideProperty(TEXT("Settings.bLimitRotation"));
-	}
-	if(!bShowLimitScale)
+	if(IsAnyControlOfType(ERigControlType::Float) ||
+		IsAnyControlOfType(ERigControlType::Integer) ||
+		IsAnyControlOfType(ERigControlType::Vector2D) ||
+		IsAnyControlOfType(ERigControlType::Position) ||
+		IsAnyControlOfType(ERigControlType::Scale) ||
+		IsAnyControlOfType(ERigControlType::Rotator) ||
+		IsAnyControlOfType(ERigControlType::Transform) ||
+		IsAnyControlOfType(ERigControlType::TransformNoScale) ||
+		IsAnyControlOfType(ERigControlType::EulerTransform))
 	{
-		DetailBuilder.HideProperty(TEXT("Settings.bLimitScale"));
-	}
+		GizmoGroup->AddPropertyRow(SettingsHandle->GetChildHandle(TEXT("bGizmoEnabled")).ToSharedRef());
+		GizmoGroup->AddPropertyRow(SettingsHandle->GetChildHandle(TEXT("bGizmoVisible")).ToSharedRef());
 
-	switch (ControlElement->Settings.ControlType)
-	{
-		case ERigControlType::Integer:
-		case ERigControlType::Float:
-		case ERigControlType::Vector2D:
-		{
-			break;
-		}
-		default:
+		// setup gizmo transform
 		{
-			DetailBuilder.HideProperty(TEXT("Settings.PrimaryAxis"));
-			break;
+			const TSharedPtr<IPropertyHandle> GizmoHandle = InStructPropertyHandle->GetChildHandle(TEXT("Gizmo"));
+			const TSharedPtr<IPropertyHandle> InitialHandle = GizmoHandle->GetChildHandle(TEXT("Initial"));
+			const TSharedPtr<IPropertyHandle> LocalHandle = InitialHandle->GetChildHandle(TEXT("Local"));
+			const TSharedPtr<IPropertyHandle> TransformHandle = LocalHandle->GetChildHandle(TEXT("Transform"));
+			GizmoGroup->AddPropertyRow(TransformHandle.ToSharedRef()).DisplayName(FText::FromString(TEXT("Gizmo Transform")))
+			.IsEnabled(TAttribute<bool>::CreateSP(this, &FRigControlElementDetails::IsGizmoEnabled));
 		}
+
+		const TSharedPtr<IPropertyHandle> GizmoNameHandle = SettingsHandle->GetChildHandle(TEXT("GizmoName"));
+		GizmoGroup->AddPropertyRow(GizmoNameHandle.ToSharedRef()).CustomWidget()
+		.NameContent()
+		[
+			SNew(STextBlock)
+			.Text(FText::FromString(TEXT("Gizmo Name")))
+			.Font(IDetailLayoutBuilder::GetDetailFont())
+			.IsEnabled(this, &FRigControlElementDetails::IsGizmoEnabled)
+		]
+		.ValueContent()
+		[
+			SNew(SControlRigGizmoNameList, ControlElements, BlueprintBeingCustomized)
+			.OnGetNameListContent(this, &FRigControlElementDetails::GetGizmoNameList)
+			.IsEnabled(this, &FRigControlElementDetails::IsGizmoEnabled)
+		];
+
+		GizmoGroup->AddPropertyRow(SettingsHandle->GetChildHandle(TEXT("GizmoColor")).ToSharedRef());
 	}
 
-	if (ControlElement->Settings.ControlType == ERigControlType::Integer)
+	if(IsAnyControlOfType(ERigControlType::Integer))
 	{
-		TSharedRef<IPropertyHandle> ControlEnum = DetailBuilder.GetProperty(TEXT("Settings.ControlEnum"));
-		ControlEnum->SetOnPropertyValueChanged(FSimpleDelegate::CreateLambda(
-			[this, &DetailBuilder]()
+		const TSharedPtr<IPropertyHandle> ControlEnumHandle = SettingsHandle->GetChildHandle(TEXT("ControlEnum"));
+		ControlGroup->AddPropertyRow(ControlEnumHandle.ToSharedRef()).DisplayName(FText::FromString(TEXT("Control Enum")));
+
+		ControlEnumHandle->SetOnPropertyValueChanged(FSimpleDelegate::CreateLambda(
+			[this, &StructCustomizationUtils, ControlElements]()
 			{
-				DetailBuilder.ForceRefreshDetails();
+				StructCustomizationUtils.GetPropertyUtilities()->ForceRefresh();
 
-				if (this->HierarchyBeingCustomized != nullptr && this->ElementKeyBeingCustomized)
+				if (this->HierarchyBeingCustomized != nullptr)
 				{
-					if(FRigControlElement* ControlBeingCustomized = this->HierarchyBeingCustomized->Find<FRigControlElement>(ElementKeyBeingCustomized))
+					for(FRigControlElement* ControlBeingCustomized : ControlElements)
 					{
 						const UEnum* ControlEnum = ControlBeingCustomized->Settings.ControlEnum;
 						if (ControlEnum != nullptr)
@@ -1795,7 +1764,7 @@ void FRigControlElementDetails::CustomizeDetails(IDetailLayoutBuilder& DetailBui
 							if (UControlRig* DebuggedRig = Cast<UControlRig>(BlueprintBeingCustomized->GetObjectBeingDebugged()))
 							{
 								URigHierarchy* DebuggedHierarchy = DebuggedRig->GetHierarchy();
-								if(FRigControlElement* DebuggedControlElement = DebuggedHierarchy->Find<FRigControlElement>(ElementKeyBeingCustomized))
+								if(FRigControlElement* DebuggedControlElement = DebuggedHierarchy->Find<FRigControlElement>(ControlBeingCustomized->GetKey()))
 								{
 									DebuggedControlElement->Settings.MinimumValue.Set<int32>(0);
                                     DebuggedControlElement->Settings.MaximumValue.Set<int32>(Maximum);
@@ -1810,110 +1779,103 @@ void FRigControlElementDetails::CustomizeDetails(IDetailLayoutBuilder& DetailBui
 				}
 			}
 		));
+		
 	}
 }
 
-FText FRigControlElementDetails::GetDisplayName() const
+bool FRigControlElementDetails::IsGizmoEnabled() const
 {
-	if (HierarchyBeingCustomized != nullptr && ElementKeyBeingCustomized)
+	for(UDetailsViewWrapperObject* ObjectBeingCustomized : ObjectsBeingCustomized)
 	{
-		if(FRigControlElement* ControlElement = HierarchyBeingCustomized->Find<FRigControlElement>(ElementKeyBeingCustomized))
+		if(FRigControlElement* ControlElement = Cast<FRigControlElement>(ObjectBeingCustomized->GetContent<FRigBaseElement>()))
 		{
-			if (ControlElement->Settings.DisplayName.IsNone())
+			if(ControlElement->Settings.bGizmoEnabled)
 			{
-				return FText();
+				return true;
 			}
-			return FText::FromName(ControlElement->GetDisplayName());
 		}
 	}
-	return FText();
+	return false;
 }
 
-void FRigControlElementDetails::SetDisplayName(const FText& InNewText, ETextCommit::Type InCommitType, const TSharedRef<IPropertyUtilities> PropertyUtilities)
+bool FRigControlElementDetails::IsEnabled(ERigControlValueType InValueType) const
 {
-	if (HierarchyBeingCustomized != nullptr && ElementKeyBeingCustomized)
+	for(UDetailsViewWrapperObject* ObjectBeingCustomized : ObjectsBeingCustomized)
 	{
-		if(FRigControlElement* ControlElement = HierarchyBeingCustomized->Find<FRigControlElement>(ElementKeyBeingCustomized))
+		if(FRigControlElement* ControlElement = Cast<FRigControlElement>(ObjectBeingCustomized->GetContent<FRigBaseElement>()))
 		{
-			if(BlueprintBeingCustomized)
+			switch (InValueType)
 			{
-				BlueprintBeingCustomized->Hierarchy->Modify();
-			}
-
-			const FString NewDisplayName = InNewText.ToString().TrimStartAndEnd();
-			if (NewDisplayName.IsEmpty())
-			{
-				ControlElement->Settings.DisplayName = FName(NAME_None);
-			}
-			else
-			{
-				ControlElement->Settings.DisplayName = *NewDisplayName;
-			}
-
-			HierarchyBeingCustomized->Notify(ERigHierarchyNotification::ControlSettingChanged, ControlElement);
-
-			if (BlueprintBeingCustomized && BlueprintBeingCustomized->Hierarchy != HierarchyBeingCustomized)
-			{
-				if(FRigControlElement* OtherControlElement = BlueprintBeingCustomized->Hierarchy->Find<FRigControlElement>(ElementKeyBeingCustomized))
+				case ERigControlValueType::Minimum:
+				case ERigControlValueType::Maximum:
 				{
-					OtherControlElement->Settings.DisplayName = ControlElement->Settings.DisplayName;
-					BlueprintBeingCustomized->Hierarchy->Notify(ERigHierarchyNotification::ControlSettingChanged, OtherControlElement);
+					if(ControlElement->Settings.bLimitTranslation || ControlElement->Settings.bLimitRotation || ControlElement->Settings.bLimitScale)
+					{
+						return true;
+					}
+				}
+				default:
+				{
+					break;
 				}
 			}
 		}
 	}
+	return true;
 }
 
-bool FRigControlElementDetails::IsGizmoEnabled() const
+const TArray<TSharedPtr<FString>>& FRigControlElementDetails::GetGizmoNameList() const
 {
-	URigHierarchy* Hierarchy = HierarchyBeingCustomized;
-	if (Hierarchy != nullptr && ElementKeyBeingCustomized)
-	{
-		if(FRigControlElement* ControlElement = Hierarchy->Find<FRigControlElement>(ElementKeyBeingCustomized))
-		{
-			return ControlElement->Settings.bGizmoEnabled;
-		}
-	}
-	return false;
+	return GizmoNameList;
 }
 
-bool FRigControlElementDetails::IsEnabled(ERigControlValueType InValueType) const
+const TArray<TSharedPtr<FString>>& FRigControlElementDetails::GetControlTypeList() const
 {
-	switch (InValueType)
+	return ControlTypeList;
+}
+
+FText FRigControlElementDetails::GetDisplayName() const
+{
+	FName DisplayName(NAME_None);
+
+	for(int32 ObjectIndex = 0; ObjectIndex < ObjectsBeingCustomized.Num(); ObjectIndex++)
 	{
-		case ERigControlValueType::Minimum:
-		case ERigControlValueType::Maximum:
+		UDetailsViewWrapperObject* ObjectBeingCustomized = ObjectsBeingCustomized[ObjectIndex];
+		const FRigControlElement* ControlElement = ObjectBeingCustomized->GetContent<FRigControlElement>();
+		if(ObjectIndex == 0)
 		{
-			if (HierarchyBeingCustomized != nullptr && ElementKeyBeingCustomized)
-			{
-				if(FRigControlElement* ControlElement = HierarchyBeingCustomized->Find<FRigControlElement>(ElementKeyBeingCustomized))
-				{
-					return ControlElement->Settings.bLimitTranslation || ControlElement->Settings.bLimitRotation || ControlElement->Settings.bLimitScale;
-				}
-			}
-			return false;
+			DisplayName = ControlElement->Settings.DisplayName;
 		}
-		default:
+		else if(DisplayName != ControlElement->Settings.DisplayName)
 		{
-			break;
+			return ControlRigDetailsMultipleValues;
 		}
 	}
-	return true;
-}
 
-const TArray<TSharedPtr<FString>>& FRigControlElementDetails::GetGizmoNameList() const
-{
-	return GizmoNameList;
+	if(!DisplayName.IsNone())
+	{
+		return FText::FromName(DisplayName);
+	}
+	return FText();
 }
 
-const TArray<TSharedPtr<FString>>& FRigControlElementDetails::GetControlTypeList() const
+void FRigControlElementDetails::SetDisplayName(const FText& InNewText, ETextCommit::Type InCommitType)
 {
-	return ControlTypeList;
-}
+	const FName DisplayName = InNewText.IsEmpty() ? FName(NAME_None) : FName(*InNewText.ToString());
+	
+	for(int32 ObjectIndex = 0; ObjectIndex < ObjectsBeingCustomized.Num(); ObjectIndex++)
+	{
+		UDetailsViewWrapperObject* ObjectBeingCustomized = ObjectsBeingCustomized[ObjectIndex];
+		FRigControlElement* ControlElement = ObjectBeingCustomized->GetContent<FRigControlElement>();
 
-void FRigNullElementDetails::CustomizeDetails(IDetailLayoutBuilder& DetailBuilder)
-{
-	FRigTransformElementDetails::CustomizeDetails(DetailBuilder);
+		FRigControlSettings Settings = ControlElement->Settings;
+		Settings.DisplayName = DisplayName;
+
+		if(GetHierarchy())
+		{
+			GetHierarchy()->SetControlSettings(ControlElement->GetKey(), Settings);
+		}
+	}
 }
 
 #undef LOCTEXT_NAMESPACE
diff --git a/Engine/Plugins/Animation/ControlRig/Source/ControlRigEditor/Private/ControlRigElementDetails.h b/Engine/Plugins/Animation/ControlRig/Source/ControlRigEditor/Private/ControlRigElementDetails.h
--- a/Engine/Plugins/Animation/ControlRig/Source/ControlRigEditor/Private/ControlRigElementDetails.h
+++ b/Engine/Plugins/Animation/ControlRig/Source/ControlRigEditor/Private/ControlRigElementDetails.h
@@ -9,11 +9,13 @@
 #include "Rigs/RigHierarchy.h"
 #include "ControlRig.h"
 #include "ControlRigBlueprint.h"
+#include "DetailsViewWrapperObject.h"
 #include "Graph/ControlRigGraph.h"
 #include "Graph/SControlRigGraphPinNameListValueWidget.h"
 #include "Styling/SlateTypes.h"
 #include "IPropertyUtilities.h"
 #include "SSearchableComboBox.h"
+#include "DetailsViewWrapperObject.h"
 
 class IPropertyHandle;
 
@@ -118,38 +120,72 @@ protected:
 	void OnTransformChanged(FEditPropertyChain* InPropertyChain);
 };
 
-class FRigBaseElementDetails : public IDetailCustomization
+class FRigBaseElementDetails : public IPropertyTypeCustomization
 {
 public:
 
-	/** IDetailCustomization interface */
-	virtual void CustomizeDetails(IDetailLayoutBuilder& DetailBuilder) override;
+	/** IPropertyTypeCustomization interface */
+	virtual void CustomizeHeader(TSharedRef<class IPropertyHandle> InStructPropertyHandle, class FDetailWidgetRow& HeaderRow, IPropertyTypeCustomizationUtils& StructCustomizationUtils) override;
+	virtual void CustomizeChildren(TSharedRef<class IPropertyHandle> InStructPropertyHandle, class IDetailChildrenBuilder& StructBuilder, IPropertyTypeCustomizationUtils& StructCustomizationUtils) override;
 
-	FRigElementKey GetElementKey() const { return ElementKeyBeingCustomized; }
 	URigHierarchy* GetHierarchy() const { return HierarchyBeingCustomized; }
-	FText GetName() const { return FText::FromName(GetElementKey().Name); }
+	FRigElementKey GetElementKey() const;
+	FText GetName() const;
 	void SetName(const FText& InNewText, ETextCommit::Type InCommitType);
 
 	void OnStructContentsChanged(FProperty* InProperty, const TSharedRef<IPropertyUtilities> PropertyUtilities);
-	bool IsSetupModeEnabled() const; 
+	bool IsSetupModeEnabled() const;
+
+	TArray<FRigElementKey> GetElementKeys() const;
+
+	template<typename T>
+	TArray<T*> GetElementsInDetailsView() const
+	{
+		TArray<T*> Elements;
+		for(UDetailsViewWrapperObject* ObjectBeingCustomized : ObjectsBeingCustomized)
+		{
+			Elements.Add(Cast<T>(ObjectBeingCustomized->GetContent<FRigBaseElement>()));
+		}
+		return Elements;
+	}
+
+	template<typename T>
+	TArray<T*> GetElementsInHierarchy() const
+	{
+		TArray<T*> Elements;
+		if(HierarchyBeingCustomized)
+		{
+			for(UDetailsViewWrapperObject* ObjectBeingCustomized : ObjectsBeingCustomized)
+			{
+				Elements.Add(HierarchyBeingCustomized->Find<T>(ObjectBeingCustomized->GetContent<FRigBaseElement>()->GetKey()));
+			}
+		}
+		return Elements;
+	}
+
+	bool IsAnyElementOfType(ERigElementType InType) const;
+	bool IsAnyElementNotOfType(ERigElementType InType) const;
+	bool IsAnyControlOfType(ERigControlType InType) const;
+	bool IsAnyControlNotOfType(ERigControlType InType) const;
 
 protected:
 
-	FRigElementKey ElementKeyBeingCustomized;
 	UControlRigBlueprint* BlueprintBeingCustomized;
 	URigHierarchy* HierarchyBeingCustomized;
+	TArray<UDetailsViewWrapperObject*> ObjectsBeingCustomized;
 };
 
 class FRigTransformElementDetails : public FRigBaseElementDetails
 {
 public:
 
-	/** IDetailCustomization interface */
-	virtual void CustomizeDetails(IDetailLayoutBuilder& DetailBuilder) override;
+	/** IPropertyTypeCustomization interface */
+	virtual void CustomizeChildren(TSharedRef<class IPropertyHandle> InStructPropertyHandle, class IDetailChildrenBuilder& StructBuilder, IPropertyTypeCustomizationUtils& StructCustomizationUtils) override;
 
 protected:
-	static void OnTransformChanged(FEditPropertyChain* InPropertyChain, UControlRigBlueprint* InBlueprint);
 
+	bool IsCurrentLocalEnabled() const;
+	
 	FEditPropertyChain PoseInitialGlobal;
 	FEditPropertyChain PoseInitialLocal;
 	FEditPropertyChain PoseCurrentGlobal;
@@ -161,41 +197,34 @@ class FRigBoneElementDetails : public FRigTransformElementDetails
 public:
 
 	/** Makes a new instance of this detail layout class for a specific detail view requesting it */
-	static TSharedRef<IDetailCustomization> MakeInstance()
+	static TSharedRef<IPropertyTypeCustomization> MakeInstance()
 	{
 		return MakeShareable(new FRigBoneElementDetails);
 	}
-
-	/** IDetailCustomization interface */
-	virtual void CustomizeDetails( IDetailLayoutBuilder& DetailBuilder ) override;
-
-private:
-
-	TSharedPtr<FRigInfluenceEntryModifier> InfluenceModifier;
-	TSharedPtr<FStructOnScope> InfluenceModifierStruct;
 };
 
 class FRigControlElementDetails : public FRigTransformElementDetails
 {
 public:
 
-	/** Makes a new instance of this detail layout class for a specific detail view requesting it */
-	static TSharedRef<IDetailCustomization> MakeInstance()
+	// Makes a new instance of this detail layout class for a specific detail view requesting it
+	static TSharedRef<IPropertyTypeCustomization> MakeInstance()
 	{
 		return MakeShareable(new FRigControlElementDetails);
 	}
 
-	/** IDetailCustomization interface */
-	virtual void CustomizeDetails(IDetailLayoutBuilder& DetailBuilder) override;
+	/** IPropertyTypeCustomization interface */
+	virtual void CustomizeChildren(TSharedRef<class IPropertyHandle> InStructPropertyHandle, class IDetailChildrenBuilder& StructBuilder, IPropertyTypeCustomizationUtils& StructCustomizationUtils) override;
 
-	FText GetDisplayName() const;
-	void SetDisplayName(const FText& InNewText, ETextCommit::Type InCommitType, const TSharedRef<IPropertyUtilities> PropertyUtilities);
 	bool IsGizmoEnabled() const;
 	bool IsEnabled(ERigControlValueType InValueType) const;
 
 	const TArray<TSharedPtr<FString>>& GetGizmoNameList() const;
 	const TArray<TSharedPtr<FString>>& GetControlTypeList() const;
 
+	FText GetDisplayName() const;
+	void SetDisplayName(const FText& InNewText, ETextCommit::Type InCommitType);
+
 private:
 	TArray<TSharedPtr<FString>> GizmoNameList;
 	static TArray<TSharedPtr<FString>> ControlTypeList;
@@ -210,12 +239,9 @@ class FRigNullElementDetails : public FRigTransformElementDetails
 {
 public:
 
-	/** Makes a new instance of this detail layout class for a specific detail view requesting it */
-	static TSharedRef<IDetailCustomization> MakeInstance()
+	// Makes a new instance of this detail layout class for a specific detail view requesting it
+	static TSharedRef<IPropertyTypeCustomization> MakeInstance()
 	{
 		return MakeShareable(new FRigNullElementDetails);
 	}
-
-	/** IDetailCustomization interface */
-	virtual void CustomizeDetails(IDetailLayoutBuilder& DetailBuilder) override;
 };
diff --git a/Engine/Plugins/Animation/ControlRig/Source/ControlRigEditor/Private/Editor/ControlRigEditor.cpp b/Engine/Plugins/Animation/ControlRig/Source/ControlRigEditor/Private/Editor/ControlRigEditor.cpp
--- a/Engine/Plugins/Animation/ControlRig/Source/ControlRigEditor/Private/Editor/ControlRigEditor.cpp
+++ b/Engine/Plugins/Animation/ControlRig/Source/ControlRigEditor/Private/Editor/ControlRigEditor.cpp
@@ -193,14 +193,11 @@ struct FControlRigZoomLevelsContainer : public FZoomLevelsContainer
 	TArray<FControlRigZoomLevelEntry> ZoomLevels;
 };
 
-
 FControlRigEditor::FControlRigEditor()
 	: ControlRig(nullptr)
 	, ActiveController(nullptr)
 	, bControlRigEditorInitialized(false)
 	, bIsSettingObjectBeingDebugged(false)
-	, NodeDetailStruct(nullptr)
-	, NodeDetailName(NAME_None)
 	, bExecutionControlRig(true)
 	, bSetupModeEnabled(false)
 	, bFirstTimeSelecting(true)
@@ -216,6 +213,8 @@ FControlRigEditor::FControlRigEditor()
 
 FControlRigEditor::~FControlRigEditor()
 {
+	ClearDetailObject();
+
 	UControlRigBlueprint* RigBlueprint = GetControlRigBlueprint();
 	if (RigBlueprint)
 	{
@@ -234,7 +233,6 @@ FControlRigEditor::~FControlRigEditor()
 		RigBlueprint->OnBreakpointAdded().RemoveAll(this);
 		RigBlueprint->OnNodeDoubleClicked().RemoveAll(this);
 		RigBlueprint->OnGraphImported().RemoveAll(this);
-		RigBlueprint->OnPostEditChangeChainProperty().RemoveAll(this);
 		RigBlueprint->OnRequestLocalizeFunctionDialog().RemoveAll(this);
 		RigBlueprint->OnRequestBulkEditDialog().Unbind();
 		RigBlueprint->OnReportCompilerMessage().RemoveAll(this);
@@ -250,11 +248,6 @@ FControlRigEditor::~FControlRigEditor()
 #endif
 	}
 
-	if (NodeDetailBuffer.Num() > 0 && NodeDetailStruct != nullptr)
-	{
-		NodeDetailStruct->DestroyStruct(NodeDetailBuffer.GetData(), 1);
-	}
-
 	if (UWorld* PreviewWorld = GetPersonaToolkit()->GetPreviewScene()->GetWorld())
 	{
 		PreviewWorld->MarkObjectsPendingKill();
@@ -488,7 +481,6 @@ void FControlRigEditor::InitControlRigEditor(const EToolkitMode::Type Mode, cons
 
 		InControlRigBlueprint->OnNodeDoubleClicked().AddSP(this, &FControlRigEditor::OnNodeDoubleClicked);
 		InControlRigBlueprint->OnGraphImported().AddSP(this, &FControlRigEditor::OnGraphImported);
-		InControlRigBlueprint->OnPostEditChangeChainProperty().AddSP(this, &FControlRigEditor::OnBlueprintPropertyChainEvent);
 		InControlRigBlueprint->OnRequestLocalizeFunctionDialog().AddSP(this, &FControlRigEditor::OnRequestLocalizeFunctionDialog);
 		InControlRigBlueprint->OnRequestBulkEditDialog().BindSP(this, &FControlRigEditor::OnRequestBulkEditDialog);
 	}
@@ -1022,9 +1014,7 @@ void FControlRigEditor::ToggleSetupMode()
 {
 	bSetupModeEnabled = !bSetupModeEnabled;
 
-	FRigElementKey PreviousRigElementInDetailPanel = RigElementInDetailPanel;
 	TArray<FRigElementKey> PreviousSelection;
-
 	if (UControlRigBlueprint* RigBlueprint = Cast<UControlRigBlueprint>(GetBlueprintObj()))
 	{
 		if (RigBlueprint->bAutoRecompileVM)
@@ -1072,12 +1062,7 @@ void FControlRigEditor::ToggleSetupMode()
 		{
 			RigBlueprint->GetHierarchyController()->SetSelection(PreviousSelection);
 		}
-	}
-
-	if (PreviousRigElementInDetailPanel.IsValid())
-	{
-		ClearDetailObject();
-		SetDetailStruct(PreviousRigElementInDetailPanel);
+		SetDetailViewForRigElements();
 	}
 }
 
@@ -1340,153 +1325,82 @@ void FControlRigEditor::SetDetailObjects(const TArray<UObject*>& InObjects)
 		}
 	}
 
-	RigElementInDetailPanel = FRigElementKey();
-	StructToDisplay.Reset();
-	SKismetInspector::FShowDetailsOptions Options;
-	Options.bForceRefresh = true;
-	Inspector->ShowDetailsForObjects(InObjects, Options);
-}
-
-void FControlRigEditor::SetDetailObject(UObject* Obj)
-{
-	if(bSuspendDetailsPanelRefresh)
-	{
-		return;
-	}
+	TArray<UObject*> FilteredObjects;
 
-	if (URigVMUnitNode* UnitNode = Cast<URigVMUnitNode>(Obj))
+	for(UObject* InObject : InObjects)
 	{
-		ClearDetailObject();
-		NodeDetailStruct = UnitNode->GetScriptStruct();
-		NodeDetailName = UnitNode->GetFName();
-
-		if (NodeDetailStruct)
+		if (URigVMUnitNode* UnitNode = Cast<URigVMUnitNode>(InObject))
 		{
-			NodeDetailBuffer.AddUninitialized(NodeDetailStruct->GetStructureSize());
-			NodeDetailStruct->InitializeDefaultValue(NodeDetailBuffer.GetData());
-
-			FString StructDefaultValue = UnitNode->GetStructDefaultValue();
-			NodeDetailStruct->ImportText(*StructDefaultValue, NodeDetailBuffer.GetData(), nullptr, PPF_None, nullptr, NodeDetailStruct->GetName());
-
-			StructToDisplay = MakeShareable(new FStructOnScope(NodeDetailStruct, (uint8*)NodeDetailBuffer.GetData()));
-			if (StructToDisplay.IsValid())
-			{
-				StructToDisplay->SetPackage(GetControlRigBlueprint()->GetOutermost());
-			}
+			UScriptStruct* NodeDetailStruct = UnitNode->GetScriptStruct();
+			const FName NodeDetailName = UnitNode->GetFName();
 
-			// mark all input properties with editanywhere
-			for (TFieldIterator<FProperty> PropertyIt(NodeDetailStruct); PropertyIt; ++PropertyIt)
+			if (NodeDetailStruct)
 			{
-				FProperty* Property = *PropertyIt;
-				if (!Property->HasMetaData(TEXT("Input")))
-				{
-					continue;
-				}
+				// create a struct matching this node
+				TSharedPtr<FStructOnScope> StructInstance  = UnitNode->ConstructStructInstance(false);
 
-				// filter out execute pins
-				if (FStructProperty* StructProperty = CastField<FStructProperty>(Property))
-				{
-					if (StructProperty->Struct->IsChildOf(FRigVMExecuteContext::StaticStruct()))
-					{
-						continue;
-					}
-				}
-
-				bool bEditable = true;
-				if (URigVMPin* Pin = UnitNode->FindPin(Property->GetName()))
-				{
-					Pin = Pin->GetPinForLink();
-					if (Pin->GetDirection() == ERigVMPinDirection::Output)
-					{
-						if (Pin->GetTargetLinks().Num() > 0)
-						{
-							bEditable = false;
-						}
-					}
-					else if (Pin->GetDirection() == ERigVMPinDirection::Input ||
-						Pin->GetDirection() == ERigVMPinDirection::IO)
-					{
-						if (Pin->GetSourceLinks().Num() > 0)
-						{
-							bEditable = false;
-						}
-					}
-				}
-
-				Property->SetPropertyFlags(Property->GetPropertyFlags() | CPF_Edit);
-
-				if (bEditable)
-				{
-					Property->ClearPropertyFlags(CPF_EditConst);
-				}
-				else
-				{
-					Property->SetPropertyFlags(Property->GetPropertyFlags() | CPF_EditConst);
-				}
+				// create the wrapper object
+				UDetailsViewWrapperObject* WrapperObject = UDetailsViewWrapperObject::MakeInstance(
+					NodeDetailStruct, StructInstance->GetStructMemory(), UnitNode);
+				WrapperObject->GetWrappedPropertyChangedChainEvent().AddSP(this, &FControlRigEditor::OnWrappedPropertyChangedChainEvent);
+				
+				FilteredObjects.Add(WrapperObject);
+				continue;
 			}
-
-			FPropertyEditorModule& PropertyEditorModule = FModuleManager::LoadModuleChecked<FPropertyEditorModule>("PropertyEditor");
-			if (!PropertyEditorModule.GetClassNameToDetailLayoutNameMap().Contains(NodeDetailStruct->GetFName()))
+		}
+		else if (URigVMLibraryNode* LibraryNode = Cast<URigVMLibraryNode>(InObject))
+		{
+			UEdGraph* EdGraph = GetControlRigBlueprint()->GetEdGraph(LibraryNode->GetContainedGraph());
+			if (EdGraph)
 			{
-				PropertyEditorModule.RegisterCustomClassLayout(NodeDetailStruct->GetFName(), FOnGetDetailCustomizationInstance::CreateStatic(&FRigUnitDetails::MakeInstance));
+				FilteredObjects.AddUnique(EdGraph);
+				continue;
 			}
-
-			Inspector->ShowSingleStruct(StructToDisplay);
 		}
-		return;
-	}
-	else if (URigVMLibraryNode* LibraryNode = Cast<URigVMLibraryNode>(Obj))
-	{
-		UEdGraph* EdGraph = GetControlRigBlueprint()->GetEdGraph(LibraryNode->GetContainedGraph());
-		if (EdGraph)
+		else if (Cast<URigVMFunctionEntryNode>(InObject) || Cast<URigVMFunctionReturnNode>(InObject))
 		{
-			TArray<UObject*> Objects;
-			Objects.Add(EdGraph);
-			SetDetailObjects(Objects);
-			return;
+			UEdGraph* EdGraph = GetControlRigBlueprint()->GetEdGraph(CastChecked<URigVMNode>(InObject)->GetGraph());
+			if (EdGraph)
+			{
+				FilteredObjects.AddUnique(EdGraph);
+				continue;
+			}
 		}
-	}
-	else if (Cast<URigVMFunctionEntryNode>(Obj) || Cast<URigVMFunctionReturnNode>(Obj))
-	{
-		UEdGraph* EdGraph = GetControlRigBlueprint()->GetEdGraph(CastChecked<URigVMNode>(Obj)->GetGraph());
-		if (EdGraph)
+		else if (URigVMNode* ModelNode = Cast<URigVMNode>(InObject))
 		{
-			TArray<UObject*> Objects;
-			Objects.Add(EdGraph);
-			SetDetailObjects(Objects);
-			return;
+			UControlRigGraph* RigGraph = Cast<UControlRigGraph>(GetFocusedGraph());
+			if (UEdGraphNode* EdNode = RigGraph->FindNodeForModelNodeName(ModelNode->GetFName()))
+			{
+				FilteredObjects.AddUnique(EdNode);
+				continue;
+			}
 		}
+
+		FilteredObjects.Add(InObject);
 	}
 
-	TArray<UObject*> Objects;
-	if (Obj)
+	for(UObject* FilteredObject : FilteredObjects)
 	{
-		if (URigVMNode* ModelNode = Cast<URigVMNode>(Obj))
+		if(UDetailsViewWrapperObject* WrapperObject = Cast<UDetailsViewWrapperObject>(FilteredObject))
 		{
-			UControlRigGraph* RigGraph = Cast<UControlRigGraph>(GetFocusedGraph());
-			if (UEdGraphNode* EdNode = RigGraph->FindNodeForModelNodeName(ModelNode->GetFName()))
-			{
-				Objects.Add(EdNode);
-			}
+			WrapperObjects.Add(WrapperObject);
 		}
 	}
-	SetDetailObjects(Objects);
+
+	SKismetInspector::FShowDetailsOptions Options;
+	Options.bForceRefresh = true;
+	Inspector->ShowDetailsForObjects(FilteredObjects, Options);
 }
 
-void FControlRigEditor::SetDetailStruct(const FRigElementKey& InElement)
+void FControlRigEditor::SetDetailViewForRigElements()
 {
 	if(bSuspendDetailsPanelRefresh)
 	{
 		return;
 	}
 
-	if (RigElementInDetailPanel == InElement)
-	{
-		return;
-	}
-
 	ClearDetailObject();
-	
+
 	UControlRigBlueprint* RigBlueprint = Cast<UControlRigBlueprint>(GetBlueprintObj());
 	URigHierarchy* Hierarchy = RigBlueprint->Hierarchy;
 	
@@ -1498,56 +1412,146 @@ void FControlRigEditor::SetDetailStruct(const FRigElementKey& InElement)
 		}
 	}
 
-	FRigBaseElement* Element = Hierarchy->Find(InElement);
-	if (Element == nullptr)
+	TArray<UObject*> Objects;
+
+	TArray<FRigElementKey> CurrentSelection = RigBlueprint->Hierarchy->GetSelectedKeys();
+	for(const FRigElementKey& SelectedKey : CurrentSelection)
+	{
+		FRigBaseElement* Element = Hierarchy->Find(SelectedKey);
+		if (Element == nullptr)
+		{
+			continue;
+		}
+
+		UDetailsViewWrapperObject* WrapperObject = UDetailsViewWrapperObject::MakeInstance(Element->GetElementStruct(), (uint8*)Element, Hierarchy);
+		WrapperObject->GetWrappedPropertyChangedChainEvent().AddSP(this, &FControlRigEditor::OnWrappedPropertyChangedChainEvent);
+
+		Objects.Add(WrapperObject);
+	}
+	
+	SetDetailObjects(Objects);
+}
+
+void FControlRigEditor::SetDetailViewForGraph(URigVMGraph* InGraph)
+{
+	check(InGraph);
+	
+	if(bSuspendDetailsPanelRefresh)
 	{
 		return;
 	}
 
-	switch (InElement.Type)
+	ClearDetailObject();
+
+	TArray<UObject*> SelectedNodes;
+	TArray<FName> SelectedNodeNames = InGraph->GetSelectNodes();
+	for(FName SelectedNodeName : SelectedNodeNames)
 	{
-		case ERigElementType::Bone:
+		if(URigVMNode* Node = InGraph->FindNodeByName(SelectedNodeName))
 		{
-			StructToDisplay = MakeShareable(new FStructOnScope(FRigBoneElement::StaticStruct(), (uint8*)Element));
-			break;
+			SelectedNodes.Add(Node);
 		}
-		case ERigElementType::Control:
-		{
-			StructToDisplay = MakeShareable(new FStructOnScope(FRigControlElement::StaticStruct(), (uint8*)Element));
-			break;
-		}
-		case ERigElementType::Null:
-		{
-			StructToDisplay = MakeShareable(new FStructOnScope(FRigNullElement::StaticStruct(), (uint8*)Element));
-			break;
-		}
-		case ERigElementType::Curve:
-		{
-			StructToDisplay = MakeShareable(new FStructOnScope(FRigCurveElement::StaticStruct(), (uint8*)Element));
-			break;
-		}
-		case ERigElementType::RigidBody:
-		{
-			StructToDisplay = MakeShareable(new FStructOnScope(FRigRigidBodyElement::StaticStruct(), (uint8*)Element));
-			break;
-		}
-		case ERigElementType::Socket:
+	}
+
+	SetDetailObjects(SelectedNodes);
+}
+
+void FControlRigEditor::SetDetailViewForFocusedGraph()
+{
+	if(bSuspendDetailsPanelRefresh)
+	{
+		return;
+	}
+
+	URigVMGraph* Model = GetFocusedModel();
+	if(Model == nullptr)
+	{
+		return;
+	}
+
+	SetDetailViewForGraph(Model);
+}
+
+void FControlRigEditor::RefreshDetailView()
+{
+	if(DetailViewShowsAnyRigElement())
+	{
+		SetDetailViewForRigElements();
+	}
+	else if(DetailViewShowsAnyRigUnit())
+	{
+		SetDetailViewForFocusedGraph();
+	}
+}
+
+bool FControlRigEditor::DetailViewShowsAnyRigElement() const
+{
+	return DetailViewShowsStruct(FRigBaseElement::StaticStruct());
+}
+
+bool FControlRigEditor::DetailViewShowsAnyRigUnit() const
+{
+	return DetailViewShowsStruct(FRigUnit::StaticStruct());
+}
+
+bool FControlRigEditor::DetailViewShowsStruct(UScriptStruct* InStruct) const
+{
+	TArray< TWeakObjectPtr<UObject> > SelectedObjects = Inspector->GetSelectedObjects();
+	for (TWeakObjectPtr<UObject> SelectedObject : SelectedObjects)
+	{
+		if (SelectedObject.IsValid())
 		{
-			StructToDisplay = MakeShareable(new FStructOnScope(FRigSocketElement::StaticStruct(), (uint8*)Element));
-			break;
+			if(UDetailsViewWrapperObject* WrapperObject = Cast<UDetailsViewWrapperObject>(SelectedObject.Get()))
+			{
+				if(WrapperObject->GetWrappedStruct()->IsChildOf(InStruct))
+				{
+					return true;
+				}
+			}
 		}
-		default:
+	}
+	return false;
+}
+
+bool FControlRigEditor::DetailViewShowsRigElement(FRigElementKey InKey) const
+{
+	TArray< TWeakObjectPtr<UObject> > SelectedObjects = Inspector->GetSelectedObjects();
+	for (TWeakObjectPtr<UObject> SelectedObject : SelectedObjects)
+	{
+		if (SelectedObject.IsValid())
 		{
-			break;
+			if(UDetailsViewWrapperObject* WrapperObject = Cast<UDetailsViewWrapperObject>(SelectedObject.Get()))
+			{
+				if(WrapperObject->GetWrappedStruct()->IsChildOf(FRigBaseElement::StaticStruct()))
+				{
+					if(WrapperObject->GetContent<FRigBaseElement>()->GetKey() == InKey)
+					{
+						return true;
+					}
+				}
+			}
 		}
 	}
+	return false;
+}
 
-	RigElementInDetailPanel = InElement;
-	if (StructToDisplay.IsValid())
+bool FControlRigEditor::DetailViewShowsRigUnit(URigVMNode* InNode) const
+{
+	TArray< TWeakObjectPtr<UObject> > SelectedObjects = Inspector->GetSelectedObjects();
+	for (TWeakObjectPtr<UObject> SelectedObject : SelectedObjects)
 	{
-		StructToDisplay->SetPackage(GetControlRigBlueprint()->GetOutermost());
+		if (SelectedObject.IsValid())
+		{
+			if(UDetailsViewWrapperObject* WrapperObject = Cast<UDetailsViewWrapperObject>(SelectedObject.Get()))
+			{
+				if(WrapperObject->GetOuter() == InNode)
+				{
+					return true;
+				}
+			}
+		}
 	}
-	Inspector->ShowSingleStruct(StructToDisplay);
+	return false;
 }
 
 void FControlRigEditor::ClearDetailObject()
@@ -1557,16 +1561,18 @@ void FControlRigEditor::ClearDetailObject()
 		return;
 	}
 
-	RigElementInDetailPanel = FRigElementKey();
-
-	if (NodeDetailBuffer.Num() > 0 && NodeDetailStruct != nullptr)
+	for(TWeakObjectPtr<UDetailsViewWrapperObject> WrapperObjectPtr : WrapperObjects)
 	{
-		NodeDetailStruct->DestroyStruct(NodeDetailBuffer.GetData(), 1);
-		NodeDetailBuffer.Reset();
-		NodeDetailStruct = nullptr;
+		if(WrapperObjectPtr.IsValid())
+		{
+			UDetailsViewWrapperObject* WrapperObject = WrapperObjectPtr.Get();
+			WrapperObject->RemoveFromRoot();
+			WrapperObject->Rename(nullptr, GetTransientPackage(), REN_ForceNoResetLoaders | REN_DoNotDirty | REN_DontCreateRedirectors | REN_NonTransactional);
+			WrapperObject->MarkPendingKill();
+		}
 	}
-	NodeDetailName = NAME_None;
-
+	WrapperObjects.Reset();
+	
 	Inspector->ShowDetailsForObjects(TArray<UObject*>());
 	Inspector->ShowSingleStruct(TSharedPtr<FStructOnScope>());
 
@@ -1681,16 +1687,7 @@ void FControlRigEditor::Compile()
 		FString LastDebuggedObjectName = GetCustomDebugObjectLabel(RigBlueprint->GetObjectBeingDebugged());
 		RigBlueprint->SetObjectBeingDebugged(nullptr);
 
-		FRigElementKey SelectedKey = RigElementInDetailPanel;
-		TArray< TWeakObjectPtr<UObject> > SelectedObjects;
-		if (SelectedKey.IsValid())
-		{
-			ClearDetailObject();
-		}
-		else
-		{
-			SelectedObjects = Inspector->GetSelectedObjects();
-		}
+		TArray< TWeakObjectPtr<UObject> > SelectedObjects = Inspector->GetSelectedObjects();
 
 		if (ControlRig)
 		{
@@ -1780,20 +1777,9 @@ void FControlRigEditor::Compile()
 			UpdateDefaultValueForVariable(NewVariable, true);
 		}
 
-		if(SelectedKey.IsValid())
+		if (SelectedObjects.Num() > 0)
 		{
-			SetDetailStruct(SelectedKey);
-		}
-		else if (SelectedObjects.Num() > 0)
-		{
-			for (TWeakObjectPtr<UObject> SelectedObject : SelectedObjects)
-			{
-				if (SelectedObject.IsValid())
-				{
-					SetDetailObject(SelectedObject.Get());
-					break;
-				}
-			}
+			RefreshDetailView();
 		}
 
 		if (UControlRigEditorSettings::Get()->bResetControlTransformsOnCompile)
@@ -2213,7 +2199,7 @@ void FControlRigEditor::PasteNodes()
 void FControlRigEditor::PostUndo(bool bSuccess)
 {
 	IControlRigEditor::PostUndo(bSuccess);
-	EnsureValidRigElementInDetailPanel();
+	EnsureValidRigElementsInDetailPanel();
 
 	if (UControlRigBlueprint* RigBlueprint = Cast<UControlRigBlueprint>(GetBlueprintObj()))
 	{
@@ -2247,7 +2233,7 @@ void FControlRigEditor::PostUndo(bool bSuccess)
 void FControlRigEditor::PostRedo(bool bSuccess)
 {
 	IControlRigEditor::PostRedo(bSuccess);
-	EnsureValidRigElementInDetailPanel();
+	EnsureValidRigElementsInDetailPanel();
 
 	if (UControlRigBlueprint* RigBlueprint = Cast<UControlRigBlueprint>(GetBlueprintObj()))
 	{
@@ -2267,15 +2253,26 @@ void FControlRigEditor::PostRedo(bool bSuccess)
 	Compile();
 }
 
-void FControlRigEditor::EnsureValidRigElementInDetailPanel()
+void FControlRigEditor::EnsureValidRigElementsInDetailPanel()
 {
-	if (RigElementInDetailPanel.IsValid())
+	UControlRigBlueprint* ControlRigBP = GetControlRigBlueprint();
+	URigHierarchy* Hierarchy = ControlRigBP->Hierarchy; 
+
+	TArray< TWeakObjectPtr<UObject> > SelectedObjects = Inspector->GetSelectedObjects();
+	for (TWeakObjectPtr<UObject> SelectedObject : SelectedObjects)
 	{
-		if (UControlRigBlueprint * RigBlueprint = Cast<UControlRigBlueprint>(GetBlueprintObj()))
+		if (SelectedObject.IsValid())
 		{
-			if (RigBlueprint->Hierarchy->GetIndex(RigElementInDetailPanel) == INDEX_NONE)
+			if(UDetailsViewWrapperObject* WrapperObject = Cast<UDetailsViewWrapperObject>(SelectedObject.Get()))
 			{
-				ClearDetailObject();
+				if(WrapperObject->GetWrappedStruct()->IsChildOf(FRigBaseElement::StaticStruct()))
+				{
+					FRigElementKey Key = WrapperObject->GetContent<FRigBaseElement>()->GetKey();
+					if(!Hierarchy->Contains(Key))
+					{
+						ClearDetailObject();
+					}
+				}
 			}
 		}
 	}
@@ -2459,22 +2456,7 @@ void FControlRigEditor::HandleModifiedEvent(ERigVMGraphNotifType InNotifType, UR
 
 				if (GraphEd.IsValid() && Node != nullptr)
 				{
-					if (InNotifType == ERigVMGraphNotifType::NodeSelected)
-					{
-						SetDetailObject(Node);
-					}
-
-					// if we used to have a rig unit selected, clear the details panel
-					else if(InGraph->GetSelectNodes().Num() == 0)
-					{
-						if (StructToDisplay.IsValid())
-						{
-							if (StructToDisplay->GetStruct()->IsChildOf(FRigUnit::StaticStruct()))
-							{
-								ClearDetailObject();
-							}
-						}
-					}
+					SetDetailViewForGraph(Node->GetGraph());
 
 					if (!RigGraph->bIsSelecting)
 					{
@@ -2495,24 +2477,33 @@ void FControlRigEditor::HandleModifiedEvent(ERigVMGraphNotifType InNotifType, UR
 		{
 			URigVMPin* Pin = Cast<URigVMPin>(InSubject);
 
-			if (ControlRigBlueprint && NodeDetailBuffer.Num() > 0 && NodeDetailStruct != nullptr && !NodeDetailName.IsNone())
+			if(URigVMPin* RootPin = Pin->GetRootPin())
 			{
-				if (Pin->GetNode()->GetFName() == NodeDetailName)
+				const FString DefaultValue = RootPin->GetDefaultValue();
+				if(!DefaultValue.IsEmpty())
 				{
-					URigVMPin* RootPin = Pin->GetRootPin();
-					if (FProperty* Property = NodeDetailStruct->FindPropertyByName(RootPin->GetFName()))
+					// sync the value change with the unit(s) displayed 
+					TArray< TWeakObjectPtr<UObject> > SelectedObjects = Inspector->GetSelectedObjects();
+					for (TWeakObjectPtr<UObject> SelectedObject : SelectedObjects)
 					{
-						FString DefaultValue = RootPin->GetDefaultValue();
-						if (!DefaultValue.IsEmpty())
+						if (SelectedObject.IsValid())
 						{
-							uint8* PropertyValuePtr = (uint8*)NodeDetailBuffer.GetData();
-							PropertyValuePtr += Property->GetOffset_ReplaceWith_ContainerPtrToValuePtr();
-							Property->ImportText(*DefaultValue, PropertyValuePtr, PPF_None, nullptr);
+							if(UDetailsViewWrapperObject* WrapperObject = Cast<UDetailsViewWrapperObject>(SelectedObject.Get()))
+							{
+								if(WrapperObject->GetOuter() == Pin->GetNode())
+								{
+									uint8* PropertyStorage = WrapperObject->GetStructProperty()->ContainerPtrToValuePtr<uint8>(WrapperObject);
+									UScriptStruct* UnitStruct = WrapperObject->GetWrappedStruct();
+									FProperty* TargetProperty = UnitStruct->FindPropertyByName(RootPin->GetFName());
+									PropertyStorage += TargetProperty->GetOffset_ReplaceWith_ContainerPtrToValuePtr();
+									TargetProperty->ImportText(*DefaultValue, PropertyStorage, PPF_None, nullptr);
+								}
+							}
 						}
 					}
 				}
 			}
-
+	
 			break;
 		}
 		case ERigVMGraphNotifType::PinArraySizeChanged:
@@ -2520,15 +2511,12 @@ void FControlRigEditor::HandleModifiedEvent(ERigVMGraphNotifType InNotifType, UR
 		{
 			URigVMPin* Pin = Cast<URigVMPin>(InSubject);
 
-			if (ControlRigBlueprint && NodeDetailBuffer.Num() > 0 && NodeDetailStruct != nullptr && !NodeDetailName.IsNone())
+			if(Pin->GetNode()->IsSelected())
 			{
-				if (Pin->GetNode()->GetFName() == NodeDetailName)
-				{
-					// refresh the details panel
-					SetDetailObject(Pin->GetNode());
-				}
+				TArray<UObject*> Objects;
+				Objects.Add(Pin->GetNode());
+				SetDetailObjects(Objects);
 			}
-
 			break;
 		}
 		case ERigVMGraphNotifType::NodeRemoved:
@@ -2570,49 +2558,62 @@ void FControlRigEditor::HandleControlRigExecutedEvent(UControlRig* InControlRig,
 {
 	if (UControlRigBlueprint* ControlRigBP = GetControlRigBlueprint())
 	{
-		if(RigElementInDetailPanel.IsValid())
-		{
-			URigHierarchy* Hierarchy = ControlRigBP->Hierarchy; 
+		URigHierarchy* Hierarchy = ControlRigBP->Hierarchy; 
 
-			if(!bSetupModeEnabled)
+		if(!bSetupModeEnabled)
+		{
+			UControlRig* DebuggedControlRig = Cast<UControlRig>(ControlRigBP->GetObjectBeingDebugged());
+			if (DebuggedControlRig == nullptr)
 			{
-				UControlRig* DebuggedControlRig = Cast<UControlRig>(ControlRigBP->GetObjectBeingDebugged());
-				if (DebuggedControlRig == nullptr)
-				{
-					DebuggedControlRig = ControlRig;
-				}
-				if(DebuggedControlRig)
-				{
-					Hierarchy = DebuggedControlRig->GetHierarchy();
-				}
+				DebuggedControlRig = ControlRig;
 			}
-
-			check(Hierarchy);
-
-			FRigBaseElement* ElementInDetailPanel = Hierarchy->Find(RigElementInDetailPanel);
-			if(FRigTransformElement* TransformElement = Cast<FRigTransformElement>(ElementInDetailPanel))
+			if(DebuggedControlRig)
 			{
-				// compute all transforms
-				Hierarchy->GetTransform(TransformElement, ERigTransformType::CurrentGlobal);
-				Hierarchy->GetTransform(TransformElement, ERigTransformType::CurrentLocal);
-				Hierarchy->GetTransform(TransformElement, ERigTransformType::InitialGlobal);
-				Hierarchy->GetTransform(TransformElement, ERigTransformType::InitialLocal);
+				Hierarchy = DebuggedControlRig->GetHierarchy();
 			}
+		}
 
-			if(FRigControlElement* ControlElement = Cast<FRigControlElement>(ElementInDetailPanel))
+		TArray< TWeakObjectPtr<UObject> > SelectedObjects = Inspector->GetSelectedObjects();
+		for (TWeakObjectPtr<UObject> SelectedObject : SelectedObjects)
+		{
+			if (SelectedObject.IsValid())
 			{
-				// compute all transforms
-				Hierarchy->GetControlOffsetTransform(ControlElement, ERigTransformType::CurrentGlobal);
-				Hierarchy->GetControlOffsetTransform(ControlElement, ERigTransformType::CurrentLocal);
-				Hierarchy->GetControlOffsetTransform(ControlElement, ERigTransformType::InitialGlobal);
-				Hierarchy->GetControlOffsetTransform(ControlElement, ERigTransformType::InitialLocal);
-				Hierarchy->GetControlGizmoTransform(ControlElement, ERigTransformType::CurrentGlobal);
-				Hierarchy->GetControlGizmoTransform(ControlElement, ERigTransformType::CurrentLocal);
-				Hierarchy->GetControlGizmoTransform(ControlElement, ERigTransformType::InitialGlobal);
-				Hierarchy->GetControlGizmoTransform(ControlElement, ERigTransformType::InitialLocal);
+				if(UDetailsViewWrapperObject* WrapperObject = Cast<UDetailsViewWrapperObject>(SelectedObject.Get()))
+				{
+					UScriptStruct* Struct = WrapperObject->GetWrappedStruct(); 
+					if(Struct->IsChildOf(FRigBaseElement::StaticStruct()))
+					{
+						FRigElementKey Key = WrapperObject->GetContent<FRigBaseElement>()->GetKey();
+
+						FRigBaseElement* Element = Hierarchy->Find(Key);
+						if(FRigTransformElement* TransformElement = Cast<FRigTransformElement>(Element))
+						{
+							// compute all transforms
+							Hierarchy->GetTransform(TransformElement, ERigTransformType::CurrentGlobal);
+							Hierarchy->GetTransform(TransformElement, ERigTransformType::CurrentLocal);
+							Hierarchy->GetTransform(TransformElement, ERigTransformType::InitialGlobal);
+							Hierarchy->GetTransform(TransformElement, ERigTransformType::InitialLocal);
+						}
+
+						if(FRigControlElement* ControlElement = Cast<FRigControlElement>(Element))
+						{
+							// compute all transforms
+							Hierarchy->GetControlOffsetTransform(ControlElement, ERigTransformType::CurrentGlobal);
+							Hierarchy->GetControlOffsetTransform(ControlElement, ERigTransformType::CurrentLocal);
+							Hierarchy->GetControlOffsetTransform(ControlElement, ERigTransformType::InitialGlobal);
+							Hierarchy->GetControlOffsetTransform(ControlElement, ERigTransformType::InitialLocal);
+							Hierarchy->GetControlGizmoTransform(ControlElement, ERigTransformType::CurrentGlobal);
+							Hierarchy->GetControlGizmoTransform(ControlElement, ERigTransformType::CurrentLocal);
+							Hierarchy->GetControlGizmoTransform(ControlElement, ERigTransformType::InitialGlobal);
+							Hierarchy->GetControlGizmoTransform(ControlElement, ERigTransformType::InitialLocal);
+						}
+
+						Struct->CopyScriptStruct(WrapperObject->GetContent<FRigBaseElement>(), Element);
+					}
+				}
 			}
 		}
-		
+
 		if(ControlRigBP->RigGraphDisplaySettings.NodeRunLimit > 1)
 		{
 			if(UControlRig* DebuggedControlRig = Cast<UControlRig>(ControlRigBP->GetObjectBeingDebugged()))
@@ -4017,10 +4018,7 @@ void FControlRigEditor::NotifyPreChange(FProperty* PropertyAboutToChange)
 
 	if (UControlRigBlueprint* ControlRigBP = GetControlRigBlueprint())
 	{
-		if (RigElementInDetailPanel.IsValid())
-		{
-			ControlRigBP->Modify();
-		}
+		ControlRigBP->Modify();
 	}
 }
 
@@ -4050,64 +4048,6 @@ void FControlRigEditor::OnFinishedChangingProperties(const FPropertyChangedEvent
 		}
 	}
 
-	if (ControlRig && ControlRigBP && NodeDetailBuffer.Num() > 0 && NodeDetailStruct != nullptr && !NodeDetailName.IsNone())
-	{
-		const FRigUnit* RigUnitPtr = (const FRigUnit * )NodeDetailBuffer.GetData();
-		const uint8* MemberMemoryPtr = PropertyChangedEvent.MemberProperty->ContainerPtrToValuePtr<uint8>(RigUnitPtr);
-		FString DefaultValue = FRigVMStruct::ExportToFullyQualifiedText(PropertyChangedEvent.MemberProperty, MemberMemoryPtr);
-		if (!DefaultValue.IsEmpty())
-		{
-			FString PinPath = FString::Printf(TEXT("%s.%s"), *NodeDetailName.ToString(), *PropertyChangedEvent.MemberProperty->GetName());
-			GetFocusedController()->SetPinDefaultValue(PinPath, DefaultValue, true, true, false, true);
-		}
-	}
-
-	if (ControlRig && ControlRigBP && RigElementInDetailPanel)
-	{
-		UControlRig* DebuggedControlRig = Cast<UControlRig>(GetBlueprintObj()->GetObjectBeingDebugged());
-		check(DebuggedControlRig);
-
-		UScriptStruct* ScriptStruct = PropertyChangedEvent.Property->GetOwner<UScriptStruct>();
-		if (ScriptStruct == FRigControlSettings::StaticStruct())
-		{
-			URigHierarchy* Hierarchy = ControlRigBP->Hierarchy;
-
-			if(FRigControlElement* ControlElement = Hierarchy->Find<FRigControlElement>(RigElementInDetailPanel))
-			{
-				FRigControlElement* DebuggedControlElement = DebuggedControlRig->GetHierarchy()->Find<FRigControlElement>(RigElementInDetailPanel);
-				if(DebuggedControlElement)
-				{
-					if(bSetupModeEnabled)
-					{
-						DebuggedControlElement->Settings = ControlElement->Settings;
-					}
-					else
-					{
-						ControlElement->Settings = DebuggedControlElement->Settings;
-					}
-				}
-			
-				if (PropertyChangedEvent.Property->GetName() == TEXT("GizmoColor"))
-				{
-					ControlElement->Settings.GizmoColor.R = FMath::Clamp<float>(ControlElement->Settings.GizmoColor.R, 0.f, 1.f);
-					ControlElement->Settings.GizmoColor.G = FMath::Clamp<float>(ControlElement->Settings.GizmoColor.G, 0.f, 1.f);
-					ControlElement->Settings.GizmoColor.B = FMath::Clamp<float>(ControlElement->Settings.GizmoColor.B, 0.f, 1.f);
-					ControlElement->Settings.GizmoColor.A = FMath::Clamp<float>(ControlElement->Settings.GizmoColor.A, 0.f, 1.f);
-				}
-
-				Hierarchy->Notify(ERigHierarchyNotification::ControlSettingChanged, ControlElement);
-				if(DebuggedControlElement)
-				{
-					DebuggedControlRig->GetHierarchy()->Notify(ERigHierarchyNotification::ControlSettingChanged, DebuggedControlElement);
-				}
-                //ControlRigBP->PropagateHierarchyFromBPToInstances();
-			}
-
-			ControlRigBP->Modify();
-			ControlRigBP->MarkPackageDirty();
-		}
-	}
-
 	// this might be a property value for a variable
 	if(ControlRig && ControlRigBP)
 	{
@@ -4123,127 +4063,137 @@ void FControlRigEditor::OnFinishedChangingProperties(const FPropertyChangedEvent
 	}
 }
 
-void FControlRigEditor::OnBlueprintPropertyChainEvent(FPropertyChangedChainEvent& PropertyChangedChainEvent)
+void FControlRigEditor::OnWrappedPropertyChangedChainEvent(UDetailsViewWrapperObject* InWrapperObject, const FString& InPropertyPath, FPropertyChangedChainEvent& InPropertyChangedChainEvent)
 {
-	if(PropertyChangedChainEvent.PropertyChain.Num() < 4)
-	{
-		return;
-	}
-	
+	check(InWrapperObject);
+
+	const FString StoredStructString = TEXT("StoredStruct->");
+
 	UControlRigBlueprint* ControlRigBP = GetControlRigBlueprint();
-	if (ControlRig && ControlRigBP && RigElementInDetailPanel)
-	{
-		UControlRig* DebuggedControlRig = Cast<UControlRig>(GetBlueprintObj()->GetObjectBeingDebugged());
-		check(DebuggedControlRig);
 
-		const TDoubleLinkedList<FProperty*>::TDoubleLinkedListNode* HeadProperty = PropertyChangedChainEvent.PropertyChain.GetHead();
-		const TDoubleLinkedList<FProperty*>::TDoubleLinkedListNode* CurrentInitialProperty = HeadProperty->GetNextNode();
-		const TDoubleLinkedList<FProperty*>::TDoubleLinkedListNode* LocalGlobalProperty = CurrentInitialProperty->GetNextNode();
+	FString PropertyPath = InPropertyPath;
+	if(InWrapperObject->GetWrappedStruct()->IsChildOf(FRigBaseElement::StaticStruct()))
+	{
+		URigHierarchy* Hierarchy = CastChecked<URigHierarchy>(InWrapperObject->GetOuter());
+		check(PropertyPath.RemoveFromStart(StoredStructString));
 
-		const bool bIsInitial = CurrentInitialProperty->GetValue()->GetFName() == TEXT("Initial");
-		const bool bIsLocal = LocalGlobalProperty->GetValue()->GetFName() == TEXT("Local");
+		const FString PoseString = TEXT("Pose->");
+		const FString OffsetString = TEXT("Offset->");
+		const FString GizmoString = TEXT("Gizmo->");
+		const FString SettingsString = TEXT("Settings->");
 
-		ERigTransformType::Type TransformType = ERigTransformType::CurrentGlobal;
-		if(bIsInitial)
-		{
-			TransformType = ERigTransformType::MakeInitial(TransformType); 
-		}
-		
-		if(bIsLocal)
+		struct Local
 		{
-			TransformType = ERigTransformType::MakeLocal(TransformType); 
-		}
+			static ERigTransformType::Type GetTransformTypeFromPath(FString& PropertyPath)
+			{
+				const FString InitialString = TEXT("Initial->");
+				const FString CurrentString = TEXT("Current->");
+				const FString GlobalString = TEXT("Global->");
+				const FString LocalString = TEXT("Local->");
 
-		UControlRig* DefaultRig = Cast<UControlRig>(ControlRigBP->GeneratedClass->GetDefaultObject(true /* create if needed */));
-		URigHierarchy* DefaultHierarchy = DefaultRig->GetHierarchy();
-		URigHierarchy* DebuggedHierarchy = DebuggedControlRig->GetHierarchy();
-		URigHierarchy* BlueprintHierarchy = ControlRigBP->Hierarchy;
+				ERigTransformType::Type TransformType = ERigTransformType::CurrentLocal;
 
-		FRigBaseElement* SourceElement = bSetupModeEnabled ?
-			BlueprintHierarchy->Find(RigElementInDetailPanel) :
-			DebuggedHierarchy->Find(RigElementInDetailPanel);
+				if(PropertyPath.RemoveFromStart(InitialString))
+				{
+					TransformType = ERigTransformType::MakeInitial(TransformType);
+				}
+				else
+				{
+					check(PropertyPath.RemoveFromStart(CurrentString));
+					TransformType = ERigTransformType::MakeCurrent(TransformType);
+				}
 
-		if(SourceElement == nullptr)
-		{
-			return;
-		}
-		
-		FRigBaseElement* TargetElement = BlueprintHierarchy->Find(RigElementInDetailPanel);
-		if(TargetElement == nullptr)
-		{
-			return;
-		}
+				if(PropertyPath.RemoveFromStart(GlobalString))
+				{
+					TransformType = ERigTransformType::MakeGlobal(TransformType);
+				}
+				else
+				{
+					check(PropertyPath.RemoveFromStart(LocalString));
+					TransformType = ERigTransformType::MakeLocal(TransformType);
+				}
 
-		// force propagation of any change from the blueprint hierarchy
-		FRigHierarchyListenerGuard PropagateToDefault(BlueprintHierarchy, true, true);
-		FRigHierarchyListenerGuard PropagateToDebugged(DefaultHierarchy, true, true, DebuggedHierarchy);
-		
-		if(HeadProperty->GetValue() == FRigTransformElement::StaticStruct()->FindPropertyByName(TEXT("Pose")))
+				return TransformType;
+			}
+		};
+
+		bool bIsInitial = false;
+		FRigBaseElement* WrappedElement = InWrapperObject->GetContent<FRigBaseElement>();
+		if(PropertyPath.RemoveFromStart(PoseString))
 		{
-			if(FRigTransformElement* SourceTransformElement = Cast<FRigTransformElement>(SourceElement))
+			const ERigTransformType::Type TransformType = Local::GetTransformTypeFromPath(PropertyPath);
+			bIsInitial = bIsInitial || ERigTransformType::IsInitial(TransformType);
+			
+			if(ERigTransformType::IsInitial(TransformType) || bSetupModeEnabled)
 			{
-				if(FRigTransformElement* TargetTransformElement = Cast<FRigTransformElement>(TargetElement))
-				{
-					const FTransform Transform = SourceTransformElement->Pose.Get(TransformType);
-					if(ERigTransformType::IsLocal(TransformType) && TargetElement->IsA<FRigControlElement>())
-					{
-						FRigControlElement* TargetControlElement = Cast<FRigControlElement>(TargetTransformElement);
-						check(TargetControlElement);
+				Hierarchy = ControlRigBP->Hierarchy;
+			}
+
+			FRigTransformElement* TransformElement = Hierarchy->Find<FRigTransformElement>(WrappedElement->GetKey());
+			if(TransformElement == nullptr)
+			{
+				return;
+			}
+			
+			const FTransform Transform = CastChecked<FRigTransformElement>(WrappedElement)->Pose.Get(TransformType);
+
+			if(ERigTransformType::IsLocal(TransformType) && TransformElement->IsA<FRigControlElement>())
+			{
+				FRigControlElement* ControlElement = Cast<FRigControlElement>(TransformElement);
 						
-						FRigControlValue Value;
-						Value.SetFromTransform(Transform, TargetControlElement->Settings.ControlType, TargetControlElement->Settings.PrimaryAxis);
+				FRigControlValue Value;
+				Value.SetFromTransform(Transform, ControlElement->Settings.ControlType, ControlElement->Settings.PrimaryAxis);
 						
-						if(ERigTransformType::IsInitial(TransformType) || bSetupModeEnabled)
-						{
-							BlueprintHierarchy->SetControlValue(TargetControlElement, Value, ERigControlValueType::Initial, true, true);
-							BlueprintHierarchy->SetControlValue(TargetControlElement, Value, ERigControlValueType::Current, true, true);
-						}
-						else
-						{
-							BlueprintHierarchy->SetControlValue(TargetControlElement, Value, ERigControlValueType::Current, true, true);
-						}
-					}
-					else
-					{
-						if(ERigTransformType::IsInitial(TransformType) || bSetupModeEnabled)
-						{
-							BlueprintHierarchy->SetTransform(TargetTransformElement, Transform, ERigTransformType::MakeInitial(TransformType), true, true, true);
-							BlueprintHierarchy->SetTransform(TargetTransformElement, Transform, ERigTransformType::MakeCurrent(TransformType), true, true, true);
-						}
-						else
-						{
-							BlueprintHierarchy->SetTransform(TargetTransformElement, Transform, TransformType, true, true, true);
-						}
-					}
+				if(ERigTransformType::IsInitial(TransformType) || bSetupModeEnabled)
+				{
+					Hierarchy->SetControlValue(ControlElement, Value, ERigControlValueType::Initial, true, true);
 				}
+				Hierarchy->SetControlValue(ControlElement, Value, ERigControlValueType::Current, true, true);
+			}
+			else
+			{
+				Hierarchy->SetTransform(TransformElement, Transform, TransformType, true, true, false);
 			}
 		}
-
-		if(HeadProperty->GetValue() == FRigControlElement::StaticStruct()->FindPropertyByName(TEXT("Offset")))
+		else if(PropertyPath.RemoveFromStart(OffsetString))
 		{
-			if(FRigControlElement* SourceControlElement = Cast<FRigControlElement>(SourceElement))
+			FRigControlElement* ControlElement = ControlRigBP->Hierarchy->Find<FRigControlElement>(WrappedElement->GetKey());
+			if(ControlElement == nullptr)
 			{
-				if(FRigControlElement* TargetControlElement = Cast<FRigControlElement>(TargetElement))
-				{
-					const FTransform Transform = SourceControlElement->Offset.Get(TransformType);
-					// this will call the current version internally as well
-					BlueprintHierarchy->SetControlOffsetTransform(TargetControlElement, Transform, ERigTransformType::MakeInitial(TransformType), true, true, true);
-				}
+				return;
 			}
+			
+			ERigTransformType::Type TransformType = Local::GetTransformTypeFromPath(PropertyPath);
+			bIsInitial = bIsInitial || ERigTransformType::IsInitial(TransformType);
+			const FTransform Transform = CastChecked<FRigControlElement>(WrappedElement)->Offset.Get(TransformType);
+			
+			ControlRigBP->Hierarchy->SetControlOffsetTransform(ControlElement, Transform, ERigTransformType::MakeInitial(TransformType), true, true, false);
 		}
+		else if(PropertyPath.RemoveFromStart(GizmoString))
+		{
+			FRigControlElement* ControlElement = ControlRigBP->Hierarchy->Find<FRigControlElement>(WrappedElement->GetKey());
+			if(ControlElement == nullptr)
+			{
+				return;
+			}
 
-		if(HeadProperty->GetValue() == FRigControlElement::StaticStruct()->FindPropertyByName(TEXT("Gizmo")))
+			ERigTransformType::Type TransformType = Local::GetTransformTypeFromPath(PropertyPath);
+			bIsInitial = bIsInitial || ERigTransformType::IsInitial(TransformType);
+			const FTransform Transform = CastChecked<FRigControlElement>(WrappedElement)->Gizmo.Get(TransformType);
+			
+			ControlRigBP->Hierarchy->SetControlGizmoTransform(ControlElement, Transform, ERigTransformType::MakeInitial(TransformType), true, false);
+		}
+		else if(PropertyPath.RemoveFromStart(SettingsString))
 		{
-			if(FRigControlElement* SourceControlElement = Cast<FRigControlElement>(SourceElement))
+			const FRigControlSettings Settings = CastChecked<FRigControlElement>(WrappedElement)->Settings;
+
+			FRigControlElement* ControlElement = ControlRigBP->Hierarchy->Find<FRigControlElement>(WrappedElement->GetKey());
+			if(ControlElement == nullptr)
 			{
-				if(FRigControlElement* TargetControlElement = Cast<FRigControlElement>(TargetElement))
-				{
-					const FTransform Transform = SourceControlElement->Gizmo.Get(TransformType);
-					// this will call the current version internally as well
-					BlueprintHierarchy->SetControlGizmoTransform(TargetControlElement, Transform, ERigTransformType::MakeInitial(TransformType), false, true);
-					return;
-				}
+				return;
 			}
+
+			ControlRigBP->Hierarchy->SetControlSettings(ControlElement, Settings, true, false);
 		}
 
 		if(bSetupModeEnabled || bIsInitial)
@@ -4253,6 +4203,26 @@ void FControlRigEditor::OnBlueprintPropertyChainEvent(FPropertyChangedChainEvent
 			ControlRigBP->MarkPackageDirty();
 		}
 	}
+	
+	if(InWrapperObject->GetWrappedStruct()->IsChildOf(FRigUnit::StaticStruct()))
+	{
+		uint8* PropertyStorage = InWrapperObject->GetStructProperty()->ContainerPtrToValuePtr<uint8>(InWrapperObject);
+		UScriptStruct* UnitStruct = InWrapperObject->GetWrappedStruct();
+		FName RootPinName = InPropertyChangedChainEvent.PropertyChain.GetHead()->GetNextNode()->GetValue()->GetFName();
+		FProperty* TargetProperty = UnitStruct->FindPropertyByName(RootPinName);
+
+		URigVMNode* Node = CastChecked<URigVMNode>(InWrapperObject->GetOuter());
+
+		PropertyStorage += TargetProperty->GetOffset_ReplaceWith_ContainerPtrToValuePtr();
+		FString DefaultValue = FRigVMStruct::ExportToFullyQualifiedText(TargetProperty, PropertyStorage);
+
+		TargetProperty->ImportText(*DefaultValue, PropertyStorage, PPF_None, nullptr);
+		if (!DefaultValue.IsEmpty())
+		{
+			FString PinPath = FString::Printf(TEXT("%s.%s"), *Node->GetName(), *RootPinName.ToString());
+			GetFocusedController()->SetPinDefaultValue(PinPath, DefaultValue, true, true, false, true);
+		}
+	}
 }
 
 void FControlRigEditor::OnRequestLocalizeFunctionDialog(URigVMLibraryNode* InFunction, UControlRigBlueprint* InTargetBlueprint, bool bForce)
@@ -4381,29 +4351,7 @@ void FControlRigEditor::OnHierarchyChanged()
 			EditorSkelComp->RebuildDebugDrawSkeleton(); 
 		}
 
-		if (NodeDetailStruct != nullptr && NodeDetailBuffer.Num() > 0 && NodeDetailName != NAME_None)
-		{
-			if (URigVMNode* Node = ControlRigBP->Model->FindNode(NodeDetailName.ToString()))
-			{
-				SetDetailObject(Node);
-			}
-			else
-			{
-				ClearDetailObject();
-			}
-		}
-		else
-		{
-			// PropagateHierarchyFromBPToInstances() recreates the rig elements in a instance's hierarchy
-			// update the details viewer to observe the newly created elements
-			const FRigElementKey CurrentElementInDetailsPanel = RigElementInDetailPanel;
-			ClearDetailObject();
-
-			if (ControlRigBP->Hierarchy->Contains(CurrentElementInDetailsPanel))
-			{
-				SetDetailStruct(CurrentElementInDetailsPanel);
-			}
-		}
+		RefreshDetailView();
 	}
 	else
 	{
@@ -4607,7 +4555,7 @@ void FControlRigEditor::OnHierarchyModified(ERigHierarchyNotification InNotif, U
 
 			if (bSelected)
 			{
-				SetDetailStruct(InElement->GetKey());
+				SetDetailViewForRigElements();
 			}
 			else
 			{
@@ -4680,7 +4628,7 @@ void FControlRigEditor::OnHierarchyModified_AnyThread(ERigHierarchyNotification
 			}
 			case ERigHierarchyNotification::ControlSettingChanged:
 			{
-				if(Key == RigElementInDetailPanel)
+				if(DetailViewShowsRigElement(Key))
 				{
 					UControlRigBlueprint* RigBlueprint = GetControlRigBlueprint();
 					check(RigBlueprint);
@@ -4697,7 +4645,7 @@ void FControlRigEditor::OnHierarchyModified_AnyThread(ERigHierarchyNotification
 			}
 			case ERigHierarchyNotification::ControlGizmoTransformChanged:
 			{
-				if(Key == RigElementInDetailPanel)
+				if(DetailViewShowsRigElement(Key))
 				{
 					UControlRigBlueprint* RigBlueprint = GetControlRigBlueprint();
 					check(RigBlueprint);
@@ -5564,7 +5512,7 @@ void FControlRigEditor::OnGraphNodeClicked(UControlRigGraphNode* InNode)
 	{
 		if (InNode->IsSelectedInEditor())
 		{
-			SetDetailObject(InNode->GetModelNode());
+			SetDetailViewForGraph(InNode->GetModel());
 		}
 	}
 }
diff --git a/Engine/Plugins/Animation/ControlRig/Source/ControlRigEditor/Private/Editor/ControlRigEditor.h b/Engine/Plugins/Animation/ControlRig/Source/ControlRigEditor/Private/Editor/ControlRigEditor.h
--- a/Engine/Plugins/Animation/ControlRig/Source/ControlRigEditor/Private/Editor/ControlRigEditor.h
+++ b/Engine/Plugins/Animation/ControlRig/Source/ControlRigEditor/Private/Editor/ControlRigEditor.h
@@ -20,6 +20,7 @@
 #include "ScopedTransaction.h"
 #include "Graph/ControlRigGraphNode.h"
 #include "RigVMModel/RigVMController.h"
+#include "DetailsViewWrapperObject.h"
 
 class UControlRigBlueprint;
 class IPersonaToolkit;
@@ -124,7 +125,7 @@ public:
 	virtual void PostUndo(bool bSuccess) override;
 	virtual void PostRedo(bool bSuccess) override;
 
-	void EnsureValidRigElementInDetailPanel();
+	void EnsureValidRigElementsInDetailPanel();
 
 	virtual void OnStartWatchingPin();
 	virtual bool CanStartWatchingPin() const;
@@ -144,10 +145,15 @@ public:
 	UControlRigBlueprint* GetControlRigBlueprint() const;
 
 	void SetDetailObjects(const TArray<UObject*>& InObjects);
-
-	void SetDetailObject(UObject* Obj);
-
-	void SetDetailStruct(const FRigElementKey& InElement);
+	void SetDetailViewForRigElements();
+	void SetDetailViewForGraph(URigVMGraph* InGraph);
+	void SetDetailViewForFocusedGraph();
+	void RefreshDetailView();
+	bool DetailViewShowsAnyRigElement() const;
+	bool DetailViewShowsAnyRigUnit() const;
+	bool DetailViewShowsStruct(UScriptStruct* InStruct) const;
+	bool DetailViewShowsRigElement(FRigElementKey InKey) const;
+	bool DetailViewShowsRigUnit(URigVMNode* InNode) const;
 
 	void ClearDetailObject();
 
@@ -399,7 +405,7 @@ protected:
 	virtual void NotifyPostChange(const FPropertyChangedEvent& PropertyChangedEvent, FProperty* PropertyThatChanged) override;
 	/** delegate for changing property */
 	virtual void OnFinishedChangingProperties(const FPropertyChangedEvent& PropertyChangedEvent) override;
-	void OnBlueprintPropertyChainEvent(FPropertyChangedChainEvent& PropertyChangedChainEvent);
+	void OnWrappedPropertyChangedChainEvent(UDetailsViewWrapperObject* InWrapperObject, const FString& InPropertyPath, FPropertyChangedChainEvent& InPropertyChangedChainEvent);
 	void OnRequestLocalizeFunctionDialog(URigVMLibraryNode* InFunction, UControlRigBlueprint* InTargetBlueprint, bool bForce);
 	FRigVMController_BulkEditResult OnRequestBulkEditDialog(UControlRigBlueprint* InBlueprint, URigVMController* InController, URigVMLibraryNode* InFunction, ERigVMControllerBulkEditType InEditType);
 	void UpdateDefaultValueForVariable(FBPVariableDescription& InVariable, bool bUseCDO);
@@ -418,12 +424,6 @@ protected:
 	URigVMController* ActiveController;
 	bool bControlRigEditorInitialized;
 	bool bIsSettingObjectBeingDebugged;
-	FRigElementKey RigElementInDetailPanel;
-	TSharedPtr<FStructOnScope> StructToDisplay;
-
-	TArray<uint8, TAlignedHeapAllocator<16>> NodeDetailBuffer;
-	UScriptStruct* NodeDetailStruct;
-	FName NodeDetailName;
 
 	/** Currently executing ControlRig or not - later maybe this will change to enum for whatever different mode*/
 	bool bExecutionControlRig;
@@ -464,7 +464,8 @@ protected:
 	bool bSuspendDetailsPanelRefresh;
 
 	TMap<int32, int32> CapsuleToHierarchyIndex;
-
+	TArray<TWeakObjectPtr<UDetailsViewWrapperObject>> WrapperObjects;
+	
 	friend class FControlRigEditorMode;
 	friend class SControlRigStackView;
 	friend class SRigHierarchy;
diff --git a/Engine/Plugins/Animation/ControlRig/Source/ControlRigEditor/Private/Editor/DetailsViewWrapperObject.cpp b/Engine/Plugins/Animation/ControlRig/Source/ControlRigEditor/Private/Editor/DetailsViewWrapperObject.cpp
new file mode 100644
--- /dev/null
+++ b/Engine/Plugins/Animation/ControlRig/Source/ControlRigEditor/Private/Editor/DetailsViewWrapperObject.cpp
@@ -0,0 +1,217 @@
+// Copyright Epic Games, Inc. All Rights Reserved.
+
+#include "DetailsViewWrapperObject.h"
+#include "Units/RigUnit.h"
+#include "ControlRigElementDetails.h"
+
+TMap<UScriptStruct*, UClass*> UDetailsViewWrapperObject::StructToClass;
+TMap<UClass*, UScriptStruct*> UDetailsViewWrapperObject::ClassToStruct;
+
+UClass* UDetailsViewWrapperObject::GetClassForStruct(UScriptStruct* InStruct)
+{
+	check(InStruct != nullptr);
+
+	UClass** ExistingClass = StructToClass.Find(InStruct);
+	if(ExistingClass)
+	{
+		return *ExistingClass;
+	}
+
+	UClass *SuperClass = UDetailsViewWrapperObject::StaticClass();
+	const FName WrapperClassName(FString::Printf(TEXT("%s_WrapperObject"), *InStruct->GetStructCPPName()));
+
+	UClass* WrapperClass = NewObject<UClass>(
+		GetTransientPackage(),
+		WrapperClassName,
+		RF_Public | RF_Transactional
+		);
+
+	// make sure this doesn't get garbage collected
+	WrapperClass->AddToRoot();
+
+	// Eviscerate the class.
+	WrapperClass->PurgeClass(false);
+	WrapperClass->PropertyLink = SuperClass->PropertyLink;
+
+	WrapperClass->SetSuperStruct(SuperClass);
+	WrapperClass->ClassWithin = UObject::StaticClass();
+	WrapperClass->ClassConfigName = SuperClass->ClassConfigName;
+	WrapperClass->SetMetaData(TEXT("DisplayName"), *InStruct->GetDisplayNameText().ToString());
+
+	struct Local
+	{
+		static bool IsStructHashable(const UScriptStruct* InStructType)
+		{
+			if (InStructType->IsNative())
+			{
+				return InStructType->GetCppStructOps() && InStructType->GetCppStructOps()->HasGetTypeHash();
+			}
+			else
+			{
+				for (TFieldIterator<FProperty> It(InStructType); It; ++It)
+				{
+					if (CastField<FBoolProperty>(*It))
+					{
+						continue;
+					}
+					else if (!It->HasAllPropertyFlags(CPF_HasGetValueTypeHash))
+					{
+						return false;
+					}
+				}
+				return true;
+			}
+		}
+	};
+
+	FStructProperty* Property = new FStructProperty(WrapperClass, TEXT("StoredStruct"), RF_Public);
+	Property->Struct = InStruct;
+	if (Local::IsStructHashable(InStruct))
+	{
+		Property->SetPropertyFlags(CPF_HasGetValueTypeHash);
+	}
+	
+	// Make sure the variables show up in the details panel
+	Property->SetPropertyFlags(CPF_Edit | CPF_BlueprintVisible);
+	Property->SetMetaData(TEXT("ShowOnlyInnerProperties"), TEXT("true"));
+
+	// For rig units mark all inputs as edit
+	if(InStruct->IsChildOf(FRigUnit::StaticStruct()))
+	{
+		// mark all input properties with editanywhere
+		for (TFieldIterator<FProperty> PropertyIt(InStruct); PropertyIt; ++PropertyIt)
+		{
+			FProperty* ChildProperty = *PropertyIt;
+			if (!ChildProperty ->HasMetaData(TEXT("Input")))
+			{
+				continue;
+			}
+
+			// filter out execute pins
+			if (FStructProperty* StructProperty = CastField<FStructProperty>(ChildProperty))
+			{
+				if (StructProperty->Struct->IsChildOf(FRigVMExecuteContext::StaticStruct()))
+				{
+					continue;
+				}
+			}
+
+			// only do this for input pins
+			if(!ChildProperty->HasMetaData(TEXT("Input")))
+			{
+				continue;
+			}
+			
+			ChildProperty->SetPropertyFlags(ChildProperty->GetPropertyFlags() | CPF_Edit);
+		}
+
+		FPropertyEditorModule& PropertyEditorModule = FModuleManager::LoadModuleChecked<FPropertyEditorModule>("PropertyEditor");
+		if (!PropertyEditorModule.GetClassNameToDetailLayoutNameMap().Contains(InStruct->GetFName()))
+		{
+			PropertyEditorModule.RegisterCustomClassLayout(InStruct->GetFName(), FOnGetDetailCustomizationInstance::CreateStatic(&FRigUnitDetails::MakeInstance));
+		}
+	}
+
+	// Given it's only one property we are adding
+	// we can set the head of the linked list directly
+	WrapperClass->ChildProperties = Property;
+
+	// Update the class
+	WrapperClass->Bind();
+	WrapperClass->StaticLink(true);
+
+	StructToClass.Add(InStruct, WrapperClass);
+	ClassToStruct.Add(WrapperClass, InStruct);
+
+	UObject* CDO = WrapperClass->GetDefaultObject(true);
+	CDO->AddToRoot();
+
+	return WrapperClass;
+}
+
+UDetailsViewWrapperObject* UDetailsViewWrapperObject::MakeInstance(UScriptStruct* InStruct, uint8* InStructMemory, UObject* InOuter)
+{
+	check(InStruct != nullptr);
+
+	InOuter = InOuter == nullptr ? GetTransientPackage() : InOuter;
+	
+	UClass* WrapperClass = GetClassForStruct(InStruct);
+	if(WrapperClass == nullptr)
+	{
+		return nullptr;
+	}
+
+	UDetailsViewWrapperObject* Instance = NewObject<UDetailsViewWrapperObject>(InOuter, WrapperClass, NAME_None, RF_Public | RF_Transactional);
+	Instance->SetContent(InStructMemory, InStruct->GetStructureSize());
+	return Instance;
+}
+
+UScriptStruct* UDetailsViewWrapperObject::GetWrappedStruct() const
+{
+	return GetStructProperty()->Struct;
+}
+
+FStructProperty* UDetailsViewWrapperObject::GetStructProperty() const
+{
+	return CastFieldChecked<FStructProperty>(GetClass()->ChildProperties);
+}
+
+void UDetailsViewWrapperObject::SetContent(const uint8* InStructMemory, int32 InStructSize)
+{
+	UScriptStruct* Struct = GetWrappedStruct();
+	check(Struct);
+	check(InStructSize == Struct->GetStructureSize());
+
+	uint8* Content = GetStoredStructPtr();
+	Struct->CopyScriptStruct(Content, InStructMemory);
+}
+
+const uint8* UDetailsViewWrapperObject::GetContent()
+{
+	return GetStoredStructPtr();
+}
+
+void UDetailsViewWrapperObject::GetContent(uint8* OutStructMemory, int32 InStructSize)
+{
+	UScriptStruct* Struct = GetWrappedStruct();
+	check(Struct);
+	check(InStructSize == Struct->GetStructureSize());
+
+	const uint8* Content = GetContent();
+	Struct->CopyScriptStruct(OutStructMemory, Content);
+}
+
+void UDetailsViewWrapperObject::PostEditChangeChainProperty(FPropertyChangedChainEvent& PropertyChangedEvent)
+{
+	Super::PostEditChangeChainProperty(PropertyChangedEvent);
+
+	FString PropertyPath;
+
+	const TDoubleLinkedList<FProperty*>::TDoubleLinkedListNode* PropertyNode = PropertyChangedEvent.PropertyChain.GetHead();
+	do
+	{
+		const FString PropertyName = PropertyNode->GetValue()->GetNameCPP();
+		if(PropertyPath.IsEmpty())
+		{
+			PropertyPath = PropertyName;
+		}
+		else
+		{
+			PropertyPath = FString::Printf(TEXT("%s->%s"), *PropertyPath, *PropertyName);
+		}
+		PropertyNode = PropertyNode->GetNextNode();
+	}
+	while(PropertyNode);
+
+	WrappedPropertyChangedChainEvent.Broadcast(this, PropertyPath, PropertyChangedEvent);
+}
+
+uint8* UDetailsViewWrapperObject::GetStoredStructPtr()
+{
+	if(FStructProperty* StructProperty = CastField<FStructProperty>(GetClass()->ChildProperties))
+	{
+		return StructProperty->ContainerPtrToValuePtr<uint8>(this);
+	}
+	checkNoEntry();
+	return nullptr;
+}
diff --git a/Engine/Plugins/Animation/ControlRig/Source/ControlRigEditor/Private/Editor/DetailsViewWrapperObject.h b/Engine/Plugins/Animation/ControlRig/Source/ControlRigEditor/Private/Editor/DetailsViewWrapperObject.h
new file mode 100644
--- /dev/null
+++ b/Engine/Plugins/Animation/ControlRig/Source/ControlRigEditor/Private/Editor/DetailsViewWrapperObject.h
@@ -0,0 +1,50 @@
+// Copyright Epic Games, Inc. All Rights Reserved.
+
+#pragma once
+
+#include "CoreMinimal.h"
+#include "DetailsViewWrapperObject.generated.h"
+
+class UDetailsViewWrapperObject;
+
+DECLARE_EVENT_ThreeParams(UDetailsViewWrapperObject, FWrappedPropertyChangedChainEvent, UDetailsViewWrapperObject*, const FString&, FPropertyChangedChainEvent&);
+
+UCLASS()
+class UDetailsViewWrapperObject : public UObject
+{
+public:
+	GENERATED_BODY()
+
+	static UClass* GetClassForStruct(UScriptStruct* InStruct);
+	static UDetailsViewWrapperObject* MakeInstance(UScriptStruct* InStruct, uint8* InStructMemory, UObject* InOuter = nullptr);
+	UScriptStruct* GetWrappedStruct() const;
+	FStructProperty* GetStructProperty() const;
+
+	void SetContent(const uint8* InStructMemory, int32 InStructSize);
+	const uint8* GetContent();
+	void GetContent(uint8* OutStructMemory, int32 InStructSize);
+
+	template<typename T>
+	T* GetContent()
+	{
+		return (T*)GetStoredStructPtr();
+	}
+
+	template<typename T>
+	void SetContent(const T* InStructMemory)
+	{
+		SetCOntent((const uint8*)InStructMemory, T::StaticStruct()->GetStructureSize());
+	}
+
+	FWrappedPropertyChangedChainEvent& GetWrappedPropertyChangedChainEvent() { return WrappedPropertyChangedChainEvent; }
+	virtual void PostEditChangeChainProperty(FPropertyChangedChainEvent& PropertyChangedEvent) override;
+
+private:
+
+	uint8* GetStoredStructPtr();
+	
+	static TMap<UScriptStruct*, UClass*> StructToClass;
+	static TMap<UClass*, UScriptStruct*> ClassToStruct;
+	
+	FWrappedPropertyChangedChainEvent WrappedPropertyChangedChainEvent;
+};
diff --git a/Engine/Plugins/Animation/ControlRig/Source/ControlRigEditor/Private/Editor/SControlRigGizmoNameList.cpp b/Engine/Plugins/Animation/ControlRig/Source/ControlRigEditor/Private/Editor/SControlRigGizmoNameList.cpp
--- a/Engine/Plugins/Animation/ControlRig/Source/ControlRigEditor/Private/Editor/SControlRigGizmoNameList.cpp
+++ b/Engine/Plugins/Animation/ControlRig/Source/ControlRigEditor/Private/Editor/SControlRigGizmoNameList.cpp
@@ -11,9 +11,21 @@
 #include "ControlRig.h"
 
 void SControlRigGizmoNameList::Construct(const FArguments& InArgs, FRigControlElement* ControlElement, UControlRigBlueprint* InBlueprint)
+{
+	TArray<FRigControlElement*> ControlElements;
+	ControlElements.Add(ControlElement);
+	return Construct(InArgs, ControlElements, InBlueprint);
+}
+
+void SControlRigGizmoNameList::Construct(const FArguments& InArgs, TArray<FRigControlElement*> ControlElements, UControlRigBlueprint* InBlueprint)
 {
 	this->OnGetNameListContent = InArgs._OnGetNameListContent;
-	this->ControlKey = ControlElement->GetKey();
+	this->ControlKeys.Reset();
+
+	for(FRigControlElement* ControlElement : ControlElements)
+	{
+		this->ControlKeys.Add(ControlElement->GetKey());
+	}
 	this->Blueprint = InBlueprint;
 
 	SBox::Construct(SBox::FArguments());
@@ -58,33 +70,51 @@ const TArray<TSharedPtr<FString>>& SControlRigGizmoNameList::GetNameList() const
 
 FText SControlRigGizmoNameList::GetNameListText() const
 {
-	int32 ControlIndex = Blueprint->Hierarchy->GetIndex(ControlKey);
-	if (ControlIndex != INDEX_NONE)
+	FName FirstName = NAME_None;
+	FText Text;
+	for(int32 KeyIndex = 0; KeyIndex < ControlKeys.Num(); KeyIndex++)
 	{
-		return FText::FromName(Blueprint->Hierarchy->GetChecked<FRigControlElement>(ControlIndex)->Settings.GizmoName);
+		const int32 ControlIndex = Blueprint->Hierarchy->GetIndex(ControlKeys[KeyIndex]);
+		if (ControlIndex != INDEX_NONE)
+		{
+			const FName GizmoName = Blueprint->Hierarchy->GetChecked<FRigControlElement>(ControlIndex)->Settings.GizmoName; 
+			if(KeyIndex == 0)
+			{
+				Text = FText::FromName(GizmoName);
+				FirstName = GizmoName;
+			}
+			else if(FirstName != GizmoName)
+			{
+				static const FString MultipleValues = TEXT("Multiple Values");
+				Text = FText::FromString(MultipleValues);
+				break;
+			}
+		}
 	}
-	return FText();
+	return Text;
 }
 
 void SControlRigGizmoNameList::SetNameListText(const FText& NewTypeInValue, ETextCommit::Type /*CommitInfo*/)
 {
-	int32 ControlIndex = Blueprint->Hierarchy->GetIndex(ControlKey);
-	if (ControlIndex != INDEX_NONE)
-	{
-		FName NewName = *NewTypeInValue.ToString();
+	const FScopedTransaction Transaction(NSLOCTEXT("ControlRigEditor", "ChangeGizmoName", "Change Gizmo Name"));
 
-		bool bIsOnInstance = false;
-		URigHierarchy* Hierarchy = Blueprint->Hierarchy;
-
-		FRigControlElement* ControlElement = Hierarchy->Get<FRigControlElement>(ControlIndex);
-		if ((ControlElement != nullptr) && (ControlElement->Settings.GizmoName != NewName))
+	for(int32 KeyIndex = 0; KeyIndex < ControlKeys.Num(); KeyIndex++)
+	{
+		const int32 ControlIndex = Blueprint->Hierarchy->GetIndex(ControlKeys[KeyIndex]);
+		if (ControlIndex != INDEX_NONE)
 		{
-			const FScopedTransaction Transaction(NSLOCTEXT("ControlRigEditor", "ChangeGizmoName", "Change Gizmo Name"));
-			Hierarchy->Modify();
+			const FName NewName = *NewTypeInValue.ToString();
+			URigHierarchy* Hierarchy = Blueprint->Hierarchy;
+
+			FRigControlElement* ControlElement = Hierarchy->Get<FRigControlElement>(ControlIndex);
+			if ((ControlElement != nullptr) && (ControlElement->Settings.GizmoName != NewName))
+			{
+				Hierarchy->Modify();
 
-			FRigControlSettings Settings = ControlElement->Settings;
-			Settings.GizmoName = NewName;
-			Hierarchy->SetControlSettings(ControlElement, Settings, true, true);
+				FRigControlSettings Settings = ControlElement->Settings;
+				Settings.GizmoName = NewName;
+				Hierarchy->SetControlSettings(ControlElement, Settings, true, true);
+			}
 		}
 	}
 }
diff --git a/Engine/Plugins/Animation/ControlRig/Source/ControlRigEditor/Private/Editor/SControlRigGizmoNameList.h b/Engine/Plugins/Animation/ControlRig/Source/ControlRigEditor/Private/Editor/SControlRigGizmoNameList.h
--- a/Engine/Plugins/Animation/ControlRig/Source/ControlRigEditor/Private/Editor/SControlRigGizmoNameList.h
+++ b/Engine/Plugins/Animation/ControlRig/Source/ControlRigEditor/Private/Editor/SControlRigGizmoNameList.h
@@ -23,6 +23,7 @@ public:
 	SLATE_END_ARGS()
 
 	void Construct(const FArguments& InArgs, FRigControlElement* ControlElement, UControlRigBlueprint* InBlueprint);
+	void Construct(const FArguments& InArgs, TArray<FRigControlElement*> ControlElements, UControlRigBlueprint* InBlueprint);
 
 protected:
 
@@ -38,6 +39,6 @@ protected:
 	TSharedPtr<SControlRigGraphPinNameListValueWidget> NameListComboBox;
 	TArray<TSharedPtr<FString>> EmptyList;
 
-	FRigElementKey ControlKey;
+	TArray<FRigElementKey> ControlKeys;
 	UControlRigBlueprint* Blueprint;
 };
diff --git a/Engine/Plugins/Animation/ControlRig/Source/ControlRigEditor/Private/Editor/SRigHierarchy.cpp b/Engine/Plugins/Animation/ControlRig/Source/ControlRigEditor/Private/Editor/SRigHierarchy.cpp
--- a/Engine/Plugins/Animation/ControlRig/Source/ControlRigEditor/Private/Editor/SRigHierarchy.cpp
+++ b/Engine/Plugins/Animation/ControlRig/Source/ControlRigEditor/Private/Editor/SRigHierarchy.cpp
@@ -868,7 +868,7 @@ void SRigHierarchy::OnSelectionChanged(TSharedPtr<FRigTreeElement> Selection, ES
 	{
 		if(LastSelectedElement.IsValid())
 		{
-			ControlRigEditor.Pin()->SetDetailStruct(LastSelectedElement);
+			ControlRigEditor.Pin()->SetDetailViewForRigElements();
 		}
 		else
 		{
@@ -1364,7 +1364,7 @@ void SRigHierarchy::OnItemClicked(TSharedPtr<FRigTreeElement> InItem)
 	{
 		if (ControlRigEditor.IsValid())
 		{
-			ControlRigEditor.Pin()->SetDetailStruct(InItem->Key);
+			ControlRigEditor.Pin()->SetDetailViewForRigElements();
 		}
 
 		if (InItem->Key.Type == ERigElementType::Bone)
diff --git a/Engine/Source/Editor/PropertyEditor/Private/DetailCategoryBuilderImpl.cpp b/Engine/Source/Editor/PropertyEditor/Private/DetailCategoryBuilderImpl.cpp
--- a/Engine/Source/Editor/PropertyEditor/Private/DetailCategoryBuilderImpl.cpp
+++ b/Engine/Source/Editor/PropertyEditor/Private/DetailCategoryBuilderImpl.cpp
@@ -852,7 +852,8 @@ bool PassesInlineFilters(FDetailItemNode& Node)
 void FDetailCategoryImpl::GenerateNodesFromCustomizations(const TArray<FDetailLayoutCustomization>& InCustomizationList, FDetailNodeList& OutNodeList)
 {
 	TAttribute<bool> IsParentEnabled(this, &FDetailCategoryImpl::IsParentEnabled);
-	for (const FDetailLayoutCustomization& Customization : InCustomizationList)
+	TArray<FDetailLayoutCustomization> CustomizationList = InCustomizationList;
+	for (const FDetailLayoutCustomization& Customization : CustomizationList)
 	{
 		if (Customization.IsValidCustomization())
 		{
