Control Rig: Fix crash when creating blueprint based on Control Rig class

diff --git a/Engine/Plugins/Animation/ControlRig/Source/ControlRigEditor/Private/ControlRigBlueprintFactory.cpp b/Engine/Plugins/Animation/ControlRig/Source/ControlRigEditor/Private/ControlRigBlueprintFactory.cpp
--- a/Engine/Plugins/Animation/ControlRig/Source/ControlRigEditor/Private/ControlRigBlueprintFactory.cpp
+++ b/Engine/Plugins/Animation/ControlRig/Source/ControlRigEditor/Private/ControlRigBlueprintFactory.cpp
@@ -309,6 +309,11 @@ UControlRigBlueprint* UControlRigBlueprintFactory::CreateControlRigFromSkeletalM
 
 void UControlRigBlueprintFactory::CreateRigGraphIfRequired(UControlRigBlueprint* InBlueprint)
 {
+	if(InBlueprint == nullptr)
+	{
+		return;
+	}
+	
 	for(UEdGraph* EdGraph : InBlueprint->UbergraphPages)
 	{
 		if(EdGraph->IsA<UControlRigGraph>())
