Control Rig: Enable toggle gizmo transform edit mode from the rig hierarchy panel

diff --git a/Engine/Plugins/Animation/ControlRig/Source/ControlRigEditor/Private/ControlRigHierarchyCommands.cpp b/Engine/Plugins/Animation/ControlRig/Source/ControlRigEditor/Private/ControlRigHierarchyCommands.cpp
--- a/Engine/Plugins/Animation/ControlRig/Source/ControlRigEditor/Private/ControlRigHierarchyCommands.cpp
+++ b/Engine/Plugins/Animation/ControlRig/Source/ControlRigEditor/Private/ControlRigHierarchyCommands.cpp
@@ -36,6 +36,7 @@ void FControlRigHierarchyCommands::RegisterCommands()
 	UI_COMMAND(ShowRigidBodies, "Show RigidBodies", "Whether to show or hide rigidbodies", EUserInterfaceActionType::ToggleButton, FInputChord());
 	UI_COMMAND(ShowSockets, "Show Sockets", "Whether to show or hide sockets", EUserInterfaceActionType::ToggleButton, FInputChord());
 	UI_COMMAND(ShowDynamicHierarchy, "Dynamic Hierarchy", "Whether to show the static or dynamic hierarchy", EUserInterfaceActionType::ToggleButton, FInputChord());
+	UI_COMMAND(ToggleGizmoTransformEdit, "Toggle Gizmo Transform Edit", "Toggle Editing Selected Control's Gizmo Transform", EUserInterfaceActionType::ToggleButton, FInputChord(EKeys::Period, EModifierKey::Control)); 
 }
 
 #undef LOCTEXT_NAMESPACE
diff --git a/Engine/Plugins/Animation/ControlRig/Source/ControlRigEditor/Private/ControlRigHierarchyCommands.h b/Engine/Plugins/Animation/ControlRig/Source/ControlRigEditor/Private/ControlRigHierarchyCommands.h
--- a/Engine/Plugins/Animation/ControlRig/Source/ControlRigEditor/Private/ControlRigHierarchyCommands.h
+++ b/Engine/Plugins/Animation/ControlRig/Source/ControlRigEditor/Private/ControlRigHierarchyCommands.h
@@ -108,6 +108,9 @@ public:
 	/* Show dynamic hierarchy */
 	TSharedPtr< FUICommandInfo > ShowDynamicHierarchy;
 
+	/** Toggle Gizmo Transform Edit*/
+	TSharedPtr< FUICommandInfo > ToggleGizmoTransformEdit;
+	
 	/**
 	 * Initialize commands
 	 */
diff --git a/Engine/Plugins/Animation/ControlRig/Source/ControlRigEditor/Private/EditMode/ControlRigEditMode.h b/Engine/Plugins/Animation/ControlRig/Source/ControlRigEditor/Private/EditMode/ControlRigEditMode.h
--- a/Engine/Plugins/Animation/ControlRig/Source/ControlRigEditor/Private/EditMode/ControlRigEditMode.h
+++ b/Engine/Plugins/Animation/ControlRig/Source/ControlRigEditor/Private/EditMode/ControlRigEditMode.h
@@ -205,9 +205,13 @@ private:
 	/** Reset Gizmo Size */
 	void ResetGizmoSize();
 
+public:
+	
 	/** Toggle Gizmo Transform Edit*/
 	void ToggleGizmoTransformEdit();
 
+private:
+	
 	/** The hotkey text is passed to a viewport notification to inform users how to toggle gizmo edit*/
 	FText GetToggleGizmoTransformEditHotKey() const;
 
diff --git a/Engine/Plugins/Animation/ControlRig/Source/ControlRigEditor/Private/Editor/SRigHierarchy.cpp b/Engine/Plugins/Animation/ControlRig/Source/ControlRigEditor/Private/Editor/SRigHierarchy.cpp
--- a/Engine/Plugins/Animation/ControlRig/Source/ControlRigEditor/Private/Editor/SRigHierarchy.cpp
+++ b/Engine/Plugins/Animation/ControlRig/Source/ControlRigEditor/Private/Editor/SRigHierarchy.cpp
@@ -615,6 +615,14 @@ void SRigHierarchy::BindCommands()
 		}),
 		FCanExecuteAction(),
 		FIsActionChecked::CreateLambda([this]() { return bShowDynamicHierarchy; }));
+
+	CommandList->MapAction(
+		Commands.ToggleGizmoTransformEdit,
+		FExecuteAction::CreateLambda([this]()
+		{
+			ControlRigEditor.Pin()->GetEditMode()->ToggleGizmoTransformEdit();
+		}));
+	
 }
 
 FReply SRigHierarchy::OnKeyDown(const FGeometry& MyGeometry, const FKeyEvent& InKeyEvent)
