Control Rig: Python Command Log when control moved

diff --git a/Engine/Plugins/Animation/ControlRig/Source/ControlRig/Private/ControlRig.cpp b/Engine/Plugins/Animation/ControlRig/Source/ControlRig/Private/ControlRig.cpp
--- a/Engine/Plugins/Animation/ControlRig/Source/ControlRig/Private/ControlRig.cpp
+++ b/Engine/Plugins/Animation/ControlRig/Source/ControlRig/Private/ControlRig.cpp
@@ -1349,7 +1349,7 @@ FTransform UControlRig::GetControlGlobalTransform(const FName& InControlName) co
 	return DynamicHierarchy->GetGlobalTransform(FRigElementKey(InControlName, ERigElementType::Control), false);
 }
 
-bool UControlRig::SetControlGlobalTransform(const FName& InControlName, const FTransform& InGlobalTransform, bool bNotify, const FRigControlModifiedContext& Context, bool bSetupUndo)
+bool UControlRig::SetControlGlobalTransform(const FName& InControlName, const FTransform& InGlobalTransform, bool bNotify, const FRigControlModifiedContext& Context, bool bSetupUndo, bool bPrintPythonCommands)
 {
 	FTransform GlobalTransform = InGlobalTransform;
 	if (IsSetupModeEnabled())
@@ -1367,7 +1367,7 @@ bool UControlRig::SetControlGlobalTransform(const FName& InControlName, const FT
 		}
 	}
 
-	SetControlValue(InControlName, Value, bNotify, Context, bSetupUndo);
+	SetControlValue(InControlName, Value, bNotify, Context, bSetupUndo, bPrintPythonCommands);
 	return true;
 }
 
diff --git a/Engine/Plugins/Animation/ControlRig/Source/ControlRig/Private/Rigs/RigHierarchy.cpp b/Engine/Plugins/Animation/ControlRig/Source/ControlRig/Private/Rigs/RigHierarchy.cpp
--- a/Engine/Plugins/Animation/ControlRig/Source/ControlRig/Private/Rigs/RigHierarchy.cpp
+++ b/Engine/Plugins/Animation/ControlRig/Source/ControlRig/Private/Rigs/RigHierarchy.cpp
@@ -2,6 +2,7 @@
 
 #include "Rigs/RigHierarchy.h"
 #include "ControlRig.h"
+#include "RigVMPythonUtils.h"
 
 #include "Rigs/RigHierarchyElements.h"
 #include "Rigs/RigHierarchyController.h"
@@ -1382,7 +1383,7 @@ FTransform URigHierarchy::GetTransform(FRigTransformElement* InTransformElement,
 	return InTransformElement->Pose.Get(InTransformType);
 }
 
-void URigHierarchy::SetTransform(FRigTransformElement* InTransformElement, const FTransform& InTransform, const ERigTransformType::Type InTransformType, bool bAffectChildren, bool bSetupUndo, bool bForce)
+void URigHierarchy::SetTransform(FRigTransformElement* InTransformElement, const FTransform& InTransform, const ERigTransformType::Type InTransformType, bool bAffectChildren, bool bSetupUndo, bool bForce, bool bPrintPythonCommands)
 {
 	if(InTransformElement == nullptr)
 	{
@@ -1458,6 +1459,33 @@ void URigHierarchy::SetTransform(FRigTransformElement* InTransformElement, const
 			}
 		}
 	}
+
+	if (bPrintPythonCommands)
+	{
+		FString MethodName;
+		switch (InTransformType)
+		{
+			case ERigTransformType::InitialLocal: 
+			case ERigTransformType::CurrentLocal:
+			{
+				MethodName = TEXT("set_local_transform");
+				break;
+			}
+			case ERigTransformType::InitialGlobal: 
+			case ERigTransformType::CurrentGlobal:
+			{
+				MethodName = TEXT("set_global_transform");
+				break;
+			}
+		}
+
+		RigVMPythonUtils::Print(FString::Printf(TEXT("hierarchy.%s(%s, %s, %s, %s)"),
+			*MethodName,
+			*InTransformElement->GetKey().ToPythonString(),
+			*RigVMPythonUtils::TransformToPythonString(InTransform),
+			(InTransformType == ERigTransformType::InitialGlobal || InTransformType == ERigTransformType::InitialLocal) ? TEXT("True") : TEXT("False"),
+			(bAffectChildren) ? TEXT("True") : TEXT("False")));
+	}
 #endif
 }
 
@@ -1489,7 +1517,7 @@ FTransform URigHierarchy::GetControlOffsetTransform(FRigControlElement* InContro
 }
 
 void URigHierarchy::SetControlOffsetTransform(FRigControlElement* InControlElement, const FTransform& InTransform,
-	const ERigTransformType::Type InTransformType, bool bAffectChildren, bool bSetupUndo, bool bForce)
+                                              const ERigTransformType::Type InTransformType, bool bAffectChildren, bool bSetupUndo, bool bForce, bool bPrintPythonCommands)
 {
 	if(InControlElement == nullptr)
 	{
@@ -1555,6 +1583,16 @@ void URigHierarchy::SetControlOffsetTransform(FRigControlElement* InControlEleme
 			}
 		}
 	}
+
+	if (bPrintPythonCommands)
+	{
+		//FRigElementKey InKey, FTransform InTransform, bool bInitial = false, bool bAffectChildren = true, bool bSetupUndo = false, bool bPrintPythonCommands = false)
+		RigVMPythonUtils::Print(FString::Printf(TEXT("hierarchy.set_control_offset_transform(%s, %s, %s, %s)"),
+			*InControlElement->GetKey().ToPythonString(),
+			*RigVMPythonUtils::TransformToPythonString(InTransform),
+			(ERigTransformType::IsInitial(InTransformType)) ? TEXT("True") : TEXT("False"),
+			(bAffectChildren) ? TEXT("True") : TEXT("False")));
+	}
 #endif
 }
 
@@ -2180,7 +2218,7 @@ FRigControlValue URigHierarchy::GetControlValue(FRigControlElement* InControlEle
 	return Value;
 }
 
-void URigHierarchy::SetControlValue(FRigControlElement* InControlElement, const FRigControlValue& InValue, ERigControlValueType InValueType, bool bSetupUndo, bool bForce)
+void URigHierarchy::SetControlValue(FRigControlElement* InControlElement, const FRigControlValue& InValue, ERigControlValueType InValueType, bool bSetupUndo, bool bForce, bool bPrintPythonCommands)
 {
 	using namespace ERigTransformType;
 
@@ -2194,16 +2232,17 @@ void URigHierarchy::SetControlValue(FRigControlElement* InControlElement, const
 				InControlElement->Settings.ApplyLimits(Value);
 					
 				SetTransform(
-                    InControlElement,
-                    Value.GetAsTransform(
-                        InControlElement->Settings.ControlType,
-                        InControlElement->Settings.PrimaryAxis
-                    ),
-                    CurrentLocal,
-                    true,
-                    bSetupUndo,
-                    bForce
-                ); 
+					InControlElement,
+					Value.GetAsTransform(
+						InControlElement->Settings.ControlType,
+						InControlElement->Settings.PrimaryAxis
+					),
+					CurrentLocal,
+					true,
+					bSetupUndo,
+					bForce,
+					bPrintPythonCommands
+				); 
 				break;
 			}
 			case ERigControlValueType::Initial:
@@ -2212,16 +2251,17 @@ void URigHierarchy::SetControlValue(FRigControlElement* InControlElement, const
 				InControlElement->Settings.ApplyLimits(Value);
 
 				SetTransform(
-                    InControlElement,
-                    Value.GetAsTransform(
-                        InControlElement->Settings.ControlType,
-                        InControlElement->Settings.PrimaryAxis
-                    ),
-                    InitialLocal,
-                    true,
-                    bSetupUndo,
-                    bForce
-                ); 
+					InControlElement,
+					Value.GetAsTransform(
+						InControlElement->Settings.ControlType,
+						InControlElement->Settings.PrimaryAxis
+					),
+					InitialLocal,
+					true,
+					bSetupUndo,
+					bForce,
+					bPrintPythonCommands
+				); 
 				break;
 			}
 			case ERigControlValueType::Minimum:
diff --git a/Engine/Plugins/Animation/ControlRig/Source/ControlRig/Private/Rigs/RigHierarchyController.cpp b/Engine/Plugins/Animation/ControlRig/Source/ControlRig/Private/Rigs/RigHierarchyController.cpp
--- a/Engine/Plugins/Animation/ControlRig/Source/ControlRig/Private/Rigs/RigHierarchyController.cpp
+++ b/Engine/Plugins/Animation/ControlRig/Source/ControlRig/Private/Rigs/RigHierarchyController.cpp
@@ -2129,7 +2129,7 @@ bool URigHierarchyController::RemoveAllParents(FRigElementKey InChild, bool bMai
 	
 	if (bRemoved && bPrintPythonCommand)
 	{
-		RigVMPythonUtils::Print(FString::Printf(TEXT("hierarchy_controller.remove_parent(%s, %s)"),
+		RigVMPythonUtils::Print(FString::Printf(TEXT("hierarchy_controller.remove_all_parents(%s, %s)"),
 			*InChild.ToPythonString(),
 			(bMaintainGlobalTransform) ? TEXT("True") : TEXT("False")));
 	}
diff --git a/Engine/Plugins/Animation/ControlRig/Source/ControlRig/Public/ControlRig.h b/Engine/Plugins/Animation/ControlRig/Source/ControlRig/Public/ControlRig.h
--- a/Engine/Plugins/Animation/ControlRig/Source/ControlRig/Public/ControlRig.h
+++ b/Engine/Plugins/Animation/ControlRig/Source/ControlRig/Public/ControlRig.h
@@ -277,9 +277,9 @@ public:
 	// Sets the relative value of a Control
 	template<class T>
 	FORCEINLINE_DEBUGGABLE void SetControlValue(const FName& InControlName, T InValue, bool bNotify = true,
-		const FRigControlModifiedContext& Context = FRigControlModifiedContext(), bool bSetupUndo = true)
+		const FRigControlModifiedContext& Context = FRigControlModifiedContext(), bool bSetupUndo = true, bool bPrintPythonCommnds = false)
 	{
-		SetControlValueImpl(InControlName, FRigControlValue::Make<T>(InValue), bNotify, Context, bSetupUndo);
+		SetControlValueImpl(InControlName, FRigControlValue::Make<T>(InValue), bNotify, Context, bSetupUndo, bPrintPythonCommnds);
 	}
 
 	// Returns the value of a Control
@@ -291,7 +291,7 @@ public:
 
 	// Sets the relative value of a Control
 	FORCEINLINE_DEBUGGABLE virtual void SetControlValueImpl(const FName& InControlName, const FRigControlValue& InValue, bool bNotify = true,
-		const FRigControlModifiedContext& Context = FRigControlModifiedContext(), bool bSetupUndo = true)
+		const FRigControlModifiedContext& Context = FRigControlModifiedContext(), bool bSetupUndo = true, bool bPrintPythonCommnds = false)
 	{
 		const FRigElementKey Key(InControlName, ERigElementType::Control);
 
@@ -301,7 +301,7 @@ public:
 			return;
 		}
 
-		DynamicHierarchy->SetControlValue(ControlElement, InValue, ERigControlValueType::Current, bSetupUndo);
+		DynamicHierarchy->SetControlValue(ControlElement, InValue, ERigControlValueType::Current, bSetupUndo, false, bPrintPythonCommnds);
 
 		if (bNotify && OnControlModified.IsBound())
 		{
@@ -309,7 +309,7 @@ public:
 		}
 	}
 
-	bool SetControlGlobalTransform(const FName& InControlName, const FTransform& InGlobalTransform, bool bNotify = true, const FRigControlModifiedContext& Context = FRigControlModifiedContext(), bool bSetupUndo = true);
+	bool SetControlGlobalTransform(const FName& InControlName, const FTransform& InGlobalTransform, bool bNotify = true, const FRigControlModifiedContext& Context = FRigControlModifiedContext(), bool bSetupUndo = true, bool bPrintPythonCommands = false);
 
 	virtual FRigControlValue GetControlValueFromGlobalTransform(const FName& InControlName, const FTransform& InGlobalTransform);
 
diff --git a/Engine/Plugins/Animation/ControlRig/Source/ControlRig/Public/Rigs/RigHierarchy.h b/Engine/Plugins/Animation/ControlRig/Source/ControlRig/Public/Rigs/RigHierarchy.h
--- a/Engine/Plugins/Animation/ControlRig/Source/ControlRig/Public/Rigs/RigHierarchy.h
+++ b/Engine/Plugins/Animation/ControlRig/Source/ControlRig/Public/Rigs/RigHierarchy.h
@@ -892,9 +892,9 @@ public:
 	 * @param bAffectChildren If set to false children will not move (maintain global).
 	 */
 	UFUNCTION(BlueprintCallable, Category = URigHierarchy)
-    FORCEINLINE_DEBUGGABLE void SetLocalTransform(FRigElementKey InKey, FTransform InTransform, bool bInitial = false, bool bAffectChildren = true, bool bSetupUndo = false)
+    FORCEINLINE_DEBUGGABLE void SetLocalTransform(FRigElementKey InKey, FTransform InTransform, bool bInitial = false, bool bAffectChildren = true, bool bSetupUndo = false, bool bPrintPythonCommands = false)
 	{
-		SetLocalTransformByIndex(GetIndex(InKey), InTransform, bInitial, bAffectChildren, bSetupUndo);
+		SetLocalTransformByIndex(GetIndex(InKey), InTransform, bInitial, bAffectChildren, bSetupUndo, bPrintPythonCommands);
 	}
 
 	/**
@@ -906,30 +906,30 @@ public:
 	 * @param bAffectChildren If set to false children will not move (maintain global).
 	 */
 	UFUNCTION(BlueprintCallable, Category = URigHierarchy)
-    FORCEINLINE_DEBUGGABLE void SetLocalTransformByIndex(int32 InElementIndex, FTransform InTransform, bool bInitial = false, bool bAffectChildren = true, bool bSetupUndo = false)
+    FORCEINLINE_DEBUGGABLE void SetLocalTransformByIndex(int32 InElementIndex, FTransform InTransform, bool bInitial = false, bool bAffectChildren = true, bool bSetupUndo = false, bool bPrintPythonCommands = false)
 	{
 		if(Elements.IsValidIndex(InElementIndex))
 		{
 			if(FRigTransformElement* TransformElement = Cast<FRigTransformElement>(Elements[InElementIndex]))
 			{
-				SetTransform(TransformElement, InTransform, bInitial ? ERigTransformType::InitialLocal : ERigTransformType::CurrentLocal, bAffectChildren, bSetupUndo);
+				SetTransform(TransformElement, InTransform, bInitial ? ERigTransformType::InitialLocal : ERigTransformType::CurrentLocal, bAffectChildren, bSetupUndo, false, bPrintPythonCommands);
 			}
 		}
 	}
 
-	FORCEINLINE_DEBUGGABLE void SetLocalTransform(int32 InElementIndex, const FTransform& InTransform, bool bAffectChildren = true, bool bSetupUndo = false)
+	FORCEINLINE_DEBUGGABLE void SetLocalTransform(int32 InElementIndex, const FTransform& InTransform, bool bAffectChildren = true, bool bSetupUndo = false, bool bPrintPythonCommands = false)
 	{
-		SetLocalTransformByIndex(InElementIndex, InTransform, false, bAffectChildren, bSetupUndo);
+		SetLocalTransformByIndex(InElementIndex, InTransform, false, bAffectChildren, bSetupUndo, bPrintPythonCommands);
 	}
 
-	FORCEINLINE_DEBUGGABLE void SetInitialLocalTransform(int32 InElementIndex, const FTransform& InTransform, bool bAffectChildren = true, bool bSetupUndo = false)
+	FORCEINLINE_DEBUGGABLE void SetInitialLocalTransform(int32 InElementIndex, const FTransform& InTransform, bool bAffectChildren = true, bool bSetupUndo = false, bool bPrintPythonCommands = false)
 	{
-		SetLocalTransformByIndex(InElementIndex, InTransform, true, bAffectChildren, bSetupUndo);
+		SetLocalTransformByIndex(InElementIndex, InTransform, true, bAffectChildren, bSetupUndo, bPrintPythonCommands);
     }
 
-	FORCEINLINE_DEBUGGABLE void SetInitialLocalTransform(const FRigElementKey& InKey, const FTransform& InTransform, bool bAffectChildren = true, bool bSetupUndo = false)
+	FORCEINLINE_DEBUGGABLE void SetInitialLocalTransform(const FRigElementKey& InKey, const FTransform& InTransform, bool bAffectChildren = true, bool bSetupUndo = false, bool bPrintPythonCommands = false)
 	{
-		SetLocalTransform(InKey, InTransform, true, bAffectChildren, bSetupUndo);
+		SetLocalTransform(InKey, InTransform, true, bAffectChildren, bSetupUndo, bPrintPythonCommands);
 	}
 
 	/**
@@ -1381,9 +1381,9 @@ public:
 	 * @param bSetupUndo If true the transform stack will be setup for undo / redo
 	 */
 	UFUNCTION(BlueprintCallable, Category = URigHierarchy)
-    FORCEINLINE_DEBUGGABLE void SetControlOffsetTransform(FRigElementKey InKey, FTransform InTransform, bool bInitial = false, bool bAffectChildren = true, bool bSetupUndo = false)
+    FORCEINLINE_DEBUGGABLE void SetControlOffsetTransform(FRigElementKey InKey, FTransform InTransform, bool bInitial = false, bool bAffectChildren = true, bool bSetupUndo = false, bool bPrintPythonCommands = false)
 	{
-		return SetControlOffsetTransformByIndex(GetIndex(InKey), InTransform, bInitial, bAffectChildren, bSetupUndo);
+		return SetControlOffsetTransformByIndex(GetIndex(InKey), InTransform, bInitial, bAffectChildren, bSetupUndo, bPrintPythonCommands);
 	}
 
 	/**
@@ -1395,13 +1395,13 @@ public:
 	 * @param bSetupUndo If true the transform stack will be setup for undo / redo
 	 */
 	UFUNCTION(BlueprintCallable, Category = URigHierarchy)
-    FORCEINLINE_DEBUGGABLE void SetControlOffsetTransformByIndex(int32 InElementIndex, FTransform InTransform, bool bInitial = false, bool bAffectChildren = true, bool bSetupUndo = false)
+    FORCEINLINE_DEBUGGABLE void SetControlOffsetTransformByIndex(int32 InElementIndex, FTransform InTransform, bool bInitial = false, bool bAffectChildren = true, bool bSetupUndo = false, bool bPrintPythonCommands = false)
 	{
 		if(Elements.IsValidIndex(InElementIndex))
 		{
 			if(FRigControlElement* ControlElement = Cast<FRigControlElement>(Elements[InElementIndex]))
 			{
-				SetControlOffsetTransform(ControlElement, InTransform, bInitial ? ERigTransformType::InitialLocal : ERigTransformType::CurrentLocal, bAffectChildren, bSetupUndo);
+				SetControlOffsetTransform(ControlElement, InTransform, bInitial ? ERigTransformType::InitialLocal : ERigTransformType::CurrentLocal, bAffectChildren, bSetupUndo, bPrintPythonCommands);
 			}
 		}
 	}
@@ -2071,7 +2071,7 @@ public:
 	 * @param bSetupUndo If true the transform stack will be setup for undo / redo
 	 * @param bForce Set the transform even if it is the same as the previously set one
 	 */
-	void SetTransform(FRigTransformElement* InTransformElement, const FTransform& InTransform, const ERigTransformType::Type InTransformType, bool bAffectChildren, bool bSetupUndo = false, bool bForce = false);
+	void SetTransform(FRigTransformElement* InTransformElement, const FTransform& InTransform, const ERigTransformType::Type InTransformType, bool bAffectChildren, bool bSetupUndo = false, bool bForce = false, bool bPrintPythonCommands = false);
 
 	/**
 	 * Returns the global offset transform for a given control element.
@@ -2090,7 +2090,7 @@ public:
 	 * @param bSetupUndo If true the transform stack will be setup for undo / redo
 	 * @param bForce Set the transform even if it is the same as the previously set one
 	 */
-	void SetControlOffsetTransform(FRigControlElement* InControlElement, const FTransform& InTransform, const ERigTransformType::Type InTransformType, bool bAffectChildren, bool bSetupUndo = false, bool bForce = false);
+	void SetControlOffsetTransform(FRigControlElement* InControlElement, const FTransform& InTransform, const ERigTransformType::Type InTransformType, bool bAffectChildren, bool bSetupUndo = false, bool bForce = false, bool bPrintPythonCommands = false);
 
 	/**
 	 * Returns the global gizmo transform for a given control element.
@@ -2140,7 +2140,7 @@ public:
 	 * @param bSetupUndo If true the transform stack will be setup for undo / redo
 	 * @param bForce Set the transform even if it is the same as the previously set one
 	 */
-	void SetControlValue(FRigControlElement* InControlElement, const FRigControlValue& InValue, ERigControlValueType InValueType, bool bSetupUndo = false, bool bForce = false);
+	void SetControlValue(FRigControlElement* InControlElement, const FRigControlValue& InValue, ERigControlValueType InValueType, bool bSetupUndo = false, bool bForce = false, bool bPrintPythonCommands = false);
 
 	template<typename T>
 	FORCEINLINE_DEBUGGABLE void SetControlValue(FRigControlElement* InControlElement, const T& InValue, ERigControlValueType InValueType, bool bSetupUndo = false, bool bForce = false) const
@@ -2402,8 +2402,7 @@ private:
 		FRigHierarchyListener()
 			: Hierarchy(nullptr)
 			, bShouldReactToInitialChanges(true)
-			// default to false since changing the current values in a BP/CDO is not meaningful
-			, bShouldReactToCurrentChanges(false) 
+			, bShouldReactToCurrentChanges(true) 
 		{}
 
 		bool ShouldReactToChange(ERigTransformType::Type InTransformType) const
diff --git a/Engine/Plugins/Animation/ControlRig/Source/ControlRigEditor/Private/ControlRigElementDetails.cpp b/Engine/Plugins/Animation/ControlRig/Source/ControlRigEditor/Private/ControlRigElementDetails.cpp
--- a/Engine/Plugins/Animation/ControlRig/Source/ControlRigEditor/Private/ControlRigElementDetails.cpp
+++ b/Engine/Plugins/Animation/ControlRig/Source/ControlRigEditor/Private/ControlRigElementDetails.cpp
@@ -1565,8 +1565,8 @@ void FRigControlElementDetails::CustomizeChildren(TSharedRef<class IPropertyHand
 					}
 
 					this->HierarchyBeingCustomized->Notify(ERigHierarchyNotification::ControlSettingChanged, ControlElement);
-					this->HierarchyBeingCustomized->SetControlValue(ControlElement, ValueToSet, ERigControlValueType::Initial, true);
-					this->HierarchyBeingCustomized->SetControlValue(ControlElement, ValueToSet, ERigControlValueType::Current, true);
+					this->HierarchyBeingCustomized->SetControlValue(ControlElement, ValueToSet, ERigControlValueType::Initial, true, false, true);
+					this->HierarchyBeingCustomized->SetControlValue(ControlElement, ValueToSet, ERigControlValueType::Current, true, false, true);
 
 					FRigControlElement::StaticStruct()->CopyScriptStruct(ViewElement, ControlElement);
 
@@ -1578,7 +1578,7 @@ void FRigControlElementDetails::CustomizeChildren(TSharedRef<class IPropertyHand
 							this->BlueprintBeingCustomized->Hierarchy->Notify(ERigHierarchyNotification::ControlSettingChanged, ControlElement);
 							this->BlueprintBeingCustomized->Hierarchy->SetControlValue(ControlElement, ValueToSet, ERigControlValueType::Initial, true);
 							this->BlueprintBeingCustomized->Hierarchy->SetControlValue(ControlElement, ValueToSet, ERigControlValueType::Current, true);
-					}
+						}
 					}
 				}
 				
@@ -1758,8 +1758,8 @@ void FRigControlElementDetails::CustomizeChildren(TSharedRef<class IPropertyHand
 
 							ControlBeingCustomized->Settings.ApplyLimits(InitialValue);
 							ControlBeingCustomized->Settings.ApplyLimits(CurrentValue);
-							HierarchyBeingCustomized->SetControlValue(ControlBeingCustomized, InitialValue, ERigControlValueType::Initial);
-							HierarchyBeingCustomized->SetControlValue(ControlBeingCustomized, CurrentValue, ERigControlValueType::Current);
+							HierarchyBeingCustomized->SetControlValue(ControlBeingCustomized, InitialValue, ERigControlValueType::Initial, false, false, true);
+							HierarchyBeingCustomized->SetControlValue(ControlBeingCustomized, CurrentValue, ERigControlValueType::Current, false, false, true);
 
 							if (UControlRig* DebuggedRig = Cast<UControlRig>(BlueprintBeingCustomized->GetObjectBeingDebugged()))
 							{
diff --git a/Engine/Plugins/Animation/ControlRig/Source/ControlRigEditor/Private/EditMode/ControlRigEditMode.cpp b/Engine/Plugins/Animation/ControlRig/Source/ControlRigEditor/Private/EditMode/ControlRigEditMode.cpp
--- a/Engine/Plugins/Animation/ControlRig/Source/ControlRigEditor/Private/EditMode/ControlRigEditMode.cpp
+++ b/Engine/Plugins/Animation/ControlRig/Source/ControlRigEditor/Private/EditMode/ControlRigEditMode.cpp
@@ -2382,7 +2382,7 @@ void FControlRigEditMode::MoveGizmo(AControlRigGizmoActor* GizmoActor, const boo
 				{
 					InOutLocal = InteractionControlRig->GetControlLocalTransform(GizmoActor->ControlName);
 				}
-				InteractionControlRig->SetControlGlobalTransform(GizmoActor->ControlName, NewTransform, true, Context);			// assumes it's attached to actor
+				InteractionControlRig->SetControlGlobalTransform(GizmoActor->ControlName, NewTransform, true, Context, true, true);			// assumes it's attached to actor
 				GizmoActor->SetGlobalTransform(CurrentTransform);
 				if (bCalcLocal)
 				{
diff --git a/Engine/Plugins/Animation/ControlRig/Source/ControlRigEditor/Private/Editor/ControlRigEditor.cpp b/Engine/Plugins/Animation/ControlRig/Source/ControlRigEditor/Private/Editor/ControlRigEditor.cpp
--- a/Engine/Plugins/Animation/ControlRig/Source/ControlRigEditor/Private/Editor/ControlRigEditor.cpp
+++ b/Engine/Plugins/Animation/ControlRig/Source/ControlRigEditor/Private/Editor/ControlRigEditor.cpp
@@ -4146,9 +4146,9 @@ void FControlRigEditor::OnWrappedPropertyChangedChainEvent(UDetailsViewWrapperOb
 						
 				if(ERigTransformType::IsInitial(TransformType) || bSetupModeEnabled)
 				{
-					Hierarchy->SetControlValue(ControlElement, Value, ERigControlValueType::Initial, true, true);
+					Hierarchy->SetControlValue(ControlElement, Value, ERigControlValueType::Initial, true, true, true);
 				}
-				Hierarchy->SetControlValue(ControlElement, Value, ERigControlValueType::Current, true, true);
+				Hierarchy->SetControlValue(ControlElement, Value, ERigControlValueType::Current, true, true, true);
 			}
 			else
 			{
diff --git a/Engine/Plugins/Animation/ControlRig/Source/ControlRigEditor/Private/Editor/SRigHierarchy.cpp b/Engine/Plugins/Animation/ControlRig/Source/ControlRigEditor/Private/Editor/SRigHierarchy.cpp
--- a/Engine/Plugins/Animation/ControlRig/Source/ControlRigEditor/Private/Editor/SRigHierarchy.cpp
+++ b/Engine/Plugins/Animation/ControlRig/Source/ControlRigEditor/Private/Editor/SRigHierarchy.cpp
@@ -2190,11 +2190,11 @@ void SRigHierarchy::HandlePasteTransforms(ERigTransformType::Type InTransformTyp
 
 			if(FRigTransformElement* TransformElement = Hierarchy->Find<FRigTransformElement>(CurrentSelection[Index]))
 			{
-				Hierarchy->SetTransform(TransformElement, Transform, InTransformType, bAffectChildren, true);
+				Hierarchy->SetTransform(TransformElement, Transform, InTransformType, bAffectChildren, true, false, true);
 			}
 			if(FRigBoneElement* BoneElement = Hierarchy->Find<FRigBoneElement>(CurrentSelection[Index]))
 			{
-				Hierarchy->SetTransform(BoneElement, Transform, ERigTransformType::MakeInitial(InTransformType), bAffectChildren, true);
+				Hierarchy->SetTransform(BoneElement, Transform, ERigTransformType::MakeInitial(InTransformType), bAffectChildren, true, false, true);
 			}
 			
             if(DebuggedHierarchy && DebuggedHierarchy != Hierarchy)
@@ -2472,7 +2472,7 @@ void SRigHierarchy::HandleResetTransform(bool bSelectionOnly)
 				for (FRigElementKey Key : KeysToReset)
 				{
 					const FTransform InitialTransform = GetHierarchy()->GetInitialLocalTransform(Key);
-					GetHierarchy()->SetLocalTransform(Key, InitialTransform, false, true, true);
+					GetHierarchy()->SetLocalTransform(Key, InitialTransform, false, true, true, true);
 					DebuggedHierarchy->SetLocalTransform(Key, InitialTransform, false, true, true);
 
 					if (Key.Type == ERigElementType::Bone)
@@ -2515,10 +2515,10 @@ void SRigHierarchy::HandleSetInitialTransformFromCurrentTransform()
 					{
 						if(FRigControlElement* ControlElement = GetHierarchy()->Find<FRigControlElement>(SelectedKey))
 						{
-							GetHierarchy()->SetControlOffsetTransform(ControlElement, LocalTransform, ERigTransformType::InitialLocal, true, true);
-							GetHierarchy()->SetControlOffsetTransform(ControlElement, LocalTransform, ERigTransformType::CurrentLocal, true, true);
-							GetHierarchy()->SetTransform(ControlElement, FTransform::Identity, ERigTransformType::InitialLocal, true, true);
-							GetHierarchy()->SetTransform(ControlElement, FTransform::Identity, ERigTransformType::CurrentLocal, true, true);
+							GetHierarchy()->SetControlOffsetTransform(ControlElement, LocalTransform, ERigTransformType::InitialLocal, true, true, false, true);
+							GetHierarchy()->SetControlOffsetTransform(ControlElement, LocalTransform, ERigTransformType::CurrentLocal, true, true, false, true);
+							GetHierarchy()->SetTransform(ControlElement, FTransform::Identity, ERigTransformType::InitialLocal, true, true, false, true);
+							GetHierarchy()->SetTransform(ControlElement, FTransform::Identity, ERigTransformType::CurrentLocal, true, true, false, true);
 						}
 						if(FRigControlElement* ControlElement = DebuggedHierarchy->Find<FRigControlElement>(SelectedKey))
 						{
@@ -2544,8 +2544,8 @@ void SRigHierarchy::HandleSetInitialTransformFromCurrentTransform()
 
 						if(FRigTransformElement* TransformElement = GetHierarchy()->Find<FRigTransformElement>(SelectedKey))
 						{
-							GetHierarchy()->SetTransform(TransformElement, LocalTransform, ERigTransformType::InitialLocal, true, true);
-							GetHierarchy()->SetTransform(TransformElement, LocalTransform, ERigTransformType::CurrentLocal, true, true);
+							GetHierarchy()->SetTransform(TransformElement, LocalTransform, ERigTransformType::InitialLocal, true, true, false, true);
+							GetHierarchy()->SetTransform(TransformElement, LocalTransform, ERigTransformType::CurrentLocal, true, true, false, true);
 						}
 						if(FRigTransformElement* TransformElement = DebuggedHierarchy->Find<FRigTransformElement>(SelectedKey))
 						{
@@ -2810,10 +2810,10 @@ FReply SRigHierarchy::ReparentOrMatchTransform(const TArray<FRigElementKey>& Dra
 					FTransform ParentTransform = DebuggedHierarchy->GetParentTransformByIndex(ControlIndex, false);
 					FTransform OffsetTransform = TargetGlobalTransform.GetRelativeTransform(ParentTransform);
 
-					Hierarchy->SetControlOffsetTransformByIndex(ControlIndex, OffsetTransform, ERigTransformType::InitialLocal, true, true);
-					Hierarchy->SetControlOffsetTransformByIndex(ControlIndex, OffsetTransform, ERigTransformType::CurrentLocal, true, true);
-					Hierarchy->SetLocalTransform(DraggedKey, FTransform::Identity, true, true, true);
-					Hierarchy->SetInitialLocalTransform(DraggedKey, FTransform::Identity, true, true);
+					Hierarchy->SetControlOffsetTransformByIndex(ControlIndex, OffsetTransform, ERigTransformType::InitialLocal, true, true, true);
+					Hierarchy->SetControlOffsetTransformByIndex(ControlIndex, OffsetTransform, ERigTransformType::CurrentLocal, true, true, true);
+					Hierarchy->SetLocalTransform(DraggedKey, FTransform::Identity, true, true, true, true);
+					Hierarchy->SetInitialLocalTransform(DraggedKey, FTransform::Identity, true, true, true);
 					DebuggedHierarchy->SetControlOffsetTransformByIndex(ControlIndex, OffsetTransform, ERigTransformType::InitialLocal, true, true);
 					DebuggedHierarchy->SetControlOffsetTransformByIndex(ControlIndex, OffsetTransform, ERigTransformType::CurrentLocal, true, true);
 					DebuggedHierarchy->SetLocalTransform(DraggedKey, FTransform::Identity, true, true, true);
@@ -2909,10 +2909,10 @@ void SRigHierarchy::HandleSetInitialTransformFromClosestBone()
 					{
 						if(FRigControlElement* ControlElement = GetHierarchy()->Find<FRigControlElement>(SelectedKey))
 						{
-							GetHierarchy()->SetControlOffsetTransform(ControlElement, LocalTransform, ERigTransformType::InitialLocal, true, true);
-							GetHierarchy()->SetControlOffsetTransform(ControlElement, LocalTransform, ERigTransformType::CurrentLocal, true, true);
-							GetHierarchy()->SetTransform(ControlElement, FTransform::Identity, ERigTransformType::InitialLocal, true, true);
-							GetHierarchy()->SetTransform(ControlElement, FTransform::Identity, ERigTransformType::CurrentLocal, true, true);
+							GetHierarchy()->SetControlOffsetTransform(ControlElement, LocalTransform, ERigTransformType::InitialLocal, true, true, false, true);
+							GetHierarchy()->SetControlOffsetTransform(ControlElement, LocalTransform, ERigTransformType::CurrentLocal, true, true, false, true);
+							GetHierarchy()->SetTransform(ControlElement, FTransform::Identity, ERigTransformType::InitialLocal, true, true, false, true);
+							GetHierarchy()->SetTransform(ControlElement, FTransform::Identity, ERigTransformType::CurrentLocal, true, true, false, true);
 						}
 						if(FRigControlElement* ControlElement = DebuggedHierarchy->Find<FRigControlElement>(SelectedKey))
 						{
@@ -2929,8 +2929,8 @@ void SRigHierarchy::HandleSetInitialTransformFromClosestBone()
 
 						if(FRigTransformElement* TransformElement = GetHierarchy()->Find<FRigTransformElement>(SelectedKey))
 						{
-							GetHierarchy()->SetTransform(TransformElement, LocalTransform, ERigTransformType::InitialLocal, true, true);
-							GetHierarchy()->SetTransform(TransformElement, LocalTransform, ERigTransformType::CurrentLocal, true, true);
+							GetHierarchy()->SetTransform(TransformElement, LocalTransform, ERigTransformType::InitialLocal, true, true, false, true);
+							GetHierarchy()->SetTransform(TransformElement, LocalTransform, ERigTransformType::CurrentLocal, true, true, false, true);
 						}
 						if(FRigTransformElement* TransformElement = DebuggedHierarchy->Find<FRigTransformElement>(SelectedKey))
 						{
@@ -2974,8 +2974,8 @@ void SRigHierarchy::HandleSetGizmoTransformFromCurrent()
 
 							DebuggedHierarchy->SetLocalTransform(Key, FTransform::Identity, false, true, true);
 							DebuggedHierarchy->SetLocalTransform(Key, FTransform::Identity, true, true, true);
-							GetHierarchy()->SetLocalTransform(Key, FTransform::Identity, false, true, true);
-							GetHierarchy()->SetLocalTransform(Key, FTransform::Identity, true, true, true);
+							GetHierarchy()->SetLocalTransform(Key, FTransform::Identity, false, true, true, true);
+							GetHierarchy()->SetLocalTransform(Key, FTransform::Identity, true, true, true, true);
 						}
 
 						if (FControlRigEditorEditMode* EditMode = ControlRigEditor.Pin()->GetEditMode())
