Control Rig: Fix control rig interactive unit test fail

diff --git a/Engine/Plugins/Animation/ControlRig/Source/ControlRig/Private/Rigs/FKControlRig.cpp b/Engine/Plugins/Animation/ControlRig/Source/ControlRig/Private/Rigs/FKControlRig.cpp
--- a/Engine/Plugins/Animation/ControlRig/Source/ControlRig/Private/Rigs/FKControlRig.cpp
+++ b/Engine/Plugins/Animation/ControlRig/Source/ControlRig/Private/Rigs/FKControlRig.cpp
@@ -143,7 +143,7 @@ void UFKControlRig::ExecuteUnits(FRigUnitContext& InOutContext, const FName& InE
 			Transform.NormalizeRotation();
 			Transform.SetScale3D(FVector::OneVector);
 
-			GetHierarchy()->SetTransform(Control, Transform, ERigTransformType::CurrentLocal, true);
+			SetControlValue(ControlName, FRigControlValue::Make(Transform), bNotify, Context, bSetupUndo);
 
 		}, true);
 
