Control Rig: Fix horde non-unity build break in DetailsViewWrapperObject

diff --git a/Engine/Plugins/Animation/ControlRig/Source/ControlRigEditor/Private/Editor/DetailsViewWrapperObject.cpp b/Engine/Plugins/Animation/ControlRig/Source/ControlRigEditor/Private/Editor/DetailsViewWrapperObject.cpp
--- a/Engine/Plugins/Animation/ControlRig/Source/ControlRigEditor/Private/Editor/DetailsViewWrapperObject.cpp
+++ b/Engine/Plugins/Animation/ControlRig/Source/ControlRigEditor/Private/Editor/DetailsViewWrapperObject.cpp
@@ -4,6 +4,10 @@
 #include "Units/RigUnit.h"
 #include "ControlRigElementDetails.h"
 
+#if WITH_EDITOR
+#include "PropertyEditorModule.h"
+#endif
+
 TMap<UScriptStruct*, UClass*> UDetailsViewWrapperObject::StructToClass;
 TMap<UClass*, UScriptStruct*> UDetailsViewWrapperObject::ClassToStruct;
 
@@ -105,11 +109,13 @@ UClass* UDetailsViewWrapperObject::GetClassForStruct(UScriptStruct* InStruct)
 			ChildProperty->SetPropertyFlags(ChildProperty->GetPropertyFlags() | CPF_Edit);
 		}
 
+#if WITH_EDITOR
 		FPropertyEditorModule& PropertyEditorModule = FModuleManager::LoadModuleChecked<FPropertyEditorModule>("PropertyEditor");
 		if (!PropertyEditorModule.GetClassNameToDetailLayoutNameMap().Contains(InStruct->GetFName()))
 		{
 			PropertyEditorModule.RegisterCustomClassLayout(InStruct->GetFName(), FOnGetDetailCustomizationInstance::CreateStatic(&FRigUnitDetails::MakeInstance));
 		}
+#endif
 	}
 
 	// Given it's only one property we are adding
