Control Rig: Minor fixes to UDetailsViewWrapperObject such that GC don't complain

diff --git a/Engine/Plugins/Animation/ControlRig/Source/ControlRigEditor/Private/Editor/DetailsViewWrapperObject.cpp b/Engine/Plugins/Animation/ControlRig/Source/ControlRigEditor/Private/Editor/DetailsViewWrapperObject.cpp
--- a/Engine/Plugins/Animation/ControlRig/Source/ControlRigEditor/Private/Editor/DetailsViewWrapperObject.cpp
+++ b/Engine/Plugins/Animation/ControlRig/Source/ControlRigEditor/Private/Editor/DetailsViewWrapperObject.cpp
@@ -125,6 +125,9 @@ UClass* UDetailsViewWrapperObject::GetClassForStruct(UScriptStruct* InStruct)
 	// Update the class
 	WrapperClass->Bind();
 	WrapperClass->StaticLink(true);
+	
+	// Similar to FConfigPropertyHelperDetails::CustomizeDetails, this is required for GC to work properly
+	WrapperClass->AssembleReferenceTokenStream();
 
 	StructToClass.Add(InStruct, WrapperClass);
 	ClassToStruct.Add(WrapperClass, InStruct);
