Control Rig: Fix crashes related to collapse nodes that are selected by the first selection. Previously on load we had a bug where we always have 1 node pre selected, dirtying the BP on load. To work around it, we ignores the first selection on load. Now that the bug has been fixed, the workaround also need to be removed to avoid invalidating the first selection causing access violation.

diff --git a/Engine/Plugins/Animation/ControlRig/Source/ControlRigEditor/Private/Editor/ControlRigEditor.cpp b/Engine/Plugins/Animation/ControlRig/Source/ControlRigEditor/Private/Editor/ControlRigEditor.cpp
--- a/Engine/Plugins/Animation/ControlRig/Source/ControlRigEditor/Private/Editor/ControlRigEditor.cpp
+++ b/Engine/Plugins/Animation/ControlRig/Source/ControlRigEditor/Private/Editor/ControlRigEditor.cpp
@@ -200,7 +200,6 @@ FControlRigEditor::FControlRigEditor()
 	, bIsSettingObjectBeingDebugged(false)
 	, bExecutionControlRig(true)
 	, bSetupModeEnabled(false)
-	, bFirstTimeSelecting(true)
 	, bAnyErrorsLeft(false)
 	, LastEventQueue(EControlRigEditorEventQueue::Setup)
 	, ExecutionMode(EControlRigExecutionModeType::EControlRigExecutionModeType_Release)
@@ -2813,12 +2812,6 @@ void FControlRigEditor::OnSelectedNodesChangedImpl(const TSet<class UObject*>& N
 		return;
 	}
 
-	if (bFirstTimeSelecting)
-	{
-		bFirstTimeSelecting = false;
-		return;
-	}
-
 	TGuardValue<bool> SelectGuard(RigGraph->bIsSelecting, true);
 
 	UControlRigBlueprint* ControlRigBlueprint = CastChecked<UControlRigBlueprint>(GetBlueprintObj());
diff --git a/Engine/Plugins/Animation/ControlRig/Source/ControlRigEditor/Private/Editor/ControlRigEditor.h b/Engine/Plugins/Animation/ControlRig/Source/ControlRigEditor/Private/Editor/ControlRigEditor.h
--- a/Engine/Plugins/Animation/ControlRig/Source/ControlRigEditor/Private/Editor/ControlRigEditor.h
+++ b/Engine/Plugins/Animation/ControlRig/Source/ControlRigEditor/Private/Editor/ControlRigEditor.h
@@ -450,7 +450,6 @@ protected:
 	void OnPinControlNameListComboBox(const TArray<TSharedPtr<FString>>* InNameList);
 	void ToggleSetupMode();
 
-	bool bFirstTimeSelecting;
 	bool bAnyErrorsLeft;
 
 	EControlRigEditorEventQueue LastEventQueue;
