Control Rig: Make sure to set auto-gen default values for custom pins on the Control Rig Anim Node, otherwise when an empty auto-gen default value is used to initialize the actual default value, the empty default value would be considered invalid by the BP compiler.

#horde https://horde.devtools.epicgames.com/log/60f6610892cf8900013c9270?issue=302

diff --git a/Engine/Plugins/Animation/ControlRig/Source/ControlRigDeveloper/Private/AnimGraphNode_ControlRig.cpp b/Engine/Plugins/Animation/ControlRig/Source/ControlRigDeveloper/Private/AnimGraphNode_ControlRig.cpp
--- a/Engine/Plugins/Animation/ControlRig/Source/ControlRigDeveloper/Private/AnimGraphNode_ControlRig.cpp
+++ b/Engine/Plugins/Animation/ControlRig/Source/ControlRigDeveloper/Private/AnimGraphNode_ControlRig.cpp
@@ -89,6 +89,15 @@ void UAnimGraphNode_ControlRig::ReallocatePinsDuringReconstruction(TArray<UEdGra
 						UEdGraphPin* NewPin = CreatePin(EEdGraphPinDirection::EGPD_Input, OldPin->PinType, PropertyName);
 						NewPin->PinFriendlyName = OldPin->PinFriendlyName;
 						
+						// Newly created pin does not have an auto generated default value, so we need to generate one here
+						// Missing the auto-gen default would cause UEdGraphPin::MovePersistentDataFromOldPin to override
+						// the actual default value with the empty auto-gen default, causing BP compiler to complain
+						// This is similar to how the following two functions create and initialize new pins, create first 
+						// and then set an auto-gen default
+						// FOptionalPinManager::CreateVisiblePins()
+						// FAnimBlueprintNodeOptionalPinManager::PostInitNewPin()
+						AnimGraphDefaultSchema->SetPinAutogeneratedDefaultValueBasedOnType(NewPin);
+
 						AnimGraphDefaultSchema->ResetPinToAutogeneratedDefaultValue(NewPin, false);
 
 						CustomizePinData(NewPin, PropertyName, INDEX_NONE);
@@ -128,6 +137,15 @@ void UAnimGraphNode_ControlRig::ReallocatePinsDuringReconstruction(TArray<UEdGra
 			UEdGraphPin* NewPin = CreatePin(EEdGraphPinDirection::EGPD_Input, PinType, PropertyName);
 			NewPin->PinFriendlyName = FText::FromName(PropertyName);
 
+			// Newly created pin does not have an auto generated default value, so we need to generate one here
+			// Missing the auto-gen default would cause UEdGraphPin::MovePersistentDataFromOldPin to override
+			// the actual default value with the empty auto-gen default, causing BP compiler to complain
+			// This is similar to how the following two functions create and initialize new pins, create first 
+			// and then set an auto-gen default
+			// FOptionalPinManager::CreateVisiblePins()
+			// FAnimBlueprintNodeOptionalPinManager::PostInitNewPin()
+			AnimGraphDefaultSchema->SetPinAutogeneratedDefaultValueBasedOnType(NewPin);
+
 			// We cant interrogate CDO here as we may be mid-compile, so we can only really
 			// reset to the autogenerated default.
 			AnimGraphDefaultSchema->ResetPinToAutogeneratedDefaultValue(NewPin, false);
