Control Rig: Print python commands when variables are added/renamed/retyped/deleted

diff --git a/Engine/Plugins/Animation/ControlRig/Source/ControlRigDeveloper/Private/ControlRigBlueprint.cpp b/Engine/Plugins/Animation/ControlRig/Source/ControlRigDeveloper/Private/ControlRigBlueprint.cpp
--- a/Engine/Plugins/Animation/ControlRig/Source/ControlRigDeveloper/Private/ControlRigBlueprint.cpp
+++ b/Engine/Plugins/Animation/ControlRig/Source/ControlRigDeveloper/Private/ControlRigBlueprint.cpp
@@ -1905,6 +1905,117 @@ FName UControlRigBlueprint::AddMemberVariable(const FName& InName, const FString
 	return Result;
 }
 
+bool UControlRigBlueprint::RemoveMemberVariable(const FName& InName)
+{
+	const int32 VarIndex = FBlueprintEditorUtils::FindNewVariableIndex(this, InName);
+	if (VarIndex == INDEX_NONE)
+	{
+		return false;
+	}
+	
+	FBlueprintEditorUtils::RemoveMemberVariable(this, InName);
+	return true;
+}
+
+bool UControlRigBlueprint::RenameMemberVariable(const FName& InOldName, const FName& InNewName)
+{
+	int32 VarIndex = FBlueprintEditorUtils::FindNewVariableIndex(this, InOldName);
+	if (VarIndex == INDEX_NONE)
+	{
+		return false;
+	}
+
+	VarIndex = FBlueprintEditorUtils::FindNewVariableIndex(this, InNewName);
+	if (VarIndex != INDEX_NONE)
+	{
+		return false;
+	}
+	
+	FBlueprintEditorUtils::RenameMemberVariable(this, InOldName, InNewName);
+	return true;
+}
+
+bool UControlRigBlueprint::ChangeMemberVariableType(const FName& InName, const FString& InCPPType, bool bIsPublic,
+	bool bIsReadOnly, FString InDefaultValue)
+{
+	int32 VarIndex = FBlueprintEditorUtils::FindNewVariableIndex(this, InName);
+	if (VarIndex == INDEX_NONE)
+	{
+		return false;
+	}
+
+	FRigVMExternalVariable Variable;
+	Variable.Name = InName;
+	Variable.bIsPublic = bIsPublic;
+	Variable.bIsReadOnly = bIsReadOnly;
+
+	FString CPPType = InCPPType;
+	if (CPPType.StartsWith(TEXT("TMap<")))
+	{
+		UE_LOG(LogControlRigDeveloper, Warning, TEXT("TMap Variables are not supported."));
+		return false;
+	}
+
+	Variable.bIsArray = RigVMUtilities::IsArrayType(CPPType);
+	if (Variable.bIsArray)
+	{
+		CPPType = CPPType.RightChop(7).LeftChop(1);
+	}
+
+	if (CPPType == TEXT("bool"))
+	{
+		Variable.TypeName = *CPPType;
+		Variable.Size = sizeof(bool);
+	}
+	else if (CPPType == TEXT("float"))
+	{
+		Variable.TypeName = *CPPType;
+		Variable.Size = sizeof(float);
+	}
+	else if (CPPType == TEXT("double"))
+	{
+		Variable.TypeName = *CPPType;
+		Variable.Size = sizeof(double);
+	}
+	else if (CPPType == TEXT("int32"))
+	{
+		Variable.TypeName = *CPPType;
+		Variable.Size = sizeof(int32);
+	}
+	else if (CPPType == TEXT("FString"))
+	{
+		Variable.TypeName = *CPPType;
+		Variable.Size = sizeof(FString);
+	}
+	else if (CPPType == TEXT("FName"))
+	{
+		Variable.TypeName = *CPPType;
+		Variable.Size = sizeof(FName);
+	}
+	else if(UScriptStruct* ScriptStruct = URigVMPin::FindObjectFromCPPTypeObjectPath<UScriptStruct>(CPPType))
+	{
+		Variable.TypeName = *ScriptStruct->GetStructCPPName();
+		Variable.TypeObject = ScriptStruct;
+		Variable.Size = ScriptStruct->GetStructureSize();
+	}
+	else if (UEnum* Enum= URigVMPin::FindObjectFromCPPTypeObjectPath<UEnum>(CPPType))
+	{
+		Variable.TypeName = *Enum->CppType;
+		Variable.TypeObject = Enum;
+		Variable.Size = Enum->GetResourceSizeBytes(EResourceSizeMode::EstimatedTotal);
+	}
+
+	FEdGraphPinType PinType = UControlRig::GetPinTypeFromExternalVariable(Variable);
+	if (!PinType.PinCategory.IsValid())
+	{
+		return false;
+	}
+
+	FBlueprintEditorUtils::ChangeMemberVariableType(this, InName, PinType);
+
+	return true;
+}
+
 FName UControlRigBlueprint::AddTransientControl(URigVMPin* InPin)
 {
 	TUniquePtr<FControlValueScope> ValueScope;
@@ -3394,6 +3505,25 @@ void UControlRigBlueprint::OnPostVariableChange(UBlueprint* InBlueprint)
 
 void UControlRigBlueprint::OnVariableAdded(const FName& InVarName)
 {
+	FBPVariableDescription Variable;
+	for (FBPVariableDescription& NewVariable : NewVariables)
+	{
+		if (NewVariable.VarName == InVarName)
+		{
+			Variable = NewVariable;
+			break;
+		}
+	}
+
+	FRigVMExternalVariable ExternalVariable = UControlRig::GetExternalVariableFromDescription(Variable);
+
+	RigVMPythonUtils::Print(FString::Printf(TEXT("blueprint.add_member_variable('%s', '%s', %s, %s, '%s')"),
+			*InVarName.ToString(),
+			*ExternalVariable.TypeName.ToString(),
+			(ExternalVariable.bIsPublic) ? TEXT("False") : TEXT("True"), 
+			(ExternalVariable.bIsReadOnly) ? TEXT("True") : TEXT("False"), 
+			*Variable.DefaultValue)); 
+	
 	BroadcastExternalVariablesChangedEvent();
 }
 
@@ -3404,9 +3534,13 @@ void UControlRigBlueprint::OnVariableRemoved(const FName& InVarName)
 	{
 		if (URigVMController* Controller = GetOrCreateController(Graph))
 		{
-			Controller->OnExternalVariableRemoved(InVarName, true, true);
+			Controller->OnExternalVariableRemoved(InVarName, true);
 		}
 	}
+
+	RigVMPythonUtils::Print(FString::Printf(TEXT("blueprint.remove_member_variable('%s')"),
+			*InVarName.ToString()));
+	
 	BroadcastExternalVariablesChangedEvent();
 }
 
@@ -3417,9 +3551,14 @@ void UControlRigBlueprint::OnVariableRenamed(const FName& InOldVarName, const FN
 	{
 		if (URigVMController* Controller = GetOrCreateController(Graph))
 		{
-			Controller->OnExternalVariableRenamed(InOldVarName, InNewVarName, true, true);
+			Controller->OnExternalVariableRenamed(InOldVarName, InNewVarName, true);
 		}
 	}
+
+	RigVMPythonUtils::Print(FString::Printf(TEXT("blueprint.rename_member_variable('%s', '%s')"),
+			*InOldVarName.ToString(),
+			*InNewVarName.ToString()));
+	
 	BroadcastExternalVariablesChangedEvent();
 }
 
@@ -3439,14 +3578,42 @@ void UControlRigBlueprint::OnVariableTypeChanged(const FName& InVarName, FEdGrap
 			FRigVMExternalVariable NewVariable = UControlRig::GetExternalVariableFromPinType(InVarName, InNewPinType);
 			if (NewVariable.IsValid(true)) // allow nullptr
 			{
-				Controller->OnExternalVariableTypeChanged(InVarName, NewVariable.TypeName.ToString(), NewVariable.TypeObject, bSetupUndoRedo, true);
+				Controller->OnExternalVariableTypeChanged(InVarName, NewVariable.TypeName.ToString(), NewVariable.TypeObject, bSetupUndoRedo);
 			}
 			else
 			{
-				Controller->OnExternalVariableRemoved(InVarName, bSetupUndoRedo, true);
+				Controller->OnExternalVariableRemoved(InVarName, bSetupUndoRedo);
 			}
 		}
 	}
+
+	FRigVMExternalVariable NewVariable = UControlRig::GetExternalVariableFromPinType(InVarName, InNewPinType);
+	FString PinType = NewVariable.TypeName.ToString();
+	if(UScriptStruct* ScriptStruct = Cast<UScriptStruct>(NewVariable.TypeObject))
+	{
+		for (auto Var : NewVariables)
+		{
+			if (Var.VarName == InVarName)
+			{
+				PinType = ScriptStruct->GetName();
+			}
+		}
+	}
+	else if (UEnum* Enum = Cast<UEnum>(NewVariable.TypeObject))
+	{
+		for (auto Var : NewVariables)
+		{
+			if (Var.VarName == InVarName)
+			{
+				PinType = Enum->GetName();
+			}
+		}
+	}
+
+	RigVMPythonUtils::Print(FString::Printf(TEXT("blueprint.change_member_variable_type('%s', '%s')"),
+		*InVarName.ToString(),
+		*PinType));
+
 	BroadcastExternalVariablesChangedEvent();
 }
 
diff --git a/Engine/Plugins/Animation/ControlRig/Source/ControlRigDeveloper/Public/ControlRigBlueprint.h b/Engine/Plugins/Animation/ControlRig/Source/ControlRigDeveloper/Public/ControlRigBlueprint.h
--- a/Engine/Plugins/Animation/ControlRig/Source/ControlRigDeveloper/Public/ControlRigBlueprint.h
+++ b/Engine/Plugins/Animation/ControlRig/Source/ControlRigDeveloper/Public/ControlRigBlueprint.h
@@ -335,6 +335,15 @@ public:
 #if WITH_EDITOR
 	UFUNCTION(BlueprintCallable, Category = "Variables")
 	FName AddMemberVariable(const FName& InName, const FString& InCPPType, bool bIsPublic = false, bool bIsReadOnly = false, FString InDefaultValue = TEXT(""));
+
+	UFUNCTION(BlueprintCallable, Category = "Variables")
+	bool RemoveMemberVariable(const FName& InName);
+
+	UFUNCTION(BlueprintCallable, Category = "Variables")
+	bool RenameMemberVariable(const FName& InOldName, const FName& InNewName);
+
+	UFUNCTION(BlueprintCallable, Category = "Variables")
+	bool ChangeMemberVariableType(const FName& InName, const FString& InCPPType, bool bIsPublic = false, bool bIsReadOnly = false, FString InDefaultValue = TEXT(""));
 #endif
 
 #if WITH_EDITORONLY_DATA
diff --git a/Engine/Source/Developer/RigVMDeveloper/Private/RigVMModel/RigVMController.cpp b/Engine/Source/Developer/RigVMDeveloper/Private/RigVMModel/RigVMController.cpp
--- a/Engine/Source/Developer/RigVMDeveloper/Private/RigVMModel/RigVMController.cpp
+++ b/Engine/Source/Developer/RigVMDeveloper/Private/RigVMModel/RigVMController.cpp
@@ -1005,7 +1005,7 @@ void URigVMController::RefreshVariableNode(const FName& InNodeName, const FName&
 	}
 }
 
-void URigVMController::OnExternalVariableRemoved(const FName& InVarName, bool bSetupUndoRedo, bool bPrintPythonCommand)
+void URigVMController::OnExternalVariableRemoved(const FName& InVarName, bool bSetupUndoRedo)
 {
 	if (!IsValidGraph())
 	{
@@ -1090,16 +1090,9 @@ void URigVMController::OnExternalVariableRemoved(const FName& InVarName, bool bS
 	{
 		CloseUndoBracket();
 	}
-
-	if (bPrintPythonCommand)
-	{
-		RigVMPythonUtils::Print(FString::Printf(TEXT("blueprint.get_controller_by_name('%s').on_external_variable_removed('%s')"),
-				*GetGraph()->GetGraphName(),
-				*InVarName.ToString()));
-	}
 }
 
-void URigVMController::OnExternalVariableRenamed(const FName& InOldVarName, const FName& InNewVarName, bool bSetupUndoRedo, bool bPrintPythonCommand)
+void URigVMController::OnExternalVariableRenamed(const FName& InOldVarName, const FName& InNewVarName, bool bSetupUndoRedo)
 {
 	if (!IsValidGraph())
 	{
@@ -1187,17 +1180,9 @@ void URigVMController::OnExternalVariableRenamed(const FName& InOldVarName, cons
 	{
 		CloseUndoBracket();
 	}
-
-	if (bPrintPythonCommand)
-	{
-		RigVMPythonUtils::Print(FString::Printf(TEXT("blueprint.get_controller_by_name('%s').on_external_variable_renamed('%s', '%s')"),
-				*GetGraph()->GetGraphName(),
-				*InOldVarName.ToString(),
-				*InNewVarName.ToString()));
-	}
 }
 
-void URigVMController::OnExternalVariableTypeChanged(const FName& InVarName, const FString& InCPPType, UObject* InCPPTypeObject, bool bSetupUndoRedo, bool bPrintPythonCommand)
+void URigVMController::OnExternalVariableTypeChanged(const FName& InVarName, const FString& InCPPType, UObject* InCPPTypeObject, bool bSetupUndoRedo)
 {
 	if (!IsValidGraph())
 	{
@@ -1283,16 +1268,6 @@ void URigVMController::OnExternalVariableTypeChanged(const FName& InVarName, con
 	{
 		CloseUndoBracket();
 	}
-	
-	if (bPrintPythonCommand)
-	{
-		//OnExternalVariableTypeChanged(const FName& InVarName, const FString& InCPPType, UObject* InCPPTypeObject
-		RigVMPythonUtils::Print(FString::Printf(TEXT("blueprint.get_controller_by_name('%s').on_external_variable_type_changed_from_object_path('%s', '%s', '%s')"),
-				*GetGraph()->GetGraphName(),
-				*InVarName.ToString(),
-				*InCPPType,
-				*InCPPTypeObject->GetPathName()));
-	}
 }
 
 void URigVMController::OnExternalVariableTypeChangedFromObjectPath(const FName& InVarName, const FString& InCPPType, const FString& InCPPTypeObjectPath, bool bSetupUndoRedo)
diff --git a/Engine/Source/Developer/RigVMDeveloper/Public/RigVMModel/RigVMController.h b/Engine/Source/Developer/RigVMDeveloper/Public/RigVMModel/RigVMController.h
--- a/Engine/Source/Developer/RigVMDeveloper/Public/RigVMModel/RigVMController.h
+++ b/Engine/Source/Developer/RigVMDeveloper/Public/RigVMModel/RigVMController.h
@@ -175,15 +175,15 @@ public:
 
 	// Removes all nodes related to a given variable
 	UFUNCTION(BlueprintCallable, Category = RigVMController)
-	void OnExternalVariableRemoved(const FName& InVarName, bool bSetupUndoRedo, bool bPrintPythonCommand = false);
+	void OnExternalVariableRemoved(const FName& InVarName, bool bSetupUndoRedo);
 
 	// Renames the variable name in all relevant nodes
 	UFUNCTION(BlueprintCallable, Category = RigVMController)
-	void OnExternalVariableRenamed(const FName& InOldVarName, const FName& InNewVarName, bool bSetupUndoRedo, bool bPrintPythonCommand = false);
+	void OnExternalVariableRenamed(const FName& InOldVarName, const FName& InNewVarName, bool bSetupUndoRedo);
 
 	// Changes the data type of all nodes matching a given variable name
 	UFUNCTION(BlueprintCallable, Category = RigVMController)
-	void OnExternalVariableTypeChanged(const FName& InVarName, const FString& InCPPType, UObject* InCPPTypeObject, bool bSetupUndoRedo, bool bPrintPythonCommand = false);
+	void OnExternalVariableTypeChanged(const FName& InVarName, const FString& InCPPType, UObject* InCPPTypeObject, bool bSetupUndoRedo);
 	UFUNCTION(BlueprintCallable, Category = RigVMController)
 	void OnExternalVariableTypeChangedFromObjectPath(const FName& InVarName, const FString& InCPPType, const FString& InCPPTypeObjectPath, bool bSetupUndoRedo);
 
