Control Rig: Use weak object ptr for details panel content

diff --git a/Engine/Plugins/Animation/ControlRig/Source/ControlRigEditor/Private/ControlRigElementDetails.cpp b/Engine/Plugins/Animation/ControlRig/Source/ControlRigEditor/Private/ControlRigElementDetails.cpp
--- a/Engine/Plugins/Animation/ControlRig/Source/ControlRigEditor/Private/ControlRigElementDetails.cpp
+++ b/Engine/Plugins/Animation/ControlRig/Source/ControlRigEditor/Private/ControlRigElementDetails.cpp
@@ -762,7 +762,11 @@ void FRigBaseElementDetails::CustomizeChildren(TSharedRef<class IPropertyHandle>
 FRigElementKey FRigBaseElementDetails::GetElementKey() const
 {
 	check(ObjectsBeingCustomized.Num() == 1);
-	return ObjectsBeingCustomized[0]->GetContent<FRigBaseElement>()->GetKey(); 
+	if(ObjectsBeingCustomized[0].IsValid())
+	{
+		return ObjectsBeingCustomized[0]->GetContent<FRigBaseElement>()->GetKey(); 
+	}
+	return FRigElementKey();
 }
 
 FText FRigBaseElementDetails::GetName() const
@@ -820,22 +824,28 @@ bool FRigBaseElementDetails::IsSetupModeEnabled() const
 TArray<FRigElementKey> FRigBaseElementDetails::GetElementKeys() const
 {
 	TArray<FRigElementKey> Keys;
-	for(UDetailsViewWrapperObject* ObjectBeingCustomized : ObjectsBeingCustomized)
+	for(TWeakObjectPtr<UDetailsViewWrapperObject> ObjectBeingCustomized : ObjectsBeingCustomized)
 	{
-		Keys.Add(ObjectBeingCustomized->GetContent<FRigBaseElement>()->GetKey());
+		if(ObjectBeingCustomized.IsValid())
+		{
+			Keys.Add(ObjectBeingCustomized->GetContent<FRigBaseElement>()->GetKey());
+		}
 	}
 	return Keys;
 }
 
 bool FRigBaseElementDetails::IsAnyControlOfType(ERigControlType InType) const
 {
-	for(UDetailsViewWrapperObject* ObjectBeingCustomized : ObjectsBeingCustomized)
+	for(TWeakObjectPtr<UDetailsViewWrapperObject> ObjectBeingCustomized : ObjectsBeingCustomized)
 	{
-		if(FRigControlElement* ControlElement = Cast<FRigControlElement>(ObjectBeingCustomized->GetContent<FRigBaseElement>()))
+		if(ObjectBeingCustomized.IsValid())
 		{
-			if(ControlElement->Settings.ControlType == InType)
+			if(FRigControlElement* ControlElement = Cast<FRigControlElement>(ObjectBeingCustomized->GetContent<FRigBaseElement>()))
 			{
-				return true;
+				if(ControlElement->Settings.ControlType == InType)
+				{
+					return true;
+				}
 			}
 		}
 	}
@@ -844,13 +854,16 @@ bool FRigBaseElementDetails::IsAnyControlOfType(ERigControlType InType) const
 
 bool FRigBaseElementDetails::IsAnyControlNotOfType(ERigControlType InType) const
 {
-	for(UDetailsViewWrapperObject* ObjectBeingCustomized : ObjectsBeingCustomized)
+	for(TWeakObjectPtr<UDetailsViewWrapperObject> ObjectBeingCustomized : ObjectsBeingCustomized)
 	{
-		if(FRigControlElement* ControlElement = Cast<FRigControlElement>(ObjectBeingCustomized->GetContent<FRigBaseElement>()))
+		if(ObjectBeingCustomized.IsValid())
 		{
-			if(ControlElement->Settings.ControlType != InType)
+			if(FRigControlElement* ControlElement = Cast<FRigControlElement>(ObjectBeingCustomized->GetContent<FRigBaseElement>()))
 			{
-				return true;
+				if(ControlElement->Settings.ControlType != InType)
+				{
+					return true;
+				}
 			}
 		}
 	}
@@ -909,11 +922,14 @@ void FRigTransformElementDetails::CustomizeChildren(TSharedRef<class IPropertyHa
 
 bool FRigTransformElementDetails::IsCurrentLocalEnabled() const
 {
-	for(UDetailsViewWrapperObject* WrapperObject : ObjectsBeingCustomized)
+	for(TWeakObjectPtr<UDetailsViewWrapperObject> ObjectBeingCustomized : ObjectsBeingCustomized)
 	{
-		if(ObjectsBeingCustomized[0]->GetContent<FRigBaseElement>()->GetType() == ERigElementType::Control)
+		if(ObjectBeingCustomized.IsValid())
 		{
-			return false;
+			if(ObjectBeingCustomized->GetContent<FRigBaseElement>()->GetType() == ERigElementType::Control)
+			{
+				return false;
+			}
 		}
 	}
 	return true;
@@ -1656,11 +1672,14 @@ void FRigControlElementDetails::CustomizeChildren(TSharedRef<class IPropertyHand
 	}
 
 	TArray<FRigControlElement*> ControlElements;
-	for(UDetailsViewWrapperObject* ObjectBeingCustomized : ObjectsBeingCustomized)
+	for(TWeakObjectPtr<UDetailsViewWrapperObject> ObjectBeingCustomized : ObjectsBeingCustomized)
 	{
-		if(FRigControlElement* ControlElement = Cast<FRigControlElement>(ObjectBeingCustomized->GetContent<FRigBaseElement>()))
+		if(ObjectBeingCustomized.IsValid())
 		{
-			ControlElements.Add(ControlElement);
+			if(FRigControlElement* ControlElement = Cast<FRigControlElement>(ObjectBeingCustomized->GetContent<FRigBaseElement>()))
+			{
+				ControlElements.Add(ControlElement);
+			}
 		}
 	}
 
@@ -1785,13 +1804,16 @@ void FRigControlElementDetails::CustomizeChildren(TSharedRef<class IPropertyHand
 
 bool FRigControlElementDetails::IsGizmoEnabled() const
 {
-	for(UDetailsViewWrapperObject* ObjectBeingCustomized : ObjectsBeingCustomized)
+	for(TWeakObjectPtr<UDetailsViewWrapperObject> ObjectBeingCustomized : ObjectsBeingCustomized)
 	{
-		if(FRigControlElement* ControlElement = Cast<FRigControlElement>(ObjectBeingCustomized->GetContent<FRigBaseElement>()))
+		if(ObjectBeingCustomized.IsValid())
 		{
-			if(ControlElement->Settings.bGizmoEnabled)
+			if(FRigControlElement* ControlElement = Cast<FRigControlElement>(ObjectBeingCustomized->GetContent<FRigBaseElement>()))
 			{
-				return true;
+				if(ControlElement->Settings.bGizmoEnabled)
+				{
+					return true;
+				}
 			}
 		}
 	}
@@ -1800,23 +1822,26 @@ bool FRigControlElementDetails::IsGizmoEnabled() const
 
 bool FRigControlElementDetails::IsEnabled(ERigControlValueType InValueType) const
 {
-	for(UDetailsViewWrapperObject* ObjectBeingCustomized : ObjectsBeingCustomized)
+	for(TWeakObjectPtr<UDetailsViewWrapperObject> ObjectBeingCustomized : ObjectsBeingCustomized)
 	{
-		if(FRigControlElement* ControlElement = Cast<FRigControlElement>(ObjectBeingCustomized->GetContent<FRigBaseElement>()))
+		if(ObjectBeingCustomized.IsValid())
 		{
-			switch (InValueType)
+			if(FRigControlElement* ControlElement = Cast<FRigControlElement>(ObjectBeingCustomized->GetContent<FRigBaseElement>()))
 			{
-				case ERigControlValueType::Minimum:
-				case ERigControlValueType::Maximum:
+				switch (InValueType)
 				{
-					if(ControlElement->Settings.bLimitTranslation || ControlElement->Settings.bLimitRotation || ControlElement->Settings.bLimitScale)
+					case ERigControlValueType::Minimum:
+					case ERigControlValueType::Maximum:
 					{
-						return true;
+						if(ControlElement->Settings.bLimitTranslation || ControlElement->Settings.bLimitRotation || ControlElement->Settings.bLimitScale)
+						{
+							return true;
+						}
+					}
+					default:
+					{
+						break;
 					}
-				}
-				default:
-				{
-					break;
 				}
 			}
 		}
@@ -1840,15 +1865,18 @@ FText FRigControlElementDetails::GetDisplayName() const
 
 	for(int32 ObjectIndex = 0; ObjectIndex < ObjectsBeingCustomized.Num(); ObjectIndex++)
 	{
-		UDetailsViewWrapperObject* ObjectBeingCustomized = ObjectsBeingCustomized[ObjectIndex];
-		const FRigControlElement* ControlElement = ObjectBeingCustomized->GetContent<FRigControlElement>();
-		if(ObjectIndex == 0)
-		{
-			DisplayName = ControlElement->Settings.DisplayName;
-		}
-		else if(DisplayName != ControlElement->Settings.DisplayName)
+		TWeakObjectPtr<UDetailsViewWrapperObject> ObjectBeingCustomized = ObjectsBeingCustomized[ObjectIndex];
+		if(ObjectBeingCustomized.IsValid())
 		{
-			return ControlRigDetailsMultipleValues;
+			const FRigControlElement* ControlElement = ObjectBeingCustomized->GetContent<FRigControlElement>();
+			if(ObjectIndex == 0)
+			{
+				DisplayName = ControlElement->Settings.DisplayName;
+			}
+			else if(DisplayName != ControlElement->Settings.DisplayName)
+			{
+				return ControlRigDetailsMultipleValues;
+			}
 		}
 	}
 
@@ -1865,15 +1893,18 @@ void FRigControlElementDetails::SetDisplayName(const FText& InNewText, ETextComm
 	
 	for(int32 ObjectIndex = 0; ObjectIndex < ObjectsBeingCustomized.Num(); ObjectIndex++)
 	{
-		UDetailsViewWrapperObject* ObjectBeingCustomized = ObjectsBeingCustomized[ObjectIndex];
-		FRigControlElement* ControlElement = ObjectBeingCustomized->GetContent<FRigControlElement>();
+		TWeakObjectPtr<UDetailsViewWrapperObject> ObjectBeingCustomized = ObjectsBeingCustomized[ObjectIndex];
+		if(ObjectBeingCustomized.IsValid())
+		{
+			FRigControlElement* ControlElement = ObjectBeingCustomized->GetContent<FRigControlElement>();
 
-		FRigControlSettings Settings = ControlElement->Settings;
-		Settings.DisplayName = DisplayName;
+			FRigControlSettings Settings = ControlElement->Settings;
+			Settings.DisplayName = DisplayName;
 
-		if(GetHierarchy())
-		{
-			GetHierarchy()->SetControlSettings(ControlElement->GetKey(), Settings);
+			if(GetHierarchy())
+			{
+				GetHierarchy()->SetControlSettings(ControlElement->GetKey(), Settings);
+			}
 		}
 	}
 }
diff --git a/Engine/Plugins/Animation/ControlRig/Source/ControlRigEditor/Private/ControlRigElementDetails.h b/Engine/Plugins/Animation/ControlRig/Source/ControlRigEditor/Private/ControlRigElementDetails.h
--- a/Engine/Plugins/Animation/ControlRig/Source/ControlRigEditor/Private/ControlRigElementDetails.h
+++ b/Engine/Plugins/Animation/ControlRig/Source/ControlRigEditor/Private/ControlRigElementDetails.h
@@ -142,9 +142,12 @@ public:
 	TArray<T*> GetElementsInDetailsView() const
 	{
 		TArray<T*> Elements;
-		for(UDetailsViewWrapperObject* ObjectBeingCustomized : ObjectsBeingCustomized)
+		for(TWeakObjectPtr<UDetailsViewWrapperObject> ObjectBeingCustomized : ObjectsBeingCustomized)
 		{
-			Elements.Add(Cast<T>(ObjectBeingCustomized->GetContent<FRigBaseElement>()));
+			if(ObjectBeingCustomized.IsValid())
+			{
+				Elements.Add(Cast<T>(ObjectBeingCustomized->GetContent<FRigBaseElement>()));
+			}
 		}
 		return Elements;
 	}
@@ -155,9 +158,12 @@ public:
 		TArray<T*> Elements;
 		if(HierarchyBeingCustomized)
 		{
-			for(UDetailsViewWrapperObject* ObjectBeingCustomized : ObjectsBeingCustomized)
+			for(TWeakObjectPtr<UDetailsViewWrapperObject> ObjectBeingCustomized : ObjectsBeingCustomized)
 			{
-				Elements.Add(HierarchyBeingCustomized->Find<T>(ObjectBeingCustomized->GetContent<FRigBaseElement>()->GetKey()));
+				if(ObjectBeingCustomized.IsValid())
+				{
+					Elements.Add(HierarchyBeingCustomized->Find<T>(ObjectBeingCustomized->GetContent<FRigBaseElement>()->GetKey()));
+				}
 			}
 		}
 		return Elements;
@@ -172,7 +178,7 @@ protected:
 
 	UControlRigBlueprint* BlueprintBeingCustomized;
 	URigHierarchy* HierarchyBeingCustomized;
-	TArray<UDetailsViewWrapperObject*> ObjectsBeingCustomized;
+	TArray<TWeakObjectPtr<UDetailsViewWrapperObject>> ObjectsBeingCustomized;
 };
 
 class FRigTransformElementDetails : public FRigBaseElementDetails
diff --git a/Engine/Plugins/Animation/ControlRig/Source/ControlRigEditor/Private/Editor/ControlRigEditor.cpp b/Engine/Plugins/Animation/ControlRig/Source/ControlRigEditor/Private/Editor/ControlRigEditor.cpp
--- a/Engine/Plugins/Animation/ControlRig/Source/ControlRigEditor/Private/Editor/ControlRigEditor.cpp
+++ b/Engine/Plugins/Animation/ControlRig/Source/ControlRigEditor/Private/Editor/ControlRigEditor.cpp
@@ -1342,6 +1342,7 @@ void FControlRigEditor::SetDetailObjects(const TArray<UObject*>& InObjects)
 				UDetailsViewWrapperObject* WrapperObject = UDetailsViewWrapperObject::MakeInstance(
 					NodeDetailStruct, StructInstance->GetStructMemory(), UnitNode);
 				WrapperObject->GetWrappedPropertyChangedChainEvent().AddSP(this, &FControlRigEditor::OnWrappedPropertyChangedChainEvent);
+				WrapperObject->AddToRoot();
 				
 				FilteredObjects.Add(WrapperObject);
 				continue;
@@ -1424,6 +1425,7 @@ void FControlRigEditor::SetDetailViewForRigElements()
 
 		UDetailsViewWrapperObject* WrapperObject = UDetailsViewWrapperObject::MakeInstance(Element->GetElementStruct(), (uint8*)Element, Hierarchy);
 		WrapperObject->GetWrappedPropertyChangedChainEvent().AddSP(this, &FControlRigEditor::OnWrappedPropertyChangedChainEvent);
+		WrapperObject->AddToRoot();
 
 		Objects.Add(WrapperObject);
 	}
