Control Rig: Details Panel to rely on the correct blueprint (not protos)

diff --git a/Engine/Plugins/Animation/ControlRig/Source/ControlRigEditor/Private/ControlRigElementDetails.cpp b/Engine/Plugins/Animation/ControlRig/Source/ControlRigEditor/Private/ControlRigElementDetails.cpp
--- a/Engine/Plugins/Animation/ControlRig/Source/ControlRigEditor/Private/ControlRigElementDetails.cpp
+++ b/Engine/Plugins/Animation/ControlRig/Source/ControlRigEditor/Private/ControlRigElementDetails.cpp
@@ -74,10 +74,14 @@ void RigElementDetails_GetCustomizedInfo(TSharedRef<IPropertyHandle> InStructPro
 			{
 				if (UControlRig* Rig = Cast<UControlRig>(SubObject))
 				{
-					OutBlueprint = Cast<UControlRigBlueprint>(Rig->GetClass()->ClassGeneratedBy);
-					if (OutBlueprint)
+					UControlRigBlueprint* Blueprint = Cast<UControlRigBlueprint>(Rig->GetClass()->ClassGeneratedBy);
+					if (Blueprint)
 					{
-						break;
+						if(Blueprint->GetOutermost() == Package)
+						{
+							OutBlueprint = Blueprint;
+							break;
+						}
 					}
 				}
 			}
@@ -723,8 +727,6 @@ void FRigBaseElementDetails::CustomizeHeader(TSharedRef<IPropertyHandle> InStruc
 	HierarchyBeingCustomized = nullptr;
 	ObjectsBeingCustomized.Reset();
 	
-	RigElementDetails_GetCustomizedInfo(InStructPropertyHandle, BlueprintBeingCustomized);
-
 	TArray<UObject*> Objects;
 	InStructPropertyHandle->GetOuterObjects(Objects);
 	for (UObject* Object : Objects)
@@ -738,6 +740,23 @@ void FRigBaseElementDetails::CustomizeHeader(TSharedRef<IPropertyHandle> InStruc
 
 		ObjectsBeingCustomized.Add(WrapperObject);
 	}
+
+	if(HierarchyBeingCustomized)
+	{
+		BlueprintBeingCustomized = HierarchyBeingCustomized->GetTypedOuter<UControlRigBlueprint>();
+		if(BlueprintBeingCustomized == nullptr)
+		{
+			if(UControlRig* ControlRig = HierarchyBeingCustomized->GetTypedOuter<UControlRig>())
+			{
+				BlueprintBeingCustomized = Cast<UControlRigBlueprint>(ControlRig->GetClass()->ClassGeneratedBy);
+			}
+		}
+	}
+
+	if(BlueprintBeingCustomized == nullptr)
+	{	
+		RigElementDetails_GetCustomizedInfo(InStructPropertyHandle, BlueprintBeingCustomized);
+	}
 }
 
 void FRigBaseElementDetails::CustomizeChildren(TSharedRef<class IPropertyHandle> InStructPropertyHandle, class IDetailChildrenBuilder& StructBuilder, IPropertyTypeCustomizationUtils& StructCustomizationUtils)
