Control Rig: Python output when pasting/duplicating/mirroring rig elements

diff --git a/Engine/Plugins/Animation/ControlRig/Source/ControlRig/Private/Rigs/RigHierarchyController.cpp b/Engine/Plugins/Animation/ControlRig/Source/ControlRig/Private/Rigs/RigHierarchyController.cpp
--- a/Engine/Plugins/Animation/ControlRig/Source/ControlRig/Private/Rigs/RigHierarchyController.cpp
+++ b/Engine/Plugins/Animation/ControlRig/Source/ControlRig/Private/Rigs/RigHierarchyController.cpp
@@ -894,7 +894,7 @@ public:
 	}
 };
 
-TArray<FRigElementKey> URigHierarchyController::ImportFromText(FString InContent, bool bReplaceExistingElements, bool bSelectNewElements, bool bSetupUndo)
+TArray<FRigElementKey> URigHierarchyController::ImportFromText(FString InContent, bool bReplaceExistingElements, bool bSelectNewElements, bool bSetupUndo, bool bPrintPythonCommands)
 {
 	TArray<FRigElementKey> PastedKeys;
 	if(!IsValid())
@@ -1084,6 +1084,18 @@ TArray<FRigElementKey> URigHierarchyController::ImportFromText(FString InContent
 		SetSelection(PastedKeys);
 	}
 
+#if WITH_EDITOR
+	if (bPrintPythonCommands)
+	{
+		FString PythonContent = InContent.Replace(TEXT("\\\""), TEXT("\\\\\""));
+		
+		RigVMPythonUtils::Print(FString::Printf(TEXT("hierarchy_controller.import_from_text('%s', %s, %s)"),
+			*PythonContent,
+			(bReplaceExistingElements) ? TEXT("True") : TEXT("False"),
+			(bSelectNewElements) ? TEXT("True") : TEXT("False")));
+	}
+#endif
+
 	return PastedKeys;
 }
 
@@ -2213,13 +2225,35 @@ bool URigHierarchyController::SetParent(FRigElementKey InChild, FRigElementKey I
 	return bParentSet;
 }
 
-TArray<FRigElementKey> URigHierarchyController::DuplicateElements(TArray<FRigElementKey> InKeys, bool bSelectNewElements, bool bSetupUndo)
+TArray<FRigElementKey> URigHierarchyController::DuplicateElements(TArray<FRigElementKey> InKeys, bool bSelectNewElements, bool bSetupUndo, bool bPrintPythonCommands)
 {
 	const FString Content = ExportToText(InKeys);
-	return ImportFromText(Content, false, bSelectNewElements, bSetupUndo);
+	TArray<FRigElementKey> Result = ImportFromText(Content, false, bSelectNewElements, bSetupUndo);
+
+#if WITH_EDITOR
+	if (!Result.IsEmpty() && bPrintPythonCommands)
+	{
+		FString ArrayStr = TEXT("[");
+		for (auto It = InKeys.CreateConstIterator(); It; ++It)
+		{
+			ArrayStr += It->ToPythonString();
+			if (It.GetIndex() < InKeys.Num() - 1)
+			{
+				ArrayStr += TEXT(", ");
+			}
+		}
+		ArrayStr += TEXT("]");		
+		
+		RigVMPythonUtils::Print(FString::Printf(TEXT("hierarchy_controller.duplicate_elements(%s, %s)"),
+			*ArrayStr,
+			(bSelectNewElements) ? TEXT("True") : TEXT("False")));
+	}
+#endif
+
+	return Result;
 }
 
-TArray<FRigElementKey> URigHierarchyController::MirrorElements(TArray<FRigElementKey> InKeys, FRigMirrorSettings InSettings, bool bSelectNewElements, bool bSetupUndo)
+TArray<FRigElementKey> URigHierarchyController::MirrorElements(TArray<FRigElementKey> InKeys, FRigMirrorSettings InSettings, bool bSelectNewElements, bool bSetupUndo, bool bPrintPythonCommands)
 {
 	TArray<FRigElementKey> OriginalKeys = Hierarchy->SortKeys(InKeys);
 	TArray<FRigElementKey> DuplicatedKeys = DuplicateElements(OriginalKeys, bSelectNewElements, bSetupUndo);
@@ -2311,6 +2345,30 @@ TArray<FRigElementKey> URigHierarchyController::MirrorElements(TArray<FRigElemen
 		}
 	}
 
+#if WITH_EDITOR
+	if (!DuplicatedKeys.IsEmpty() && bPrintPythonCommands)
+	{
+		FString ArrayStr = TEXT("[");
+		for (auto It = InKeys.CreateConstIterator(); It; ++It)
+		{
+			ArrayStr += It->ToPythonString();
+			if (It.GetIndex() < InKeys.Num() - 1)
+			{
+				ArrayStr += TEXT(", ");
+			}
+		}
+		ArrayStr += TEXT("]");
+
+		RigVMPythonUtils::Print(FString::Printf(TEXT("hierarchy_controller.mirror_elements(%s, unreal.RigMirrorSettings(unreal.AxisType.%s, unreal.AxisType.%s, '%s', '%s'), %s)"),
+			*ArrayStr,
+			(InSettings.MirrorAxis.GetValue() == EAxis::X) ? TEXT("X") : (InSettings.MirrorAxis.GetValue() == EAxis::Y) ? TEXT("Y") : TEXT("Z"),
+			(InSettings.AxisToFlip.GetValue() == EAxis::X) ? TEXT("X") : (InSettings.AxisToFlip.GetValue() == EAxis::Y) ? TEXT("Y") : TEXT("Z"),
+			*InSettings.OldName,
+			*InSettings.NewName,
+			(bSelectNewElements) ? TEXT("True") : TEXT("False")));
+	}
+#endif
+	
 	return DuplicatedKeys;
 }
 
diff --git a/Engine/Plugins/Animation/ControlRig/Source/ControlRig/Public/Rigs/RigHierarchyController.h b/Engine/Plugins/Animation/ControlRig/Source/ControlRig/Public/Rigs/RigHierarchyController.h
--- a/Engine/Plugins/Animation/ControlRig/Source/ControlRig/Public/Rigs/RigHierarchyController.h
+++ b/Engine/Plugins/Animation/ControlRig/Source/ControlRig/Public/Rigs/RigHierarchyController.h
@@ -321,13 +321,15 @@ public:
 	 * @param bReplaceExistingElements If set to true existing items will be replaced / updated with the content in the buffer
 	 * @param bSelectNewElements If set to true the new elements will be selected
 	 * @param bSetupUndo If set to true the stack will record the change for undo / redo
+	 * @param bPrintPythonCommand If set to true a python command equivalent to this call will be printed out
 	 */
 	UFUNCTION(BlueprintCallable, Category = URigHierarchyController)
 	TArray<FRigElementKey> ImportFromText(
 		FString InContent,
 		bool bReplaceExistingElements = false,
 		bool bSelectNewElements = true,
-		bool bSetupUndo = false);
+		bool bSetupUndo = false,
+		bool bPrintPythonCommands = false);
 
 	/**
 	* Imports the content of a RigHierachyContainer (the hierarchy v1 pre 5.0)
@@ -408,10 +410,11 @@ public:
 	 * @param InKeys The keys of the elements to duplicate
 	 * @param bSelectNewElements If set to true the new elements will be selected
 	 * @param bSetupUndo If set to true the stack will record the change for undo / redo
+	 * @param bPrintPythonCommand If set to true a python command equivalent to this call will be printed out
 	 * @return The keys of the 4d items
 	 */
 	UFUNCTION(BlueprintCallable, Category = URigHierarchyController)
-    TArray<FRigElementKey> DuplicateElements(TArray<FRigElementKey> InKeys, bool bSelectNewElements = true, bool bSetupUndo = false);
+    TArray<FRigElementKey> DuplicateElements(TArray<FRigElementKey> InKeys, bool bSelectNewElements = true, bool bSetupUndo = false, bool bPrintPythonCommands = false);
 
 	/**
 	 * Mirrors the given elements
@@ -419,10 +422,11 @@ public:
 	 * @param InSettings The settings to use for the mirror operation
 	 * @param bSelectNewElements If set to true the new elements will be selected
 	 * @param bSetupUndo If set to true the stack will record the change for undo / redo
+	 * @param bPrintPythonCommand If set to true a python command equivalent to this call will be printed out
 	 * @return The keys of the mirrored items
 	 */
 	UFUNCTION(BlueprintCallable, Category = URigHierarchyController)
-    TArray<FRigElementKey> MirrorElements(TArray<FRigElementKey> InKeys, FRigMirrorSettings InSettings, bool bSelectNewElements = true, bool bSetupUndo = false);
+    TArray<FRigElementKey> MirrorElements(TArray<FRigElementKey> InKeys, FRigMirrorSettings InSettings, bool bSelectNewElements = true, bool bSetupUndo = false, bool bPrintPythonCommands = false);
 
 	/**
 	 * Returns the modified event, which can be used to 
diff --git a/Engine/Plugins/Animation/ControlRig/Source/ControlRig/Public/Rigs/RigHierarchyDefines.h b/Engine/Plugins/Animation/ControlRig/Source/ControlRig/Public/Rigs/RigHierarchyDefines.h
--- a/Engine/Plugins/Animation/ControlRig/Source/ControlRig/Public/Rigs/RigHierarchyDefines.h
+++ b/Engine/Plugins/Animation/ControlRig/Source/ControlRig/Public/Rigs/RigHierarchyDefines.h
@@ -1257,19 +1257,19 @@ struct CONTROLRIG_API FRigMirrorSettings
 	}
 
 	// the axis to mirror against
-	UPROPERTY(EditAnywhere, Category = Settings)
+	UPROPERTY(BlueprintReadWrite, EditAnywhere, Category = Settings)
 	TEnumAsByte<EAxis::Type> MirrorAxis;
 
 	// the axis to flip for rotations
-	UPROPERTY(EditAnywhere, Category = Settings)
+	UPROPERTY(BlueprintReadWrite, EditAnywhere, Category = Settings)
 	TEnumAsByte<EAxis::Type> AxisToFlip;
 
 	// the string to replace all occurences of with New Name
-	UPROPERTY(EditAnywhere, Category = Settings)
+	UPROPERTY(BlueprintReadWrite, EditAnywhere, Category = Settings)
 	FString OldName;
 
 	// the string to replace all occurences of Old Name with
-	UPROPERTY(EditAnywhere, Category = Settings)
+	UPROPERTY(BlueprintReadWrite, EditAnywhere, Category = Settings)
 	FString NewName;
 
 	FTransform MirrorTransform(const FTransform& InTransform) const;
diff --git a/Engine/Plugins/Animation/ControlRig/Source/ControlRigEditor/Private/Editor/SRigHierarchy.cpp b/Engine/Plugins/Animation/ControlRig/Source/ControlRigEditor/Private/Editor/SRigHierarchy.cpp
--- a/Engine/Plugins/Animation/ControlRig/Source/ControlRigEditor/Private/Editor/SRigHierarchy.cpp
+++ b/Engine/Plugins/Animation/ControlRig/Source/ControlRigEditor/Private/Editor/SRigHierarchy.cpp
@@ -1967,7 +1967,7 @@ void SRigHierarchy::HandleDuplicateItem()
 			check(Controller);
 
 			const TArray<FRigElementKey> KeysToDuplicate = GetSelectedKeys();
-			Controller->DuplicateElements(KeysToDuplicate, true, true);
+			Controller->DuplicateElements(KeysToDuplicate, true, true, true);
 		}
 
 		ControlRigBlueprint->PropagateHierarchyFromBPToInstances();
@@ -2021,7 +2021,7 @@ void SRigHierarchy::HandleMirrorItem()
 
 				const TArray<FRigElementKey> KeysToMirror = GetSelectedKeys();
 				const TArray<FRigElementKey> KeysToDuplicate = GetSelectedKeys();
-				Controller->MirrorElements(KeysToDuplicate, Settings, true, true);
+				Controller->MirrorElements(KeysToDuplicate, Settings, true, true, true);
 			}
 			ControlRigBlueprint->PropagateHierarchyFromBPToInstances();
 		}
@@ -2145,7 +2145,7 @@ void SRigHierarchy::HandlePasteItems()
 		URigHierarchyController* Controller = Hierarchy->GetController(true);
 		check(Controller);
 
-		Controller->ImportFromText(Content, false, true, true);
+		Controller->ImportFromText(Content, false, true, true, true);
 	}
 
 	//ControlRigBlueprint->PropagateHierarchyFromBPToInstances();
