Crash using ControlRig in Sequencer
#info save/use cached pose nodes are not correctly invalidated/refreshed due to ControlRig edit mode invoking evaluation of the same anim instance multiple times per frame
#fix ensure that various counters used for synchronization/invalidation of cached data are used by ControlRigLayer proxy - and applied whenever update/evaluating the source "input" proxy

diff --git a/Engine/Plugins/Animation/ControlRig/Source/ControlRig/Private/Sequencer/ControlRigLayerInstanceProxy.cpp b/Engine/Plugins/Animation/ControlRig/Source/ControlRig/Private/Sequencer/ControlRigLayerInstanceProxy.cpp
--- a/Engine/Plugins/Animation/ControlRig/Source/ControlRig/Private/Sequencer/ControlRigLayerInstanceProxy.cpp
+++ b/Engine/Plugins/Animation/ControlRig/Source/ControlRig/Private/Sequencer/ControlRigLayerInstanceProxy.cpp
@@ -11,13 +11,18 @@ void FControlRigLayerInstanceProxy::Initialize(UAnimInstance* InAnimInstance)
 
 	FAnimInstanceProxy::Initialize(InAnimInstance);
 
+	InitializationCounter.Increment();
 	UpdateCounter.Reset();
+	EvaluationCounter.Reset();
+	CachedBonesCounter.Reset();
 }
 
 void FControlRigLayerInstanceProxy::CacheBones()
 {
 	if (bBoneCachesInvalidated)
 	{
+		CachedBonesCounter.Increment();
+		
 		FAnimationCacheBonesContext Context(this);
 		check(CurrentRoot);
 		CurrentRoot->CacheBones_AnyThread(Context);
@@ -28,6 +33,8 @@ void FControlRigLayerInstanceProxy::CacheBones()
 
 bool FControlRigLayerInstanceProxy::Evaluate(FPoseContext& Output)
 {
+	EvaluationCounter.Increment();
+	
 	check(CurrentRoot);
 	CurrentRoot->Evaluate_AnyThread(Output);
 	return true;
@@ -292,6 +299,27 @@ void FControlRigLayerInstanceProxy::ResetCounter(FAnimInstanceProxy* InAnimInsta
 {
 	FAnimInstanceProxy::ResetCounterInputProxy(InAnimInstanceProxy);
 }
+
+void FControlRigLayerInstanceProxy::SynchronizeProxyInitializationCounter(FAnimInstanceProxy* Proxy, FAnimInstanceProxy* SyncProxy)
+{
+	FAnimInstanceProxy::SynchronizeInitializationCounter(Proxy, SyncProxy);
+}
+
+void FControlRigLayerInstanceProxy::SynchronizeProxyCacheBonesCounter(FAnimInstanceProxy* Proxy, FAnimInstanceProxy* SyncProxy)
+{
+	FAnimInstanceProxy::SynchronizeCacheBonesCounter(Proxy, SyncProxy);
+}
+
+void FControlRigLayerInstanceProxy::SynchronizeProxyUpdateCounter(FAnimInstanceProxy* Proxy, FAnimInstanceProxy* SyncProxy)
+{
+	FAnimInstanceProxy::SynchronizeUpdateCounter(Proxy, SyncProxy);
+}
+
+void FControlRigLayerInstanceProxy::SynchronizeProxyEvaluationCounter(FAnimInstanceProxy* Proxy, FAnimInstanceProxy* SyncProxy)
+{
+	FAnimInstanceProxy::SynchronizeEvaluationCounter(Proxy, SyncProxy);
+}
+
 //////////////////////////////////////////////////////////////////////////////////////////////////
 void FAnimNode_ControlRigInputPose::Initialize_AnyThread(const FAnimationInitializeContext& Context)
 {
@@ -300,6 +328,7 @@ void FAnimNode_ControlRigInputPose::Initialize_AnyThread(const FAnimationInitial
 		FAnimationInitializeContext InputContext(InputProxy);
 		if (InputPose.GetLinkNode())
 		{
+			FControlRigLayerInstanceProxy::SynchronizeProxyInitializationCounter(InputProxy, Context.AnimInstanceProxy);
 			InputPose.Initialize(InputContext);
 		}
  		else 
@@ -316,6 +345,7 @@ void FAnimNode_ControlRigInputPose::CacheBones_AnyThread(const FAnimationCacheBo
 		FAnimationCacheBonesContext InputContext(InputProxy);
 		if (InputPose.GetLinkNode())
 		{
+			FControlRigLayerInstanceProxy::SynchronizeProxyCacheBonesCounter(InputProxy, Context.AnimInstanceProxy);
 			InputPose.CacheBones(InputContext);
 		}
 		else
@@ -332,6 +362,7 @@ void FAnimNode_ControlRigInputPose::Update_AnyThread(const FAnimationUpdateConte
 		FAnimationUpdateContext InputContext = Context.WithOtherProxy(InputProxy);
 		if (InputPose.GetLinkNode())
 		{
+			FControlRigLayerInstanceProxy::SynchronizeProxyUpdateCounter(InputProxy, Context.AnimInstanceProxy);
 			InputPose.Update(InputContext);
 		}
 		else
@@ -354,6 +385,7 @@ void FAnimNode_ControlRigInputPose::Evaluate_AnyThread(FPoseContext& Output)
 			// if no linked node, just use Evaluate of proxy
 			if (InputPose.GetLinkNode())
 			{
+				FControlRigLayerInstanceProxy::SynchronizeProxyEvaluationCounter(InputProxy, Output.AnimInstanceProxy);
 				InputPose.Evaluate(InputContext);
 			}
 			else
diff --git a/Engine/Plugins/Animation/ControlRig/Source/ControlRig/Private/Sequencer/ControlRigLayerInstanceProxy.h b/Engine/Plugins/Animation/ControlRig/Source/ControlRig/Private/Sequencer/ControlRigLayerInstanceProxy.h
--- a/Engine/Plugins/Animation/ControlRig/Source/ControlRig/Private/Sequencer/ControlRigLayerInstanceProxy.h
+++ b/Engine/Plugins/Animation/ControlRig/Source/ControlRig/Private/Sequencer/ControlRigLayerInstanceProxy.h
@@ -135,4 +135,9 @@ protected:
 	static void EvaluateCustomProxy(FAnimInstanceProxy* InputProxy, FPoseContext& Output);
 	/** reset internal Counters of given animinstance proxy */
 	static void ResetCounter(FAnimInstanceProxy* InAnimInstanceProxy);
+	
+	static void SynchronizeProxyInitializationCounter(FAnimInstanceProxy* Proxy, FAnimInstanceProxy* SyncProxy);
+	static void SynchronizeProxyCacheBonesCounter(FAnimInstanceProxy* Proxy, FAnimInstanceProxy* SyncProxy);
+	static void SynchronizeProxyUpdateCounter(FAnimInstanceProxy* Proxy, FAnimInstanceProxy* SyncProxy);
+	static void SynchronizeProxyEvaluationCounter(FAnimInstanceProxy* Proxy, FAnimInstanceProxy* SyncProxy);
 };
diff --git a/Engine/Source/Runtime/AnimGraphRuntime/Public/AnimCustomInstanceHelper.h b/Engine/Source/Runtime/AnimGraphRuntime/Public/AnimCustomInstanceHelper.h
--- a/Engine/Source/Runtime/AnimGraphRuntime/Public/AnimCustomInstanceHelper.h
+++ b/Engine/Source/Runtime/AnimGraphRuntime/Public/AnimCustomInstanceHelper.h
@@ -272,6 +272,7 @@ public:
 				}
 
 				AnimInstance->UpdateAnimation(0.0f, false);
+				AnimInstance->RecalcRequiredBones();
 			}
 
 			// Update space bases to reset it back to ref pose
diff --git a/Engine/Source/Runtime/Engine/Private/Animation/AnimInstanceProxy.cpp b/Engine/Source/Runtime/Engine/Private/Animation/AnimInstanceProxy.cpp
--- a/Engine/Source/Runtime/Engine/Private/Animation/AnimInstanceProxy.cpp
+++ b/Engine/Source/Runtime/Engine/Private/Animation/AnimInstanceProxy.cpp
@@ -3257,4 +3257,36 @@ void FAnimInstanceProxy::ResetCounterInputProxy(FAnimInstanceProxy* InputProxy)
 	}
 }
 
+void FAnimInstanceProxy::SynchronizeInitializationCounter(FAnimInstanceProxy* Proxy, FAnimInstanceProxy* SyncProxy)
+{
+	if (Proxy && SyncProxy)
+	{
+		Proxy->InitializationCounter.SynchronizeWith(SyncProxy->InitializationCounter);
+	}
+}
+
+void FAnimInstanceProxy::SynchronizeCacheBonesCounter(FAnimInstanceProxy* Proxy, FAnimInstanceProxy* SyncProxy)
+{
+	if (Proxy && SyncProxy)
+	{
+		Proxy->CachedBonesCounter.SynchronizeWith(SyncProxy->CachedBonesCounter);
+	}
+}
+
+void FAnimInstanceProxy::SynchronizeUpdateCounter(FAnimInstanceProxy* Proxy, FAnimInstanceProxy* SyncProxy)
+{
+	if (Proxy && SyncProxy)
+	{
+		Proxy->UpdateCounter.SynchronizeWith(SyncProxy->UpdateCounter);
+	}
+}
+
+void FAnimInstanceProxy::SynchronizeEvaluationCounter(FAnimInstanceProxy* Proxy, FAnimInstanceProxy* SyncProxy)
+{
+	if (Proxy && SyncProxy)
+	{
+		Proxy->EvaluationCounter.SynchronizeWith(SyncProxy->EvaluationCounter);
+	}
+}
+
 #undef LOCTEXT_NAMESPACE
diff --git a/Engine/Source/Runtime/Engine/Public/Animation/AnimInstanceProxy.h b/Engine/Source/Runtime/Engine/Public/Animation/AnimInstanceProxy.h
--- a/Engine/Source/Runtime/Engine/Public/Animation/AnimInstanceProxy.h
+++ b/Engine/Source/Runtime/Engine/Public/Animation/AnimInstanceProxy.h
@@ -964,6 +964,10 @@ protected:
 	static void EvaluateInputProxy(FAnimInstanceProxy* InputProxy, FPoseContext& Output);
 	static void ResetCounterInputProxy(FAnimInstanceProxy* InputProxy);
 
+	static void SynchronizeInitializationCounter(FAnimInstanceProxy* Proxy, FAnimInstanceProxy* SyncProxy);
+	static void SynchronizeCacheBonesCounter(FAnimInstanceProxy* Proxy, FAnimInstanceProxy* SyncProxy);
+	static void SynchronizeUpdateCounter(FAnimInstanceProxy* Proxy, FAnimInstanceProxy* SyncProxy);
+	static void SynchronizeEvaluationCounter(FAnimInstanceProxy* Proxy, FAnimInstanceProxy* SyncProxy);
 private:
 	/** Evaluate the slot when there are blend spaces involved in any of the active anim montages. */
 	void SlotEvaluatePoseWithBlendProfiles(const FName& SlotNodeName, const FAnimationPoseData& SourceAnimationPoseData, float InSourceWeight, FAnimationPoseData& OutBlendedAnimationPoseData, float InBlendWeight);
