Control Rig: Make property on detailswrapperobject hidden to blueprint

diff --git a/Engine/Plugins/Animation/ControlRig/Source/ControlRigEditor/Private/Editor/DetailsViewWrapperObject.cpp b/Engine/Plugins/Animation/ControlRig/Source/ControlRigEditor/Private/Editor/DetailsViewWrapperObject.cpp
--- a/Engine/Plugins/Animation/ControlRig/Source/ControlRigEditor/Private/Editor/DetailsViewWrapperObject.cpp
+++ b/Engine/Plugins/Animation/ControlRig/Source/ControlRigEditor/Private/Editor/DetailsViewWrapperObject.cpp
@@ -76,7 +76,7 @@ UClass* UDetailsViewWrapperObject::GetClassForStruct(UScriptStruct* InStruct)
 	}
 	
 	// Make sure the variables show up in the details panel
-	Property->SetPropertyFlags(CPF_Edit | CPF_BlueprintVisible);
+	Property->SetPropertyFlags(CPF_Edit);
 	Property->SetMetaData(TEXT("ShowOnlyInnerProperties"), TEXT("true"));
 
 	// For rig units mark all inputs as edit
