Control Rig: Guard against silent failure when collapsing nodes to function

diff --git a/Engine/Plugins/Animation/ControlRig/Source/ControlRigEditor/Private/ControlRigEditorModule.cpp b/Engine/Plugins/Animation/ControlRig/Source/ControlRigEditor/Private/ControlRigEditorModule.cpp
--- a/Engine/Plugins/Animation/ControlRig/Source/ControlRigEditor/Private/ControlRigEditorModule.cpp
+++ b/Engine/Plugins/Animation/ControlRig/Source/ControlRigEditor/Private/ControlRigEditorModule.cpp
@@ -1999,8 +1999,15 @@ void FControlRigEditorModule::GetContextMenuActions(const UControlRigGraphSchema
 							TArray<FName> Nodes = Model->GetSelectNodes();
 							Controller->OpenUndoBracket(TEXT("Collapse to Function"));
 							URigVMCollapseNode* CollapseNode = Controller->CollapseNodes(Nodes, TEXT("New Function"), true, true);
-							Controller->PromoteCollapseNodeToFunctionReferenceNode(CollapseNode->GetFName(), true, true);
-							Controller->CloseUndoBracket();
+							if(CollapseNode)
+							{
+								Controller->PromoteCollapseNodeToFunctionReferenceNode(CollapseNode->GetFName(), true, true);
+								Controller->CloseUndoBracket();
+							}
+							else
+							{
+								Controller->CancelUndoBracket();
+							}
 						})
 					));
 
