Control Rig: Support for LWC for Control Values

* Switched serialization / internal storage over to float variants / custom float structs
* Introduced template specizializations to route usage to correct storage / structure
* Marked certain template specializations as delete to catch / avoid usages

diff --git a/Engine/Plugins/Animation/ControlRig/Source/ControlRig/Private/Rigs/RigControlHierarchy.cpp b/Engine/Plugins/Animation/ControlRig/Source/ControlRig/Private/Rigs/RigControlHierarchy.cpp
--- a/Engine/Plugins/Animation/ControlRig/Source/ControlRig/Private/Rigs/RigControlHierarchy.cpp
+++ b/Engine/Plugins/Animation/ControlRig/Source/ControlRig/Private/Rigs/RigControlHierarchy.cpp
@@ -78,7 +78,7 @@ void FRigControlHierarchy::PostLoad()
 			FRigControlValue& Value = Control.GetValue((ERigControlValueType)ValueType);
 			if (!Value.IsValid())
 			{
-				Value.GetRef<FTransform>() = Value.Storage_DEPRECATED;
+				Value.GetRef<FRigControlValue::FTransform_Float>() = Value.Storage_DEPRECATED;
 			}
 		}
 	}
diff --git a/Engine/Plugins/Animation/ControlRig/Source/ControlRig/Public/Rigs/RigHierarchy.h b/Engine/Plugins/Animation/ControlRig/Source/ControlRig/Public/Rigs/RigHierarchy.h
--- a/Engine/Plugins/Animation/ControlRig/Source/ControlRig/Public/Rigs/RigHierarchy.h
+++ b/Engine/Plugins/Animation/ControlRig/Source/ControlRig/Public/Rigs/RigHierarchy.h
@@ -1912,7 +1912,7 @@ public:
 	UFUNCTION(BlueprintCallable, Category = URigHierarchy)
     static FORCEINLINE FRigControlValue MakeControlValueFromVector2D(FVector2D InValue)
 	{
-		return FRigControlValue::Make<FVector2D>(InValue);
+		return FRigControlValue::Make<FVector3f>(FVector3f(InValue.X, InValue.Y, 0.f));
 	}
 
 	/**
@@ -1923,7 +1923,8 @@ public:
 	UFUNCTION(BlueprintCallable, Category = URigHierarchy)
 	static FORCEINLINE FVector2D GetVector2DFromControlValue(FRigControlValue InValue)
 	{
-		return InValue.Get<FVector2D>();
+		const FVector3f Vector = InValue.Get<FVector3f>();
+		return FVector2D(Vector.X, Vector.Y);
 	}
 
 	/**
@@ -1934,7 +1935,7 @@ public:
 	UFUNCTION(BlueprintCallable, Category = URigHierarchy)
     static FORCEINLINE FRigControlValue MakeControlValueFromVector(FVector InValue)
 	{
-		return FRigControlValue::Make<FVector>(InValue);
+		return FRigControlValue::Make<FVector3f>(InValue);
 	}
 
 	/**
@@ -1945,7 +1946,7 @@ public:
 	UFUNCTION(BlueprintCallable, Category = URigHierarchy)
 	static FORCEINLINE FVector GetVectorFromControlValue(FRigControlValue InValue)
 	{
-		return InValue.Get<FVector>();
+		return InValue.Get<FVector3f>();
 	}
 
 	/**
@@ -1956,7 +1957,7 @@ public:
 	UFUNCTION(BlueprintCallable, Category = URigHierarchy)
     static FORCEINLINE FRigControlValue MakeControlValueFromRotator(FRotator InValue)
 	{
-		return FRigControlValue::Make<FRotator>(InValue);
+		return FRigControlValue::Make<FVector3f>(InValue.Vector());
 	}
 
 	/**
@@ -1967,7 +1968,7 @@ public:
 	UFUNCTION(BlueprintCallable, Category = URigHierarchy)
 	static FORCEINLINE FRotator GetRotatorFromControlValue(FRigControlValue InValue)
 	{
-		return InValue.Get<FRotator>();
+		return FRotator::MakeFromEuler(InValue.Get<FVector3f>());
 	}
 
 	/**
@@ -1978,7 +1979,7 @@ public:
 	UFUNCTION(BlueprintCallable, Category = URigHierarchy)
     static FORCEINLINE FRigControlValue MakeControlValueFromTransform(FTransform InValue)
 	{
-		return FRigControlValue::Make<FTransform>(InValue);
+		return FRigControlValue::Make<FRigControlValue::FTransform_Float>(InValue);
 	}
 
 	/**
@@ -1989,7 +1990,7 @@ public:
 	UFUNCTION(BlueprintCallable, Category = URigHierarchy)
 	static FORCEINLINE FTransform GetTransformFromControlValue(FRigControlValue InValue)
 	{
-		return InValue.Get<FTransform>();
+		return InValue.Get<FRigControlValue::FTransform_Float>().ToTransform();
 	}
 
 	/**
@@ -2000,7 +2001,7 @@ public:
 	UFUNCTION(BlueprintCallable, Category = URigHierarchy)
     static FORCEINLINE FRigControlValue MakeControlValueFromEulerTransform(FEulerTransform InValue)
 	{
-		return FRigControlValue::Make<FEulerTransform>(InValue);
+		return FRigControlValue::Make<FRigControlValue::FEulerTransform_Float>(InValue);
 	}
 
 	/**
@@ -2011,7 +2012,7 @@ public:
 	UFUNCTION(BlueprintCallable, Category = URigHierarchy)
 	static FORCEINLINE FEulerTransform GetEulerTransformFromControlValue(FRigControlValue InValue)
 	{
-		return InValue.Get<FEulerTransform>();
+		return InValue.Get<FRigControlValue::FEulerTransform_Float>().ToTransform();
 	}
 
 	/**
@@ -2022,7 +2023,7 @@ public:
 	UFUNCTION(BlueprintCallable, Category = URigHierarchy)
     static FORCEINLINE FRigControlValue MakeControlValueFromTransformNoScale(FTransformNoScale InValue)
 	{
-		return FRigControlValue::Make<FTransformNoScale>(InValue);
+		return FRigControlValue::Make<FRigControlValue::FTransformNoScale_Float>(InValue);
 	}
 
 	/**
@@ -2033,7 +2034,7 @@ public:
 	UFUNCTION(BlueprintCallable, Category = URigHierarchy)
 	static FORCEINLINE FTransformNoScale GetTransformNoScaleFromControlValue(FRigControlValue InValue)
 	{
-		return InValue.Get<FTransformNoScale>();
+		return InValue.Get<FRigControlValue::FTransformNoScale_Float>().ToTransform();
 	}
 
 private:
@@ -2453,6 +2454,23 @@ protected:
 	friend class FControlRigEditor;
 };
 
+#if !UE_LARGE_WORLD_COORDINATES_DISABLED
+
+template<>
+FORCEINLINE_DEBUGGABLE FVector2D URigHierarchy::GetControlValue(FRigControlElement* InControlElement, ERigControlValueType InValueType) const
+{
+	const FVector3f Value = GetControlValue(InControlElement, InValueType).Get<FVector3f>();
+	return FVector2D(Value.X, Value.Y);
+}
+
+template<>
+FORCEINLINE_DEBUGGABLE void URigHierarchy::SetControlValue(int32 InElementIndex, const FVector2D& InValue, ERigControlValueType InValueType, bool bSetupUndo) const
+{
+	return SetControlValue(InElementIndex, FRigControlValue::Make<FVector3f>(FVector3f(InValue.X, InValue.Y, 0.f)), InValueType, bSetupUndo);
+}
+
+#endif
+
 #if WITH_EDITOR
 
 class FRigHierarchyListenerGuard
diff --git a/Engine/Plugins/Animation/ControlRig/Source/ControlRig/Public/Rigs/RigHierarchyDefines.h b/Engine/Plugins/Animation/ControlRig/Source/ControlRig/Public/Rigs/RigHierarchyDefines.h
--- a/Engine/Plugins/Animation/ControlRig/Source/ControlRig/Public/Rigs/RigHierarchyDefines.h
+++ b/Engine/Plugins/Animation/ControlRig/Source/ControlRig/Public/Rigs/RigHierarchyDefines.h
@@ -344,7 +344,7 @@ public:
 		{
 			case ERigControlType::Bool:
 			{
-				Transform.SetLocation(FVector(Get<bool>() ? 1.f : 0.f, 0.f, 0.f));
+				Transform.SetLocation(FVector3f(Get<bool>() ? 1.f : 0.f, 0.f, 0.f));
 				break;
 			}
 			case ERigControlType::Float:
@@ -354,17 +354,17 @@ public:
 				{
 					case ERigControlAxis::X:
 					{
-						Transform.SetLocation(FVector(ValueToGet, 0.f, 0.f));
+						Transform.SetLocation(FVector3f(ValueToGet, 0.f, 0.f));
 						break;
 					}
 					case ERigControlAxis::Y:
 					{
-						Transform.SetLocation(FVector(0.f, ValueToGet, 0.f));
+						Transform.SetLocation(FVector3f(0.f, ValueToGet, 0.f));
 						break;
 					}
 					case ERigControlAxis::Z:
 					{
-						Transform.SetLocation(FVector(0.f, 0.f, ValueToGet));
+						Transform.SetLocation(FVector3f(0.f, 0.f, ValueToGet));
 						break;
 					}
 				}
@@ -377,17 +377,17 @@ public:
 				{
 					case ERigControlAxis::X:
 					{
-						Transform.SetLocation(FVector((float)ValueToGet, 0.f, 0.f));
+						Transform.SetLocation(FVector3f((float)ValueToGet, 0.f, 0.f));
 						break;
 					}
 					case ERigControlAxis::Y:
 					{
-						Transform.SetLocation(FVector(0.f, (float)ValueToGet, 0.f));
+						Transform.SetLocation(FVector3f(0.f, (float)ValueToGet, 0.f));
 						break;
 					}
 					case ERigControlAxis::Z:
 					{
-						Transform.SetLocation(FVector(0.f, 0.f, (float)ValueToGet));
+						Transform.SetLocation(FVector3f(0.f, 0.f, (float)ValueToGet));
 						break;
 					}
 				}
@@ -395,22 +395,22 @@ public:
 			}
 			case ERigControlType::Vector2D:
 			{
-				const FVector2D ValueToGet = Get<FVector2D>();
+				const FVector3f ValueToGet = Get<FVector3f>();
 				switch (InPrimaryAxis)
 				{
 					case ERigControlAxis::X:
 					{
-						Transform.SetLocation(FVector(0.f, ValueToGet.X, ValueToGet.Y));
+						Transform.SetLocation(FVector3f(0.f, ValueToGet.X, ValueToGet.Y));
 						break;
 					}
 					case ERigControlAxis::Y:
 					{
-						Transform.SetLocation(FVector(ValueToGet.X, 0.f, ValueToGet.Y));
+						Transform.SetLocation(FVector3f(ValueToGet.X, 0.f, ValueToGet.Y));
 						break;
 					}
 					case ERigControlAxis::Z:
 					{
-						Transform.SetLocation(FVector(ValueToGet.X, ValueToGet.Y, 0.f));
+						Transform.SetLocation(FVector3f(ValueToGet.X, ValueToGet.Y, 0.f));
 						break;
 					}
 				}
@@ -418,33 +418,34 @@ public:
 			}
 			case ERigControlType::Position:
 			{
-				Transform.SetLocation(Get<FVector>());
+				Transform.SetLocation(Get<FVector3f>());
 				break;
 			}
 			case ERigControlType::Scale:
 			{
-				Transform.SetScale3D(Get<FVector>());
+				Transform.SetScale3D(Get<FVector3f>());
 				break;
 			}
 			case ERigControlType::Rotator:
 			{
-				Transform.SetRotation(FQuat(Get<FRotator>()));
+				const FVector3f RotatorAxes = Get<FVector3f>();
+				Transform.SetRotation(FQuat(FRotator(RotatorAxes.X, RotatorAxes.Y, RotatorAxes.Z)));
 				break;
 			}
 			case ERigControlType::Transform:
 			{
-				return Get<FTransform>();
+				return Get<FTransform_Float>().ToTransform();
 			}
 			case ERigControlType::TransformNoScale:
 			{
-				const FTransformNoScale TransformNoScale = Get<FTransformNoScale>();
+				const FTransformNoScale TransformNoScale = Get<FTransformNoScale_Float>().ToTransform();
 				Transform = TransformNoScale;
 				Transform.NormalizeRotation();
 				break;
 			}
 			case ERigControlType::EulerTransform:
 			{
-				const FEulerTransform EulerTransform = Get<FEulerTransform>();
+				const FEulerTransform EulerTransform = Get<FEulerTransform_Float>().ToTransform();
 				Transform = FTransform(EulerTransform.ToFTransform());
 				Transform.NormalizeRotation();
 				break;
@@ -518,17 +519,17 @@ public:
 				{
 					case ERigControlAxis::X:
 					{
-						Set<FVector2D>(FVector2D(Location.Y, Location.Z));
+						Set<FVector3f>(FVector3f(Location.Y, Location.Z, 0.f));
 						break;
 					}
 					case ERigControlAxis::Y:
 					{
-						Set<FVector2D>(FVector2D(Location.X, Location.Z));
+						Set<FVector3f>(FVector3f(Location.X, Location.Z, 0.f));
 						break;
 					}
 					case ERigControlAxis::Z:
 					{
-						Set<FVector2D>(FVector2D(Location.X, Location.Y));
+						Set<FVector3f>(FVector3f(Location.X, Location.Y, 0.f));
 						break;
 					}
 				}
@@ -536,12 +537,12 @@ public:
 			}
 			case ERigControlType::Position:
 			{
-				Set<FVector>(InTransform.GetLocation());
+				Set<FVector3f>(InTransform.GetLocation());
 				break;
 			}
 			case ERigControlType::Scale:
 			{
-				Set<FVector>(InTransform.GetScale3D());
+				Set<FVector3f>(InTransform.GetScale3D());
 				break;
 			}
 			case ERigControlType::Rotator:
@@ -558,18 +559,18 @@ public:
 
 				//Add Diff
 				CurrentRotator = CurrentRotator + DeltaRot;
-				Set<FRotator>(CurrentRotator);
+				Set<FVector3f>(CurrentRotator.Vector());
 				break;
 			}
 			case ERigControlType::Transform:
 			{
-				Set<FTransform>(InTransform);
+				Set<FTransform_Float>(InTransform);
 				break;
 			}
 			case ERigControlType::TransformNoScale:
 			{
 				const FTransformNoScale NoScale = InTransform;
-				Set<FTransformNoScale>(NoScale);
+				Set<FTransformNoScale_Float>(NoScale);
 				break;
 			}
 			case ERigControlType::EulerTransform:
@@ -586,7 +587,7 @@ public:
 				DeltaRot.Normalize();
 				const FRotator NewRotation(CurrentEulerTransform.Rotation + DeltaRot);
 				NewTransform.Rotation = NewRotation;
-				Set<FEulerTransform>(NewTransform);
+				Set<FEulerTransform_Float>(NewTransform);
 				break;
 			}
 			default:
@@ -597,6 +598,16 @@ public:
 		}
 	}
 
+	template<typename T>
+	FORCEINLINE static T Clamp(const T Value, const T Minimum, const T Maximum)
+	{
+		if (Minimum < Maximum)
+		{
+			return FMath::Clamp<T>(Value, Minimum, Maximum);
+		}
+		return FMath::Clamp<T>(Value, Maximum, Minimum);
+	}
+
 	FORCEINLINE void ApplyLimits(
 		bool bLimitTranslation,
 		bool bLimitRotation,
@@ -610,36 +621,6 @@ public:
 			return;
 		}
 
-		struct Local
-		{
-			FORCEINLINE static float Clamp(const float Value, const float Minimum, const float Maximum)
-			{
-				if (Minimum < Maximum)
-				{
-					return FMath::Clamp<float>(Value, Minimum, Maximum);
-				}
-				return FMath::Clamp<float>(Value, Maximum, Minimum);
-			}
-
-			FORCEINLINE static double Clamp(const double Value, const double Minimum, const double Maximum)
-			{
-				if (Minimum < Maximum)
-				{
-					return FMath::Clamp<double>(Value, Minimum, Maximum);
-				}
-				return FMath::Clamp<double>(Value, Maximum, Minimum);
-			}
-
-			FORCEINLINE static int32 Clamp(const int32 Value, const int32 Minimum, const int32 Maximum)
-			{
-				if (Minimum < Maximum)
-				{
-					return FMath::Clamp<int32>(Value, Minimum, Maximum);
-				}
-				return FMath::Clamp<int32>(Value, Maximum, Minimum);
-			}
-		};
-
 		switch(InControlType)
 		{
 			case ERigControlType::Float:
@@ -647,7 +628,7 @@ public:
 				if (bLimitTranslation)
 				{
 					float& ValueRef = GetRef<float>();
-					ValueRef = Local::Clamp(ValueRef, InMinimumValue.Get<float>(), InMaximumValue.Get<float>());
+					ValueRef = Clamp<float>(ValueRef, InMinimumValue.Get<float>(), InMaximumValue.Get<float>());
 				}
 				break;
 			}
@@ -656,7 +637,7 @@ public:
 				if (bLimitTranslation)
 				{
 					int32& ValueRef = GetRef<int32>();
-					ValueRef = Local::Clamp(ValueRef, InMinimumValue.Get<int32>(), InMaximumValue.Get<int32>());
+					ValueRef = Clamp<int32>(ValueRef, InMinimumValue.Get<int32>(), InMaximumValue.Get<int32>());
 				}
 				break;
 			}
@@ -664,11 +645,11 @@ public:
 			{
 				if (bLimitTranslation)
 				{
-					FVector2D& ValueRef = GetRef<FVector2D>();
-					const FVector2D& Min = InMinimumValue.GetRef<FVector2D>();
-					const FVector2D& Max = InMaximumValue.GetRef<FVector2D>();
-					ValueRef.X = Local::Clamp(ValueRef.X, Min.X, Max.X);
-					ValueRef.Y = Local::Clamp(ValueRef.Y, Min.Y, Max.Y);
+					FVector3f& ValueRef = GetRef<FVector3f>();
+					const FVector3f& Min = InMinimumValue.GetRef<FVector3f>();
+					const FVector3f& Max = InMaximumValue.GetRef<FVector3f>();
+					ValueRef.X = Clamp<float>(ValueRef.X, Min.X, Max.X);
+					ValueRef.Y = Clamp<float>(ValueRef.Y, Min.Y, Max.Y);
 				}
 				break;
 			}
@@ -676,12 +657,12 @@ public:
 			{
 				if (bLimitTranslation)
 				{
-					FVector& ValueRef = GetRef<FVector>();
-					const FVector& Min = InMinimumValue.GetRef<FVector>();
-					const FVector& Max = InMaximumValue.GetRef<FVector>();
-					ValueRef.X = Local::Clamp(ValueRef.X, Min.X, Max.X);
-					ValueRef.Y = Local::Clamp(ValueRef.Y, Min.Y, Max.Y);
-					ValueRef.Z = Local::Clamp(ValueRef.Z, Min.Z, Max.Z);
+					FVector3f& ValueRef = GetRef<FVector3f>();
+					const FVector3f& Min = InMinimumValue.GetRef<FVector3f>();
+					const FVector3f& Max = InMaximumValue.GetRef<FVector3f>();
+					ValueRef.X = Clamp<float>(ValueRef.X, Min.X, Max.X);
+					ValueRef.Y = Clamp<float>(ValueRef.Y, Min.Y, Max.Y);
+					ValueRef.Z = Clamp<float>(ValueRef.Z, Min.Z, Max.Z);
 				}
 				break;
 			}
@@ -689,12 +670,12 @@ public:
 			{
 				if (bLimitScale)
 				{
-					FVector& ValueRef = GetRef<FVector>();
-					const FVector& Min = InMinimumValue.GetRef<FVector>();
-					const FVector& Max = InMaximumValue.GetRef<FVector>();
-					ValueRef.X = Local::Clamp(ValueRef.X, Min.X, Max.X);
-					ValueRef.Y = Local::Clamp(ValueRef.Y, Min.Y, Max.Y);
-					ValueRef.Z = Local::Clamp(ValueRef.Z, Min.Z, Max.Z);
+					FVector3f& ValueRef = GetRef<FVector3f>();
+					const FVector3f& Min = InMinimumValue.GetRef<FVector3f>();
+					const FVector3f& Max = InMaximumValue.GetRef<FVector3f>();
+					ValueRef.X = Clamp<float>(ValueRef.X, Min.X, Max.X);
+					ValueRef.Y = Clamp<float>(ValueRef.Y, Min.Y, Max.Y);
+					ValueRef.Z = Clamp<float>(ValueRef.Z, Min.Z, Max.Z);
 				}
 				break;
 			}
@@ -702,113 +683,107 @@ public:
 			{
 				if (bLimitRotation)
 				{
-					FRotator& ValueRef = GetRef<FRotator>();
-					const FRotator& Min = InMinimumValue.GetRef<FRotator>();
-					const FRotator& Max = InMaximumValue.GetRef<FRotator>();
-					ValueRef.Pitch = Local::Clamp(ValueRef.Pitch, Min.Pitch, Max.Pitch);
-					ValueRef.Yaw = Local::Clamp(ValueRef.Yaw, Min.Yaw, Max.Yaw);
-					ValueRef.Roll = Local::Clamp(ValueRef.Roll, Min.Roll, Max.Roll);
+					FVector3f& ValueRef = GetRef<FVector3f>();
+					const FVector3f& Min = InMinimumValue.GetRef<FVector3f>();
+					const FVector3f& Max = InMaximumValue.GetRef<FVector3f>();
+					ValueRef.X = Clamp<float>(ValueRef.X, Min.X, Max.X);
+					ValueRef.Y = Clamp<float>(ValueRef.Y, Min.Y, Max.Y);
+					ValueRef.Z = Clamp<float>(ValueRef.Z, Min.Z, Max.Z);
 				}
 				break;
 			}
 			case ERigControlType::Transform:
 			{
-				FTransform& ValueRef = GetRef<FTransform>();
-				const FTransform& Min = InMinimumValue.GetRef<FTransform>();
-				const FTransform& Max = InMaximumValue.GetRef<FTransform>();
+				FTransform_Float& ValueRef = GetRef<FTransform_Float>();
+				const FTransform Min = InMinimumValue.GetRef<FTransform_Float>().ToTransform();
+				const FTransform Max = InMaximumValue.GetRef<FTransform_Float>().ToTransform();
 
 				if (bLimitTranslation)
 				{
-					ValueRef.SetLocation(FVector(
-						Local::Clamp(ValueRef.GetLocation().X, Min.GetLocation().X, Max.GetLocation().X),
-						Local::Clamp(ValueRef.GetLocation().Y, Min.GetLocation().Y, Max.GetLocation().Y),
-						Local::Clamp(ValueRef.GetLocation().Z, Min.GetLocation().Z, Max.GetLocation().Z)
-					));
+					ValueRef.TranslationX = Clamp<float>(ValueRef.TranslationX, Min.GetLocation().X, Max.GetLocation().X);
+					ValueRef.TranslationY = Clamp<float>(ValueRef.TranslationY, Min.GetLocation().Y, Max.GetLocation().Y);
+					ValueRef.TranslationZ = Clamp<float>(ValueRef.TranslationZ, Min.GetLocation().Z, Max.GetLocation().Z);
 				}
 				if (bLimitRotation)
 				{
-					const FRotator Rotator = ValueRef.GetRotation().Rotator();
+					const FRotator Rotator = FQuat(ValueRef.RotationX, ValueRef.RotationY, ValueRef.RotationZ, ValueRef.RotationW).Rotator();
 					const FRotator MinRotator = Min.GetRotation().Rotator();
 					const FRotator MaxRotator = Max.GetRotation().Rotator();
 
-					ValueRef.SetRotation(FQuat(FRotator(
-						Local::Clamp(Rotator.Pitch, MinRotator.Pitch, MaxRotator.Pitch),
-						Local::Clamp(Rotator.Yaw, MinRotator.Yaw, MaxRotator.Yaw),
-						Local::Clamp(Rotator.Roll, MinRotator.Roll, MaxRotator.Roll)
-					)));
+					FQuat LimitedQuat(FRotator(
+						Clamp<float>(Rotator.Pitch, MinRotator.Pitch, MaxRotator.Pitch),
+						Clamp<float>(Rotator.Yaw, MinRotator.Yaw, MaxRotator.Yaw),
+						Clamp<float>(Rotator.Roll, MinRotator.Roll, MaxRotator.Roll)
+					));
+					
+					ValueRef.RotationX = LimitedQuat.X;
+					ValueRef.RotationY = LimitedQuat.Y;
+					ValueRef.RotationZ = LimitedQuat.Z;
+					ValueRef.RotationW = LimitedQuat.W;
 				}
 				if (bLimitScale)
 				{
-					ValueRef.SetScale3D(FVector(
-						Local::Clamp(ValueRef.GetScale3D().X, Min.GetScale3D().X, Max.GetScale3D().X),
-						Local::Clamp(ValueRef.GetScale3D().Y, Min.GetScale3D().Y, Max.GetScale3D().Y),
-						Local::Clamp(ValueRef.GetScale3D().Z, Min.GetScale3D().Z, Max.GetScale3D().Z)
-					));
+					ValueRef.ScaleX = Clamp<float>(ValueRef.ScaleX, Min.GetScale3D().X, Max.GetScale3D().X);
+					ValueRef.ScaleY = Clamp<float>(ValueRef.ScaleY, Min.GetScale3D().Y, Max.GetScale3D().Y);
+					ValueRef.ScaleZ = Clamp<float>(ValueRef.ScaleZ, Min.GetScale3D().Z, Max.GetScale3D().Z);
 				}
 				break;
 			}
 			case ERigControlType::TransformNoScale:
 			{
-				FTransformNoScale& ValueRef = GetRef<FTransformNoScale>();
-				const FTransformNoScale& Min = InMinimumValue.GetRef<FTransformNoScale>();
-				const FTransformNoScale& Max = InMaximumValue.GetRef<FTransformNoScale>();
+				FTransformNoScale_Float& ValueRef = GetRef<FTransformNoScale_Float>();
+				const FTransformNoScale Min = InMinimumValue.GetRef<FTransformNoScale_Float>().ToTransform();
+				const FTransformNoScale Max = InMaximumValue.GetRef<FTransformNoScale_Float>().ToTransform();
 
 				if (bLimitTranslation)
 				{
-					ValueRef.Location = FVector(
-						Local::Clamp(ValueRef.Location.X, Min.Location.X, Max.Location.X),
-						Local::Clamp(ValueRef.Location.Y, Min.Location.Y, Max.Location.Y),
-						Local::Clamp(ValueRef.Location.Z, Min.Location.Z, Max.Location.Z)
-					);
+					ValueRef.TranslationX = Clamp<float>(ValueRef.TranslationX, Min.Location.X, Max.Location.X);
+					ValueRef.TranslationY = Clamp<float>(ValueRef.TranslationY, Min.Location.Y, Max.Location.Y);
+					ValueRef.TranslationZ = Clamp<float>(ValueRef.TranslationZ, Min.Location.Z, Max.Location.Z);
 				}
 				if (bLimitRotation)
 				{
-					const FRotator Rotator = ValueRef.Rotation.Rotator();
+					const FRotator Rotator = FQuat(ValueRef.RotationX, ValueRef.RotationY, ValueRef.RotationZ, ValueRef.RotationW).Rotator();
 					const FRotator MinRotator = Min.Rotation.Rotator();
 					const FRotator MaxRotator = Max.Rotation.Rotator();
 
-					ValueRef.Rotation = FQuat(FRotator(
-						Local::Clamp(Rotator.Pitch, MinRotator.Pitch, MaxRotator.Pitch),
-						Local::Clamp(Rotator.Yaw, MinRotator.Yaw, MaxRotator.Yaw),
-						Local::Clamp(Rotator.Roll, MinRotator.Roll, MaxRotator.Roll)
+					FQuat LimitedQuat(FRotator(
+						Clamp<float>(Rotator.Pitch, MinRotator.Pitch, MaxRotator.Pitch),
+						Clamp<float>(Rotator.Yaw, MinRotator.Yaw, MaxRotator.Yaw),
+						Clamp<float>(Rotator.Roll, MinRotator.Roll, MaxRotator.Roll)
 					));
+
+					ValueRef.RotationX = LimitedQuat.X;
+					ValueRef.RotationY = LimitedQuat.Y;
+					ValueRef.RotationZ = LimitedQuat.Z;
+					ValueRef.RotationW = LimitedQuat.W;
 				}
 				break;
 			}
 
 			case ERigControlType::EulerTransform:
 			{
-				FEulerTransform& ValueRef = GetRef<FEulerTransform>();
-				const FEulerTransform& Min = InMinimumValue.GetRef<FEulerTransform>();
-				const FEulerTransform& Max = InMaximumValue.GetRef<FEulerTransform>();
+				FEulerTransform_Float& ValueRef = GetRef<FEulerTransform_Float>();
+				const FEulerTransform_Float& Min = InMinimumValue.GetRef<FEulerTransform_Float>();
+				const FEulerTransform_Float& Max = InMaximumValue.GetRef<FEulerTransform_Float>();
 
 				if (bLimitTranslation)
 				{
-					ValueRef.Location = FVector(
-						Local::Clamp(ValueRef.Location.X, Min.Location.X, Max.Location.X),
-						Local::Clamp(ValueRef.Location.Y, Min.Location.Y, Max.Location.Y),
-						Local::Clamp(ValueRef.Location.Z, Min.Location.Z, Max.Location.Z)
-					);
+					ValueRef.TranslationX = Clamp<float>(ValueRef.TranslationX, Min.TranslationX, Max.TranslationX);
+					ValueRef.TranslationY = Clamp<float>(ValueRef.TranslationY, Min.TranslationY, Max.TranslationY);
+					ValueRef.TranslationZ = Clamp<float>(ValueRef.TranslationZ, Min.TranslationZ, Max.TranslationZ);
 				}
 				if (bLimitRotation)
 				{
-					const FRotator Rotator = ValueRef.Rotation;
-					const FRotator MinRotator = Min.Rotation;
-					const FRotator MaxRotator = Max.Rotation;
-
-					ValueRef.Rotation = FRotator(
-						Local::Clamp(Rotator.Pitch, MinRotator.Pitch, MaxRotator.Pitch),
-						Local::Clamp(Rotator.Yaw, MinRotator.Yaw, MaxRotator.Yaw),
-						Local::Clamp(Rotator.Roll, MinRotator.Roll, MaxRotator.Roll)
-					);
+					ValueRef.RotationX = Clamp<float>(ValueRef.RotationX, Min.RotationX, Max.RotationX);
+					ValueRef.RotationY = Clamp<float>(ValueRef.RotationY, Min.RotationY, Max.RotationY);
+					ValueRef.RotationZ = Clamp<float>(ValueRef.RotationZ, Min.RotationZ, Max.RotationZ);
 				}
 				if (bLimitScale)
 				{
-					ValueRef.Scale = FVector(
-						Local::Clamp(ValueRef.Scale.X, Min.Scale.X, Max.Scale.X),
-						Local::Clamp(ValueRef.Scale.Y, Min.Scale.Y, Max.Scale.Y),
-						Local::Clamp(ValueRef.Scale.Z, Min.Scale.Z, Max.Scale.Z)
-					);
+					ValueRef.ScaleX = Clamp<float>(ValueRef.ScaleX, Min.ScaleX, Max.ScaleX);
+					ValueRef.ScaleY = Clamp<float>(ValueRef.ScaleY, Min.ScaleY, Max.ScaleY);
+					ValueRef.ScaleZ = Clamp<float>(ValueRef.ScaleZ, Min.ScaleZ, Max.ScaleZ);
 				}
 				break;
 			}
@@ -820,8 +795,140 @@ public:
 		}
 	}
 
-private:
+	struct FTransform_Float
+	{
+		float RotationX;
+		float RotationY;
+		float RotationZ;
+		float RotationW;
+		float TranslationX;
+		float TranslationY;
+		float TranslationZ;
+#if ENABLE_VECTORIZED_TRANSFORM
+		float TranslationW;
+#endif
+		float ScaleX;
+		float ScaleY;
+		float ScaleZ;
+#if ENABLE_VECTORIZED_TRANSFORM
+		float ScaleW;
+#endif
+
+		FTransform_Float()
+		{
+			*this = FTransform_Float(FTransform::Identity);
+		}
+
+		FTransform_Float(const FTransform& InTransform)
+		{
+			RotationX = InTransform.GetRotation().X;
+			RotationY = InTransform.GetRotation().Y;
+			RotationZ = InTransform.GetRotation().Z;
+			RotationW = InTransform.GetRotation().W;
+			TranslationX = InTransform.GetTranslation().X;
+			TranslationY = InTransform.GetTranslation().Y;
+			TranslationZ = InTransform.GetTranslation().Z;
+			ScaleX = InTransform.GetScale3D().X;
+			ScaleY = InTransform.GetScale3D().Y;
+			ScaleZ = InTransform.GetScale3D().Z;
+#if ENABLE_VECTORIZED_TRANSFORM
+			TranslationW = ScaleW = 0.f;
+#endif
+		}
 
+		FTransform ToTransform() const
+		{
+			FTransform Transform;
+			Transform.SetRotation(FQuat(RotationX, RotationY, RotationZ, RotationW));
+			Transform.SetTranslation(FVector(TranslationX, TranslationY, TranslationZ));
+			Transform.SetScale3D(FVector(ScaleX, ScaleY, ScaleZ));
+			return Transform;
+		}
+	};
+
+	struct FTransformNoScale_Float
+	{
+		float RotationX;
+		float RotationY;
+		float RotationZ;
+		float RotationW;
+		float TranslationX;
+		float TranslationY;
+		float TranslationZ;
+#if ENABLE_VECTORIZED_TRANSFORM
+		float TranslationW;
+#endif
+
+		FTransformNoScale_Float()
+		{
+			*this = FTransformNoScale_Float(FTransformNoScale::Identity);
+		}
+
+		FTransformNoScale_Float(const FTransformNoScale& InTransform)
+		{
+			RotationX = InTransform.Rotation.X;
+			RotationY = InTransform.Rotation.Y;
+			RotationZ = InTransform.Rotation.Z;
+			RotationW = InTransform.Rotation.W;
+			TranslationX = InTransform.Location.X;
+			TranslationY = InTransform.Location.Y;
+			TranslationZ = InTransform.Location.Z;
+#if ENABLE_VECTORIZED_TRANSFORM
+			TranslationW = 0.f;
+#endif
+		}
+
+		FTransformNoScale ToTransform() const
+		{
+			FTransformNoScale Transform;
+			Transform.Rotation = FQuat(RotationX, RotationY, RotationZ, RotationW);
+			Transform.Location = FVector(TranslationX, TranslationY, TranslationZ);
+			return Transform;
+		}
+	};
+
+	struct FEulerTransform_Float
+	{
+		float RotationX;
+		float RotationY;
+		float RotationZ;
+		float TranslationX;
+		float TranslationY;
+		float TranslationZ;
+		float ScaleX;
+		float ScaleY;
+		float ScaleZ;
+
+		FEulerTransform_Float()
+		{
+			*this = FEulerTransform_Float(FEulerTransform::Identity);
+		}
+
+		FEulerTransform_Float(const FEulerTransform& InTransform)
+		{
+			RotationX = InTransform.Rotation.Pitch;
+			RotationY = InTransform.Rotation.Yaw;
+			RotationZ = InTransform.Rotation.Roll;
+			TranslationX = InTransform.Location.X;
+			TranslationY = InTransform.Location.Y;
+			TranslationZ = InTransform.Location.Z;
+			ScaleX = InTransform.Scale.X;
+			ScaleY = InTransform.Scale.Y;
+			ScaleZ = InTransform.Scale.Z;
+		}
+
+		FEulerTransform ToTransform() const
+		{
+			FEulerTransform Transform;
+			Transform.Rotation = FRotator(RotationX, RotationY, RotationZ);
+			Transform.Location = FVector(TranslationX, TranslationY, TranslationZ);
+			Transform.Scale = FVector(ScaleX, ScaleY, ScaleZ);
+			return Transform;
+		}
+	};
+
+private:
+	
 	UPROPERTY()
 	FRigControlValueStorage FloatStorage;
 	
@@ -833,8 +940,10 @@ private:
 	friend class URigHierarchyController;
 };
 
+#if UE_LARGE_WORLD_COORDINATES_DISABLED
+
 template<>
-FORCEINLINE_DEBUGGABLE FQuat FRigControlValue::SetFromString<FQuat>(const FString& InString)
+inline FQuat FRigControlValue::SetFromString<FQuat>(const FString& InString)
 {
 	FQuat Value;
 	TBaseStructure<FQuat>::Get()->ImportText(*InString, &Value, nullptr, PPF_None, nullptr, TBaseStructure<FQuat>::Get()->GetName());
@@ -842,6 +951,138 @@ FORCEINLINE_DEBUGGABLE FQuat FRigControlValue::SetFromString<FQuat>(const FStrin
 	return Value;
 }
 
+#else
+
+template<>
+inline static FRigControlValue FRigControlValue::Make(FVector2D InValue)
+{
+	return Make<FVector3f>(FVector3f(InValue.X, InValue.Y, 0.f));
+}
+
+template<>
+inline static FRigControlValue FRigControlValue::Make(FVector InValue)
+{
+	return Make<FVector3f>(InValue);
+}
+
+template<>
+inline static FRigControlValue FRigControlValue::Make(FRotator InValue)
+{
+	return Make<FVector3f>(InValue.Vector());
+}
+
+template<>
+inline static FRigControlValue FRigControlValue::Make(FTransform InValue)
+{
+	return Make<FTransform_Float>(InValue);
+}
+
+template<>
+inline static FRigControlValue FRigControlValue::Make(FTransformNoScale InValue)
+{
+	return Make<FTransformNoScale_Float>(InValue);
+}
+
+template<>
+inline static FRigControlValue FRigControlValue::Make(FEulerTransform InValue)
+{
+	return Make<FEulerTransform_Float>(InValue);
+}
+
+template<>
+inline FString FRigControlValue::ToString<FVector>() const
+{
+	FVector Value = GetRef<FVector3f>();
+	FString Result;
+	TBaseStructure<FVector>::Get()->ExportText(Result, &Value, nullptr, nullptr, PPF_None, nullptr);
+	return Result;
+}
+
+template<>
+inline FString FRigControlValue::ToString<FRotator>() const
+{
+	FRotator Value = FRotator::MakeFromEuler(GetRef<FVector3f>());
+	FString Result;
+	TBaseStructure<FRotator>::Get()->ExportText(Result, &Value, nullptr, nullptr, PPF_None, nullptr);
+	return Result;
+}
+
+template<>
+inline FString FRigControlValue::ToString<FTransform>() const
+{
+	FTransform Value = GetRef<FTransform_Float>().ToTransform();
+	FString Result;
+	TBaseStructure<FTransform>::Get()->ExportText(Result, &Value, nullptr, nullptr, PPF_None, nullptr);
+	return Result;
+}
+
+template<>
+inline FVector FRigControlValue::SetFromString<FVector>(const FString& InString)
+{
+	FVector Value;
+	TBaseStructure<FVector>::Get()->ImportText(*InString, &Value, nullptr, PPF_None, nullptr, TBaseStructure<FVector>::Get()->GetName());
+	Set<FVector3f>(Value);
+	return Value;
+}
+
+template<>
+inline FQuat FRigControlValue::SetFromString<FQuat>(const FString& InString)
+{
+	FQuat Value;
+	TBaseStructure<FQuat>::Get()->ImportText(*InString, &Value, nullptr, PPF_None, nullptr, TBaseStructure<FQuat>::Get()->GetName());
+	Set<FVector3f>(Value.Rotator().Vector());
+	return Value;
+}
+
+template<>
+inline FRotator FRigControlValue::SetFromString<FRotator>(const FString& InString)
+{
+	FRotator Value;
+	TBaseStructure<FRotator>::Get()->ImportText(*InString, &Value, nullptr, PPF_None, nullptr, TBaseStructure<FRotator>::Get()->GetName());
+	Set<FVector3f>(Value.Vector());
+	return Value;
+}
+
+template<>
+inline FTransform FRigControlValue::SetFromString<FTransform>(const FString& InString)
+{
+	FTransform Value;
+	TBaseStructure<FTransform>::Get()->ImportText(*InString, &Value, nullptr, PPF_None, nullptr, TBaseStructure<FTransform>::Get()->GetName());
+	Set<FTransform_Float>(Value);
+	return Value;
+}
+
+template<>
+FVector2D& FRigControlValue::GetRef() = delete;
+template<>
+FVector& FRigControlValue::GetRef() = delete;
+template<>
+FRotator& FRigControlValue::GetRef() = delete;
+template<>
+FQuat& FRigControlValue::GetRef() = delete;
+template<>
+FTransform& FRigControlValue::GetRef() = delete;
+template<>
+FTransformNoScale& FRigControlValue::GetRef() = delete;
+template<>
+FEulerTransform& FRigControlValue::GetRef() = delete;
+template<>
+void FRigControlValue::Set(FVector2D InValue) = delete;
+template<>
+void FRigControlValue::Set(FVector InValue) = delete;
+template<>
+void FRigControlValue::Set(FRotator InValue) = delete;
+template<>
+void FRigControlValue::Set(FQuat InValue) = delete;
+template<>
+void FRigControlValue::Set(FTransform InValue) = delete;
+template<>
+void FRigControlValue::Set(FTransformNoScale InValue) = delete;
+template<>
+void FRigControlValue::Set(FEulerTransform InValue) = delete;
+
+#endif
+
 enum class EControlRigContextChannelToKey : uint32
 {
 	None = 0x000,
diff --git a/Engine/Plugins/Animation/ControlRig/Source/ControlRigEditor/Private/Editor/ControlRigEditor.cpp b/Engine/Plugins/Animation/ControlRig/Source/ControlRigEditor/Private/Editor/ControlRigEditor.cpp
--- a/Engine/Plugins/Animation/ControlRig/Source/ControlRigEditor/Private/Editor/ControlRigEditor.cpp
+++ b/Engine/Plugins/Animation/ControlRig/Source/ControlRigEditor/Private/Editor/ControlRigEditor.cpp
@@ -5354,8 +5354,8 @@ void FControlRigEditor::HandleOnControlModified(UControlRig* Subject, FRigContro
 				case ERigControlType::Rotator:
 				{
 					FRotator Rotator = ControlValue.Get<FRotator>();
-					FRigControlValue QuatValue = FRigControlValue::Make<FQuat>(FQuat(Rotator));
-					NewDefaultValue = QuatValue.ToString<FQuat>();
+					FRigControlValue RotatorValue = FRigControlValue::Make<FRotator>(Rotator);
+					NewDefaultValue = RotatorValue.ToString<FRotator>();
 					break;
 				}
 				case ERigControlType::Transform:
