Control Rig: Catch nullptr during curves import

diff --git a/Engine/Plugins/Animation/ControlRig/Source/ControlRig/Private/Rigs/RigHierarchyController.cpp b/Engine/Plugins/Animation/ControlRig/Source/ControlRig/Private/Rigs/RigHierarchyController.cpp
--- a/Engine/Plugins/Animation/ControlRig/Source/ControlRig/Private/Rigs/RigHierarchyController.cpp
+++ b/Engine/Plugins/Animation/ControlRig/Source/ControlRig/Private/Rigs/RigHierarchyController.cpp
@@ -708,7 +708,10 @@ USkeleton* URigHierarchyController::GetSkeletonFromAssetPath(const FString& InAs
 TArray<FRigElementKey> URigHierarchyController::ImportCurves(USkeleton* InSkeleton, FName InNameSpace,
                                                              bool bSelectCurves, bool bSetupUndo, bool bPrintPythonCommand)
 {
-	check(InSkeleton);
+	if (InSkeleton == nullptr)
+	{
+		return TArray<FRigElementKey>();
+	}
 
 	TArray<FRigElementKey> Keys;
 	if(!IsValid())
