Control Rig: Build break fix for certain platforms

diff --git a/Engine/Plugins/Animation/ControlRig/Source/ControlRig/Public/Rigs/RigHierarchyDefines.h b/Engine/Plugins/Animation/ControlRig/Source/ControlRig/Public/Rigs/RigHierarchyDefines.h
--- a/Engine/Plugins/Animation/ControlRig/Source/ControlRig/Public/Rigs/RigHierarchyDefines.h
+++ b/Engine/Plugins/Animation/ControlRig/Source/ControlRig/Public/Rigs/RigHierarchyDefines.h
@@ -954,37 +954,37 @@ inline FQuat FRigControlValue::SetFromString<FQuat>(const FString& InString)
 #else
 
 template<>
-inline static FRigControlValue FRigControlValue::Make(FVector2D InValue)
+inline FRigControlValue FRigControlValue::Make(FVector2D InValue)
 {
 	return Make<FVector3f>(FVector3f(InValue.X, InValue.Y, 0.f));
 }
 
 template<>
-inline static FRigControlValue FRigControlValue::Make(FVector InValue)
+inline FRigControlValue FRigControlValue::Make(FVector InValue)
 {
 	return Make<FVector3f>(InValue);
 }
 
 template<>
-inline static FRigControlValue FRigControlValue::Make(FRotator InValue)
+inline FRigControlValue FRigControlValue::Make(FRotator InValue)
 {
 	return Make<FVector3f>(InValue.Vector());
 }
 
 template<>
-inline static FRigControlValue FRigControlValue::Make(FTransform InValue)
+inline FRigControlValue FRigControlValue::Make(FTransform InValue)
 {
 	return Make<FTransform_Float>(InValue);
 }
 
 template<>
-inline static FRigControlValue FRigControlValue::Make(FTransformNoScale InValue)
+inline FRigControlValue FRigControlValue::Make(FTransformNoScale InValue)
 {
 	return Make<FTransformNoScale_Float>(InValue);
 }
 
 template<>
-inline static FRigControlValue FRigControlValue::Make(FEulerTransform InValue)
+inline FRigControlValue FRigControlValue::Make(FEulerTransform InValue)
 {
 	return Make<FEulerTransform_Float>(InValue);
 }
