Control Rig: Restrict Local variables to functions

diff --git a/Engine/Plugins/Animation/ControlRig/Source/ControlRigEditor/Private/Editor/ControlRigEditor.cpp b/Engine/Plugins/Animation/ControlRig/Source/ControlRigEditor/Private/Editor/ControlRigEditor.cpp
--- a/Engine/Plugins/Animation/ControlRig/Source/ControlRigEditor/Private/Editor/ControlRigEditor.cpp
+++ b/Engine/Plugins/Animation/ControlRig/Source/ControlRigEditor/Private/Editor/ControlRigEditor.cpp
@@ -1929,6 +1929,17 @@ bool FControlRigEditor::TransactionObjectAffectsBlueprint(UObject* InTransactedO
 	return FBlueprintEditor::TransactionObjectAffectsBlueprint(InTransactedObject);
 }
 
+bool FControlRigEditor::CanAddNewLocalVariable() const
+{
+	const URigVMGraph* Graph = GetFocusedModel();
+	const URigVMGraph* ParentGraph = Graph->GetParentGraph();				
+	if (ParentGraph && ParentGraph->IsA<URigVMFunctionLibrary>())
+	{
+		return true;
+	}
+	return false;
+}
+
 void FControlRigEditor::OnAddNewLocalVariable()
 {
 	if (!CanAddNewLocalVariable())
@@ -2411,10 +2422,18 @@ bool FControlRigEditor::IsSectionVisible(NodeSectionID::Type InSectionID) const
 		case NodeSectionID::GRAPH:
 		case NodeSectionID::VARIABLE:
 		case NodeSectionID::FUNCTION:
-		case NodeSectionID::LOCAL_VARIABLE:
 		{
 			return true;
 		}
+		case NodeSectionID::LOCAL_VARIABLE:
+		{
+			const URigVMGraph* Graph = GetFocusedModel();
+			const URigVMGraph* ParentGraph = Graph->GetParentGraph();				
+			if (ParentGraph && ParentGraph->IsA<URigVMFunctionLibrary>())
+			{
+				return true;
+			}
+		}
 		default:
 		{
 			break;
diff --git a/Engine/Plugins/Animation/ControlRig/Source/ControlRigEditor/Private/Editor/ControlRigEditor.h b/Engine/Plugins/Animation/ControlRig/Source/ControlRigEditor/Private/Editor/ControlRigEditor.h
--- a/Engine/Plugins/Animation/ControlRig/Source/ControlRigEditor/Private/Editor/ControlRigEditor.h
+++ b/Engine/Plugins/Animation/ControlRig/Source/ControlRigEditor/Private/Editor/ControlRigEditor.h
@@ -107,7 +107,7 @@ public:
 
 	// BlueprintEditor interface
 	virtual bool TransactionObjectAffectsBlueprint(UObject* InTransactedObject) override;
-	virtual bool CanAddNewLocalVariable() const override { return true; }
+	virtual bool CanAddNewLocalVariable() const override;
 	virtual void OnAddNewLocalVariable() override;
 	
 	virtual void DeleteSelectedNodes() override;
