Control Rig: disallow set transform node to change rig element transform during init, since the purpose of init is to simply pre-allocate memory

diff --git a/Engine/Plugins/Animation/ControlRig/Source/ControlRig/Private/Units/Hierarchy/RigUnit_SetTransform.cpp b/Engine/Plugins/Animation/ControlRig/Source/ControlRig/Private/Units/Hierarchy/RigUnit_SetTransform.cpp
--- a/Engine/Plugins/Animation/ControlRig/Source/ControlRig/Private/Units/Hierarchy/RigUnit_SetTransform.cpp
+++ b/Engine/Plugins/Animation/ControlRig/Source/ControlRig/Private/Units/Hierarchy/RigUnit_SetTransform.cpp
@@ -31,10 +31,7 @@ FRigUnit_SetTransform_Execute()
 			{
 				CachedIndex.Reset();
 
-				if (ExecuteContext.EventName == FRigUnit_InverseExecution::EventName)
-				{
-					return;
-				}
+				return;
 			}
 			case EControlRigState::Update:
 			{
