Perf improvement for meta human face control rigs from [REDACTED_HANDLE]
Optimizes control rig pose input/output
Adds option to pass partial hierarchy as control rig input.

#rnx

diff --git a/Engine/Plugins/Animation/ControlRig/Source/ControlRig/Private/AnimNode_ControlRigBase.cpp b/Engine/Plugins/Animation/ControlRig/Source/ControlRig/Private/AnimNode_ControlRigBase.cpp
--- a/Engine/Plugins/Animation/ControlRig/Source/ControlRig/Private/AnimNode_ControlRigBase.cpp
+++ b/Engine/Plugins/Animation/ControlRig/Source/ControlRig/Private/AnimNode_ControlRigBase.cpp
@@ -17,6 +17,9 @@ static TAutoConsoleVariable<int32> CVarControlRigDisableExecutionAnimNode(TEXT("
 
 FAnimNode_ControlRigBase::FAnimNode_ControlRigBase()
 	: FAnimNode_CustomProperty()
+	, bTransferInputPose(true)
+	, bTransferInputCurves(true)
+	, bTransferPoseInGlobalSpace(true)
 	, InputSettings(FControlRigIOSettings())
 	, OutputSettings(FControlRigIOSettings())
 	, bExecute(true)
@@ -122,53 +125,132 @@ void FAnimNode_ControlRigBase::UpdateInput(UControlRig* ControlRig, const FPoseC
 
 	DECLARE_SCOPE_HIERARCHICAL_COUNTER_FUNC()
 
-	if (InputSettings.bUpdatePose)
+	if (InputSettings.bUpdatePose && bTransferInputPose)
 	{
 		const FBoneContainer& RequiredBones = InOutput.Pose.GetBoneContainer();
 
-		// get component pose from control rig
-		FCSPose<FCompactPose> MeshPoses;
-		// first I need to convert to local pose
-		MeshPoses.InitPose(InOutput.Pose);
-
 		// reset transforms here to prevent additive transforms from accumulating to INF
 		// we only update transforms from the mesh pose for bones in the current LOD, 
 		// so the reset here ensures excluded bones are also reset
-		ControlRig->GetHierarchy()->ResetPoseToInitial(ERigElementType::Bone);
+		if(!ControlRigBoneInputMappingIncomplete.IsEmpty())
+		{
+			ControlRig->GetHierarchy()->ResetPoseToInitial(ERigElementType::Bone);
+		}
 
-		// @re-think - now control rig contains init pose from their default hierarchy and current pose from this instance.
-		// we may need this init pose somewhere (instance refpose)
-		for (auto Iter = ControlRigBoneMapping.CreateConstIterator(); Iter; ++Iter)
+		if(bTransferPoseInGlobalSpace || NodeMappingContainer.IsValid())
 		{
-			const FName& Name = Iter.Key();
-			const uint16 Index = Iter.Value();
-			const FRigElementKey Key(Name, ERigElementType::Bone);
+			// get component pose from control rig
+			FCSPose<FCompactPose> MeshPoses;
+			// first I need to convert to local pose
+			MeshPoses.InitPose(InOutput.Pose);
 
-			FTransform ComponentTransform = MeshPoses.GetComponentSpaceTransform(FCompactPoseBoneIndex(Index));
-			if (NodeMappingContainer.IsValid())
+			if(!ControlRigBoneInputMappingComplete.IsEmpty())
 			{
-				ComponentTransform = NodeMappingContainer->GetSourceToTargetTransform(Name).GetRelativeTransformReverse(ComponentTransform);
+				for (const TPair<uint16, uint16>& Pair : ControlRigBoneInputMappingComplete)
+				{
+					const uint16 SkeletonIndex = Pair.Key;
+					const uint16 ControlRigIndex = Pair.Value;
+					
+					FCompactPoseBoneIndex CompactPoseIndex(SkeletonIndex);
+					FTransform ComponentTransform = MeshPoses.GetComponentSpaceTransform(CompactPoseIndex);
+					ControlRig->GetHierarchy()->SetGlobalTransformByIndex(ControlRigIndex, ComponentTransform, false);
+				}
 			}
+			else
+			{
+				for (auto Iter = ControlRigBoneInputMappingIncomplete.CreateConstIterator(); Iter; ++Iter)
+				{
+					const FName& Name = Iter.Key();
+					const uint16 Index = Iter.Value();
+					const FRigElementKey Key(Name, ERigElementType::Bone);
 
-			ControlRig->GetHierarchy()->SetGlobalTransform(Key, ComponentTransform, false);
+					FCompactPoseBoneIndex CompactPoseIndex(Index);
+					FTransform ComponentTransform = MeshPoses.GetComponentSpaceTransform(CompactPoseIndex);
+					if (NodeMappingContainer.IsValid())
+					{
+						ComponentTransform = NodeMappingContainer->GetSourceToTargetTransform(Name).GetRelativeTransformReverse(ComponentTransform);
+					}
+					ControlRig->GetHierarchy()->SetGlobalTransform(Key, ComponentTransform, false);
+				}
+			}
 		}
+		else
+		{
+			if(!ControlRigBoneInputMappingComplete.IsEmpty())
+			{
+				for (const TPair<uint16, uint16>& Pair : ControlRigBoneInputMappingComplete)
+				{
+					const uint16 SkeletonIndex = Pair.Key;
+					const uint16 ControlRigIndex = Pair.Value;
+					
+					FCompactPoseBoneIndex CompactPoseIndex(SkeletonIndex);
+					FTransform LocalTransform = InOutput.Pose[CompactPoseIndex];
+					ControlRig->GetHierarchy()->SetLocalTransformByIndex(ControlRigIndex, LocalTransform, false);
+				}
+			}
+			else
+			{
+				for (auto Iter = ControlRigBoneInputMappingIncomplete.CreateConstIterator(); Iter; ++Iter)
+				{
+					const FName& Name = Iter.Key();
+					const uint16 Index = Iter.Value();
+					const FRigElementKey Key(Name, ERigElementType::Bone);
 
+					FCompactPoseBoneIndex CompactPoseIndex(Index);
+					FTransform LocalTransform = InOutput.Pose[CompactPoseIndex];
+					ControlRig->GetHierarchy()->SetLocalTransform(Key, LocalTransform, false);
+				}
+			}
+		}
+		
+		if(ControlRigBoneInputMappingIncomplete.IsEmpty())
+		{
+			if(bTransferPoseInGlobalSpace)
+			{
+				for(uint16 ElementIndexToReset : ControlRigTransformElementsToReset)
+				{
+					FTransform InitialTransform = ControlRig->GetHierarchy()->GetGlobalTransformByIndex(ElementIndexToReset, true);
+					ControlRig->GetHierarchy()->SetGlobalTransformByIndex(ElementIndexToReset, InitialTransform, false);
+				}
+			}
+			else
+			{
+				for(uint16 ElementIndexToReset : ControlRigTransformElementsToReset)
+				{
+					FTransform InitialTransform = ControlRig->GetHierarchy()->GetLocalTransformByIndex(ElementIndexToReset, true);
+					ControlRig->GetHierarchy()->SetLocalTransformByIndex(ElementIndexToReset, InitialTransform, false);
+				}
+			}
+		}
 					
 #if WITH_EDITOR
 		ControlRig->ApplyTransformOverrideForUserCreatedBones();
 #endif
 	}
 
-	if (InputSettings.bUpdateCurves)
+	if (InputSettings.bUpdateCurves && bTransferInputCurves)
 	{
-		// we just do name mapping 
-		for (auto Iter = ControlRigCurveMapping.CreateConstIterator(); Iter; ++Iter)
+		if(!ControlRigCurveMappingComplete.IsEmpty())
+		{
+			for (const TPair<uint16, uint16>& Pair : ControlRigCurveMappingComplete)
+			{
+				const uint16 SkeletonIndex = Pair.Key;
+				const uint16 ControlRigIndex = Pair.Value;
+
+				const float Value = InOutput.Curve.Get(SkeletonIndex);
+				ControlRig->GetHierarchy()->SetCurveValueByIndex(ControlRigIndex, Value);
+			}
+		}
+		else
 		{
-			const FName& Name = Iter.Key();
-			const uint16 Index = Iter.Value();
-			const FRigElementKey Key(Name, ERigElementType::Curve);
+			for (auto Iter = ControlRigCurveMappingIncomplete.CreateConstIterator(); Iter; ++Iter)
+			{
+				const FName& Name = Iter.Key();
+				const uint16 Index = Iter.Value();
+				const FRigElementKey Key(Name, ERigElementType::Curve);
 
-			ControlRig->GetHierarchy()->SetCurveValue(Key, InOutput.Curve.Get(Index));
+				ControlRig->GetHierarchy()->SetCurveValue(Key, InOutput.Curve.Get(Index));
+			}
 		}
 	}
 
@@ -197,47 +279,123 @@ void FAnimNode_ControlRigBase::UpdateOutput(UControlRig* ControlRig, FPoseContex
 		// copy output of the rig
 		const FBoneContainer& RequiredBones = InOutput.Pose.GetBoneContainer();
 
-		// get component pose from control rig
-		FCSPose<FCompactPose> MeshPoses;
-		MeshPoses.InitPose(InOutput.Pose);
+		TMap<FName, uint16>& IncompleteMapping = ControlRigBoneOutputMappingIncomplete;
+		TArray<TPair<uint16, uint16>>& CompleteMapping = ControlRigBoneOutputMappingComplete;
 
-		for (auto Iter = ControlRigBoneMapping.CreateConstIterator(); Iter; ++Iter)
+		// if we don't have a different mapping for outputs, use the input mapping
+		if(IncompleteMapping.IsEmpty() && CompleteMapping.IsEmpty())
 		{
-			const FName& Name = Iter.Key();
-			const uint16 Index = Iter.Value();
-			const FRigElementKey Key(Name, ERigElementType::Bone);
+			IncompleteMapping = ControlRigBoneInputMappingIncomplete;
+			CompleteMapping = ControlRigBoneInputMappingComplete;
+		}
 
-			FCompactPoseBoneIndex CompactPoseIndex(Index);
-			FTransform ComponentTransform = ControlRig->GetHierarchy()->GetGlobalTransform(Key);
-			if (NodeMappingContainer.IsValid())
+		if(bTransferPoseInGlobalSpace || NodeMappingContainer.IsValid())
+		{
+			// get component pose from control rig
+			FCSPose<FCompactPose> MeshPoses;
+			MeshPoses.InitPose(InOutput.Pose);
+
+			if(!CompleteMapping.IsEmpty())
 			{
-				ComponentTransform = NodeMappingContainer->GetSourceToTargetTransform(Name) * ComponentTransform;
+				for (const TPair<uint16, uint16>& Pair : CompleteMapping)
+				{
+					const uint16 SkeletonIndex = Pair.Key;
+					const uint16 ControlRigIndex = Pair.Value;
+
+					FCompactPoseBoneIndex CompactPoseIndex(SkeletonIndex);
+					FTransform ComponentTransform = ControlRig->GetHierarchy()->GetGlobalTransformByIndex(ControlRigIndex);
+					MeshPoses.SetComponentSpaceTransform(CompactPoseIndex, ComponentTransform);
+				}
 			}
+			else
+			{
+				for (auto Iter = IncompleteMapping.CreateConstIterator(); Iter; ++Iter)
+				{
+					const FName& Name = Iter.Key();
+					const uint16 Index = Iter.Value();
+					const FRigElementKey Key(Name, ERigElementType::Bone);
 
-			MeshPoses.SetComponentSpaceTransform(CompactPoseIndex, ComponentTransform);
+					FCompactPoseBoneIndex CompactPoseIndex(Index);
+					FTransform ComponentTransform = ControlRig->GetHierarchy()->GetGlobalTransform(Key);
+					if (NodeMappingContainer.IsValid())
+					{
+						ComponentTransform = NodeMappingContainer->GetSourceToTargetTransform(Name) * ComponentTransform;
+					}
+
+					MeshPoses.SetComponentSpaceTransform(CompactPoseIndex, ComponentTransform);
+				}
+			}
+
+			FCSPose<FCompactPose>::ConvertComponentPosesToLocalPosesSafe(MeshPoses, InOutput.Pose);
+			InOutput.Pose.NormalizeRotations();
 		}
+		else
+		{
+			if(!CompleteMapping.IsEmpty())
+			{
+				for (const TPair<uint16, uint16>& Pair : CompleteMapping)
+				{
+					const uint16 SkeletonIndex = Pair.Key;
+					const uint16 ControlRigIndex = Pair.Value;
 
-		FCSPose<FCompactPose>::ConvertComponentPosesToLocalPosesSafe(MeshPoses, InOutput.Pose);
-		InOutput.Pose.NormalizeRotations();
+					FCompactPoseBoneIndex CompactPoseIndex(SkeletonIndex);
+					FTransform LocalTransform = ControlRig->GetHierarchy()->GetLocalTransformByIndex(ControlRigIndex);
+					InOutput.Pose[CompactPoseIndex] = LocalTransform;
+				}
+			}
+			else
+			{
+				for (auto Iter = IncompleteMapping.CreateConstIterator(); Iter; ++Iter)
+				{
+					const FName& Name = Iter.Key();
+					const uint16 Index = Iter.Value();
+					const FRigElementKey Key(Name, ERigElementType::Bone);
+
+					FCompactPoseBoneIndex CompactPoseIndex(Index);
+					FTransform LocalTransform = ControlRig->GetHierarchy()->GetLocalTransform(Key);
+					InOutput.Pose[CompactPoseIndex] = LocalTransform;
+				}
+			}
+		}
 	}
 
 	if (OutputSettings.bUpdateCurves)
 	{
-		// update curve
-		for (auto Iter = ControlRigCurveMapping.CreateConstIterator(); Iter; ++Iter)
+		if(!ControlRigCurveMappingComplete.IsEmpty())
 		{
-			const FName& Name = Iter.Key();
-			const uint16 Index = Iter.Value();
-			const FRigElementKey Key(Name, ERigElementType::Curve);
+			for (const TPair<uint16, uint16>& Pair : ControlRigCurveMappingComplete)
+			{
+				const uint16 SkeletonIndex = Pair.Key;
+				const uint16 ControlRigIndex = Pair.Value;
 
-			const float PreviousValue = InOutput.Curve.Get(Index);
-			const float Value = ControlRig->GetHierarchy()->GetCurveValue(Key);
+				const float PreviousValue = InOutput.Curve.Get(SkeletonIndex);
+				const float Value = ControlRig->GetHierarchy()->GetCurveValueByIndex(ControlRigIndex);
 
-			if(!FMath::IsNearlyEqual(PreviousValue, Value))
+				if(!FMath::IsNearlyEqual(PreviousValue, Value))
+				{
+					// this causes a side effect of marking the curve as "valid"
+					// so only apply it for curves that have really changed
+					InOutput.Curve.Set(SkeletonIndex, Value);
+				}
+			}
+		}
+		else
+		{
+			for (auto Iter = ControlRigCurveMappingIncomplete.CreateConstIterator(); Iter; ++Iter)
 			{
-				// this causes a side effect of marking the curve as "valid"
-				// so only apply it for curves that have really changed
-				InOutput.Curve.Set(Index, Value);
+				const FName& Name = Iter.Key();
+				const uint16 Index = Iter.Value();
+				const FRigElementKey Key(Name, ERigElementType::Curve);
+
+				const float PreviousValue = InOutput.Curve.Get(Index);
+				const float Value = ControlRig->GetHierarchy()->GetCurveValue(Key);
+
+				if(!FMath::IsNearlyEqual(PreviousValue, Value))
+				{
+					// this causes a side effect of marking the curve as "valid"
+					// so only apply it for curves that have really changed
+					InOutput.Curve.Set(Index, Value);
+				}
 			}
 		}
 	}
@@ -394,8 +552,13 @@ void FAnimNode_ControlRigBase::CacheBones_AnyThread(const FAnimationCacheBonesCo
 		// fill up node names
 		FBoneContainer& RequiredBones = Context.AnimInstanceProxy->GetRequiredBones();
 
-		ControlRigBoneMapping.Reset();
-		ControlRigCurveMapping.Reset();
+		ControlRigBoneInputMappingComplete.Reset();
+		ControlRigBoneOutputMappingComplete.Reset();
+		ControlRigCurveMappingComplete.Reset();
+		ControlRigTransformElementsToReset.Reset();
+		ControlRigBoneInputMappingIncomplete.Reset();
+		ControlRigBoneOutputMappingIncomplete.Reset();
+		ControlRigCurveMappingIncomplete.Reset();
 
 		if(RequiredBones.IsValid())
 		{
@@ -419,7 +582,7 @@ void FAnimNode_ControlRigBase::CacheBones_AnyThread(const FAnimationCacheBonesCo
 					FName* SourceName = TargetToSourceMappingTable.Find(TargetNodeName);
 					if (SourceName)
 					{
-						ControlRigBoneMapping.Add(*SourceName, Index);
+						ControlRigBoneInputMappingIncomplete.Add(*SourceName, Index);
 					}
 				}
 			}
@@ -435,7 +598,53 @@ void FAnimNode_ControlRigBase::CacheBones_AnyThread(const FAnimationCacheBonesCo
 					const FName& BoneName = RefSkeleton.GetBoneName(RequiredBonesArray[Index]);
 					if (NodeNames.Contains(BoneName))
 					{
-						ControlRigBoneMapping.Add(BoneName, Index);
+						ControlRigBoneInputMappingIncomplete.Add(BoneName, Index);
+					}
+				}
+			}
+
+			if(!InputBonesToTransfer.IsEmpty())
+			{
+				ControlRigBoneOutputMappingIncomplete = ControlRigBoneInputMappingIncomplete;
+				ControlRigBoneInputMappingIncomplete.Reset();
+				
+				if (NodeMappingContainer.IsValid())
+				{
+					// get target to source mapping table - this is reversed mapping table
+					TMap<FName, FName> TargetToSourceMappingTable;
+					NodeMappingContainer->GetTargetToSourceMappingTable(TargetToSourceMappingTable);
+
+					for(FBoneReference& InputBoneToTransfer : InputBonesToTransfer)
+					{
+						if(!InputBoneToTransfer.Initialize(RequiredBones))
+						{
+							continue;
+						}
+						FName TargetNodeName = RefSkeleton.GetBoneName(InputBoneToTransfer.BoneIndex);
+						FName* SourceName = TargetToSourceMappingTable.Find(TargetNodeName);
+						if (SourceName)
+						{
+							ControlRigBoneInputMappingIncomplete.Add(*SourceName, InputBoneToTransfer.BoneIndex);
+						}
+					}
+				}
+				else
+				{
+					TArray<FName> NodeNames;
+					TArray<FNodeItem> NodeItems;
+					ControlRig->GetMappableNodeData(NodeNames, NodeItems);
+
+					for(FBoneReference& InputBoneToTransfer : InputBonesToTransfer)
+					{
+						if(!InputBoneToTransfer.Initialize(RequiredBones))
+						{
+							continue;
+						}
+						const FName& BoneName = RefSkeleton.GetBoneName(RequiredBonesArray[InputBoneToTransfer.BoneIndex]);
+						if (NodeNames.Contains(BoneName))
+						{
+							ControlRigBoneInputMappingIncomplete.Add(BoneName, InputBoneToTransfer.BoneIndex);
+						}
 					}
 				}
 			}
@@ -448,9 +657,90 @@ void FAnimNode_ControlRigBase::CacheBones_AnyThread(const FAnimationCacheBonesCo
 				// see if the curve name exists in the control rig
 				if (Hierarchy->GetIndex(FRigElementKey(CurveNames[Index], ERigElementType::Curve)) != INDEX_NONE)
 				{
-					ControlRigCurveMapping.Add(CurveNames[Index], Index);
+					ControlRigCurveMappingIncomplete.Add(CurveNames[Index], Index);
 				}
 			}
+
+			// check if we can switch the bones to a complete mapping.
+			// we can only do that if there is no node mapping container set.
+			if(!NodeMappingContainer.IsValid())
+			{
+				for(int32 InputOutput = 0; InputOutput < 2; InputOutput++)
+				{
+					bool bIsCompleteBoneMapping = true;
+					TMap<FName, uint16>& IncompleteMapping = InputOutput == 0 ? ControlRigBoneInputMappingIncomplete : ControlRigBoneOutputMappingIncomplete;
+					if(IncompleteMapping.IsEmpty())
+					{
+						continue;
+					}
+					
+					TArray<TPair<uint16, uint16>>& CompleteMapping = InputOutput == 0 ? ControlRigBoneInputMappingComplete : ControlRigBoneOutputMappingComplete;
+					
+					for (auto Iter = IncompleteMapping.CreateConstIterator(); Iter; ++Iter)
+					{
+						const uint16 SkeletonIndex = Iter.Value();
+						const int32 ControlRigIndex = ControlRig->GetHierarchy()->GetIndex(FRigElementKey(Iter.Key(), ERigElementType::Bone));
+						if(ControlRigIndex != INDEX_NONE)
+						{
+							CompleteMapping.Add(TPair<uint16, uint16>(SkeletonIndex, (uint16)ControlRigIndex));
+						}
+						else
+						{
+							bIsCompleteBoneMapping = false;
+						}
+					}
+
+					if(bIsCompleteBoneMapping)
+					{
+						IncompleteMapping.Reset();
+
+						if(InputOutput == 0)
+						{
+							ControlRigTransformElementsToReset.Reset();
+
+							for(int32 ElementIndex = 0; ElementIndex < ControlRig->GetHierarchy()->Num(); ElementIndex++)
+							{
+								const FRigElementKey Key = ControlRig->GetHierarchy()->GetKey(ElementIndex);
+								if(Key.Type == ERigElementType::Bone || Key.Type == ERigElementType::Curve)
+								{
+									continue;
+								}
+								ControlRigTransformElementsToReset.Add((uint16)ElementIndex);
+							}
+						}
+					}
+					else
+					{
+						CompleteMapping.Reset();
+					}
+				}
+			}
+
+			bool bIsCompleteCurveMapping = true;
+			
+			// check if we can switch the curves to a complete mapping as well
+			for (auto Iter = ControlRigCurveMappingIncomplete.CreateConstIterator(); Iter; ++Iter)
+			{
+				const uint16 SkeletonIndex = Iter.Value();
+				const int32 ControlRigIndex = ControlRig->GetHierarchy()->GetIndex(FRigElementKey(Iter.Key(), ERigElementType::Curve));
+				if(ControlRigIndex != INDEX_NONE)
+				{
+					ControlRigCurveMappingComplete.Add(TPair<uint16, uint16>(SkeletonIndex, (uint16)ControlRigIndex));
+				}
+				else
+				{
+					bIsCompleteCurveMapping = false;
+				}
+			}
+
+			if(bIsCompleteCurveMapping)
+			{
+				ControlRigCurveMappingIncomplete.Reset();
+			}
+			else
+			{
+				ControlRigCurveMappingComplete.Reset();
+			}
 		}
 
 		// re-init when LOD changes
diff --git a/Engine/Plugins/Animation/ControlRig/Source/ControlRig/Public/AnimNode_ControlRigBase.h b/Engine/Plugins/Animation/ControlRig/Source/ControlRig/Public/AnimNode_ControlRigBase.h
--- a/Engine/Plugins/Animation/ControlRig/Source/ControlRig/Public/AnimNode_ControlRigBase.h
+++ b/Engine/Plugins/Animation/ControlRig/Source/ControlRig/Public/AnimNode_ControlRigBase.h
@@ -65,15 +65,58 @@ protected:
 	UPROPERTY(EditAnywhere, Category = Links)
 	FPoseLink Source;
 
+	/**
+	 * If this is checked the bone pose coming from the AnimBP will be
+	 * transferred into the Control Rig.
+	 */
+	UPROPERTY(EditAnywhere, Category = Settings)
+	bool bTransferInputPose;
+
+	/**
+	 * If this is checked the curves coming from the AnimBP will be
+	 * transferred into the Control Rig.
+	 */
+	UPROPERTY(EditAnywhere, Category = Settings)
+	bool bTransferInputCurves;
+
+	/**
+	 * Transferring the pose in global space guarantees a global pose match,
+	 * while transferring in local space ensures a match of the local transforms.
+	 * In general transforms only differ if the hierarchy topology differs
+	 * between the Control Rig and the skeleton used in the AnimBP.
+	 * Note: Turning this off can potentially improve performance.
+	 */
+	UPROPERTY(EditAnywhere, Category = Settings)
+	bool bTransferPoseInGlobalSpace;
+
+	/**
+	 * An inclusive list of bones to transfer as part
+	 * of the input pose transfer phase.
+	 * If this list is empty all bones will be transferred.
+	 */
+	UPROPERTY(EditAnywhere, Category = Settings)
+	TArray<FBoneReference> InputBonesToTransfer;
+
+	/** Complete mapping from skeleton to control rig bone index */
+	TArray<TPair<uint16, uint16>> ControlRigBoneInputMappingComplete;
+	TArray<TPair<uint16, uint16>> ControlRigBoneOutputMappingComplete;
+
+	/** Complete mapping from skeleton to curve index */
+	TArray<TPair<uint16, uint16>> ControlRigCurveMappingComplete;
+
+	/**
+	 * Indices of transform elements requiring a reset
+	 * this is only used in conjunctions with ControlRigBoneMappingComplete
+	 */
+	TArray<uint16> ControlRigTransformElementsToReset;
+
 	/** Rig Hierarchy bone name to required array index mapping */
-	UPROPERTY(transient)
-	TMap<FName, uint16> ControlRigBoneMapping;
+	TMap<FName, uint16> ControlRigBoneInputMappingIncomplete;
+	TMap<FName, uint16> ControlRigBoneOutputMappingIncomplete;
 
 	/** Rig Curve name to Curve LUI mapping */
-	UPROPERTY(transient)
-	TMap<FName, uint16> ControlRigCurveMapping;
+	TMap<FName, uint16> ControlRigCurveMappingIncomplete;
 
-	UPROPERTY(transient)
 	TMap<FName, uint16> InputToCurveMappingUIDs;
 
 	/** Node Mapping Container */
