Control Rig: make sure to call PostLoad() for all referenced CRBP before PostLoad() the current CRBP, otherwise, when current CRBP compiles during post load, it will try to inline functions from the referenced BPs which will fail due to referenced BP not fully loaded

diff --git a/Engine/Plugins/Animation/ControlRig/Source/ControlRigDeveloper/Private/ControlRigBlueprint.cpp b/Engine/Plugins/Animation/ControlRig/Source/ControlRigDeveloper/Private/ControlRigBlueprint.cpp
--- a/Engine/Plugins/Animation/ControlRig/Source/ControlRigDeveloper/Private/ControlRigBlueprint.cpp
+++ b/Engine/Plugins/Animation/ControlRig/Source/ControlRigDeveloper/Private/ControlRigBlueprint.cpp
@@ -281,39 +281,9 @@ void UControlRigBlueprint::Serialize(FArchive& Ar)
 
 	if(Ar.IsObjectReferenceCollector())
 	{
-		TArray<UBlueprint*> ReferencedBlueprints;
-		
-		TArray<UEdGraph*> EdGraphs;
-		GetAllGraphs(EdGraphs);
-		for (UEdGraph* EdGraph : EdGraphs)
-		{
-			for(UEdGraphNode* Node : EdGraph->Nodes)
-			{
-				if(UControlRigGraphNode* RigNode = Cast<UControlRigGraphNode>(Node))
-				{
-					if(URigVMFunctionReferenceNode* FunctionRefNode = Cast<URigVMFunctionReferenceNode>(RigNode->GetModelNode()))
-					{
-						if(URigVMLibraryNode* ReferencedNode = FunctionRefNode->GetReferencedNode())
-						{
-							if(URigVMFunctionLibrary* ReferencedFunctionLibrary = ReferencedNode->GetLibrary())
-							{
-								if(ReferencedFunctionLibrary == GetLocalFunctionLibrary())
-								{
-									continue;
-								}
+		TArray<UControlRigBlueprint*> ReferencedBlueprints = GetReferencedControlRigBlueprints();
 
-								if(UBlueprint* ReferencedBlueprint = Cast<UBlueprint>(ReferencedFunctionLibrary->GetOuter()))
-								{
-									ReferencedBlueprints.AddUnique(ReferencedBlueprint);
-								}
-							}
-						}
-					}
-				}
-			}
-		}
-
-		for(UBlueprint* ReferencedBlueprint : ReferencedBlueprints)
+		for(UControlRigBlueprint* ReferencedBlueprint : ReferencedBlueprints)
 		{
 			Ar << ReferencedBlueprints;
 		}
@@ -367,6 +337,18 @@ void UControlRigBlueprint::PostLoad()
 {
 	Super::PostLoad();
 
+	TArray<UControlRigBlueprint*> ReferencedBlueprints = GetReferencedControlRigBlueprints();
+
+	// PostLoad all referenced BPs so that their function graphs are fully loaded 
+	// and ready to be inlined into this BP during compilation
+	for (UControlRigBlueprint* BP : ReferencedBlueprints)
+	{
+		if (BP->HasAllFlags(RF_NeedPostLoad))
+		{
+			BP->ConditionalPostLoad();
+		}
+	}
+	
 	// temporarily disable default value validation during load time, serialized values should always be accepted
 	TGuardValue<bool> DisablePinDefaultValueValidation(GetOrCreateController()->bValidatePinDefaults, false);
 
@@ -726,6 +708,43 @@ void UControlRigBlueprint::HandleReportFromCompiler(EMessageSeverity::Type InSev
 	}
 }
 
+TArray<UControlRigBlueprint*> UControlRigBlueprint::GetReferencedControlRigBlueprints()
+{
+	TArray<UControlRigBlueprint*> ReferencedBlueprints;
+	
+	TArray<UEdGraph*> EdGraphs;
+	GetAllGraphs(EdGraphs);
+	for (UEdGraph* EdGraph : EdGraphs)
+	{
+		for(UEdGraphNode* Node : EdGraph->Nodes)
+		{
+			if(UControlRigGraphNode* RigNode = Cast<UControlRigGraphNode>(Node))
+			{
+				if(URigVMFunctionReferenceNode* FunctionRefNode = Cast<URigVMFunctionReferenceNode>(RigNode->GetModelNode()))
+				{
+					if(URigVMLibraryNode* ReferencedNode = FunctionRefNode->GetReferencedNode())
+					{
+						if(URigVMFunctionLibrary* ReferencedFunctionLibrary = ReferencedNode->GetLibrary())
+						{
+							if(ReferencedFunctionLibrary == GetLocalFunctionLibrary())
+							{
+								continue;
+							}
+
+							if(UControlRigBlueprint* ReferencedBlueprint = Cast<UControlRigBlueprint>(ReferencedFunctionLibrary->GetOuter()))
+							{
+								ReferencedBlueprints.AddUnique(ReferencedBlueprint);
+							}
+						}
+					}
+				}
+			}
+		}
+	}
+	
+	return ReferencedBlueprints;
+}
+
 #if WITH_EDITOR
 
 void UControlRigBlueprint::ClearBreakpoints()
diff --git a/Engine/Plugins/Animation/ControlRig/Source/ControlRigDeveloper/Public/ControlRigBlueprint.h b/Engine/Plugins/Animation/ControlRig/Source/ControlRigDeveloper/Public/ControlRigBlueprint.h
--- a/Engine/Plugins/Animation/ControlRig/Source/ControlRigDeveloper/Public/ControlRigBlueprint.h
+++ b/Engine/Plugins/Animation/ControlRig/Source/ControlRigDeveloper/Public/ControlRigBlueprint.h
@@ -567,6 +567,8 @@ private:
 	bool RemoveEdGraphForCollapseNode(URigVMCollapseNode* InNode, bool bNotify = false);
 	void HandleReportFromCompiler(EMessageSeverity::Type InSeverity, UObject* InSubject, const FString& InMessage);
 
+	TArray<UControlRigBlueprint*> GetReferencedControlRigBlueprints();
+	
 #if WITH_EDITOR
 private:
 	bool bCompileInDebugMode;
