[Backout] - CL16985307
Original CL Desc
-----------------------------------------------------------------
Perf improvement for meta human face control rigs from [REDACTED_HANDLE]
Optimizes control rig pose input/output
Adds option to pass partial hierarchy as control rig input.

#rnx

diff --git a/Engine/Plugins/Animation/ControlRig/Source/ControlRig/Private/AnimNode_ControlRigBase.cpp b/Engine/Plugins/Animation/ControlRig/Source/ControlRig/Private/AnimNode_ControlRigBase.cpp
--- a/Engine/Plugins/Animation/ControlRig/Source/ControlRig/Private/AnimNode_ControlRigBase.cpp
+++ b/Engine/Plugins/Animation/ControlRig/Source/ControlRig/Private/AnimNode_ControlRigBase.cpp
@@ -17,9 +17,6 @@ static TAutoConsoleVariable<int32> CVarControlRigDisableExecutionAnimNode(TEXT("
 
 FAnimNode_ControlRigBase::FAnimNode_ControlRigBase()
 	: FAnimNode_CustomProperty()
-	, bTransferInputPose(true)
-	, bTransferInputCurves(true)
-	, bTransferPoseInGlobalSpace(true)
 	, InputSettings(FControlRigIOSettings())
 	, OutputSettings(FControlRigIOSettings())
 	, bExecute(true)
@@ -125,132 +122,53 @@ void FAnimNode_ControlRigBase::UpdateInput(UControlRig* ControlRig, const FPoseC
 
 	DECLARE_SCOPE_HIERARCHICAL_COUNTER_FUNC()
 
-	if (InputSettings.bUpdatePose && bTransferInputPose)
+	if (InputSettings.bUpdatePose)
 	{
 		const FBoneContainer& RequiredBones = InOutput.Pose.GetBoneContainer();
 
+		// get component pose from control rig
+		FCSPose<FCompactPose> MeshPoses;
+		// first I need to convert to local pose
+		MeshPoses.InitPose(InOutput.Pose);
+
 		// reset transforms here to prevent additive transforms from accumulating to INF
 		// we only update transforms from the mesh pose for bones in the current LOD, 
 		// so the reset here ensures excluded bones are also reset
-		if(!ControlRigBoneInputMappingIncomplete.IsEmpty())
-		{
-			ControlRig->GetHierarchy()->ResetPoseToInitial(ERigElementType::Bone);
-		}
+		ControlRig->GetHierarchy()->ResetPoseToInitial(ERigElementType::Bone);
 
-		if(bTransferPoseInGlobalSpace || NodeMappingContainer.IsValid())
+		// @re-think - now control rig contains init pose from their default hierarchy and current pose from this instance.
+		// we may need this init pose somewhere (instance refpose)
+		for (auto Iter = ControlRigBoneMapping.CreateConstIterator(); Iter; ++Iter)
 		{
-			// get component pose from control rig
-			FCSPose<FCompactPose> MeshPoses;
-			// first I need to convert to local pose
-			MeshPoses.InitPose(InOutput.Pose);
+			const FName& Name = Iter.Key();
+			const uint16 Index = Iter.Value();
+			const FRigElementKey Key(Name, ERigElementType::Bone);
 
-			if(!ControlRigBoneInputMappingComplete.IsEmpty())
+			FTransform ComponentTransform = MeshPoses.GetComponentSpaceTransform(FCompactPoseBoneIndex(Index));
+			if (NodeMappingContainer.IsValid())
 			{
-				for (const TPair<uint16, uint16>& Pair : ControlRigBoneInputMappingComplete)
-				{
-					const uint16 SkeletonIndex = Pair.Key;
-					const uint16 ControlRigIndex = Pair.Value;
-					
-					FCompactPoseBoneIndex CompactPoseIndex(SkeletonIndex);
-					FTransform ComponentTransform = MeshPoses.GetComponentSpaceTransform(CompactPoseIndex);
-					ControlRig->GetHierarchy()->SetGlobalTransformByIndex(ControlRigIndex, ComponentTransform, false);
-				}
+				ComponentTransform = NodeMappingContainer->GetSourceToTargetTransform(Name).GetRelativeTransformReverse(ComponentTransform);
 			}
-			else
-			{
-				for (auto Iter = ControlRigBoneInputMappingIncomplete.CreateConstIterator(); Iter; ++Iter)
-				{
-					const FName& Name = Iter.Key();
-					const uint16 Index = Iter.Value();
-					const FRigElementKey Key(Name, ERigElementType::Bone);
 
-					FCompactPoseBoneIndex CompactPoseIndex(Index);
-					FTransform ComponentTransform = MeshPoses.GetComponentSpaceTransform(CompactPoseIndex);
-					if (NodeMappingContainer.IsValid())
-					{
-						ComponentTransform = NodeMappingContainer->GetSourceToTargetTransform(Name).GetRelativeTransformReverse(ComponentTransform);
-					}
-					ControlRig->GetHierarchy()->SetGlobalTransform(Key, ComponentTransform, false);
-				}
-			}
+			ControlRig->GetHierarchy()->SetGlobalTransform(Key, ComponentTransform, false);
 		}
-		else
-		{
-			if(!ControlRigBoneInputMappingComplete.IsEmpty())
-			{
-				for (const TPair<uint16, uint16>& Pair : ControlRigBoneInputMappingComplete)
-				{
-					const uint16 SkeletonIndex = Pair.Key;
-					const uint16 ControlRigIndex = Pair.Value;
-					
-					FCompactPoseBoneIndex CompactPoseIndex(SkeletonIndex);
-					FTransform LocalTransform = InOutput.Pose[CompactPoseIndex];
-					ControlRig->GetHierarchy()->SetLocalTransformByIndex(ControlRigIndex, LocalTransform, false);
-				}
-			}
-			else
-			{
-				for (auto Iter = ControlRigBoneInputMappingIncomplete.CreateConstIterator(); Iter; ++Iter)
-				{
-					const FName& Name = Iter.Key();
-					const uint16 Index = Iter.Value();
-					const FRigElementKey Key(Name, ERigElementType::Bone);
 
-					FCompactPoseBoneIndex CompactPoseIndex(Index);
-					FTransform LocalTransform = InOutput.Pose[CompactPoseIndex];
-					ControlRig->GetHierarchy()->SetLocalTransform(Key, LocalTransform, false);
-				}
-			}
-		}
-		
-		if(ControlRigBoneInputMappingIncomplete.IsEmpty())
-		{
-			if(bTransferPoseInGlobalSpace)
-			{
-				for(uint16 ElementIndexToReset : ControlRigTransformElementsToReset)
-				{
-					FTransform InitialTransform = ControlRig->GetHierarchy()->GetGlobalTransformByIndex(ElementIndexToReset, true);
-					ControlRig->GetHierarchy()->SetGlobalTransformByIndex(ElementIndexToReset, InitialTransform, false);
-				}
-			}
-			else
-			{
-				for(uint16 ElementIndexToReset : ControlRigTransformElementsToReset)
-				{
-					FTransform InitialTransform = ControlRig->GetHierarchy()->GetLocalTransformByIndex(ElementIndexToReset, true);
-					ControlRig->GetHierarchy()->SetLocalTransformByIndex(ElementIndexToReset, InitialTransform, false);
-				}
-			}
-		}
 					
 #if WITH_EDITOR
 		ControlRig->ApplyTransformOverrideForUserCreatedBones();
 #endif
 	}
 
-	if (InputSettings.bUpdateCurves && bTransferInputCurves)
+	if (InputSettings.bUpdateCurves)
 	{
-		if(!ControlRigCurveMappingComplete.IsEmpty())
-		{
-			for (const TPair<uint16, uint16>& Pair : ControlRigCurveMappingComplete)
-			{
-				const uint16 SkeletonIndex = Pair.Key;
-				const uint16 ControlRigIndex = Pair.Value;
-
-				const float Value = InOutput.Curve.Get(SkeletonIndex);
-				ControlRig->GetHierarchy()->SetCurveValueByIndex(ControlRigIndex, Value);
-			}
-		}
-		else
+		// we just do name mapping 
+		for (auto Iter = ControlRigCurveMapping.CreateConstIterator(); Iter; ++Iter)
 		{
-			for (auto Iter = ControlRigCurveMappingIncomplete.CreateConstIterator(); Iter; ++Iter)
-			{
-				const FName& Name = Iter.Key();
-				const uint16 Index = Iter.Value();
-				const FRigElementKey Key(Name, ERigElementType::Curve);
+			const FName& Name = Iter.Key();
+			const uint16 Index = Iter.Value();
+			const FRigElementKey Key(Name, ERigElementType::Curve);
 
-				ControlRig->GetHierarchy()->SetCurveValue(Key, InOutput.Curve.Get(Index));
-			}
+			ControlRig->GetHierarchy()->SetCurveValue(Key, InOutput.Curve.Get(Index));
 		}
 	}
 
@@ -279,123 +197,47 @@ void FAnimNode_ControlRigBase::UpdateOutput(UControlRig* ControlRig, FPoseContex
 		// copy output of the rig
 		const FBoneContainer& RequiredBones = InOutput.Pose.GetBoneContainer();
 
-		TMap<FName, uint16>& IncompleteMapping = ControlRigBoneOutputMappingIncomplete;
-		TArray<TPair<uint16, uint16>>& CompleteMapping = ControlRigBoneOutputMappingComplete;
+		// get component pose from control rig
+		FCSPose<FCompactPose> MeshPoses;
+		MeshPoses.InitPose(InOutput.Pose);
 
-		// if we don't have a different mapping for outputs, use the input mapping
-		if(IncompleteMapping.IsEmpty() && CompleteMapping.IsEmpty())
+		for (auto Iter = ControlRigBoneMapping.CreateConstIterator(); Iter; ++Iter)
 		{
-			IncompleteMapping = ControlRigBoneInputMappingIncomplete;
-			CompleteMapping = ControlRigBoneInputMappingComplete;
-		}
-
-		if(bTransferPoseInGlobalSpace || NodeMappingContainer.IsValid())
-		{
-			// get component pose from control rig
-			FCSPose<FCompactPose> MeshPoses;
-			MeshPoses.InitPose(InOutput.Pose);
+			const FName& Name = Iter.Key();
+			const uint16 Index = Iter.Value();
+			const FRigElementKey Key(Name, ERigElementType::Bone);
 
-			if(!CompleteMapping.IsEmpty())
-			{
-				for (const TPair<uint16, uint16>& Pair : CompleteMapping)
-				{
-					const uint16 SkeletonIndex = Pair.Key;
-					const uint16 ControlRigIndex = Pair.Value;
-
-					FCompactPoseBoneIndex CompactPoseIndex(SkeletonIndex);
-					FTransform ComponentTransform = ControlRig->GetHierarchy()->GetGlobalTransformByIndex(ControlRigIndex);
-					MeshPoses.SetComponentSpaceTransform(CompactPoseIndex, ComponentTransform);
-				}
-			}
-			else
+			FCompactPoseBoneIndex CompactPoseIndex(Index);
+			FTransform ComponentTransform = ControlRig->GetHierarchy()->GetGlobalTransform(Key);
+			if (NodeMappingContainer.IsValid())
 			{
-				for (auto Iter = IncompleteMapping.CreateConstIterator(); Iter; ++Iter)
-				{
-					const FName& Name = Iter.Key();
-					const uint16 Index = Iter.Value();
-					const FRigElementKey Key(Name, ERigElementType::Bone);
-
-					FCompactPoseBoneIndex CompactPoseIndex(Index);
-					FTransform ComponentTransform = ControlRig->GetHierarchy()->GetGlobalTransform(Key);
-					if (NodeMappingContainer.IsValid())
-					{
-						ComponentTransform = NodeMappingContainer->GetSourceToTargetTransform(Name) * ComponentTransform;
-					}
-
-					MeshPoses.SetComponentSpaceTransform(CompactPoseIndex, ComponentTransform);
-				}
+				ComponentTransform = NodeMappingContainer->GetSourceToTargetTransform(Name) * ComponentTransform;
 			}
 
-			FCSPose<FCompactPose>::ConvertComponentPosesToLocalPosesSafe(MeshPoses, InOutput.Pose);
-			InOutput.Pose.NormalizeRotations();
+			MeshPoses.SetComponentSpaceTransform(CompactPoseIndex, ComponentTransform);
 		}
-		else
-		{
-			if(!CompleteMapping.IsEmpty())
-			{
-				for (const TPair<uint16, uint16>& Pair : CompleteMapping)
-				{
-					const uint16 SkeletonIndex = Pair.Key;
-					const uint16 ControlRigIndex = Pair.Value;
 
-					FCompactPoseBoneIndex CompactPoseIndex(SkeletonIndex);
-					FTransform LocalTransform = ControlRig->GetHierarchy()->GetLocalTransformByIndex(ControlRigIndex);
-					InOutput.Pose[CompactPoseIndex] = LocalTransform;
-				}
-			}
-			else
-			{
-				for (auto Iter = IncompleteMapping.CreateConstIterator(); Iter; ++Iter)
-				{
-					const FName& Name = Iter.Key();
-					const uint16 Index = Iter.Value();
-					const FRigElementKey Key(Name, ERigElementType::Bone);
-
-					FCompactPoseBoneIndex CompactPoseIndex(Index);
-					FTransform LocalTransform = ControlRig->GetHierarchy()->GetLocalTransform(Key);
-					InOutput.Pose[CompactPoseIndex] = LocalTransform;
-				}
-			}
-		}
+		FCSPose<FCompactPose>::ConvertComponentPosesToLocalPosesSafe(MeshPoses, InOutput.Pose);
+		InOutput.Pose.NormalizeRotations();
 	}
 
 	if (OutputSettings.bUpdateCurves)
 	{
-		if(!ControlRigCurveMappingComplete.IsEmpty())
+		// update curve
+		for (auto Iter = ControlRigCurveMapping.CreateConstIterator(); Iter; ++Iter)
 		{
-			for (const TPair<uint16, uint16>& Pair : ControlRigCurveMappingComplete)
-			{
-				const uint16 SkeletonIndex = Pair.Key;
-				const uint16 ControlRigIndex = Pair.Value;
+			const FName& Name = Iter.Key();
+			const uint16 Index = Iter.Value();
+			const FRigElementKey Key(Name, ERigElementType::Curve);
 
-				const float PreviousValue = InOutput.Curve.Get(SkeletonIndex);
-				const float Value = ControlRig->GetHierarchy()->GetCurveValueByIndex(ControlRigIndex);
+			const float PreviousValue = InOutput.Curve.Get(Index);
+			const float Value = ControlRig->GetHierarchy()->GetCurveValue(Key);
 
-				if(!FMath::IsNearlyEqual(PreviousValue, Value))
-				{
-					// this causes a side effect of marking the curve as "valid"
-					// so only apply it for curves that have really changed
-					InOutput.Curve.Set(SkeletonIndex, Value);
-				}
-			}
-		}
-		else
-		{
-			for (auto Iter = ControlRigCurveMappingIncomplete.CreateConstIterator(); Iter; ++Iter)
+			if(!FMath::IsNearlyEqual(PreviousValue, Value))
 			{
-				const FName& Name = Iter.Key();
-				const uint16 Index = Iter.Value();
-				const FRigElementKey Key(Name, ERigElementType::Curve);
-
-				const float PreviousValue = InOutput.Curve.Get(Index);
-				const float Value = ControlRig->GetHierarchy()->GetCurveValue(Key);
-
-				if(!FMath::IsNearlyEqual(PreviousValue, Value))
-				{
-					// this causes a side effect of marking the curve as "valid"
-					// so only apply it for curves that have really changed
-					InOutput.Curve.Set(Index, Value);
-				}
+				// this causes a side effect of marking the curve as "valid"
+				// so only apply it for curves that have really changed
+				InOutput.Curve.Set(Index, Value);
 			}
 		}
 	}
@@ -552,13 +394,8 @@ void FAnimNode_ControlRigBase::CacheBones_AnyThread(const FAnimationCacheBonesCo
 		// fill up node names
 		FBoneContainer& RequiredBones = Context.AnimInstanceProxy->GetRequiredBones();
 
-		ControlRigBoneInputMappingComplete.Reset();
-		ControlRigBoneOutputMappingComplete.Reset();
-		ControlRigCurveMappingComplete.Reset();
-		ControlRigTransformElementsToReset.Reset();
-		ControlRigBoneInputMappingIncomplete.Reset();
-		ControlRigBoneOutputMappingIncomplete.Reset();
-		ControlRigCurveMappingIncomplete.Reset();
+		ControlRigBoneMapping.Reset();
+		ControlRigCurveMapping.Reset();
 
 		if(RequiredBones.IsValid())
 		{
@@ -582,7 +419,7 @@ void FAnimNode_ControlRigBase::CacheBones_AnyThread(const FAnimationCacheBonesCo
 					FName* SourceName = TargetToSourceMappingTable.Find(TargetNodeName);
 					if (SourceName)
 					{
-						ControlRigBoneInputMappingIncomplete.Add(*SourceName, Index);
+						ControlRigBoneMapping.Add(*SourceName, Index);
 					}
 				}
 			}
@@ -598,53 +435,7 @@ void FAnimNode_ControlRigBase::CacheBones_AnyThread(const FAnimationCacheBonesCo
 					const FName& BoneName = RefSkeleton.GetBoneName(RequiredBonesArray[Index]);
 					if (NodeNames.Contains(BoneName))
 					{
-						ControlRigBoneInputMappingIncomplete.Add(BoneName, Index);
-					}
-				}
-			}
-
-			if(!InputBonesToTransfer.IsEmpty())
-			{
-				ControlRigBoneOutputMappingIncomplete = ControlRigBoneInputMappingIncomplete;
-				ControlRigBoneInputMappingIncomplete.Reset();
-				
-				if (NodeMappingContainer.IsValid())
-				{
-					// get target to source mapping table - this is reversed mapping table
-					TMap<FName, FName> TargetToSourceMappingTable;
-					NodeMappingContainer->GetTargetToSourceMappingTable(TargetToSourceMappingTable);
-
-					for(FBoneReference& InputBoneToTransfer : InputBonesToTransfer)
-					{
-						if(!InputBoneToTransfer.Initialize(RequiredBones))
-						{
-							continue;
-						}
-						FName TargetNodeName = RefSkeleton.GetBoneName(InputBoneToTransfer.BoneIndex);
-						FName* SourceName = TargetToSourceMappingTable.Find(TargetNodeName);
-						if (SourceName)
-						{
-							ControlRigBoneInputMappingIncomplete.Add(*SourceName, InputBoneToTransfer.BoneIndex);
-						}
-					}
-				}
-				else
-				{
-					TArray<FName> NodeNames;
-					TArray<FNodeItem> NodeItems;
-					ControlRig->GetMappableNodeData(NodeNames, NodeItems);
-
-					for(FBoneReference& InputBoneToTransfer : InputBonesToTransfer)
-					{
-						if(!InputBoneToTransfer.Initialize(RequiredBones))
-						{
-							continue;
-						}
-						const FName& BoneName = RefSkeleton.GetBoneName(RequiredBonesArray[InputBoneToTransfer.BoneIndex]);
-						if (NodeNames.Contains(BoneName))
-						{
-							ControlRigBoneInputMappingIncomplete.Add(BoneName, InputBoneToTransfer.BoneIndex);
-						}
+						ControlRigBoneMapping.Add(BoneName, Index);
 					}
 				}
 			}
@@ -657,90 +448,9 @@ void FAnimNode_ControlRigBase::CacheBones_AnyThread(const FAnimationCacheBonesCo
 				// see if the curve name exists in the control rig
 				if (Hierarchy->GetIndex(FRigElementKey(CurveNames[Index], ERigElementType::Curve)) != INDEX_NONE)
 				{
-					ControlRigCurveMappingIncomplete.Add(CurveNames[Index], Index);
+					ControlRigCurveMapping.Add(CurveNames[Index], Index);
 				}
 			}
-
-			// check if we can switch the bones to a complete mapping.
-			// we can only do that if there is no node mapping container set.
-			if(!NodeMappingContainer.IsValid())
-			{
-				for(int32 InputOutput = 0; InputOutput < 2; InputOutput++)
-				{
-					bool bIsCompleteBoneMapping = true;
-					TMap<FName, uint16>& IncompleteMapping = InputOutput == 0 ? ControlRigBoneInputMappingIncomplete : ControlRigBoneOutputMappingIncomplete;
-					if(IncompleteMapping.IsEmpty())
-					{
-						continue;
-					}
-					
-					TArray<TPair<uint16, uint16>>& CompleteMapping = InputOutput == 0 ? ControlRigBoneInputMappingComplete : ControlRigBoneOutputMappingComplete;
-					
-					for (auto Iter = IncompleteMapping.CreateConstIterator(); Iter; ++Iter)
-					{
-						const uint16 SkeletonIndex = Iter.Value();
-						const int32 ControlRigIndex = ControlRig->GetHierarchy()->GetIndex(FRigElementKey(Iter.Key(), ERigElementType::Bone));
-						if(ControlRigIndex != INDEX_NONE)
-						{
-							CompleteMapping.Add(TPair<uint16, uint16>(SkeletonIndex, (uint16)ControlRigIndex));
-						}
-						else
-						{
-							bIsCompleteBoneMapping = false;
-						}
-					}
-
-					if(bIsCompleteBoneMapping)
-					{
-						IncompleteMapping.Reset();
-
-						if(InputOutput == 0)
-						{
-							ControlRigTransformElementsToReset.Reset();
-
-							for(int32 ElementIndex = 0; ElementIndex < ControlRig->GetHierarchy()->Num(); ElementIndex++)
-							{
-								const FRigElementKey Key = ControlRig->GetHierarchy()->GetKey(ElementIndex);
-								if(Key.Type == ERigElementType::Bone || Key.Type == ERigElementType::Curve)
-								{
-									continue;
-								}
-								ControlRigTransformElementsToReset.Add((uint16)ElementIndex);
-							}
-						}
-					}
-					else
-					{
-						CompleteMapping.Reset();
-					}
-				}
-			}
-
-			bool bIsCompleteCurveMapping = true;
-			
-			// check if we can switch the curves to a complete mapping as well
-			for (auto Iter = ControlRigCurveMappingIncomplete.CreateConstIterator(); Iter; ++Iter)
-			{
-				const uint16 SkeletonIndex = Iter.Value();
-				const int32 ControlRigIndex = ControlRig->GetHierarchy()->GetIndex(FRigElementKey(Iter.Key(), ERigElementType::Curve));
-				if(ControlRigIndex != INDEX_NONE)
-				{
-					ControlRigCurveMappingComplete.Add(TPair<uint16, uint16>(SkeletonIndex, (uint16)ControlRigIndex));
-				}
-				else
-				{
-					bIsCompleteCurveMapping = false;
-				}
-			}
-
-			if(bIsCompleteCurveMapping)
-			{
-				ControlRigCurveMappingIncomplete.Reset();
-			}
-			else
-			{
-				ControlRigCurveMappingComplete.Reset();
-			}
 		}
 
 		// re-init when LOD changes
diff --git a/Engine/Plugins/Animation/ControlRig/Source/ControlRig/Public/AnimNode_ControlRigBase.h b/Engine/Plugins/Animation/ControlRig/Source/ControlRig/Public/AnimNode_ControlRigBase.h
--- a/Engine/Plugins/Animation/ControlRig/Source/ControlRig/Public/AnimNode_ControlRigBase.h
+++ b/Engine/Plugins/Animation/ControlRig/Source/ControlRig/Public/AnimNode_ControlRigBase.h
@@ -65,58 +65,15 @@ protected:
 	UPROPERTY(EditAnywhere, Category = Links)
 	FPoseLink Source;
 
-	/**
-	 * If this is checked the bone pose coming from the AnimBP will be
-	 * transferred into the Control Rig.
-	 */
-	UPROPERTY(EditAnywhere, Category = Settings)
-	bool bTransferInputPose;
-
-	/**
-	 * If this is checked the curves coming from the AnimBP will be
-	 * transferred into the Control Rig.
-	 */
-	UPROPERTY(EditAnywhere, Category = Settings)
-	bool bTransferInputCurves;
-
-	/**
-	 * Transferring the pose in global space guarantees a global pose match,
-	 * while transferring in local space ensures a match of the local transforms.
-	 * In general transforms only differ if the hierarchy topology differs
-	 * between the Control Rig and the skeleton used in the AnimBP.
-	 * Note: Turning this off can potentially improve performance.
-	 */
-	UPROPERTY(EditAnywhere, Category = Settings)
-	bool bTransferPoseInGlobalSpace;
-
-	/**
-	 * An inclusive list of bones to transfer as part
-	 * of the input pose transfer phase.
-	 * If this list is empty all bones will be transferred.
-	 */
-	UPROPERTY(EditAnywhere, Category = Settings)
-	TArray<FBoneReference> InputBonesToTransfer;
-
-	/** Complete mapping from skeleton to control rig bone index */
-	TArray<TPair<uint16, uint16>> ControlRigBoneInputMappingComplete;
-	TArray<TPair<uint16, uint16>> ControlRigBoneOutputMappingComplete;
-
-	/** Complete mapping from skeleton to curve index */
-	TArray<TPair<uint16, uint16>> ControlRigCurveMappingComplete;
-
-	/**
-	 * Indices of transform elements requiring a reset
-	 * this is only used in conjunctions with ControlRigBoneMappingComplete
-	 */
-	TArray<uint16> ControlRigTransformElementsToReset;
-
 	/** Rig Hierarchy bone name to required array index mapping */
-	TMap<FName, uint16> ControlRigBoneInputMappingIncomplete;
-	TMap<FName, uint16> ControlRigBoneOutputMappingIncomplete;
+	UPROPERTY(transient)
+	TMap<FName, uint16> ControlRigBoneMapping;
 
 	/** Rig Curve name to Curve LUI mapping */
-	TMap<FName, uint16> ControlRigCurveMappingIncomplete;
+	UPROPERTY(transient)
+	TMap<FName, uint16> ControlRigCurveMapping;
 
+	UPROPERTY(transient)
 	TMap<FName, uint16> InputToCurveMappingUIDs;
 
 	/** Node Mapping Container */
