Control Rig & Sequencer: Control Rig Components created by sequencer for different worlds should not share the same control rig instance, each CR component should have its own control rig instance

diff --git a/Engine/Plugins/Animation/ControlRig/Source/ControlRig/Private/Sequencer/MovieSceneControlRigParameterTemplate.cpp b/Engine/Plugins/Animation/ControlRig/Source/ControlRig/Private/Sequencer/MovieSceneControlRigParameterTemplate.cpp
--- a/Engine/Plugins/Animation/ControlRig/Source/ControlRig/Private/Sequencer/MovieSceneControlRigParameterTemplate.cpp
+++ b/Engine/Plugins/Animation/ControlRig/Source/ControlRig/Private/Sequencer/MovieSceneControlRigParameterTemplate.cpp
@@ -815,7 +815,7 @@ struct FControlRigParameterPreAnimatedTokenProducer : IMovieScenePreAnimatedToke
 				{
 					if (ControlRigComponent->GetControlRig() != ControlRig)
 					{
-						ControlRigComponent->SetControlRig(ControlRig);
+						ControlRigComponent->Initialize();
 					}
 					else
 					{
@@ -913,7 +913,7 @@ struct FControlRigParameterExecutionToken : IMovieSceneExecutionToken
 									ControlRig->GetObjectBinding()->BindToObject(BoundObjects[0].Get());
 									if (NewControlRigComponent->GetControlRig() != ControlRig)
 									{
-										NewControlRigComponent->SetControlRig(ControlRig);
+										NewControlRigComponent->Initialize();
 									}
 									else
 									{
@@ -929,7 +929,7 @@ struct FControlRigParameterExecutionToken : IMovieSceneExecutionToken
 								ControlRig->GetObjectBinding()->BindToObject(BoundObjects[0].Get());
 								if (NewControlRigComponent->GetControlRig() != ControlRig)
 								{
-									NewControlRigComponent->SetControlRig(ControlRig);
+									NewControlRigComponent->Initialize();
 								}
 								else
 								{
