Control Rig: Mark transient DetailsViewWrapperObject as transient, non-transactional

diff --git a/Engine/Plugins/Animation/ControlRig/Source/ControlRigEditor/Private/Editor/DetailsViewWrapperObject.cpp b/Engine/Plugins/Animation/ControlRig/Source/ControlRigEditor/Private/Editor/DetailsViewWrapperObject.cpp
--- a/Engine/Plugins/Animation/ControlRig/Source/ControlRigEditor/Private/Editor/DetailsViewWrapperObject.cpp
+++ b/Engine/Plugins/Animation/ControlRig/Source/ControlRigEditor/Private/Editor/DetailsViewWrapperObject.cpp
@@ -27,8 +27,8 @@ UClass* UDetailsViewWrapperObject::GetClassForStruct(UScriptStruct* InStruct)
 	UClass* WrapperClass = NewObject<UClass>(
 		GetTransientPackage(),
 		WrapperClassName,
-		RF_Public | RF_Transactional
-		);
+		RF_Public | RF_Transient
+ 		);
 
 	// make sure this doesn't get garbage collected
 	WrapperClass->AddToRoot();
@@ -150,7 +150,7 @@ UDetailsViewWrapperObject* UDetailsViewWrapperObject::MakeInstance(UScriptStruct
 		return nullptr;
 	}
 
-	UDetailsViewWrapperObject* Instance = NewObject<UDetailsViewWrapperObject>(InOuter, WrapperClass, NAME_None, RF_Public | RF_Transactional);
+	UDetailsViewWrapperObject* Instance = NewObject<UDetailsViewWrapperObject>(InOuter, WrapperClass, NAME_None, RF_Public | RF_Transient);
 	Instance->SetContent(InStructMemory, InStruct->GetStructureSize());
 	return Instance;
 }
