Sequencer: Control Rig: Setup edit mode if not set up when a controller is selected. This fixes a bug when a controller is selected from the picker for the first time.

Refactor selecting rigs and controls out from OnSelectionChanged so that it can be called from HandleControlSelected

diff --git a/Engine/Plugins/Animation/ControlRig/Source/ControlRigEditor/Private/Sequencer/ControlRigParameterTrackEditor.cpp b/Engine/Plugins/Animation/ControlRig/Source/ControlRigEditor/Private/Sequencer/ControlRigParameterTrackEditor.cpp
--- a/Engine/Plugins/Animation/ControlRig/Source/ControlRigEditor/Private/Sequencer/ControlRigParameterTrackEditor.cpp
+++ b/Engine/Plugins/Animation/ControlRig/Source/ControlRigEditor/Private/Sequencer/ControlRigParameterTrackEditor.cpp
@@ -1665,7 +1665,6 @@ void FControlRigParameterTrackEditor::OnSelectionChanged(TArray<UMovieSceneTrack
 	
 	TGuardValue<bool> Guard(bIsDoingSelection, true);
 
-	TArray<FString> StringArray;
 	FControlRigEditMode* ControlRigEditMode = GetEditMode();
 	UControlRig* ControlRig = nullptr;
 
@@ -1717,6 +1716,14 @@ void FControlRigParameterTrackEditor::OnSelectionChanged(TArray<UMovieSceneTrack
 		return;
 	}
 
+	SelectRigsAndControls(ControlRig, KeyAreas);
+}
+
+void FControlRigParameterTrackEditor::SelectRigsAndControls(UControlRig* ControlRig, const TArray<const IKeyArea*>& KeyAreas)
+{
+	FControlRigEditMode* ControlRigEditMode = GetEditMode();
+
+	TArray<FString> StringArray;
 	TMap<UControlRig *, TSet<FName>> RigsAndControls;
 	for (const IKeyArea* KeyArea : KeyAreas)
 	{
@@ -1938,6 +1945,8 @@ void FControlRigParameterTrackEditor::HandleControlSelected(UControlRig* Subject
 	}
 	TGuardValue<bool> Guard(bIsDoingSelection, true);
 
+	FControlRigEditMode* ControlRigEditMode = GetEditMode();
+
 	FName ControlRigName(*Subject->GetName());
 	if (TSharedPtr<IControlRigObjectBinding> ObjectBinding = Subject->GetObjectBinding())
 	{
@@ -1958,7 +1967,7 @@ void FControlRigParameterTrackEditor::HandleControlSelected(UControlRig* Subject
 
 		FFindOrCreateTrackResult TrackResult = FindOrCreateControlRigTrackForObject(ObjectHandle, Subject, ControlRigName, bCreateTrack);
 		UMovieSceneControlRigParameterTrack* Track = CastChecked<UMovieSceneControlRigParameterTrack>(TrackResult.Track, ECastCheckedType::NullAllowed);
-			if (Track)
+		if (Track)
 		{
 			GetSequencer()->SuspendSelectionBroadcast();
 			for (UMovieSceneSection* Section : Track->GetAllSections())
@@ -1968,6 +1977,10 @@ void FControlRigParameterTrackEditor::HandleControlSelected(UControlRig* Subject
 			}
 			GetSequencer()->ResumeSelectionBroadcast();
 
+			TArray<const IKeyArea*> KeyAreas;
+			GetSequencer()->GetSelectedKeyAreas(KeyAreas);
+			SelectRigsAndControls(Subject, KeyAreas);
+
 			//Force refresh now, not later
 			GetSequencer()->RefreshTree();
 
diff --git a/Engine/Plugins/Animation/ControlRig/Source/ControlRigEditor/Private/Sequencer/ControlRigParameterTrackEditor.h b/Engine/Plugins/Animation/ControlRig/Source/ControlRigEditor/Private/Sequencer/ControlRigParameterTrackEditor.h
--- a/Engine/Plugins/Animation/ControlRig/Source/ControlRigEditor/Private/Sequencer/ControlRigParameterTrackEditor.h
+++ b/Engine/Plugins/Animation/ControlRig/Source/ControlRigEditor/Private/Sequencer/ControlRigParameterTrackEditor.h
@@ -118,6 +118,9 @@ private:
 	void HandleControlSelected(UControlRig* Subject, FRigControlElement* ControlElement, bool bSelected);
 	void HandleOnInitialized(UControlRig* Subject, const EControlRigState InState, const FName& InEventName);
 
+	/** Select control rig if not selected, select controls from key areas */
+	void SelectRigsAndControls(UControlRig* Subject, const TArray<const IKeyArea*>& KeyAreas);
+
 	/** Handle Creation for SkelMeshComp or Actor Owner, either may have a binding*/
 	FMovieSceneTrackEditor::FFindOrCreateHandleResult FindOrCreateHandleToSceneCompOrOwner(USceneComponent* InComp);
 
