Control Rig: Pop up a confirmation message when unlinking a level sequence

diff --git a/Engine/Plugins/Animation/ControlRig/Source/ControlRigEditor/Private/ControlRigEditorModule.cpp b/Engine/Plugins/Animation/ControlRig/Source/ControlRigEditor/Private/ControlRigEditorModule.cpp
--- a/Engine/Plugins/Animation/ControlRig/Source/ControlRigEditor/Private/ControlRigEditorModule.cpp
+++ b/Engine/Plugins/Animation/ControlRig/Source/ControlRigEditor/Private/ControlRigEditorModule.cpp
@@ -115,6 +115,8 @@
 #include "EditMode/SControlRigTweenWidget.h"
 #include "EditMode/SControlRigSnapper.h"
 #include "Dialogs/CustomDialog.h"
+#include "Framework/Notifications/NotificationManager.h"
+#include "Widgets/Notifications/SNotificationList.h"
 #include "Tools/SMotionTrailOptions.h"
 
 #define LOCTEXT_NAMESPACE "ControlRigEditorModule"
@@ -897,6 +899,19 @@ void FControlRigEditorModule::UnLinkLevelSequence(UAnimSequence* AnimSequence)
 							if (LevelAnimLinkItem.ResolveAnimSequence() == AnimSequence)
 							{
 								LevelAnimLink->AnimSequenceLinks.RemoveAtSwap(Index);
+							
+								const FText NotificationText = FText::Format(LOCTEXT("UnlinkLevelSequenceSuccess", "{0} unlinked from "), FText::FromString(AnimSequence->GetName()));
+								FNotificationInfo Info(NotificationText);
+								Info.ExpireDuration = 5.f;
+								Info.Hyperlink = FSimpleDelegate::CreateLambda([=]()
+								{
+									TArray<UObject*> Assets;
+									Assets.Add(LevelSequence);
+									GEditor->GetEditorSubsystem<UAssetEditorSubsystem>()->OpenEditorForAssets(Assets);
+								});
+								Info.HyperlinkText = FText::Format(LOCTEXT("OpenUnlinkedLevelSequenceLink", "{0}"), FText::FromString(LevelSequence->GetName()));
+								FSlateNotificationManager::Get().AddNotification(Info)->SetCompletionState(SNotificationItem::CS_Success);
+														
 								break;
 							}
 						}
