Control Rig: Improve AnimNode runtime performance by adding more options for filtering input transfer from [REDACTED_HANDLE]

diff --git a/Engine/Plugins/Animation/ControlRig/Source/ControlRig/Private/AnimNode_ControlRigBase.cpp b/Engine/Plugins/Animation/ControlRig/Source/ControlRig/Private/AnimNode_ControlRigBase.cpp
--- a/Engine/Plugins/Animation/ControlRig/Source/ControlRig/Private/AnimNode_ControlRigBase.cpp
+++ b/Engine/Plugins/Animation/ControlRig/Source/ControlRig/Private/AnimNode_ControlRigBase.cpp
@@ -17,6 +17,10 @@ static TAutoConsoleVariable<int32> CVarControlRigDisableExecutionAnimNode(TEXT("
 
 FAnimNode_ControlRigBase::FAnimNode_ControlRigBase()
 	: FAnimNode_CustomProperty()
+	, bResetInputPoseToInitial(true) 
+	, bTransferInputPose(true)
+	, bTransferInputCurves(true)
+	, bTransferPoseInGlobalSpace(true)
 	, InputSettings(FControlRigIOSettings())
 	, OutputSettings(FControlRigIOSettings())
 	, bExecute(true)
@@ -104,6 +108,9 @@ bool FAnimNode_ControlRigBase::CanExecute()
 
 void FAnimNode_ControlRigBase::UpdateInput(UControlRig* ControlRig, const FPoseContext& InOutput)
 {
+	QUICK_SCOPE_CYCLE_COUNTER(STAT_ControlRig_UpdateInput);
+
+
 	if(!CanExecute())
 	{
 		return;
@@ -122,53 +129,112 @@ void FAnimNode_ControlRigBase::UpdateInput(UControlRig* ControlRig, const FPoseC
 
 	DECLARE_SCOPE_HIERARCHICAL_COUNTER_FUNC()
 
-	if (InputSettings.bUpdatePose)
+	if (InputSettings.bUpdatePose && bTransferInputPose)
 	{
 		const FBoneContainer& RequiredBones = InOutput.Pose.GetBoneContainer();
 
-		// get component pose from control rig
-		FCSPose<FCompactPose> MeshPoses;
-		// first I need to convert to local pose
-		MeshPoses.InitPose(InOutput.Pose);
-
 		// reset transforms here to prevent additive transforms from accumulating to INF
 		// we only update transforms from the mesh pose for bones in the current LOD, 
 		// so the reset here ensures excluded bones are also reset
-		ControlRig->GetHierarchy()->ResetPoseToInitial(ERigElementType::Bone);
+		if(!ControlRigBoneInputMappingByName.IsEmpty() || bResetInputPoseToInitial)
+		{
+			ControlRig->GetHierarchy()->ResetPoseToInitial(ERigElementType::Bone);
+		}
 
-		// @re-think - now control rig contains init pose from their default hierarchy and current pose from this instance.
-		// we may need this init pose somewhere (instance refpose)
-		for (auto Iter = ControlRigBoneMapping.CreateConstIterator(); Iter; ++Iter)
+		if(bTransferPoseInGlobalSpace || NodeMappingContainer.IsValid())
 		{
-			const FName& Name = Iter.Key();
-			const uint16 Index = Iter.Value();
-			const FRigElementKey Key(Name, ERigElementType::Bone);
+			// get component pose from control rig
+			FCSPose<FCompactPose> MeshPoses;
+			// first I need to convert to local pose
+			MeshPoses.InitPose(InOutput.Pose);
 
-			FTransform ComponentTransform = MeshPoses.GetComponentSpaceTransform(FCompactPoseBoneIndex(Index));
-			if (NodeMappingContainer.IsValid())
+			if(!ControlRigBoneInputMappingByIndex.IsEmpty())
 			{
-				ComponentTransform = NodeMappingContainer->GetSourceToTargetTransform(Name).GetRelativeTransformReverse(ComponentTransform);
+				for (const TPair<uint16, uint16>& Pair : ControlRigBoneInputMappingByIndex)
+				{
+					const uint16 ControlRigIndex = Pair.Key;
+					const uint16 SkeletonIndex = Pair.Value;
+					
+					FCompactPoseBoneIndex CompactPoseIndex(SkeletonIndex);
+					FTransform ComponentTransform = MeshPoses.GetComponentSpaceTransform(CompactPoseIndex);
+					ControlRig->GetHierarchy()->SetGlobalTransformByIndex(ControlRigIndex, ComponentTransform, false);
+				}
 			}
+			else
+			{
+				for (auto Iter = ControlRigBoneInputMappingByName.CreateConstIterator(); Iter; ++Iter)
+				{
+					const FName& Name = Iter.Key();
+					const uint16 Index = Iter.Value();
+					const FRigElementKey Key(Name, ERigElementType::Bone);
 
-			ControlRig->GetHierarchy()->SetGlobalTransform(Key, ComponentTransform, false);
+					FCompactPoseBoneIndex CompactPoseIndex(Index);
+					FTransform ComponentTransform = MeshPoses.GetComponentSpaceTransform(CompactPoseIndex);
+					if (NodeMappingContainer.IsValid())
+					{
+						ComponentTransform = NodeMappingContainer->GetSourceToTargetTransform(Name).GetRelativeTransformReverse(ComponentTransform);
+					}
+					ControlRig->GetHierarchy()->SetGlobalTransform(Key, ComponentTransform, false);
+				}
+			}
 		}
-
+		else
+		{
+			if(!ControlRigBoneInputMappingByIndex.IsEmpty())
+			{
+				for (const TPair<uint16, uint16>& Pair : ControlRigBoneInputMappingByIndex)
+				{
+					const uint16 ControlRigIndex = Pair.Key;
+					const uint16 SkeletonIndex = Pair.Value;
 					
+					FCompactPoseBoneIndex CompactPoseIndex(SkeletonIndex);
+					FTransform LocalTransform = InOutput.Pose[CompactPoseIndex];
+					ControlRig->GetHierarchy()->SetLocalTransformByIndex(ControlRigIndex, LocalTransform, false);
+				}
+			}
+			else
+			{
+				for (auto Iter = ControlRigBoneInputMappingByName.CreateConstIterator(); Iter; ++Iter)
+				{
+					const FName& Name = Iter.Key();
+					const uint16 SkeletonIndex = Iter.Value();
+					const FRigElementKey Key(Name, ERigElementType::Bone);
+
+					FCompactPoseBoneIndex CompactPoseIndex(SkeletonIndex);
+					FTransform LocalTransform = InOutput.Pose[CompactPoseIndex];
+					ControlRig->GetHierarchy()->SetLocalTransform(Key, LocalTransform, false);
+				}
+			}
+		}
+		
 #if WITH_EDITOR
 		ControlRig->ApplyTransformOverrideForUserCreatedBones();
 #endif
 	}
 
-	if (InputSettings.bUpdateCurves)
+	if (InputSettings.bUpdateCurves && bTransferInputCurves)
 	{
-		// we just do name mapping 
-		for (auto Iter = ControlRigCurveMapping.CreateConstIterator(); Iter; ++Iter)
+		if(!ControlRigCurveMappingByIndex.IsEmpty())
 		{
-			const FName& Name = Iter.Key();
-			const uint16 Index = Iter.Value();
-			const FRigElementKey Key(Name, ERigElementType::Curve);
+			for (const TPair<uint16, uint16>& Pair : ControlRigCurveMappingByIndex)
+			{
+				const uint16 ControlRigIndex = Pair.Key;
+				const uint16 SkeletonIndex = Pair.Value;
 
-			ControlRig->GetHierarchy()->SetCurveValue(Key, InOutput.Curve.Get(Index));
+				const float Value = InOutput.Curve.Get(SkeletonIndex);
+				ControlRig->GetHierarchy()->SetCurveValueByIndex(ControlRigIndex, Value);
+			}
+		}
+		else
+		{
+			for (auto Iter = ControlRigCurveMappingByName.CreateConstIterator(); Iter; ++Iter)
+			{
+				const FName& Name = Iter.Key();
+				const uint16 SkeletonIndex = Iter.Value();
+				const FRigElementKey Key(Name, ERigElementType::Curve);
+
+				ControlRig->GetHierarchy()->SetCurveValue(Key, InOutput.Curve.Get(SkeletonIndex));
+			}
 		}
 	}
 
@@ -185,6 +251,8 @@ void FAnimNode_ControlRigBase::UpdateInput(UControlRig* ControlRig, const FPoseC
 
 void FAnimNode_ControlRigBase::UpdateOutput(UControlRig* ControlRig, FPoseContext& InOutput)
 {
+	QUICK_SCOPE_CYCLE_COUNTER(STAT_ControlRig_UpdateOutput);
+
 	if(!CanExecute())
 	{
 		return;
@@ -197,47 +265,123 @@ void FAnimNode_ControlRigBase::UpdateOutput(UControlRig* ControlRig, FPoseContex
 		// copy output of the rig
 		const FBoneContainer& RequiredBones = InOutput.Pose.GetBoneContainer();
 
-		// get component pose from control rig
-		FCSPose<FCompactPose> MeshPoses;
-		MeshPoses.InitPose(InOutput.Pose);
+		TMap<FName, uint16>& NameBasedMapping = ControlRigBoneOutputMappingByName;
+		TArray<TPair<uint16, uint16>>& IndexBasedMapping = ControlRigBoneOutputMappingByIndex;
 
-		for (auto Iter = ControlRigBoneMapping.CreateConstIterator(); Iter; ++Iter)
+		// if we don't have a different mapping for outputs, use the input mapping
+		if(NameBasedMapping.IsEmpty() && IndexBasedMapping.IsEmpty())
 		{
-			const FName& Name = Iter.Key();
-			const uint16 Index = Iter.Value();
-			const FRigElementKey Key(Name, ERigElementType::Bone);
+			NameBasedMapping = ControlRigBoneInputMappingByName;
+			IndexBasedMapping = ControlRigBoneInputMappingByIndex;
+		}
 
-			FCompactPoseBoneIndex CompactPoseIndex(Index);
-			FTransform ComponentTransform = ControlRig->GetHierarchy()->GetGlobalTransform(Key);
-			if (NodeMappingContainer.IsValid())
+		if(bTransferPoseInGlobalSpace || NodeMappingContainer.IsValid())
+		{
+			// get component pose from control rig
+			FCSPose<FCompactPose> MeshPoses;
+			MeshPoses.InitPose(InOutput.Pose);
+
+			if(!IndexBasedMapping.IsEmpty())
 			{
-				ComponentTransform = NodeMappingContainer->GetSourceToTargetTransform(Name) * ComponentTransform;
+				for (const TPair<uint16, uint16>& Pair : IndexBasedMapping)
+				{
+					const uint16 ControlRigIndex = Pair.Key;
+					const uint16 SkeletonIndex = Pair.Value;
+
+					FCompactPoseBoneIndex CompactPoseIndex(SkeletonIndex);
+					FTransform ComponentTransform = ControlRig->GetHierarchy()->GetGlobalTransformByIndex(ControlRigIndex);
+					MeshPoses.SetComponentSpaceTransform(CompactPoseIndex, ComponentTransform);
+				}
+			}
+			else
+			{
+				for (auto Iter = NameBasedMapping.CreateConstIterator(); Iter; ++Iter)
+				{
+					const FName& Name = Iter.Key();
+					const uint16 SkeletonIndex = Iter.Value();
+					const FRigElementKey Key(Name, ERigElementType::Bone);
+
+					FCompactPoseBoneIndex CompactPoseIndex(SkeletonIndex);
+					FTransform ComponentTransform = ControlRig->GetHierarchy()->GetGlobalTransform(Key);
+					if (NodeMappingContainer.IsValid())
+					{
+						ComponentTransform = NodeMappingContainer->GetSourceToTargetTransform(Name) * ComponentTransform;
+					}
+
+					MeshPoses.SetComponentSpaceTransform(CompactPoseIndex, ComponentTransform);
+				}
 			}
 
-			MeshPoses.SetComponentSpaceTransform(CompactPoseIndex, ComponentTransform);
+			FCSPose<FCompactPose>::ConvertComponentPosesToLocalPosesSafe(MeshPoses, InOutput.Pose);
+			InOutput.Pose.NormalizeRotations();
 		}
+		else
+		{
+			if(!IndexBasedMapping.IsEmpty())
+			{
+				for (const TPair<uint16, uint16>& Pair : IndexBasedMapping)
+				{
+					const uint16 ControlRigIndex = Pair.Key;
+					const uint16 SkeletonIndex = Pair.Value;
+
+					FCompactPoseBoneIndex CompactPoseIndex(SkeletonIndex);
+					FTransform LocalTransform = ControlRig->GetHierarchy()->GetLocalTransformByIndex(ControlRigIndex);
+					InOutput.Pose[CompactPoseIndex] = LocalTransform;
+				}
+			}
+			else
+			{
+				for (auto Iter = NameBasedMapping.CreateConstIterator(); Iter; ++Iter)
+				{
+					const FName& Name = Iter.Key();
+					const uint16 Index = Iter.Value();
+					const FRigElementKey Key(Name, ERigElementType::Bone);
 
-		FCSPose<FCompactPose>::ConvertComponentPosesToLocalPosesSafe(MeshPoses, InOutput.Pose);
-		InOutput.Pose.NormalizeRotations();
+					FCompactPoseBoneIndex CompactPoseIndex(Index);
+					FTransform LocalTransform = ControlRig->GetHierarchy()->GetLocalTransform(Key);
+					InOutput.Pose[CompactPoseIndex] = LocalTransform;
+				}
+			}
+		}
 	}
 
 	if (OutputSettings.bUpdateCurves)
 	{
-		// update curve
-		for (auto Iter = ControlRigCurveMapping.CreateConstIterator(); Iter; ++Iter)
+		if(!ControlRigCurveMappingByIndex.IsEmpty())
 		{
-			const FName& Name = Iter.Key();
-			const uint16 Index = Iter.Value();
-			const FRigElementKey Key(Name, ERigElementType::Curve);
+			for (const TPair<uint16, uint16>& Pair : ControlRigCurveMappingByIndex)
+			{
+				const uint16 ControlRigIndex = Pair.Key;
+				const uint16 SkeletonIndex = Pair.Value;
 
-			const float PreviousValue = InOutput.Curve.Get(Index);
-			const float Value = ControlRig->GetHierarchy()->GetCurveValue(Key);
+				const float PreviousValue = InOutput.Curve.Get(SkeletonIndex);
+				const float Value = ControlRig->GetHierarchy()->GetCurveValueByIndex(ControlRigIndex);
 
-			if(!FMath::IsNearlyEqual(PreviousValue, Value))
+				if(!FMath::IsNearlyEqual(PreviousValue, Value))
+				{
+					// this causes a side effect of marking the curve as "valid"
+					// so only apply it for curves that have really changed
+					InOutput.Curve.Set(SkeletonIndex, Value);
+				}
+			}
+		}
+		else
+		{
+			for (auto Iter = ControlRigCurveMappingByName.CreateConstIterator(); Iter; ++Iter)
 			{
-				// this causes a side effect of marking the curve as "valid"
-				// so only apply it for curves that have really changed
-				InOutput.Curve.Set(Index, Value);
+				const FName& Name = Iter.Key();
+				const uint16 Index = Iter.Value();
+				const FRigElementKey Key(Name, ERigElementType::Curve);
+
+				const float PreviousValue = InOutput.Curve.Get(Index);
+				const float Value = ControlRig->GetHierarchy()->GetCurveValue(Key);
+
+				if(!FMath::IsNearlyEqual(PreviousValue, Value))
+				{
+					// this causes a side effect of marking the curve as "valid"
+					// so only apply it for curves that have really changed
+					InOutput.Curve.Set(Index, Value);
+				}
 			}
 		}
 	}
@@ -394,8 +538,12 @@ void FAnimNode_ControlRigBase::CacheBones_AnyThread(const FAnimationCacheBonesCo
 		// fill up node names
 		FBoneContainer& RequiredBones = Context.AnimInstanceProxy->GetRequiredBones();
 
-		ControlRigBoneMapping.Reset();
-		ControlRigCurveMapping.Reset();
+		ControlRigBoneInputMappingByIndex.Reset();
+		ControlRigBoneOutputMappingByIndex.Reset();
+		ControlRigCurveMappingByIndex.Reset();
+		ControlRigBoneInputMappingByName.Reset();
+		ControlRigBoneOutputMappingByName.Reset();
+		ControlRigCurveMappingByName.Reset();
 
 		if(RequiredBones.IsValid())
 		{
@@ -419,7 +567,7 @@ void FAnimNode_ControlRigBase::CacheBones_AnyThread(const FAnimationCacheBonesCo
 					FName* SourceName = TargetToSourceMappingTable.Find(TargetNodeName);
 					if (SourceName)
 					{
-						ControlRigBoneMapping.Add(*SourceName, Index);
+						ControlRigBoneInputMappingByName.Add(*SourceName, Index);
 					}
 				}
 			}
@@ -435,7 +583,53 @@ void FAnimNode_ControlRigBase::CacheBones_AnyThread(const FAnimationCacheBonesCo
 					const FName& BoneName = RefSkeleton.GetBoneName(RequiredBonesArray[Index]);
 					if (NodeNames.Contains(BoneName))
 					{
-						ControlRigBoneMapping.Add(BoneName, Index);
+						ControlRigBoneInputMappingByName.Add(BoneName, Index);
+					}
+				}
+			}
+
+			if(!InputBonesToTransfer.IsEmpty())
+			{
+				ControlRigBoneOutputMappingByName = ControlRigBoneInputMappingByName;
+				ControlRigBoneInputMappingByName.Reset();
+				
+				if (NodeMappingContainer.IsValid())
+				{
+					// get target to source mapping table - this is reversed mapping table
+					TMap<FName, FName> TargetToSourceMappingTable;
+					NodeMappingContainer->GetTargetToSourceMappingTable(TargetToSourceMappingTable);
+
+					for(FBoneReference& InputBoneToTransfer : InputBonesToTransfer)
+					{
+						if(!InputBoneToTransfer.Initialize(RequiredBones))
+						{
+							continue;
+						}
+						FName TargetNodeName = RefSkeleton.GetBoneName(InputBoneToTransfer.BoneIndex);
+						FName* SourceName = TargetToSourceMappingTable.Find(TargetNodeName);
+						if (SourceName)
+						{
+							ControlRigBoneInputMappingByName.Add(*SourceName, InputBoneToTransfer.BoneIndex);
+						}
+					}
+				}
+				else
+				{
+					TArray<FName> NodeNames;
+					TArray<FNodeItem> NodeItems;
+					ControlRig->GetMappableNodeData(NodeNames, NodeItems);
+
+					for(FBoneReference& InputBoneToTransfer : InputBonesToTransfer)
+					{
+						if(!InputBoneToTransfer.Initialize(RequiredBones))
+						{
+							continue;
+						}
+						const FName& BoneName = RefSkeleton.GetBoneName(RequiredBonesArray[InputBoneToTransfer.BoneIndex]);
+						if (NodeNames.Contains(BoneName))
+						{
+							ControlRigBoneInputMappingByName.Add(BoneName, InputBoneToTransfer.BoneIndex);
+						}
 					}
 				}
 			}
@@ -448,9 +642,75 @@ void FAnimNode_ControlRigBase::CacheBones_AnyThread(const FAnimationCacheBonesCo
 				// see if the curve name exists in the control rig
 				if (Hierarchy->GetIndex(FRigElementKey(CurveNames[Index], ERigElementType::Curve)) != INDEX_NONE)
 				{
-					ControlRigCurveMapping.Add(CurveNames[Index], Index);
+					ControlRigCurveMappingByName.Add(CurveNames[Index], Index);
 				}
 			}
+
+			// check if we can switch the bones to an index based mapping.
+			// we can only do that if there is no node mapping container set.
+			if(!NodeMappingContainer.IsValid())
+			{
+				for(int32 InputOutput = 0; InputOutput < 2; InputOutput++)
+				{
+					bool bIsMappingByIndex = true;
+					TMap<FName, uint16>& NameBasedMapping = InputOutput == 0 ? ControlRigBoneInputMappingByName : ControlRigBoneOutputMappingByName;
+					if(NameBasedMapping.IsEmpty())
+					{
+						continue;
+					}
+					
+					TArray<TPair<uint16, uint16>>& IndexBasedMapping = InputOutput == 0 ? ControlRigBoneInputMappingByIndex : ControlRigBoneOutputMappingByIndex;
+					
+					for (auto Iter = NameBasedMapping.CreateConstIterator(); Iter; ++Iter)
+					{
+						const uint16 SkeletonIndex = Iter.Value();
+						const int32 ControlRigIndex = ControlRig->GetHierarchy()->GetIndex(FRigElementKey(Iter.Key(), ERigElementType::Bone));
+						if(ControlRigIndex != INDEX_NONE)
+						{
+							IndexBasedMapping.Add(TPair<uint16, uint16>((uint16)ControlRigIndex, SkeletonIndex));
+						}
+						else
+						{
+							bIsMappingByIndex = false;
+						}
+					}
+
+					if(bIsMappingByIndex)
+					{
+						NameBasedMapping.Reset();
+					}
+					else
+					{
+						IndexBasedMapping.Reset();
+					}
+				}
+			}
+
+			bool bIsCurveMappingByIndex = true;
+			
+			// check if we can switch the curves to a index based mapping as well
+			for (auto Iter = ControlRigCurveMappingByName.CreateConstIterator(); Iter; ++Iter)
+			{
+				const uint16 SkeletonIndex = Iter.Value();
+				const int32 ControlRigIndex = ControlRig->GetHierarchy()->GetIndex(FRigElementKey(Iter.Key(), ERigElementType::Curve));
+				if(ControlRigIndex != INDEX_NONE)
+				{
+					ControlRigCurveMappingByIndex.Add(TPair<uint16, uint16>((uint16)ControlRigIndex, SkeletonIndex));
+				}
+				else
+				{
+					bIsCurveMappingByIndex = false;
+				}
+			}
+
+			if(bIsCurveMappingByIndex)
+			{
+				ControlRigCurveMappingByName.Reset();
+			}
+			else
+			{
+				ControlRigCurveMappingByIndex.Reset();
+			}
 		}
 
 		// re-init when LOD changes
diff --git a/Engine/Plugins/Animation/ControlRig/Source/ControlRig/Public/AnimNode_ControlRigBase.h b/Engine/Plugins/Animation/ControlRig/Source/ControlRig/Public/AnimNode_ControlRigBase.h
--- a/Engine/Plugins/Animation/ControlRig/Source/ControlRig/Public/AnimNode_ControlRigBase.h
+++ b/Engine/Plugins/Animation/ControlRig/Source/ControlRig/Public/AnimNode_ControlRigBase.h
@@ -65,15 +65,59 @@ protected:
 	UPROPERTY(EditAnywhere, Category = Links)
 	FPoseLink Source;
 
+	/**
+	 * If this is checked the rig's pose needs to be reset to its initial
+	 * prior to evaluating the rig.
+	 */
+	UPROPERTY(EditAnywhere, AdvancedDisplay, Category = Settings)
+	bool bResetInputPoseToInitial;
+
+	/**
+	 * If this is checked the bone pose coming from the AnimBP will be
+	 * transferred into the Control Rig.
+	 */
+	UPROPERTY(EditAnywhere, AdvancedDisplay, Category = Settings)
+	bool bTransferInputPose;
+
+	/**
+	 * If this is checked the curves coming from the AnimBP will be
+	 * transferred into the Control Rig.
+	 */
+	UPROPERTY(EditAnywhere, AdvancedDisplay, Category = Settings)
+	bool bTransferInputCurves;
+
+	/**
+	 * Transferring the pose in global space guarantees a global pose match,
+	 * while transferring in local space ensures a match of the local transforms.
+	 * In general transforms only differ if the hierarchy topology differs
+	 * between the Control Rig and the skeleton used in the AnimBP.
+	 * Note: Turning this off can potentially improve performance.
+	 */
+	UPROPERTY(EditAnywhere, AdvancedDisplay, Category = Settings)
+	bool bTransferPoseInGlobalSpace;
+
+	/**
+	 * An inclusive list of bones to transfer as part
+	 * of the input pose transfer phase.
+	 * If this list is empty all bones will be transferred.
+	 */
+	UPROPERTY(EditAnywhere, AdvancedDisplay, Category = Settings)
+	TArray<FBoneReference> InputBonesToTransfer;
+
+	/** Complete mapping from skeleton to control rig bone index */
+	TArray<TPair<uint16, uint16>> ControlRigBoneInputMappingByIndex;
+	TArray<TPair<uint16, uint16>> ControlRigBoneOutputMappingByIndex;
+
+	/** Complete mapping from skeleton to curve index */
+	TArray<TPair<uint16, uint16>> ControlRigCurveMappingByIndex;
+
 	/** Rig Hierarchy bone name to required array index mapping */
-	UPROPERTY(transient)
-	TMap<FName, uint16> ControlRigBoneMapping;
+	TMap<FName, uint16> ControlRigBoneInputMappingByName;
+	TMap<FName, uint16> ControlRigBoneOutputMappingByName;
 
 	/** Rig Curve name to Curve LUI mapping */
-	UPROPERTY(transient)
-	TMap<FName, uint16> ControlRigCurveMapping;
+	TMap<FName, uint16> ControlRigCurveMappingByName;
 
-	UPROPERTY(transient)
 	TMap<FName, uint16> InputToCurveMappingUIDs;
 
 	/** Node Mapping Container */
