Sequencer: Control Rig: Fix for baking subsequences to control rigs, need to get transform, similar to what we did for animation baking.
#trivial

diff --git a/Engine/Plugins/Animation/ControlRig/Source/ControlRigEditor/Private/Sequencer/ControlRigParameterTrackEditor.cpp b/Engine/Plugins/Animation/ControlRig/Source/ControlRigEditor/Private/Sequencer/ControlRigParameterTrackEditor.cpp
--- a/Engine/Plugins/Animation/ControlRig/Source/ControlRigEditor/Private/Sequencer/ControlRigParameterTrackEditor.cpp
+++ b/Engine/Plugins/Animation/ControlRig/Source/ControlRigEditor/Private/Sequencer/ControlRigParameterTrackEditor.cpp
@@ -848,7 +848,7 @@ void FControlRigParameterTrackEditor::BakeToControlRig(UClass* InClass, FGuid Ob
 			TempAnimSequence->SetSkeleton(Skeleton);
 			const TSharedPtr<ISequencer> ParentSequencer = GetSequencer();
 			FMovieSceneSequenceIDRef Template = ParentSequencer->GetFocusedTemplateID();
-			FMovieSceneSequenceTransform RootToLocalTransform;
+			FMovieSceneSequenceTransform RootToLocalTransform = ParentSequencer->GetFocusedMovieSceneSequenceTransform();
 			UAnimSeqExportOption* AnimSeqExportOption = NewObject<UAnimSeqExportOption>(GetTransientPackage(), NAME_None);
 			UBakeToControlRigSettings* BakeSettings = GetMutableDefault<UBakeToControlRigSettings>();
 
