Control Rig: Offer inspection of RigVM Memory in optional details panel tab

diff --git a/Engine/Plugins/Animation/ControlRig/Source/ControlRigDeveloper/Public/ControlRigBlueprint.h b/Engine/Plugins/Animation/ControlRig/Source/ControlRigDeveloper/Public/ControlRigBlueprint.h
--- a/Engine/Plugins/Animation/ControlRig/Source/ControlRigDeveloper/Public/ControlRigBlueprint.h
+++ b/Engine/Plugins/Animation/ControlRig/Source/ControlRigDeveloper/Public/ControlRigBlueprint.h
@@ -42,6 +42,7 @@ DECLARE_EVENT_ThreeParams(UControlRigBlueprint, FOnLocalizeFunctionDialogRequest
 DECLARE_EVENT_ThreeParams(UControlRigBlueprint, FOnReportCompilerMessage, EMessageSeverity::Type, UObject*, const FString&);
 DECLARE_DELEGATE_RetVal_FourParams(FRigVMController_BulkEditResult, FControlRigOnBulkEditDialogRequestedDelegate, UControlRigBlueprint*, URigVMController*, URigVMLibraryNode*, ERigVMControllerBulkEditType);
 DECLARE_EVENT(UControlRigBlueprint, FOnBreakpointAdded);
+DECLARE_EVENT_OneParam(UControlRigBlueprint, FOnRequestInspectObject, const TArray<UObject*>& );
 
 USTRUCT()
 struct CONTROLRIGDEVELOPER_API FControlRigPublicFunctionArg
@@ -190,6 +191,9 @@ public:
 	// UObject interface
 	virtual void PostEditChangeChainProperty(struct FPropertyChangedChainEvent& PropertyChangedEvent) override;
 
+	FOnRequestInspectObject& OnRequestInspectObject() { return OnRequestInspectObjectEvent; }
+	void RequestInspectObject(const TArray<UObject*>& InObjects) { OnRequestInspectObjectEvent.Broadcast(InObjects); }
+
 #endif	// #if WITH_EDITOR
 
 	virtual bool ShouldBeMarkedDirtyUponTransaction() const override { return false; }
@@ -575,6 +579,8 @@ private:
 	
 	TArray<URigVMNode*> RigVMBreakpointNodes;
 
+	FOnRequestInspectObject OnRequestInspectObjectEvent;
+	
 public:
 
 	/** Sets the execution mode. In Release mode the rig will ignore all breakpoints. */
diff --git a/Engine/Plugins/Animation/ControlRig/Source/ControlRigEditor/Private/ControlRigCompilerDetails.cpp b/Engine/Plugins/Animation/ControlRig/Source/ControlRigEditor/Private/ControlRigCompilerDetails.cpp
--- a/Engine/Plugins/Animation/ControlRig/Source/ControlRigEditor/Private/ControlRigCompilerDetails.cpp
+++ b/Engine/Plugins/Animation/ControlRig/Source/ControlRigEditor/Private/ControlRigCompilerDetails.cpp
@@ -10,6 +10,7 @@
 #include "Widgets/Text/STextBlock.h"
 #include "HAL/PlatformApplicationMisc.h"
 #include "ControlRig.h"
+#include "IPropertyUtilities.h"
 #include "IPythonScriptPlugin.h"
 #include "RigVMPythonUtils.h"
 
@@ -52,6 +53,45 @@ void FRigVMCompileSettingsDetails::CustomizeChildren(TSharedRef<IPropertyHandle>
 			StructBuilder.AddProperty(InStructPropertyHandle->GetChildHandle(ChildIndex).ToSharedRef());
 		}
 
+#if !UE_RIGVM_UCLASS_BASED_STORAGE_DISABLED
+
+		StructBuilder.AddCustomRow(LOCTEXT("MemoryInspection", "Memory Inspection"))
+			.NameContent()
+			[
+				SNew(STextBlock)
+				.Text(FText::FromString(TEXT("Memory Inspection")))
+			.Font(IDetailLayoutBuilder::GetDetailFont())
+			]
+			.ValueContent()
+			[
+				SNew(SVerticalBox)
+				+ SVerticalBox::Slot()
+				[
+					SNew(SButton)
+					.OnClicked(this, &FRigVMCompileSettingsDetails::OnInspectMemory, ERigVMMemoryType::Literal)
+					.ContentPadding(FMargin(2))
+					.Content()
+					[
+						SNew(STextBlock)
+						.Justification(ETextJustify::Center)
+						.Text(LOCTEXT("InspectLiteralMemory", "Inspect Literal Memory"))
+					]
+				]
+				+ SVerticalBox::Slot()
+				[
+					SNew(SButton)
+					.OnClicked(this, &FRigVMCompileSettingsDetails::OnInspectMemory, ERigVMMemoryType::Work)
+					.ContentPadding(FMargin(2))
+					.Content()
+					[
+						SNew(STextBlock)
+						.Justification(ETextJustify::Center)
+						.Text(LOCTEXT("InspectWorkMemory", "Inspect Work Memory"))
+					]
+				]
+			];
+#endif
+
 		StructBuilder.AddCustomRow(LOCTEXT("ASTTools", "AST Tools"))
 			.NameContent()
 			[
@@ -114,6 +154,26 @@ void FRigVMCompileSettingsDetails::CustomizeChildren(TSharedRef<IPropertyHandle>
 	}
 }
 
+#if !UE_RIGVM_UCLASS_BASED_STORAGE_DISABLED
+
+FReply FRigVMCompileSettingsDetails::OnInspectMemory(ERigVMMemoryType InMemoryType)
+{
+	if (BlueprintBeingCustomized)
+	{
+		if(UControlRig* DebuggedRig = Cast<UControlRig>(BlueprintBeingCustomized->GetObjectBeingDebugged()))
+		{
+			if(URigVMMemoryStorage* MemoryStorage = DebuggedRig->GetVM()->GetMemoryByType(InMemoryType))
+			{
+				TArray<UObject*> ObjectsToSelect = {MemoryStorage};
+				BlueprintBeingCustomized->RequestInspectObject(ObjectsToSelect);
+			}
+		}
+	}
+	return FReply::Handled();
+}
+
+#endif
+
 FReply FRigVMCompileSettingsDetails::OnCopyASTClicked()
 {
 	if (BlueprintBeingCustomized)
diff --git a/Engine/Plugins/Animation/ControlRig/Source/ControlRigEditor/Private/ControlRigCompilerDetails.h b/Engine/Plugins/Animation/ControlRig/Source/ControlRigEditor/Private/ControlRigCompilerDetails.h
--- a/Engine/Plugins/Animation/ControlRig/Source/ControlRigEditor/Private/ControlRigCompilerDetails.h
+++ b/Engine/Plugins/Animation/ControlRig/Source/ControlRigEditor/Private/ControlRigCompilerDetails.h
@@ -24,6 +24,9 @@ public:
 
 protected:
 
+#if !UE_RIGVM_UCLASS_BASED_STORAGE_DISABLED
+	FReply OnInspectMemory(ERigVMMemoryType InMemoryType);
+#endif
 	FReply OnCopyASTClicked();
 	FReply OnCopyByteCodeClicked();
 	FReply OnCopyPythonScriptClicked();
diff --git a/Engine/Plugins/Animation/ControlRig/Source/ControlRigEditor/Private/Editor/ControlRigEditor.cpp b/Engine/Plugins/Animation/ControlRig/Source/ControlRigEditor/Private/Editor/ControlRigEditor.cpp
--- a/Engine/Plugins/Animation/ControlRig/Source/ControlRigEditor/Private/Editor/ControlRigEditor.cpp
+++ b/Engine/Plugins/Animation/ControlRig/Source/ControlRigEditor/Private/Editor/ControlRigEditor.cpp
@@ -345,6 +345,7 @@ void FControlRigEditor::InitControlRigEditor(const EToolkitMode::Type Mode, cons
 
 	InControlRigBlueprint->OnModified().AddSP(this, &FControlRigEditor::HandleModifiedEvent);
 	InControlRigBlueprint->OnVMCompiled().AddSP(this, &FControlRigEditor::HandleVMCompiledEvent);
+	InControlRigBlueprint->OnRequestInspectObject().AddSP(this, &FControlRigEditor::SetDetailObjects);
 
 	BindCommands();
 
@@ -1308,6 +1309,44 @@ void FControlRigEditor::SetDetailObjects(const TArray<UObject*>& InObjects)
 		return;
 	}
 
+#if !UE_RIGVM_UCLASS_BASED_STORAGE_DISABLED
+	if(InObjects.Num() == 1)
+	{
+		if(URigVMMemoryStorage* Memory = Cast<URigVMMemoryStorage>(InObjects[0]))
+		{
+			FPropertyEditorModule& EditModule = FModuleManager::Get().GetModuleChecked<FPropertyEditorModule>("PropertyEditor");
+
+			FDetailsViewArgs DetailsViewArgs;
+			DetailsViewArgs.NameAreaSettings = FDetailsViewArgs::HideNameArea;
+			DetailsViewArgs.bHideSelectionTip = true;
+
+			TSharedRef<IDetailsView> DetailsView = EditModule.CreateDetailView( DetailsViewArgs );
+			TSharedRef<SDockTab> DockTab = SNew(SDockTab)
+			.Label( LOCTEXT("ControlRigMemoryDetails", "Control Rig Memory Details") )
+			.TabRole(ETabRole::NomadTab)
+			[
+				DetailsView
+			];
+
+			GetTabManager()->InsertNewDocumentTab(
+				FBlueprintEditorTabs::DetailsID,
+				TEXT("ControlRigMemoryDetails"),
+				FTabManager::FLastMajorOrNomadTab(TEXT("ControlRigMemoryDetails")),
+				DockTab
+			);
+
+			FFunctionGraphTask::CreateAndDispatchWhenReady([DetailsView, InObjects]()
+			{
+				
+				DetailsView->SetObject(InObjects[0]);
+				
+			}, TStatId(), NULL, ENamedThreads::GameThread);
+
+			return;
+		}
+	}
+#endif
+
 	ClearDetailObject();
 
 	if (InObjects.Num() == 1)
diff --git a/Engine/Source/Runtime/RigVM/Private/RigVMCore/RigVMMemoryStorage.cpp b/Engine/Source/Runtime/RigVM/Private/RigVMCore/RigVMMemoryStorage.cpp
--- a/Engine/Source/Runtime/RigVM/Private/RigVMCore/RigVMMemoryStorage.cpp
+++ b/Engine/Source/Runtime/RigVM/Private/RigVMCore/RigVMMemoryStorage.cpp
@@ -444,7 +444,21 @@ FProperty* URigVMMemoryStorageGeneratorClass::AddProperty(URigVMMemoryStorageGen
 			}
 		}
 
-		Result->SetPropertyFlags(CPF_Edit | CPF_NonTransactional);
+		Result->SetPropertyFlags(CPF_Edit | CPF_EditConst | CPF_NonTransactional);
+
+		static const FName NAME_DisplayName(TEXT("DisplayName"));
+		static const FName NAME_ToolTipName(TEXT("ToolTip"));
+
+		FString DisplayName = Result->GetName();
+		while(DisplayName.ReplaceInline(TEXT("__"), TEXT("_")) > 0)
+		{}
+		DisplayName.ReplaceInline(TEXT("_"), TEXT(" "));
+
+		DisplayName = FString::Printf(TEXT("[%d] %s"), InClass->LinkedProperties.Num(), *DisplayName);
+		Result->SetMetaData(NAME_DisplayName, *DisplayName);
+		Result->SetMetaData(NAME_ToolTipName, *FString::Printf(TEXT("Name %s\nCPPType %s"), *Result->GetName(), *InProperty.CPPType));
+
+		InClass->LinkedProperties.Add(Result);
 		(*LinkToProperty) = Result;
 	}
 
